Enabled UHT to digest GENERATED_BODY instead of GENERATED_UCLASS_BODY, GENERATED_USTRUCT_BODY, GENERATED_UINTERFACE_BODY or GENERATED_IINTERFACE_BODY, changed every occurence to the new syntax and fix

diff --git a/Engine/Source/Runtime/Engine/Classes/Materials/MaterialExpressionTextureSample.h b/Engine/Source/Runtime/Engine/Classes/Materials/MaterialExpressionTextureSample.h
--- a/Engine/Source/Runtime/Engine/Classes/Materials/MaterialExpressionTextureSample.h
+++ b/Engine/Source/Runtime/Engine/Classes/Materials/MaterialExpressionTextureSample.h
@@ -19,7 +19,9 @@ enum ETextureMipValueMode
 UCLASS(collapsecategories, hidecategories=Object)
 class ENGINE_API UMaterialExpressionTextureSample : public UMaterialExpressionTextureBase
 {
-	GENERATED_UCLASS_BODY()
+	GENERATED_BODY()
+public:
+	UMaterialExpressionTextureSample(const FObjectInitializer& ObjectInitializer = FObjectInitializer::Get());
 
 	UPROPERTY(meta = (RequiredInput = "false", ToolTip = "Defaults to 'ConstCoordinate' if not specified"))
 	FExpressionInput Coordinates;
diff --git a/Engine/Source/Runtime/Engine/Classes/Materials/MaterialExpressionTextureSampleParameter.h b/Engine/Source/Runtime/Engine/Classes/Materials/MaterialExpressionTextureSampleParameter.h
--- a/Engine/Source/Runtime/Engine/Classes/Materials/MaterialExpressionTextureSampleParameter.h
+++ b/Engine/Source/Runtime/Engine/Classes/Materials/MaterialExpressionTextureSampleParameter.h
@@ -8,7 +8,9 @@
 UCLASS(collapsecategories, abstract, hidecategories=Object)
 class ENGINE_API UMaterialExpressionTextureSampleParameter : public UMaterialExpressionTextureSample
 {
-	GENERATED_UCLASS_BODY()
+	GENERATED_BODY()
+public:
+	UMaterialExpressionTextureSampleParameter(const FObjectInitializer& ObjectInitializer = FObjectInitializer::Get());
 
 	UPROPERTY(EditAnywhere, Category=MaterialExpressionTextureSampleParameter)
 	FName ParameterName;
diff --git a/Engine/Source/Runtime/Engine/Classes/Materials/MaterialExpressionTextureSampleParameter2D.h b/Engine/Source/Runtime/Engine/Classes/Materials/MaterialExpressionTextureSampleParameter2D.h
--- a/Engine/Source/Runtime/Engine/Classes/Materials/MaterialExpressionTextureSampleParameter2D.h
+++ b/Engine/Source/Runtime/Engine/Classes/Materials/MaterialExpressionTextureSampleParameter2D.h
@@ -8,7 +8,9 @@
 UCLASS(collapsecategories, hidecategories=Object)
 class ENGINE_API UMaterialExpressionTextureSampleParameter2D : public UMaterialExpressionTextureSampleParameter
 {
-	GENERATED_UCLASS_BODY()
+	GENERATED_BODY()
+public:
+	UMaterialExpressionTextureSampleParameter2D(const FObjectInitializer& ObjectInitializer = FObjectInitializer::Get());
 
 
 	// Begin UMaterialExpression Interface
diff --git a/Engine/Source/Runtime/Engine/Classes/Materials/MaterialExpressionTextureSampleParameterCube.h b/Engine/Source/Runtime/Engine/Classes/Materials/MaterialExpressionTextureSampleParameterCube.h
--- a/Engine/Source/Runtime/Engine/Classes/Materials/MaterialExpressionTextureSampleParameterCube.h
+++ b/Engine/Source/Runtime/Engine/Classes/Materials/MaterialExpressionTextureSampleParameterCube.h
@@ -8,7 +8,9 @@
 UCLASS(collapsecategories, hidecategories=Object, MinimalAPI)
 class UMaterialExpressionTextureSampleParameterCube : public UMaterialExpressionTextureSampleParameter
 {
-	GENERATED_UCLASS_BODY()
+	GENERATED_BODY()
+public:
+	ENGINE_API UMaterialExpressionTextureSampleParameterCube(const FObjectInitializer& ObjectInitializer = FObjectInitializer::Get());
 
 
 
diff --git a/Engine/Source/Runtime/Engine/Classes/Materials/MaterialExpressionTextureSampleParameterSubUV.h b/Engine/Source/Runtime/Engine/Classes/Materials/MaterialExpressionTextureSampleParameterSubUV.h
--- a/Engine/Source/Runtime/Engine/Classes/Materials/MaterialExpressionTextureSampleParameterSubUV.h
+++ b/Engine/Source/Runtime/Engine/Classes/Materials/MaterialExpressionTextureSampleParameterSubUV.h
@@ -8,7 +8,9 @@
 UCLASS(collapsecategories, hidecategories=Object, MinimalAPI)
 class UMaterialExpressionTextureSampleParameterSubUV : public UMaterialExpressionTextureSampleParameter2D
 {
-	GENERATED_UCLASS_BODY()
+	GENERATED_BODY()
+public:
+	ENGINE_API UMaterialExpressionTextureSampleParameterSubUV(const FObjectInitializer& ObjectInitializer = FObjectInitializer::Get());
 
 	UPROPERTY(EditAnywhere, Category=MaterialExpressionTextureSampleParameterSubUV)
 	uint32 bBlend:1;
diff --git a/Engine/Source/Runtime/Engine/Classes/Materials/MaterialExpressionTime.h b/Engine/Source/Runtime/Engine/Classes/Materials/MaterialExpressionTime.h
--- a/Engine/Source/Runtime/Engine/Classes/Materials/MaterialExpressionTime.h
+++ b/Engine/Source/Runtime/Engine/Classes/Materials/MaterialExpressionTime.h
@@ -9,7 +9,9 @@
 UCLASS(collapsecategories, hidecategories=Object)
 class UMaterialExpressionTime : public UMaterialExpression
 {
-	GENERATED_UCLASS_BODY()
+	GENERATED_BODY()
+public:
+	UMaterialExpressionTime(const FObjectInitializer& ObjectInitializer = FObjectInitializer::Get());
 
 	/** This time continues advancing regardless of whether the game is paused. */
 	UPROPERTY(EditAnywhere, Category=MaterialExpressionTime)
diff --git a/Engine/Source/Runtime/Engine/Classes/Materials/MaterialExpressionTransform.h b/Engine/Source/Runtime/Engine/Classes/Materials/MaterialExpressionTransform.h
--- a/Engine/Source/Runtime/Engine/Classes/Materials/MaterialExpressionTransform.h
+++ b/Engine/Source/Runtime/Engine/Classes/Materials/MaterialExpressionTransform.h
@@ -28,7 +28,9 @@ enum EMaterialVectorCoordTransform
 UCLASS(collapsecategories, hidecategories=Object)
 class UMaterialExpressionTransform : public UMaterialExpression
 {
-	GENERATED_UCLASS_BODY()
+	GENERATED_BODY()
+public:
+	UMaterialExpressionTransform(const FObjectInitializer& ObjectInitializer = FObjectInitializer::Get());
 
 	/** input expression for this transform */
 	UPROPERTY()
diff --git a/Engine/Source/Runtime/Engine/Classes/Materials/MaterialExpressionTransformPosition.h b/Engine/Source/Runtime/Engine/Classes/Materials/MaterialExpressionTransformPosition.h
--- a/Engine/Source/Runtime/Engine/Classes/Materials/MaterialExpressionTransformPosition.h
+++ b/Engine/Source/Runtime/Engine/Classes/Materials/MaterialExpressionTransformPosition.h
@@ -16,7 +16,9 @@ enum EMaterialPositionTransformSource
 UCLASS(collapsecategories, hidecategories=Object, MinimalAPI)
 class UMaterialExpressionTransformPosition : public UMaterialExpression
 {
-	GENERATED_UCLASS_BODY()
+	GENERATED_BODY()
+public:
+	ENGINE_API UMaterialExpressionTransformPosition(const FObjectInitializer& ObjectInitializer = FObjectInitializer::Get());
 
 	/** input expression for this transform */
 	UPROPERTY()
diff --git a/Engine/Source/Runtime/Engine/Classes/Materials/MaterialExpressionTwoSidedSign.h b/Engine/Source/Runtime/Engine/Classes/Materials/MaterialExpressionTwoSidedSign.h
--- a/Engine/Source/Runtime/Engine/Classes/Materials/MaterialExpressionTwoSidedSign.h
+++ b/Engine/Source/Runtime/Engine/Classes/Materials/MaterialExpressionTwoSidedSign.h
@@ -8,7 +8,9 @@
 UCLASS(collapsecategories, hidecategories=Object, MinimalAPI)
 class UMaterialExpressionTwoSidedSign : public UMaterialExpression
 {
-	GENERATED_UCLASS_BODY()
+	GENERATED_BODY()
+public:
+	ENGINE_API UMaterialExpressionTwoSidedSign(const FObjectInitializer& ObjectInitializer = FObjectInitializer::Get());
 
 
 	// Begin UMaterialExpression Interface
diff --git a/Engine/Source/Runtime/Engine/Classes/Materials/MaterialExpressionVectorParameter.h b/Engine/Source/Runtime/Engine/Classes/Materials/MaterialExpressionVectorParameter.h
--- a/Engine/Source/Runtime/Engine/Classes/Materials/MaterialExpressionVectorParameter.h
+++ b/Engine/Source/Runtime/Engine/Classes/Materials/MaterialExpressionVectorParameter.h
@@ -8,7 +8,9 @@
 UCLASS(collapsecategories, hidecategories=Object, MinimalAPI)
 class UMaterialExpressionVectorParameter : public UMaterialExpressionParameter
 {
-	GENERATED_UCLASS_BODY()
+	GENERATED_BODY()
+public:
+	ENGINE_API UMaterialExpressionVectorParameter(const FObjectInitializer& ObjectInitializer = FObjectInitializer::Get());
 
 	UPROPERTY(EditAnywhere, Category=MaterialExpressionVectorParameter)
 	FLinearColor DefaultValue;
diff --git a/Engine/Source/Runtime/Engine/Classes/Materials/MaterialExpressionVertexColor.h b/Engine/Source/Runtime/Engine/Classes/Materials/MaterialExpressionVertexColor.h
--- a/Engine/Source/Runtime/Engine/Classes/Materials/MaterialExpressionVertexColor.h
+++ b/Engine/Source/Runtime/Engine/Classes/Materials/MaterialExpressionVertexColor.h
@@ -8,7 +8,9 @@
 UCLASS(collapsecategories, hidecategories=Object, MinimalAPI)
 class UMaterialExpressionVertexColor : public UMaterialExpression
 {
-	GENERATED_UCLASS_BODY()
+	GENERATED_BODY()
+public:
+	ENGINE_API UMaterialExpressionVertexColor(const FObjectInitializer& ObjectInitializer = FObjectInitializer::Get());
 
 
 	// Begin UMaterialExpression Interface
diff --git a/Engine/Source/Runtime/Engine/Classes/Materials/MaterialExpressionVertexNormalWS.h b/Engine/Source/Runtime/Engine/Classes/Materials/MaterialExpressionVertexNormalWS.h
--- a/Engine/Source/Runtime/Engine/Classes/Materials/MaterialExpressionVertexNormalWS.h
+++ b/Engine/Source/Runtime/Engine/Classes/Materials/MaterialExpressionVertexNormalWS.h
@@ -8,7 +8,9 @@
 UCLASS(collapsecategories, hidecategories=Object)
 class UMaterialExpressionVertexNormalWS : public UMaterialExpression
 {
-	GENERATED_UCLASS_BODY()
+	GENERATED_BODY()
+public:
+	UMaterialExpressionVertexNormalWS(const FObjectInitializer& ObjectInitializer = FObjectInitializer::Get());
 
 
 	// Begin UMaterialExpression Interface
diff --git a/Engine/Source/Runtime/Engine/Classes/Materials/MaterialExpressionViewSize.h b/Engine/Source/Runtime/Engine/Classes/Materials/MaterialExpressionViewSize.h
--- a/Engine/Source/Runtime/Engine/Classes/Materials/MaterialExpressionViewSize.h
+++ b/Engine/Source/Runtime/Engine/Classes/Materials/MaterialExpressionViewSize.h
@@ -8,7 +8,9 @@
 UCLASS(collapsecategories, hidecategories=Object)
 class UMaterialExpressionViewSize : public UMaterialExpression
 {
-	GENERATED_UCLASS_BODY()
+	GENERATED_BODY()
+public:
+	UMaterialExpressionViewSize(const FObjectInitializer& ObjectInitializer = FObjectInitializer::Get());
 
 	// Begin UMaterialExpression Interface
 	virtual int32 Compile(class FMaterialCompiler* Compiler, int32 OutputIndex, int32 MultiplexIndex) override;
diff --git a/Engine/Source/Runtime/Engine/Classes/Materials/MaterialExpressionWorldPosition.h b/Engine/Source/Runtime/Engine/Classes/Materials/MaterialExpressionWorldPosition.h
--- a/Engine/Source/Runtime/Engine/Classes/Materials/MaterialExpressionWorldPosition.h
+++ b/Engine/Source/Runtime/Engine/Classes/Materials/MaterialExpressionWorldPosition.h
@@ -27,7 +27,9 @@ enum EWorldPositionIncludedOffsets
 UCLASS(collapsecategories, hidecategories=Object, MinimalAPI)
 class UMaterialExpressionWorldPosition : public UMaterialExpression
 {
-	GENERATED_UCLASS_BODY()
+	GENERATED_BODY()
+public:
+	ENGINE_API UMaterialExpressionWorldPosition(const FObjectInitializer& ObjectInitializer = FObjectInitializer::Get());
 
 	UPROPERTY(EditAnywhere, Category=UMaterialExpressionWorldPosition, meta=(DisplayName = "Shader Offsets"))
 	TEnumAsByte<EWorldPositionIncludedOffsets> WorldPositionShaderOffset;
diff --git a/Engine/Source/Runtime/Engine/Classes/Materials/MaterialFunction.h b/Engine/Source/Runtime/Engine/Classes/Materials/MaterialFunction.h
--- a/Engine/Source/Runtime/Engine/Classes/Materials/MaterialFunction.h
+++ b/Engine/Source/Runtime/Engine/Classes/Materials/MaterialFunction.h
@@ -9,7 +9,9 @@
 UCLASS(hidecategories=object, MinimalAPI)
 class UMaterialFunction : public UObject
 {
-	GENERATED_UCLASS_BODY()
+	GENERATED_BODY()
+public:
+	ENGINE_API UMaterialFunction(const FObjectInitializer& ObjectInitializer = FObjectInitializer::Get());
 
 	/** Used by materials using this function to know when to recompile. */
 	UPROPERTY(duplicatetransient)
diff --git a/Engine/Source/Runtime/Engine/Classes/Materials/MaterialInstance.h b/Engine/Source/Runtime/Engine/Classes/Materials/MaterialInstance.h
--- a/Engine/Source/Runtime/Engine/Classes/Materials/MaterialInstance.h
+++ b/Engine/Source/Runtime/Engine/Classes/Materials/MaterialInstance.h
@@ -18,7 +18,7 @@ class FSHAHash;
 USTRUCT()
 struct FFontParameterValue
 {
-	GENERATED_USTRUCT_BODY()
+	GENERATED_BODY()
 
 	UPROPERTY(EditAnywhere, BlueprintReadWrite, Category=FontParameterValue)
 	FName ParameterName;
@@ -48,7 +48,7 @@ struct FFontParameterValue
 USTRUCT()
 struct FScalarParameterValue
 {
-	GENERATED_USTRUCT_BODY()
+	GENERATED_BODY()
 
 	UPROPERTY(EditAnywhere, BlueprintReadWrite, Category=ScalarParameterValue)
 	FName ParameterName;
@@ -74,7 +74,7 @@ struct FScalarParameterValue
 USTRUCT()
 struct FTextureParameterValue
 {
-	GENERATED_USTRUCT_BODY()
+	GENERATED_BODY()
 
 	UPROPERTY(EditAnywhere, BlueprintReadWrite, Category=TextureParameterValue)
 	FName ParameterName;
@@ -100,7 +100,7 @@ struct FTextureParameterValue
 USTRUCT()
 struct FVectorParameterValue
 {
-	GENERATED_USTRUCT_BODY()
+	GENERATED_BODY()
 
 	UPROPERTY(EditAnywhere, BlueprintReadWrite, Category=VectorParameterValue)
 	FName ParameterName;
@@ -125,7 +125,9 @@ struct FVectorParameterValue
 UCLASS(abstract, BlueprintType,MinimalAPI)
 class UMaterialInstance : public UMaterialInterface
 {
-	GENERATED_UCLASS_BODY()
+	GENERATED_BODY()
+public:
+	ENGINE_API UMaterialInstance(const FObjectInitializer& ObjectInitializer = FObjectInitializer::Get());
 
 	/** Physical material to use for this graphics material. Used for sounds, effects etc.*/
 	UPROPERTY(EditAnywhere, BlueprintReadWrite, Category=MaterialInstance)
diff --git a/Engine/Source/Runtime/Engine/Classes/Materials/MaterialInstanceActor.h b/Engine/Source/Runtime/Engine/Classes/Materials/MaterialInstanceActor.h
--- a/Engine/Source/Runtime/Engine/Classes/Materials/MaterialInstanceActor.h
+++ b/Engine/Source/Runtime/Engine/Classes/Materials/MaterialInstanceActor.h
@@ -11,7 +11,9 @@
 UCLASS(hidecategories=Movement, hidecategories=Advanced, hidecategories=Collision, hidecategories=Display, hidecategories=Actor, hidecategories=Attachment, MinimalAPI)
 class AMaterialInstanceActor : public AActor
 {
-	GENERATED_UCLASS_BODY()
+	GENERATED_BODY()
+public:
+	ENGINE_API AMaterialInstanceActor(const FObjectInitializer& ObjectInitializer = FObjectInitializer::Get());
 
 	/** Pointer to actors that we want to control paramters of using Matinee. */
 	UPROPERTY(EditAnywhere, BlueprintReadWrite, Category=MaterialInstanceActor)
diff --git a/Engine/Source/Runtime/Engine/Classes/Materials/MaterialInstanceConstant.h b/Engine/Source/Runtime/Engine/Classes/Materials/MaterialInstanceConstant.h
--- a/Engine/Source/Runtime/Engine/Classes/Materials/MaterialInstanceConstant.h
+++ b/Engine/Source/Runtime/Engine/Classes/Materials/MaterialInstanceConstant.h
@@ -14,7 +14,9 @@
 UCLASS(hidecategories=Object, collapsecategories, BlueprintType,MinimalAPI)
 class UMaterialInstanceConstant : public UMaterialInstance
 {
-	GENERATED_UCLASS_BODY()
+	GENERATED_BODY()
+public:
+	ENGINE_API UMaterialInstanceConstant(const FObjectInitializer& ObjectInitializer = FObjectInitializer::Get());
 
 #if WITH_EDITOR
 	/** For constructing new MICs. */
diff --git a/Engine/Source/Runtime/Engine/Classes/Materials/MaterialInstanceDynamic.h b/Engine/Source/Runtime/Engine/Classes/Materials/MaterialInstanceDynamic.h
--- a/Engine/Source/Runtime/Engine/Classes/Materials/MaterialInstanceDynamic.h
+++ b/Engine/Source/Runtime/Engine/Classes/Materials/MaterialInstanceDynamic.h
@@ -8,7 +8,9 @@
 UCLASS(hidecategories=Object, collapsecategories, BlueprintType, MinimalAPI)
 class UMaterialInstanceDynamic : public UMaterialInstance
 {
-	GENERATED_UCLASS_BODY()
+	GENERATED_BODY()
+public:
+	ENGINE_API UMaterialInstanceDynamic(const FObjectInitializer& ObjectInitializer = FObjectInitializer::Get());
 
 	/** Set a MID scalar (float) parameter value */
 	UFUNCTION(BlueprintCallable, meta=(Keywords = "SetFloatParameterValue"), Category="Rendering|Material")
diff --git a/Engine/Source/Runtime/Engine/Classes/Materials/MaterialInterface.h b/Engine/Source/Runtime/Engine/Classes/Materials/MaterialInterface.h
--- a/Engine/Source/Runtime/Engine/Classes/Materials/MaterialInterface.h
+++ b/Engine/Source/Runtime/Engine/Classes/Materials/MaterialInterface.h
@@ -34,7 +34,7 @@ enum EMaterialUsage
 USTRUCT()
 struct ENGINE_API FMaterialRelevance
 {
-	GENERATED_USTRUCT_BODY()
+	GENERATED_BODY()
 
 	UPROPERTY()
 	uint32 bOpaque : 1;
@@ -104,7 +104,7 @@ struct ENGINE_API FMaterialRelevance
 USTRUCT()
 struct FLightmassMaterialInterfaceSettings
 {
-	GENERATED_USTRUCT_BODY()
+	GENERATED_BODY()
 
 	/** If true, forces translucency to cast static shadows as if the material were masked. */
 	UPROPERTY(EditAnywhere, BlueprintReadWrite, Category=Material)
@@ -171,7 +171,9 @@ struct FLightmassMaterialInterfaceSettings
 UCLASS(abstract, BlueprintType,MinimalAPI)
 class UMaterialInterface : public UObject, public IBlendableInterface
 {
-	GENERATED_UCLASS_BODY()
+	GENERATED_BODY()
+public:
+	ENGINE_API UMaterialInterface(const FObjectInitializer& ObjectInitializer = FObjectInitializer::Get());
 
 	/** SubsurfaceProfile, for Screen Space Subsurface Scattering */
 	UPROPERTY(EditAnywhere, BlueprintReadOnly, Category = Material, meta = (DisplayName = "Subsurface Profile"))
@@ -622,4 +624,4 @@ private:
 };
 
 /** Helper function to serialize inline shader maps for the given material resources. */
-extern void SerializeInlineShaderMaps(const TMap<const class ITargetPlatform*,TArray<FMaterialResource*>>* PlatformMaterialResourcesToSave, FArchive& Ar, FMaterialResource* OutMaterialResourcesLoaded[][ERHIFeatureLevel::Num]);
\ No newline at end of file
+extern void SerializeInlineShaderMaps(const TMap<const class ITargetPlatform*,TArray<FMaterialResource*>>* PlatformMaterialResourcesToSave, FArchive& Ar, FMaterialResource* OutMaterialResourcesLoaded[][ERHIFeatureLevel::Num]);
diff --git a/Engine/Source/Runtime/Engine/Classes/Materials/MaterialParameterCollection.h b/Engine/Source/Runtime/Engine/Classes/Materials/MaterialParameterCollection.h
--- a/Engine/Source/Runtime/Engine/Classes/Materials/MaterialParameterCollection.h
+++ b/Engine/Source/Runtime/Engine/Classes/Materials/MaterialParameterCollection.h
@@ -11,7 +11,7 @@
 USTRUCT()
 struct FCollectionParameterBase
 {
-	GENERATED_USTRUCT_BODY()
+	GENERATED_BODY()
 	
 	FCollectionParameterBase()
 	{
@@ -31,7 +31,7 @@ struct FCollectionParameterBase
 USTRUCT()
 struct FCollectionScalarParameter : public FCollectionParameterBase
 {
-	GENERATED_USTRUCT_BODY()
+	GENERATED_BODY()
 	
 	UPROPERTY(EditAnywhere, Category=Parameter)
 	float DefaultValue;
@@ -41,7 +41,7 @@ struct FCollectionScalarParameter : public FCollectionParameterBase
 USTRUCT()
 struct FCollectionVectorParameter : public FCollectionParameterBase
 {
-	GENERATED_USTRUCT_BODY()
+	GENERATED_BODY()
 	
 	UPROPERTY(EditAnywhere, Category=Parameter)
 	FLinearColor DefaultValue;
@@ -54,7 +54,9 @@ struct FCollectionVectorParameter : public FCollectionParameterBase
 UCLASS(hidecategories=object, MinimalAPI)
 class UMaterialParameterCollection : public UObject
 {
-	GENERATED_UCLASS_BODY()
+	GENERATED_BODY()
+public:
+	ENGINE_API UMaterialParameterCollection(const FObjectInitializer& ObjectInitializer = FObjectInitializer::Get());
 
 	/** Used by materials using this collection to know when to recompile. */
 	UPROPERTY(duplicatetransient)
diff --git a/Engine/Source/Runtime/Engine/Classes/Materials/MaterialParameterCollectionInstance.h b/Engine/Source/Runtime/Engine/Classes/Materials/MaterialParameterCollectionInstance.h
--- a/Engine/Source/Runtime/Engine/Classes/Materials/MaterialParameterCollectionInstance.h
+++ b/Engine/Source/Runtime/Engine/Classes/Materials/MaterialParameterCollectionInstance.h
@@ -14,7 +14,9 @@
 UCLASS(hidecategories=object, MinimalAPI)
 class UMaterialParameterCollectionInstance : public UObject
 {
-	GENERATED_UCLASS_BODY()
+	GENERATED_BODY()
+public:
+	ENGINE_API UMaterialParameterCollectionInstance(const FObjectInitializer& ObjectInitializer = FObjectInitializer::Get());
 
 	// Begin UObject interface.
 	ENGINE_API virtual void PostInitProperties();	
diff --git a/Engine/Source/Runtime/Engine/Classes/Matinee/InterpData.h b/Engine/Source/Runtime/Engine/Classes/Matinee/InterpData.h
--- a/Engine/Source/Runtime/Engine/Classes/Matinee/InterpData.h
+++ b/Engine/Source/Runtime/Engine/Classes/Matinee/InterpData.h
@@ -13,7 +13,9 @@
 UCLASS(hidecategories=Object, MinimalAPI)
 class UInterpData : public UObject
 {
-	GENERATED_UCLASS_BODY()
+	GENERATED_BODY()
+public:
+	ENGINE_API UInterpData(const FObjectInitializer& ObjectInitializer = FObjectInitializer::Get());
 
 	/** Duration of interpolation sequence - in seconds. */
 	UPROPERTY()
diff --git a/Engine/Source/Runtime/Engine/Classes/Matinee/InterpFilter.h b/Engine/Source/Runtime/Engine/Classes/Matinee/InterpFilter.h
--- a/Engine/Source/Runtime/Engine/Classes/Matinee/InterpFilter.h
+++ b/Engine/Source/Runtime/Engine/Classes/Matinee/InterpFilter.h
@@ -12,7 +12,9 @@
 UCLASS()
 class UInterpFilter : public UObject
 {
-	GENERATED_UCLASS_BODY()
+	GENERATED_BODY()
+public:
+	UInterpFilter(const FObjectInitializer& ObjectInitializer = FObjectInitializer::Get());
 
 	/** Caption for this filter. */
 	UPROPERTY()
diff --git a/Engine/Source/Runtime/Engine/Classes/Matinee/InterpFilter_Classes.h b/Engine/Source/Runtime/Engine/Classes/Matinee/InterpFilter_Classes.h
--- a/Engine/Source/Runtime/Engine/Classes/Matinee/InterpFilter_Classes.h
+++ b/Engine/Source/Runtime/Engine/Classes/Matinee/InterpFilter_Classes.h
@@ -13,7 +13,9 @@
 UCLASS()
 class UInterpFilter_Classes : public UInterpFilter
 {
-	GENERATED_UCLASS_BODY()
+	GENERATED_BODY()
+public:
+	UInterpFilter_Classes(const FObjectInitializer& ObjectInitializer = FObjectInitializer::Get());
 
 #if WITH_EDITORONLY_DATA
 	/** Which class to filter groups by. */
diff --git a/Engine/Source/Runtime/Engine/Classes/Matinee/InterpFilter_Custom.h b/Engine/Source/Runtime/Engine/Classes/Matinee/InterpFilter_Custom.h
--- a/Engine/Source/Runtime/Engine/Classes/Matinee/InterpFilter_Custom.h
+++ b/Engine/Source/Runtime/Engine/Classes/Matinee/InterpFilter_Custom.h
@@ -13,7 +13,9 @@
 UCLASS(MinimalAPI)
 class UInterpFilter_Custom : public UInterpFilter
 {
-	GENERATED_UCLASS_BODY()
+	GENERATED_BODY()
+public:
+	ENGINE_API UInterpFilter_Custom(const FObjectInitializer& ObjectInitializer = FObjectInitializer::Get());
 
 #if WITH_EDITORONLY_DATA
 	/** Which groups are included in this filter. */
diff --git a/Engine/Source/Runtime/Engine/Classes/Matinee/InterpGroup.h b/Engine/Source/Runtime/Engine/Classes/Matinee/InterpGroup.h
--- a/Engine/Source/Runtime/Engine/Classes/Matinee/InterpGroup.h
+++ b/Engine/Source/Runtime/Engine/Classes/Matinee/InterpGroup.h
@@ -13,7 +13,7 @@
 USTRUCT()
 struct FInterpEdSelKey
 {
-	GENERATED_USTRUCT_BODY()
+	GENERATED_BODY()
 
 	UPROPERTY()
 	class UInterpGroup* Group;
@@ -77,7 +77,9 @@ struct FInterpEdSelKey
 UCLASS(collapsecategories, hidecategories=Object, MinimalAPI)
 class UInterpGroup : public UObject, public FInterpEdInputInterface
 {
-	GENERATED_UCLASS_BODY()
+	GENERATED_BODY()
+public:
+	ENGINE_API UInterpGroup(const FObjectInitializer& ObjectInitializer = FObjectInitializer::Get());
 
 	UPROPERTY(export)
 	TArray<class UInterpTrack*> InterpTracks;
diff --git a/Engine/Source/Runtime/Engine/Classes/Matinee/InterpGroupCamera.h b/Engine/Source/Runtime/Engine/Classes/Matinee/InterpGroupCamera.h
--- a/Engine/Source/Runtime/Engine/Classes/Matinee/InterpGroupCamera.h
+++ b/Engine/Source/Runtime/Engine/Classes/Matinee/InterpGroupCamera.h
@@ -15,7 +15,7 @@
 USTRUCT()
 struct FCameraPreviewInfo
 {
-	GENERATED_USTRUCT_BODY()
+	GENERATED_BODY()
 
 	UPROPERTY(EditAnywhere, Category=CameraPreviewInfo)
 	TSubclassOf<class APawn>  PawnClass;
@@ -50,7 +50,9 @@ struct FCameraPreviewInfo
 UCLASS(collapsecategories, hidecategories=Object, MinimalAPI)
 class UInterpGroupCamera : public UInterpGroup
 {
-	GENERATED_UCLASS_BODY()
+	GENERATED_BODY()
+public:
+	ENGINE_API UInterpGroupCamera(const FObjectInitializer& ObjectInitializer = FObjectInitializer::Get());
 
 	UPROPERTY(transient)
 	class UCameraAnim* CameraAnimInst;
diff --git a/Engine/Source/Runtime/Engine/Classes/Matinee/InterpGroupDirector.h b/Engine/Source/Runtime/Engine/Classes/Matinee/InterpGroupDirector.h
--- a/Engine/Source/Runtime/Engine/Classes/Matinee/InterpGroupDirector.h
+++ b/Engine/Source/Runtime/Engine/Classes/Matinee/InterpGroupDirector.h
@@ -12,7 +12,9 @@
 UCLASS(collapsecategories, hidecategories=Object, MinimalAPI)
 class UInterpGroupDirector : public UInterpGroup
 {
-	GENERATED_UCLASS_BODY()
+	GENERATED_BODY()
+public:
+	ENGINE_API UInterpGroupDirector(const FObjectInitializer& ObjectInitializer = FObjectInitializer::Get());
 
 	// Begin UInterpGroup Interface
 	virtual AActor* SelectGroupActor( class UInterpGroupInst* GrInst, bool bDeselectActors ) override;
diff --git a/Engine/Source/Runtime/Engine/Classes/Matinee/InterpGroupInst.h b/Engine/Source/Runtime/Engine/Classes/Matinee/InterpGroupInst.h
--- a/Engine/Source/Runtime/Engine/Classes/Matinee/InterpGroupInst.h
+++ b/Engine/Source/Runtime/Engine/Classes/Matinee/InterpGroupInst.h
@@ -6,7 +6,9 @@
 UCLASS(MinimalAPI)
 class UInterpGroupInst : public UObject
 {
-	GENERATED_UCLASS_BODY()
+	GENERATED_BODY()
+public:
+	ENGINE_API UInterpGroupInst(const FObjectInitializer& ObjectInitializer = FObjectInitializer::Get());
 
 	/**
 	 *
diff --git a/Engine/Source/Runtime/Engine/Classes/Matinee/InterpGroupInstCamera.h b/Engine/Source/Runtime/Engine/Classes/Matinee/InterpGroupInstCamera.h
--- a/Engine/Source/Runtime/Engine/Classes/Matinee/InterpGroupInstCamera.h
+++ b/Engine/Source/Runtime/Engine/Classes/Matinee/InterpGroupInstCamera.h
@@ -7,7 +7,9 @@
 UCLASS()
 class UInterpGroupInstCamera : public UInterpGroupInst
 {
-	GENERATED_UCLASS_BODY()
+	GENERATED_BODY()
+public:
+	UInterpGroupInstCamera(const FObjectInitializer& ObjectInitializer = FObjectInitializer::Get());
 
 };
 
diff --git a/Engine/Source/Runtime/Engine/Classes/Matinee/InterpGroupInstDirector.h b/Engine/Source/Runtime/Engine/Classes/Matinee/InterpGroupInstDirector.h
--- a/Engine/Source/Runtime/Engine/Classes/Matinee/InterpGroupInstDirector.h
+++ b/Engine/Source/Runtime/Engine/Classes/Matinee/InterpGroupInstDirector.h
@@ -8,7 +8,9 @@
 UCLASS(MinimalAPI)
 class UInterpGroupInstDirector : public UInterpGroupInst
 {
-	GENERATED_UCLASS_BODY()
+	GENERATED_BODY()
+public:
+	ENGINE_API UInterpGroupInstDirector(const FObjectInitializer& ObjectInitializer = FObjectInitializer::Get());
 
 };
 
diff --git a/Engine/Source/Runtime/Engine/Classes/Matinee/InterpTrack.h b/Engine/Source/Runtime/Engine/Classes/Matinee/InterpTrack.h
--- a/Engine/Source/Runtime/Engine/Classes/Matinee/InterpTrack.h
+++ b/Engine/Source/Runtime/Engine/Classes/Matinee/InterpTrack.h
@@ -33,7 +33,7 @@ enum ETrackActiveCondition
 USTRUCT()
 struct FSupportedSubTrackInfo
 {
-	GENERATED_USTRUCT_BODY()
+	GENERATED_BODY()
 
 	/** The sub track class which is supported by this track */
 	UPROPERTY()
@@ -60,7 +60,7 @@ struct FSupportedSubTrackInfo
 USTRUCT()
 struct FSubTrackGroup
 {
-	GENERATED_USTRUCT_BODY()
+	GENERATED_BODY()
 
 	/** Name of the subtrack  group */
 	UPROPERTY()
@@ -90,7 +90,9 @@ struct FSubTrackGroup
 UCLASS(collapsecategories, hidecategories=Object, abstract, MinimalAPI)
 class UInterpTrack : public UObject, public FCurveEdInterface, public FInterpEdInputInterface
 {
-	GENERATED_UCLASS_BODY()
+	GENERATED_BODY()
+public:
+	ENGINE_API UInterpTrack(const FObjectInitializer& ObjectInitializer = FObjectInitializer::Get());
 
 	/** A list of subtracks that belong to this track */
 	UPROPERTY()
diff --git a/Engine/Source/Runtime/Engine/Classes/Matinee/InterpTrackAnimControl.h b/Engine/Source/Runtime/Engine/Classes/Matinee/InterpTrackAnimControl.h
--- a/Engine/Source/Runtime/Engine/Classes/Matinee/InterpTrackAnimControl.h
+++ b/Engine/Source/Runtime/Engine/Classes/Matinee/InterpTrackAnimControl.h
@@ -8,7 +8,7 @@
 USTRUCT()
 struct FAnimControlTrackKey
 {
-	GENERATED_USTRUCT_BODY()
+	GENERATED_BODY()
 
 	/** Position in the Matinee sequence to start playing this animation. */
 	UPROPERTY()
@@ -55,7 +55,9 @@ struct FAnimControlTrackKey
 UCLASS(MinimalAPI, meta=( DisplayName = "Animation Track" ))
 class UInterpTrackAnimControl : public UInterpTrackFloatBase
 {
-	GENERATED_UCLASS_BODY()
+	GENERATED_BODY()
+public:
+	ENGINE_API UInterpTrackAnimControl(const FObjectInitializer& ObjectInitializer = FObjectInitializer::Get());
 	
 	/** 
 	 *	Name of slot to use when playing animation. Passed to Actor. 
diff --git a/Engine/Source/Runtime/Engine/Classes/Matinee/InterpTrackAudioMaster.h b/Engine/Source/Runtime/Engine/Classes/Matinee/InterpTrackAudioMaster.h
--- a/Engine/Source/Runtime/Engine/Classes/Matinee/InterpTrackAudioMaster.h
+++ b/Engine/Source/Runtime/Engine/Classes/Matinee/InterpTrackAudioMaster.h
@@ -8,7 +8,9 @@
 UCLASS(MinimalAPI, meta=( DisplayName = "Audio Master Track" ) )
 class UInterpTrackAudioMaster : public UInterpTrackVectorBase
 {
-	GENERATED_UCLASS_BODY()
+	GENERATED_BODY()
+public:
+	ENGINE_API UInterpTrackAudioMaster(const FObjectInitializer& ObjectInitializer = FObjectInitializer::Get());
 
 
 	// Begin UInterpTrack interface.
diff --git a/Engine/Source/Runtime/Engine/Classes/Matinee/InterpTrackBoolProp.h b/Engine/Source/Runtime/Engine/Classes/Matinee/InterpTrackBoolProp.h
--- a/Engine/Source/Runtime/Engine/Classes/Matinee/InterpTrackBoolProp.h
+++ b/Engine/Source/Runtime/Engine/Classes/Matinee/InterpTrackBoolProp.h
@@ -8,7 +8,7 @@
 USTRUCT()
 struct FBoolTrackKey
 {
-	GENERATED_USTRUCT_BODY()
+	GENERATED_BODY()
 
 	UPROPERTY()
 	float Time;
@@ -28,7 +28,9 @@ struct FBoolTrackKey
 UCLASS(MinimalAPI, meta=( DisplayName = "Bool Property Track" ) )
 class UInterpTrackBoolProp : public UInterpTrack
 {
-	GENERATED_UCLASS_BODY()
+	GENERATED_BODY()
+public:
+	ENGINE_API UInterpTrackBoolProp(const FObjectInitializer& ObjectInitializer = FObjectInitializer::Get());
 
 	/** Array of booleans to set. */
 	UPROPERTY()
diff --git a/Engine/Source/Runtime/Engine/Classes/Matinee/InterpTrackColorProp.h b/Engine/Source/Runtime/Engine/Classes/Matinee/InterpTrackColorProp.h
--- a/Engine/Source/Runtime/Engine/Classes/Matinee/InterpTrackColorProp.h
+++ b/Engine/Source/Runtime/Engine/Classes/Matinee/InterpTrackColorProp.h
@@ -8,7 +8,9 @@
 UCLASS(MinimalAPI, meta=( DisplayName = "Color Property Track" ) )
 class UInterpTrackColorProp : public UInterpTrackVectorBase
 {
-	GENERATED_UCLASS_BODY()
+	GENERATED_BODY()
+public:
+	ENGINE_API UInterpTrackColorProp(const FObjectInitializer& ObjectInitializer = FObjectInitializer::Get());
 
 	/** Name of property in Group  AActor  which this track mill modify over time. */
 	UPROPERTY(Category=InterpTrackColorProp, VisibleAnywhere)
diff --git a/Engine/Source/Runtime/Engine/Classes/Matinee/InterpTrackColorScale.h b/Engine/Source/Runtime/Engine/Classes/Matinee/InterpTrackColorScale.h
--- a/Engine/Source/Runtime/Engine/Classes/Matinee/InterpTrackColorScale.h
+++ b/Engine/Source/Runtime/Engine/Classes/Matinee/InterpTrackColorScale.h
@@ -12,7 +12,9 @@
 UCLASS(meta=( DisplayName = "Color Scale Track" ) )
 class UInterpTrackColorScale : public UInterpTrackVectorBase
 {
-	GENERATED_UCLASS_BODY()
+	GENERATED_BODY()
+public:
+	UInterpTrackColorScale(const FObjectInitializer& ObjectInitializer = FObjectInitializer::Get());
 
 
 	// Begin UInterpTrack interface.
diff --git a/Engine/Source/Runtime/Engine/Classes/Matinee/InterpTrackDirector.h b/Engine/Source/Runtime/Engine/Classes/Matinee/InterpTrackDirector.h
--- a/Engine/Source/Runtime/Engine/Classes/Matinee/InterpTrackDirector.h
+++ b/Engine/Source/Runtime/Engine/Classes/Matinee/InterpTrackDirector.h
@@ -14,7 +14,7 @@
 USTRUCT()
 struct FDirectorTrackCut
 {
-	GENERATED_USTRUCT_BODY()
+	GENERATED_BODY()
 
 	/** Time to perform the cut. */
 	UPROPERTY()
@@ -45,7 +45,9 @@ struct FDirectorTrackCut
 UCLASS(MinimalAPI, meta=( DisplayName = "Director Track" ) )
 class UInterpTrackDirector : public UInterpTrack
 {
-	GENERATED_UCLASS_BODY()
+	GENERATED_BODY()
+public:
+	ENGINE_API UInterpTrackDirector(const FObjectInitializer& ObjectInitializer = FObjectInitializer::Get());
 
 	/** Array of cuts between cameras. */
 	UPROPERTY()
diff --git a/Engine/Source/Runtime/Engine/Classes/Matinee/InterpTrackEvent.h b/Engine/Source/Runtime/Engine/Classes/Matinee/InterpTrackEvent.h
--- a/Engine/Source/Runtime/Engine/Classes/Matinee/InterpTrackEvent.h
+++ b/Engine/Source/Runtime/Engine/Classes/Matinee/InterpTrackEvent.h
@@ -15,7 +15,7 @@
 USTRUCT()
 struct FEventTrackKey
 {
-	GENERATED_USTRUCT_BODY()
+	GENERATED_BODY()
 
 	UPROPERTY()
 	float Time;
@@ -34,7 +34,9 @@ struct FEventTrackKey
 UCLASS(MinimalAPI, meta=( DisplayName = "Event Track" ) )
 class UInterpTrackEvent : public UInterpTrack
 {
-	GENERATED_UCLASS_BODY()
+	GENERATED_BODY()
+public:
+	ENGINE_API UInterpTrackEvent(const FObjectInitializer& ObjectInitializer = FObjectInitializer::Get());
 
 	/** Array of events to fire off. */
 	UPROPERTY()
diff --git a/Engine/Source/Runtime/Engine/Classes/Matinee/InterpTrackFade.h b/Engine/Source/Runtime/Engine/Classes/Matinee/InterpTrackFade.h
--- a/Engine/Source/Runtime/Engine/Classes/Matinee/InterpTrackFade.h
+++ b/Engine/Source/Runtime/Engine/Classes/Matinee/InterpTrackFade.h
@@ -7,7 +7,9 @@
 UCLASS(meta=( DisplayName = "Fade Track" ) )
 class UInterpTrackFade : public UInterpTrackFloatBase
 {
-	GENERATED_UCLASS_BODY()
+	GENERATED_BODY()
+public:
+	UInterpTrackFade(const FObjectInitializer& ObjectInitializer = FObjectInitializer::Get());
 
 	/** 
 	 * InterpTrackFade
diff --git a/Engine/Source/Runtime/Engine/Classes/Matinee/InterpTrackFloatAnimBPParam.h b/Engine/Source/Runtime/Engine/Classes/Matinee/InterpTrackFloatAnimBPParam.h
--- a/Engine/Source/Runtime/Engine/Classes/Matinee/InterpTrackFloatAnimBPParam.h
+++ b/Engine/Source/Runtime/Engine/Classes/Matinee/InterpTrackFloatAnimBPParam.h
@@ -6,7 +6,9 @@
 UCLASS(meta=( DisplayName = "Float Anim BP Parameter Track" ) )
 class UInterpTrackFloatAnimBPParam : public UInterpTrackFloatBase
 {
-	GENERATED_UCLASS_BODY()
+	GENERATED_BODY()
+public:
+	UInterpTrackFloatAnimBPParam(const FObjectInitializer& ObjectInitializer = FObjectInitializer::Get());
 	
 	/** Materials whose parameters we want to change and the references to those materials. */
 	UPROPERTY(EditAnywhere, Category=InterpTrackFloatAnimBPParam)
diff --git a/Engine/Source/Runtime/Engine/Classes/Matinee/InterpTrackFloatBase.h b/Engine/Source/Runtime/Engine/Classes/Matinee/InterpTrackFloatBase.h
--- a/Engine/Source/Runtime/Engine/Classes/Matinee/InterpTrackFloatBase.h
+++ b/Engine/Source/Runtime/Engine/Classes/Matinee/InterpTrackFloatBase.h
@@ -7,7 +7,9 @@
 UCLASS(abstract, MinimalAPI)
 class UInterpTrackFloatBase : public UInterpTrack
 {
-	GENERATED_UCLASS_BODY()
+	GENERATED_BODY()
+public:
+	ENGINE_API UInterpTrackFloatBase(const FObjectInitializer& ObjectInitializer = FObjectInitializer::Get());
 	
 	/** Actually track data containing keyframes of float as it varies over time. */
 	UPROPERTY()
diff --git a/Engine/Source/Runtime/Engine/Classes/Matinee/InterpTrackFloatMaterialParam.h b/Engine/Source/Runtime/Engine/Classes/Matinee/InterpTrackFloatMaterialParam.h
--- a/Engine/Source/Runtime/Engine/Classes/Matinee/InterpTrackFloatMaterialParam.h
+++ b/Engine/Source/Runtime/Engine/Classes/Matinee/InterpTrackFloatMaterialParam.h
@@ -6,7 +6,9 @@
 UCLASS(meta=( DisplayName = "Float Material Parameter Track" ) )
 class UInterpTrackFloatMaterialParam : public UInterpTrackFloatBase
 {
-	GENERATED_UCLASS_BODY()
+	GENERATED_BODY()
+public:
+	UInterpTrackFloatMaterialParam(const FObjectInitializer& ObjectInitializer = FObjectInitializer::Get());
 	
 	/** Materials whose parameters we want to change and the references to those materials. */
 	UPROPERTY(EditAnywhere, Category=InterpTrackFloatMaterialParam)
diff --git a/Engine/Source/Runtime/Engine/Classes/Matinee/InterpTrackFloatParticleParam.h b/Engine/Source/Runtime/Engine/Classes/Matinee/InterpTrackFloatParticleParam.h
--- a/Engine/Source/Runtime/Engine/Classes/Matinee/InterpTrackFloatParticleParam.h
+++ b/Engine/Source/Runtime/Engine/Classes/Matinee/InterpTrackFloatParticleParam.h
@@ -6,7 +6,9 @@
 UCLASS(meta=( DisplayName = "Float Particle Param Track" ) )
 class UInterpTrackFloatParticleParam : public UInterpTrackFloatBase
 {
-	GENERATED_UCLASS_BODY()
+	GENERATED_BODY()
+public:
+	UInterpTrackFloatParticleParam(const FObjectInitializer& ObjectInitializer = FObjectInitializer::Get());
 	
 	/** Name of property in the Emitter which this track mill modify over time. */
 	UPROPERTY(EditAnywhere, Category=InterpTrackFloatParticleParam)
diff --git a/Engine/Source/Runtime/Engine/Classes/Matinee/InterpTrackFloatProp.h b/Engine/Source/Runtime/Engine/Classes/Matinee/InterpTrackFloatProp.h
--- a/Engine/Source/Runtime/Engine/Classes/Matinee/InterpTrackFloatProp.h
+++ b/Engine/Source/Runtime/Engine/Classes/Matinee/InterpTrackFloatProp.h
@@ -7,7 +7,9 @@
 UCLASS(MinimalAPI, meta=( DisplayName = "Float Property Track" ) )
 class UInterpTrackFloatProp : public UInterpTrackFloatBase
 {
-	GENERATED_UCLASS_BODY()
+	GENERATED_BODY()
+public:
+	ENGINE_API UInterpTrackFloatProp(const FObjectInitializer& ObjectInitializer = FObjectInitializer::Get());
 	
 	/** Name of property in Group  AActor  which this track mill modify over time. */
 	UPROPERTY(Category=InterpTrackFloatProp, VisibleAnywhere)
diff --git a/Engine/Source/Runtime/Engine/Classes/Matinee/InterpTrackInst.h b/Engine/Source/Runtime/Engine/Classes/Matinee/InterpTrackInst.h
--- a/Engine/Source/Runtime/Engine/Classes/Matinee/InterpTrackInst.h
+++ b/Engine/Source/Runtime/Engine/Classes/Matinee/InterpTrackInst.h
@@ -14,7 +14,9 @@ class UInterpTrack;
 UCLASS(MinimalAPI)
 class UInterpTrackInst : public UObject
 {
-	GENERATED_UCLASS_BODY()
+	GENERATED_BODY()
+public:
+	ENGINE_API UInterpTrackInst(const FObjectInitializer& ObjectInitializer = FObjectInitializer::Get());
 
 
 	/** 
diff --git a/Engine/Source/Runtime/Engine/Classes/Matinee/InterpTrackInstAnimControl.h b/Engine/Source/Runtime/Engine/Classes/Matinee/InterpTrackInstAnimControl.h
--- a/Engine/Source/Runtime/Engine/Classes/Matinee/InterpTrackInstAnimControl.h
+++ b/Engine/Source/Runtime/Engine/Classes/Matinee/InterpTrackInstAnimControl.h
@@ -6,7 +6,9 @@
 UCLASS(MinimalAPI)
 class UInterpTrackInstAnimControl : public UInterpTrackInst
 {
-	GENERATED_UCLASS_BODY()
+	GENERATED_BODY()
+public:
+	ENGINE_API UInterpTrackInstAnimControl(const FObjectInitializer& ObjectInitializer = FObjectInitializer::Get());
 
 	/**
 	 */
diff --git a/Engine/Source/Runtime/Engine/Classes/Matinee/InterpTrackInstAudioMaster.h b/Engine/Source/Runtime/Engine/Classes/Matinee/InterpTrackInstAudioMaster.h
--- a/Engine/Source/Runtime/Engine/Classes/Matinee/InterpTrackInstAudioMaster.h
+++ b/Engine/Source/Runtime/Engine/Classes/Matinee/InterpTrackInstAudioMaster.h
@@ -7,7 +7,9 @@
 UCLASS()
 class UInterpTrackInstAudioMaster : public UInterpTrackInst
 {
-	GENERATED_UCLASS_BODY()
+	GENERATED_BODY()
+public:
+	UInterpTrackInstAudioMaster(const FObjectInitializer& ObjectInitializer = FObjectInitializer::Get());
 
 
 	// Begin UInterpTrackInst Instance
diff --git a/Engine/Source/Runtime/Engine/Classes/Matinee/InterpTrackInstBoolProp.h b/Engine/Source/Runtime/Engine/Classes/Matinee/InterpTrackInstBoolProp.h
--- a/Engine/Source/Runtime/Engine/Classes/Matinee/InterpTrackInstBoolProp.h
+++ b/Engine/Source/Runtime/Engine/Classes/Matinee/InterpTrackInstBoolProp.h
@@ -7,7 +7,9 @@
 UCLASS()
 class UInterpTrackInstBoolProp : public UInterpTrackInstProperty
 {
-	GENERATED_UCLASS_BODY()
+	GENERATED_BODY()
+public:
+	UInterpTrackInstBoolProp(const FObjectInitializer& ObjectInitializer = FObjectInitializer::Get());
 
 	/** Pointer to boolean property in TrackObject. */
 	void* BoolPropertyAddress;
diff --git a/Engine/Source/Runtime/Engine/Classes/Matinee/InterpTrackInstColorProp.h b/Engine/Source/Runtime/Engine/Classes/Matinee/InterpTrackInstColorProp.h
--- a/Engine/Source/Runtime/Engine/Classes/Matinee/InterpTrackInstColorProp.h
+++ b/Engine/Source/Runtime/Engine/Classes/Matinee/InterpTrackInstColorProp.h
@@ -7,7 +7,9 @@
 UCLASS()
 class UInterpTrackInstColorProp : public UInterpTrackInstProperty
 {
-	GENERATED_UCLASS_BODY()
+	GENERATED_BODY()
+public:
+	UInterpTrackInstColorProp(const FObjectInitializer& ObjectInitializer = FObjectInitializer::Get());
 
 	/** Pointer to color property in TrackObject. */
 	FColor* ColorProp;
diff --git a/Engine/Source/Runtime/Engine/Classes/Matinee/InterpTrackInstColorScale.h b/Engine/Source/Runtime/Engine/Classes/Matinee/InterpTrackInstColorScale.h
--- a/Engine/Source/Runtime/Engine/Classes/Matinee/InterpTrackInstColorScale.h
+++ b/Engine/Source/Runtime/Engine/Classes/Matinee/InterpTrackInstColorScale.h
@@ -6,7 +6,9 @@
 UCLASS()
 class UInterpTrackInstColorScale : public UInterpTrackInst
 {
-	GENERATED_UCLASS_BODY()
+	GENERATED_BODY()
+public:
+	UInterpTrackInstColorScale(const FObjectInitializer& ObjectInitializer = FObjectInitializer::Get());
 
 
 	// Begin UInterpTrackInst Instance
diff --git a/Engine/Source/Runtime/Engine/Classes/Matinee/InterpTrackInstDirector.h b/Engine/Source/Runtime/Engine/Classes/Matinee/InterpTrackInstDirector.h
--- a/Engine/Source/Runtime/Engine/Classes/Matinee/InterpTrackInstDirector.h
+++ b/Engine/Source/Runtime/Engine/Classes/Matinee/InterpTrackInstDirector.h
@@ -8,7 +8,9 @@
 UCLASS(MinimalAPI)
 class UInterpTrackInstDirector : public UInterpTrackInst
 {
-	GENERATED_UCLASS_BODY()
+	GENERATED_BODY()
+public:
+	ENGINE_API UInterpTrackInstDirector(const FObjectInitializer& ObjectInitializer = FObjectInitializer::Get());
 
 	UPROPERTY()
 	class AActor* OldViewTarget;
diff --git a/Engine/Source/Runtime/Engine/Classes/Matinee/InterpTrackInstEvent.h b/Engine/Source/Runtime/Engine/Classes/Matinee/InterpTrackInstEvent.h
--- a/Engine/Source/Runtime/Engine/Classes/Matinee/InterpTrackInstEvent.h
+++ b/Engine/Source/Runtime/Engine/Classes/Matinee/InterpTrackInstEvent.h
@@ -7,7 +7,9 @@
 UCLASS(MinimalAPI)
 class UInterpTrackInstEvent : public UInterpTrackInst
 {
-	GENERATED_UCLASS_BODY()
+	GENERATED_BODY()
+public:
+	ENGINE_API UInterpTrackInstEvent(const FObjectInitializer& ObjectInitializer = FObjectInitializer::Get());
 
 	/** 
 	 *	Position we were in last time we evaluated Events. 
diff --git a/Engine/Source/Runtime/Engine/Classes/Matinee/InterpTrackInstFade.h b/Engine/Source/Runtime/Engine/Classes/Matinee/InterpTrackInstFade.h
--- a/Engine/Source/Runtime/Engine/Classes/Matinee/InterpTrackInstFade.h
+++ b/Engine/Source/Runtime/Engine/Classes/Matinee/InterpTrackInstFade.h
@@ -7,7 +7,9 @@
 UCLASS()
 class UInterpTrackInstFade : public UInterpTrackInst
 {
-	GENERATED_UCLASS_BODY()
+	GENERATED_BODY()
+public:
+	UInterpTrackInstFade(const FObjectInitializer& ObjectInitializer = FObjectInitializer::Get());
 
 
 	// Begin UInterpTrackInst Instance
diff --git a/Engine/Source/Runtime/Engine/Classes/Matinee/InterpTrackInstFloatAnimBPParam.h b/Engine/Source/Runtime/Engine/Classes/Matinee/InterpTrackInstFloatAnimBPParam.h
--- a/Engine/Source/Runtime/Engine/Classes/Matinee/InterpTrackInstFloatAnimBPParam.h
+++ b/Engine/Source/Runtime/Engine/Classes/Matinee/InterpTrackInstFloatAnimBPParam.h
@@ -7,7 +7,9 @@
 UCLASS()
 class UInterpTrackInstFloatAnimBPParam : public UInterpTrackInst
 {
-	GENERATED_UCLASS_BODY()
+	GENERATED_BODY()
+public:
+	UInterpTrackInstFloatAnimBPParam(const FObjectInitializer& ObjectInitializer = FObjectInitializer::Get());
 
 	/** MIDs we're using to set the desired parameter. */
 	UPROPERTY(transient)
diff --git a/Engine/Source/Runtime/Engine/Classes/Matinee/InterpTrackInstFloatMaterialParam.h b/Engine/Source/Runtime/Engine/Classes/Matinee/InterpTrackInstFloatMaterialParam.h
--- a/Engine/Source/Runtime/Engine/Classes/Matinee/InterpTrackInstFloatMaterialParam.h
+++ b/Engine/Source/Runtime/Engine/Classes/Matinee/InterpTrackInstFloatMaterialParam.h
@@ -7,7 +7,9 @@
 UCLASS()
 class UInterpTrackInstFloatMaterialParam : public UInterpTrackInst
 {
-	GENERATED_UCLASS_BODY()
+	GENERATED_BODY()
+public:
+	UInterpTrackInstFloatMaterialParam(const FObjectInitializer& ObjectInitializer = FObjectInitializer::Get());
 
 	/** MIDs we're using to set the desired parameter. */
 	UPROPERTY()
diff --git a/Engine/Source/Runtime/Engine/Classes/Matinee/InterpTrackInstFloatParticleParam.h b/Engine/Source/Runtime/Engine/Classes/Matinee/InterpTrackInstFloatParticleParam.h
--- a/Engine/Source/Runtime/Engine/Classes/Matinee/InterpTrackInstFloatParticleParam.h
+++ b/Engine/Source/Runtime/Engine/Classes/Matinee/InterpTrackInstFloatParticleParam.h
@@ -7,7 +7,9 @@
 UCLASS()
 class UInterpTrackInstFloatParticleParam : public UInterpTrackInst
 {
-	GENERATED_UCLASS_BODY()
+	GENERATED_BODY()
+public:
+	UInterpTrackInstFloatParticleParam(const FObjectInitializer& ObjectInitializer = FObjectInitializer::Get());
 
 	/** Saved value for restoring state when exiting Matinee. */
 	UPROPERTY()
diff --git a/Engine/Source/Runtime/Engine/Classes/Matinee/InterpTrackInstFloatProp.h b/Engine/Source/Runtime/Engine/Classes/Matinee/InterpTrackInstFloatProp.h
--- a/Engine/Source/Runtime/Engine/Classes/Matinee/InterpTrackInstFloatProp.h
+++ b/Engine/Source/Runtime/Engine/Classes/Matinee/InterpTrackInstFloatProp.h
@@ -8,7 +8,9 @@
 UCLASS(MinimalAPI)
 class UInterpTrackInstFloatProp : public UInterpTrackInstProperty
 {
-	GENERATED_UCLASS_BODY()
+	GENERATED_BODY()
+public:
+	ENGINE_API UInterpTrackInstFloatProp(const FObjectInitializer& ObjectInitializer = FObjectInitializer::Get());
 
 	/** Pointer to float property in TrackObject. */
 	float* FloatProp;
diff --git a/Engine/Source/Runtime/Engine/Classes/Matinee/InterpTrackInstLinearColorProp.h b/Engine/Source/Runtime/Engine/Classes/Matinee/InterpTrackInstLinearColorProp.h
--- a/Engine/Source/Runtime/Engine/Classes/Matinee/InterpTrackInstLinearColorProp.h
+++ b/Engine/Source/Runtime/Engine/Classes/Matinee/InterpTrackInstLinearColorProp.h
@@ -7,7 +7,9 @@
 UCLASS()
 class UInterpTrackInstLinearColorProp : public UInterpTrackInstProperty
 {
-	GENERATED_UCLASS_BODY()
+	GENERATED_BODY()
+public:
+	UInterpTrackInstLinearColorProp(const FObjectInitializer& ObjectInitializer = FObjectInitializer::Get());
 
 	/** Pointer to color property in TrackObject. */
 	FLinearColor* ColorProp;
diff --git a/Engine/Source/Runtime/Engine/Classes/Matinee/InterpTrackInstMove.h b/Engine/Source/Runtime/Engine/Classes/Matinee/InterpTrackInstMove.h
--- a/Engine/Source/Runtime/Engine/Classes/Matinee/InterpTrackInstMove.h
+++ b/Engine/Source/Runtime/Engine/Classes/Matinee/InterpTrackInstMove.h
@@ -8,7 +8,9 @@
 UCLASS(MinimalAPI)
 class UInterpTrackInstMove : public UInterpTrackInst
 {
-	GENERATED_UCLASS_BODY()
+	GENERATED_BODY()
+public:
+	ENGINE_API UInterpTrackInstMove(const FObjectInitializer& ObjectInitializer = FObjectInitializer::Get());
 
 	/** Saved position. Used in editor for resetting when quitting Matinee. */
 	UPROPERTY()
diff --git a/Engine/Source/Runtime/Engine/Classes/Matinee/InterpTrackInstParticleReplay.h b/Engine/Source/Runtime/Engine/Classes/Matinee/InterpTrackInstParticleReplay.h
--- a/Engine/Source/Runtime/Engine/Classes/Matinee/InterpTrackInstParticleReplay.h
+++ b/Engine/Source/Runtime/Engine/Classes/Matinee/InterpTrackInstParticleReplay.h
@@ -7,7 +7,9 @@
 UCLASS()
 class UInterpTrackInstParticleReplay : public UInterpTrackInst
 {
-	GENERATED_UCLASS_BODY()
+	GENERATED_BODY()
+public:
+	UInterpTrackInstParticleReplay(const FObjectInitializer& ObjectInitializer = FObjectInitializer::Get());
 
 	/** 
 	 *	Position we were in last time we evaluated.
diff --git a/Engine/Source/Runtime/Engine/Classes/Matinee/InterpTrackInstProperty.h b/Engine/Source/Runtime/Engine/Classes/Matinee/InterpTrackInstProperty.h
--- a/Engine/Source/Runtime/Engine/Classes/Matinee/InterpTrackInstProperty.h
+++ b/Engine/Source/Runtime/Engine/Classes/Matinee/InterpTrackInstProperty.h
@@ -8,7 +8,9 @@
 UCLASS()
 class UInterpTrackInstProperty : public UInterpTrackInst
 {
-	GENERATED_UCLASS_BODY()
+	GENERATED_BODY()
+public:
+	UInterpTrackInstProperty(const FObjectInitializer& ObjectInitializer = FObjectInitializer::Get());
 
 	/** Function to call after updating the value of the color property. */
 	UPROPERTY()
diff --git a/Engine/Source/Runtime/Engine/Classes/Matinee/InterpTrackInstSlomo.h b/Engine/Source/Runtime/Engine/Classes/Matinee/InterpTrackInstSlomo.h
--- a/Engine/Source/Runtime/Engine/Classes/Matinee/InterpTrackInstSlomo.h
+++ b/Engine/Source/Runtime/Engine/Classes/Matinee/InterpTrackInstSlomo.h
@@ -7,7 +7,9 @@
 UCLASS()
 class UInterpTrackInstSlomo : public UInterpTrackInst
 {
-	GENERATED_UCLASS_BODY()
+	GENERATED_BODY()
+public:
+	UInterpTrackInstSlomo(const FObjectInitializer& ObjectInitializer = FObjectInitializer::Get());
 
 	/** Backup of initial LevelInfo MatineeTimeDilation setting when interpolation started. */
 	UPROPERTY()
diff --git a/Engine/Source/Runtime/Engine/Classes/Matinee/InterpTrackInstSound.h b/Engine/Source/Runtime/Engine/Classes/Matinee/InterpTrackInstSound.h
--- a/Engine/Source/Runtime/Engine/Classes/Matinee/InterpTrackInstSound.h
+++ b/Engine/Source/Runtime/Engine/Classes/Matinee/InterpTrackInstSound.h
@@ -7,7 +7,9 @@
 UCLASS()
 class UInterpTrackInstSound : public UInterpTrackInst
 {
-	GENERATED_UCLASS_BODY()
+	GENERATED_BODY()
+public:
+	UInterpTrackInstSound(const FObjectInitializer& ObjectInitializer = FObjectInitializer::Get());
 
 	UPROPERTY()
 	float LastUpdatePosition;
diff --git a/Engine/Source/Runtime/Engine/Classes/Matinee/InterpTrackInstToggle.h b/Engine/Source/Runtime/Engine/Classes/Matinee/InterpTrackInstToggle.h
--- a/Engine/Source/Runtime/Engine/Classes/Matinee/InterpTrackInstToggle.h
+++ b/Engine/Source/Runtime/Engine/Classes/Matinee/InterpTrackInstToggle.h
@@ -10,7 +10,9 @@
 UCLASS()
 class UInterpTrackInstToggle : public UInterpTrackInst
 {
-	GENERATED_UCLASS_BODY()
+	GENERATED_BODY()
+public:
+	UInterpTrackInstToggle(const FObjectInitializer& ObjectInitializer = FObjectInitializer::Get());
 
 	UPROPERTY(EditAnywhere, Category=InterpTrackInstToggle)
 	TEnumAsByte<enum ETrackToggleAction> Action;
diff --git a/Engine/Source/Runtime/Engine/Classes/Matinee/InterpTrackInstVectorMaterialParam.h b/Engine/Source/Runtime/Engine/Classes/Matinee/InterpTrackInstVectorMaterialParam.h
--- a/Engine/Source/Runtime/Engine/Classes/Matinee/InterpTrackInstVectorMaterialParam.h
+++ b/Engine/Source/Runtime/Engine/Classes/Matinee/InterpTrackInstVectorMaterialParam.h
@@ -7,7 +7,9 @@
 UCLASS()
 class UInterpTrackInstVectorMaterialParam : public UInterpTrackInst
 {
-	GENERATED_UCLASS_BODY()
+	GENERATED_BODY()
+public:
+	UInterpTrackInstVectorMaterialParam(const FObjectInitializer& ObjectInitializer = FObjectInitializer::Get());
 
 	/** MIDs we're using to set the desired parameter. */
 	UPROPERTY()
diff --git a/Engine/Source/Runtime/Engine/Classes/Matinee/InterpTrackInstVectorProp.h b/Engine/Source/Runtime/Engine/Classes/Matinee/InterpTrackInstVectorProp.h
--- a/Engine/Source/Runtime/Engine/Classes/Matinee/InterpTrackInstVectorProp.h
+++ b/Engine/Source/Runtime/Engine/Classes/Matinee/InterpTrackInstVectorProp.h
@@ -7,7 +7,9 @@
 UCLASS()
 class UInterpTrackInstVectorProp : public UInterpTrackInstProperty
 {
-	GENERATED_UCLASS_BODY()
+	GENERATED_BODY()
+public:
+	UInterpTrackInstVectorProp(const FObjectInitializer& ObjectInitializer = FObjectInitializer::Get());
 
 	/** Pointer to FVector property in TrackObject. */
 	FVector* VectorProp;
diff --git a/Engine/Source/Runtime/Engine/Classes/Matinee/InterpTrackInstVisibility.h b/Engine/Source/Runtime/Engine/Classes/Matinee/InterpTrackInstVisibility.h
--- a/Engine/Source/Runtime/Engine/Classes/Matinee/InterpTrackInstVisibility.h
+++ b/Engine/Source/Runtime/Engine/Classes/Matinee/InterpTrackInstVisibility.h
@@ -9,7 +9,9 @@
 UCLASS()
 class UInterpTrackInstVisibility : public UInterpTrackInst
 {
-	GENERATED_UCLASS_BODY()
+	GENERATED_BODY()
+public:
+	UInterpTrackInstVisibility(const FObjectInitializer& ObjectInitializer = FObjectInitializer::Get());
 
 	UPROPERTY(EditAnywhere, Category=InterpTrackInstVisibility)
 	TEnumAsByte<enum EVisibilityTrackAction> Action;
diff --git a/Engine/Source/Runtime/Engine/Classes/Matinee/InterpTrackLinearColorBase.h b/Engine/Source/Runtime/Engine/Classes/Matinee/InterpTrackLinearColorBase.h
--- a/Engine/Source/Runtime/Engine/Classes/Matinee/InterpTrackLinearColorBase.h
+++ b/Engine/Source/Runtime/Engine/Classes/Matinee/InterpTrackLinearColorBase.h
@@ -7,7 +7,9 @@
 UCLASS(abstract, MinimalAPI)
 class UInterpTrackLinearColorBase : public UInterpTrack
 {
-	GENERATED_UCLASS_BODY()
+	GENERATED_BODY()
+public:
+	ENGINE_API UInterpTrackLinearColorBase(const FObjectInitializer& ObjectInitializer = FObjectInitializer::Get());
 
 	/** Actually track data containing keyframes of a FVector as it varies over time. */
 	UPROPERTY()
diff --git a/Engine/Source/Runtime/Engine/Classes/Matinee/InterpTrackLinearColorProp.h b/Engine/Source/Runtime/Engine/Classes/Matinee/InterpTrackLinearColorProp.h
--- a/Engine/Source/Runtime/Engine/Classes/Matinee/InterpTrackLinearColorProp.h
+++ b/Engine/Source/Runtime/Engine/Classes/Matinee/InterpTrackLinearColorProp.h
@@ -8,7 +8,9 @@
 UCLASS(MinimalAPI, meta=( DisplayName = "LinearColor Property Track" ) )
 class UInterpTrackLinearColorProp : public UInterpTrackLinearColorBase
 {
-	GENERATED_UCLASS_BODY()
+	GENERATED_BODY()
+public:
+	ENGINE_API UInterpTrackLinearColorProp(const FObjectInitializer& ObjectInitializer = FObjectInitializer::Get());
 
 	/** Name of property in Group  AActor  which this track mill modify over time. */
 	UPROPERTY(Category=InterpTrackLinearColorProp, VisibleAnywhere)
diff --git a/Engine/Source/Runtime/Engine/Classes/Matinee/InterpTrackMove.h b/Engine/Source/Runtime/Engine/Classes/Matinee/InterpTrackMove.h
--- a/Engine/Source/Runtime/Engine/Classes/Matinee/InterpTrackMove.h
+++ b/Engine/Source/Runtime/Engine/Classes/Matinee/InterpTrackMove.h
@@ -29,7 +29,7 @@ enum EInterpTrackMoveRotMode
 USTRUCT()
 struct FInterpLookupPoint
 {
-	GENERATED_USTRUCT_BODY()
+	GENERATED_BODY()
 
 	UPROPERTY()
 	FName GroupName;
@@ -48,7 +48,7 @@ struct FInterpLookupPoint
 USTRUCT()
 struct FInterpLookupTrack
 {
-	GENERATED_USTRUCT_BODY()
+	GENERATED_BODY()
 
 	UPROPERTY()
 	TArray<struct FInterpLookupPoint> Points;
@@ -92,7 +92,9 @@ struct FInterpLookupTrack
 UCLASS(MinimalAPI, meta=( DisplayName = "Movement Track" ) )
 class UInterpTrackMove : public UInterpTrack
 {
-	GENERATED_UCLASS_BODY()
+	GENERATED_BODY()
+public:
+	ENGINE_API UInterpTrackMove(const FObjectInitializer& ObjectInitializer = FObjectInitializer::Get());
 	
 	/** Actual position keyframe data. */
 	UPROPERTY()
diff --git a/Engine/Source/Runtime/Engine/Classes/Matinee/InterpTrackMoveAxis.h b/Engine/Source/Runtime/Engine/Classes/Matinee/InterpTrackMoveAxis.h
--- a/Engine/Source/Runtime/Engine/Classes/Matinee/InterpTrackMoveAxis.h
+++ b/Engine/Source/Runtime/Engine/Classes/Matinee/InterpTrackMoveAxis.h
@@ -26,7 +26,9 @@ enum EInterpMoveAxis
 UCLASS(MinimalAPI, meta=( DisplayName = "Move Axis Track" ) )
 class UInterpTrackMoveAxis : public UInterpTrackFloatBase
 {
-	GENERATED_UCLASS_BODY()
+	GENERATED_BODY()
+public:
+	ENGINE_API UInterpTrackMoveAxis(const FObjectInitializer& ObjectInitializer = FObjectInitializer::Get());
 
 	/** The axis which this track will use when transforming an actor */
 	UPROPERTY()
diff --git a/Engine/Source/Runtime/Engine/Classes/Matinee/InterpTrackParticleReplay.h b/Engine/Source/Runtime/Engine/Classes/Matinee/InterpTrackParticleReplay.h
--- a/Engine/Source/Runtime/Engine/Classes/Matinee/InterpTrackParticleReplay.h
+++ b/Engine/Source/Runtime/Engine/Classes/Matinee/InterpTrackParticleReplay.h
@@ -12,7 +12,7 @@
 USTRUCT()
 struct FParticleReplayTrackKey
 {
-	GENERATED_USTRUCT_BODY()
+	GENERATED_BODY()
 
 	/** Position along timeline */
 	UPROPERTY()
@@ -39,7 +39,9 @@ struct FParticleReplayTrackKey
 UCLASS(MinimalAPI, meta=( DisplayName = "Particle Replay Track" ) )
 class UInterpTrackParticleReplay : public UInterpTrack
 {
-	GENERATED_UCLASS_BODY()
+	GENERATED_BODY()
+public:
+	ENGINE_API UInterpTrackParticleReplay(const FObjectInitializer& ObjectInitializer = FObjectInitializer::Get());
 
 	/** Array of keys */
 	UPROPERTY()
diff --git a/Engine/Source/Runtime/Engine/Classes/Matinee/InterpTrackSlomo.h b/Engine/Source/Runtime/Engine/Classes/Matinee/InterpTrackSlomo.h
--- a/Engine/Source/Runtime/Engine/Classes/Matinee/InterpTrackSlomo.h
+++ b/Engine/Source/Runtime/Engine/Classes/Matinee/InterpTrackSlomo.h
@@ -7,7 +7,9 @@
 UCLASS(meta=( DisplayName = "Slomo Track" ) )
 class UInterpTrackSlomo : public UInterpTrackFloatBase
 {
-	GENERATED_UCLASS_BODY()
+	GENERATED_BODY()
+public:
+	UInterpTrackSlomo(const FObjectInitializer& ObjectInitializer = FObjectInitializer::Get());
 
 
 	// Begin InterpTrack interface.
diff --git a/Engine/Source/Runtime/Engine/Classes/Matinee/InterpTrackSound.h b/Engine/Source/Runtime/Engine/Classes/Matinee/InterpTrackSound.h
--- a/Engine/Source/Runtime/Engine/Classes/Matinee/InterpTrackSound.h
+++ b/Engine/Source/Runtime/Engine/Classes/Matinee/InterpTrackSound.h
@@ -13,7 +13,7 @@
 USTRUCT()
 struct FSoundTrackKey
 {
-	GENERATED_USTRUCT_BODY()
+	GENERATED_BODY()
 
 	UPROPERTY()
 	float Time;
@@ -42,7 +42,9 @@ struct FSoundTrackKey
 UCLASS(MinimalAPI, meta=( DisplayName = "Sound Track" ) )
 class UInterpTrackSound : public UInterpTrackVectorBase
 {
-	GENERATED_UCLASS_BODY()
+	GENERATED_BODY()
+public:
+	ENGINE_API UInterpTrackSound(const FObjectInitializer& ObjectInitializer = FObjectInitializer::Get());
 
 	/** Array of sounds to play at specific times. */
 	UPROPERTY()
diff --git a/Engine/Source/Runtime/Engine/Classes/Matinee/InterpTrackToggle.h b/Engine/Source/Runtime/Engine/Classes/Matinee/InterpTrackToggle.h
--- a/Engine/Source/Runtime/Engine/Classes/Matinee/InterpTrackToggle.h
+++ b/Engine/Source/Runtime/Engine/Classes/Matinee/InterpTrackToggle.h
@@ -26,7 +26,7 @@ enum ETrackToggleAction
 USTRUCT()
 struct FToggleTrackKey
 {
-	GENERATED_USTRUCT_BODY()
+	GENERATED_BODY()
 
 	UPROPERTY()
 	float Time;
@@ -46,7 +46,9 @@ struct FToggleTrackKey
 UCLASS(MinimalAPI, meta=( DisplayName = "Toggle Track" ) )
 class UInterpTrackToggle : public UInterpTrack
 {
-	GENERATED_UCLASS_BODY()
+	GENERATED_BODY()
+public:
+	ENGINE_API UInterpTrackToggle(const FObjectInitializer& ObjectInitializer = FObjectInitializer::Get());
 
 	/** Array of events to fire off. */
 	UPROPERTY()
diff --git a/Engine/Source/Runtime/Engine/Classes/Matinee/InterpTrackVectorBase.h b/Engine/Source/Runtime/Engine/Classes/Matinee/InterpTrackVectorBase.h
--- a/Engine/Source/Runtime/Engine/Classes/Matinee/InterpTrackVectorBase.h
+++ b/Engine/Source/Runtime/Engine/Classes/Matinee/InterpTrackVectorBase.h
@@ -7,7 +7,9 @@
 UCLASS(abstract, MinimalAPI)
 class UInterpTrackVectorBase : public UInterpTrack
 {
-	GENERATED_UCLASS_BODY()
+	GENERATED_BODY()
+public:
+	ENGINE_API UInterpTrackVectorBase(const FObjectInitializer& ObjectInitializer = FObjectInitializer::Get());
 
 	/** Actually track data containing keyframes of a FVector as it varies over time. */
 	UPROPERTY()
diff --git a/Engine/Source/Runtime/Engine/Classes/Matinee/InterpTrackVectorMaterialParam.h b/Engine/Source/Runtime/Engine/Classes/Matinee/InterpTrackVectorMaterialParam.h
--- a/Engine/Source/Runtime/Engine/Classes/Matinee/InterpTrackVectorMaterialParam.h
+++ b/Engine/Source/Runtime/Engine/Classes/Matinee/InterpTrackVectorMaterialParam.h
@@ -6,7 +6,9 @@
 UCLASS(MinimalAPI, meta=( DisplayName = "Vector Material Parameter Track" ))
 class UInterpTrackVectorMaterialParam : public UInterpTrackVectorBase
 {
-	GENERATED_UCLASS_BODY()
+	GENERATED_BODY()
+public:
+	ENGINE_API UInterpTrackVectorMaterialParam(const FObjectInitializer& ObjectInitializer = FObjectInitializer::Get());
 
 	/** Materials whose parameters we want to change and the references to those materials. */
 	UPROPERTY(EditAnywhere, Category=InterpTrackVectorMaterialParam)
diff --git a/Engine/Source/Runtime/Engine/Classes/Matinee/InterpTrackVectorProp.h b/Engine/Source/Runtime/Engine/Classes/Matinee/InterpTrackVectorProp.h
--- a/Engine/Source/Runtime/Engine/Classes/Matinee/InterpTrackVectorProp.h
+++ b/Engine/Source/Runtime/Engine/Classes/Matinee/InterpTrackVectorProp.h
@@ -8,7 +8,9 @@
 UCLASS(MinimalAPI, meta=( DisplayName = "Vector Property Track" ) )
 class UInterpTrackVectorProp : public UInterpTrackVectorBase
 {
-	GENERATED_UCLASS_BODY()
+	GENERATED_BODY()
+public:
+	ENGINE_API UInterpTrackVectorProp(const FObjectInitializer& ObjectInitializer = FObjectInitializer::Get());
 
 	/** Name of property in Group  AActor  which this track mill modify over time. */
 	UPROPERTY(Category=InterpTrackVectorProp, VisibleAnywhere)
diff --git a/Engine/Source/Runtime/Engine/Classes/Matinee/InterpTrackVisibility.h b/Engine/Source/Runtime/Engine/Classes/Matinee/InterpTrackVisibility.h
--- a/Engine/Source/Runtime/Engine/Classes/Matinee/InterpTrackVisibility.h
+++ b/Engine/Source/Runtime/Engine/Classes/Matinee/InterpTrackVisibility.h
@@ -39,7 +39,7 @@ enum EVisibilityTrackCondition
 USTRUCT()
 struct FVisibilityTrackKey
 {
-	GENERATED_USTRUCT_BODY()
+	GENERATED_BODY()
 
 	UPROPERTY()
 	float Time;
@@ -64,7 +64,9 @@ struct FVisibilityTrackKey
 UCLASS(MinimalAPI, meta=( DisplayName = "Visibility Track" ) )
 class UInterpTrackVisibility : public UInterpTrack
 {
-	GENERATED_UCLASS_BODY()
+	GENERATED_BODY()
+public:
+	ENGINE_API UInterpTrackVisibility(const FObjectInitializer& ObjectInitializer = FObjectInitializer::Get());
 
 	/** Array of events to fire off. */
 	UPROPERTY()
diff --git a/Engine/Source/Runtime/Engine/Classes/Matinee/MatineeActor.h b/Engine/Source/Runtime/Engine/Classes/Matinee/MatineeActor.h
--- a/Engine/Source/Runtime/Engine/Classes/Matinee/MatineeActor.h
+++ b/Engine/Source/Runtime/Engine/Classes/Matinee/MatineeActor.h
@@ -16,7 +16,7 @@ DECLARE_DYNAMIC_DELEGATE( FOnMatineeEvent );
 USTRUCT()
 struct FCameraCutInfo
 {
-	GENERATED_USTRUCT_BODY()
+	GENERATED_BODY()
 
 	UPROPERTY()
 	FVector Location;
@@ -39,7 +39,7 @@ struct FCameraCutInfo
 USTRUCT()
 struct FInterpGroupActorInfo
 {
-	GENERATED_USTRUCT_BODY()
+	GENERATED_BODY()
 
 	UPROPERTY(EditAnywhere, BlueprintReadWrite, Category=InterpGroupActorInfo)
 	FName ObjectName;
@@ -52,7 +52,9 @@ struct FInterpGroupActorInfo
 UCLASS(MinimalAPI, NotBlueprintable, hidecategories=(Collision, Game, Input), showcategories=("Input|MouseInput", "Input|TouchInput", "Game|Damage"))
 class AMatineeActor : public AActor
 {
-	GENERATED_UCLASS_BODY()
+	GENERATED_BODY()
+public:
+	ENGINE_API AMatineeActor(const FObjectInitializer& ObjectInitializer = FObjectInitializer::Get());
 
 	/** The matinee data used by this actor*/
 	UPROPERTY(EditAnywhere, BlueprintReadWrite, Category=MatineeActor, replicated, meta=(ForceRebuildProperty = "GroupActorInfos"))
diff --git a/Engine/Source/Runtime/Engine/Classes/Matinee/MatineeActorCameraAnim.h b/Engine/Source/Runtime/Engine/Classes/Matinee/MatineeActorCameraAnim.h
--- a/Engine/Source/Runtime/Engine/Classes/Matinee/MatineeActorCameraAnim.h
+++ b/Engine/Source/Runtime/Engine/Classes/Matinee/MatineeActorCameraAnim.h
@@ -12,7 +12,9 @@ class UCameraAnim;
 UCLASS(notplaceable, MinimalAPI, NotBlueprintable)
 class AMatineeActorCameraAnim : public AMatineeActor
 {
-	GENERATED_UCLASS_BODY()
+	GENERATED_BODY()
+public:
+	ENGINE_API AMatineeActorCameraAnim(const FObjectInitializer& ObjectInitializer = FObjectInitializer::Get());
 
 	/** The camera anim we are editing */
 	UPROPERTY(Transient)
@@ -29,4 +31,4 @@ class AMatineeActorCameraAnim : public AMatineeActor
 		return false;
 	}
 	// End UObject interface
-};
\ No newline at end of file
+};
diff --git a/Engine/Source/Runtime/Engine/Classes/Matinee/MatineeAnimInterface.h b/Engine/Source/Runtime/Engine/Classes/Matinee/MatineeAnimInterface.h
--- a/Engine/Source/Runtime/Engine/Classes/Matinee/MatineeAnimInterface.h
+++ b/Engine/Source/Runtime/Engine/Classes/Matinee/MatineeAnimInterface.h
@@ -13,12 +13,15 @@
 UINTERFACE(MinimalAPI, meta=(CannotImplementInterfaceInBlueprint))
 class UMatineeAnimInterface : public UInterface
 {
-	GENERATED_UINTERFACE_BODY()
+	GENERATED_BODY()
+public:
+	ENGINE_API UMatineeAnimInterface(const FObjectInitializer& ObjectInitializer = FObjectInitializer::Get());
 };
 
 class IMatineeAnimInterface
 {
-	GENERATED_IINTERFACE_BODY()
+	GENERATED_BODY()
+public:
 
 
 	
diff --git a/Engine/Source/Runtime/Engine/Classes/Matinee/MatineeInterface.h b/Engine/Source/Runtime/Engine/Classes/Matinee/MatineeInterface.h
--- a/Engine/Source/Runtime/Engine/Classes/Matinee/MatineeInterface.h
+++ b/Engine/Source/Runtime/Engine/Classes/Matinee/MatineeInterface.h
@@ -12,12 +12,15 @@
 UINTERFACE(meta=(CannotImplementInterfaceInBlueprint))
 class UMatineeInterface : public UInterface
 {
-	GENERATED_UINTERFACE_BODY()
+	GENERATED_BODY()
+public:
+	UMatineeInterface(const FObjectInitializer& ObjectInitializer = FObjectInitializer::Get());
 };
 
 class IMatineeInterface
 {
-	GENERATED_IINTERFACE_BODY()
+	GENERATED_BODY()
+public:
 
 
 	/** Called when a matinee actor starts interpolating this  AActor via matinee.
diff --git a/Engine/Source/Runtime/Engine/Classes/MovieScene/RuntimeMovieScenePlayerInterface.h b/Engine/Source/Runtime/Engine/Classes/MovieScene/RuntimeMovieScenePlayerInterface.h
--- a/Engine/Source/Runtime/Engine/Classes/MovieScene/RuntimeMovieScenePlayerInterface.h
+++ b/Engine/Source/Runtime/Engine/Classes/MovieScene/RuntimeMovieScenePlayerInterface.h
@@ -7,12 +7,15 @@
 UINTERFACE(MinimalAPI, meta=(CannotImplementInterfaceInBlueprint))
 class URuntimeMovieScenePlayerInterface : public UInterface
 {
-	GENERATED_UINTERFACE_BODY()
+	GENERATED_BODY()
+public:
+	ENGINE_API URuntimeMovieScenePlayerInterface(const FObjectInitializer& ObjectInitializer = FObjectInitializer::Get());
 };
 
 class IRuntimeMovieScenePlayerInterface
 {
-	GENERATED_IINTERFACE_BODY()
+	GENERATED_BODY()
+public:
 
 	/**
 	 * Ticks this MovieScenePlayer.  Should be called every frame.
diff --git a/Engine/Source/Runtime/Engine/Classes/Particles/Acceleration/ParticleModuleAcceleration.h b/Engine/Source/Runtime/Engine/Classes/Particles/Acceleration/ParticleModuleAcceleration.h
--- a/Engine/Source/Runtime/Engine/Classes/Particles/Acceleration/ParticleModuleAcceleration.h
+++ b/Engine/Source/Runtime/Engine/Classes/Particles/Acceleration/ParticleModuleAcceleration.h
@@ -8,7 +8,9 @@
 UCLASS(editinlinenew, hidecategories=Object, meta=(DisplayName = "Acceleration"), MinimalAPI)
 class UParticleModuleAcceleration : public UParticleModuleAccelerationBase
 {
-	GENERATED_UCLASS_BODY()
+	GENERATED_BODY()
+public:
+	ENGINE_API UParticleModuleAcceleration(const FObjectInitializer& ObjectInitializer = FObjectInitializer::Get());
 
 	/**
 	 *	The initial acceleration of the particle.
diff --git a/Engine/Source/Runtime/Engine/Classes/Particles/Acceleration/ParticleModuleAccelerationBase.h b/Engine/Source/Runtime/Engine/Classes/Particles/Acceleration/ParticleModuleAccelerationBase.h
--- a/Engine/Source/Runtime/Engine/Classes/Particles/Acceleration/ParticleModuleAccelerationBase.h
+++ b/Engine/Source/Runtime/Engine/Classes/Particles/Acceleration/ParticleModuleAccelerationBase.h
@@ -8,7 +8,9 @@
 UCLASS(editinlinenew, hidecategories=Object, meta=(DisplayName = "Acceleration"), abstract)
 class UParticleModuleAccelerationBase : public UParticleModule
 {
-	GENERATED_UCLASS_BODY()
+	GENERATED_BODY()
+public:
+	UParticleModuleAccelerationBase(const FObjectInitializer& ObjectInitializer = FObjectInitializer::Get());
 
 	/**
 	 *	If true, then treat the acceleration as world-space
diff --git a/Engine/Source/Runtime/Engine/Classes/Particles/Acceleration/ParticleModuleAccelerationConstant.h b/Engine/Source/Runtime/Engine/Classes/Particles/Acceleration/ParticleModuleAccelerationConstant.h
--- a/Engine/Source/Runtime/Engine/Classes/Particles/Acceleration/ParticleModuleAccelerationConstant.h
+++ b/Engine/Source/Runtime/Engine/Classes/Particles/Acceleration/ParticleModuleAccelerationConstant.h
@@ -11,7 +11,9 @@
 UCLASS(editinlinenew, hidecategories=(Object, Acceleration), meta=(DisplayName = "Const Acceleration"))
 class UParticleModuleAccelerationConstant : public UParticleModuleAccelerationBase
 {
-	GENERATED_UCLASS_BODY()
+	GENERATED_BODY()
+public:
+	UParticleModuleAccelerationConstant(const FObjectInitializer& ObjectInitializer = FObjectInitializer::Get());
 
 	/** Constant acceleration for particles in this system. */
 	UPROPERTY(EditAnywhere, Category=ParticleModuleAccelerationConstant)
diff --git a/Engine/Source/Runtime/Engine/Classes/Particles/Acceleration/ParticleModuleAccelerationDrag.h b/Engine/Source/Runtime/Engine/Classes/Particles/Acceleration/ParticleModuleAccelerationDrag.h
--- a/Engine/Source/Runtime/Engine/Classes/Particles/Acceleration/ParticleModuleAccelerationDrag.h
+++ b/Engine/Source/Runtime/Engine/Classes/Particles/Acceleration/ParticleModuleAccelerationDrag.h
@@ -11,7 +11,9 @@
 UCLASS(editinlinenew, hidecategories=(Object, Acceleration), MinimalAPI, meta=(DisplayName = "Drag"))
 class UParticleModuleAccelerationDrag : public UParticleModuleAccelerationBase
 {
-	GENERATED_UCLASS_BODY()
+	GENERATED_BODY()
+public:
+	ENGINE_API UParticleModuleAccelerationDrag(const FObjectInitializer& ObjectInitializer = FObjectInitializer::Get());
 
 	/** Per-particle drag coefficient. Evaluted using emitter time. */
 	UPROPERTY(EditAnywhere, Category=Drag)
diff --git a/Engine/Source/Runtime/Engine/Classes/Particles/Acceleration/ParticleModuleAccelerationDragScaleOverLife.h b/Engine/Source/Runtime/Engine/Classes/Particles/Acceleration/ParticleModuleAccelerationDragScaleOverLife.h
--- a/Engine/Source/Runtime/Engine/Classes/Particles/Acceleration/ParticleModuleAccelerationDragScaleOverLife.h
+++ b/Engine/Source/Runtime/Engine/Classes/Particles/Acceleration/ParticleModuleAccelerationDragScaleOverLife.h
@@ -11,7 +11,9 @@
 UCLASS(editinlinenew, hidecategories=(UObject, Acceleration), MinimalAPI, meta=(DisplayName = "Drag Scale/Life"))
 class UParticleModuleAccelerationDragScaleOverLife : public UParticleModuleAccelerationBase
 {
-	GENERATED_UCLASS_BODY()
+	GENERATED_BODY()
+public:
+	ENGINE_API UParticleModuleAccelerationDragScaleOverLife(const FObjectInitializer& ObjectInitializer = FObjectInitializer::Get());
 
 	/** Per-particle drag scale. Evaluted using particle relative time. */
 	UPROPERTY(EditAnywhere, Category=Drag)
diff --git a/Engine/Source/Runtime/Engine/Classes/Particles/Acceleration/ParticleModuleAccelerationOverLifetime.h b/Engine/Source/Runtime/Engine/Classes/Particles/Acceleration/ParticleModuleAccelerationOverLifetime.h
--- a/Engine/Source/Runtime/Engine/Classes/Particles/Acceleration/ParticleModuleAccelerationOverLifetime.h
+++ b/Engine/Source/Runtime/Engine/Classes/Particles/Acceleration/ParticleModuleAccelerationOverLifetime.h
@@ -8,7 +8,9 @@
 UCLASS(editinlinenew, hidecategories=Object, meta=(DisplayName = "Acceleration/Life"), MinimalAPI)
 class UParticleModuleAccelerationOverLifetime : public UParticleModuleAccelerationBase
 {
-	GENERATED_UCLASS_BODY()
+	GENERATED_BODY()
+public:
+	ENGINE_API UParticleModuleAccelerationOverLifetime(const FObjectInitializer& ObjectInitializer = FObjectInitializer::Get());
 
 	/** 
 	 *	The acceleration of the particle over its lifetime.
diff --git a/Engine/Source/Runtime/Engine/Classes/Particles/Attractor/ParticleModuleAttractorBase.h b/Engine/Source/Runtime/Engine/Classes/Particles/Attractor/ParticleModuleAttractorBase.h
--- a/Engine/Source/Runtime/Engine/Classes/Particles/Attractor/ParticleModuleAttractorBase.h
+++ b/Engine/Source/Runtime/Engine/Classes/Particles/Attractor/ParticleModuleAttractorBase.h
@@ -8,7 +8,9 @@
 UCLASS(editinlinenew, hidecategories=Object, abstract, meta=(DisplayName = "Attraction"))
 class UParticleModuleAttractorBase : public UParticleModule
 {
-	GENERATED_UCLASS_BODY()
+	GENERATED_BODY()
+public:
+	UParticleModuleAttractorBase(const FObjectInitializer& ObjectInitializer = FObjectInitializer::Get());
 
 };
 
diff --git a/Engine/Source/Runtime/Engine/Classes/Particles/Attractor/ParticleModuleAttractorLine.h b/Engine/Source/Runtime/Engine/Classes/Particles/Attractor/ParticleModuleAttractorLine.h
--- a/Engine/Source/Runtime/Engine/Classes/Particles/Attractor/ParticleModuleAttractorLine.h
+++ b/Engine/Source/Runtime/Engine/Classes/Particles/Attractor/ParticleModuleAttractorLine.h
@@ -8,7 +8,9 @@
 UCLASS(editinlinenew, hidecategories=Object, MinimalAPI, meta=(DisplayName = "Line Attractor"))
 class UParticleModuleAttractorLine : public UParticleModuleAttractorBase
 {
-	GENERATED_UCLASS_BODY()
+	GENERATED_BODY()
+public:
+	ENGINE_API UParticleModuleAttractorLine(const FObjectInitializer& ObjectInitializer = FObjectInitializer::Get());
 
 	/** The first endpoint of the line. */
 	UPROPERTY(EditAnywhere, Category=Attractor)
diff --git a/Engine/Source/Runtime/Engine/Classes/Particles/Attractor/ParticleModuleAttractorParticle.h b/Engine/Source/Runtime/Engine/Classes/Particles/Attractor/ParticleModuleAttractorParticle.h
--- a/Engine/Source/Runtime/Engine/Classes/Particles/Attractor/ParticleModuleAttractorParticle.h
+++ b/Engine/Source/Runtime/Engine/Classes/Particles/Attractor/ParticleModuleAttractorParticle.h
@@ -16,7 +16,9 @@ enum EAttractorParticleSelectionMethod
 UCLASS(editinlinenew, hidecategories=Object, MinimalAPI, meta=(DisplayName = "Particle Attractor"))
 class UParticleModuleAttractorParticle : public UParticleModuleAttractorBase
 {
-	GENERATED_UCLASS_BODY()
+	GENERATED_BODY()
+public:
+	ENGINE_API UParticleModuleAttractorParticle(const FObjectInitializer& ObjectInitializer = FObjectInitializer::Get());
 
 	/**
 	 *	The source emitter for attractors
diff --git a/Engine/Source/Runtime/Engine/Classes/Particles/Attractor/ParticleModuleAttractorPoint.h b/Engine/Source/Runtime/Engine/Classes/Particles/Attractor/ParticleModuleAttractorPoint.h
--- a/Engine/Source/Runtime/Engine/Classes/Particles/Attractor/ParticleModuleAttractorPoint.h
+++ b/Engine/Source/Runtime/Engine/Classes/Particles/Attractor/ParticleModuleAttractorPoint.h
@@ -8,7 +8,9 @@
 UCLASS(editinlinenew, hidecategories=Object, MinimalAPI, meta=(DisplayName = "Point Attractor"))
 class UParticleModuleAttractorPoint : public UParticleModuleAttractorBase
 {
-	GENERATED_UCLASS_BODY()
+	GENERATED_BODY()
+public:
+	ENGINE_API UParticleModuleAttractorPoint(const FObjectInitializer& ObjectInitializer = FObjectInitializer::Get());
 
 	/**	The position of the point attractor from the source of the emitter.		*/
 	UPROPERTY(EditAnywhere, Category=Attractor)
diff --git a/Engine/Source/Runtime/Engine/Classes/Particles/Attractor/ParticleModuleAttractorPointGravity.h b/Engine/Source/Runtime/Engine/Classes/Particles/Attractor/ParticleModuleAttractorPointGravity.h
--- a/Engine/Source/Runtime/Engine/Classes/Particles/Attractor/ParticleModuleAttractorPointGravity.h
+++ b/Engine/Source/Runtime/Engine/Classes/Particles/Attractor/ParticleModuleAttractorPointGravity.h
@@ -12,7 +12,9 @@
 UCLASS(editinlinenew, hidecategories=Object, meta=(DisplayName = "Point Gravity"))
 class UParticleModuleAttractorPointGravity : public UParticleModuleAttractorBase
 {
-	GENERATED_UCLASS_BODY()
+	GENERATED_BODY()
+public:
+	UParticleModuleAttractorPointGravity(const FObjectInitializer& ObjectInitializer = FObjectInitializer::Get());
 
 	/** The position of the point gravity source. */
 	UPROPERTY(EditAnywhere, Category=PointGravitySource)
diff --git a/Engine/Source/Runtime/Engine/Classes/Particles/Beam/ParticleModuleBeamBase.h b/Engine/Source/Runtime/Engine/Classes/Particles/Beam/ParticleModuleBeamBase.h
--- a/Engine/Source/Runtime/Engine/Classes/Particles/Beam/ParticleModuleBeamBase.h
+++ b/Engine/Source/Runtime/Engine/Classes/Particles/Beam/ParticleModuleBeamBase.h
@@ -55,7 +55,9 @@ enum Beam2SourceTargetTangentMethod
 UCLASS(editinlinenew, hidecategories=Object, abstract, meta=(DisplayName = "Beam"))
 class UParticleModuleBeamBase : public UParticleModule
 {
-	GENERATED_UCLASS_BODY()
+	GENERATED_BODY()
+public:
+	UParticleModuleBeamBase(const FObjectInitializer& ObjectInitializer = FObjectInitializer::Get());
 
 
 	//Begin UParticleModule Interface
diff --git a/Engine/Source/Runtime/Engine/Classes/Particles/Beam/ParticleModuleBeamModifier.h b/Engine/Source/Runtime/Engine/Classes/Particles/Beam/ParticleModuleBeamModifier.h
--- a/Engine/Source/Runtime/Engine/Classes/Particles/Beam/ParticleModuleBeamModifier.h
+++ b/Engine/Source/Runtime/Engine/Classes/Particles/Beam/ParticleModuleBeamModifier.h
@@ -25,7 +25,7 @@ enum BeamModifierType
 USTRUCT()
 struct FBeamModifierOptions
 {
-	GENERATED_USTRUCT_BODY()
+	GENERATED_BODY()
 
 	/** If true, modify the value associated with this grouping.	*/
 	UPROPERTY(EditAnywhere, Category=BeamModifierOptions)
@@ -52,7 +52,9 @@ struct FBeamModifierOptions
 UCLASS(editinlinenew, hidecategories=Object, meta=(DisplayName = "Beam Modifier"))
 class UParticleModuleBeamModifier : public UParticleModuleBeamBase
 {
-	GENERATED_UCLASS_BODY()
+	GENERATED_BODY()
+public:
+	UParticleModuleBeamModifier(const FObjectInitializer& ObjectInitializer = FObjectInitializer::Get());
 
 	/**	Whether this module modifies the Source or the Target. */
 	UPROPERTY(EditAnywhere, Category=Modifier)
diff --git a/Engine/Source/Runtime/Engine/Classes/Particles/Beam/ParticleModuleBeamNoise.h b/Engine/Source/Runtime/Engine/Classes/Particles/Beam/ParticleModuleBeamNoise.h
--- a/Engine/Source/Runtime/Engine/Classes/Particles/Beam/ParticleModuleBeamNoise.h
+++ b/Engine/Source/Runtime/Engine/Classes/Particles/Beam/ParticleModuleBeamNoise.h
@@ -14,7 +14,9 @@
 UCLASS(editinlinenew, hidecategories=Object, MinimalAPI, meta=(DisplayName = "Noise"))
 class UParticleModuleBeamNoise : public UParticleModuleBeamBase
 {
-	GENERATED_UCLASS_BODY()
+	GENERATED_BODY()
+public:
+	ENGINE_API UParticleModuleBeamNoise(const FObjectInitializer& ObjectInitializer = FObjectInitializer::Get());
 
 	/** Is low frequency noise enabled. */
 	UPROPERTY(EditAnywhere, Category=LowFreq)
diff --git a/Engine/Source/Runtime/Engine/Classes/Particles/Beam/ParticleModuleBeamSource.h b/Engine/Source/Runtime/Engine/Classes/Particles/Beam/ParticleModuleBeamSource.h
--- a/Engine/Source/Runtime/Engine/Classes/Particles/Beam/ParticleModuleBeamSource.h
+++ b/Engine/Source/Runtime/Engine/Classes/Particles/Beam/ParticleModuleBeamSource.h
@@ -14,7 +14,9 @@
 UCLASS(editinlinenew, hidecategories=Object, MinimalAPI, meta=(DisplayName = "Source"))
 class UParticleModuleBeamSource : public UParticleModuleBeamBase
 {
-	GENERATED_UCLASS_BODY()
+	GENERATED_BODY()
+public:
+	ENGINE_API UParticleModuleBeamSource(const FObjectInitializer& ObjectInitializer = FObjectInitializer::Get());
 
 	/** The method flag. */
 	UPROPERTY(EditAnywhere, Category=Source)
diff --git a/Engine/Source/Runtime/Engine/Classes/Particles/Beam/ParticleModuleBeamTarget.h b/Engine/Source/Runtime/Engine/Classes/Particles/Beam/ParticleModuleBeamTarget.h
--- a/Engine/Source/Runtime/Engine/Classes/Particles/Beam/ParticleModuleBeamTarget.h
+++ b/Engine/Source/Runtime/Engine/Classes/Particles/Beam/ParticleModuleBeamTarget.h
@@ -14,7 +14,9 @@
 UCLASS(editinlinenew, hidecategories=Object, MinimalAPI, meta=(DisplayName = "Target"))
 class UParticleModuleBeamTarget : public UParticleModuleBeamBase
 {
-	GENERATED_UCLASS_BODY()
+	GENERATED_BODY()
+public:
+	ENGINE_API UParticleModuleBeamTarget(const FObjectInitializer& ObjectInitializer = FObjectInitializer::Get());
 
 	/** The method flag. */
 	UPROPERTY(EditAnywhere, Category=Target)
diff --git a/Engine/Source/Runtime/Engine/Classes/Particles/Camera/ParticleModuleCameraBase.h b/Engine/Source/Runtime/Engine/Classes/Particles/Camera/ParticleModuleCameraBase.h
--- a/Engine/Source/Runtime/Engine/Classes/Particles/Camera/ParticleModuleCameraBase.h
+++ b/Engine/Source/Runtime/Engine/Classes/Particles/Camera/ParticleModuleCameraBase.h
@@ -8,7 +8,9 @@
 UCLASS(editinlinenew, hidecategories=Object, abstract, meta=(DisplayName = "Camera"))
 class UParticleModuleCameraBase : public UParticleModule
 {
-	GENERATED_UCLASS_BODY()
+	GENERATED_BODY()
+public:
+	UParticleModuleCameraBase(const FObjectInitializer& ObjectInitializer = FObjectInitializer::Get());
 
 	virtual bool CanTickInAnyThread() override
 	{
diff --git a/Engine/Source/Runtime/Engine/Classes/Particles/Camera/ParticleModuleCameraOffset.h b/Engine/Source/Runtime/Engine/Classes/Particles/Camera/ParticleModuleCameraOffset.h
--- a/Engine/Source/Runtime/Engine/Classes/Particles/Camera/ParticleModuleCameraOffset.h
+++ b/Engine/Source/Runtime/Engine/Classes/Particles/Camera/ParticleModuleCameraOffset.h
@@ -20,7 +20,9 @@ enum EParticleCameraOffsetUpdateMethod
 UCLASS(editinlinenew, hidecategories=Object, meta=(DisplayName = "Camera Offset"))
 class UParticleModuleCameraOffset : public UParticleModuleCameraBase
 {
-	GENERATED_UCLASS_BODY()
+	GENERATED_BODY()
+public:
+	UParticleModuleCameraOffset(const FObjectInitializer& ObjectInitializer = FObjectInitializer::Get());
 
 	/** 
 	 *	The camera-relative offset to apply to sprite location
diff --git a/Engine/Source/Runtime/Engine/Classes/Particles/Collision/ParticleModuleCollision.h b/Engine/Source/Runtime/Engine/Classes/Particles/Collision/ParticleModuleCollision.h
--- a/Engine/Source/Runtime/Engine/Classes/Particles/Collision/ParticleModuleCollision.h
+++ b/Engine/Source/Runtime/Engine/Classes/Particles/Collision/ParticleModuleCollision.h
@@ -8,7 +8,9 @@
 UCLASS(editinlinenew, hidecategories=Object, MinimalAPI, meta=(DisplayName = "Collision"))
 class UParticleModuleCollision : public UParticleModuleCollisionBase
 {
-	GENERATED_UCLASS_BODY()
+	GENERATED_BODY()
+public:
+	ENGINE_API UParticleModuleCollision(const FObjectInitializer& ObjectInitializer = FObjectInitializer::Get());
 
 	/**
 	 *	How much to `slow' the velocity of the particle after a collision.
diff --git a/Engine/Source/Runtime/Engine/Classes/Particles/Collision/ParticleModuleCollisionBase.h b/Engine/Source/Runtime/Engine/Classes/Particles/Collision/ParticleModuleCollisionBase.h
--- a/Engine/Source/Runtime/Engine/Classes/Particles/Collision/ParticleModuleCollisionBase.h
+++ b/Engine/Source/Runtime/Engine/Classes/Particles/Collision/ParticleModuleCollisionBase.h
@@ -29,7 +29,9 @@ enum EParticleCollisionComplete
 UCLASS(editinlinenew, hidecategories=Object, abstract, meta=(DisplayName = "Collision"))
 class UParticleModuleCollisionBase : public UParticleModule
 {
-	GENERATED_UCLASS_BODY()
+	GENERATED_BODY()
+public:
+	UParticleModuleCollisionBase(const FObjectInitializer& ObjectInitializer = FObjectInitializer::Get());
 
 };
 
diff --git a/Engine/Source/Runtime/Engine/Classes/Particles/Collision/ParticleModuleCollisionGPU.h b/Engine/Source/Runtime/Engine/Classes/Particles/Collision/ParticleModuleCollisionGPU.h
--- a/Engine/Source/Runtime/Engine/Classes/Particles/Collision/ParticleModuleCollisionGPU.h
+++ b/Engine/Source/Runtime/Engine/Classes/Particles/Collision/ParticleModuleCollisionGPU.h
@@ -24,7 +24,9 @@ namespace EParticleCollisionResponse
 UCLASS(editinlinenew, hidecategories=Object, MinimalAPI, meta=(DisplayName = "Collision (Scene Depth)"))
 class UParticleModuleCollisionGPU : public UParticleModuleCollisionBase
 {
-	GENERATED_UCLASS_BODY()
+	GENERATED_BODY()
+public:
+	ENGINE_API UParticleModuleCollisionGPU(const FObjectInitializer& ObjectInitializer = FObjectInitializer::Get());
 
 	/**
 	 * Dampens the velocity of a particle in the direction normal to the
diff --git a/Engine/Source/Runtime/Engine/Classes/Particles/Color/ParticleModuleColor.h b/Engine/Source/Runtime/Engine/Classes/Particles/Color/ParticleModuleColor.h
--- a/Engine/Source/Runtime/Engine/Classes/Particles/Color/ParticleModuleColor.h
+++ b/Engine/Source/Runtime/Engine/Classes/Particles/Color/ParticleModuleColor.h
@@ -8,7 +8,9 @@
 UCLASS(editinlinenew, hidecategories=Object, MinimalAPI, meta=(DisplayName = "Initial Color"))
 class UParticleModuleColor : public UParticleModuleColorBase
 {
-	GENERATED_UCLASS_BODY()
+	GENERATED_BODY()
+public:
+	ENGINE_API UParticleModuleColor(const FObjectInitializer& ObjectInitializer = FObjectInitializer::Get());
 
 	/** Initial color for a particle as a function of Emitter time. */
 	UPROPERTY(EditAnywhere, Category = Color, meta = (TreatAsColor))
diff --git a/Engine/Source/Runtime/Engine/Classes/Particles/Color/ParticleModuleColorBase.h b/Engine/Source/Runtime/Engine/Classes/Particles/Color/ParticleModuleColorBase.h
--- a/Engine/Source/Runtime/Engine/Classes/Particles/Color/ParticleModuleColorBase.h
+++ b/Engine/Source/Runtime/Engine/Classes/Particles/Color/ParticleModuleColorBase.h
@@ -8,7 +8,9 @@
 UCLASS(editinlinenew, hidecategories=Object, abstract, meta=(DisplayName = "Color"))
 class UParticleModuleColorBase : public UParticleModule
 {
-	GENERATED_UCLASS_BODY()
+	GENERATED_BODY()
+public:
+	UParticleModuleColorBase(const FObjectInitializer& ObjectInitializer = FObjectInitializer::Get());
 
 };
 
diff --git a/Engine/Source/Runtime/Engine/Classes/Particles/Color/ParticleModuleColorOverLife.h b/Engine/Source/Runtime/Engine/Classes/Particles/Color/ParticleModuleColorOverLife.h
--- a/Engine/Source/Runtime/Engine/Classes/Particles/Color/ParticleModuleColorOverLife.h
+++ b/Engine/Source/Runtime/Engine/Classes/Particles/Color/ParticleModuleColorOverLife.h
@@ -8,7 +8,9 @@
 UCLASS(editinlinenew, hidecategories=Object, MinimalAPI, meta=(DisplayName = "Color Over Life"))
 class UParticleModuleColorOverLife : public UParticleModuleColorBase
 {
-	GENERATED_UCLASS_BODY()
+	GENERATED_BODY()
+public:
+	ENGINE_API UParticleModuleColorOverLife(const FObjectInitializer& ObjectInitializer = FObjectInitializer::Get());
 
 	/** The color to apply to the particle, as a function of the particle RelativeTime. */
 	UPROPERTY(EditAnywhere, Category=Color, meta=(TreatAsColor))
diff --git a/Engine/Source/Runtime/Engine/Classes/Particles/Color/ParticleModuleColorScaleOverLife.h b/Engine/Source/Runtime/Engine/Classes/Particles/Color/ParticleModuleColorScaleOverLife.h
--- a/Engine/Source/Runtime/Engine/Classes/Particles/Color/ParticleModuleColorScaleOverLife.h
+++ b/Engine/Source/Runtime/Engine/Classes/Particles/Color/ParticleModuleColorScaleOverLife.h
@@ -14,7 +14,9 @@
 UCLASS(editinlinenew, hidecategories=Object, MinimalAPI, meta=(DisplayName = "Scale Color / Life"))
 class UParticleModuleColorScaleOverLife : public UParticleModuleColorBase
 {
-	GENERATED_UCLASS_BODY()
+	GENERATED_BODY()
+public:
+	ENGINE_API UParticleModuleColorScaleOverLife(const FObjectInitializer& ObjectInitializer = FObjectInitializer::Get());
 
 	/** The scale factor for the color.													*/
 	UPROPERTY(EditAnywhere, Category = Color, meta = (TreatAsColor))
diff --git a/Engine/Source/Runtime/Engine/Classes/Particles/Color/ParticleModuleColor_Seeded.h b/Engine/Source/Runtime/Engine/Classes/Particles/Color/ParticleModuleColor_Seeded.h
--- a/Engine/Source/Runtime/Engine/Classes/Particles/Color/ParticleModuleColor_Seeded.h
+++ b/Engine/Source/Runtime/Engine/Classes/Particles/Color/ParticleModuleColor_Seeded.h
@@ -8,7 +8,9 @@
 UCLASS(editinlinenew, hidecategories=Object, meta=(DisplayName = "Init Color (Seed)"))
 class UParticleModuleColor_Seeded : public UParticleModuleColor
 {
-	GENERATED_UCLASS_BODY()
+	GENERATED_BODY()
+public:
+	UParticleModuleColor_Seeded(const FObjectInitializer& ObjectInitializer = FObjectInitializer::Get());
 
 	/** The random seed(s) to use for looking up values in StartLocation */
 	UPROPERTY(EditAnywhere, Category=RandomSeed)
diff --git a/Engine/Source/Runtime/Engine/Classes/Particles/Emitter.h b/Engine/Source/Runtime/Engine/Classes/Particles/Emitter.h
--- a/Engine/Source/Runtime/Engine/Classes/Particles/Emitter.h
+++ b/Engine/Source/Runtime/Engine/Classes/Particles/Emitter.h
@@ -47,7 +47,9 @@ DECLARE_DYNAMIC_MULTICAST_DELEGATE_EightParams( FParticleCollisionSignature, FNa
 UCLASS(hideCategories=(Activation,"Components|Activation",Input,Collision,"Game|Damage"), ComponentWrapperClass)
 class ENGINE_API AEmitter : public AActor
 {
-	GENERATED_UCLASS_BODY()
+	GENERATED_BODY()
+public:
+	AEmitter(const FObjectInitializer& ObjectInitializer = FObjectInitializer::Get());
 
 private_subobject:
 	DEPRECATED_FORGAME(4.6, "ParticleSystemComponent should not be accessed directly, please use GetParticleSystemComponent() function instead. ParticleSystemComponent will soon be private and your code will not compile.")
diff --git a/Engine/Source/Runtime/Engine/Classes/Particles/EmitterCameraLensEffectBase.h b/Engine/Source/Runtime/Engine/Classes/Particles/EmitterCameraLensEffectBase.h
--- a/Engine/Source/Runtime/Engine/Classes/Particles/EmitterCameraLensEffectBase.h
+++ b/Engine/Source/Runtime/Engine/Classes/Particles/EmitterCameraLensEffectBase.h
@@ -13,7 +13,9 @@
 UCLASS(abstract, Blueprintable, MinimalAPI)
 class AEmitterCameraLensEffectBase : public AEmitter
 {
-	GENERATED_UCLASS_BODY()
+	GENERATED_BODY()
+public:
+	ENGINE_API AEmitterCameraLensEffectBase(const FObjectInitializer& ObjectInitializer = FObjectInitializer::Get());
 
 protected:
 	/** Particle System to use */
diff --git a/Engine/Source/Runtime/Engine/Classes/Particles/Event/ParticleModuleEventBase.h b/Engine/Source/Runtime/Engine/Classes/Particles/Event/ParticleModuleEventBase.h
--- a/Engine/Source/Runtime/Engine/Classes/Particles/Event/ParticleModuleEventBase.h
+++ b/Engine/Source/Runtime/Engine/Classes/Particles/Event/ParticleModuleEventBase.h
@@ -8,7 +8,9 @@
 UCLASS(editinlinenew, hidecategories=Object, abstract, meta=(DisplayName = "Event"))
 class UParticleModuleEventBase : public UParticleModule
 {
-	GENERATED_UCLASS_BODY()
+	GENERATED_BODY()
+public:
+	UParticleModuleEventBase(const FObjectInitializer& ObjectInitializer = FObjectInitializer::Get());
 
 
 	// Begin UParticleModule Interface
diff --git a/Engine/Source/Runtime/Engine/Classes/Particles/Event/ParticleModuleEventGenerator.h b/Engine/Source/Runtime/Engine/Classes/Particles/Event/ParticleModuleEventGenerator.h
--- a/Engine/Source/Runtime/Engine/Classes/Particles/Event/ParticleModuleEventGenerator.h
+++ b/Engine/Source/Runtime/Engine/Classes/Particles/Event/ParticleModuleEventGenerator.h
@@ -10,7 +10,7 @@
 USTRUCT()
 struct FParticleEvent_GenerateInfo
 {
-	GENERATED_USTRUCT_BODY()
+	GENERATED_BODY()
 
 	/** The type of event. */
 	UPROPERTY(EditAnywhere, Category=ParticleEvent_GenerateInfo)
@@ -61,7 +61,9 @@ struct FParticleEvent_GenerateInfo
 UCLASS(editinlinenew, hidecategories=Object, MinimalAPI, meta=(DisplayName = "Event Generator"))
 class UParticleModuleEventGenerator : public UParticleModuleEventBase
 {
-	GENERATED_UCLASS_BODY()
+	GENERATED_BODY()
+public:
+	ENGINE_API UParticleModuleEventGenerator(const FObjectInitializer& ObjectInitializer = FObjectInitializer::Get());
 
 	/**
 	 */
diff --git a/Engine/Source/Runtime/Engine/Classes/Particles/Event/ParticleModuleEventReceiverBase.h b/Engine/Source/Runtime/Engine/Classes/Particles/Event/ParticleModuleEventReceiverBase.h
--- a/Engine/Source/Runtime/Engine/Classes/Particles/Event/ParticleModuleEventReceiverBase.h
+++ b/Engine/Source/Runtime/Engine/Classes/Particles/Event/ParticleModuleEventReceiverBase.h
@@ -9,7 +9,9 @@
 UCLASS(editinlinenew, hidecategories=Object, abstract)
 class UParticleModuleEventReceiverBase : public UParticleModuleEventBase
 {
-	GENERATED_UCLASS_BODY()
+	GENERATED_BODY()
+public:
+	UParticleModuleEventReceiverBase(const FObjectInitializer& ObjectInitializer = FObjectInitializer::Get());
 
 	/** The type of event that will generate the kill. */
 	UPROPERTY(EditAnywhere, Category=Source)
diff --git a/Engine/Source/Runtime/Engine/Classes/Particles/Event/ParticleModuleEventReceiverKillParticles.h b/Engine/Source/Runtime/Engine/Classes/Particles/Event/ParticleModuleEventReceiverKillParticles.h
--- a/Engine/Source/Runtime/Engine/Classes/Particles/Event/ParticleModuleEventReceiverKillParticles.h
+++ b/Engine/Source/Runtime/Engine/Classes/Particles/Event/ParticleModuleEventReceiverKillParticles.h
@@ -8,7 +8,9 @@
 UCLASS(editinlinenew, hidecategories=Object, meta=(DisplayName = "EventReceiver Kill All"))
 class UParticleModuleEventReceiverKillParticles : public UParticleModuleEventReceiverBase
 {
-	GENERATED_UCLASS_BODY()
+	GENERATED_BODY()
+public:
+	UParticleModuleEventReceiverKillParticles(const FObjectInitializer& ObjectInitializer = FObjectInitializer::Get());
 
 	/** If true, stop this emitter from spawning as well. */
 	UPROPERTY(EditAnywhere, Category=ParticleModuleEventReceiverKillParticles)
diff --git a/Engine/Source/Runtime/Engine/Classes/Particles/Event/ParticleModuleEventReceiverSpawn.h b/Engine/Source/Runtime/Engine/Classes/Particles/Event/ParticleModuleEventReceiverSpawn.h
--- a/Engine/Source/Runtime/Engine/Classes/Particles/Event/ParticleModuleEventReceiverSpawn.h
+++ b/Engine/Source/Runtime/Engine/Classes/Particles/Event/ParticleModuleEventReceiverSpawn.h
@@ -8,7 +8,9 @@
 UCLASS(editinlinenew, hidecategories=Object, meta=(DisplayName = "EventReceiver Spawn"))
 class UParticleModuleEventReceiverSpawn : public UParticleModuleEventReceiverBase
 {
-	GENERATED_UCLASS_BODY()
+	GENERATED_BODY()
+public:
+	UParticleModuleEventReceiverSpawn(const FObjectInitializer& ObjectInitializer = FObjectInitializer::Get());
 
 	/** The number of particles to spawn. */
 	UPROPERTY(EditAnywhere, Category=Spawn)
diff --git a/Engine/Source/Runtime/Engine/Classes/Particles/Event/ParticleModuleEventSendToGame.h b/Engine/Source/Runtime/Engine/Classes/Particles/Event/ParticleModuleEventSendToGame.h
--- a/Engine/Source/Runtime/Engine/Classes/Particles/Event/ParticleModuleEventSendToGame.h
+++ b/Engine/Source/Runtime/Engine/Classes/Particles/Event/ParticleModuleEventSendToGame.h
@@ -7,7 +7,9 @@
 UCLASS(abstract, editinlinenew, hidecategories=Object)
 class UParticleModuleEventSendToGame : public UObject
 {
-	GENERATED_UCLASS_BODY()
+	GENERATED_BODY()
+public:
+	UParticleModuleEventSendToGame(const FObjectInitializer& ObjectInitializer = FObjectInitializer::Get());
 
 
 	/** This is our function to allow subclasses to "do the event action" **/
diff --git a/Engine/Source/Runtime/Engine/Classes/Particles/Kill/ParticleModuleKillBase.h b/Engine/Source/Runtime/Engine/Classes/Particles/Kill/ParticleModuleKillBase.h
--- a/Engine/Source/Runtime/Engine/Classes/Particles/Kill/ParticleModuleKillBase.h
+++ b/Engine/Source/Runtime/Engine/Classes/Particles/Kill/ParticleModuleKillBase.h
@@ -8,7 +8,9 @@
 UCLASS(editinlinenew, hidecategories=Object, abstract, meta=(DisplayName = "Kill"))
 class UParticleModuleKillBase : public UParticleModule
 {
-	GENERATED_UCLASS_BODY()
+	GENERATED_BODY()
+public:
+	UParticleModuleKillBase(const FObjectInitializer& ObjectInitializer = FObjectInitializer::Get());
 
 };
 
diff --git a/Engine/Source/Runtime/Engine/Classes/Particles/Kill/ParticleModuleKillBox.h b/Engine/Source/Runtime/Engine/Classes/Particles/Kill/ParticleModuleKillBox.h
--- a/Engine/Source/Runtime/Engine/Classes/Particles/Kill/ParticleModuleKillBox.h
+++ b/Engine/Source/Runtime/Engine/Classes/Particles/Kill/ParticleModuleKillBox.h
@@ -8,7 +8,9 @@
 UCLASS(editinlinenew, hidecategories=Object, meta=(DisplayName = "Kill Box"))
 class UParticleModuleKillBox : public UParticleModuleKillBase
 {
-	GENERATED_UCLASS_BODY()
+	GENERATED_BODY()
+public:
+	UParticleModuleKillBox(const FObjectInitializer& ObjectInitializer = FObjectInitializer::Get());
 
 	/** The lower left corner of the box. */
 	UPROPERTY(EditAnywhere, Category=Kill)
diff --git a/Engine/Source/Runtime/Engine/Classes/Particles/Kill/ParticleModuleKillHeight.h b/Engine/Source/Runtime/Engine/Classes/Particles/Kill/ParticleModuleKillHeight.h
--- a/Engine/Source/Runtime/Engine/Classes/Particles/Kill/ParticleModuleKillHeight.h
+++ b/Engine/Source/Runtime/Engine/Classes/Particles/Kill/ParticleModuleKillHeight.h
@@ -8,7 +8,9 @@
 UCLASS(editinlinenew, hidecategories=Object, meta=(DisplayName = "Kill Height"))
 class UParticleModuleKillHeight : public UParticleModuleKillBase
 {
-	GENERATED_UCLASS_BODY()
+	GENERATED_BODY()
+public:
+	UParticleModuleKillHeight(const FObjectInitializer& ObjectInitializer = FObjectInitializer::Get());
 
 	/** The height at which to kill the particle. */
 	UPROPERTY(EditAnywhere, Category=Kill)
diff --git a/Engine/Source/Runtime/Engine/Classes/Particles/Lifetime/ParticleModuleLifetime.h b/Engine/Source/Runtime/Engine/Classes/Particles/Lifetime/ParticleModuleLifetime.h
--- a/Engine/Source/Runtime/Engine/Classes/Particles/Lifetime/ParticleModuleLifetime.h
+++ b/Engine/Source/Runtime/Engine/Classes/Particles/Lifetime/ParticleModuleLifetime.h
@@ -8,7 +8,9 @@
 UCLASS(editinlinenew, hidecategories=Object, MinimalAPI, meta=(DisplayName = "Lifetime"))
 class UParticleModuleLifetime : public UParticleModuleLifetimeBase
 {
-	GENERATED_UCLASS_BODY()
+	GENERATED_BODY()
+public:
+	ENGINE_API UParticleModuleLifetime(const FObjectInitializer& ObjectInitializer = FObjectInitializer::Get());
 
 	/** The lifetime of the particle, in seconds. Retrieved using the EmitterTime at the spawn of the particle. */
 	UPROPERTY(EditAnywhere, Category=Lifetime)
diff --git a/Engine/Source/Runtime/Engine/Classes/Particles/Lifetime/ParticleModuleLifetimeBase.h b/Engine/Source/Runtime/Engine/Classes/Particles/Lifetime/ParticleModuleLifetimeBase.h
--- a/Engine/Source/Runtime/Engine/Classes/Particles/Lifetime/ParticleModuleLifetimeBase.h
+++ b/Engine/Source/Runtime/Engine/Classes/Particles/Lifetime/ParticleModuleLifetimeBase.h
@@ -8,7 +8,9 @@
 UCLASS(editinlinenew, hidecategories=Object, abstract, meta=(DisplayName = "Lifetime"))
 class UParticleModuleLifetimeBase : public UParticleModule
 {
-	GENERATED_UCLASS_BODY()
+	GENERATED_BODY()
+public:
+	UParticleModuleLifetimeBase(const FObjectInitializer& ObjectInitializer = FObjectInitializer::Get());
 
 
 	/** Return the maximum lifetime this module would return. */
diff --git a/Engine/Source/Runtime/Engine/Classes/Particles/Lifetime/ParticleModuleLifetime_Seeded.h b/Engine/Source/Runtime/Engine/Classes/Particles/Lifetime/ParticleModuleLifetime_Seeded.h
--- a/Engine/Source/Runtime/Engine/Classes/Particles/Lifetime/ParticleModuleLifetime_Seeded.h
+++ b/Engine/Source/Runtime/Engine/Classes/Particles/Lifetime/ParticleModuleLifetime_Seeded.h
@@ -8,7 +8,9 @@
 UCLASS(editinlinenew, hidecategories=Object, meta=(DisplayName = "Lifetime (Seed)"))
 class UParticleModuleLifetime_Seeded : public UParticleModuleLifetime
 {
-	GENERATED_UCLASS_BODY()
+	GENERATED_BODY()
+public:
+	UParticleModuleLifetime_Seeded(const FObjectInitializer& ObjectInitializer = FObjectInitializer::Get());
 
 	/** The random seed(s) to use for looking up values in StartLocation */
 	UPROPERTY(EditAnywhere, Category=RandomSeed)
diff --git a/Engine/Source/Runtime/Engine/Classes/Particles/Light/ParticleModuleLight.h b/Engine/Source/Runtime/Engine/Classes/Particles/Light/ParticleModuleLight.h
--- a/Engine/Source/Runtime/Engine/Classes/Particles/Light/ParticleModuleLight.h
+++ b/Engine/Source/Runtime/Engine/Classes/Particles/Light/ParticleModuleLight.h
@@ -9,7 +9,9 @@
 UCLASS(editinlinenew, hidecategories=Object, MinimalAPI, meta=(DisplayName = "Light"))
 class UParticleModuleLight : public UParticleModuleLightBase
 {
-	GENERATED_UCLASS_BODY()
+	GENERATED_BODY()
+public:
+	ENGINE_API UParticleModuleLight(const FObjectInitializer& ObjectInitializer = FObjectInitializer::Get());
 
 	/** Whether to use physically based inverse squared falloff from the light.  If unchecked, the LightExponent distribution will be used instead. */
 	UPROPERTY(EditAnywhere, Category=Light)
diff --git a/Engine/Source/Runtime/Engine/Classes/Particles/Light/ParticleModuleLightBase.h b/Engine/Source/Runtime/Engine/Classes/Particles/Light/ParticleModuleLightBase.h
--- a/Engine/Source/Runtime/Engine/Classes/Particles/Light/ParticleModuleLightBase.h
+++ b/Engine/Source/Runtime/Engine/Classes/Particles/Light/ParticleModuleLightBase.h
@@ -8,7 +8,9 @@
 UCLASS(editinlinenew, hidecategories=Object, abstract, meta=(DisplayName = "Light"))
 class UParticleModuleLightBase : public UParticleModule
 {
-	GENERATED_UCLASS_BODY()
+	GENERATED_BODY()
+public:
+	UParticleModuleLightBase(const FObjectInitializer& ObjectInitializer = FObjectInitializer::Get());
 	virtual bool CanTickInAnyThread() override
 	{
 		return false;
