Enabled UHT to digest GENERATED_BODY instead of GENERATED_UCLASS_BODY, GENERATED_USTRUCT_BODY, GENERATED_UINTERFACE_BODY or GENERATED_IINTERFACE_BODY, changed every occurence to the new syntax and fix

diff --git a/Engine/Source/Runtime/Engine/Classes/Particles/Light/ParticleModuleLight_Seeded.h b/Engine/Source/Runtime/Engine/Classes/Particles/Light/ParticleModuleLight_Seeded.h
--- a/Engine/Source/Runtime/Engine/Classes/Particles/Light/ParticleModuleLight_Seeded.h
+++ b/Engine/Source/Runtime/Engine/Classes/Particles/Light/ParticleModuleLight_Seeded.h
@@ -8,7 +8,9 @@
 UCLASS(editinlinenew, hidecategories=Object, MinimalAPI, meta=(DisplayName = "Light (Seed)"))
 class UParticleModuleLight_Seeded : public UParticleModuleLight
 {
-	GENERATED_UCLASS_BODY()
+	GENERATED_BODY()
+public:
+	ENGINE_API UParticleModuleLight_Seeded(const FObjectInitializer& ObjectInitializer = FObjectInitializer::Get());
 
 	/** The random seed(s) to use for looking up values in StartLocation */
 	UPROPERTY(EditAnywhere, Category=RandomSeed)
diff --git a/Engine/Source/Runtime/Engine/Classes/Particles/Location/ParticleModuleLocation.h b/Engine/Source/Runtime/Engine/Classes/Particles/Location/ParticleModuleLocation.h
--- a/Engine/Source/Runtime/Engine/Classes/Particles/Location/ParticleModuleLocation.h
+++ b/Engine/Source/Runtime/Engine/Classes/Particles/Location/ParticleModuleLocation.h
@@ -8,7 +8,9 @@
 UCLASS(editinlinenew, hidecategories=Object, MinimalAPI, meta=(DisplayName = "Initial Location"))
 class UParticleModuleLocation : public UParticleModuleLocationBase
 {
-	GENERATED_UCLASS_BODY()
+	GENERATED_BODY()
+public:
+	ENGINE_API UParticleModuleLocation(const FObjectInitializer& ObjectInitializer = FObjectInitializer::Get());
 
 	/** 
 	 *	The location the particle should be emitted.
diff --git a/Engine/Source/Runtime/Engine/Classes/Particles/Location/ParticleModuleLocationBase.h b/Engine/Source/Runtime/Engine/Classes/Particles/Location/ParticleModuleLocationBase.h
--- a/Engine/Source/Runtime/Engine/Classes/Particles/Location/ParticleModuleLocationBase.h
+++ b/Engine/Source/Runtime/Engine/Classes/Particles/Location/ParticleModuleLocationBase.h
@@ -8,7 +8,9 @@
 UCLASS(editinlinenew, hidecategories=Object, abstract, meta=(DisplayName = "Location"))
 class UParticleModuleLocationBase : public UParticleModule
 {
-	GENERATED_UCLASS_BODY()
+	GENERATED_BODY()
+public:
+	UParticleModuleLocationBase(const FObjectInitializer& ObjectInitializer = FObjectInitializer::Get());
 
 };
 
diff --git a/Engine/Source/Runtime/Engine/Classes/Particles/Location/ParticleModuleLocationBoneSocket.h b/Engine/Source/Runtime/Engine/Classes/Particles/Location/ParticleModuleLocationBoneSocket.h
--- a/Engine/Source/Runtime/Engine/Classes/Particles/Location/ParticleModuleLocationBoneSocket.h
+++ b/Engine/Source/Runtime/Engine/Classes/Particles/Location/ParticleModuleLocationBoneSocket.h
@@ -24,7 +24,7 @@ enum ELocationBoneSocketSelectionMethod
 USTRUCT()
 struct FLocationBoneSocketInfo
 {
-	GENERATED_USTRUCT_BODY()
+	GENERATED_BODY()
 
 	/** The name of the bone/socket on the skeletal mesh */
 	UPROPERTY(EditAnywhere, Category=LocationBoneSocketInfo)
@@ -45,7 +45,9 @@ struct FLocationBoneSocketInfo
 UCLASS(editinlinenew, hidecategories=Object, meta=(DisplayName = "Bone/Socket Location"))
 class UParticleModuleLocationBoneSocket : public UParticleModuleLocationBase
 {
-	GENERATED_UCLASS_BODY()
+	GENERATED_BODY()
+public:
+	UParticleModuleLocationBoneSocket(const FObjectInitializer& ObjectInitializer = FObjectInitializer::Get());
 
 	/**
 	 *	Whether the module uses Bones or Sockets for locations.
diff --git a/Engine/Source/Runtime/Engine/Classes/Particles/Location/ParticleModuleLocationDirect.h b/Engine/Source/Runtime/Engine/Classes/Particles/Location/ParticleModuleLocationDirect.h
--- a/Engine/Source/Runtime/Engine/Classes/Particles/Location/ParticleModuleLocationDirect.h
+++ b/Engine/Source/Runtime/Engine/Classes/Particles/Location/ParticleModuleLocationDirect.h
@@ -14,7 +14,9 @@
 UCLASS(editinlinenew, hidecategories=Object, MinimalAPI, meta=(DisplayName = "Direct Location"))
 class UParticleModuleLocationDirect : public UParticleModuleLocationBase
 {
-	GENERATED_UCLASS_BODY()
+	GENERATED_BODY()
+public:
+	ENGINE_API UParticleModuleLocationDirect(const FObjectInitializer& ObjectInitializer = FObjectInitializer::Get());
 
 	/** 
 	 *	The location of the particle at a give time. Retrieved using the particle RelativeTime. 
diff --git a/Engine/Source/Runtime/Engine/Classes/Particles/Location/ParticleModuleLocationEmitter.h b/Engine/Source/Runtime/Engine/Classes/Particles/Location/ParticleModuleLocationEmitter.h
--- a/Engine/Source/Runtime/Engine/Classes/Particles/Location/ParticleModuleLocationEmitter.h
+++ b/Engine/Source/Runtime/Engine/Classes/Particles/Location/ParticleModuleLocationEmitter.h
@@ -23,7 +23,9 @@ enum ELocationEmitterSelectionMethod
 UCLASS(editinlinenew, hidecategories=Object, MinimalAPI, meta=(DisplayName = "Emitter Initial Location"))
 class UParticleModuleLocationEmitter : public UParticleModuleLocationBase
 {
-	GENERATED_UCLASS_BODY()
+	GENERATED_BODY()
+public:
+	ENGINE_API UParticleModuleLocationEmitter(const FObjectInitializer& ObjectInitializer = FObjectInitializer::Get());
 
 	//=============================================================================
 	// Variables
diff --git a/Engine/Source/Runtime/Engine/Classes/Particles/Location/ParticleModuleLocationEmitterDirect.h b/Engine/Source/Runtime/Engine/Classes/Particles/Location/ParticleModuleLocationEmitterDirect.h
--- a/Engine/Source/Runtime/Engine/Classes/Particles/Location/ParticleModuleLocationEmitterDirect.h
+++ b/Engine/Source/Runtime/Engine/Classes/Particles/Location/ParticleModuleLocationEmitterDirect.h
@@ -15,7 +15,9 @@
 UCLASS(editinlinenew, hidecategories=Object, meta=(DisplayName = "Emitter Direct Location"))
 class UParticleModuleLocationEmitterDirect : public UParticleModuleLocationBase
 {
-	GENERATED_UCLASS_BODY()
+	GENERATED_BODY()
+public:
+	UParticleModuleLocationEmitterDirect(const FObjectInitializer& ObjectInitializer = FObjectInitializer::Get());
 
 	/** The name of the emitter to use as a source for the location of the particles. */
 	UPROPERTY(EditAnywhere, export, noclear, Category=Location)
diff --git a/Engine/Source/Runtime/Engine/Classes/Particles/Location/ParticleModuleLocationPrimitiveBase.h b/Engine/Source/Runtime/Engine/Classes/Particles/Location/ParticleModuleLocationPrimitiveBase.h
--- a/Engine/Source/Runtime/Engine/Classes/Particles/Location/ParticleModuleLocationPrimitiveBase.h
+++ b/Engine/Source/Runtime/Engine/Classes/Particles/Location/ParticleModuleLocationPrimitiveBase.h
@@ -12,7 +12,9 @@
 UCLASS(editinlinenew, hidecategories=Object)
 class UParticleModuleLocationPrimitiveBase : public UParticleModuleLocationBase
 {
-	GENERATED_UCLASS_BODY()
+	GENERATED_BODY()
+public:
+	UParticleModuleLocationPrimitiveBase(const FObjectInitializer& ObjectInitializer = FObjectInitializer::Get());
 
 	/** Whether the positive X axis is valid for spawning. */
 	UPROPERTY(EditAnywhere, Category=Location)
diff --git a/Engine/Source/Runtime/Engine/Classes/Particles/Location/ParticleModuleLocationPrimitiveCylinder.h b/Engine/Source/Runtime/Engine/Classes/Particles/Location/ParticleModuleLocationPrimitiveCylinder.h
--- a/Engine/Source/Runtime/Engine/Classes/Particles/Location/ParticleModuleLocationPrimitiveCylinder.h
+++ b/Engine/Source/Runtime/Engine/Classes/Particles/Location/ParticleModuleLocationPrimitiveCylinder.h
@@ -21,7 +21,9 @@ enum CylinderHeightAxis
 UCLASS(editinlinenew, hidecategories = Object, MinimalAPI, meta = (DisplayName = "Cylinder"))
 class UParticleModuleLocationPrimitiveCylinder : public UParticleModuleLocationPrimitiveBase
 {
-	GENERATED_UCLASS_BODY()
+	GENERATED_BODY()
+public:
+	ENGINE_API UParticleModuleLocationPrimitiveCylinder(const FObjectInitializer& ObjectInitializer = FObjectInitializer::Get());
 
 	/** If true, get the particle velocity form the radial distance inside the primitive. */
 	UPROPERTY(EditAnywhere, Category=Location)
diff --git a/Engine/Source/Runtime/Engine/Classes/Particles/Location/ParticleModuleLocationPrimitiveCylinder_Seeded.h b/Engine/Source/Runtime/Engine/Classes/Particles/Location/ParticleModuleLocationPrimitiveCylinder_Seeded.h
--- a/Engine/Source/Runtime/Engine/Classes/Particles/Location/ParticleModuleLocationPrimitiveCylinder_Seeded.h
+++ b/Engine/Source/Runtime/Engine/Classes/Particles/Location/ParticleModuleLocationPrimitiveCylinder_Seeded.h
@@ -8,7 +8,9 @@
 UCLASS(editinlinenew, hidecategories=Object, meta=(DisplayName = "Cylinder (Seed)"))
 class UParticleModuleLocationPrimitiveCylinder_Seeded : public UParticleModuleLocationPrimitiveCylinder
 {
-	GENERATED_UCLASS_BODY()
+	GENERATED_BODY()
+public:
+	UParticleModuleLocationPrimitiveCylinder_Seeded(const FObjectInitializer& ObjectInitializer = FObjectInitializer::Get());
 
 	/** The random seed(s) to use for looking up values in StartLocation */
 	UPROPERTY(EditAnywhere, Category=RandomSeed)
diff --git a/Engine/Source/Runtime/Engine/Classes/Particles/Location/ParticleModuleLocationPrimitiveSphere.h b/Engine/Source/Runtime/Engine/Classes/Particles/Location/ParticleModuleLocationPrimitiveSphere.h
--- a/Engine/Source/Runtime/Engine/Classes/Particles/Location/ParticleModuleLocationPrimitiveSphere.h
+++ b/Engine/Source/Runtime/Engine/Classes/Particles/Location/ParticleModuleLocationPrimitiveSphere.h
@@ -12,7 +12,9 @@
 UCLASS(editinlinenew, hidecategories=Object, MinimalAPI, meta=(DisplayName = "Sphere"))
 class UParticleModuleLocationPrimitiveSphere : public UParticleModuleLocationPrimitiveBase
 {
-	GENERATED_UCLASS_BODY()
+	GENERATED_BODY()
+public:
+	ENGINE_API UParticleModuleLocationPrimitiveSphere(const FObjectInitializer& ObjectInitializer = FObjectInitializer::Get());
 
 	/** The radius of the sphere. Retrieved using EmitterTime. */
 	UPROPERTY(EditAnywhere, Category=Location)
diff --git a/Engine/Source/Runtime/Engine/Classes/Particles/Location/ParticleModuleLocationPrimitiveSphere_Seeded.h b/Engine/Source/Runtime/Engine/Classes/Particles/Location/ParticleModuleLocationPrimitiveSphere_Seeded.h
--- a/Engine/Source/Runtime/Engine/Classes/Particles/Location/ParticleModuleLocationPrimitiveSphere_Seeded.h
+++ b/Engine/Source/Runtime/Engine/Classes/Particles/Location/ParticleModuleLocationPrimitiveSphere_Seeded.h
@@ -8,7 +8,9 @@
 UCLASS(editinlinenew, hidecategories=Object, meta=(DisplayName = "Sphere (Seed)"))
 class UParticleModuleLocationPrimitiveSphere_Seeded : public UParticleModuleLocationPrimitiveSphere
 {
-	GENERATED_UCLASS_BODY()
+	GENERATED_BODY()
+public:
+	UParticleModuleLocationPrimitiveSphere_Seeded(const FObjectInitializer& ObjectInitializer = FObjectInitializer::Get());
 
 	/** The random seed(s) to use for looking up values in StartLocation */
 	UPROPERTY(EditAnywhere, Category=RandomSeed)
diff --git a/Engine/Source/Runtime/Engine/Classes/Particles/Location/ParticleModuleLocationPrimitiveTriangle.h b/Engine/Source/Runtime/Engine/Classes/Particles/Location/ParticleModuleLocationPrimitiveTriangle.h
--- a/Engine/Source/Runtime/Engine/Classes/Particles/Location/ParticleModuleLocationPrimitiveTriangle.h
+++ b/Engine/Source/Runtime/Engine/Classes/Particles/Location/ParticleModuleLocationPrimitiveTriangle.h
@@ -12,7 +12,9 @@
 UCLASS(editinlinenew, hidecategories=Object, MinimalAPI, meta=(DisplayName = "Triangle"))
 class UParticleModuleLocationPrimitiveTriangle : public UParticleModuleLocationBase
 {
-	GENERATED_UCLASS_BODY()
+	GENERATED_BODY()
+public:
+	ENGINE_API UParticleModuleLocationPrimitiveTriangle(const FObjectInitializer& ObjectInitializer = FObjectInitializer::Get());
 
 
 	UPROPERTY(EditAnywhere, Category=Location)
diff --git a/Engine/Source/Runtime/Engine/Classes/Particles/Location/ParticleModuleLocationSkelVertSurface.h b/Engine/Source/Runtime/Engine/Classes/Particles/Location/ParticleModuleLocationSkelVertSurface.h
--- a/Engine/Source/Runtime/Engine/Classes/Particles/Location/ParticleModuleLocationSkelVertSurface.h
+++ b/Engine/Source/Runtime/Engine/Classes/Particles/Location/ParticleModuleLocationSkelVertSurface.h
@@ -16,7 +16,9 @@ enum ELocationSkelVertSurfaceSource
 UCLASS(editinlinenew, hidecategories=Object, meta=(DisplayName = "Skel Vert/Surf Location"))
 class UParticleModuleLocationSkelVertSurface : public UParticleModuleLocationBase
 {
-	GENERATED_UCLASS_BODY()
+	GENERATED_BODY()
+public:
+	UParticleModuleLocationSkelVertSurface(const FObjectInitializer& ObjectInitializer = FObjectInitializer::Get());
 
 	/**
 	 *	Whether the module uses Verts or Surfaces for locations.
diff --git a/Engine/Source/Runtime/Engine/Classes/Particles/Location/ParticleModuleLocationWorldOffset.h b/Engine/Source/Runtime/Engine/Classes/Particles/Location/ParticleModuleLocationWorldOffset.h
--- a/Engine/Source/Runtime/Engine/Classes/Particles/Location/ParticleModuleLocationWorldOffset.h
+++ b/Engine/Source/Runtime/Engine/Classes/Particles/Location/ParticleModuleLocationWorldOffset.h
@@ -8,7 +8,9 @@
 UCLASS(editinlinenew, meta=(DisplayName = "World Offset"))
 class UParticleModuleLocationWorldOffset : public UParticleModuleLocation
 {
-	GENERATED_UCLASS_BODY()
+	GENERATED_BODY()
+public:
+	UParticleModuleLocationWorldOffset(const FObjectInitializer& ObjectInitializer = FObjectInitializer::Get());
 
 
 protected:
diff --git a/Engine/Source/Runtime/Engine/Classes/Particles/Location/ParticleModuleLocationWorldOffset_Seeded.h b/Engine/Source/Runtime/Engine/Classes/Particles/Location/ParticleModuleLocationWorldOffset_Seeded.h
--- a/Engine/Source/Runtime/Engine/Classes/Particles/Location/ParticleModuleLocationWorldOffset_Seeded.h
+++ b/Engine/Source/Runtime/Engine/Classes/Particles/Location/ParticleModuleLocationWorldOffset_Seeded.h
@@ -8,7 +8,9 @@
 UCLASS(editinlinenew, meta=(DisplayName = "World Offset (Seed)"))
 class UParticleModuleLocationWorldOffset_Seeded : public UParticleModuleLocationWorldOffset
 {
-	GENERATED_UCLASS_BODY()
+	GENERATED_BODY()
+public:
+	UParticleModuleLocationWorldOffset_Seeded(const FObjectInitializer& ObjectInitializer = FObjectInitializer::Get());
 
 	/** The random seed(s) to use for looking up values in StartLocation */
 	UPROPERTY(EditAnywhere, Category=RandomSeed)
diff --git a/Engine/Source/Runtime/Engine/Classes/Particles/Location/ParticleModuleLocation_Seeded.h b/Engine/Source/Runtime/Engine/Classes/Particles/Location/ParticleModuleLocation_Seeded.h
--- a/Engine/Source/Runtime/Engine/Classes/Particles/Location/ParticleModuleLocation_Seeded.h
+++ b/Engine/Source/Runtime/Engine/Classes/Particles/Location/ParticleModuleLocation_Seeded.h
@@ -8,7 +8,9 @@
 UCLASS(editinlinenew, hidecategories=Object, meta=(DisplayName = "Initial Location (Seed)"))
 class UParticleModuleLocation_Seeded : public UParticleModuleLocation
 {
-	GENERATED_UCLASS_BODY()
+	GENERATED_BODY()
+public:
+	UParticleModuleLocation_Seeded(const FObjectInitializer& ObjectInitializer = FObjectInitializer::Get());
 
 	/** The random seed(s) to use for looking up values in StartLocation */
 	UPROPERTY(EditAnywhere, Category=RandomSeed)
diff --git a/Engine/Source/Runtime/Engine/Classes/Particles/Location/ParticleModuleSourceMovement.h b/Engine/Source/Runtime/Engine/Classes/Particles/Location/ParticleModuleSourceMovement.h
--- a/Engine/Source/Runtime/Engine/Classes/Particles/Location/ParticleModuleSourceMovement.h
+++ b/Engine/Source/Runtime/Engine/Classes/Particles/Location/ParticleModuleSourceMovement.h
@@ -8,7 +8,9 @@
 UCLASS(editinlinenew, hidecategories=Object, meta=(DisplayName = "Source Movement"))
 class UParticleModuleSourceMovement : public UParticleModuleLocationBase
 {
-	GENERATED_UCLASS_BODY()
+	GENERATED_BODY()
+public:
+	UParticleModuleSourceMovement(const FObjectInitializer& ObjectInitializer = FObjectInitializer::Get());
 
 	/** 
 	 *	The scale factor to apply to the source movement before adding to the particle location.
diff --git a/Engine/Source/Runtime/Engine/Classes/Particles/Material/ParticleModuleMaterialBase.h b/Engine/Source/Runtime/Engine/Classes/Particles/Material/ParticleModuleMaterialBase.h
--- a/Engine/Source/Runtime/Engine/Classes/Particles/Material/ParticleModuleMaterialBase.h
+++ b/Engine/Source/Runtime/Engine/Classes/Particles/Material/ParticleModuleMaterialBase.h
@@ -8,7 +8,9 @@
 UCLASS(editinlinenew, hidecategories=Object, abstract, meta=(DisplayName = "Material"))
 class UParticleModuleMaterialBase : public UParticleModule
 {
-	GENERATED_UCLASS_BODY()
+	GENERATED_BODY()
+public:
+	UParticleModuleMaterialBase(const FObjectInitializer& ObjectInitializer = FObjectInitializer::Get());
 
 	virtual bool CanTickInAnyThread() override
 	{
diff --git a/Engine/Source/Runtime/Engine/Classes/Particles/Material/ParticleModuleMeshMaterial.h b/Engine/Source/Runtime/Engine/Classes/Particles/Material/ParticleModuleMeshMaterial.h
--- a/Engine/Source/Runtime/Engine/Classes/Particles/Material/ParticleModuleMeshMaterial.h
+++ b/Engine/Source/Runtime/Engine/Classes/Particles/Material/ParticleModuleMeshMaterial.h
@@ -8,7 +8,9 @@
 UCLASS(editinlinenew, hidecategories=Object, MinimalAPI, meta=(DisplayName = "Mesh Material"))
 class UParticleModuleMeshMaterial : public UParticleModuleMaterialBase
 {
-	GENERATED_UCLASS_BODY()
+	GENERATED_BODY()
+public:
+	ENGINE_API UParticleModuleMeshMaterial(const FObjectInitializer& ObjectInitializer = FObjectInitializer::Get());
 
 	/** The array of materials to apply to the mesh particles. */
 	UPROPERTY(EditAnywhere, Category=MeshMaterials)
diff --git a/Engine/Source/Runtime/Engine/Classes/Particles/Modules/Location/ParticleModulePivotOffset.h b/Engine/Source/Runtime/Engine/Classes/Particles/Modules/Location/ParticleModulePivotOffset.h
--- a/Engine/Source/Runtime/Engine/Classes/Particles/Modules/Location/ParticleModulePivotOffset.h
+++ b/Engine/Source/Runtime/Engine/Classes/Particles/Modules/Location/ParticleModulePivotOffset.h
@@ -7,7 +7,9 @@
 UCLASS(editinlinenew, hidecategories=Object, MinimalAPI, DisplayName="Pivot Offset")
 class UParticleModulePivotOffset : public UParticleModuleLocationBase
 {
-	GENERATED_UCLASS_BODY()
+	GENERATED_BODY()
+public:
+	ENGINE_API UParticleModulePivotOffset(const FObjectInitializer& ObjectInitializer = FObjectInitializer::Get());
 
 	/** Offset applied in UV space to the particle vertex positions. Defaults to (0.5,0.5) putting the pivot in the centre of the partilce. */
 	UPROPERTY(EditAnywhere, Category=PivotOffset)
diff --git a/Engine/Source/Runtime/Engine/Classes/Particles/Orbit/ParticleModuleOrbit.h b/Engine/Source/Runtime/Engine/Classes/Particles/Orbit/ParticleModuleOrbit.h
--- a/Engine/Source/Runtime/Engine/Classes/Particles/Orbit/ParticleModuleOrbit.h
+++ b/Engine/Source/Runtime/Engine/Classes/Particles/Orbit/ParticleModuleOrbit.h
@@ -23,7 +23,7 @@ enum EOrbitChainMode
 USTRUCT()
 struct FOrbitOptions
 {
-	GENERATED_USTRUCT_BODY()
+	GENERATED_BODY()
 
 	/**
 	 *	Whether to process the data during spawning.
@@ -57,7 +57,9 @@ struct FOrbitOptions
 UCLASS(editinlinenew, hidecategories=(Object, Orbit), MinimalAPI, meta=(DisplayName = "Orbit"))
 class UParticleModuleOrbit : public UParticleModuleOrbitBase
 {
-	GENERATED_UCLASS_BODY()
+	GENERATED_BODY()
+public:
+	ENGINE_API UParticleModuleOrbit(const FObjectInitializer& ObjectInitializer = FObjectInitializer::Get());
 
 	/**
 	 *	Orbit modules will chain together in the order they appear in the module stack.
diff --git a/Engine/Source/Runtime/Engine/Classes/Particles/Orbit/ParticleModuleOrbitBase.h b/Engine/Source/Runtime/Engine/Classes/Particles/Orbit/ParticleModuleOrbitBase.h
--- a/Engine/Source/Runtime/Engine/Classes/Particles/Orbit/ParticleModuleOrbitBase.h
+++ b/Engine/Source/Runtime/Engine/Classes/Particles/Orbit/ParticleModuleOrbitBase.h
@@ -8,7 +8,9 @@
 UCLASS(editinlinenew, hidecategories=Object, abstract, meta=(DisplayName = "Orbit"))
 class UParticleModuleOrbitBase : public UParticleModule
 {
-	GENERATED_UCLASS_BODY()
+	GENERATED_BODY()
+public:
+	UParticleModuleOrbitBase(const FObjectInitializer& ObjectInitializer = FObjectInitializer::Get());
 
 	/** 
 	 *	If true, distribution values will be retrieved using the EmitterTime.
diff --git a/Engine/Source/Runtime/Engine/Classes/Particles/Orientation/ParticleModuleOrientationAxisLock.h b/Engine/Source/Runtime/Engine/Classes/Particles/Orientation/ParticleModuleOrientationAxisLock.h
--- a/Engine/Source/Runtime/Engine/Classes/Particles/Orientation/ParticleModuleOrientationAxisLock.h
+++ b/Engine/Source/Runtime/Engine/Classes/Particles/Orientation/ParticleModuleOrientationAxisLock.h
@@ -35,7 +35,9 @@ enum EParticleAxisLock
 UCLASS(editinlinenew, hidecategories=Object, MinimalAPI, meta=(DisplayName = "Lock Axis"))
 class UParticleModuleOrientationAxisLock : public UParticleModuleOrientationBase
 {
-	GENERATED_UCLASS_BODY()
+	GENERATED_BODY()
+public:
+	ENGINE_API UParticleModuleOrientationAxisLock(const FObjectInitializer& ObjectInitializer = FObjectInitializer::Get());
 
 	/**
 	 *	The lock axis flag setting.
diff --git a/Engine/Source/Runtime/Engine/Classes/Particles/Orientation/ParticleModuleOrientationBase.h b/Engine/Source/Runtime/Engine/Classes/Particles/Orientation/ParticleModuleOrientationBase.h
--- a/Engine/Source/Runtime/Engine/Classes/Particles/Orientation/ParticleModuleOrientationBase.h
+++ b/Engine/Source/Runtime/Engine/Classes/Particles/Orientation/ParticleModuleOrientationBase.h
@@ -8,7 +8,9 @@
 UCLASS(editinlinenew, hidecategories=Object, abstract, meta=(DisplayName = "Orientation"))
 class UParticleModuleOrientationBase : public UParticleModule
 {
-	GENERATED_UCLASS_BODY()
+	GENERATED_BODY()
+public:
+	UParticleModuleOrientationBase(const FObjectInitializer& ObjectInitializer = FObjectInitializer::Get());
 
 };
 
diff --git a/Engine/Source/Runtime/Engine/Classes/Particles/Parameter/ParticleModuleParameterBase.h b/Engine/Source/Runtime/Engine/Classes/Particles/Parameter/ParticleModuleParameterBase.h
--- a/Engine/Source/Runtime/Engine/Classes/Particles/Parameter/ParticleModuleParameterBase.h
+++ b/Engine/Source/Runtime/Engine/Classes/Particles/Parameter/ParticleModuleParameterBase.h
@@ -8,7 +8,9 @@
 UCLASS(editinlinenew, hidecategories=Object, abstract, meta=(DisplayName = "Parameter"))
 class UParticleModuleParameterBase : public UParticleModule
 {
-	GENERATED_UCLASS_BODY()
+	GENERATED_BODY()
+public:
+	UParticleModuleParameterBase(const FObjectInitializer& ObjectInitializer = FObjectInitializer::Get());
 	virtual bool CanTickInAnyThread() override
 	{
 		return false;
diff --git a/Engine/Source/Runtime/Engine/Classes/Particles/Parameter/ParticleModuleParameterDynamic.h b/Engine/Source/Runtime/Engine/Classes/Particles/Parameter/ParticleModuleParameterDynamic.h
--- a/Engine/Source/Runtime/Engine/Classes/Particles/Parameter/ParticleModuleParameterDynamic.h
+++ b/Engine/Source/Runtime/Engine/Classes/Particles/Parameter/ParticleModuleParameterDynamic.h
@@ -37,7 +37,7 @@ enum EEmitterDynamicParameterValue
 USTRUCT()
 struct FEmitterDynamicParameter
 {
-	GENERATED_USTRUCT_BODY()
+	GENERATED_BODY()
 
 	/** The parameter name - from the material DynamicParameter expression. READ-ONLY */
 	UPROPERTY(Category=EmitterDynamicParameter, VisibleAnywhere, BlueprintReadWrite)
@@ -86,7 +86,9 @@ struct FEmitterDynamicParameter
 UCLASS(editinlinenew, hidecategories=Object, MinimalAPI, meta=(DisplayName = "Dynamic"))
 class UParticleModuleParameterDynamic : public UParticleModuleParameterBase
 {
-	GENERATED_UCLASS_BODY()
+	GENERATED_BODY()
+public:
+	ENGINE_API UParticleModuleParameterDynamic(const FObjectInitializer& ObjectInitializer = FObjectInitializer::Get());
 
 	/** The dynamic parameters this module uses. */
 	UPROPERTY(EditAnywhere, editfixedsize, Category=ParticleModuleParameterDynamic)
diff --git a/Engine/Source/Runtime/Engine/Classes/Particles/Parameter/ParticleModuleParameterDynamic_Seeded.h b/Engine/Source/Runtime/Engine/Classes/Particles/Parameter/ParticleModuleParameterDynamic_Seeded.h
--- a/Engine/Source/Runtime/Engine/Classes/Particles/Parameter/ParticleModuleParameterDynamic_Seeded.h
+++ b/Engine/Source/Runtime/Engine/Classes/Particles/Parameter/ParticleModuleParameterDynamic_Seeded.h
@@ -8,7 +8,9 @@
 UCLASS(editinlinenew, hidecategories=Object, meta=(DisplayName = "Dynamic (Seed)"))
 class UParticleModuleParameterDynamic_Seeded : public UParticleModuleParameterDynamic
 {
-	GENERATED_UCLASS_BODY()
+	GENERATED_BODY()
+public:
+	UParticleModuleParameterDynamic_Seeded(const FObjectInitializer& ObjectInitializer = FObjectInitializer::Get());
 
 	/** The random seed(s) to use for looking up values in StartLocation */
 	UPROPERTY(EditAnywhere, Category=RandomSeed)
diff --git a/Engine/Source/Runtime/Engine/Classes/Particles/ParticleEmitter.h b/Engine/Source/Runtime/Engine/Classes/Particles/ParticleEmitter.h
--- a/Engine/Source/Runtime/Engine/Classes/Particles/ParticleEmitter.h
+++ b/Engine/Source/Runtime/Engine/Classes/Particles/ParticleEmitter.h
@@ -55,7 +55,7 @@ enum EEmitterRenderMode
 USTRUCT()
 struct FParticleBurst
 {
-	GENERATED_USTRUCT_BODY()
+	GENERATED_BODY()
 
 	/** The number of particles to burst */
 	UPROPERTY(EditAnywhere, Category=ParticleBurst)
@@ -83,7 +83,9 @@ struct FParticleBurst
 UCLASS(hidecategories=Object, editinlinenew, abstract, MinimalAPI)
 class UParticleEmitter : public UObject
 {
-	GENERATED_UCLASS_BODY()
+	GENERATED_BODY()
+public:
+	ENGINE_API UParticleEmitter(const FObjectInitializer& ObjectInitializer = FObjectInitializer::Get());
 
 	//=============================================================================
 	//	General variables
diff --git a/Engine/Source/Runtime/Engine/Classes/Particles/ParticleEventManager.h b/Engine/Source/Runtime/Engine/Classes/Particles/ParticleEventManager.h
--- a/Engine/Source/Runtime/Engine/Classes/Particles/ParticleEventManager.h
+++ b/Engine/Source/Runtime/Engine/Classes/Particles/ParticleEventManager.h
@@ -16,7 +16,9 @@ class UParticleSystemComponent;
 UCLASS(config=Game, NotBlueprintable, notplaceable)
 class AParticleEventManager : public AActor
 {
-	GENERATED_UCLASS_BODY()
+	GENERATED_BODY()
+public:
+	AParticleEventManager(const FObjectInitializer& ObjectInitializer = FObjectInitializer::Get());
 
 	virtual void HandleParticleSpawnEvents( UParticleSystemComponent* Component, const TArray<FParticleEventSpawnData>& SpawnEvents );
 	virtual void HandleParticleDeathEvents( UParticleSystemComponent* Component, const TArray<FParticleEventDeathData>& DeathEvents );
diff --git a/Engine/Source/Runtime/Engine/Classes/Particles/ParticleLODLevel.h b/Engine/Source/Runtime/Engine/Classes/Particles/ParticleLODLevel.h
--- a/Engine/Source/Runtime/Engine/Classes/Particles/ParticleLODLevel.h
+++ b/Engine/Source/Runtime/Engine/Classes/Particles/ParticleLODLevel.h
@@ -13,7 +13,9 @@ class UInterpCurveEdSetup;
 UCLASS(collapsecategories, hidecategories=Object, editinlinenew, MinimalAPI)
 class UParticleLODLevel : public UObject
 {
-	GENERATED_UCLASS_BODY()
+	GENERATED_BODY()
+public:
+	ENGINE_API UParticleLODLevel(const FObjectInitializer& ObjectInitializer = FObjectInitializer::Get());
 
 	/** The index value of the LOD level												*/
 	UPROPERTY()
diff --git a/Engine/Source/Runtime/Engine/Classes/Particles/ParticleModule.h b/Engine/Source/Runtime/Engine/Classes/Particles/ParticleModule.h
--- a/Engine/Source/Runtime/Engine/Classes/Particles/ParticleModule.h
+++ b/Engine/Source/Runtime/Engine/Classes/Particles/ParticleModule.h
@@ -56,7 +56,7 @@ enum EParticleSourceSelectionMethod
 USTRUCT()
 struct FParticleCurvePair
 {
-	GENERATED_USTRUCT_BODY()
+	GENERATED_BODY()
 
 	UPROPERTY()
 	FString CurveName;
@@ -75,7 +75,7 @@ struct FParticleCurvePair
 USTRUCT()
 struct FParticleRandomSeedInfo
 {
-	GENERATED_USTRUCT_BODY()
+	GENERATED_BODY()
 
 	/** The name to expose to the placed instances for setting this seed */
 	UPROPERTY(EditAnywhere, Category=ParticleRandomSeedInfo)
@@ -136,7 +136,9 @@ struct FParticleRandomSeedInfo
 UCLASS(editinlinenew, hidecategories=Object, abstract, MinimalAPI)
 class UParticleModule : public UObject
 {
-	GENERATED_UCLASS_BODY()
+	GENERATED_BODY()
+public:
+	ENGINE_API UParticleModule(const FObjectInitializer& ObjectInitializer = FObjectInitializer::Get());
 
 	/** If true, the module performs operations on particles during Spawning		*/
 	UPROPERTY()
diff --git a/Engine/Source/Runtime/Engine/Classes/Particles/ParticleModuleRequired.h b/Engine/Source/Runtime/Engine/Classes/Particles/ParticleModuleRequired.h
--- a/Engine/Source/Runtime/Engine/Classes/Particles/ParticleModuleRequired.h
+++ b/Engine/Source/Runtime/Engine/Classes/Particles/ParticleModuleRequired.h
@@ -77,7 +77,9 @@ enum EEmitterNormalsMode
 UCLASS(editinlinenew, hidecategories=(Object, Cascade), meta=(DisplayName = "Required"), MinimalAPI)
 class UParticleModuleRequired : public UParticleModule
 {
-	GENERATED_UCLASS_BODY()
+	GENERATED_BODY()
+public:
+	ENGINE_API UParticleModuleRequired(const FObjectInitializer& ObjectInitializer = FObjectInitializer::Get());
 
 	//
 	// General.
diff --git a/Engine/Source/Runtime/Engine/Classes/Particles/ParticleSpriteEmitter.h b/Engine/Source/Runtime/Engine/Classes/Particles/ParticleSpriteEmitter.h
--- a/Engine/Source/Runtime/Engine/Classes/Particles/ParticleSpriteEmitter.h
+++ b/Engine/Source/Runtime/Engine/Classes/Particles/ParticleSpriteEmitter.h
@@ -20,7 +20,9 @@ enum EParticleScreenAlignment
 UCLASS(collapsecategories, hidecategories=Object, editinlinenew, MinimalAPI)
 class UParticleSpriteEmitter : public UParticleEmitter
 {
-	GENERATED_UCLASS_BODY()
+	GENERATED_BODY()
+public:
+	ENGINE_API UParticleSpriteEmitter(const FObjectInitializer& ObjectInitializer = FObjectInitializer::Get());
 
 
 	// Begin UObject Interface
diff --git a/Engine/Source/Runtime/Engine/Classes/Particles/ParticleSystem.h b/Engine/Source/Runtime/Engine/Classes/Particles/ParticleSystem.h
--- a/Engine/Source/Runtime/Engine/Classes/Particles/ParticleSystem.h
+++ b/Engine/Source/Runtime/Engine/Classes/Particles/ParticleSystem.h
@@ -47,7 +47,7 @@ enum EParticleSystemOcclusionBoundsMethod
 USTRUCT()
 struct FParticleSystemLOD
 {
-	GENERATED_USTRUCT_BODY()
+	GENERATED_BODY()
 
 	FParticleSystemLOD()
 	{
@@ -67,7 +67,7 @@ struct FParticleSystemLOD
 USTRUCT()
 struct FLODSoloTrack
 {
-	GENERATED_USTRUCT_BODY()
+	GENERATED_BODY()
 
 	UPROPERTY(transient)
 	TArray<uint8> SoloEnableSetting;
@@ -77,7 +77,7 @@ struct FLODSoloTrack
 USTRUCT()
 struct FNamedEmitterMaterial
 {
-	GENERATED_USTRUCT_BODY()
+	GENERATED_BODY()
 
 	FNamedEmitterMaterial()
 	: Name(NAME_None)
@@ -100,7 +100,9 @@ struct FNamedEmitterMaterial
 UCLASS(hidecategories=Object, MinimalAPI, BlueprintType)
 class UParticleSystem : public UObject
 {
-	GENERATED_UCLASS_BODY()
+	GENERATED_BODY()
+public:
+	ENGINE_API UParticleSystem(const FObjectInitializer& ObjectInitializer = FObjectInitializer::Get());
 
 	UPROPERTY(EditAnywhere, Category=ParticleSystem, AssetRegistrySearchable)
 	TEnumAsByte<enum EParticleSystemUpdateMode> SystemUpdateMode;
diff --git a/Engine/Source/Runtime/Engine/Classes/Particles/ParticleSystemComponent.h b/Engine/Source/Runtime/Engine/Classes/Particles/ParticleSystemComponent.h
--- a/Engine/Source/Runtime/Engine/Classes/Particles/ParticleSystemComponent.h
+++ b/Engine/Source/Runtime/Engine/Classes/Particles/ParticleSystemComponent.h
@@ -78,7 +78,7 @@ DECLARE_DYNAMIC_MULTICAST_DELEGATE_OneParam( FOnSystemFinished, class UParticleS
 USTRUCT(BlueprintType)
 struct FParticleSysParam
 {
-	GENERATED_USTRUCT_BODY()
+	GENERATED_BODY()
 
 	/** The name of the parameter */
 	UPROPERTY(EditAnywhere, BlueprintReadWrite, Category=ParticleSysParam)
@@ -250,7 +250,9 @@ struct FParticleEventKismetData : public FParticleEventData
 UCLASS(ClassGroup=(Rendering, Common), hidecategories=Object, hidecategories=Physics, hidecategories=Collision, showcategories=Trigger, editinlinenew, meta=(BlueprintSpawnableComponent))
 class ENGINE_API UParticleSystemComponent : public UPrimitiveComponent
 {
-	GENERATED_UCLASS_BODY()
+	GENERATED_BODY()
+public:
+	UParticleSystemComponent(const FObjectInitializer& ObjectInitializer = FObjectInitializer::Get());
 
 	UPROPERTY(EditAnywhere, BlueprintReadOnly, Category=Particles)
 	class UParticleSystem* Template;
diff --git a/Engine/Source/Runtime/Engine/Classes/Particles/ParticleSystemReplay.h b/Engine/Source/Runtime/Engine/Classes/Particles/ParticleSystemReplay.h
--- a/Engine/Source/Runtime/Engine/Classes/Particles/ParticleSystemReplay.h
+++ b/Engine/Source/Runtime/Engine/Classes/Particles/ParticleSystemReplay.h
@@ -7,7 +7,7 @@
 USTRUCT()
 struct FParticleEmitterReplayFrame
 {
-	GENERATED_USTRUCT_BODY()
+	GENERATED_BODY()
 
 	/** Type of emitter (EDynamicEmitterType) */
 	int32 EmitterType;
@@ -46,7 +46,7 @@ struct FParticleEmitterReplayFrame
 USTRUCT()
 struct FParticleSystemReplayFrame
 {
-	GENERATED_USTRUCT_BODY()
+	GENERATED_BODY()
 
 	/** Emitter frame state data */
 	TArray<struct FParticleEmitterReplayFrame> Emitters;
@@ -61,7 +61,9 @@ struct FParticleSystemReplayFrame
 UCLASS(hidecategories=Object, AutoExpandCategories=ParticleSystemReplay)
 class UParticleSystemReplay : public UObject
 {
-	GENERATED_UCLASS_BODY()
+	GENERATED_BODY()
+public:
+	UParticleSystemReplay(const FObjectInitializer& ObjectInitializer = FObjectInitializer::Get());
 
 	/** Unique ID number for this replay clip */
 	UPROPERTY(EditAnywhere, transient, duplicatetransient, Category=ParticleSystemReplay)
diff --git a/Engine/Source/Runtime/Engine/Classes/Particles/Rotation/ParticleModuleMeshRotation.h b/Engine/Source/Runtime/Engine/Classes/Particles/Rotation/ParticleModuleMeshRotation.h
--- a/Engine/Source/Runtime/Engine/Classes/Particles/Rotation/ParticleModuleMeshRotation.h
+++ b/Engine/Source/Runtime/Engine/Classes/Particles/Rotation/ParticleModuleMeshRotation.h
@@ -8,7 +8,9 @@
 UCLASS(editinlinenew, hidecategories=Object, MinimalAPI, meta=(DisplayName = "Init Mesh Rotation"))
 class UParticleModuleMeshRotation : public UParticleModuleRotationBase
 {
-	GENERATED_UCLASS_BODY()
+	GENERATED_BODY()
+public:
+	ENGINE_API UParticleModuleMeshRotation(const FObjectInitializer& ObjectInitializer = FObjectInitializer::Get());
 
 	/**
 	 *	Initial rotation in ROTATIONS PER SECOND (1 = 360 degrees).
diff --git a/Engine/Source/Runtime/Engine/Classes/Particles/Rotation/ParticleModuleMeshRotation_Seeded.h b/Engine/Source/Runtime/Engine/Classes/Particles/Rotation/ParticleModuleMeshRotation_Seeded.h
--- a/Engine/Source/Runtime/Engine/Classes/Particles/Rotation/ParticleModuleMeshRotation_Seeded.h
+++ b/Engine/Source/Runtime/Engine/Classes/Particles/Rotation/ParticleModuleMeshRotation_Seeded.h
@@ -8,7 +8,9 @@
 UCLASS(editinlinenew, hidecategories=Object, meta=(DisplayName = "Init Mesh Rot (Seed)"))
 class UParticleModuleMeshRotation_Seeded : public UParticleModuleMeshRotation
 {
-	GENERATED_UCLASS_BODY()
+	GENERATED_BODY()
+public:
+	UParticleModuleMeshRotation_Seeded(const FObjectInitializer& ObjectInitializer = FObjectInitializer::Get());
 
 	/** The random seed(s) to use for looking up values in StartLocation */
 	UPROPERTY(EditAnywhere, Category=RandomSeed)
diff --git a/Engine/Source/Runtime/Engine/Classes/Particles/Rotation/ParticleModuleRotation.h b/Engine/Source/Runtime/Engine/Classes/Particles/Rotation/ParticleModuleRotation.h
--- a/Engine/Source/Runtime/Engine/Classes/Particles/Rotation/ParticleModuleRotation.h
+++ b/Engine/Source/Runtime/Engine/Classes/Particles/Rotation/ParticleModuleRotation.h
@@ -8,7 +8,9 @@
 UCLASS(editinlinenew, hidecategories=Object, MinimalAPI, meta=(DisplayName = "Initial Rotation"))
 class UParticleModuleRotation : public UParticleModuleRotationBase
 {
-	GENERATED_UCLASS_BODY()
+	GENERATED_BODY()
+public:
+	ENGINE_API UParticleModuleRotation(const FObjectInitializer& ObjectInitializer = FObjectInitializer::Get());
 
 	/**
 	 *	Initial rotation of the particle (1 = 360 degrees).
diff --git a/Engine/Source/Runtime/Engine/Classes/Particles/Rotation/ParticleModuleRotationBase.h b/Engine/Source/Runtime/Engine/Classes/Particles/Rotation/ParticleModuleRotationBase.h
--- a/Engine/Source/Runtime/Engine/Classes/Particles/Rotation/ParticleModuleRotationBase.h
+++ b/Engine/Source/Runtime/Engine/Classes/Particles/Rotation/ParticleModuleRotationBase.h
@@ -8,7 +8,9 @@
 UCLASS(editinlinenew, hidecategories=Object, abstract, meta=(DisplayName = "Rotation"))
 class UParticleModuleRotationBase : public UParticleModule
 {
-	GENERATED_UCLASS_BODY()
+	GENERATED_BODY()
+public:
+	UParticleModuleRotationBase(const FObjectInitializer& ObjectInitializer = FObjectInitializer::Get());
 
 };
 
diff --git a/Engine/Source/Runtime/Engine/Classes/Particles/Rotation/ParticleModuleRotationOverLifetime.h b/Engine/Source/Runtime/Engine/Classes/Particles/Rotation/ParticleModuleRotationOverLifetime.h
--- a/Engine/Source/Runtime/Engine/Classes/Particles/Rotation/ParticleModuleRotationOverLifetime.h
+++ b/Engine/Source/Runtime/Engine/Classes/Particles/Rotation/ParticleModuleRotationOverLifetime.h
@@ -8,7 +8,9 @@
 UCLASS(editinlinenew, hidecategories=Object, MinimalAPI, meta=(DisplayName = "Rotation/Life"))
 class UParticleModuleRotationOverLifetime : public UParticleModuleRotationBase
 {
-	GENERATED_UCLASS_BODY()
+	GENERATED_BODY()
+public:
+	ENGINE_API UParticleModuleRotationOverLifetime(const FObjectInitializer& ObjectInitializer = FObjectInitializer::Get());
 
 	/** 
 	 *	The rotation of the particle (1.0 = 360 degrees).
diff --git a/Engine/Source/Runtime/Engine/Classes/Particles/Rotation/ParticleModuleRotation_Seeded.h b/Engine/Source/Runtime/Engine/Classes/Particles/Rotation/ParticleModuleRotation_Seeded.h
--- a/Engine/Source/Runtime/Engine/Classes/Particles/Rotation/ParticleModuleRotation_Seeded.h
+++ b/Engine/Source/Runtime/Engine/Classes/Particles/Rotation/ParticleModuleRotation_Seeded.h
@@ -8,7 +8,9 @@
 UCLASS(editinlinenew, hidecategories=Object, meta=(DisplayName = "Init Rotation (Seed)"))
 class UParticleModuleRotation_Seeded : public UParticleModuleRotation
 {
-	GENERATED_UCLASS_BODY()
+	GENERATED_BODY()
+public:
+	UParticleModuleRotation_Seeded(const FObjectInitializer& ObjectInitializer = FObjectInitializer::Get());
 
 	/** The random seed(s) to use for looking up values in StartLocation */
 	UPROPERTY(EditAnywhere, Category=RandomSeed)
diff --git a/Engine/Source/Runtime/Engine/Classes/Particles/RotationRate/ParticleModuleMeshRotationRate.h b/Engine/Source/Runtime/Engine/Classes/Particles/RotationRate/ParticleModuleMeshRotationRate.h
--- a/Engine/Source/Runtime/Engine/Classes/Particles/RotationRate/ParticleModuleMeshRotationRate.h
+++ b/Engine/Source/Runtime/Engine/Classes/Particles/RotationRate/ParticleModuleMeshRotationRate.h
@@ -8,7 +8,9 @@
 UCLASS(editinlinenew, hidecategories=Object, MinimalAPI, meta=(DisplayName = "Init Mesh Rotation Rate"))
 class UParticleModuleMeshRotationRate : public UParticleModuleRotationRateBase
 {
-	GENERATED_UCLASS_BODY()
+	GENERATED_BODY()
+public:
+	ENGINE_API UParticleModuleMeshRotationRate(const FObjectInitializer& ObjectInitializer = FObjectInitializer::Get());
 
 	/**
 	 *	Initial rotation rate, in rotations per second.
diff --git a/Engine/Source/Runtime/Engine/Classes/Particles/RotationRate/ParticleModuleMeshRotationRateMultiplyLife.h b/Engine/Source/Runtime/Engine/Classes/Particles/RotationRate/ParticleModuleMeshRotationRateMultiplyLife.h
--- a/Engine/Source/Runtime/Engine/Classes/Particles/RotationRate/ParticleModuleMeshRotationRateMultiplyLife.h
+++ b/Engine/Source/Runtime/Engine/Classes/Particles/RotationRate/ParticleModuleMeshRotationRateMultiplyLife.h
@@ -8,7 +8,9 @@
 UCLASS(editinlinenew, hidecategories=Object, meta=(DisplayName = "Mesh Rotation Rate * Life"))
 class UParticleModuleMeshRotationRateMultiplyLife : public UParticleModuleRotationRateBase
 {
-	GENERATED_UCLASS_BODY()
+	GENERATED_BODY()
+public:
+	UParticleModuleMeshRotationRateMultiplyLife(const FObjectInitializer& ObjectInitializer = FObjectInitializer::Get());
 
 	/**
 	 *	The scale factor that should be applied to the rotation rate.
diff --git a/Engine/Source/Runtime/Engine/Classes/Particles/RotationRate/ParticleModuleMeshRotationRateOverLife.h b/Engine/Source/Runtime/Engine/Classes/Particles/RotationRate/ParticleModuleMeshRotationRateOverLife.h
--- a/Engine/Source/Runtime/Engine/Classes/Particles/RotationRate/ParticleModuleMeshRotationRateOverLife.h
+++ b/Engine/Source/Runtime/Engine/Classes/Particles/RotationRate/ParticleModuleMeshRotationRateOverLife.h
@@ -8,7 +8,9 @@
 UCLASS(editinlinenew, hidecategories=Object, meta=(DisplayName = "Mesh Rotation Rate over Life"))
 class UParticleModuleMeshRotationRateOverLife : public UParticleModuleRotationRateBase
 {
-	GENERATED_UCLASS_BODY()
+	GENERATED_BODY()
+public:
+	UParticleModuleMeshRotationRateOverLife(const FObjectInitializer& ObjectInitializer = FObjectInitializer::Get());
 
 	/**
 	 *	The rotation rate desired.
diff --git a/Engine/Source/Runtime/Engine/Classes/Particles/RotationRate/ParticleModuleMeshRotationRate_Seeded.h b/Engine/Source/Runtime/Engine/Classes/Particles/RotationRate/ParticleModuleMeshRotationRate_Seeded.h
--- a/Engine/Source/Runtime/Engine/Classes/Particles/RotationRate/ParticleModuleMeshRotationRate_Seeded.h
+++ b/Engine/Source/Runtime/Engine/Classes/Particles/RotationRate/ParticleModuleMeshRotationRate_Seeded.h
@@ -8,7 +8,9 @@
 UCLASS(editinlinenew, hidecategories=Object, meta=(DisplayName = "Init Mesh Rotation Rate (Seed)"))
 class UParticleModuleMeshRotationRate_Seeded : public UParticleModuleMeshRotationRate
 {
-	GENERATED_UCLASS_BODY()
+	GENERATED_BODY()
+public:
+	UParticleModuleMeshRotationRate_Seeded(const FObjectInitializer& ObjectInitializer = FObjectInitializer::Get());
 
 	/** The random seed(s) to use for looking up values in StartLocation */
 	UPROPERTY(EditAnywhere, Category=RandomSeed)
diff --git a/Engine/Source/Runtime/Engine/Classes/Particles/RotationRate/ParticleModuleRotationRate.h b/Engine/Source/Runtime/Engine/Classes/Particles/RotationRate/ParticleModuleRotationRate.h
--- a/Engine/Source/Runtime/Engine/Classes/Particles/RotationRate/ParticleModuleRotationRate.h
+++ b/Engine/Source/Runtime/Engine/Classes/Particles/RotationRate/ParticleModuleRotationRate.h
@@ -8,7 +8,9 @@
 UCLASS(editinlinenew, hidecategories=Object, MinimalAPI, meta=(DisplayName = "Initial Rotation Rate"))
 class UParticleModuleRotationRate : public UParticleModuleRotationRateBase
 {
-	GENERATED_UCLASS_BODY()
+	GENERATED_BODY()
+public:
+	ENGINE_API UParticleModuleRotationRate(const FObjectInitializer& ObjectInitializer = FObjectInitializer::Get());
 
 	/** 
 	 *	Initial rotation rate, in rotations per second.
diff --git a/Engine/Source/Runtime/Engine/Classes/Particles/RotationRate/ParticleModuleRotationRateBase.h b/Engine/Source/Runtime/Engine/Classes/Particles/RotationRate/ParticleModuleRotationRateBase.h
--- a/Engine/Source/Runtime/Engine/Classes/Particles/RotationRate/ParticleModuleRotationRateBase.h
+++ b/Engine/Source/Runtime/Engine/Classes/Particles/RotationRate/ParticleModuleRotationRateBase.h
@@ -8,7 +8,9 @@
 UCLASS(editinlinenew, hidecategories=Object, abstract, meta=(DisplayName = "Rotation Rate"))
 class UParticleModuleRotationRateBase : public UParticleModule
 {
-	GENERATED_UCLASS_BODY()
+	GENERATED_BODY()
+public:
+	UParticleModuleRotationRateBase(const FObjectInitializer& ObjectInitializer = FObjectInitializer::Get());
 
 };
 
diff --git a/Engine/Source/Runtime/Engine/Classes/Particles/RotationRate/ParticleModuleRotationRateMultiplyLife.h b/Engine/Source/Runtime/Engine/Classes/Particles/RotationRate/ParticleModuleRotationRateMultiplyLife.h
--- a/Engine/Source/Runtime/Engine/Classes/Particles/RotationRate/ParticleModuleRotationRateMultiplyLife.h
+++ b/Engine/Source/Runtime/Engine/Classes/Particles/RotationRate/ParticleModuleRotationRateMultiplyLife.h
@@ -8,7 +8,9 @@
 UCLASS(editinlinenew, hidecategories=Object, MinimalAPI, meta=(DisplayName = "Rotation Rate * Life"))
 class UParticleModuleRotationRateMultiplyLife : public UParticleModuleRotationRateBase
 {
-	GENERATED_UCLASS_BODY()
+	GENERATED_BODY()
+public:
+	ENGINE_API UParticleModuleRotationRateMultiplyLife(const FObjectInitializer& ObjectInitializer = FObjectInitializer::Get());
 
 	/**
 	 *	The scale factor that should be applied to the rotation rate.
diff --git a/Engine/Source/Runtime/Engine/Classes/Particles/RotationRate/ParticleModuleRotationRate_Seeded.h b/Engine/Source/Runtime/Engine/Classes/Particles/RotationRate/ParticleModuleRotationRate_Seeded.h
--- a/Engine/Source/Runtime/Engine/Classes/Particles/RotationRate/ParticleModuleRotationRate_Seeded.h
+++ b/Engine/Source/Runtime/Engine/Classes/Particles/RotationRate/ParticleModuleRotationRate_Seeded.h
@@ -8,7 +8,9 @@
 UCLASS(editinlinenew, hidecategories=Object, meta=(DisplayName = "Init Rotation Rate (Seed)"))
 class UParticleModuleRotationRate_Seeded : public UParticleModuleRotationRate
 {
-	GENERATED_UCLASS_BODY()
+	GENERATED_BODY()
+public:
+	UParticleModuleRotationRate_Seeded(const FObjectInitializer& ObjectInitializer = FObjectInitializer::Get());
 
 	/** The random seed(s) to use for looking up values in StartLocation */
 	UPROPERTY(EditAnywhere, Category=RandomSeed)
diff --git a/Engine/Source/Runtime/Engine/Classes/Particles/Size/ParticleModuleSize.h b/Engine/Source/Runtime/Engine/Classes/Particles/Size/ParticleModuleSize.h
--- a/Engine/Source/Runtime/Engine/Classes/Particles/Size/ParticleModuleSize.h
+++ b/Engine/Source/Runtime/Engine/Classes/Particles/Size/ParticleModuleSize.h
@@ -8,7 +8,9 @@
 UCLASS(editinlinenew, hidecategories=Object, MinimalAPI, meta=(DisplayName = "Initial Size"))
 class UParticleModuleSize : public UParticleModuleSizeBase
 {
-	GENERATED_UCLASS_BODY()
+	GENERATED_BODY()
+public:
+	ENGINE_API UParticleModuleSize(const FObjectInitializer& ObjectInitializer = FObjectInitializer::Get());
 
 	/**
 	 *	The initial size that should be used for a particle.
diff --git a/Engine/Source/Runtime/Engine/Classes/Particles/Size/ParticleModuleSizeBase.h b/Engine/Source/Runtime/Engine/Classes/Particles/Size/ParticleModuleSizeBase.h
--- a/Engine/Source/Runtime/Engine/Classes/Particles/Size/ParticleModuleSizeBase.h
+++ b/Engine/Source/Runtime/Engine/Classes/Particles/Size/ParticleModuleSizeBase.h
@@ -8,7 +8,9 @@
 UCLASS(editinlinenew, hidecategories=Object, abstract, meta=(DisplayName = "Size"))
 class UParticleModuleSizeBase : public UParticleModule
 {
-	GENERATED_UCLASS_BODY()
+	GENERATED_BODY()
+public:
+	UParticleModuleSizeBase(const FObjectInitializer& ObjectInitializer = FObjectInitializer::Get());
 
 };
 
diff --git a/Engine/Source/Runtime/Engine/Classes/Particles/Size/ParticleModuleSizeMultiplyLife.h b/Engine/Source/Runtime/Engine/Classes/Particles/Size/ParticleModuleSizeMultiplyLife.h
--- a/Engine/Source/Runtime/Engine/Classes/Particles/Size/ParticleModuleSizeMultiplyLife.h
+++ b/Engine/Source/Runtime/Engine/Classes/Particles/Size/ParticleModuleSizeMultiplyLife.h
@@ -8,7 +8,9 @@
 UCLASS(editinlinenew, hidecategories=Object, MinimalAPI, meta=(DisplayName = "Size By Life"))
 class UParticleModuleSizeMultiplyLife : public UParticleModuleSizeBase
 {
-	GENERATED_UCLASS_BODY()
+	GENERATED_BODY()
+public:
+	ENGINE_API UParticleModuleSizeMultiplyLife(const FObjectInitializer& ObjectInitializer = FObjectInitializer::Get());
 
 	/**
 	 *	The scale factor for the size that should be used for a particle.
diff --git a/Engine/Source/Runtime/Engine/Classes/Particles/Size/ParticleModuleSizeScale.h b/Engine/Source/Runtime/Engine/Classes/Particles/Size/ParticleModuleSizeScale.h
--- a/Engine/Source/Runtime/Engine/Classes/Particles/Size/ParticleModuleSizeScale.h
+++ b/Engine/Source/Runtime/Engine/Classes/Particles/Size/ParticleModuleSizeScale.h
@@ -8,7 +8,9 @@
 UCLASS(editinlinenew, hidecategories=Object, MinimalAPI, meta=(DisplayName = "Size Scale"))
 class UParticleModuleSizeScale : public UParticleModuleSizeBase
 {
-	GENERATED_UCLASS_BODY()
+	GENERATED_BODY()
+public:
+	ENGINE_API UParticleModuleSizeScale(const FObjectInitializer& ObjectInitializer = FObjectInitializer::Get());
 
 	/**
 	 *	The amount the BaseSize should be scaled before being used as the size of the particle. 
diff --git a/Engine/Source/Runtime/Engine/Classes/Particles/Size/ParticleModuleSizeScaleBySpeed.h b/Engine/Source/Runtime/Engine/Classes/Particles/Size/ParticleModuleSizeScaleBySpeed.h
--- a/Engine/Source/Runtime/Engine/Classes/Particles/Size/ParticleModuleSizeScaleBySpeed.h
+++ b/Engine/Source/Runtime/Engine/Classes/Particles/Size/ParticleModuleSizeScaleBySpeed.h
@@ -11,7 +11,9 @@
 UCLASS(editinlinenew, hidecategories=Object, meta=(DisplayName = "Size By Speed"))
 class UParticleModuleSizeScaleBySpeed : public UParticleModuleSizeBase
 {
-	GENERATED_UCLASS_BODY()
+	GENERATED_BODY()
+public:
+	UParticleModuleSizeScaleBySpeed(const FObjectInitializer& ObjectInitializer = FObjectInitializer::Get());
 
 	/** By how much speed affects the size of the particle in each dimension. */
 	UPROPERTY(EditAnywhere, Category=ParticleModuleSizeScaleBySpeed)
diff --git a/Engine/Source/Runtime/Engine/Classes/Particles/Size/ParticleModuleSize_Seeded.h b/Engine/Source/Runtime/Engine/Classes/Particles/Size/ParticleModuleSize_Seeded.h
--- a/Engine/Source/Runtime/Engine/Classes/Particles/Size/ParticleModuleSize_Seeded.h
+++ b/Engine/Source/Runtime/Engine/Classes/Particles/Size/ParticleModuleSize_Seeded.h
@@ -8,7 +8,9 @@
 UCLASS(editinlinenew, hidecategories=Object, meta=(DisplayName = "Initial Size (Seed)"))
 class UParticleModuleSize_Seeded : public UParticleModuleSize
 {
-	GENERATED_UCLASS_BODY()
+	GENERATED_BODY()
+public:
+	UParticleModuleSize_Seeded(const FObjectInitializer& ObjectInitializer = FObjectInitializer::Get());
 
 	/** The random seed(s) to use for looking up values in StartLocation */
 	UPROPERTY(EditAnywhere, Category=RandomSeed)
diff --git a/Engine/Source/Runtime/Engine/Classes/Particles/Spawn/ParticleModuleSpawn.h b/Engine/Source/Runtime/Engine/Classes/Particles/Spawn/ParticleModuleSpawn.h
--- a/Engine/Source/Runtime/Engine/Classes/Particles/Spawn/ParticleModuleSpawn.h
+++ b/Engine/Source/Runtime/Engine/Classes/Particles/Spawn/ParticleModuleSpawn.h
@@ -11,7 +11,9 @@ class UParticleLODLevel;
 UCLASS(editinlinenew, hidecategories=Object, hidecategories=ParticleModuleSpawnBase, MinimalAPI, meta=(DisplayName = "Spawn"))
 class UParticleModuleSpawn : public UParticleModuleSpawnBase
 {
-	GENERATED_UCLASS_BODY()
+	GENERATED_BODY()
+public:
+	ENGINE_API UParticleModuleSpawn(const FObjectInitializer& ObjectInitializer = FObjectInitializer::Get());
 
 	/** The rate at which to spawn particles. */
 	UPROPERTY(EditAnywhere, Category=Spawn)
diff --git a/Engine/Source/Runtime/Engine/Classes/Particles/Spawn/ParticleModuleSpawnBase.h b/Engine/Source/Runtime/Engine/Classes/Particles/Spawn/ParticleModuleSpawnBase.h
--- a/Engine/Source/Runtime/Engine/Classes/Particles/Spawn/ParticleModuleSpawnBase.h
+++ b/Engine/Source/Runtime/Engine/Classes/Particles/Spawn/ParticleModuleSpawnBase.h
@@ -10,7 +10,9 @@ struct FParticleEmitterInstance;
 UCLASS(editinlinenew, hidecategories=Object, abstract, meta=(DisplayName = "Spawn"))
 class UParticleModuleSpawnBase : public UParticleModule
 {
-	GENERATED_UCLASS_BODY()
+	GENERATED_BODY()
+public:
+	UParticleModuleSpawnBase(const FObjectInitializer& ObjectInitializer = FObjectInitializer::Get());
 
 	/** 
 	 *	If true, the SpawnRate of the SpawnModule of the emitter will be processed.
diff --git a/Engine/Source/Runtime/Engine/Classes/Particles/Spawn/ParticleModuleSpawnPerUnit.h b/Engine/Source/Runtime/Engine/Classes/Particles/Spawn/ParticleModuleSpawnPerUnit.h
--- a/Engine/Source/Runtime/Engine/Classes/Particles/Spawn/ParticleModuleSpawnPerUnit.h
+++ b/Engine/Source/Runtime/Engine/Classes/Particles/Spawn/ParticleModuleSpawnPerUnit.h
@@ -8,7 +8,9 @@
 UCLASS(editinlinenew, hidecategories=Object, meta=(DisplayName = "Spawn PerUnit"))
 class UParticleModuleSpawnPerUnit : public UParticleModuleSpawnBase
 {
-	GENERATED_UCLASS_BODY()
+	GENERATED_BODY()
+public:
+	UParticleModuleSpawnPerUnit(const FObjectInitializer& ObjectInitializer = FObjectInitializer::Get());
 
 	/** 
 	 *	The scalar to apply to the distance traveled.
diff --git a/Engine/Source/Runtime/Engine/Classes/Particles/SubUV/ParticleModuleSubUV.h b/Engine/Source/Runtime/Engine/Classes/Particles/SubUV/ParticleModuleSubUV.h
--- a/Engine/Source/Runtime/Engine/Classes/Particles/SubUV/ParticleModuleSubUV.h
+++ b/Engine/Source/Runtime/Engine/Classes/Particles/SubUV/ParticleModuleSubUV.h
@@ -9,7 +9,9 @@
 UCLASS(editinlinenew, hidecategories=Object, MinimalAPI, meta=(DisplayName = "SubImage Index"))
 class UParticleModuleSubUV : public UParticleModuleSubUVBase
 {
-	GENERATED_UCLASS_BODY()
+	GENERATED_BODY()
+public:
+	ENGINE_API UParticleModuleSubUV(const FObjectInitializer& ObjectInitializer = FObjectInitializer::Get());
 
 	/**
 	 *	The index of the sub-image that should be used for the particle.
diff --git a/Engine/Source/Runtime/Engine/Classes/Particles/SubUV/ParticleModuleSubUVBase.h b/Engine/Source/Runtime/Engine/Classes/Particles/SubUV/ParticleModuleSubUVBase.h
--- a/Engine/Source/Runtime/Engine/Classes/Particles/SubUV/ParticleModuleSubUVBase.h
+++ b/Engine/Source/Runtime/Engine/Classes/Particles/SubUV/ParticleModuleSubUVBase.h
@@ -8,7 +8,9 @@
 UCLASS(editinlinenew, hidecategories=Object, abstract, meta=(DisplayName = "SubUV"))
 class UParticleModuleSubUVBase : public UParticleModule
 {
-	GENERATED_UCLASS_BODY()
+	GENERATED_BODY()
+public:
+	UParticleModuleSubUVBase(const FObjectInitializer& ObjectInitializer = FObjectInitializer::Get());
 
 };
 
diff --git a/Engine/Source/Runtime/Engine/Classes/Particles/SubUV/ParticleModuleSubUVMovie.h b/Engine/Source/Runtime/Engine/Classes/Particles/SubUV/ParticleModuleSubUVMovie.h
--- a/Engine/Source/Runtime/Engine/Classes/Particles/SubUV/ParticleModuleSubUVMovie.h
+++ b/Engine/Source/Runtime/Engine/Classes/Particles/SubUV/ParticleModuleSubUVMovie.h
@@ -8,7 +8,9 @@
 UCLASS(editinlinenew, hidecategories=Object, hidecategories=SubUV, meta=(DisplayName = "SubUV Movie"))
 class UParticleModuleSubUVMovie : public UParticleModuleSubUV
 {
-	GENERATED_UCLASS_BODY()
+	GENERATED_BODY()
+public:
+	UParticleModuleSubUVMovie(const FObjectInitializer& ObjectInitializer = FObjectInitializer::Get());
 
 	/**
 	 *	If true, use the emitter time to look up the frame rate.
diff --git a/Engine/Source/Runtime/Engine/Classes/Particles/Trail/ParticleModuleTrailBase.h b/Engine/Source/Runtime/Engine/Classes/Particles/Trail/ParticleModuleTrailBase.h
--- a/Engine/Source/Runtime/Engine/Classes/Particles/Trail/ParticleModuleTrailBase.h
+++ b/Engine/Source/Runtime/Engine/Classes/Particles/Trail/ParticleModuleTrailBase.h
@@ -13,7 +13,9 @@
 UCLASS(editinlinenew, hidecategories=Object, abstract, meta=(DisplayName = "Trail"))
 class UParticleModuleTrailBase : public UParticleModule
 {
-	GENERATED_UCLASS_BODY()
+	GENERATED_BODY()
+public:
+	UParticleModuleTrailBase(const FObjectInitializer& ObjectInitializer = FObjectInitializer::Get());
 
 
 	// Begin UParticleModule Interface
diff --git a/Engine/Source/Runtime/Engine/Classes/Particles/Trail/ParticleModuleTrailSource.h b/Engine/Source/Runtime/Engine/Classes/Particles/Trail/ParticleModuleTrailSource.h
--- a/Engine/Source/Runtime/Engine/Classes/Particles/Trail/ParticleModuleTrailSource.h
+++ b/Engine/Source/Runtime/Engine/Classes/Particles/Trail/ParticleModuleTrailSource.h
@@ -32,7 +32,9 @@ enum ETrail2SourceMethod
 UCLASS(editinlinenew, hidecategories=Object, MinimalAPI, meta=(DisplayName = "Source"))
 class UParticleModuleTrailSource : public UParticleModuleTrailBase
 {
-	GENERATED_UCLASS_BODY()
+	GENERATED_BODY()
+public:
+	ENGINE_API UParticleModuleTrailSource(const FObjectInitializer& ObjectInitializer = FObjectInitializer::Get());
 
 	/** The source method for the trail. */
 	UPROPERTY(EditAnywhere, Category=Source)
diff --git a/Engine/Source/Runtime/Engine/Classes/Particles/TypeData/ParticleModuleTypeDataAnimTrail.h b/Engine/Source/Runtime/Engine/Classes/Particles/TypeData/ParticleModuleTypeDataAnimTrail.h
--- a/Engine/Source/Runtime/Engine/Classes/Particles/TypeData/ParticleModuleTypeDataAnimTrail.h
+++ b/Engine/Source/Runtime/Engine/Classes/Particles/TypeData/ParticleModuleTypeDataAnimTrail.h
@@ -13,7 +13,9 @@
 UCLASS(editinlinenew, hidecategories=Object, meta=(DisplayName = "AnimTrail Data"))
 class UParticleModuleTypeDataAnimTrail : public UParticleModuleTypeDataBase
 {
-	GENERATED_UCLASS_BODY()
+	GENERATED_BODY()
+public:
+	UParticleModuleTypeDataAnimTrail(const FObjectInitializer& ObjectInitializer = FObjectInitializer::Get());
 
 	//*****************************************************************************
 	// General Trail Variables
diff --git a/Engine/Source/Runtime/Engine/Classes/Particles/TypeData/ParticleModuleTypeDataBase.h b/Engine/Source/Runtime/Engine/Classes/Particles/TypeData/ParticleModuleTypeDataBase.h
--- a/Engine/Source/Runtime/Engine/Classes/Particles/TypeData/ParticleModuleTypeDataBase.h
+++ b/Engine/Source/Runtime/Engine/Classes/Particles/TypeData/ParticleModuleTypeDataBase.h
@@ -10,7 +10,9 @@ struct FParticleEmitterInstance;
 UCLASS(editinlinenew, hidecategories=Object, abstract, MinimalAPI)
 class UParticleModuleTypeDataBase : public UParticleModule
 {
-	GENERATED_UCLASS_BODY()
+	GENERATED_BODY()
+public:
+	ENGINE_API UParticleModuleTypeDataBase(const FObjectInitializer& ObjectInitializer = FObjectInitializer::Get());
 
 
 	// Begin UParticleModule Interface
diff --git a/Engine/Source/Runtime/Engine/Classes/Particles/TypeData/ParticleModuleTypeDataBeam2.h b/Engine/Source/Runtime/Engine/Classes/Particles/TypeData/ParticleModuleTypeDataBeam2.h
--- a/Engine/Source/Runtime/Engine/Classes/Particles/TypeData/ParticleModuleTypeDataBeam2.h
+++ b/Engine/Source/Runtime/Engine/Classes/Particles/TypeData/ParticleModuleTypeDataBeam2.h
@@ -34,7 +34,7 @@ enum EBeamTaperMethod
 USTRUCT()
 struct FBeamTargetData
 {
-	GENERATED_USTRUCT_BODY()
+	GENERATED_BODY()
 
 	/** Name of the target.																	*/
 	UPROPERTY(EditAnywhere, Category=BeamTargetData)
@@ -55,7 +55,9 @@ struct FBeamTargetData
 UCLASS(editinlinenew, dontcollapsecategories, hidecategories=Object, MinimalAPI, meta=(DisplayName = "Beam Data"))
 class UParticleModuleTypeDataBeam2 : public UParticleModuleTypeDataBase
 {
-	GENERATED_UCLASS_BODY()
+	GENERATED_BODY()
+public:
+	ENGINE_API UParticleModuleTypeDataBeam2(const FObjectInitializer& ObjectInitializer = FObjectInitializer::Get());
 
 	//
 	// General Beam Variables.
diff --git a/Engine/Source/Runtime/Engine/Classes/Particles/TypeData/ParticleModuleTypeDataGpu.h b/Engine/Source/Runtime/Engine/Classes/Particles/TypeData/ParticleModuleTypeDataGpu.h
--- a/Engine/Source/Runtime/Engine/Classes/Particles/TypeData/ParticleModuleTypeDataGpu.h
+++ b/Engine/Source/Runtime/Engine/Classes/Particles/TypeData/ParticleModuleTypeDataGpu.h
@@ -18,7 +18,7 @@
 USTRUCT()
 struct FGPUSpriteLocalVectorFieldInfo
 {
-	GENERATED_USTRUCT_BODY()
+	GENERATED_BODY()
 
 	/** Local vector field to apply to this emitter. */
 	UPROPERTY()
@@ -84,7 +84,7 @@ struct FGPUSpriteLocalVectorFieldInfo
 USTRUCT()
 struct FGPUSpriteEmitterInfo
 {
-	GENERATED_USTRUCT_BODY()
+	GENERATED_BODY()
 
 	/** The required module. Needed for now, but should be divorced from the runtime. */
 	UPROPERTY()
@@ -215,7 +215,7 @@ struct FGPUSpriteEmitterInfo
 USTRUCT()
 struct FGPUSpriteResourceData
 {
-	GENERATED_USTRUCT_BODY()
+	GENERATED_BODY()
 
 	/** Quantized color samples. */
 	UPROPERTY()
@@ -392,7 +392,9 @@ struct FGPUSpriteResourceData
 UCLASS(editinlinenew, hidecategories=Object, MinimalAPI, meta=(DisplayName = "GPU Sprites"))
 class UParticleModuleTypeDataGpu : public UParticleModuleTypeDataBase
 {
-	GENERATED_UCLASS_BODY()
+	GENERATED_BODY()
+public:
+	ENGINE_API UParticleModuleTypeDataGpu(const FObjectInitializer& ObjectInitializer = FObjectInitializer::Get());
 
 	/** Information for runtime simulation. */
 	UPROPERTY(transient)
diff --git a/Engine/Source/Runtime/Engine/Classes/Particles/TypeData/ParticleModuleTypeDataMesh.h b/Engine/Source/Runtime/Engine/Classes/Particles/TypeData/ParticleModuleTypeDataMesh.h
--- a/Engine/Source/Runtime/Engine/Classes/Particles/TypeData/ParticleModuleTypeDataMesh.h
+++ b/Engine/Source/Runtime/Engine/Classes/Particles/TypeData/ParticleModuleTypeDataMesh.h
@@ -53,7 +53,9 @@ enum EMeshCameraFacingOptions
 UCLASS(editinlinenew, hidecategories=Object, MinimalAPI, meta=(DisplayName = "Mesh Data"))
 class UParticleModuleTypeDataMesh : public UParticleModuleTypeDataBase
 {
-	GENERATED_UCLASS_BODY()
+	GENERATED_BODY()
+public:
+	ENGINE_API UParticleModuleTypeDataMesh(const FObjectInitializer& ObjectInitializer = FObjectInitializer::Get());
 
 	/** The static mesh to render at the particle positions */
 	UPROPERTY(EditAnywhere, Category=Mesh)
diff --git a/Engine/Source/Runtime/Engine/Classes/Particles/TypeData/ParticleModuleTypeDataRibbon.h b/Engine/Source/Runtime/Engine/Classes/Particles/TypeData/ParticleModuleTypeDataRibbon.h
--- a/Engine/Source/Runtime/Engine/Classes/Particles/TypeData/ParticleModuleTypeDataRibbon.h
+++ b/Engine/Source/Runtime/Engine/Classes/Particles/TypeData/ParticleModuleTypeDataRibbon.h
@@ -22,7 +22,9 @@ enum ETrailsRenderAxisOption
 UCLASS(editinlinenew, hidecategories=Object, meta=(DisplayName = "Ribbon Data"))
 class UParticleModuleTypeDataRibbon : public UParticleModuleTypeDataBase
 {
-	GENERATED_UCLASS_BODY()
+	GENERATED_BODY()
+public:
+	UParticleModuleTypeDataRibbon(const FObjectInitializer& ObjectInitializer = FObjectInitializer::Get());
 
 	//
 	// General Trail Variables.
diff --git a/Engine/Source/Runtime/Engine/Classes/Particles/VectorField/ParticleModuleVectorFieldBase.h b/Engine/Source/Runtime/Engine/Classes/Particles/VectorField/ParticleModuleVectorFieldBase.h
--- a/Engine/Source/Runtime/Engine/Classes/Particles/VectorField/ParticleModuleVectorFieldBase.h
+++ b/Engine/Source/Runtime/Engine/Classes/Particles/VectorField/ParticleModuleVectorFieldBase.h
@@ -12,7 +12,9 @@
 UCLASS(editinlinenew, hidecategories=Object, abstract, meta=(DisplayName = "Vector Field"))
 class UParticleModuleVectorFieldBase : public UParticleModule
 {
-	GENERATED_UCLASS_BODY()
+	GENERATED_BODY()
+public:
+	UParticleModuleVectorFieldBase(const FObjectInitializer& ObjectInitializer = FObjectInitializer::Get());
 
 };
 
diff --git a/Engine/Source/Runtime/Engine/Classes/Particles/VectorField/ParticleModuleVectorFieldGlobal.h b/Engine/Source/Runtime/Engine/Classes/Particles/VectorField/ParticleModuleVectorFieldGlobal.h
--- a/Engine/Source/Runtime/Engine/Classes/Particles/VectorField/ParticleModuleVectorFieldGlobal.h
+++ b/Engine/Source/Runtime/Engine/Classes/Particles/VectorField/ParticleModuleVectorFieldGlobal.h
@@ -11,7 +11,9 @@
 UCLASS(editinlinenew, hidecategories=Object, meta=(DisplayName = "Global Vector Fields"))
 class UParticleModuleVectorFieldGlobal : public UParticleModuleVectorFieldBase
 {
-	GENERATED_UCLASS_BODY()
+	GENERATED_BODY()
+public:
+	UParticleModuleVectorFieldGlobal(const FObjectInitializer& ObjectInitializer = FObjectInitializer::Get());
 
 	/** Property override value for global vector field tightness.  */
 	UPROPERTY()
diff --git a/Engine/Source/Runtime/Engine/Classes/Particles/VectorField/ParticleModuleVectorFieldLocal.h b/Engine/Source/Runtime/Engine/Classes/Particles/VectorField/ParticleModuleVectorFieldLocal.h
--- a/Engine/Source/Runtime/Engine/Classes/Particles/VectorField/ParticleModuleVectorFieldLocal.h
+++ b/Engine/Source/Runtime/Engine/Classes/Particles/VectorField/ParticleModuleVectorFieldLocal.h
@@ -11,7 +11,9 @@
 UCLASS(editinlinenew, hidecategories=Object, MinimalAPI, meta=(DisplayName = "Local Vector Field"))
 class UParticleModuleVectorFieldLocal : public UParticleModuleVectorFieldBase
 {
-	GENERATED_UCLASS_BODY()
+	GENERATED_BODY()
+public:
+	ENGINE_API UParticleModuleVectorFieldLocal(const FObjectInitializer& ObjectInitializer = FObjectInitializer::Get());
 
 	/** Vector field asset to use. */
 	UPROPERTY(EditAnywhere, Category=VectorField)
diff --git a/Engine/Source/Runtime/Engine/Classes/Particles/VectorField/ParticleModuleVectorFieldRotation.h b/Engine/Source/Runtime/Engine/Classes/Particles/VectorField/ParticleModuleVectorFieldRotation.h
--- a/Engine/Source/Runtime/Engine/Classes/Particles/VectorField/ParticleModuleVectorFieldRotation.h
+++ b/Engine/Source/Runtime/Engine/Classes/Particles/VectorField/ParticleModuleVectorFieldRotation.h
@@ -12,7 +12,9 @@
 UCLASS(editinlinenew, hidecategories=Object, meta=(DisplayName = "VF Init Rotation"))
 class UParticleModuleVectorFieldRotation : public UParticleModuleVectorFieldBase
 {
-	GENERATED_UCLASS_BODY()
+	GENERATED_BODY()
+public:
+	UParticleModuleVectorFieldRotation(const FObjectInitializer& ObjectInitializer = FObjectInitializer::Get());
 
 	/** Minimum initial rotation applied to the local vector field. */
 	UPROPERTY(EditAnywhere, Category=VectorField)
diff --git a/Engine/Source/Runtime/Engine/Classes/Particles/VectorField/ParticleModuleVectorFieldRotationRate.h b/Engine/Source/Runtime/Engine/Classes/Particles/VectorField/ParticleModuleVectorFieldRotationRate.h
--- a/Engine/Source/Runtime/Engine/Classes/Particles/VectorField/ParticleModuleVectorFieldRotationRate.h
+++ b/Engine/Source/Runtime/Engine/Classes/Particles/VectorField/ParticleModuleVectorFieldRotationRate.h
@@ -11,7 +11,9 @@
 UCLASS(editinlinenew, hidecategories=Object, meta=(DisplayName = "VF Rotation Rate"))
 class UParticleModuleVectorFieldRotationRate : public UParticleModuleVectorFieldBase
 {
-	GENERATED_UCLASS_BODY()
+	GENERATED_BODY()
+public:
+	UParticleModuleVectorFieldRotationRate(const FObjectInitializer& ObjectInitializer = FObjectInitializer::Get());
 
 	/** Constant rotation rate applied to the local vector field. */
 	UPROPERTY(EditAnywhere, Category=VectorField)
diff --git a/Engine/Source/Runtime/Engine/Classes/Particles/VectorField/ParticleModuleVectorFieldScale.h b/Engine/Source/Runtime/Engine/Classes/Particles/VectorField/ParticleModuleVectorFieldScale.h
--- a/Engine/Source/Runtime/Engine/Classes/Particles/VectorField/ParticleModuleVectorFieldScale.h
+++ b/Engine/Source/Runtime/Engine/Classes/Particles/VectorField/ParticleModuleVectorFieldScale.h
@@ -11,7 +11,9 @@
 UCLASS(editinlinenew, hidecategories=Object, MinimalAPI, meta=(DisplayName = "Vector Field Scale"))
 class UParticleModuleVectorFieldScale : public UParticleModuleVectorFieldBase
 {
-	GENERATED_UCLASS_BODY()
+	GENERATED_BODY()
+public:
+	ENGINE_API UParticleModuleVectorFieldScale(const FObjectInitializer& ObjectInitializer = FObjectInitializer::Get());
 
 	/** Per-particle vector field scale. Evaluated using emitter time. */
 	UPROPERTY(EditAnywhere, Category=VectorField)
diff --git a/Engine/Source/Runtime/Engine/Classes/Particles/VectorField/ParticleModuleVectorFieldScaleOverLife.h b/Engine/Source/Runtime/Engine/Classes/Particles/VectorField/ParticleModuleVectorFieldScaleOverLife.h
--- a/Engine/Source/Runtime/Engine/Classes/Particles/VectorField/ParticleModuleVectorFieldScaleOverLife.h
+++ b/Engine/Source/Runtime/Engine/Classes/Particles/VectorField/ParticleModuleVectorFieldScaleOverLife.h
@@ -11,7 +11,9 @@
 UCLASS(editinlinenew, hidecategories=Object, MinimalAPI, meta=(DisplayName = "VF Scale/Life"))
 class UParticleModuleVectorFieldScaleOverLife : public UParticleModuleVectorFieldBase
 {
-	GENERATED_UCLASS_BODY()
+	GENERATED_BODY()
+public:
+	ENGINE_API UParticleModuleVectorFieldScaleOverLife(const FObjectInitializer& ObjectInitializer = FObjectInitializer::Get());
 
 	/** Per-particle vector field scale. Evaluated using particle relative time. */
 	UPROPERTY(EditAnywhere, Category=VectorField)
diff --git a/Engine/Source/Runtime/Engine/Classes/Particles/Velocity/ParticleModuleVelocity.h b/Engine/Source/Runtime/Engine/Classes/Particles/Velocity/ParticleModuleVelocity.h
--- a/Engine/Source/Runtime/Engine/Classes/Particles/Velocity/ParticleModuleVelocity.h
+++ b/Engine/Source/Runtime/Engine/Classes/Particles/Velocity/ParticleModuleVelocity.h
@@ -8,7 +8,9 @@
 UCLASS(editinlinenew, hidecategories=Object, MinimalAPI, meta=(DisplayName = "Initial Velocity"))
 class UParticleModuleVelocity : public UParticleModuleVelocityBase
 {
-	GENERATED_UCLASS_BODY()
+	GENERATED_BODY()
+public:
+	ENGINE_API UParticleModuleVelocity(const FObjectInitializer& ObjectInitializer = FObjectInitializer::Get());
 
 	/** 
 	 *	The velocity to apply to a particle when it is spawned.
diff --git a/Engine/Source/Runtime/Engine/Classes/Particles/Velocity/ParticleModuleVelocityBase.h b/Engine/Source/Runtime/Engine/Classes/Particles/Velocity/ParticleModuleVelocityBase.h
--- a/Engine/Source/Runtime/Engine/Classes/Particles/Velocity/ParticleModuleVelocityBase.h
+++ b/Engine/Source/Runtime/Engine/Classes/Particles/Velocity/ParticleModuleVelocityBase.h
@@ -8,7 +8,9 @@
 UCLASS(editinlinenew, hidecategories=Object, abstract, meta=(DisplayName = "Velocity"))
 class UParticleModuleVelocityBase : public UParticleModule
 {
-	GENERATED_UCLASS_BODY()
+	GENERATED_BODY()
+public:
+	UParticleModuleVelocityBase(const FObjectInitializer& ObjectInitializer = FObjectInitializer::Get());
 
 	/**
 	 *	If true, then treat the velocity as world-space defined.
diff --git a/Engine/Source/Runtime/Engine/Classes/Particles/Velocity/ParticleModuleVelocityCone.h b/Engine/Source/Runtime/Engine/Classes/Particles/Velocity/ParticleModuleVelocityCone.h
--- a/Engine/Source/Runtime/Engine/Classes/Particles/Velocity/ParticleModuleVelocityCone.h
+++ b/Engine/Source/Runtime/Engine/Classes/Particles/Velocity/ParticleModuleVelocityCone.h
@@ -8,7 +8,9 @@
 UCLASS(editinlinenew, hidecategories=Object, MinimalAPI, meta=( DisplayName = "Velocity Cone" ))
 class UParticleModuleVelocityCone : public UParticleModuleVelocityBase
 {
-	GENERATED_UCLASS_BODY()
+	GENERATED_BODY()
+public:
+	ENGINE_API UParticleModuleVelocityCone(const FObjectInitializer& ObjectInitializer = FObjectInitializer::Get());
 
 	/** The Min value represents the inner cone angle value and the Max value represents the outer cone angle value. */
 	UPROPERTY(EditAnywhere, Category=Cone)
diff --git a/Engine/Source/Runtime/Engine/Classes/Particles/Velocity/ParticleModuleVelocityInheritParent.h b/Engine/Source/Runtime/Engine/Classes/Particles/Velocity/ParticleModuleVelocityInheritParent.h
--- a/Engine/Source/Runtime/Engine/Classes/Particles/Velocity/ParticleModuleVelocityInheritParent.h
+++ b/Engine/Source/Runtime/Engine/Classes/Particles/Velocity/ParticleModuleVelocityInheritParent.h
@@ -8,7 +8,9 @@
 UCLASS(editinlinenew, hidecategories=Object, MinimalAPI, meta=(DisplayName = "Inherit Parent Velocity"))
 class UParticleModuleVelocityInheritParent : public UParticleModuleVelocityBase
 {
-	GENERATED_UCLASS_BODY()
+	GENERATED_BODY()
+public:
+	ENGINE_API UParticleModuleVelocityInheritParent(const FObjectInitializer& ObjectInitializer = FObjectInitializer::Get());
 
 	/**
 	 *	The scale to apply tot he parent velocity prior to adding it to the particle velocity during spawn.
diff --git a/Engine/Source/Runtime/Engine/Classes/Particles/Velocity/ParticleModuleVelocityOverLifetime.h b/Engine/Source/Runtime/Engine/Classes/Particles/Velocity/ParticleModuleVelocityOverLifetime.h
--- a/Engine/Source/Runtime/Engine/Classes/Particles/Velocity/ParticleModuleVelocityOverLifetime.h
+++ b/Engine/Source/Runtime/Engine/Classes/Particles/Velocity/ParticleModuleVelocityOverLifetime.h
@@ -8,7 +8,9 @@
 UCLASS(editinlinenew, hidecategories=Object, MinimalAPI, meta=(DisplayName = "Velocity/Life"))
 class UParticleModuleVelocityOverLifetime : public UParticleModuleVelocityBase
 {
-	GENERATED_UCLASS_BODY()
+	GENERATED_BODY()
+public:
+	ENGINE_API UParticleModuleVelocityOverLifetime(const FObjectInitializer& ObjectInitializer = FObjectInitializer::Get());
 
 	/**
 	 *	The scaling  value applied to the velocity.
diff --git a/Engine/Source/Runtime/Engine/Classes/Particles/Velocity/ParticleModuleVelocity_Seeded.h b/Engine/Source/Runtime/Engine/Classes/Particles/Velocity/ParticleModuleVelocity_Seeded.h
--- a/Engine/Source/Runtime/Engine/Classes/Particles/Velocity/ParticleModuleVelocity_Seeded.h
+++ b/Engine/Source/Runtime/Engine/Classes/Particles/Velocity/ParticleModuleVelocity_Seeded.h
@@ -8,7 +8,9 @@
 UCLASS(editinlinenew, hidecategories=Object, meta=(DisplayName = "Initial Velocity (Seed)"))
 class UParticleModuleVelocity_Seeded : public UParticleModuleVelocity
 {
-	GENERATED_UCLASS_BODY()
+	GENERATED_BODY()
+public:
+	UParticleModuleVelocity_Seeded(const FObjectInitializer& ObjectInitializer = FObjectInitializer::Get());
 
 	/** The random seed(s) to use for looking up values in StartLocation */
 	UPROPERTY(EditAnywhere, Category=RandomSeed)
diff --git a/Engine/Source/Runtime/Engine/Classes/PhysicalMaterials/PhysicalMaterial.h b/Engine/Source/Runtime/Engine/Classes/PhysicalMaterials/PhysicalMaterial.h
--- a/Engine/Source/Runtime/Engine/Classes/PhysicalMaterials/PhysicalMaterial.h
+++ b/Engine/Source/Runtime/Engine/Classes/PhysicalMaterials/PhysicalMaterial.h
@@ -13,7 +13,7 @@ namespace physx
 USTRUCT()
 struct FTireFrictionScalePair
 {
-	GENERATED_USTRUCT_BODY()
+	GENERATED_BODY()
 
 	/** Tire type */
 	UPROPERTY(EditAnywhere, Category=TireFrictionScalePair)
@@ -36,7 +36,9 @@ struct FTireFrictionScalePair
 UCLASS(collapsecategories, hidecategories=Object)
 class ENGINE_API UPhysicalMaterial : public UObject
 {
-	GENERATED_UCLASS_BODY()
+	GENERATED_BODY()
+public:
+	UPhysicalMaterial(const FObjectInitializer& ObjectInitializer = FObjectInitializer::Get());
 
 	//
 	// Surface properties.
diff --git a/Engine/Source/Runtime/Engine/Classes/PhysicalMaterials/PhysicalMaterialPropertyBase.h b/Engine/Source/Runtime/Engine/Classes/PhysicalMaterials/PhysicalMaterialPropertyBase.h
--- a/Engine/Source/Runtime/Engine/Classes/PhysicalMaterials/PhysicalMaterialPropertyBase.h
+++ b/Engine/Source/Runtime/Engine/Classes/PhysicalMaterials/PhysicalMaterialPropertyBase.h
@@ -26,7 +26,9 @@
 UCLASS(collapsecategories, hidecategories=Object, editinlinenew, abstract,MinimalAPI, deprecated)
 class UDEPRECATED_PhysicalMaterialPropertyBase : public UObject
 {
-	GENERATED_UCLASS_BODY()
+	GENERATED_BODY()
+public:
+	ENGINE_API UDEPRECATED_PhysicalMaterialPropertyBase(const FObjectInitializer& ObjectInitializer = FObjectInitializer::Get());
 
 	// Begin UObject Interface
 	virtual bool IsAsset() const override { return false; }
diff --git a/Engine/Source/Runtime/Engine/Classes/PhysicsEngine/AggregateGeometry2D.h b/Engine/Source/Runtime/Engine/Classes/PhysicsEngine/AggregateGeometry2D.h
--- a/Engine/Source/Runtime/Engine/Classes/PhysicsEngine/AggregateGeometry2D.h
+++ b/Engine/Source/Runtime/Engine/Classes/PhysicsEngine/AggregateGeometry2D.h
@@ -9,7 +9,7 @@
 USTRUCT()
 struct ENGINE_API FConvexElement2D
 {
-	GENERATED_USTRUCT_BODY()
+	GENERATED_BODY()
 
 	// Vertices that lie on the convex hull
 	UPROPERTY()
@@ -30,7 +30,7 @@ struct ENGINE_API FConvexElement2D
 USTRUCT()
 struct ENGINE_API FCircleElement2D
 {
-	GENERATED_USTRUCT_BODY()
+	GENERATED_BODY()
 
 	UPROPERTY(Category=KSphereElem, VisibleAnywhere)
 	FVector2D Center;
@@ -68,7 +68,7 @@ public:
 USTRUCT()
 struct ENGINE_API FBoxElement2D
 {
-	GENERATED_USTRUCT_BODY()
+	GENERATED_BODY()
 
 	// Center of the box
 	UPROPERTY(Category=KBoxElem, VisibleAnywhere)
@@ -114,7 +114,7 @@ struct ENGINE_API FBoxElement2D
 USTRUCT()
 struct ENGINE_API FAggregateGeometry2D
 {
-	GENERATED_USTRUCT_BODY()
+	GENERATED_BODY()
 
 public:
 	UPROPERTY(EditAnywhere, editfixedsize, Category = KAggregateGeom)
diff --git a/Engine/Source/Runtime/Engine/Classes/PhysicsEngine/BodyInstance.h b/Engine/Source/Runtime/Engine/Classes/PhysicsEngine/BodyInstance.h
--- a/Engine/Source/Runtime/Engine/Classes/PhysicsEngine/BodyInstance.h
+++ b/Engine/Source/Runtime/Engine/Classes/PhysicsEngine/BodyInstance.h
@@ -64,7 +64,7 @@ struct FCollisionNotifyInfo;
 USTRUCT()
 struct ENGINE_API FCollisionResponse
 {
-	GENERATED_USTRUCT_BODY()
+	GENERATED_BODY()
 
 	FCollisionResponse();
 	FCollisionResponse(ECollisionResponse DefaultResponse);
@@ -126,7 +126,7 @@ enum class BodyInstanceSceneState
 USTRUCT()
 struct ENGINE_API FBodyInstance
 {
-	GENERATED_USTRUCT_BODY()
+	GENERATED_BODY()
 
 	/** 
 	 *	Index of this BodyInstance within the SkeletalMeshComponent/PhysicsAsset. 
@@ -862,4 +862,4 @@ struct FBodyInstanceEditorHelpers
 private:
 	FBodyInstanceEditorHelpers() {}
 };
-#endif
\ No newline at end of file
+#endif
diff --git a/Engine/Source/Runtime/Engine/Classes/PhysicsEngine/BodySetup.h b/Engine/Source/Runtime/Engine/Classes/PhysicsEngine/BodySetup.h
--- a/Engine/Source/Runtime/Engine/Classes/PhysicsEngine/BodySetup.h
+++ b/Engine/Source/Runtime/Engine/Classes/PhysicsEngine/BodySetup.h
@@ -24,7 +24,9 @@ namespace physx
 UCLASS(hidecategories=Object, MinimalAPI)
 class UBodySetup : public UObject
 {
-	GENERATED_UCLASS_BODY()
+	GENERATED_BODY()
+public:
+	ENGINE_API UBodySetup(const FObjectInitializer& ObjectInitializer = FObjectInitializer::Get());
 
 	/** Simplified collision representation of this  */
 	UPROPERTY()
diff --git a/Engine/Source/Runtime/Engine/Classes/PhysicsEngine/BodySetup2D.h b/Engine/Source/Runtime/Engine/Classes/PhysicsEngine/BodySetup2D.h
--- a/Engine/Source/Runtime/Engine/Classes/PhysicsEngine/BodySetup2D.h
+++ b/Engine/Source/Runtime/Engine/Classes/PhysicsEngine/BodySetup2D.h
@@ -8,7 +8,9 @@
 UCLASS()
 class ENGINE_API UBodySetup2D : public UBodySetup
 {
-	GENERATED_UCLASS_BODY()
+	GENERATED_BODY()
+public:
+	UBodySetup2D(const FObjectInitializer& ObjectInitializer = FObjectInitializer::Get());
 
 	UPROPERTY()
 	FAggregateGeometry2D AggGeom2D;
diff --git a/Engine/Source/Runtime/Engine/Classes/PhysicsEngine/DestructibleActor.h b/Engine/Source/Runtime/Engine/Classes/PhysicsEngine/DestructibleActor.h
--- a/Engine/Source/Runtime/Engine/Classes/PhysicsEngine/DestructibleActor.h
+++ b/Engine/Source/Runtime/Engine/Classes/PhysicsEngine/DestructibleActor.h
@@ -14,7 +14,9 @@ DECLARE_DYNAMIC_MULTICAST_DELEGATE_TwoParams(FActorFractureSignature, const FVec
 UCLASS(MinimalAPI, hideCategories=(Input), showCategories=("Input|MouseInput", "Input|TouchInput"), ComponentWrapperClass)
 class ADestructibleActor : public AActor
 {
-	GENERATED_UCLASS_BODY()
+	GENERATED_BODY()
+public:
+	ENGINE_API ADestructibleActor(const FObjectInitializer& ObjectInitializer = FObjectInitializer::Get());
 
 	/**
 	 * The component which holds the skinned mesh and physics data for this actor.
diff --git a/Engine/Source/Runtime/Engine/Classes/PhysicsEngine/PhysicsAsset.h b/Engine/Source/Runtime/Engine/Classes/PhysicsEngine/PhysicsAsset.h
--- a/Engine/Source/Runtime/Engine/Classes/PhysicsEngine/PhysicsAsset.h
+++ b/Engine/Source/Runtime/Engine/Classes/PhysicsEngine/PhysicsAsset.h
@@ -18,7 +18,9 @@
 UCLASS(hidecategories=Object, BlueprintType, MinimalAPI)
 class UPhysicsAsset : public UObject
 {
-	GENERATED_UCLASS_BODY()
+	GENERATED_BODY()
+public:
+	ENGINE_API UPhysicsAsset(const FObjectInitializer& ObjectInitializer = FObjectInitializer::Get());
 
 #if WITH_EDITORONLY_DATA
 	/** 
@@ -124,4 +126,4 @@ public:
 };
 
 #define		RB_MinSizeToLockDOF				(0.1)
-#define		RB_MinAngleToLockDOF			(5.0)
\ No newline at end of file
+#define		RB_MinAngleToLockDOF			(5.0)
diff --git a/Engine/Source/Runtime/Engine/Classes/PhysicsEngine/PhysicsCollisionHandler.h b/Engine/Source/Runtime/Engine/Classes/PhysicsEngine/PhysicsCollisionHandler.h
--- a/Engine/Source/Runtime/Engine/Classes/PhysicsEngine/PhysicsCollisionHandler.h
+++ b/Engine/Source/Runtime/Engine/Classes/PhysicsEngine/PhysicsCollisionHandler.h
@@ -8,7 +8,9 @@ struct FCollisionNotifyInfo;
 UCLASS(hidecategories=(Object), Blueprintable)
 class ENGINE_API UPhysicsCollisionHandler : public UObject
 {
-	GENERATED_UCLASS_BODY()
+	GENERATED_BODY()
+public:
+	UPhysicsCollisionHandler(const FObjectInitializer& ObjectInitializer = FObjectInitializer::Get());
 
 	// This class is SUPER BASIC, but games can implement their own handlers to do more advanced physics collisions
 
diff --git a/Engine/Source/Runtime/Engine/Classes/PhysicsEngine/PhysicsConstraintActor.h b/Engine/Source/Runtime/Engine/Classes/PhysicsEngine/PhysicsConstraintActor.h
--- a/Engine/Source/Runtime/Engine/Classes/PhysicsEngine/PhysicsConstraintActor.h
+++ b/Engine/Source/Runtime/Engine/Classes/PhysicsEngine/PhysicsConstraintActor.h
@@ -10,7 +10,9 @@
 UCLASS(ConversionRoot, MinimalAPI, ComponentWrapperClass)
 class APhysicsConstraintActor : public ARigidBodyBase
 {
-	GENERATED_UCLASS_BODY()
+	GENERATED_BODY()
+public:
+	ENGINE_API APhysicsConstraintActor(const FObjectInitializer& ObjectInitializer = FObjectInitializer::Get());
 
 	// Cached reference to constraint component
 private_subobject:
diff --git a/Engine/Source/Runtime/Engine/Classes/PhysicsEngine/PhysicsConstraintComponent.h b/Engine/Source/Runtime/Engine/Classes/PhysicsEngine/PhysicsConstraintComponent.h
--- a/Engine/Source/Runtime/Engine/Classes/PhysicsEngine/PhysicsConstraintComponent.h
+++ b/Engine/Source/Runtime/Engine/Classes/PhysicsEngine/PhysicsConstraintComponent.h
@@ -12,7 +12,9 @@ DECLARE_DYNAMIC_MULTICAST_DELEGATE_OneParam(FConstraintBrokenSignature, int32, C
 UCLASS(ClassGroup=Physics, MinimalAPI, meta=(BlueprintSpawnableComponent), HideCategories=(Activation,"Components|Activation", Physics, Mobility), ShowCategories=("Physics|Components|PhysicsConstraint"))
 class UPhysicsConstraintComponent : public USceneComponent
 {
-	GENERATED_UCLASS_BODY()
+	GENERATED_BODY()
+public:
+	ENGINE_API UPhysicsConstraintComponent(const FObjectInitializer& ObjectInitializer = FObjectInitializer::Get());
 
 	/** Pointer to first Actor to constrain.  */
 	UPROPERTY(EditInstanceOnly, Category=Constraint)
diff --git a/Engine/Source/Runtime/Engine/Classes/PhysicsEngine/PhysicsConstraintTemplate.h b/Engine/Source/Runtime/Engine/Classes/PhysicsEngine/PhysicsConstraintTemplate.h
--- a/Engine/Source/Runtime/Engine/Classes/PhysicsEngine/PhysicsConstraintTemplate.h
+++ b/Engine/Source/Runtime/Engine/Classes/PhysicsEngine/PhysicsConstraintTemplate.h
@@ -17,7 +17,9 @@
 UCLASS(hidecategories=Object, MinimalAPI)
 class UPhysicsConstraintTemplate : public UObject
 {
-	GENERATED_UCLASS_BODY()
+	GENERATED_BODY()
+public:
+	ENGINE_API UPhysicsConstraintTemplate(const FObjectInitializer& ObjectInitializer = FObjectInitializer::Get());
 
 	UPROPERTY()
 	FName JointName_DEPRECATED;
diff --git a/Engine/Source/Runtime/Engine/Classes/PhysicsEngine/PhysicsHandleComponent.h b/Engine/Source/Runtime/Engine/Classes/PhysicsEngine/PhysicsHandleComponent.h
--- a/Engine/Source/Runtime/Engine/Classes/PhysicsEngine/PhysicsHandleComponent.h
+++ b/Engine/Source/Runtime/Engine/Classes/PhysicsEngine/PhysicsHandleComponent.h
@@ -17,7 +17,9 @@ namespace physx
 UCLASS(collapsecategories, ClassGroup=Physics, hidecategories=Object, MinimalAPI, meta=(BlueprintSpawnableComponent))
 class UPhysicsHandleComponent : public UActorComponent
 {
-	GENERATED_UCLASS_BODY()
+	GENERATED_BODY()
+public:
+	ENGINE_API UPhysicsHandleComponent(const FObjectInitializer& ObjectInitializer = FObjectInitializer::Get());
 
 	/** Component we are currently holding */
 	UPROPERTY()
diff --git a/Engine/Source/Runtime/Engine/Classes/PhysicsEngine/PhysicsSettings.h b/Engine/Source/Runtime/Engine/Classes/PhysicsEngine/PhysicsSettings.h
--- a/Engine/Source/Runtime/Engine/Classes/PhysicsEngine/PhysicsSettings.h
+++ b/Engine/Source/Runtime/Engine/Classes/PhysicsEngine/PhysicsSettings.h
@@ -15,7 +15,7 @@
 USTRUCT()
 struct FPhysicalSurfaceName
 {
-	GENERATED_USTRUCT_BODY()
+	GENERATED_BODY()
 
 	UPROPERTY()
 	TEnumAsByte<enum EPhysicalSurface> Type;
@@ -88,7 +88,9 @@ UCLASS(config=Engine, defaultconfig)
 class ENGINE_API UPhysicsSettings
 	: public UObject
 {
-	GENERATED_UCLASS_BODY()
+	GENERATED_BODY()
+public:
+	UPhysicsSettings(const FObjectInitializer& ObjectInitializer = FObjectInitializer::Get());
 
 	/** Default gravity. */
 	UPROPERTY(config, EditAnywhere, Category = Constants)
diff --git a/Engine/Source/Runtime/Engine/Classes/PhysicsEngine/PhysicsSpringComponent.h b/Engine/Source/Runtime/Engine/Classes/PhysicsEngine/PhysicsSpringComponent.h
--- a/Engine/Source/Runtime/Engine/Classes/PhysicsEngine/PhysicsSpringComponent.h
+++ b/Engine/Source/Runtime/Engine/Classes/PhysicsEngine/PhysicsSpringComponent.h
@@ -12,7 +12,9 @@
 UCLASS(hidecategories=(Object, Mobility, LOD), ClassGroup=Physics, showcategories=Trigger)
 class ENGINE_API UPhysicsSpringComponent : public USceneComponent
 {
-	GENERATED_UCLASS_BODY()
+	GENERATED_BODY()
+public:
+	UPhysicsSpringComponent(const FObjectInitializer& ObjectInitializer = FObjectInitializer::Get());
 
 	/** Specifies how much strength the spring has. The higher the SpringStiffness the more force the spring can push on a body with. */
 	UPROPERTY(BlueprintReadWrite, EditAnywhere, Category = Physics)
diff --git a/Engine/Source/Runtime/Engine/Classes/PhysicsEngine/PhysicsThruster.h b/Engine/Source/Runtime/Engine/Classes/PhysicsEngine/PhysicsThruster.h
--- a/Engine/Source/Runtime/Engine/Classes/PhysicsEngine/PhysicsThruster.h
+++ b/Engine/Source/Runtime/Engine/Classes/PhysicsEngine/PhysicsThruster.h
@@ -13,7 +13,9 @@
 UCLASS(hideCategories=(Input,Collision,Replication), showCategories=("Input|MouseInput", "Input|TouchInput"), ComponentWrapperClass)
 class APhysicsThruster : public ARigidBodyBase
 {
-	GENERATED_UCLASS_BODY()
+	GENERATED_BODY()
+public:
+	APhysicsThruster(const FObjectInitializer& ObjectInitializer = FObjectInitializer::Get());
 
 private_subobject:
 	/** Thruster component */
diff --git a/Engine/Source/Runtime/Engine/Classes/PhysicsEngine/PhysicsThrusterComponent.h b/Engine/Source/Runtime/Engine/Classes/PhysicsEngine/PhysicsThrusterComponent.h
--- a/Engine/Source/Runtime/Engine/Classes/PhysicsEngine/PhysicsThrusterComponent.h
+++ b/Engine/Source/Runtime/Engine/Classes/PhysicsEngine/PhysicsThrusterComponent.h
@@ -11,7 +11,9 @@
 UCLASS(hidecategories=(Object, Mobility, LOD), ClassGroup=Physics, showcategories=Trigger, MinimalAPI, meta=(BlueprintSpawnableComponent))
 class UPhysicsThrusterComponent : public USceneComponent
 {
-	GENERATED_UCLASS_BODY()
+	GENERATED_BODY()
+public:
+	ENGINE_API UPhysicsThrusterComponent(const FObjectInitializer& ObjectInitializer = FObjectInitializer::Get());
 
 	/** Strength of thrust force applied to the base object. */
 	UPROPERTY(BlueprintReadWrite, interp, Category=Physics)
diff --git a/Engine/Source/Runtime/Engine/Classes/PhysicsEngine/RadialForceActor.h b/Engine/Source/Runtime/Engine/Classes/PhysicsEngine/RadialForceActor.h
--- a/Engine/Source/Runtime/Engine/Classes/PhysicsEngine/RadialForceActor.h
+++ b/Engine/Source/Runtime/Engine/Classes/PhysicsEngine/RadialForceActor.h
@@ -8,7 +8,9 @@
 UCLASS(MinimalAPI, hideCategories=(Collision, Input), showCategories=("Input|MouseInput", "Input|TouchInput"), ComponentWrapperClass)
 class ARadialForceActor : public ARigidBodyBase
 {
-	GENERATED_UCLASS_BODY()
+	GENERATED_BODY()
+public:
+	ENGINE_API ARadialForceActor(const FObjectInitializer& ObjectInitializer = FObjectInitializer::Get());
 
 private_subobject:
 	/** Force component */
diff --git a/Engine/Source/Runtime/Engine/Classes/PhysicsEngine/RadialForceComponent.h b/Engine/Source/Runtime/Engine/Classes/PhysicsEngine/RadialForceComponent.h
--- a/Engine/Source/Runtime/Engine/Classes/PhysicsEngine/RadialForceComponent.h
+++ b/Engine/Source/Runtime/Engine/Classes/PhysicsEngine/RadialForceComponent.h
@@ -10,7 +10,9 @@
 UCLASS(hidecategories=(Object, Mobility), ClassGroup=Physics, showcategories=Trigger, meta=(BlueprintSpawnableComponent), MinimalAPI)
 class URadialForceComponent : public USceneComponent
 {
-	GENERATED_UCLASS_BODY()
+	GENERATED_BODY()
+public:
+	ENGINE_API URadialForceComponent(const FObjectInitializer& ObjectInitializer = FObjectInitializer::Get());
 
 	/** The radius to apply the force or impulse in */
 	UPROPERTY(interp, EditAnywhere, BlueprintReadWrite, Category=RadialForceComponent)
diff --git a/Engine/Source/Runtime/Engine/Classes/PhysicsEngine/RigidBodyBase.h b/Engine/Source/Runtime/Engine/Classes/PhysicsEngine/RigidBodyBase.h
--- a/Engine/Source/Runtime/Engine/Classes/PhysicsEngine/RigidBodyBase.h
+++ b/Engine/Source/Runtime/Engine/Classes/PhysicsEngine/RigidBodyBase.h
@@ -10,7 +10,9 @@
 UCLASS(ClassGroup=Physics, abstract,MinimalAPI)
 class ARigidBodyBase : public AActor
 {
-	GENERATED_UCLASS_BODY()
+	GENERATED_BODY()
+public:
+	ENGINE_API ARigidBodyBase(const FObjectInitializer& ObjectInitializer = FObjectInitializer::Get());
 
 };
 
diff --git a/Engine/Source/Runtime/Engine/Classes/Slate/ButtonStyleAsset.h b/Engine/Source/Runtime/Engine/Classes/Slate/ButtonStyleAsset.h
--- a/Engine/Source/Runtime/Engine/Classes/Slate/ButtonStyleAsset.h
+++ b/Engine/Source/Runtime/Engine/Classes/Slate/ButtonStyleAsset.h
@@ -12,7 +12,9 @@
 UCLASS(hidecategories=Object, MinimalAPI)
 class UButtonStyleAsset : public UObject
 {
-	GENERATED_UCLASS_BODY()
+	GENERATED_BODY()
+public:
+	ENGINE_API UButtonStyleAsset(const FObjectInitializer& ObjectInitializer = FObjectInitializer::Get());
 
 public:
 	/** The actual data describing the button's appearance. */
diff --git a/Engine/Source/Runtime/Engine/Classes/Slate/CheckboxStyleAsset.h b/Engine/Source/Runtime/Engine/Classes/Slate/CheckboxStyleAsset.h
--- a/Engine/Source/Runtime/Engine/Classes/Slate/CheckboxStyleAsset.h
+++ b/Engine/Source/Runtime/Engine/Classes/Slate/CheckboxStyleAsset.h
@@ -12,7 +12,9 @@
 UCLASS(hidecategories=Object, MinimalAPI)
 class UCheckBoxStyleAsset : public UObject
 {
-	GENERATED_UCLASS_BODY()
+	GENERATED_BODY()
+public:
+	ENGINE_API UCheckBoxStyleAsset(const FObjectInitializer& ObjectInitializer = FObjectInitializer::Get());
 
 public:
 	/** The actual data describing the Check Box's appearance. */
diff --git a/Engine/Source/Runtime/Engine/Classes/Slate/SlateBrushAsset.h b/Engine/Source/Runtime/Engine/Classes/Slate/SlateBrushAsset.h
--- a/Engine/Source/Runtime/Engine/Classes/Slate/SlateBrushAsset.h
+++ b/Engine/Source/Runtime/Engine/Classes/Slate/SlateBrushAsset.h
@@ -12,7 +12,9 @@
 UCLASS(hidecategories=Object, MinimalAPI, BlueprintType)
 class USlateBrushAsset : public UObject
 {
-	GENERATED_UCLASS_BODY()
+	GENERATED_BODY()
+public:
+	ENGINE_API USlateBrushAsset(const FObjectInitializer& ObjectInitializer = FObjectInitializer::Get());
 
 public:
 	/** The slate brush resource describing the texture's behavior. */
diff --git a/Engine/Source/Runtime/Engine/Classes/Sound/AmbientSound.h b/Engine/Source/Runtime/Engine/Classes/Sound/AmbientSound.h
--- a/Engine/Source/Runtime/Engine/Classes/Sound/AmbientSound.h
+++ b/Engine/Source/Runtime/Engine/Classes/Sound/AmbientSound.h
@@ -7,7 +7,9 @@
 UCLASS(AutoExpandCategories=Audio, ClassGroup=Sounds, MinimalAPI, hideCategories(Collision, Input, Game), showCategories=("Input|MouseInput", "Input|TouchInput", "Game|Damage"), ComponentWrapperClass)
 class AAmbientSound : public AActor
 {
-	GENERATED_UCLASS_BODY()
+	GENERATED_BODY()
+public:
+	ENGINE_API AAmbientSound(const FObjectInitializer& ObjectInitializer = FObjectInitializer::Get());
 
 private_subobject:
 	/** Audio component that handles sound playing */
diff --git a/Engine/Source/Runtime/Engine/Classes/Sound/AudioSettings.h b/Engine/Source/Runtime/Engine/Classes/Sound/AudioSettings.h
--- a/Engine/Source/Runtime/Engine/Classes/Sound/AudioSettings.h
+++ b/Engine/Source/Runtime/Engine/Classes/Sound/AudioSettings.h
@@ -12,7 +12,9 @@ UCLASS(config=Engine, defaultconfig)
 class ENGINE_API UAudioSettings
 	: public UObject
 {
-	GENERATED_UCLASS_BODY()
+	GENERATED_BODY()
+public:
+	UAudioSettings(const FObjectInitializer& ObjectInitializer = FObjectInitializer::Get());
 
 	/** The SoundClass assigned to newly created sounds */
 	UPROPERTY(config, EditAnywhere, Category="Audio", meta=(AllowedClasses="SoundClass", DisplayName="Default Sound Class"))
diff --git a/Engine/Source/Runtime/Engine/Classes/Sound/AudioVolume.h b/Engine/Source/Runtime/Engine/Classes/Sound/AudioVolume.h
--- a/Engine/Source/Runtime/Engine/Classes/Sound/AudioVolume.h
+++ b/Engine/Source/Runtime/Engine/Classes/Sound/AudioVolume.h
@@ -47,7 +47,7 @@ enum ReverbPreset
 USTRUCT(BlueprintType)
 struct FReverbSettings
 {
-	GENERATED_USTRUCT_BODY()
+	GENERATED_BODY()
 
 	/* Whether to apply the reverb settings below. */
 	UPROPERTY(EditAnywhere, BlueprintReadWrite, Category=ReverbSettings )
@@ -97,7 +97,7 @@ struct TStructOpsTypeTraits<FReverbSettings> : public TStructOpsTypeTraitsBase
 USTRUCT(BlueprintType)
 struct FInteriorSettings
 {
-	GENERATED_USTRUCT_BODY()
+	GENERATED_BODY()
 
 	// Whether these interior settings are the default values for the world
 	UPROPERTY()
@@ -155,7 +155,9 @@ struct FInteriorSettings
 UCLASS(hidecategories=(Advanced, Attachment, Collision, Volume), MinimalAPI)
 class AAudioVolume : public AVolume
 {
-	GENERATED_UCLASS_BODY()
+	GENERATED_BODY()
+public:
+	ENGINE_API AAudioVolume(const FObjectInitializer& ObjectInitializer = FObjectInitializer::Get());
 
 	/**
 	 * Priority of this volume. In the case of overlapping volumes the one with the highest priority
diff --git a/Engine/Source/Runtime/Engine/Classes/Sound/DialogueSoundWaveProxy.h b/Engine/Source/Runtime/Engine/Classes/Sound/DialogueSoundWaveProxy.h
--- a/Engine/Source/Runtime/Engine/Classes/Sound/DialogueSoundWaveProxy.h
+++ b/Engine/Source/Runtime/Engine/Classes/Sound/DialogueSoundWaveProxy.h
@@ -9,7 +9,9 @@
 UCLASS()
 class UDialogueSoundWaveProxy : public USoundBase
 {
-	GENERATED_UCLASS_BODY()
+	GENERATED_BODY()
+public:
+	UDialogueSoundWaveProxy(const FObjectInitializer& ObjectInitializer = FObjectInitializer::Get());
 
 	friend class UDialogueWave;
 
diff --git a/Engine/Source/Runtime/Engine/Classes/Sound/DialogueTypes.h b/Engine/Source/Runtime/Engine/Classes/Sound/DialogueTypes.h
--- a/Engine/Source/Runtime/Engine/Classes/Sound/DialogueTypes.h
+++ b/Engine/Source/Runtime/Engine/Classes/Sound/DialogueTypes.h
@@ -36,7 +36,7 @@ class UDialogueVoice;
 USTRUCT(BlueprintType)
 struct ENGINE_API FDialogueContext
 {
-	GENERATED_USTRUCT_BODY()
+	GENERATED_BODY()
 
 	FDialogueContext();
 
@@ -61,7 +61,7 @@ class UDialogueWave;
 USTRUCT()
 struct FDialogueWaveParameter
 {
-	GENERATED_USTRUCT_BODY()
+	GENERATED_BODY()
 
 	FDialogueWaveParameter();
 
@@ -77,5 +77,7 @@ struct FDialogueWaveParameter
 UCLASS(abstract)
 class UDialogueTypes : public UObject
 {
-	GENERATED_UCLASS_BODY()
+	GENERATED_BODY()
+public:
+	UDialogueTypes(const FObjectInitializer& ObjectInitializer = FObjectInitializer::Get());
 };
diff --git a/Engine/Source/Runtime/Engine/Classes/Sound/DialogueVoice.h b/Engine/Source/Runtime/Engine/Classes/Sound/DialogueVoice.h
--- a/Engine/Source/Runtime/Engine/Classes/Sound/DialogueVoice.h
+++ b/Engine/Source/Runtime/Engine/Classes/Sound/DialogueVoice.h
@@ -13,7 +13,9 @@
 UCLASS(hidecategories=Object, editinlinenew, MinimalAPI, BlueprintType)
 class UDialogueVoice : public UObject
 {
-	GENERATED_UCLASS_BODY()
+	GENERATED_BODY()
+public:
+	ENGINE_API UDialogueVoice(const FObjectInitializer& ObjectInitializer = FObjectInitializer::Get());
 
 	UPROPERTY(EditAnywhere, Category=DialogueVoice, AssetRegistrySearchable)
 	TEnumAsByte<EGrammaticalGender::Type> Gender;
@@ -33,4 +35,4 @@ public:
 
 	virtual void PostDuplicate(bool bDuplicateForPIE);
 	// End UObject interface. 
-};
\ No newline at end of file
+};
diff --git a/Engine/Source/Runtime/Engine/Classes/Sound/DialogueWave.h b/Engine/Source/Runtime/Engine/Classes/Sound/DialogueWave.h
--- a/Engine/Source/Runtime/Engine/Classes/Sound/DialogueWave.h
+++ b/Engine/Source/Runtime/Engine/Classes/Sound/DialogueWave.h
@@ -16,7 +16,7 @@ class UDialogueSoundWaveProxy;
 USTRUCT()
 struct FDialogueContextMapping
 {
-	GENERATED_USTRUCT_BODY()
+	GENERATED_BODY()
 
 	FDialogueContextMapping();
 
@@ -39,7 +39,9 @@ bool operator!=(const FDialogueContextMapping& LHS, const FDialogueContextMappin
 UCLASS(hidecategories=Object, editinlinenew, MinimalAPI, BlueprintType)
 class UDialogueWave : public UObject
 {
-	GENERATED_UCLASS_BODY()
+	GENERATED_BODY()
+public:
+	ENGINE_API UDialogueWave(const FObjectInitializer& ObjectInitializer = FObjectInitializer::Get());
 
 	/** true if this dialogue is considered to contain mature/adult content. */
 	UPROPERTY(EditAnywhere, Category=Filter, AssetRegistrySearchable)
@@ -86,4 +88,4 @@ public:
 
 protected:
 	void UpdateMappingProxy(FDialogueContextMapping& ContextMapping);
-};
\ No newline at end of file
+};
diff --git a/Engine/Source/Runtime/Engine/Classes/Sound/ReverbEffect.h b/Engine/Source/Runtime/Engine/Classes/Sound/ReverbEffect.h
--- a/Engine/Source/Runtime/Engine/Classes/Sound/ReverbEffect.h
+++ b/Engine/Source/Runtime/Engine/Classes/Sound/ReverbEffect.h
@@ -7,7 +7,9 @@
 UCLASS(BlueprintType, hidecategories=object, MinimalAPI)
 class UReverbEffect : public UObject
 {
-	GENERATED_UCLASS_BODY()
+	GENERATED_BODY()
+public:
+	ENGINE_API UReverbEffect(const FObjectInitializer& ObjectInitializer = FObjectInitializer::Get());
 
 	/** Density - 0.0 < 1.0 < 1.0 - Coloration of the late reverb - lower value is more */
 	UPROPERTY(Category=ReverbParameters, meta=(ClampMin = "0.0", ClampMax = "1.0"), EditAnywhere)
diff --git a/Engine/Source/Runtime/Engine/Classes/Sound/SoundAttenuation.h b/Engine/Source/Runtime/Engine/Classes/Sound/SoundAttenuation.h
--- a/Engine/Source/Runtime/Engine/Classes/Sound/SoundAttenuation.h
+++ b/Engine/Source/Runtime/Engine/Classes/Sound/SoundAttenuation.h
@@ -43,7 +43,7 @@ The settings for attenuating.
 USTRUCT(BlueprintType)
 struct ENGINE_API FAttenuationSettings
 {
-	GENERATED_USTRUCT_BODY()
+	GENERATED_BODY()
 
 	/* Enable attenuation via volume. */
 	UPROPERTY(EditAnywhere, BlueprintReadWrite, Category=Attenuation )
@@ -162,7 +162,9 @@ struct TStructOpsTypeTraits<FAttenuationSettings> : public TStructOpsTypeTraitsB
 UCLASS(BlueprintType, hidecategories=Object, editinlinenew, MinimalAPI)
 class USoundAttenuation : public UObject
 {
-	GENERATED_UCLASS_BODY()
+	GENERATED_BODY()
+public:
+	ENGINE_API USoundAttenuation(const FObjectInitializer& ObjectInitializer = FObjectInitializer::Get());
 
 	UPROPERTY(EditAnywhere, Category=Settings)
 	FAttenuationSettings Attenuation;
diff --git a/Engine/Source/Runtime/Engine/Classes/Sound/SoundBase.h b/Engine/Source/Runtime/Engine/Classes/Sound/SoundBase.h
--- a/Engine/Source/Runtime/Engine/Classes/Sound/SoundBase.h
+++ b/Engine/Source/Runtime/Engine/Classes/Sound/SoundBase.h
@@ -31,7 +31,9 @@ namespace EMaxConcurrentResolutionRule
 UCLASS(config=Engine, hidecategories=Object, abstract, editinlinenew, BlueprintType)
 class ENGINE_API USoundBase : public UObject
 {
-	GENERATED_UCLASS_BODY()
+	GENERATED_BODY()
+public:
+	USoundBase(const FObjectInitializer& ObjectInitializer = FObjectInitializer::Get());
 
 protected:
 	/** Sound class this sound belongs to */
diff --git a/Engine/Source/Runtime/Engine/Classes/Sound/SoundClass.h b/Engine/Source/Runtime/Engine/Classes/Sound/SoundClass.h
--- a/Engine/Source/Runtime/Engine/Classes/Sound/SoundClass.h
+++ b/Engine/Source/Runtime/Engine/Classes/Sound/SoundClass.h
@@ -17,7 +17,7 @@ namespace EAudioOutputTarget
 USTRUCT()
 struct FSoundClassEditorData
 {
-	GENERATED_USTRUCT_BODY()
+	GENERATED_BODY()
 
 	int32 NodePosX;
 
@@ -43,7 +43,7 @@ struct FSoundClassEditorData
 USTRUCT()
 struct FSoundClassProperties
 {
-	GENERATED_USTRUCT_BODY()
+	GENERATED_BODY()
 
 	/** Volume multiplier. */
 	UPROPERTY(EditAnywhere, Category=SoundClassProperties)
@@ -132,7 +132,7 @@ struct FSoundClassProperties
 USTRUCT()
 struct FPassiveSoundMixModifier
 {
-	GENERATED_USTRUCT_BODY()
+	GENERATED_BODY()
 
 	/** The SoundMix to activate */
 	UPROPERTY(EditAnywhere, Category=PassiveSoundMixModifier)
@@ -158,7 +158,9 @@ struct FPassiveSoundMixModifier
 UCLASS(hidecategories=object, MinimalAPI)
 class USoundClass : public UObject
 {
-	GENERATED_UCLASS_BODY()
+	GENERATED_BODY()
+public:
+	ENGINE_API USoundClass(const FObjectInitializer& ObjectInitializer = FObjectInitializer::Get());
 
 	/** Configurable properties like volume and priority. */
 	UPROPERTY(EditAnywhere, Category=SoundClass)
