Enabled UHT to digest GENERATED_BODY instead of GENERATED_UCLASS_BODY, GENERATED_USTRUCT_BODY, GENERATED_UINTERFACE_BODY or GENERATED_IINTERFACE_BODY, changed every occurence to the new syntax and fix

diff --git a/Engine/Source/Runtime/Engine/Classes/Sound/SoundCue.h b/Engine/Source/Runtime/Engine/Classes/Sound/SoundCue.h
--- a/Engine/Source/Runtime/Engine/Classes/Sound/SoundCue.h
+++ b/Engine/Source/Runtime/Engine/Classes/Sound/SoundCue.h
@@ -10,7 +10,7 @@ struct FSoundParseParameters;
 USTRUCT()
 struct FSoundNodeEditorData
 {
-	GENERATED_USTRUCT_BODY()
+	GENERATED_BODY()
 
 	int32 NodePosX;
 
@@ -36,7 +36,9 @@ struct FSoundNodeEditorData
 UCLASS(hidecategories=object, MinimalAPI, BlueprintType)
 class USoundCue : public USoundBase
 {
-	GENERATED_UCLASS_BODY()
+	GENERATED_BODY()
+public:
+	ENGINE_API USoundCue(const FObjectInitializer& ObjectInitializer = FObjectInitializer::Get());
 
 	/* Indicates whether attenuation should use the Attenuation Overrides or the Attenuation Settings asset */
 	UPROPERTY(EditAnywhere, Category=Attenuation)
diff --git a/Engine/Source/Runtime/Engine/Classes/Sound/SoundGroups.h b/Engine/Source/Runtime/Engine/Classes/Sound/SoundGroups.h
--- a/Engine/Source/Runtime/Engine/Classes/Sound/SoundGroups.h
+++ b/Engine/Source/Runtime/Engine/Classes/Sound/SoundGroups.h
@@ -37,7 +37,7 @@ enum ESoundGroup
 USTRUCT()
 struct FSoundGroup
 {
-	GENERATED_USTRUCT_BODY()
+	GENERATED_BODY()
 
 	// The sound group enumeration we are setting values for
 	UPROPERTY(config)
@@ -71,7 +71,9 @@ struct FSoundGroup
 UCLASS(config=Engine, abstract)
 class USoundGroups : public UObject
 {
-	GENERATED_UCLASS_BODY()
+	GENERATED_BODY()
+public:
+	USoundGroups(const FObjectInitializer& ObjectInitializer = FObjectInitializer::Get());
 
 	void Initialize() const;
 
diff --git a/Engine/Source/Runtime/Engine/Classes/Sound/SoundMix.h b/Engine/Source/Runtime/Engine/Classes/Sound/SoundMix.h
--- a/Engine/Source/Runtime/Engine/Classes/Sound/SoundMix.h
+++ b/Engine/Source/Runtime/Engine/Classes/Sound/SoundMix.h
@@ -7,7 +7,7 @@
 USTRUCT()
 struct FAudioEQEffect
 {
-	GENERATED_USTRUCT_BODY()
+	GENERATED_BODY()
 
 	/* Start time of effect */
 	double RootTime;
@@ -73,7 +73,7 @@ struct FAudioEQEffect
 USTRUCT()
 struct FSoundClassAdjuster
 {
-	GENERATED_USTRUCT_BODY()
+	GENERATED_BODY()
 
 	/* The sound class this adjuster affects. */
 	UPROPERTY(EditAnywhere, Category=SoundClassAdjuster, DisplayName = "Sound Class" )
@@ -111,7 +111,9 @@ struct FSoundClassAdjuster
 UCLASS(BlueprintType, hidecategories=object, MinimalAPI)
 class USoundMix : public UObject
 {
-	GENERATED_UCLASS_BODY()
+	GENERATED_BODY()
+public:
+	ENGINE_API USoundMix(const FObjectInitializer& ObjectInitializer = FObjectInitializer::Get());
 
 	/** Whether to apply the EQ effect */
 	UPROPERTY(EditAnywhere, Category=EQ, AssetRegistrySearchable )
diff --git a/Engine/Source/Runtime/Engine/Classes/Sound/SoundNode.h b/Engine/Source/Runtime/Engine/Classes/Sound/SoundNode.h
--- a/Engine/Source/Runtime/Engine/Classes/Sound/SoundNode.h
+++ b/Engine/Source/Runtime/Engine/Classes/Sound/SoundNode.h
@@ -43,7 +43,9 @@ struct FActiveSound;
 UCLASS(abstract, hidecategories=Object, editinlinenew)
 class ENGINE_API USoundNode : public UObject
 {
-	GENERATED_UCLASS_BODY()
+	GENERATED_BODY()
+public:
+	USoundNode(const FObjectInitializer& ObjectInitializer = FObjectInitializer::Get());
 
 	static const int32 MAX_ALLOWED_CHILD_NODES = 32;
 
diff --git a/Engine/Source/Runtime/Engine/Classes/Sound/SoundNodeAttenuation.h b/Engine/Source/Runtime/Engine/Classes/Sound/SoundNodeAttenuation.h
--- a/Engine/Source/Runtime/Engine/Classes/Sound/SoundNodeAttenuation.h
+++ b/Engine/Source/Runtime/Engine/Classes/Sound/SoundNodeAttenuation.h
@@ -15,7 +15,9 @@ class USoundAttenuation;
 UCLASS(hidecategories=Object, editinlinenew, MinimalAPI, meta=( DisplayName="Attenuation" ))
 class USoundNodeAttenuation : public USoundNode
 {
-	GENERATED_UCLASS_BODY()
+	GENERATED_BODY()
+public:
+	ENGINE_API USoundNodeAttenuation(const FObjectInitializer& ObjectInitializer = FObjectInitializer::Get());
 
 	UPROPERTY(EditAnywhere, Category=Attenuation, meta=(EditCondition="!bOverrideAttenuation"))
 	USoundAttenuation* AttenuationSettings;
diff --git a/Engine/Source/Runtime/Engine/Classes/Sound/SoundNodeBranch.h b/Engine/Source/Runtime/Engine/Classes/Sound/SoundNodeBranch.h
--- a/Engine/Source/Runtime/Engine/Classes/Sound/SoundNodeBranch.h
+++ b/Engine/Source/Runtime/Engine/Classes/Sound/SoundNodeBranch.h
@@ -10,7 +10,9 @@
 UCLASS(hidecategories=Object, editinlinenew, MinimalAPI, meta=( DisplayName="Branch" ))
 class USoundNodeBranch : public USoundNode
 {
-	GENERATED_UCLASS_BODY()
+	GENERATED_BODY()
+public:
+	ENGINE_API USoundNodeBranch(const FObjectInitializer& ObjectInitializer = FObjectInitializer::Get());
 
 	/** The name of the boolean parameter to use to determine which branch we should take */
 	UPROPERTY(EditAnywhere, Category=Branch)
diff --git a/Engine/Source/Runtime/Engine/Classes/Sound/SoundNodeConcatenator.h b/Engine/Source/Runtime/Engine/Classes/Sound/SoundNodeConcatenator.h
--- a/Engine/Source/Runtime/Engine/Classes/Sound/SoundNodeConcatenator.h
+++ b/Engine/Source/Runtime/Engine/Classes/Sound/SoundNodeConcatenator.h
@@ -16,7 +16,9 @@
 UCLASS(hidecategories = Object, editinlinenew, MinimalAPI, meta = (DisplayName = "Concatenator"))
 class USoundNodeConcatenator : public USoundNode
 {
-	GENERATED_UCLASS_BODY()
+	GENERATED_BODY()
+public:
+	ENGINE_API USoundNodeConcatenator(const FObjectInitializer& ObjectInitializer = FObjectInitializer::Get());
 
 	/** Volume multiplier for each input. */
 	UPROPERTY(EditAnywhere, editfixedsize, Category=Concatenator)
diff --git a/Engine/Source/Runtime/Engine/Classes/Sound/SoundNodeDelay.h b/Engine/Source/Runtime/Engine/Classes/Sound/SoundNodeDelay.h
--- a/Engine/Source/Runtime/Engine/Classes/Sound/SoundNodeDelay.h
+++ b/Engine/Source/Runtime/Engine/Classes/Sound/SoundNodeDelay.h
@@ -10,7 +10,9 @@
 UCLASS(hidecategories=Object, editinlinenew, MinimalAPI, meta=( DisplayName="Delay" ))
 class USoundNodeDelay : public USoundNode
 {
-	GENERATED_UCLASS_BODY()
+	GENERATED_BODY()
+public:
+	ENGINE_API USoundNodeDelay(const FObjectInitializer& ObjectInitializer = FObjectInitializer::Get());
 
 	/* The lower bound of delay time in seconds. */
 	UPROPERTY(EditAnywhere, Category=Delay )
diff --git a/Engine/Source/Runtime/Engine/Classes/Sound/SoundNodeDialoguePlayer.h b/Engine/Source/Runtime/Engine/Classes/Sound/SoundNodeDialoguePlayer.h
--- a/Engine/Source/Runtime/Engine/Classes/Sound/SoundNodeDialoguePlayer.h
+++ b/Engine/Source/Runtime/Engine/Classes/Sound/SoundNodeDialoguePlayer.h
@@ -11,7 +11,9 @@
 UCLASS(hidecategories = Object, editinlinenew, MinimalAPI, meta = (DisplayName = "Dialogue Player"))
 class USoundNodeDialoguePlayer : public USoundNode
 {
-	GENERATED_UCLASS_BODY()
+	GENERATED_BODY()
+public:
+	ENGINE_API USoundNodeDialoguePlayer(const FObjectInitializer& ObjectInitializer = FObjectInitializer::Get());
 
 	UPROPERTY(EditAnywhere, Category=DialoguePlayer)
 	FDialogueWaveParameter DialogueWaveParameter;
diff --git a/Engine/Source/Runtime/Engine/Classes/Sound/SoundNodeDistanceCrossFade.h b/Engine/Source/Runtime/Engine/Classes/Sound/SoundNodeDistanceCrossFade.h
--- a/Engine/Source/Runtime/Engine/Classes/Sound/SoundNodeDistanceCrossFade.h
+++ b/Engine/Source/Runtime/Engine/Classes/Sound/SoundNodeDistanceCrossFade.h
@@ -8,7 +8,7 @@
 USTRUCT()
 struct FDistanceDatum
 {
-	GENERATED_USTRUCT_BODY()
+	GENERATED_BODY()
 
 	/* The FadeInDistance at which to start hearing this sound.
 	 * If you want to hear the sound up close then setting this to 0 might be a good option. 
@@ -73,7 +73,9 @@ struct FDistanceDatum
 UCLASS(hidecategories=Object, editinlinenew, MinimalAPI, meta=( DisplayName="Crossfade by Distance" ))
 class USoundNodeDistanceCrossFade : public USoundNode
 {
-	GENERATED_UCLASS_BODY()
+	GENERATED_BODY()
+public:
+	ENGINE_API USoundNodeDistanceCrossFade(const FObjectInitializer& ObjectInitializer = FObjectInitializer::Get());
 
 	/**
 	 * Each input needs to have the correct data filled in so the SoundNodeDistanceCrossFade is able
diff --git a/Engine/Source/Runtime/Engine/Classes/Sound/SoundNodeDoppler.h b/Engine/Source/Runtime/Engine/Classes/Sound/SoundNodeDoppler.h
--- a/Engine/Source/Runtime/Engine/Classes/Sound/SoundNodeDoppler.h
+++ b/Engine/Source/Runtime/Engine/Classes/Sound/SoundNodeDoppler.h
@@ -12,7 +12,9 @@ struct FListener;
 UCLASS(hidecategories=Object, editinlinenew, meta=( DisplayName="Doppler" ))
 class USoundNodeDoppler : public USoundNode
 {
-	GENERATED_UCLASS_BODY()
+	GENERATED_BODY()
+public:
+	USoundNodeDoppler(const FObjectInitializer& ObjectInitializer = FObjectInitializer::Get());
 
 	/* How much to scale the doppler shift (1.0 is normal). */
 	UPROPERTY(EditAnywhere, Category=Doppler )
diff --git a/Engine/Source/Runtime/Engine/Classes/Sound/SoundNodeEnveloper.h b/Engine/Source/Runtime/Engine/Classes/Sound/SoundNodeEnveloper.h
--- a/Engine/Source/Runtime/Engine/Classes/Sound/SoundNodeEnveloper.h
+++ b/Engine/Source/Runtime/Engine/Classes/Sound/SoundNodeEnveloper.h
@@ -10,7 +10,9 @@
 UCLASS(hidecategories=Object, editinlinenew, MinimalAPI, meta=( DisplayName="Enveloper" ))
 class USoundNodeEnveloper : public USoundNode
 {
-	GENERATED_UCLASS_BODY()
+	GENERATED_BODY()
+public:
+	ENGINE_API USoundNodeEnveloper(const FObjectInitializer& ObjectInitializer = FObjectInitializer::Get());
 
 	// The time in seconds where the envelope's loop begins.
 	UPROPERTY(EditAnywhere, Category=Looping)
diff --git a/Engine/Source/Runtime/Engine/Classes/Sound/SoundNodeGroupControl.h b/Engine/Source/Runtime/Engine/Classes/Sound/SoundNodeGroupControl.h
--- a/Engine/Source/Runtime/Engine/Classes/Sound/SoundNodeGroupControl.h
+++ b/Engine/Source/Runtime/Engine/Classes/Sound/SoundNodeGroupControl.h
@@ -11,7 +11,9 @@
 UCLASS(hidecategories=Object, editinlinenew, MinimalAPI, meta=( DisplayName="Group Control" ))
 class USoundNodeGroupControl : public USoundNode
 {
-	GENERATED_UCLASS_BODY()
+	GENERATED_BODY()
+public:
+	ENGINE_API USoundNodeGroupControl(const FObjectInitializer& ObjectInitializer = FObjectInitializer::Get());
 
 	/* How many active sounds are allowed for each group */
 	UPROPERTY(EditAnywhere, editfixedsize, Category=GroupControl)
diff --git a/Engine/Source/Runtime/Engine/Classes/Sound/SoundNodeLooping.h b/Engine/Source/Runtime/Engine/Classes/Sound/SoundNodeLooping.h
--- a/Engine/Source/Runtime/Engine/Classes/Sound/SoundNodeLooping.h
+++ b/Engine/Source/Runtime/Engine/Classes/Sound/SoundNodeLooping.h
@@ -11,7 +11,9 @@
 UCLASS(hidecategories=Object, editinlinenew, MinimalAPI, meta=( DisplayName="Looping" ))
 class USoundNodeLooping : public USoundNode
 {
-	GENERATED_UCLASS_BODY()
+	GENERATED_BODY()
+public:
+	ENGINE_API USoundNodeLooping(const FObjectInitializer& ObjectInitializer = FObjectInitializer::Get());
 
 	/* The amount of times to loop */
 	UPROPERTY(EditAnywhere, Category = Looping, meta = (ClampMin = 1))
diff --git a/Engine/Source/Runtime/Engine/Classes/Sound/SoundNodeMature.h b/Engine/Source/Runtime/Engine/Classes/Sound/SoundNodeMature.h
--- a/Engine/Source/Runtime/Engine/Classes/Sound/SoundNodeMature.h
+++ b/Engine/Source/Runtime/Engine/Classes/Sound/SoundNodeMature.h
@@ -12,7 +12,9 @@
 UCLASS(hidecategories=Object, editinlinenew, MinimalAPI, meta=( DisplayName="Mature" ))
 class USoundNodeMature : public USoundNode
 {
-	GENERATED_UCLASS_BODY()
+	GENERATED_BODY()
+public:
+	ENGINE_API USoundNodeMature(const FObjectInitializer& ObjectInitializer = FObjectInitializer::Get());
 
 
 public:
diff --git a/Engine/Source/Runtime/Engine/Classes/Sound/SoundNodeMixer.h b/Engine/Source/Runtime/Engine/Classes/Sound/SoundNodeMixer.h
--- a/Engine/Source/Runtime/Engine/Classes/Sound/SoundNodeMixer.h
+++ b/Engine/Source/Runtime/Engine/Classes/Sound/SoundNodeMixer.h
@@ -11,7 +11,9 @@
 UCLASS(hidecategories=Object, editinlinenew, MinimalAPI, meta=( DisplayName="Mixer" ))
 class USoundNodeMixer : public USoundNode
 {
-	GENERATED_UCLASS_BODY()
+	GENERATED_BODY()
+public:
+	ENGINE_API USoundNodeMixer(const FObjectInitializer& ObjectInitializer = FObjectInitializer::Get());
 
 	/** A volume for each input.  Automatically sized. */
 	UPROPERTY(EditAnywhere, export, editfixedsize, Category=Mixer)
diff --git a/Engine/Source/Runtime/Engine/Classes/Sound/SoundNodeModulator.h b/Engine/Source/Runtime/Engine/Classes/Sound/SoundNodeModulator.h
--- a/Engine/Source/Runtime/Engine/Classes/Sound/SoundNodeModulator.h
+++ b/Engine/Source/Runtime/Engine/Classes/Sound/SoundNodeModulator.h
@@ -11,7 +11,9 @@
 UCLASS(hidecategories=Object, editinlinenew, MinimalAPI, meta=( DisplayName="Modulator" ))
 class USoundNodeModulator : public USoundNode
 {
-	GENERATED_UCLASS_BODY()
+	GENERATED_BODY()
+public:
+	ENGINE_API USoundNodeModulator(const FObjectInitializer& ObjectInitializer = FObjectInitializer::Get());
 
 	/* The lower bound of pitch (1.0 is no change). */
 	UPROPERTY(EditAnywhere, Category=Modulation )
diff --git a/Engine/Source/Runtime/Engine/Classes/Sound/SoundNodeModulatorContinuous.h b/Engine/Source/Runtime/Engine/Classes/Sound/SoundNodeModulatorContinuous.h
--- a/Engine/Source/Runtime/Engine/Classes/Sound/SoundNodeModulatorContinuous.h
+++ b/Engine/Source/Runtime/Engine/Classes/Sound/SoundNodeModulatorContinuous.h
@@ -17,7 +17,7 @@ enum ModulationParamMode
 USTRUCT()
 struct FModulatorContinuousParams
 {
-	GENERATED_USTRUCT_BODY()
+	GENERATED_BODY()
 
 	FModulatorContinuousParams()
 		: Default(1.f)
@@ -66,7 +66,9 @@ struct FModulatorContinuousParams
 UCLASS(hidecategories=Object, editinlinenew, meta=( DisplayName="Continuous Modulator" ))
 class USoundNodeModulatorContinuous : public USoundNode
 {
-	GENERATED_UCLASS_BODY()
+	GENERATED_BODY()
+public:
+	USoundNodeModulatorContinuous(const FObjectInitializer& ObjectInitializer = FObjectInitializer::Get());
 
 	UPROPERTY(EditAnywhere, Category=ContinuousModulator)
 	FModulatorContinuousParams PitchModulationParams;
diff --git a/Engine/Source/Runtime/Engine/Classes/Sound/SoundNodeOscillator.h b/Engine/Source/Runtime/Engine/Classes/Sound/SoundNodeOscillator.h
--- a/Engine/Source/Runtime/Engine/Classes/Sound/SoundNodeOscillator.h
+++ b/Engine/Source/Runtime/Engine/Classes/Sound/SoundNodeOscillator.h
@@ -11,7 +11,9 @@
 UCLASS(hidecategories=Object, editinlinenew, meta=( DisplayName="Oscillator" ))
 class USoundNodeOscillator : public USoundNode
 {
-	GENERATED_UCLASS_BODY()
+	GENERATED_BODY()
+public:
+	USoundNodeOscillator(const FObjectInitializer& ObjectInitializer = FObjectInitializer::Get());
 
 	/* Whether to oscillate volume. */
 	UPROPERTY(EditAnywhere, Category=Oscillator )
diff --git a/Engine/Source/Runtime/Engine/Classes/Sound/SoundNodeParamCrossFade.h b/Engine/Source/Runtime/Engine/Classes/Sound/SoundNodeParamCrossFade.h
--- a/Engine/Source/Runtime/Engine/Classes/Sound/SoundNodeParamCrossFade.h
+++ b/Engine/Source/Runtime/Engine/Classes/Sound/SoundNodeParamCrossFade.h
@@ -10,7 +10,9 @@
 UCLASS(editinlinenew, MinimalAPI, meta=( DisplayName="Crossfade by Param" ))
 class USoundNodeParamCrossFade : public USoundNodeDistanceCrossFade
 {
-	GENERATED_UCLASS_BODY()
+	GENERATED_BODY()
+public:
+	ENGINE_API USoundNodeParamCrossFade(const FObjectInitializer& ObjectInitializer = FObjectInitializer::Get());
 
 public:
 
diff --git a/Engine/Source/Runtime/Engine/Classes/Sound/SoundNodeRandom.h b/Engine/Source/Runtime/Engine/Classes/Sound/SoundNodeRandom.h
--- a/Engine/Source/Runtime/Engine/Classes/Sound/SoundNodeRandom.h
+++ b/Engine/Source/Runtime/Engine/Classes/Sound/SoundNodeRandom.h
@@ -11,7 +11,9 @@
 UCLASS(hidecategories=Object, editinlinenew, MinimalAPI, meta=( DisplayName="Random" ))
 class USoundNodeRandom : public USoundNode
 {
-	GENERATED_UCLASS_BODY()
+	GENERATED_BODY()
+public:
+	ENGINE_API USoundNodeRandom(const FObjectInitializer& ObjectInitializer = FObjectInitializer::Get());
 
 	UPROPERTY(EditAnywhere, editfixedsize, Category=Random)
 	TArray<float> Weights;
diff --git a/Engine/Source/Runtime/Engine/Classes/Sound/SoundNodeSoundClass.h b/Engine/Source/Runtime/Engine/Classes/Sound/SoundNodeSoundClass.h
--- a/Engine/Source/Runtime/Engine/Classes/Sound/SoundNodeSoundClass.h
+++ b/Engine/Source/Runtime/Engine/Classes/Sound/SoundNodeSoundClass.h
@@ -11,7 +11,9 @@
 UCLASS(hidecategories=Object, editinlinenew, MinimalAPI, meta=( DisplayName="SoundClass" ))
 class USoundNodeSoundClass : public USoundNode
 {
-	GENERATED_UCLASS_BODY()
+	GENERATED_BODY()
+public:
+	ENGINE_API USoundNodeSoundClass(const FObjectInitializer& ObjectInitializer = FObjectInitializer::Get());
 
 	UPROPERTY(EditAnywhere, Category=SoundClass)
 	USoundClass* SoundClassOverride;
diff --git a/Engine/Source/Runtime/Engine/Classes/Sound/SoundNodeSwitch.h b/Engine/Source/Runtime/Engine/Classes/Sound/SoundNodeSwitch.h
--- a/Engine/Source/Runtime/Engine/Classes/Sound/SoundNodeSwitch.h
+++ b/Engine/Source/Runtime/Engine/Classes/Sound/SoundNodeSwitch.h
@@ -10,7 +10,9 @@
 UCLASS(hidecategories=Object, editinlinenew, MinimalAPI, meta=(DisplayName="Switch"))
 class USoundNodeSwitch : public USoundNode
 {
-	GENERATED_UCLASS_BODY()
+	GENERATED_BODY()
+public:
+	ENGINE_API USoundNodeSwitch(const FObjectInitializer& ObjectInitializer = FObjectInitializer::Get());
 
 	/** The name of the integer parameter to use to determine which branch we should take */
 	UPROPERTY(EditAnywhere, Category=Switch)
diff --git a/Engine/Source/Runtime/Engine/Classes/Sound/SoundNodeWaveParam.h b/Engine/Source/Runtime/Engine/Classes/Sound/SoundNodeWaveParam.h
--- a/Engine/Source/Runtime/Engine/Classes/Sound/SoundNodeWaveParam.h
+++ b/Engine/Source/Runtime/Engine/Classes/Sound/SoundNodeWaveParam.h
@@ -10,7 +10,9 @@
 UCLASS(hidecategories=Object, editinlinenew, MinimalAPI, meta=( DisplayName="Wave Param" ))
 class USoundNodeWaveParam : public USoundNode
 {
-	GENERATED_UCLASS_BODY()
+	GENERATED_BODY()
+public:
+	ENGINE_API USoundNodeWaveParam(const FObjectInitializer& ObjectInitializer = FObjectInitializer::Get());
 
 	/** The name of the wave parameter to use to look up the SoundWave we should play */
 	UPROPERTY(EditAnywhere, Category=WaveParam)
diff --git a/Engine/Source/Runtime/Engine/Classes/Sound/SoundNodeWavePlayer.h b/Engine/Source/Runtime/Engine/Classes/Sound/SoundNodeWavePlayer.h
--- a/Engine/Source/Runtime/Engine/Classes/Sound/SoundNodeWavePlayer.h
+++ b/Engine/Source/Runtime/Engine/Classes/Sound/SoundNodeWavePlayer.h
@@ -11,7 +11,9 @@
 UCLASS(hidecategories=Object, editinlinenew, MinimalAPI, meta=( DisplayName="Wave Player" ))
 class USoundNodeWavePlayer : public USoundNode
 {
-	GENERATED_UCLASS_BODY()
+	GENERATED_BODY()
+public:
+	ENGINE_API USoundNodeWavePlayer(const FObjectInitializer& ObjectInitializer = FObjectInitializer::Get());
 
 	UPROPERTY(EditAnywhere, Category=WavePlayer)
 	USoundWave* SoundWave;
diff --git a/Engine/Source/Runtime/Engine/Classes/Sound/SoundWave.h b/Engine/Source/Runtime/Engine/Classes/Sound/SoundWave.h
--- a/Engine/Source/Runtime/Engine/Classes/Sound/SoundWave.h
+++ b/Engine/Source/Runtime/Engine/Classes/Sound/SoundWave.h
@@ -64,7 +64,7 @@ struct FStreamedAudioChunk
 USTRUCT()
 struct FStreamedAudioPlatformData
 {
-	GENERATED_USTRUCT_BODY()
+	GENERATED_BODY()
 
 	/** Number of audio chunks. */
 	int32 NumChunks;
@@ -111,7 +111,9 @@ struct FStreamedAudioPlatformData
 UCLASS(hidecategories=Object, editinlinenew, BlueprintType)
 class ENGINE_API USoundWave : public USoundBase
 {
-	GENERATED_UCLASS_BODY()
+	GENERATED_BODY()
+public:
+	USoundWave(const FObjectInitializer& ObjectInitializer = FObjectInitializer::Get());
 
 	/** Platform agnostic compression quality. 1..100 with 1 being best compression and 100 being best quality. */
 	UPROPERTY(EditAnywhere, Category=Compression, meta=(ClampMin = "1", ClampMax = "100"), AssetRegistrySearchable)
diff --git a/Engine/Source/Runtime/Engine/Classes/Sound/SoundWaveStreaming.h b/Engine/Source/Runtime/Engine/Classes/Sound/SoundWaveStreaming.h
--- a/Engine/Source/Runtime/Engine/Classes/Sound/SoundWaveStreaming.h
+++ b/Engine/Source/Runtime/Engine/Classes/Sound/SoundWaveStreaming.h
@@ -14,7 +14,9 @@ DECLARE_DELEGATE_TwoParams( FOnSoundWaveStreamingUnderflow, class USoundWaveStre
 UCLASS()
 class ENGINE_API USoundWaveStreaming : public USoundWave
 {
-	GENERATED_UCLASS_BODY()
+	GENERATED_BODY()
+public:
+	USoundWaveStreaming(const FObjectInitializer& ObjectInitializer = FObjectInitializer::Get());
 
 private:
 	TArray<uint8> QueuedAudio;
@@ -48,4 +50,4 @@ public:
 	/** Called when GeneratePCMData is called but not enough data is available. Allows more data to be added, and will try again */
 	FOnSoundWaveStreamingUnderflow OnSoundWaveStreamingUnderflow;
 
-};
\ No newline at end of file
+};
diff --git a/Engine/Source/Runtime/Engine/Classes/Tests/AutomationTestSettings.h b/Engine/Source/Runtime/Engine/Classes/Tests/AutomationTestSettings.h
--- a/Engine/Source/Runtime/Engine/Classes/Tests/AutomationTestSettings.h
+++ b/Engine/Source/Runtime/Engine/Classes/Tests/AutomationTestSettings.h
@@ -15,7 +15,7 @@
 USTRUCT()
 struct FOpenTestAsset
 {
-	GENERATED_USTRUCT_BODY()
+	GENERATED_BODY()
 
 	/**
 	 * Asset reference
@@ -34,7 +34,7 @@ struct FOpenTestAsset
 USTRUCT()
 struct FExternalToolDefinition
 {
-	GENERATED_USTRUCT_BODY()
+	GENERATED_BODY()
 
 	/* The name of the tool / test. */
 	UPROPERTY(config, EditAnywhere, Category=ExternalTools)
@@ -67,7 +67,7 @@ struct FExternalToolDefinition
 USTRUCT()
 struct FImportFactorySettingValues
 {
-	GENERATED_USTRUCT_BODY()
+	GENERATED_BODY()
 
 	/* The name of the UProperty to change */
 	UPROPERTY(config, EditAnywhere, Category = Automation, meta = (ToolTip = "Name of the property to change.  Nested settings can be modified using \"Outer.Property\""))
@@ -85,7 +85,7 @@ struct FImportFactorySettingValues
 USTRUCT()
 struct FEditorImportExportTestDefinition
 {
-	GENERATED_USTRUCT_BODY()
+	GENERATED_BODY()
 
 	/* The file to import */
 	UPROPERTY(config, EditAnywhere, Category = Automation, meta = (FilePathFilter = "*"))
@@ -110,7 +110,7 @@ struct FEditorImportExportTestDefinition
 USTRUCT()
 struct FEditorImportWorkflowDefinition
 {
-	GENERATED_USTRUCT_BODY()
+	GENERATED_BODY()
 
 	/* The file to import */
 	UPROPERTY(config, EditAnywhere, Category = Automation, meta = (FilePathFilter = "*"))
@@ -127,7 +127,7 @@ struct FEditorImportWorkflowDefinition
 USTRUCT()
 struct FBuildPromotionImportWorkflowSettings
 {
-	GENERATED_USTRUCT_BODY()
+	GENERATED_BODY()
 
 	/* Import settings for the Diffuse texture */
 	UPROPERTY(config, EditAnywhere, Category = Automation)
@@ -180,7 +180,7 @@ struct FBuildPromotionImportWorkflowSettings
 USTRUCT()
 struct FBuildPromotionOpenAssetSettings
 {
-	GENERATED_USTRUCT_BODY()
+	GENERATED_BODY()
 
 	/* The blueprint asset to open */
 	UPROPERTY(config, EditAnywhere, Category = Automation, meta = (FilePathFilter = "uasset"))
@@ -213,7 +213,7 @@ struct FBuildPromotionOpenAssetSettings
 USTRUCT()
 struct FBuildPromotionBlueprintSettings
 {
-	GENERATED_USTRUCT_BODY()
+	GENERATED_BODY()
 
 	/** The starting mesh for the blueprint **/
 	UPROPERTY(EditAnywhere, Category = Automation, meta = (FilePathFilter = "uasset"))
@@ -230,7 +230,7 @@ struct FBuildPromotionBlueprintSettings
 USTRUCT()
 struct FBuildPromotionNewProjectSettings
 {
-	GENERATED_USTRUCT_BODY()
+	GENERATED_BODY()
 
 	/** The path for the new project */
 	UPROPERTY(EditAnywhere, Category = Automation)
@@ -247,7 +247,7 @@ struct FBuildPromotionNewProjectSettings
 USTRUCT()
 struct FBuildPromotionTestSettings
 {
-	GENERATED_USTRUCT_BODY()
+	GENERATED_BODY()
 
 	/** Default static mesh asset to apply materials to **/
 	UPROPERTY(EditAnywhere, Category = Automation)
@@ -280,7 +280,7 @@ struct FBuildPromotionTestSettings
 USTRUCT()
 struct FEditorMapPerformanceTestDefinition
 {
-	GENERATED_USTRUCT_BODY()
+	GENERATED_BODY()
 	
 	/** Map to be used for the Performance Capture **/
 	UPROPERTY(config, EditAnywhere, Category = Automation, meta = (FilePathFilter = "umap"))
@@ -297,7 +297,7 @@ struct FEditorMapPerformanceTestDefinition
 USTRUCT()
 struct FLaunchOnTestSettings
 {
-	GENERATED_USTRUCT_BODY()
+	GENERATED_BODY()
 
 	/** Map to be used for the Launch On test **/
 	UPROPERTY(config, EditAnywhere, Category = Automation, meta = (FilePathFilter = "umap"))
@@ -314,7 +314,9 @@ struct FLaunchOnTestSettings
 UCLASS(config=Engine)
 class ENGINE_API UAutomationTestSettings : public UObject
 {
-	GENERATED_UCLASS_BODY()
+	GENERATED_BODY()
+public:
+	UAutomationTestSettings(const FObjectInitializer& ObjectInitializer = FObjectInitializer::Get());
 
 public:
 
diff --git a/Engine/Source/Runtime/Engine/Classes/Tests/TextPropertyTestObject.h b/Engine/Source/Runtime/Engine/Classes/Tests/TextPropertyTestObject.h
--- a/Engine/Source/Runtime/Engine/Classes/Tests/TextPropertyTestObject.h
+++ b/Engine/Source/Runtime/Engine/Classes/Tests/TextPropertyTestObject.h
@@ -10,7 +10,9 @@
 UCLASS()
 class UTextPropertyTestObject : public UObject
 {
-	GENERATED_UCLASS_BODY()
+	GENERATED_BODY()
+public:
+	UTextPropertyTestObject(const FObjectInitializer& ObjectInitializer = FObjectInitializer::Get());
 
 public:
 	UPROPERTY()
@@ -19,4 +21,4 @@ public:
 	FText UndefaultedText;
 	UPROPERTY()
 	FText TransientText;
-};
\ No newline at end of file
+};
diff --git a/Engine/Source/Runtime/Engine/Classes/VectorField/VectorField.h b/Engine/Source/Runtime/Engine/Classes/VectorField/VectorField.h
--- a/Engine/Source/Runtime/Engine/Classes/VectorField/VectorField.h
+++ b/Engine/Source/Runtime/Engine/Classes/VectorField/VectorField.h
@@ -11,7 +11,9 @@
 UCLASS(MinimalAPI)
 class UVectorField : public UObject
 {
-	GENERATED_UCLASS_BODY()
+	GENERATED_BODY()
+public:
+	ENGINE_API UVectorField(const FObjectInitializer& ObjectInitializer = FObjectInitializer::Get());
 
 	/** Bounds of the volume in local space. */
 	UPROPERTY(EditAnywhere, Category=VectorFieldBounds)
diff --git a/Engine/Source/Runtime/Engine/Classes/VectorField/VectorFieldAnimated.h b/Engine/Source/Runtime/Engine/Classes/VectorField/VectorFieldAnimated.h
--- a/Engine/Source/Runtime/Engine/Classes/VectorField/VectorFieldAnimated.h
+++ b/Engine/Source/Runtime/Engine/Classes/VectorField/VectorFieldAnimated.h
@@ -22,7 +22,9 @@ enum EVectorFieldConstructionOp
 UCLASS(MinimalAPI)
 class UVectorFieldAnimated : public UVectorField
 {
-	GENERATED_UCLASS_BODY()
+	GENERATED_BODY()
+public:
+	ENGINE_API UVectorFieldAnimated(const FObjectInitializer& ObjectInitializer = FObjectInitializer::Get());
 
 	/** The texture from which to create the vector field. */
 	UPROPERTY(EditAnywhere, Category=Reconstruction)
diff --git a/Engine/Source/Runtime/Engine/Classes/VectorField/VectorFieldStatic.h b/Engine/Source/Runtime/Engine/Classes/VectorField/VectorFieldStatic.h
--- a/Engine/Source/Runtime/Engine/Classes/VectorField/VectorFieldStatic.h
+++ b/Engine/Source/Runtime/Engine/Classes/VectorField/VectorFieldStatic.h
@@ -11,7 +11,9 @@
 UCLASS(hidecategories=VectorFieldBounds, MinimalAPI)
 class UVectorFieldStatic : public UVectorField
 {
-	GENERATED_UCLASS_BODY()
+	GENERATED_BODY()
+public:
+	ENGINE_API UVectorFieldStatic(const FObjectInitializer& ObjectInitializer = FObjectInitializer::Get());
 
 	/** Size of the vector field volume. */
 	UPROPERTY(Category=VectorFieldStatic, VisibleAnywhere)
diff --git a/Engine/Source/Runtime/Engine/Classes/VectorField/VectorFieldVolume.h b/Engine/Source/Runtime/Engine/Classes/VectorField/VectorFieldVolume.h
--- a/Engine/Source/Runtime/Engine/Classes/VectorField/VectorFieldVolume.h
+++ b/Engine/Source/Runtime/Engine/Classes/VectorField/VectorFieldVolume.h
@@ -11,7 +11,9 @@
 UCLASS(hidecategories=(Object, Advanced, Collision), MinimalAPI)
 class AVectorFieldVolume : public AActor
 {
-	GENERATED_UCLASS_BODY()
+	GENERATED_BODY()
+public:
+	ENGINE_API AVectorFieldVolume(const FObjectInitializer& ObjectInitializer = FObjectInitializer::Get());
 
 private_subobject:
 	DEPRECATED_FORGAME(4.6, "VectorFieldComponent should not be accessed directly, please use GetVectorFieldComponent() function instead. VectorFieldComponent will soon be private and your code will not compile.")
diff --git a/Engine/Source/Runtime/Engine/Classes/Vehicles/TireType.h b/Engine/Source/Runtime/Engine/Classes/Vehicles/TireType.h
--- a/Engine/Source/Runtime/Engine/Classes/Vehicles/TireType.h
+++ b/Engine/Source/Runtime/Engine/Classes/Vehicles/TireType.h
@@ -11,7 +11,9 @@
 UCLASS()
 class ENGINE_API UTireType : public UDataAsset
 {
-	GENERATED_UCLASS_BODY()
+	GENERATED_BODY()
+public:
+	UTireType(const FObjectInitializer& ObjectInitializer = FObjectInitializer::Get());
 
 private:
 
diff --git a/Engine/Source/Runtime/Engine/Classes/Vehicles/VehicleAnimInstance.h b/Engine/Source/Runtime/Engine/Classes/Vehicles/VehicleAnimInstance.h
--- a/Engine/Source/Runtime/Engine/Classes/Vehicles/VehicleAnimInstance.h
+++ b/Engine/Source/Runtime/Engine/Classes/Vehicles/VehicleAnimInstance.h
@@ -10,7 +10,9 @@
 UCLASS(transient)
 class ENGINE_API UVehicleAnimInstance : public UAnimInstance
 {
-	GENERATED_UCLASS_BODY()
+	GENERATED_BODY()
+public:
+	UVehicleAnimInstance(const FObjectInitializer& ObjectInitializer = FObjectInitializer::Get());
 
 	/** Makes a montage jump to the end of a named section. */
 	UFUNCTION(BlueprintCallable, Category="Animation")
diff --git a/Engine/Source/Runtime/Engine/Classes/Vehicles/VehicleWheel.h b/Engine/Source/Runtime/Engine/Classes/Vehicles/VehicleWheel.h
--- a/Engine/Source/Runtime/Engine/Classes/Vehicles/VehicleWheel.h
+++ b/Engine/Source/Runtime/Engine/Classes/Vehicles/VehicleWheel.h
@@ -19,7 +19,9 @@ namespace physx
 UCLASS(BlueprintType, Blueprintable)
 class ENGINE_API UVehicleWheel : public UObject
 {
-	GENERATED_UCLASS_BODY()
+	GENERATED_BODY()
+public:
+	UVehicleWheel(const FObjectInitializer& ObjectInitializer = FObjectInitializer::Get());
 
 	/** 
 	 * Static mesh with collision setup for wheel, will be used to create wheel shape
diff --git a/Engine/Source/Runtime/Engine/Classes/Vehicles/WheeledVehicleMovementComponent.h b/Engine/Source/Runtime/Engine/Classes/Vehicles/WheeledVehicleMovementComponent.h
--- a/Engine/Source/Runtime/Engine/Classes/Vehicles/WheeledVehicleMovementComponent.h
+++ b/Engine/Source/Runtime/Engine/Classes/Vehicles/WheeledVehicleMovementComponent.h
@@ -86,7 +86,7 @@ struct FTireShaderOutput
 USTRUCT()
 struct ENGINE_API FWheelSetup
 {
-	GENERATED_USTRUCT_BODY()
+	GENERATED_BODY()
 
 	// The wheel class to use
 	UPROPERTY(EditAnywhere, Category=WheelSetup)
@@ -106,7 +106,7 @@ struct ENGINE_API FWheelSetup
 USTRUCT()
 struct ENGINE_API FReplicatedVehicleState
 {
-	GENERATED_USTRUCT_BODY()
+	GENERATED_BODY()
 
 	// input replication: steering
 	UPROPERTY()
@@ -132,7 +132,7 @@ struct ENGINE_API FReplicatedVehicleState
 USTRUCT()
 struct ENGINE_API FVehicleInputRate
 {
-	GENERATED_USTRUCT_BODY()
+	GENERATED_BODY()
 
 	// Rate at which the input value rises
 	UPROPERTY(EditAnywhere, Category=VehicleInputRate)
@@ -161,7 +161,9 @@ struct ENGINE_API FVehicleInputRate
 UCLASS(Abstract, hidecategories=(PlanarMovement, "Components|Movement|Planar", Activation, "Components|Activation"))
 class ENGINE_API UWheeledVehicleMovementComponent : public UPawnMovementComponent, public IRVOAvoidanceInterface
 {
-	GENERATED_UCLASS_BODY()
+	GENERATED_BODY()
+public:
+	UWheeledVehicleMovementComponent(const FObjectInitializer& ObjectInitializer = FObjectInitializer::Get());
 
 	/** Wheels to create */
 	UPROPERTY(EditAnywhere, Category=VehicleSetup)
@@ -498,8 +500,10 @@ protected:
 	void UpdateState(float DeltaTime);
 
 	/** Pass current state to server */
-	UFUNCTION(reliable, server, WithValidation)
+	UFUNCTION(reliable, server="ServerUpdateState_Implementation", WithValidation="ServerUpdateState_Validate")
 	void ServerUpdateState(float InSteeringInput, float InThrottleInput, float InBrakeInput, float InHandbrakeInput, int32 CurrentGear);
+	virtual void ServerUpdateState_Implementation(float InSteeringInput, float InThrottleInput, float InBrakeInput, float InHandbrakeInput, int32 CurrentGear);
+	virtual bool ServerUpdateState_Validate(float , float , float , float , int32 );
 
 	/** Update RVO Avoidance for simulation */
 	void UpdateAvoidance(float DeltaTime);
diff --git a/Engine/Source/Runtime/Engine/Classes/Vehicles/WheeledVehicleMovementComponent4W.h b/Engine/Source/Runtime/Engine/Classes/Vehicles/WheeledVehicleMovementComponent4W.h
--- a/Engine/Source/Runtime/Engine/Classes/Vehicles/WheeledVehicleMovementComponent4W.h
+++ b/Engine/Source/Runtime/Engine/Classes/Vehicles/WheeledVehicleMovementComponent4W.h
@@ -25,7 +25,7 @@ namespace EVehicleDifferential4W
 USTRUCT()
 struct FVehicleDifferential4WData
 {
-	GENERATED_USTRUCT_BODY()
+	GENERATED_BODY()
 
 	/** Type of differential */
 	UPROPERTY(EditAnywhere, Category=Setup)
@@ -59,7 +59,7 @@ struct FVehicleDifferential4WData
 USTRUCT()
 struct FVehicleEngineData
 {
-	GENERATED_USTRUCT_BODY()
+	GENERATED_BODY()
 
 	/** Torque (Nm) at a given RPM*/
 	UPROPERTY(EditAnywhere, Category = Setup)
@@ -93,7 +93,7 @@ struct FVehicleEngineData
 USTRUCT()
 struct FVehicleGearData
 {
-	GENERATED_USTRUCT_BODY()
+	GENERATED_BODY()
 
 	/** Determines the amount of torque multiplication*/
 	UPROPERTY(EditAnywhere, Category = Setup)
@@ -111,7 +111,7 @@ struct FVehicleGearData
 USTRUCT()
 struct FVehicleTransmissionData
 {
-	GENERATED_USTRUCT_BODY()
+	GENERATED_BODY()
 	/** Whether to use automatic transmission */
 	UPROPERTY(EditAnywhere, Category = VehicleSetup, meta=(FriendlyName = "Automatic Transmission"))
 	bool bUseGearAutoBox;
@@ -148,7 +148,9 @@ struct FVehicleTransmissionData
 UCLASS()
 class ENGINE_API UWheeledVehicleMovementComponent4W : public UWheeledVehicleMovementComponent
 {
-	GENERATED_UCLASS_BODY()
+	GENERATED_BODY()
+public:
+	UWheeledVehicleMovementComponent4W(const FObjectInitializer& ObjectInitializer = FObjectInitializer::Get());
 
 	/** Engine */
 	UPROPERTY(EditAnywhere, Category = MechanicalSetup)
diff --git a/Engine/Source/Runtime/Engine/Classes/VisualLogger/VisualLoggerAutomationTests.h b/Engine/Source/Runtime/Engine/Classes/VisualLogger/VisualLoggerAutomationTests.h
--- a/Engine/Source/Runtime/Engine/Classes/VisualLogger/VisualLoggerAutomationTests.h
+++ b/Engine/Source/Runtime/Engine/Classes/VisualLogger/VisualLoggerAutomationTests.h
@@ -8,6 +8,8 @@
 UCLASS(NotBlueprintable, Transient, hidecategories = UObject, notplaceable)
 class UVisualLoggerAutomationTests : public UObject
 {
-	GENERATED_UCLASS_BODY()
+	GENERATED_BODY()
+public:
+	UVisualLoggerAutomationTests(const FObjectInitializer& ObjectInitializer = FObjectInitializer::Get());
 };
 
diff --git a/Engine/Source/Runtime/Engine/Classes/VisualLogger/VisualLoggerKismetLibrary.h b/Engine/Source/Runtime/Engine/Classes/VisualLogger/VisualLoggerKismetLibrary.h
--- a/Engine/Source/Runtime/Engine/Classes/VisualLogger/VisualLoggerKismetLibrary.h
+++ b/Engine/Source/Runtime/Engine/Classes/VisualLogger/VisualLoggerKismetLibrary.h
@@ -7,7 +7,9 @@
 UCLASS(MinimalAPI)
 class UVisualLoggerKismetLibrary : public UBlueprintFunctionLibrary
 {
-	GENERATED_UCLASS_BODY()
+	GENERATED_BODY()
+public:
+	ENGINE_API UVisualLoggerKismetLibrary(const FObjectInitializer& ObjectInitializer = FObjectInitializer::Get());
 
 	/** Logs simple text string with Visual Logger - recording for Visual Logs has to be enabled to record this data */
 	UFUNCTION(BlueprintCallable, Category = "Debug|VisualLogger", meta = (HidePin = "WorldContextObject", DefaultToSelf = "WorldContextObject", FriendlyName = "LogText"))
diff --git a/Engine/Source/Runtime/Engine/Public/AI/RVOAvoidanceInterface.h b/Engine/Source/Runtime/Engine/Public/AI/RVOAvoidanceInterface.h
--- a/Engine/Source/Runtime/Engine/Public/AI/RVOAvoidanceInterface.h
+++ b/Engine/Source/Runtime/Engine/Public/AI/RVOAvoidanceInterface.h
@@ -12,12 +12,15 @@
 UINTERFACE(MinimalAPI, meta=(CannotImplementInterfaceInBlueprint))
 class URVOAvoidanceInterface : public UInterface
 {
-	GENERATED_UINTERFACE_BODY()
+	GENERATED_BODY()
+public:
+	ENGINE_API URVOAvoidanceInterface(const FObjectInitializer& ObjectInitializer = FObjectInitializer::Get());
 };
 
 class ENGINE_API IRVOAvoidanceInterface
 {
-	GENERATED_IINTERFACE_BODY()
+	GENERATED_BODY()
+public:
 
 	/** Store the AvoidanceUID generated by the Avoidance Manager **/
 	virtual void SetRVOAvoidanceUID(int32 UID) = 0;
diff --git a/Engine/Source/Runtime/EngineMessages/Classes/EngineServiceMessages.h b/Engine/Source/Runtime/EngineMessages/Classes/EngineServiceMessages.h
--- a/Engine/Source/Runtime/EngineMessages/Classes/EngineServiceMessages.h
+++ b/Engine/Source/Runtime/EngineMessages/Classes/EngineServiceMessages.h
@@ -14,7 +14,7 @@
 USTRUCT()
 struct FEngineServicePing
 {
-	GENERATED_USTRUCT_BODY()
+	GENERATED_BODY()
 };
 
 template<>
@@ -33,7 +33,7 @@ struct TStructOpsTypeTraits<FEngineServicePing> : public TStructOpsTypeTraitsBas
 USTRUCT()
 struct FEngineServicePong
 {
-	GENERATED_USTRUCT_BODY()
+	GENERATED_BODY()
 
 	/** Holds the name of the currently loaded level, if any. */
 	UPROPERTY()
@@ -83,7 +83,7 @@ struct TStructOpsTypeTraits<FEngineServicePong> : public TStructOpsTypeTraitsBas
 USTRUCT()
 struct FEngineServiceAuthDeny
 {
-	GENERATED_USTRUCT_BODY()
+	GENERATED_BODY()
 
 	/** Holds the name of the user that denied access. */
 	UPROPERTY()
@@ -110,7 +110,7 @@ struct TStructOpsTypeTraits<FEngineServiceAuthDeny> : public TStructOpsTypeTrait
 USTRUCT()
 struct FEngineServiceAuthGrant
 {
-	GENERATED_USTRUCT_BODY()
+	GENERATED_BODY()
 
 	/** Holds the name of the user that granted access. */
 	UPROPERTY()
@@ -141,7 +141,7 @@ struct TStructOpsTypeTraits<FEngineServiceAuthGrant> : public TStructOpsTypeTrai
 USTRUCT()
 struct FEngineServiceExecuteCommand
 {
-	GENERATED_USTRUCT_BODY()
+	GENERATED_BODY()
 
 	/** Holds the command to execute. */
 	UPROPERTY()
@@ -181,7 +181,7 @@ struct TStructOpsTypeTraits<FEngineServiceExecuteCommand> : public TStructOpsTyp
 USTRUCT()
 struct FEngineServiceTerminate
 {
-	GENERATED_USTRUCT_BODY()
+	GENERATED_BODY()
 
 	/** Holds the name of the user that wants to terminate the engine. */
 	UPROPERTY()
@@ -219,7 +219,7 @@ struct TStructOpsTypeTraits<FEngineServiceTerminate> : public TStructOpsTypeTrai
 USTRUCT()
 struct FEngineServiceNotification
 {
-	GENERATED_USTRUCT_BODY()
+	GENERATED_BODY()
 
 	/** Holds the notification text. */
 	UPROPERTY()
diff --git a/Engine/Source/Runtime/EngineSettings/Classes/ConsoleSettings.h b/Engine/Source/Runtime/EngineSettings/Classes/ConsoleSettings.h
--- a/Engine/Source/Runtime/EngineSettings/Classes/ConsoleSettings.h
+++ b/Engine/Source/Runtime/EngineSettings/Classes/ConsoleSettings.h
@@ -11,7 +11,7 @@
 USTRUCT()
 struct FAutoCompleteCommand
 {
-	GENERATED_USTRUCT_BODY()
+	GENERATED_BODY()
 
 	UPROPERTY(config, EditAnywhere, Category=Command)
 	FString Command;
@@ -28,7 +28,9 @@ UCLASS(config=Input, defaultconfig)
 class ENGINESETTINGS_API UConsoleSettings
 	: public UObject
 {
-	GENERATED_UCLASS_BODY()
+	GENERATED_BODY()
+public:
+	UConsoleSettings(const FObjectInitializer& ObjectInitializer = FObjectInitializer::Get());
 
 	/**  Visible Console stuff */
 	UPROPERTY(globalconfig, EditAnywhere, Category=General)
diff --git a/Engine/Source/Runtime/EngineSettings/Classes/GameMapsSettings.h b/Engine/Source/Runtime/EngineSettings/Classes/GameMapsSettings.h
--- a/Engine/Source/Runtime/EngineSettings/Classes/GameMapsSettings.h
+++ b/Engine/Source/Runtime/EngineSettings/Classes/GameMapsSettings.h
@@ -33,7 +33,9 @@ UCLASS(config=Engine, defaultconfig)
 class ENGINESETTINGS_API UGameMapsSettings
 	: public UObject
 {
-	GENERATED_UCLASS_BODY()
+	GENERATED_BODY()
+public:
+	UGameMapsSettings(const FObjectInitializer& ObjectInitializer = FObjectInitializer::Get());
 
 	/**
 	 * Get the default map specified in the settings.
diff --git a/Engine/Source/Runtime/EngineSettings/Classes/GameNetworkManagerSettings.h b/Engine/Source/Runtime/EngineSettings/Classes/GameNetworkManagerSettings.h
--- a/Engine/Source/Runtime/EngineSettings/Classes/GameNetworkManagerSettings.h
+++ b/Engine/Source/Runtime/EngineSettings/Classes/GameNetworkManagerSettings.h
@@ -12,7 +12,9 @@ UCLASS(config=Game)
 class ENGINESETTINGS_API UGameNetworkManagerSettings
 	: public UObject
 {
-	GENERATED_UCLASS_BODY()
+	GENERATED_BODY()
+public:
+	UGameNetworkManagerSettings(const FObjectInitializer& ObjectInitializer = FObjectInitializer::Get());
 
 	/** Minimum bandwidth dynamically set per connection. */
 	UPROPERTY(globalconfig, EditAnywhere, Category=Bandwidth)
diff --git a/Engine/Source/Runtime/EngineSettings/Classes/GameSessionSettings.h b/Engine/Source/Runtime/EngineSettings/Classes/GameSessionSettings.h
--- a/Engine/Source/Runtime/EngineSettings/Classes/GameSessionSettings.h
+++ b/Engine/Source/Runtime/EngineSettings/Classes/GameSessionSettings.h
@@ -9,7 +9,9 @@ UCLASS(config=Game)
 class ENGINESETTINGS_API UGameSessionSettings
 	: public UObject
 {
-	GENERATED_UCLASS_BODY()
+	GENERATED_BODY()
+public:
+	UGameSessionSettings(const FObjectInitializer& ObjectInitializer = FObjectInitializer::Get());
 
 	/** Maximum number of spectators allowed by this server. */
 	UPROPERTY(globalconfig, EditAnywhere, Category=GameSessionSettings)
diff --git a/Engine/Source/Runtime/EngineSettings/Classes/GeneralEngineSettings.h b/Engine/Source/Runtime/EngineSettings/Classes/GeneralEngineSettings.h
--- a/Engine/Source/Runtime/EngineSettings/Classes/GeneralEngineSettings.h
+++ b/Engine/Source/Runtime/EngineSettings/Classes/GeneralEngineSettings.h
@@ -9,6 +9,8 @@ UCLASS(config=Engine)
 class ENGINESETTINGS_API UGeneralEngineSettings
 	: public UObject
 {
-	GENERATED_UCLASS_BODY()
+	GENERATED_BODY()
+public:
+	UGeneralEngineSettings(const FObjectInitializer& ObjectInitializer = FObjectInitializer::Get());
 
 };
diff --git a/Engine/Source/Runtime/EngineSettings/Classes/GeneralProjectSettings.h b/Engine/Source/Runtime/EngineSettings/Classes/GeneralProjectSettings.h
--- a/Engine/Source/Runtime/EngineSettings/Classes/GeneralProjectSettings.h
+++ b/Engine/Source/Runtime/EngineSettings/Classes/GeneralProjectSettings.h
@@ -9,7 +9,9 @@ UCLASS(config=Game, defaultconfig)
 class ENGINESETTINGS_API UGeneralProjectSettings
 	: public UObject
 {
-	GENERATED_UCLASS_BODY()
+	GENERATED_BODY()
+public:
+	UGeneralProjectSettings(const FObjectInitializer& ObjectInitializer = FObjectInitializer::Get());
 
 	/** The name of the company (author, provider) that created the project. */
 	UPROPERTY(config, EditAnywhere, Category=Publisher)
diff --git a/Engine/Source/Runtime/EngineSettings/Classes/HudSettings.h b/Engine/Source/Runtime/EngineSettings/Classes/HudSettings.h
--- a/Engine/Source/Runtime/EngineSettings/Classes/HudSettings.h
+++ b/Engine/Source/Runtime/EngineSettings/Classes/HudSettings.h
@@ -9,7 +9,9 @@ UCLASS(config=Game)
 class ENGINESETTINGS_API UHudSettings
 	: public UObject
 {
-	GENERATED_UCLASS_BODY()
+	GENERATED_BODY()
+public:
+	UHudSettings(const FObjectInitializer& ObjectInitializer = FObjectInitializer::Get());
 
 	/** Whether the HUD is visible at all.	 */
 	UPROPERTY(config, EditAnywhere, Category=HudSettings)
diff --git a/Engine/Source/Runtime/Foliage/Private/FoliageStatistics.h b/Engine/Source/Runtime/Foliage/Private/FoliageStatistics.h
--- a/Engine/Source/Runtime/Foliage/Private/FoliageStatistics.h
+++ b/Engine/Source/Runtime/Foliage/Private/FoliageStatistics.h
@@ -7,7 +7,9 @@
 UCLASS()
 class FOLIAGE_API UFoliageStatistics : public UBlueprintFunctionLibrary
 {
-	GENERATED_UCLASS_BODY()
+	GENERATED_BODY()
+public:
+	UFoliageStatistics(const FObjectInitializer& ObjectInitializer = FObjectInitializer::Get());
 
 	/**
 	* Counts how many foliage instances overlap a given sphere
diff --git a/Engine/Source/Runtime/Foliage/Private/InteractiveFoliageComponent.h b/Engine/Source/Runtime/Foliage/Private/InteractiveFoliageComponent.h
--- a/Engine/Source/Runtime/Foliage/Private/InteractiveFoliageComponent.h
+++ b/Engine/Source/Runtime/Foliage/Private/InteractiveFoliageComponent.h
@@ -7,7 +7,9 @@
 UCLASS(hidecategories=Object)
 class UInteractiveFoliageComponent : public UStaticMeshComponent
 {
-	GENERATED_UCLASS_BODY()
+	GENERATED_BODY()
+public:
+	UInteractiveFoliageComponent(const FObjectInitializer& ObjectInitializer = FObjectInitializer::Get());
 
 
 public:
diff --git a/Engine/Source/Runtime/Foliage/Private/ProceduralFoliageBlockingVolume.h b/Engine/Source/Runtime/Foliage/Private/ProceduralFoliageBlockingVolume.h
--- a/Engine/Source/Runtime/Foliage/Private/ProceduralFoliageBlockingVolume.h
+++ b/Engine/Source/Runtime/Foliage/Private/ProceduralFoliageBlockingVolume.h
@@ -11,7 +11,9 @@ class AProceduralFoliageActor;
 UCLASS(MinimalAPI)
 class AProceduralFoliageBlockingVolume : public AVolume
 {
-	GENERATED_UCLASS_BODY()
+	GENERATED_BODY()
+public:
+	FOLIAGE_API AProceduralFoliageBlockingVolume(const FObjectInitializer& ObjectInitializer = FObjectInitializer::Get());
 
 	UPROPERTY(Category = ProceduralFoliage, EditAnywhere)
 	AProceduralFoliageActor* ProceduralFoliageActor;
diff --git a/Engine/Source/Runtime/Foliage/Public/FoliageType.h b/Engine/Source/Runtime/Foliage/Public/FoliageType.h
--- a/Engine/Source/Runtime/Foliage/Public/FoliageType.h
+++ b/Engine/Source/Runtime/Foliage/Public/FoliageType.h
@@ -37,7 +37,9 @@ enum class EFoliageScaling : uint8
 UCLASS(hidecategories = Object, editinlinenew, MinimalAPI, BlueprintType, Blueprintable)
 class UFoliageType : public UObject
 {
-	GENERATED_UCLASS_BODY()
+	GENERATED_BODY()
+public:
+	FOLIAGE_API UFoliageType(const FObjectInitializer& ObjectInitializer = FObjectInitializer::Get());
 
 	/* A GUID that is updated every time the foliage type is modified, 
 	   so foliage placed in the level can detect the FoliageType has changed. */
diff --git a/Engine/Source/Runtime/Foliage/Public/FoliageType_InstancedStaticMesh.h b/Engine/Source/Runtime/Foliage/Public/FoliageType_InstancedStaticMesh.h
--- a/Engine/Source/Runtime/Foliage/Public/FoliageType_InstancedStaticMesh.h
+++ b/Engine/Source/Runtime/Foliage/Public/FoliageType_InstancedStaticMesh.h
@@ -8,7 +8,9 @@
 UCLASS(hidecategories=Object, editinlinenew, MinimalAPI)
 class UFoliageType_InstancedStaticMesh : public UFoliageType
 {
-	GENERATED_UCLASS_BODY()
+	GENERATED_BODY()
+public:
+	FOLIAGE_API UFoliageType_InstancedStaticMesh(const FObjectInitializer& ObjectInitializer = FObjectInitializer::Get());
 
 	UPROPERTY(EditAnywhere, Category=FoliageType)
 	UStaticMesh* Mesh;
diff --git a/Engine/Source/Runtime/Foliage/Public/InstancedFoliageActor.h b/Engine/Source/Runtime/Foliage/Public/InstancedFoliageActor.h
--- a/Engine/Source/Runtime/Foliage/Public/InstancedFoliageActor.h
+++ b/Engine/Source/Runtime/Foliage/Public/InstancedFoliageActor.h
@@ -23,7 +23,9 @@ struct FDesiredFoliageInstance;
 UCLASS(notplaceable, hidecategories = Object, MinimalAPI, NotBlueprintable)
 class AInstancedFoliageActor : public AActor
 {
-	GENERATED_UCLASS_BODY()
+	GENERATED_BODY()
+public:
+	FOLIAGE_API AInstancedFoliageActor(const FObjectInitializer& ObjectInitializer = FObjectInitializer::Get());
 
 public:
 #if WITH_EDITORONLY_DATA
diff --git a/Engine/Source/Runtime/Foliage/Public/InteractiveFoliageActor.h b/Engine/Source/Runtime/Foliage/Public/InteractiveFoliageActor.h
--- a/Engine/Source/Runtime/Foliage/Public/InteractiveFoliageActor.h
+++ b/Engine/Source/Runtime/Foliage/Public/InteractiveFoliageActor.h
@@ -10,7 +10,9 @@
 UCLASS(MinimalAPI)
 class AInteractiveFoliageActor : public AStaticMeshActor
 {
-	GENERATED_UCLASS_BODY()
+	GENERATED_BODY()
+public:
+	FOLIAGE_API AInteractiveFoliageActor(const FObjectInitializer& ObjectInitializer = FObjectInitializer::Get());
 
 private:
 	/** Collision cylinder */
diff --git a/Engine/Source/Runtime/Foliage/Public/ProceduralFoliage.h b/Engine/Source/Runtime/Foliage/Public/ProceduralFoliage.h
--- a/Engine/Source/Runtime/Foliage/Public/ProceduralFoliage.h
+++ b/Engine/Source/Runtime/Foliage/Public/ProceduralFoliage.h
@@ -12,7 +12,7 @@ class UProceduralFoliageTile;
 USTRUCT()
 struct FProceduralFoliageTypeData
 {
-	GENERATED_USTRUCT_BODY()
+	GENERATED_BODY()
 
 	UPROPERTY(Category = ProceduralFoliageSimulation, EditAnywhere, BlueprintReadOnly)
 	TSubclassOf<UFoliageType_InstancedStaticMesh> Type;
@@ -27,7 +27,9 @@ struct FProceduralFoliageTypeData
 UCLASS(BlueprintType, Blueprintable)
 class FOLIAGE_API UProceduralFoliage : public UObject
 {
-	GENERATED_UCLASS_BODY()
+	GENERATED_BODY()
+public:
+	UProceduralFoliage(const FObjectInitializer& ObjectInitializer = FObjectInitializer::Get());
 
 	/**The seed used for generating the randomness of the simulation. */
 	UPROPERTY(Category = ProceduralFoliageSimulation, EditAnywhere, BlueprintReadOnly)
diff --git a/Engine/Source/Runtime/Foliage/Public/ProceduralFoliageActor.h b/Engine/Source/Runtime/Foliage/Public/ProceduralFoliageActor.h
--- a/Engine/Source/Runtime/Foliage/Public/ProceduralFoliageActor.h
+++ b/Engine/Source/Runtime/Foliage/Public/ProceduralFoliageActor.h
@@ -9,7 +9,9 @@ class UProceduralFoliageComponent;
 UCLASS()
 class FOLIAGE_API AProceduralFoliageActor: public AVolume
 {
-	GENERATED_UCLASS_BODY()
+	GENERATED_BODY()
+public:
+	AProceduralFoliageActor(const FObjectInitializer& ObjectInitializer = FObjectInitializer::Get());
 
 	UPROPERTY(Category = ProceduralFoliageTile, VisibleAnywhere, BlueprintReadOnly)
 	UProceduralFoliageComponent* ProceduralComponent;
@@ -17,4 +19,4 @@ class FOLIAGE_API AProceduralFoliageActor: public AVolume
 #if WITH_EDITOR
 	virtual bool GetReferencedContentObjects(TArray<UObject*>& Objects) const override;
 #endif
-};
\ No newline at end of file
+};
diff --git a/Engine/Source/Runtime/Foliage/Public/ProceduralFoliageComponent.h b/Engine/Source/Runtime/Foliage/Public/ProceduralFoliageComponent.h
--- a/Engine/Source/Runtime/Foliage/Public/ProceduralFoliageComponent.h
+++ b/Engine/Source/Runtime/Foliage/Public/ProceduralFoliageComponent.h
@@ -12,7 +12,9 @@ struct FDesiredFoliageInstance;
 UCLASS(BlueprintType)
 class FOLIAGE_API UProceduralFoliageComponent : public UActorComponent
 {
-	GENERATED_UCLASS_BODY()
+	GENERATED_BODY()
+public:
+	UProceduralFoliageComponent(const FObjectInitializer& ObjectInitializer = FObjectInitializer::Get());
 
 	/** The overlap in Cms between two tiles. */
 	UPROPERTY(Category = "ProceduralFoliage", BlueprintReadWrite, EditAnywhere)
@@ -43,4 +45,4 @@ private:
 private:
 	UPROPERTY()
 	FGuid ProceduralGuid;
-};
\ No newline at end of file
+};
diff --git a/Engine/Source/Runtime/Foliage/Public/ProceduralFoliageInstance.h b/Engine/Source/Runtime/Foliage/Public/ProceduralFoliageInstance.h
--- a/Engine/Source/Runtime/Foliage/Public/ProceduralFoliageInstance.h
+++ b/Engine/Source/Runtime/Foliage/Public/ProceduralFoliageInstance.h
@@ -61,8 +61,8 @@ struct FProceduralFoliageOverlap
 USTRUCT(BlueprintType)
 struct FOLIAGE_API FProceduralFoliageInstance
 {
+	GENERATED_BODY()
 public:
-	GENERATED_USTRUCT_BODY()
 	FProceduralFoliageInstance();
 	FProceduralFoliageInstance(const FProceduralFoliageInstance& Other);
 
@@ -103,4 +103,4 @@ public:
 	UActorComponent* BaseComponent;
 private:
 	bool bAlive;
-};
\ No newline at end of file
+};
diff --git a/Engine/Source/Runtime/Foliage/Public/ProceduralFoliageTile.h b/Engine/Source/Runtime/Foliage/Public/ProceduralFoliageTile.h
--- a/Engine/Source/Runtime/Foliage/Public/ProceduralFoliageTile.h
+++ b/Engine/Source/Runtime/Foliage/Public/ProceduralFoliageTile.h
@@ -16,7 +16,9 @@ class UFoliageType_InstancedStaticMesh;
 UCLASS()
 class FOLIAGE_API UProceduralFoliageTile : public UObject
 {
-	GENERATED_UCLASS_BODY()
+	GENERATED_BODY()
+public:
+	UProceduralFoliageTile(const FObjectInitializer& ObjectInitializer = FObjectInitializer::Get());
 
 	UPROPERTY()
 	const UProceduralFoliage* ProceduralFoliage;
@@ -79,4 +81,4 @@ private:
 private:
 	float GetRandomGaussian();
 	FVector GetSeedOffset(const UFoliageType_InstancedStaticMesh* Type, float MinDistance);
-};
\ No newline at end of file
+};
diff --git a/Engine/Source/Runtime/GameLiveStreaming/GameLiveStreamingFunctionLibrary.h b/Engine/Source/Runtime/GameLiveStreaming/GameLiveStreamingFunctionLibrary.h
--- a/Engine/Source/Runtime/GameLiveStreaming/GameLiveStreamingFunctionLibrary.h
+++ b/Engine/Source/Runtime/GameLiveStreaming/GameLiveStreamingFunctionLibrary.h
@@ -5,7 +5,9 @@
 UCLASS()
 class UGameLiveStreamingFunctionLibrary : public UBlueprintFunctionLibrary
 {
-	GENERATED_UCLASS_BODY()
+	GENERATED_BODY()
+public:
+	UGameLiveStreamingFunctionLibrary(const FObjectInitializer& ObjectInitializer = FObjectInitializer::Get());
 
 
 	/**
diff --git a/Engine/Source/Runtime/GameLiveStreaming/QueryLiveStreamsCallbackProxy.h b/Engine/Source/Runtime/GameLiveStreaming/QueryLiveStreamsCallbackProxy.h
--- a/Engine/Source/Runtime/GameLiveStreaming/QueryLiveStreamsCallbackProxy.h
+++ b/Engine/Source/Runtime/GameLiveStreaming/QueryLiveStreamsCallbackProxy.h
@@ -6,7 +6,7 @@
 USTRUCT( BlueprintType )
 struct FBlueprintLiveStreamInfo 
 {
-	GENERATED_USTRUCT_BODY()
+	GENERATED_BODY()
 
 	/** Name of the game that is streaming */
 	UPROPERTY( EditAnywhere, Category="LiveStreaming" )
@@ -27,7 +27,9 @@ DECLARE_DYNAMIC_MULTICAST_DELEGATE_TwoParams( FOnQueryLiveStreamsCompleted, cons
 UCLASS()
 class UQueryLiveStreamsCallbackProxy : public UOnlineBlueprintCallProxyBase
 {
-	GENERATED_UCLASS_BODY()
+	GENERATED_BODY()
+public:
+	UQueryLiveStreamsCallbackProxy(const FObjectInitializer& ObjectInitializer = FObjectInitializer::Get());
 
 	// Called when the asynchronous request for live streams completes, passing along the list of live streams currently active, along with a boolean value that indicates whether the request was successful at all
 	UPROPERTY( BlueprintAssignable )
@@ -45,4 +47,4 @@ protected:
 
 	/** Name of the game that we're querying about */
 	FString GameName;
-};
\ No newline at end of file
+};
diff --git a/Engine/Source/Runtime/GameMenuBuilder/Public/GameMenuWidgetStyle.h b/Engine/Source/Runtime/GameMenuBuilder/Public/GameMenuWidgetStyle.h
--- a/Engine/Source/Runtime/GameMenuBuilder/Public/GameMenuWidgetStyle.h
+++ b/Engine/Source/Runtime/GameMenuBuilder/Public/GameMenuWidgetStyle.h
@@ -20,7 +20,7 @@ namespace GameMenuLayoutType
 USTRUCT()
 struct FGameMenuStyle : public FSlateWidgetStyle
 {
-	GENERATED_USTRUCT_BODY()
+	GENERATED_BODY()
 
 	FGameMenuStyle();
 	virtual ~FGameMenuStyle();
@@ -178,7 +178,9 @@ struct FGameMenuStyle : public FSlateWidgetStyle
 UCLASS(hidecategories = Object, MinimalAPI)
 class UGameMenuWidgetStyle : public USlateWidgetStyleContainerBase
 {
-	GENERATED_UCLASS_BODY()
+	GENERATED_BODY()
+public:
+	GAMEMENUBUILDER_API UGameMenuWidgetStyle(const FObjectInitializer& ObjectInitializer = FObjectInitializer::Get());
 
 public:
 	/** The actual data describing the menu's appearance. */
@@ -189,4 +191,4 @@ public:
 	{
 		return static_cast< const struct FSlateWidgetStyle* >(&MenuStyle);
 	}
-};
\ No newline at end of file
+};
diff --git a/Engine/Source/Runtime/GameplayAbilities/Classes/GameplayAbilityBlueprint.h b/Engine/Source/Runtime/GameplayAbilities/Classes/GameplayAbilityBlueprint.h
--- a/Engine/Source/Runtime/GameplayAbilities/Classes/GameplayAbilityBlueprint.h
+++ b/Engine/Source/Runtime/GameplayAbilities/Classes/GameplayAbilityBlueprint.h
@@ -13,7 +13,9 @@
 UCLASS(BlueprintType)
 class GAMEPLAYABILITIES_API UGameplayAbilityBlueprint : public UBlueprint
 {
-	GENERATED_UCLASS_BODY()
+	GENERATED_BODY()
+public:
+	UGameplayAbilityBlueprint(const FObjectInitializer& ObjectInitializer = FObjectInitializer::Get());
 
 #if WITH_EDITOR
 
diff --git a/Engine/Source/Runtime/GameplayAbilities/Public/Abilities/GameplayAbility.h b/Engine/Source/Runtime/GameplayAbilities/Public/Abilities/GameplayAbility.h
--- a/Engine/Source/Runtime/GameplayAbilities/Public/Abilities/GameplayAbility.h
+++ b/Engine/Source/Runtime/GameplayAbilities/Public/Abilities/GameplayAbility.h
@@ -77,7 +77,7 @@ DECLARE_MULTICAST_DELEGATE_OneParam(FOnGameplayAbilityStateEnded, FName);
 USTRUCT()
 struct FAbilityTriggerData
 {
-	GENERATED_USTRUCT_BODY()
+	GENERATED_BODY()
 
 	FAbilityTriggerData() 
 	: TriggerSource(EGameplayAbilityTriggerSource::GameplayEvent)
@@ -98,7 +98,9 @@ struct FAbilityTriggerData
 UCLASS(Blueprintable)
 class GAMEPLAYABILITIES_API UGameplayAbility : public UObject
 {
-	GENERATED_UCLASS_BODY()
+	GENERATED_BODY()
+public:
+	UGameplayAbility(const FObjectInitializer& ObjectInitializer = FObjectInitializer::Get());
 
 	friend class UAbilitySystemComponent;
 	friend class UGameplayAbilitySet;
diff --git a/Engine/Source/Runtime/GameplayAbilities/Public/Abilities/GameplayAbilityTargetActor.h b/Engine/Source/Runtime/GameplayAbilities/Public/Abilities/GameplayAbilityTargetActor.h
--- a/Engine/Source/Runtime/GameplayAbilities/Public/Abilities/GameplayAbilityTargetActor.h
+++ b/Engine/Source/Runtime/GameplayAbilities/Public/Abilities/GameplayAbilityTargetActor.h
@@ -20,7 +20,9 @@ namespace ETargetAbilitySelfSelection
 UCLASS(Blueprintable, abstract, notplaceable)
 class GAMEPLAYABILITIES_API AGameplayAbilityTargetActor : public AActor
 {
-	GENERATED_UCLASS_BODY()
+	GENERATED_BODY()
+public:
+	AGameplayAbilityTargetActor(const FObjectInitializer& ObjectInitializer = FObjectInitializer::Get());
 
 public:
 
@@ -97,4 +99,4 @@ public:
 
 	FDelegateHandle GenericConfirmHandle;
 	FDelegateHandle GenericCancelHandle;
-};
\ No newline at end of file
+};
diff --git a/Engine/Source/Runtime/GameplayAbilities/Public/Abilities/GameplayAbilityTargetActor_ActorPlacement.h b/Engine/Source/Runtime/GameplayAbilities/Public/Abilities/GameplayAbilityTargetActor_ActorPlacement.h
--- a/Engine/Source/Runtime/GameplayAbilities/Public/Abilities/GameplayAbilityTargetActor_ActorPlacement.h
+++ b/Engine/Source/Runtime/GameplayAbilities/Public/Abilities/GameplayAbilityTargetActor_ActorPlacement.h
@@ -9,7 +9,9 @@
 UCLASS(Blueprintable)
 class AGameplayAbilityTargetActor_ActorPlacement : public AGameplayAbilityTargetActor_GroundTrace
 {
-	GENERATED_UCLASS_BODY()
+	GENERATED_BODY()
+public:
+	AGameplayAbilityTargetActor_ActorPlacement(const FObjectInitializer& ObjectInitializer = FObjectInitializer::Get());
 
 	virtual void EndPlay(const EEndPlayReason::Type EndPlayReason) override;
 
@@ -25,4 +27,4 @@ class AGameplayAbilityTargetActor_ActorPlacement : public AGameplayAbilityTarget
 
 	/** Visualization for the intended location of the placed actor. */
 	TWeakObjectPtr<AGameplayAbilityWorldReticle_ActorVisualization> ActorVisualizationReticle;
-};
\ No newline at end of file
+};
diff --git a/Engine/Source/Runtime/GameplayAbilities/Public/Abilities/GameplayAbilityTargetActor_GroundTrace.h b/Engine/Source/Runtime/GameplayAbilities/Public/Abilities/GameplayAbilityTargetActor_GroundTrace.h
--- a/Engine/Source/Runtime/GameplayAbilities/Public/Abilities/GameplayAbilityTargetActor_GroundTrace.h
+++ b/Engine/Source/Runtime/GameplayAbilities/Public/Abilities/GameplayAbilityTargetActor_GroundTrace.h
@@ -9,7 +9,9 @@
 UCLASS(Blueprintable)
 class GAMEPLAYABILITIES_API AGameplayAbilityTargetActor_GroundTrace : public AGameplayAbilityTargetActor_Trace
 {
-	GENERATED_UCLASS_BODY()
+	GENERATED_BODY()
+public:
+	AGameplayAbilityTargetActor_GroundTrace(const FObjectInitializer& ObjectInitializer = FObjectInitializer::Get());
 
 	virtual void StartTargeting(UGameplayAbility* InAbility) override;
 
@@ -33,4 +35,4 @@ protected:
 	FCollisionShape CollisionShape;
 	float CollisionHeightOffset;		//When tracing, give this much extra height to avoid start-in-ground problems. Dealing with thick placement actors while standing near walls may be trickier.
 	bool bLastTraceWasGood;
-};
\ No newline at end of file
+};
diff --git a/Engine/Source/Runtime/GameplayAbilities/Public/Abilities/GameplayAbilityTargetActor_Radius.h b/Engine/Source/Runtime/GameplayAbilities/Public/Abilities/GameplayAbilityTargetActor_Radius.h
--- a/Engine/Source/Runtime/GameplayAbilities/Public/Abilities/GameplayAbilityTargetActor_Radius.h
+++ b/Engine/Source/Runtime/GameplayAbilities/Public/Abilities/GameplayAbilityTargetActor_Radius.h
@@ -11,7 +11,9 @@
 UCLASS(Blueprintable, notplaceable)
 class GAMEPLAYABILITIES_API AGameplayAbilityTargetActor_Radius : public AGameplayAbilityTargetActor
 {
-	GENERATED_UCLASS_BODY()
+	GENERATED_BODY()
+public:
+	AGameplayAbilityTargetActor_Radius(const FObjectInitializer& ObjectInitializer = FObjectInitializer::Get());
 
 public:
 
@@ -29,4 +31,4 @@ protected:
 
 	FGameplayAbilityTargetDataHandle MakeTargetData(const TArray<TWeakObjectPtr<AActor> > Actors, const FVector& Origin) const;
 
-};
\ No newline at end of file
+};
diff --git a/Engine/Source/Runtime/GameplayAbilities/Public/Abilities/GameplayAbilityTargetActor_SingleLineTrace.h b/Engine/Source/Runtime/GameplayAbilities/Public/Abilities/GameplayAbilityTargetActor_SingleLineTrace.h
--- a/Engine/Source/Runtime/GameplayAbilities/Public/Abilities/GameplayAbilityTargetActor_SingleLineTrace.h
+++ b/Engine/Source/Runtime/GameplayAbilities/Public/Abilities/GameplayAbilityTargetActor_SingleLineTrace.h
@@ -9,8 +9,10 @@
 UCLASS(Blueprintable)
 class GAMEPLAYABILITIES_API AGameplayAbilityTargetActor_SingleLineTrace : public AGameplayAbilityTargetActor_Trace
 {
-	GENERATED_UCLASS_BODY()
+	GENERATED_BODY()
+public:
+	AGameplayAbilityTargetActor_SingleLineTrace(const FObjectInitializer& ObjectInitializer = FObjectInitializer::Get());
 
 protected:
 	virtual FHitResult PerformTrace(AActor* InSourceActor) override;
-};
\ No newline at end of file
+};
diff --git a/Engine/Source/Runtime/GameplayAbilities/Public/Abilities/GameplayAbilityTargetActor_Trace.h b/Engine/Source/Runtime/GameplayAbilities/Public/Abilities/GameplayAbilityTargetActor_Trace.h
--- a/Engine/Source/Runtime/GameplayAbilities/Public/Abilities/GameplayAbilityTargetActor_Trace.h
+++ b/Engine/Source/Runtime/GameplayAbilities/Public/Abilities/GameplayAbilityTargetActor_Trace.h
@@ -10,7 +10,9 @@
 UCLASS(Abstract, Blueprintable, notplaceable, config=Game)
 class GAMEPLAYABILITIES_API AGameplayAbilityTargetActor_Trace : public AGameplayAbilityTargetActor
 {
-	GENERATED_UCLASS_BODY()
+	GENERATED_BODY()
+public:
+	AGameplayAbilityTargetActor_Trace(const FObjectInitializer& ObjectInitializer = FObjectInitializer::Get());
 
 public:
 	virtual void EndPlay(const EEndPlayReason::Type EndPlayReason) override;
@@ -47,4 +49,4 @@ protected:
 	FGameplayAbilityTargetDataHandle MakeTargetData(FHitResult HitResult) const;
 
 	TWeakObjectPtr<AGameplayAbilityWorldReticle> ReticleActor;
-};
\ No newline at end of file
+};
diff --git a/Engine/Source/Runtime/GameplayAbilities/Public/Abilities/GameplayAbilityTargetTypes.h b/Engine/Source/Runtime/GameplayAbilities/Public/Abilities/GameplayAbilityTargetTypes.h
--- a/Engine/Source/Runtime/GameplayAbilities/Public/Abilities/GameplayAbilityTargetTypes.h
+++ b/Engine/Source/Runtime/GameplayAbilities/Public/Abilities/GameplayAbilityTargetTypes.h
@@ -59,7 +59,7 @@ namespace EGameplayTargetingConfirmation
 USTRUCT()
 struct GAMEPLAYABILITIES_API FGameplayAbilityTargetData
 {
-	GENERATED_USTRUCT_BODY()
+	GENERATED_BODY()
 
 	virtual ~FGameplayAbilityTargetData() { }
 
@@ -160,7 +160,7 @@ namespace EGameplayAbilityTargetingLocationType
 USTRUCT(BlueprintType)
 struct GAMEPLAYABILITIES_API FGameplayAbilityTargetDataHandle
 {
-	GENERATED_USTRUCT_BODY()
+	GENERATED_BODY()
 
 	FGameplayAbilityTargetDataHandle() { }
 	FGameplayAbilityTargetDataHandle(struct FGameplayAbilityTargetData* DataPtr)
@@ -248,7 +248,7 @@ struct TStructOpsTypeTraits<FGameplayAbilityTargetDataHandle> : public TStructOp
 USTRUCT(BlueprintType)
 struct GAMEPLAYABILITIES_API FGameplayAbilityTargetingLocationInfo
 {
-	GENERATED_USTRUCT_BODY()
+	GENERATED_BODY()
 
 	FGameplayAbilityTargetingLocationInfo()
 	: LocationType(EGameplayAbilityTargetingLocationType::LiteralTransform)
@@ -345,7 +345,7 @@ struct TStructOpsTypeTraits<FGameplayAbilityTargetingLocationInfo> : public TStr
 USTRUCT(BlueprintType)
 struct GAMEPLAYABILITIES_API FGameplayAbilityTargetData_LocationInfo : public FGameplayAbilityTargetData
 {
-	GENERATED_USTRUCT_BODY()
+	GENERATED_BODY()
 
 	/** Generic location data for source */
 	UPROPERTY(BlueprintReadWrite, EditAnywhere, Category = Targeting)
@@ -406,7 +406,7 @@ struct TStructOpsTypeTraits<FGameplayAbilityTargetData_LocationInfo> : public TS
 USTRUCT(BlueprintType)
 struct GAMEPLAYABILITIES_API FGameplayAbilityTargetData_ActorArray : public FGameplayAbilityTargetData
 {
-	GENERATED_USTRUCT_BODY()
+	GENERATED_BODY()
 
 	/** We could be selecting this group of actors from any type of location, so use a generic location type */
 	UPROPERTY(BlueprintReadWrite, EditAnywhere, Category = Targeting)
@@ -508,7 +508,7 @@ struct TStructOpsTypeTraits<FGameplayAbilityTargetData_ActorArray> : public TStr
 USTRUCT(BlueprintType)
 struct GAMEPLAYABILITIES_API FGameplayAbilityTargetData_SingleTargetHit : public FGameplayAbilityTargetData
 {
-	GENERATED_USTRUCT_BODY()
+	GENERATED_BODY()
 
 	FGameplayAbilityTargetData_SingleTargetHit()
 	{ }
@@ -670,4 +670,4 @@ struct FAbilityReplicatedDataCache
 		}
 
 	}
-};
\ No newline at end of file
+};
diff --git a/Engine/Source/Runtime/GameplayAbilities/Public/Abilities/GameplayAbilityTypes.h b/Engine/Source/Runtime/GameplayAbilities/Public/Abilities/GameplayAbilityTypes.h
--- a/Engine/Source/Runtime/GameplayAbilities/Public/Abilities/GameplayAbilityTypes.h
+++ b/Engine/Source/Runtime/GameplayAbilities/Public/Abilities/GameplayAbilityTypes.h
@@ -115,7 +115,7 @@ namespace EGameplayAbilityTriggerSource
 USTRUCT(BlueprintType)
 struct GAMEPLAYABILITIES_API FGameplayAbilityActorInfo
 {
-	GENERATED_USTRUCT_BODY()
+	GENERATED_BODY()
 
 	virtual ~FGameplayAbilityActorInfo() {}
 
@@ -164,7 +164,7 @@ struct GAMEPLAYABILITIES_API FGameplayAbilityActorInfo
 USTRUCT()
 struct FGameplayAbilitySpecHandleAndPredictionKey
 {
-	GENERATED_USTRUCT_BODY()
+	GENERATED_BODY()
 
 	FGameplayAbilitySpecHandleAndPredictionKey()
 	: PredictionKeyAtCreation(0)
@@ -206,7 +206,7 @@ private:
 USTRUCT()
 struct GAMEPLAYABILITIES_API FGameplayAbilityRepAnimMontage
 {
-	GENERATED_USTRUCT_BODY()
+	GENERATED_BODY()
 
 	/** AnimMontage ref */
 	UPROPERTY()
@@ -255,7 +255,7 @@ struct GAMEPLAYABILITIES_API FGameplayAbilityRepAnimMontage
 USTRUCT()
 struct GAMEPLAYABILITIES_API FGameplayAbilityLocalAnimMontage
 {
-	GENERATED_USTRUCT_BODY()
+	GENERATED_BODY()
 
 	FGameplayAbilityLocalAnimMontage()
 	: AnimMontage(nullptr), PlayBit(false), AnimatingAbility(nullptr)
@@ -281,7 +281,7 @@ struct GAMEPLAYABILITIES_API FGameplayAbilityLocalAnimMontage
 USTRUCT(BlueprintType)
 struct GAMEPLAYABILITIES_API FGameplayEventData
 {
-	GENERATED_USTRUCT_BODY()
+	GENERATED_BODY()
 
 	FGameplayEventData()
 	: Instigator(nullptr)
@@ -364,7 +364,7 @@ struct FGameplayAbiliyInputBinds
 USTRUCT()
 struct GAMEPLAYABILITIES_API FAttributeDefaults
 {
-	GENERATED_USTRUCT_BODY()
+	GENERATED_BODY()
 
 	UPROPERTY(EditAnywhere, Category = "AttributeTest")
 	TSubclassOf<UAttributeSet> Attributes;
diff --git a/Engine/Source/Runtime/GameplayAbilities/Public/Abilities/GameplayAbilityWorldReticle.h b/Engine/Source/Runtime/GameplayAbilities/Public/Abilities/GameplayAbilityWorldReticle.h
--- a/Engine/Source/Runtime/GameplayAbilities/Public/Abilities/GameplayAbilityWorldReticle.h
+++ b/Engine/Source/Runtime/GameplayAbilities/Public/Abilities/GameplayAbilityWorldReticle.h
@@ -8,7 +8,7 @@ class AGameplayAbilityTargetActor;
 USTRUCT(BlueprintType)
 struct FWorldReticleParameters
 {
-	GENERATED_USTRUCT_BODY()
+	GENERATED_BODY()
 
 	//Use this so that we can't slip in new parameters without some actor knowing about it.
 	void Initialize(FVector InAOEScale)
@@ -24,7 +24,9 @@ struct FWorldReticleParameters
 UCLASS(Blueprintable, notplaceable)
 class GAMEPLAYABILITIES_API AGameplayAbilityWorldReticle : public AActor
 {
-	GENERATED_UCLASS_BODY()
+	GENERATED_BODY()
+public:
+	AGameplayAbilityWorldReticle(const FObjectInitializer& ObjectInitializer = FObjectInitializer::Get());
 
 public:
 
diff --git a/Engine/Source/Runtime/GameplayAbilities/Public/Abilities/GameplayAbilityWorldReticle_ActorVisualization.h b/Engine/Source/Runtime/GameplayAbilities/Public/Abilities/GameplayAbilityWorldReticle_ActorVisualization.h
--- a/Engine/Source/Runtime/GameplayAbilities/Public/Abilities/GameplayAbilityWorldReticle_ActorVisualization.h
+++ b/Engine/Source/Runtime/GameplayAbilities/Public/Abilities/GameplayAbilityWorldReticle_ActorVisualization.h
@@ -11,7 +11,9 @@ class AGameplayAbilityTargetActor;
 UCLASS(notplaceable)
 class AGameplayAbilityWorldReticle_ActorVisualization : public AGameplayAbilityWorldReticle
 {
-	GENERATED_UCLASS_BODY()
+	GENERATED_BODY()
+public:
+	AGameplayAbilityWorldReticle_ActorVisualization(const FObjectInitializer& ObjectInitializer = FObjectInitializer::Get());
 
 public:
 	void InitializeReticleVisualizationInformation(AGameplayAbilityTargetActor* InTargetingActor, AActor* VisualizationActor, UMaterialInterface *VisualizationMaterial);
diff --git a/Engine/Source/Runtime/GameplayAbilities/Public/Abilities/GameplayAbility_CharacterJump.h b/Engine/Source/Runtime/GameplayAbilities/Public/Abilities/GameplayAbility_CharacterJump.h
--- a/Engine/Source/Runtime/GameplayAbilities/Public/Abilities/GameplayAbility_CharacterJump.h
+++ b/Engine/Source/Runtime/GameplayAbilities/Public/Abilities/GameplayAbility_CharacterJump.h
@@ -11,7 +11,9 @@
 UCLASS()
 class GAMEPLAYABILITIES_API UGameplayAbility_CharacterJump : public UGameplayAbility
 {
-	GENERATED_UCLASS_BODY()
+	GENERATED_BODY()
+public:
+	UGameplayAbility_CharacterJump(const FObjectInitializer& ObjectInitializer = FObjectInitializer::Get());
 
 public:
 
@@ -22,4 +24,4 @@ public:
 	virtual void InputReleased(const FGameplayAbilitySpecHandle Handle, const FGameplayAbilityActorInfo* ActorInfo, const FGameplayAbilityActivationInfo ActivationInfo) override;
 
 	virtual void CancelAbility(const FGameplayAbilitySpecHandle Handle, const FGameplayAbilityActorInfo* ActorInfo, const FGameplayAbilityActivationInfo ActivationInfo, bool bReplicateCancelAbility) override;
-};
\ No newline at end of file
+};
diff --git a/Engine/Source/Runtime/GameplayAbilities/Public/Abilities/GameplayAbility_Montage.h b/Engine/Source/Runtime/GameplayAbilities/Public/Abilities/GameplayAbility_Montage.h
--- a/Engine/Source/Runtime/GameplayAbilities/Public/Abilities/GameplayAbility_Montage.h
+++ b/Engine/Source/Runtime/GameplayAbilities/Public/Abilities/GameplayAbility_Montage.h
@@ -11,7 +11,9 @@
 UCLASS()
 class GAMEPLAYABILITIES_API UGameplayAbility_Montage : public UGameplayAbility
 {
-	GENERATED_UCLASS_BODY()
+	GENERATED_BODY()
+public:
+	UGameplayAbility_Montage(const FObjectInitializer& ObjectInitializer = FObjectInitializer::Get());
 
 public:
 	
@@ -39,4 +41,4 @@ public:
 	void GetGameplayEffectsWhileAnimating(TArray<const UGameplayEffect*>& OutEffects) const;
 
 	virtual void ConvertDeprecatedGameplayEffectReferencesToBlueprintReferences(UGameplayEffect* OldGE, TSubclassOf<UGameplayEffect> NewGEClass) override;
-};
\ No newline at end of file
+};
diff --git a/Engine/Source/Runtime/GameplayAbilities/Public/Abilities/Tasks/AbilityTask.h b/Engine/Source/Runtime/GameplayAbilities/Public/Abilities/Tasks/AbilityTask.h
--- a/Engine/Source/Runtime/GameplayAbilities/Public/Abilities/Tasks/AbilityTask.h
+++ b/Engine/Source/Runtime/GameplayAbilities/Public/Abilities/Tasks/AbilityTask.h
@@ -60,7 +60,9 @@
 UCLASS()
 class GAMEPLAYABILITIES_API UAbilityTask : public UObject
 {
-	GENERATED_UCLASS_BODY()
+	GENERATED_BODY()
+public:
+	UAbilityTask(const FObjectInitializer& ObjectInitializer = FObjectInitializer::Get());
 
 	DECLARE_DYNAMIC_MULTICAST_DELEGATE(FGenericAbilityTaskDelegate);
 
diff --git a/Engine/Source/Runtime/GameplayAbilities/Public/Abilities/Tasks/AbilityTask_MoveToLocation.h b/Engine/Source/Runtime/GameplayAbilities/Public/Abilities/Tasks/AbilityTask_MoveToLocation.h
--- a/Engine/Source/Runtime/GameplayAbilities/Public/Abilities/Tasks/AbilityTask_MoveToLocation.h
+++ b/Engine/Source/Runtime/GameplayAbilities/Public/Abilities/Tasks/AbilityTask_MoveToLocation.h
@@ -23,7 +23,9 @@ class UPrimitiveComponent;
 UCLASS(MinimalAPI)
 class UAbilityTask_MoveToLocation : public UAbilityTask
 {
-	GENERATED_UCLASS_BODY()
+	GENERATED_BODY()
+public:
+	GAMEPLAYABILITIES_API UAbilityTask_MoveToLocation(const FObjectInitializer& ObjectInitializer = FObjectInitializer::Get());
 
 	UPROPERTY(BlueprintAssignable)
 	FMoveToLocationDelegate		OnTargetLocationReached;
@@ -65,4 +67,4 @@ protected:
 	UCurveFloat* LerpCurve;
 
 	
-};
\ No newline at end of file
+};
diff --git a/Engine/Source/Runtime/GameplayAbilities/Public/Abilities/Tasks/AbilityTask_NetworkSyncPoint.h b/Engine/Source/Runtime/GameplayAbilities/Public/Abilities/Tasks/AbilityTask_NetworkSyncPoint.h
--- a/Engine/Source/Runtime/GameplayAbilities/Public/Abilities/Tasks/AbilityTask_NetworkSyncPoint.h
+++ b/Engine/Source/Runtime/GameplayAbilities/Public/Abilities/Tasks/AbilityTask_NetworkSyncPoint.h
@@ -24,7 +24,9 @@ enum class EAbilityTaskNetSyncType : uint8
 UCLASS(MinimalAPI)
 class UAbilityTask_NetworkSyncPoint : public UAbilityTask
 {
-	GENERATED_UCLASS_BODY()
+	GENERATED_BODY()
+public:
+	GAMEPLAYABILITIES_API UAbilityTask_NetworkSyncPoint(const FObjectInitializer& ObjectInitializer = FObjectInitializer::Get());
 
 	UPROPERTY(BlueprintAssignable)
 	FNetworkSyncDelegate	OnSync;
@@ -61,4 +63,4 @@ protected:
 	EAbilityGenericReplicatedEvent::Type	ReplicatedEventToListenFor;
 
 	EAbilityTaskNetSyncType SyncType;
-};
\ No newline at end of file
+};
diff --git a/Engine/Source/Runtime/GameplayAbilities/Public/Abilities/Tasks/AbilityTask_PlayMontageAndWait.h b/Engine/Source/Runtime/GameplayAbilities/Public/Abilities/Tasks/AbilityTask_PlayMontageAndWait.h
--- a/Engine/Source/Runtime/GameplayAbilities/Public/Abilities/Tasks/AbilityTask_PlayMontageAndWait.h
+++ b/Engine/Source/Runtime/GameplayAbilities/Public/Abilities/Tasks/AbilityTask_PlayMontageAndWait.h
@@ -8,7 +8,9 @@ DECLARE_DYNAMIC_MULTICAST_DELEGATE(FMontageWaitSimpleDelegate);
 UCLASS(MinimalAPI)
 class UAbilityTask_PlayMontageAndWait : public UAbilityTask
 {
-	GENERATED_UCLASS_BODY()
+	GENERATED_BODY()
+public:
+	GAMEPLAYABILITIES_API UAbilityTask_PlayMontageAndWait(const FObjectInitializer& ObjectInitializer = FObjectInitializer::Get());
 
 	UPROPERTY(BlueprintAssignable)
 	FMontageWaitSimpleDelegate	OnComplete;
@@ -47,4 +49,4 @@ private:
 	UAnimMontage* MontageToPlay;
 	float Rate;	
 	FName StartSection;
-};
\ No newline at end of file
+};
diff --git a/Engine/Source/Runtime/GameplayAbilities/Public/Abilities/Tasks/AbilityTask_Repeat.h b/Engine/Source/Runtime/GameplayAbilities/Public/Abilities/Tasks/AbilityTask_Repeat.h
--- a/Engine/Source/Runtime/GameplayAbilities/Public/Abilities/Tasks/AbilityTask_Repeat.h
+++ b/Engine/Source/Runtime/GameplayAbilities/Public/Abilities/Tasks/AbilityTask_Repeat.h
@@ -13,7 +13,9 @@ DECLARE_DYNAMIC_MULTICAST_DELEGATE_OneParam(FRepeatedActionDelegate, int32, Acti
 UCLASS()
 class GAMEPLAYABILITIES_API UAbilityTask_Repeat : public UAbilityTask
 {
-	GENERATED_UCLASS_BODY()
+	GENERATED_BODY()
+public:
+	UAbilityTask_Repeat(const FObjectInitializer& ObjectInitializer = FObjectInitializer::Get());
 
 	UPROPERTY(BlueprintAssignable)
 	FRepeatedActionDelegate		OnPerformAction;
@@ -40,4 +42,4 @@ protected:
 	FTimerHandle TimerHandle_PerformAction;
 
 	virtual void OnDestroy(bool AbilityIsEnding) override;
-};
\ No newline at end of file
+};
diff --git a/Engine/Source/Runtime/GameplayAbilities/Public/Abilities/Tasks/AbilityTask_SpawnActor.h b/Engine/Source/Runtime/GameplayAbilities/Public/Abilities/Tasks/AbilityTask_SpawnActor.h
--- a/Engine/Source/Runtime/GameplayAbilities/Public/Abilities/Tasks/AbilityTask_SpawnActor.h
+++ b/Engine/Source/Runtime/GameplayAbilities/Public/Abilities/Tasks/AbilityTask_SpawnActor.h
@@ -23,7 +23,9 @@ DECLARE_DYNAMIC_MULTICAST_DELEGATE_OneParam(FSpawnActorDelegate, AActor*, Spawne
 UCLASS(MinimalAPI)
 class UAbilityTask_SpawnActor: public UAbilityTask
 {
-	GENERATED_UCLASS_BODY()
+	GENERATED_BODY()
+public:
+	GAMEPLAYABILITIES_API UAbilityTask_SpawnActor(const FObjectInitializer& ObjectInitializer = FObjectInitializer::Get());
 
 	UPROPERTY(BlueprintAssignable)
 	FSpawnActorDelegate	Success;
@@ -44,4 +46,4 @@ class UAbilityTask_SpawnActor: public UAbilityTask
 
 protected:
 	FGameplayAbilityTargetDataHandle CachedTargetDataHandle;
-};
\ No newline at end of file
+};
diff --git a/Engine/Source/Runtime/GameplayAbilities/Public/Abilities/Tasks/AbilityTask_StartAbilityState.h b/Engine/Source/Runtime/GameplayAbilities/Public/Abilities/Tasks/AbilityTask_StartAbilityState.h
--- a/Engine/Source/Runtime/GameplayAbilities/Public/Abilities/Tasks/AbilityTask_StartAbilityState.h
+++ b/Engine/Source/Runtime/GameplayAbilities/Public/Abilities/Tasks/AbilityTask_StartAbilityState.h
@@ -25,7 +25,9 @@ DECLARE_DYNAMIC_MULTICAST_DELEGATE(FAbilityStateDelegate);
 UCLASS(MinimalAPI)
 class UAbilityTask_StartAbilityState : public UAbilityTask
 {
-	GENERATED_UCLASS_BODY()
+	GENERATED_BODY()
+public:
+	GAMEPLAYABILITIES_API UAbilityTask_StartAbilityState(const FObjectInitializer& ObjectInitializer = FObjectInitializer::Get());
 
 	/** Invoked if 'EndAbilityState' is called or if 'EndAbility' is called and this state is active. */
 	UPROPERTY(BlueprintAssignable)
diff --git a/Engine/Source/Runtime/GameplayAbilities/Public/Abilities/Tasks/AbilityTask_VisualizeTargeting.h b/Engine/Source/Runtime/GameplayAbilities/Public/Abilities/Tasks/AbilityTask_VisualizeTargeting.h
--- a/Engine/Source/Runtime/GameplayAbilities/Public/Abilities/Tasks/AbilityTask_VisualizeTargeting.h
+++ b/Engine/Source/Runtime/GameplayAbilities/Public/Abilities/Tasks/AbilityTask_VisualizeTargeting.h
@@ -12,7 +12,9 @@ DECLARE_DYNAMIC_MULTICAST_DELEGATE(FVisualizeTargetingDelegate);
 UCLASS(MinimalAPI)
 class UAbilityTask_VisualizeTargeting: public UAbilityTask
 {
-	GENERATED_UCLASS_BODY()
+	GENERATED_BODY()
+public:
+	GAMEPLAYABILITIES_API UAbilityTask_VisualizeTargeting(const FObjectInitializer& ObjectInitializer = FObjectInitializer::Get());
 
 	UPROPERTY(BlueprintAssignable)
 	FVisualizeTargetingDelegate TimeElapsed;
@@ -64,4 +66,4 @@ protected:
 *			-This function should spawn the actor with SpawnActorDeferred and return true/false if it spawned something.
 *		-Have a function named FinishSpawningActor w/ an AActor* of the class you spawned
 *			-This function *must* call ExecuteConstruction + PostActorConstruction
-*/
\ No newline at end of file
+*/
diff --git a/Engine/Source/Runtime/GameplayAbilities/Public/Abilities/Tasks/AbilityTask_WaitAbilityActivate.h b/Engine/Source/Runtime/GameplayAbilities/Public/Abilities/Tasks/AbilityTask_WaitAbilityActivate.h
--- a/Engine/Source/Runtime/GameplayAbilities/Public/Abilities/Tasks/AbilityTask_WaitAbilityActivate.h
+++ b/Engine/Source/Runtime/GameplayAbilities/Public/Abilities/Tasks/AbilityTask_WaitAbilityActivate.h
@@ -14,7 +14,9 @@ class AActor;
 UCLASS(MinimalAPI)
 class UAbilityTask_WaitAbilityActivate : public UAbilityTask
 {
-	GENERATED_UCLASS_BODY()
+	GENERATED_BODY()
+public:
+	GAMEPLAYABILITIES_API UAbilityTask_WaitAbilityActivate(const FObjectInitializer& ObjectInitializer = FObjectInitializer::Get());
 
 	UPROPERTY(BlueprintAssignable)
 	FWaitAbilityActivateDelegate	OnActivate;
@@ -37,4 +39,4 @@ protected:
 	virtual void OnDestroy(bool AbilityEnded) override;
 
 	FDelegateHandle OnAbilityActivateDelegateHandle;
-};
\ No newline at end of file
+};
diff --git a/Engine/Source/Runtime/GameplayAbilities/Public/Abilities/Tasks/AbilityTask_WaitAbilityCommit.h b/Engine/Source/Runtime/GameplayAbilities/Public/Abilities/Tasks/AbilityTask_WaitAbilityCommit.h
--- a/Engine/Source/Runtime/GameplayAbilities/Public/Abilities/Tasks/AbilityTask_WaitAbilityCommit.h
+++ b/Engine/Source/Runtime/GameplayAbilities/Public/Abilities/Tasks/AbilityTask_WaitAbilityCommit.h
@@ -15,7 +15,9 @@ DECLARE_DYNAMIC_MULTICAST_DELEGATE_OneParam(FWaitAbilityCommitDelegate, UGamepla
 UCLASS(MinimalAPI)
 class UAbilityTask_WaitAbilityCommit : public UAbilityTask
 {
-	GENERATED_UCLASS_BODY()
+	GENERATED_BODY()
+public:
+	GAMEPLAYABILITIES_API UAbilityTask_WaitAbilityCommit(const FObjectInitializer& ObjectInitializer = FObjectInitializer::Get());
 
 	UPROPERTY(BlueprintAssignable)
 	FWaitAbilityCommitDelegate	OnCommit;
@@ -37,4 +39,4 @@ protected:
 	virtual void OnDestroy(bool AbilityEnded) override;
 
 	FDelegateHandle OnAbilityCommitDelegateHandle;
-};
\ No newline at end of file
+};
diff --git a/Engine/Source/Runtime/GameplayAbilities/Public/Abilities/Tasks/AbilityTask_WaitAttributeChange.h b/Engine/Source/Runtime/GameplayAbilities/Public/Abilities/Tasks/AbilityTask_WaitAttributeChange.h
--- a/Engine/Source/Runtime/GameplayAbilities/Public/Abilities/Tasks/AbilityTask_WaitAttributeChange.h
+++ b/Engine/Source/Runtime/GameplayAbilities/Public/Abilities/Tasks/AbilityTask_WaitAttributeChange.h
@@ -31,7 +31,9 @@ DECLARE_DYNAMIC_MULTICAST_DELEGATE(FWaitAttributeChangeDelegate);
 UCLASS(MinimalAPI)
 class UAbilityTask_WaitAttributeChange : public UAbilityTask
 {
-	GENERATED_UCLASS_BODY()
+	GENERATED_BODY()
+public:
+	GAMEPLAYABILITIES_API UAbilityTask_WaitAttributeChange(const FObjectInitializer& ObjectInitializer = FObjectInitializer::Get());
 
 	UPROPERTY(BlueprintAssignable)
 	FWaitAttributeChangeDelegate	OnChange;
diff --git a/Engine/Source/Runtime/GameplayAbilities/Public/Abilities/Tasks/AbilityTask_WaitConfirm.h b/Engine/Source/Runtime/GameplayAbilities/Public/Abilities/Tasks/AbilityTask_WaitConfirm.h
--- a/Engine/Source/Runtime/GameplayAbilities/Public/Abilities/Tasks/AbilityTask_WaitConfirm.h
+++ b/Engine/Source/Runtime/GameplayAbilities/Public/Abilities/Tasks/AbilityTask_WaitConfirm.h
@@ -7,7 +7,9 @@
 UCLASS(MinimalAPI)
 class UAbilityTask_WaitConfirm : public UAbilityTask
 {
-	GENERATED_UCLASS_BODY()
+	GENERATED_BODY()
+public:
+	GAMEPLAYABILITIES_API UAbilityTask_WaitConfirm(const FObjectInitializer& ObjectInitializer = FObjectInitializer::Get());
 
 	UPROPERTY(BlueprintAssignable)
 	FGenericAbilityTaskDelegate	OnConfirm;
@@ -28,4 +30,4 @@ protected:
 	bool RegisteredCallback;
 
 	FDelegateHandle OnConfirmCallbackDelegateHandle;
-};
\ No newline at end of file
+};
diff --git a/Engine/Source/Runtime/GameplayAbilities/Public/Abilities/Tasks/AbilityTask_WaitConfirmCancel.h b/Engine/Source/Runtime/GameplayAbilities/Public/Abilities/Tasks/AbilityTask_WaitConfirmCancel.h
--- a/Engine/Source/Runtime/GameplayAbilities/Public/Abilities/Tasks/AbilityTask_WaitConfirmCancel.h
+++ b/Engine/Source/Runtime/GameplayAbilities/Public/Abilities/Tasks/AbilityTask_WaitConfirmCancel.h
@@ -14,7 +14,9 @@ DECLARE_DYNAMIC_MULTICAST_DELEGATE(FWaitConfirmCancelDelegate);
 UCLASS(MinimalAPI)
 class UAbilityTask_WaitConfirmCancel : public UAbilityTask
 {
-	GENERATED_UCLASS_BODY()
+	GENERATED_BODY()
+public:
+	GAMEPLAYABILITIES_API UAbilityTask_WaitConfirmCancel(const FObjectInitializer& ObjectInitializer = FObjectInitializer::Get());
 
 	UPROPERTY(BlueprintAssignable)
 	FWaitConfirmCancelDelegate	OnConfirm;
@@ -45,4 +47,4 @@ protected:
 
 	bool RegisteredCallbacks;
 	
-};
\ No newline at end of file
+};
diff --git a/Engine/Source/Runtime/GameplayAbilities/Public/Abilities/Tasks/AbilityTask_WaitDelay.h b/Engine/Source/Runtime/GameplayAbilities/Public/Abilities/Tasks/AbilityTask_WaitDelay.h
--- a/Engine/Source/Runtime/GameplayAbilities/Public/Abilities/Tasks/AbilityTask_WaitDelay.h
+++ b/Engine/Source/Runtime/GameplayAbilities/Public/Abilities/Tasks/AbilityTask_WaitDelay.h
@@ -9,7 +9,9 @@ DECLARE_DYNAMIC_MULTICAST_DELEGATE(FWaitDelayDelegate);
 UCLASS(MinimalAPI)
 class UAbilityTask_WaitDelay : public UAbilityTask
 {
-	GENERATED_UCLASS_BODY()
+	GENERATED_BODY()
+public:
+	GAMEPLAYABILITIES_API UAbilityTask_WaitDelay(const FObjectInitializer& ObjectInitializer = FObjectInitializer::Get());
 
 	UPROPERTY(BlueprintAssignable)
 	FWaitDelayDelegate	OnFinish;
@@ -30,4 +32,4 @@ private:
 	virtual void OnDestroy(bool AbilityEnded) override;
 	float Time;
 	float TimeStarted;
-};
\ No newline at end of file
+};
diff --git a/Engine/Source/Runtime/GameplayAbilities/Public/Abilities/Tasks/AbilityTask_WaitGameplayEffectApplied.h b/Engine/Source/Runtime/GameplayAbilities/Public/Abilities/Tasks/AbilityTask_WaitGameplayEffectApplied.h
--- a/Engine/Source/Runtime/GameplayAbilities/Public/Abilities/Tasks/AbilityTask_WaitGameplayEffectApplied.h
+++ b/Engine/Source/Runtime/GameplayAbilities/Public/Abilities/Tasks/AbilityTask_WaitGameplayEffectApplied.h
@@ -9,7 +9,9 @@
 UCLASS(MinimalAPI)
 class UAbilityTask_WaitGameplayEffectApplied : public UAbilityTask
 {
-	GENERATED_UCLASS_BODY()
+	GENERATED_BODY()
+public:
+	GAMEPLAYABILITIES_API UAbilityTask_WaitGameplayEffectApplied(const FObjectInitializer& ObjectInitializer = FObjectInitializer::Get());
 
 	UFUNCTION()
 	void OnApplyGameplayEffectCallback(UAbilitySystemComponent* Target, const FGameplayEffectSpec& SpecApplied, FActiveGameplayEffectHandle ActiveHandle);
@@ -37,4 +39,4 @@ protected:
 	bool UseExternalOwner;
 
 	TWeakObjectPtr<UAbilitySystemComponent> ExternalOwner;
-};
\ No newline at end of file
+};
diff --git a/Engine/Source/Runtime/GameplayAbilities/Public/Abilities/Tasks/AbilityTask_WaitGameplayEffectApplied_Self.h b/Engine/Source/Runtime/GameplayAbilities/Public/Abilities/Tasks/AbilityTask_WaitGameplayEffectApplied_Self.h
--- a/Engine/Source/Runtime/GameplayAbilities/Public/Abilities/Tasks/AbilityTask_WaitGameplayEffectApplied_Self.h
+++ b/Engine/Source/Runtime/GameplayAbilities/Public/Abilities/Tasks/AbilityTask_WaitGameplayEffectApplied_Self.h
@@ -10,7 +10,9 @@ DECLARE_DYNAMIC_MULTICAST_DELEGATE_ThreeParams(FGameplayEffectAppliedSelfDelegat
 UCLASS(MinimalAPI)
 class UAbilityTask_WaitGameplayEffectApplied_Self : public UAbilityTask_WaitGameplayEffectApplied
 {
-	GENERATED_UCLASS_BODY()
+	GENERATED_BODY()
+public:
+	GAMEPLAYABILITIES_API UAbilityTask_WaitGameplayEffectApplied_Self(const FObjectInitializer& ObjectInitializer = FObjectInitializer::Get());
 
 	UPROPERTY(BlueprintAssignable)
 	FGameplayEffectAppliedSelfDelegate OnApplied;
@@ -30,4 +32,4 @@ protected:
 
 private:
 	FDelegateHandle OnApplyGameplayEffectCallbackDelegateHandle;
-};
\ No newline at end of file
+};
diff --git a/Engine/Source/Runtime/GameplayAbilities/Public/Abilities/Tasks/AbilityTask_WaitGameplayEffectApplied_Target.h b/Engine/Source/Runtime/GameplayAbilities/Public/Abilities/Tasks/AbilityTask_WaitGameplayEffectApplied_Target.h
--- a/Engine/Source/Runtime/GameplayAbilities/Public/Abilities/Tasks/AbilityTask_WaitGameplayEffectApplied_Target.h
+++ b/Engine/Source/Runtime/GameplayAbilities/Public/Abilities/Tasks/AbilityTask_WaitGameplayEffectApplied_Target.h
@@ -10,7 +10,9 @@ DECLARE_DYNAMIC_MULTICAST_DELEGATE_ThreeParams(FGameplayEffectAppliedTargetDeleg
 UCLASS(MinimalAPI)
 class UAbilityTask_WaitGameplayEffectApplied_Target : public UAbilityTask_WaitGameplayEffectApplied
 {
-	GENERATED_UCLASS_BODY()
+	GENERATED_BODY()
+public:
+	GAMEPLAYABILITIES_API UAbilityTask_WaitGameplayEffectApplied_Target(const FObjectInitializer& ObjectInitializer = FObjectInitializer::Get());
 
 	UPROPERTY(BlueprintAssignable)
 	FGameplayEffectAppliedTargetDelegate OnApplied;
@@ -30,4 +32,4 @@ protected:
 
 private:
 	FDelegateHandle OnApplyGameplayEffectCallbackDelegateHandle;
-};
\ No newline at end of file
+};
diff --git a/Engine/Source/Runtime/GameplayAbilities/Public/Abilities/Tasks/AbilityTask_WaitGameplayEffectRemoved.h b/Engine/Source/Runtime/GameplayAbilities/Public/Abilities/Tasks/AbilityTask_WaitGameplayEffectRemoved.h
--- a/Engine/Source/Runtime/GameplayAbilities/Public/Abilities/Tasks/AbilityTask_WaitGameplayEffectRemoved.h
+++ b/Engine/Source/Runtime/GameplayAbilities/Public/Abilities/Tasks/AbilityTask_WaitGameplayEffectRemoved.h
@@ -17,7 +17,9 @@ class UPrimitiveComponent;
 UCLASS(MinimalAPI)
 class UAbilityTask_WaitGameplayEffectRemoved : public UAbilityTask
 {
-	GENERATED_UCLASS_BODY()
+	GENERATED_BODY()
+public:
+	GAMEPLAYABILITIES_API UAbilityTask_WaitGameplayEffectRemoved(const FObjectInitializer& ObjectInitializer = FObjectInitializer::Get());
 
 	UPROPERTY(BlueprintAssignable)
 	FWaitGameplayEffectRemovedDelegate	OnRemoved;
@@ -42,4 +44,4 @@ protected:
 	bool Registered;
 
 	FDelegateHandle OnGameplayEffectRemovedDelegateHandle;
-};
\ No newline at end of file
+};
diff --git a/Engine/Source/Runtime/GameplayAbilities/Public/Abilities/Tasks/AbilityTask_WaitGameplayTag.h b/Engine/Source/Runtime/GameplayAbilities/Public/Abilities/Tasks/AbilityTask_WaitGameplayTag.h
--- a/Engine/Source/Runtime/GameplayAbilities/Public/Abilities/Tasks/AbilityTask_WaitGameplayTag.h
+++ b/Engine/Source/Runtime/GameplayAbilities/Public/Abilities/Tasks/AbilityTask_WaitGameplayTag.h
@@ -11,7 +11,9 @@ DECLARE_DYNAMIC_MULTICAST_DELEGATE(FWaitGameplayTagDelegate);
 UCLASS(MinimalAPI)
 class UAbilityTask_WaitGameplayTagAdded : public UAbilityTask_WaitGameplayTag
 {
-	GENERATED_UCLASS_BODY()
+	GENERATED_BODY()
+public:
+	GAMEPLAYABILITIES_API UAbilityTask_WaitGameplayTagAdded(const FObjectInitializer& ObjectInitializer = FObjectInitializer::Get());
 
 	UPROPERTY(BlueprintAssignable)
 	FWaitGameplayTagDelegate	Added;
@@ -31,7 +33,9 @@ class UAbilityTask_WaitGameplayTagAdded : public UAbilityTask_WaitGameplayTag
 UCLASS(MinimalAPI)
 class UAbilityTask_WaitGameplayTagRemoved : public UAbilityTask_WaitGameplayTag
 {
-	GENERATED_UCLASS_BODY()
+	GENERATED_BODY()
+public:
+	GAMEPLAYABILITIES_API UAbilityTask_WaitGameplayTagRemoved(const FObjectInitializer& ObjectInitializer = FObjectInitializer::Get());
 
 	UPROPERTY(BlueprintAssignable)
 	FWaitGameplayTagDelegate	Removed;
@@ -46,4 +50,4 @@ class UAbilityTask_WaitGameplayTagRemoved : public UAbilityTask_WaitGameplayTag
 	virtual void Activate() override;
 
 	virtual void GameplayTagCallback(const FGameplayTag Tag, int32 NewCount) override;
-};
\ No newline at end of file
+};
diff --git a/Engine/Source/Runtime/GameplayAbilities/Public/Abilities/Tasks/AbilityTask_WaitGameplayTagBase.h b/Engine/Source/Runtime/GameplayAbilities/Public/Abilities/Tasks/AbilityTask_WaitGameplayTagBase.h
--- a/Engine/Source/Runtime/GameplayAbilities/Public/Abilities/Tasks/AbilityTask_WaitGameplayTagBase.h
+++ b/Engine/Source/Runtime/GameplayAbilities/Public/Abilities/Tasks/AbilityTask_WaitGameplayTagBase.h
@@ -8,7 +8,9 @@
 UCLASS(MinimalAPI)
 class UAbilityTask_WaitGameplayTag : public UAbilityTask
 {
-	GENERATED_UCLASS_BODY()
+	GENERATED_BODY()
+public:
+	GAMEPLAYABILITIES_API UAbilityTask_WaitGameplayTag(const FObjectInitializer& ObjectInitializer = FObjectInitializer::Get());
 
 	virtual void Activate() override;
 
@@ -32,4 +34,4 @@ protected:
 	bool OnlyTriggerOnce;
 
 	FDelegateHandle DelegateHandle;
-};
\ No newline at end of file
+};
diff --git a/Engine/Source/Runtime/GameplayAbilities/Public/Abilities/Tasks/AbilityTask_WaitInputPress.h b/Engine/Source/Runtime/GameplayAbilities/Public/Abilities/Tasks/AbilityTask_WaitInputPress.h
--- a/Engine/Source/Runtime/GameplayAbilities/Public/Abilities/Tasks/AbilityTask_WaitInputPress.h
+++ b/Engine/Source/Runtime/GameplayAbilities/Public/Abilities/Tasks/AbilityTask_WaitInputPress.h
@@ -14,7 +14,9 @@ DECLARE_DYNAMIC_MULTICAST_DELEGATE_OneParam(FInputPressDelegate, float, TimeWait
 UCLASS(MinimalAPI)
 class UAbilityTask_WaitInputPress : public UAbilityTask
 {
-	GENERATED_UCLASS_BODY()
+	GENERATED_BODY()
+public:
+	GAMEPLAYABILITIES_API UAbilityTask_WaitInputPress(const FObjectInitializer& ObjectInitializer = FObjectInitializer::Get());
 
 	UPROPERTY(BlueprintAssignable)
 	FInputPressDelegate		OnPress;
@@ -35,4 +37,4 @@ protected:
 	float StartTime;
 	bool bTestInitialState;
 	FDelegateHandle DelegateHandle;
-};
\ No newline at end of file
+};
diff --git a/Engine/Source/Runtime/GameplayAbilities/Public/Abilities/Tasks/AbilityTask_WaitInputRelease.h b/Engine/Source/Runtime/GameplayAbilities/Public/Abilities/Tasks/AbilityTask_WaitInputRelease.h
--- a/Engine/Source/Runtime/GameplayAbilities/Public/Abilities/Tasks/AbilityTask_WaitInputRelease.h
+++ b/Engine/Source/Runtime/GameplayAbilities/Public/Abilities/Tasks/AbilityTask_WaitInputRelease.h
@@ -14,7 +14,9 @@ DECLARE_DYNAMIC_MULTICAST_DELEGATE_OneParam(FInputReleaseDelegate, float, TimeHe
 UCLASS(MinimalAPI)
 class UAbilityTask_WaitInputRelease : public UAbilityTask
 {
-	GENERATED_UCLASS_BODY()
+	GENERATED_BODY()
+public:
+	GAMEPLAYABILITIES_API UAbilityTask_WaitInputRelease(const FObjectInitializer& ObjectInitializer = FObjectInitializer::Get());
 
 	UPROPERTY(BlueprintAssignable)
 	FInputReleaseDelegate	OnRelease;
@@ -35,4 +37,4 @@ protected:
 	float StartTime;
 	bool bTestInitialState;
 	FDelegateHandle DelegateHandle;
-};
\ No newline at end of file
+};
diff --git a/Engine/Source/Runtime/GameplayAbilities/Public/Abilities/Tasks/AbilityTask_WaitMovementModeChange.h b/Engine/Source/Runtime/GameplayAbilities/Public/Abilities/Tasks/AbilityTask_WaitMovementModeChange.h
--- a/Engine/Source/Runtime/GameplayAbilities/Public/Abilities/Tasks/AbilityTask_WaitMovementModeChange.h
+++ b/Engine/Source/Runtime/GameplayAbilities/Public/Abilities/Tasks/AbilityTask_WaitMovementModeChange.h
@@ -10,7 +10,9 @@ class ACharacter;
 UCLASS(MinimalAPI)
 class UAbilityTask_WaitMovementModeChange : public UAbilityTask
 {
-	GENERATED_UCLASS_BODY()
+	GENERATED_BODY()
+public:
+	GAMEPLAYABILITIES_API UAbilityTask_WaitMovementModeChange(const FObjectInitializer& ObjectInitializer = FObjectInitializer::Get());
 
 public:
 
@@ -32,4 +34,4 @@ private:
 	virtual void OnDestroy(bool AbilityEnded) override;
 
 	TWeakObjectPtr<ACharacter>	MyCharacter;
-};
\ No newline at end of file
+};
diff --git a/Engine/Source/Runtime/GameplayAbilities/Public/Abilities/Tasks/AbilityTask_WaitOverlap.h b/Engine/Source/Runtime/GameplayAbilities/Public/Abilities/Tasks/AbilityTask_WaitOverlap.h
--- a/Engine/Source/Runtime/GameplayAbilities/Public/Abilities/Tasks/AbilityTask_WaitOverlap.h
+++ b/Engine/Source/Runtime/GameplayAbilities/Public/Abilities/Tasks/AbilityTask_WaitOverlap.h
@@ -17,7 +17,9 @@ class UPrimitiveComponent;
 UCLASS(MinimalAPI)
 class UAbilityTask_WaitOverlap : public UAbilityTask
 {
-	GENERATED_UCLASS_BODY()
+	GENERATED_BODY()
+public:
+	GAMEPLAYABILITIES_API UAbilityTask_WaitOverlap(const FObjectInitializer& ObjectInitializer = FObjectInitializer::Get());
 
 	UPROPERTY(BlueprintAssignable)
 	FWaitOverlapDelegate	OnOverlap;
@@ -40,4 +42,4 @@ private:
 
 	UPrimitiveComponent* GetComponent();
 	
-};
\ No newline at end of file
+};
diff --git a/Engine/Source/Runtime/GameplayAbilities/Public/Abilities/Tasks/AbilityTask_WaitTargetData.h b/Engine/Source/Runtime/GameplayAbilities/Public/Abilities/Tasks/AbilityTask_WaitTargetData.h
--- a/Engine/Source/Runtime/GameplayAbilities/Public/Abilities/Tasks/AbilityTask_WaitTargetData.h
+++ b/Engine/Source/Runtime/GameplayAbilities/Public/Abilities/Tasks/AbilityTask_WaitTargetData.h
@@ -13,7 +13,9 @@ DECLARE_DYNAMIC_MULTICAST_DELEGATE_OneParam(FWaitTargetDataDelegate, FGameplayAb
 UCLASS(notplaceable)
 class GAMEPLAYABILITIES_API UAbilityTask_WaitTargetData: public UAbilityTask
 {
-	GENERATED_UCLASS_BODY()
+	GENERATED_BODY()
+public:
+	UAbilityTask_WaitTargetData(const FObjectInitializer& ObjectInitializer = FObjectInitializer::Get());
 
 	UPROPERTY(BlueprintAssignable)
 	FWaitTargetDataDelegate	ValidData;
diff --git a/Engine/Source/Runtime/GameplayAbilities/Public/Abilities/Tasks/AbilityTask_WaitVelocityChange.h b/Engine/Source/Runtime/GameplayAbilities/Public/Abilities/Tasks/AbilityTask_WaitVelocityChange.h
--- a/Engine/Source/Runtime/GameplayAbilities/Public/Abilities/Tasks/AbilityTask_WaitVelocityChange.h
+++ b/Engine/Source/Runtime/GameplayAbilities/Public/Abilities/Tasks/AbilityTask_WaitVelocityChange.h
@@ -9,7 +9,9 @@ DECLARE_DYNAMIC_MULTICAST_DELEGATE(FWaitVelocityChangeDelegate);
 UCLASS(MinimalAPI)
 class UAbilityTask_WaitVelocityChange: public UAbilityTask
 {
-	GENERATED_UCLASS_BODY()
+	GENERATED_BODY()
+public:
+	GAMEPLAYABILITIES_API UAbilityTask_WaitVelocityChange(const FObjectInitializer& ObjectInitializer = FObjectInitializer::Get());
 
 	/** Delegate called when velocity requirements are met */
 	UPROPERTY(BlueprintAssignable)
@@ -30,4 +32,4 @@ protected:
 
 	float	MinimumMagnitude;
 	FVector Direction;
-};
\ No newline at end of file
+};
diff --git a/Engine/Source/Runtime/GameplayAbilities/Public/AbilitySystemBlueprintLibrary.h b/Engine/Source/Runtime/GameplayAbilities/Public/AbilitySystemBlueprintLibrary.h
--- a/Engine/Source/Runtime/GameplayAbilities/Public/AbilitySystemBlueprintLibrary.h
+++ b/Engine/Source/Runtime/GameplayAbilities/Public/AbilitySystemBlueprintLibrary.h
@@ -17,7 +17,9 @@ class UAbilityTask_WaitConfirmCancel;
 UCLASS()
 class GAMEPLAYABILITIES_API UAbilitySystemBlueprintLibrary : public UBlueprintFunctionLibrary
 {
-	GENERATED_UCLASS_BODY()
+	GENERATED_BODY()
+public:
+	UAbilitySystemBlueprintLibrary(const FObjectInitializer& ObjectInitializer = FObjectInitializer::Get());
 
 	UFUNCTION(BlueprintPure, Category = Ability)
 	static UAbilitySystemComponent* GetAbilitySystemComponent(AActor *Actor);
