Back out changelist 2481333

diff --git a/Engine/Source/Runtime/Engine/Classes/AI/Navigation/NavMeshRenderingComponent.h b/Engine/Source/Runtime/Engine/Classes/AI/Navigation/NavMeshRenderingComponent.h
--- a/Engine/Source/Runtime/Engine/Classes/AI/Navigation/NavMeshRenderingComponent.h
+++ b/Engine/Source/Runtime/Engine/Classes/AI/Navigation/NavMeshRenderingComponent.h
@@ -8,9 +8,7 @@ struct FNavMeshSceneProxyData;
 UCLASS(hidecategories=Object, editinlinenew)
 class ENGINE_API UNavMeshRenderingComponent : public UPrimitiveComponent
 {
-	GENERATED_BODY()
-public:
-	UNavMeshRenderingComponent(const FObjectInitializer& ObjectInitializer = FObjectInitializer::Get());
+	GENERATED_UCLASS_BODY()
 
 public:
 	// Begin UObject Interface
diff --git a/Engine/Source/Runtime/Engine/Classes/AI/Navigation/NavModifierComponent.h b/Engine/Source/Runtime/Engine/Classes/AI/Navigation/NavModifierComponent.h
--- a/Engine/Source/Runtime/Engine/Classes/AI/Navigation/NavModifierComponent.h
+++ b/Engine/Source/Runtime/Engine/Classes/AI/Navigation/NavModifierComponent.h
@@ -11,9 +11,7 @@ struct FCompositeNavModifier;
 UCLASS(ClassGroup = (Navigation), meta = (BlueprintSpawnableComponent), hidecategories = (Activation))
 class UNavModifierComponent : public UNavRelevantComponent
 {
-	GENERATED_BODY()
-public:
-	UNavModifierComponent(const FObjectInitializer& ObjectInitializer = FObjectInitializer::Get());
+	GENERATED_UCLASS_BODY()
 
 	UPROPERTY(EditAnywhere, Category = Navigation)
 	TSubclassOf<UNavArea> AreaClass;
diff --git a/Engine/Source/Runtime/Engine/Classes/AI/Navigation/NavModifierVolume.h b/Engine/Source/Runtime/Engine/Classes/AI/Navigation/NavModifierVolume.h
--- a/Engine/Source/Runtime/Engine/Classes/AI/Navigation/NavModifierVolume.h
+++ b/Engine/Source/Runtime/Engine/Classes/AI/Navigation/NavModifierVolume.h
@@ -11,9 +11,7 @@
 UCLASS(hidecategories=(Navigation))
 class ANavModifierVolume : public AVolume, public INavRelevantInterface
 {
-	GENERATED_BODY()
-public:
-	ANavModifierVolume(const FObjectInitializer& ObjectInitializer = FObjectInitializer::Get());
+	GENERATED_UCLASS_BODY()
 
 	UPROPERTY(EditAnywhere, BlueprintReadWrite, Category=Default)
 	TSubclassOf<class UNavArea> AreaClass;
diff --git a/Engine/Source/Runtime/Engine/Classes/AI/Navigation/NavNodeInterface.h b/Engine/Source/Runtime/Engine/Classes/AI/Navigation/NavNodeInterface.h
--- a/Engine/Source/Runtime/Engine/Classes/AI/Navigation/NavNodeInterface.h
+++ b/Engine/Source/Runtime/Engine/Classes/AI/Navigation/NavNodeInterface.h
@@ -6,15 +6,12 @@
 UINTERFACE(MinimalAPI, meta=(CannotImplementInterfaceInBlueprint))
 class UNavNodeInterface : public UInterface
 {
-	GENERATED_BODY()
-public:
-	ENGINE_API UNavNodeInterface(const FObjectInitializer& ObjectInitializer = FObjectInitializer::Get());
+	GENERATED_UINTERFACE_BODY()
 };
 
 class INavNodeInterface
 {
-	GENERATED_BODY()
-public:
+	GENERATED_IINTERFACE_BODY()
 
 
 	/**
diff --git a/Engine/Source/Runtime/Engine/Classes/AI/Navigation/NavPathObserverInterface.h b/Engine/Source/Runtime/Engine/Classes/AI/Navigation/NavPathObserverInterface.h
--- a/Engine/Source/Runtime/Engine/Classes/AI/Navigation/NavPathObserverInterface.h
+++ b/Engine/Source/Runtime/Engine/Classes/AI/Navigation/NavPathObserverInterface.h
@@ -6,15 +6,12 @@
 UINTERFACE(MinimalAPI, meta=(CannotImplementInterfaceInBlueprint))
 class UNavPathObserverInterface : public UInterface
 {
-	GENERATED_BODY()
-public:
-	ENGINE_API UNavPathObserverInterface(const FObjectInitializer& ObjectInitializer = FObjectInitializer::Get());
+	GENERATED_UINTERFACE_BODY()
 };
 
 class INavPathObserverInterface
 {
-	GENERATED_BODY()
-public:
+	GENERATED_IINTERFACE_BODY()
 
 	/** 
 	 *	Called on Navigation Agent when given PathGenerator updated it's path.
diff --git a/Engine/Source/Runtime/Engine/Classes/AI/Navigation/NavRelevantComponent.h b/Engine/Source/Runtime/Engine/Classes/AI/Navigation/NavRelevantComponent.h
--- a/Engine/Source/Runtime/Engine/Classes/AI/Navigation/NavRelevantComponent.h
+++ b/Engine/Source/Runtime/Engine/Classes/AI/Navigation/NavRelevantComponent.h
@@ -8,9 +8,7 @@
 UCLASS()
 class ENGINE_API UNavRelevantComponent : public UActorComponent, public INavRelevantInterface
 {
-	GENERATED_BODY()
-public:
-	UNavRelevantComponent(const FObjectInitializer& ObjectInitializer = FObjectInitializer::Get());
+	GENERATED_UCLASS_BODY()
 
 	// Begin UActorComponent Interface
 	virtual void OnRegister() override;
diff --git a/Engine/Source/Runtime/Engine/Classes/AI/Navigation/NavRelevantInterface.h b/Engine/Source/Runtime/Engine/Classes/AI/Navigation/NavRelevantInterface.h
--- a/Engine/Source/Runtime/Engine/Classes/AI/Navigation/NavRelevantInterface.h
+++ b/Engine/Source/Runtime/Engine/Classes/AI/Navigation/NavRelevantInterface.h
@@ -11,15 +11,12 @@ struct FNavigableGeometryExport;
 UINTERFACE(MinimalAPI, meta=(CannotImplementInterfaceInBlueprint))
 class UNavRelevantInterface : public UInterface
 {
-	GENERATED_BODY()
-public:
-	ENGINE_API UNavRelevantInterface(const FObjectInitializer& ObjectInitializer = FObjectInitializer::Get());
+	GENERATED_UINTERFACE_BODY()
 };
 
 class INavRelevantInterface
 {
-	GENERATED_BODY()
-public:
+	GENERATED_IINTERFACE_BODY()
 
 	/** Prepare navigation modifiers */
 	virtual void GetNavigationData(FNavigationRelevantData& Data) const {}
diff --git a/Engine/Source/Runtime/Engine/Classes/AI/Navigation/NavTestRenderingComponent.h b/Engine/Source/Runtime/Engine/Classes/AI/Navigation/NavTestRenderingComponent.h
--- a/Engine/Source/Runtime/Engine/Classes/AI/Navigation/NavTestRenderingComponent.h
+++ b/Engine/Source/Runtime/Engine/Classes/AI/Navigation/NavTestRenderingComponent.h
@@ -6,9 +6,7 @@
 UCLASS(hidecategories=Object)
 class UNavTestRenderingComponent: public UPrimitiveComponent
 {
-	GENERATED_BODY()
-public:
-	UNavTestRenderingComponent(const FObjectInitializer& ObjectInitializer = FObjectInitializer::Get());
+	GENERATED_UCLASS_BODY()
 
 	virtual FPrimitiveSceneProxy* CreateSceneProxy() override;
 	virtual FBoxSphereBounds CalcBounds(const FTransform &LocalToWorld) const override;
diff --git a/Engine/Source/Runtime/Engine/Classes/AI/Navigation/NavigationDataChunk.h b/Engine/Source/Runtime/Engine/Classes/AI/Navigation/NavigationDataChunk.h
--- a/Engine/Source/Runtime/Engine/Classes/AI/Navigation/NavigationDataChunk.h
+++ b/Engine/Source/Runtime/Engine/Classes/AI/Navigation/NavigationDataChunk.h
@@ -10,9 +10,7 @@
 UCLASS(NotBlueprintable, abstract)
 class ENGINE_API UNavigationDataChunk : public UObject
 {
-	GENERATED_BODY()
-public:
-	UNavigationDataChunk(const FObjectInitializer& ObjectInitializer = FObjectInitializer::Get());
+	GENERATED_UCLASS_BODY()
 
 public:
 	/** Name of NavigationData actor that owns this chunk */
diff --git a/Engine/Source/Runtime/Engine/Classes/AI/Navigation/NavigationGraph.h b/Engine/Source/Runtime/Engine/Classes/AI/Navigation/NavigationGraph.h
--- a/Engine/Source/Runtime/Engine/Classes/AI/Navigation/NavigationGraph.h
+++ b/Engine/Source/Runtime/Engine/Classes/AI/Navigation/NavigationGraph.h
@@ -9,7 +9,7 @@ class UNavigationSystem;
 USTRUCT()
 struct FNavGraphEdge
 {
-	GENERATED_BODY()
+	GENERATED_USTRUCT_BODY()
 
 	enum {
 		EgdeFlagsCount = 7
@@ -35,7 +35,7 @@ struct FNavGraphEdge
 USTRUCT()
 struct FNavGraphNode
 {
-	GENERATED_BODY()
+	GENERATED_USTRUCT_BODY()
 
 	/** Who's this node referring to? This will most commonly point to an actor or a component */
 	UPROPERTY()
@@ -56,9 +56,7 @@ struct FNavGraphNode
 UCLASS(config=Engine, MinimalAPI, abstract)
 class ANavigationGraph : public ANavigationData
 {
-	GENERATED_BODY()
-public:
-	ENGINE_API ANavigationGraph(const FObjectInitializer& ObjectInitializer = FObjectInitializer::Get());
+	GENERATED_UCLASS_BODY()
 
 public:
 	static ANavigationData* CreateNavigationInstances(UNavigationSystem* NavSys);
diff --git a/Engine/Source/Runtime/Engine/Classes/AI/Navigation/NavigationGraphNode.h b/Engine/Source/Runtime/Engine/Classes/AI/Navigation/NavigationGraphNode.h
--- a/Engine/Source/Runtime/Engine/Classes/AI/Navigation/NavigationGraphNode.h
+++ b/Engine/Source/Runtime/Engine/Classes/AI/Navigation/NavigationGraphNode.h
@@ -9,7 +9,5 @@
 UCLASS(config=Engine, MinimalAPI, NotBlueprintable, abstract)
 class ANavigationGraphNode : public AActor
 {
-	GENERATED_BODY()
-public:
-	ENGINE_API ANavigationGraphNode(const FObjectInitializer& ObjectInitializer = FObjectInitializer::Get());
+	GENERATED_UCLASS_BODY()
 };
diff --git a/Engine/Source/Runtime/Engine/Classes/AI/Navigation/NavigationGraphNodeComponent.h b/Engine/Source/Runtime/Engine/Classes/AI/Navigation/NavigationGraphNodeComponent.h
--- a/Engine/Source/Runtime/Engine/Classes/AI/Navigation/NavigationGraphNodeComponent.h
+++ b/Engine/Source/Runtime/Engine/Classes/AI/Navigation/NavigationGraphNodeComponent.h
@@ -7,9 +7,7 @@
 UCLASS(config=Engine, MinimalAPI, HideCategories=(Mobility))
 class UNavigationGraphNodeComponent : public USceneComponent
 {
-	GENERATED_BODY()
-public:
-	ENGINE_API UNavigationGraphNodeComponent(const FObjectInitializer& ObjectInitializer = FObjectInitializer::Get());
+	GENERATED_UCLASS_BODY()
 
 	UPROPERTY()
 	FNavGraphNode Node;
diff --git a/Engine/Source/Runtime/Engine/Classes/AI/Navigation/NavigationPath.h b/Engine/Source/Runtime/Engine/Classes/AI/Navigation/NavigationPath.h
--- a/Engine/Source/Runtime/Engine/Classes/AI/Navigation/NavigationPath.h
+++ b/Engine/Source/Runtime/Engine/Classes/AI/Navigation/NavigationPath.h
@@ -16,9 +16,7 @@ DECLARE_DYNAMIC_MULTICAST_DELEGATE_TwoParams(FOnNavigationPathUpdated, UNavigati
 UCLASS(BlueprintType)
 class ENGINE_API UNavigationPath : public UObject
 {
-	GENERATED_BODY()
-public:
-	UNavigationPath(const FObjectInitializer& ObjectInitializer = FObjectInitializer::Get());
+	GENERATED_UCLASS_BODY()
 
 	UPROPERTY(BlueprintAssignable)
 	FOnNavigationPathUpdated PathUpdatedNotifier;
diff --git a/Engine/Source/Runtime/Engine/Classes/AI/Navigation/NavigationPathGenerator.h b/Engine/Source/Runtime/Engine/Classes/AI/Navigation/NavigationPathGenerator.h
--- a/Engine/Source/Runtime/Engine/Classes/AI/Navigation/NavigationPathGenerator.h
+++ b/Engine/Source/Runtime/Engine/Classes/AI/Navigation/NavigationPathGenerator.h
@@ -6,15 +6,12 @@
 UINTERFACE(MinimalAPI, meta=(CannotImplementInterfaceInBlueprint))
 class UNavigationPathGenerator : public UInterface
 {
-	GENERATED_BODY()
-public:
-	ENGINE_API UNavigationPathGenerator(const FObjectInitializer& ObjectInitializer = FObjectInitializer::Get());
+	GENERATED_UINTERFACE_BODY()
 };
 
 class INavigationPathGenerator
 {
-	GENERATED_BODY()
-public:
+	GENERATED_IINTERFACE_BODY()
 
 
 	/** 
diff --git a/Engine/Source/Runtime/Engine/Classes/AI/Navigation/NavigationSystem.h b/Engine/Source/Runtime/Engine/Classes/AI/Navigation/NavigationSystem.h
--- a/Engine/Source/Runtime/Engine/Classes/AI/Navigation/NavigationSystem.h
+++ b/Engine/Source/Runtime/Engine/Classes/AI/Navigation/NavigationSystem.h
@@ -127,9 +127,7 @@ private:
 UCLASS(Within=World, config=Engine, defaultconfig)
 class ENGINE_API UNavigationSystem : public UBlueprintFunctionLibrary
 {
-	GENERATED_BODY()
-public:
-	UNavigationSystem(const FObjectInitializer& ObjectInitializer = FObjectInitializer::Get());
+	GENERATED_UCLASS_BODY()
 
 	virtual ~UNavigationSystem();
 
diff --git a/Engine/Source/Runtime/Engine/Classes/AI/Navigation/NavigationTestingActor.h b/Engine/Source/Runtime/Engine/Classes/AI/Navigation/NavigationTestingActor.h
--- a/Engine/Source/Runtime/Engine/Classes/AI/Navigation/NavigationTestingActor.h
+++ b/Engine/Source/Runtime/Engine/Classes/AI/Navigation/NavigationTestingActor.h
@@ -33,9 +33,7 @@ namespace ENavCostDisplay
 UCLASS(hidecategories=(Object, Actor, Input, Rendering), showcategories=("Input|MouseInput", "Input|TouchInput"), Blueprintable)
 class ENGINE_API ANavigationTestingActor : public AActor, public INavAgentInterface, public INavPathObserverInterface
 {
-	GENERATED_BODY()
-public:
-	ANavigationTestingActor(const FObjectInitializer& ObjectInitializer = FObjectInitializer::Get());
+	GENERATED_UCLASS_BODY()
 
 private_subobject:
 	DEPRECATED_FORGAME(4.6, "CapsuleComponent should not be accessed directly, please use GetCapsuleComponent() function instead. CapsuleComponent will soon be private and your code will not compile.")
diff --git a/Engine/Source/Runtime/Engine/Classes/AI/Navigation/NavigationTypes.h b/Engine/Source/Runtime/Engine/Classes/AI/Navigation/NavigationTypes.h
--- a/Engine/Source/Runtime/Engine/Classes/AI/Navigation/NavigationTypes.h
+++ b/Engine/Source/Runtime/Engine/Classes/AI/Navigation/NavigationTypes.h
@@ -315,7 +315,7 @@ typedef TWeakPtr<struct FNavigationPath, ESPMode::ThreadSafe> FNavPathWeakPtr;
 USTRUCT()
 struct FMovementProperties
 {
-	GENERATED_BODY()
+	GENERATED_USTRUCT_BODY()
 
 	/** If true, this Pawn is capable of crouching. */
 	UPROPERTY(EditAnywhere, BlueprintReadWrite, Category=MovementProperties)
@@ -351,7 +351,7 @@ struct FMovementProperties
 USTRUCT()
 struct ENGINE_API FNavAgentProperties : public FMovementProperties
 {
-	GENERATED_BODY()
+	GENERATED_USTRUCT_BODY()
 
 	/** Radius of the capsule used for navigation/pathfinding. */
 	UPROPERTY(EditAnywhere, BlueprintReadWrite, Category = MovementProperties)
@@ -407,7 +407,7 @@ inline uint32 GetTypeHash(const FNavAgentProperties& A)
 USTRUCT()
 struct ENGINE_API FNavDataConfig : public FNavAgentProperties
 {
-	GENERATED_BODY()
+	GENERATED_USTRUCT_BODY()
 
 	UPROPERTY(EditAnywhere, BlueprintReadWrite, Category = Display)
 	FName Name;
@@ -520,11 +520,10 @@ struct ENGINE_API FMoveRequestCustomData
 typedef TSharedPtr<FMoveRequestCustomData, ESPMode::ThreadSafe> FCustomMoveSharedPtr;
 typedef TWeakPtr<FMoveRequestCustomData, ESPMode::ThreadSafe> FCustomMoveWeakPtr;
 
-UCLASS(Abstract)
+UCLASS(Abstract, CustomConstructor)
 class UNavigationTypes : public UObject
 {
-	GENERATED_BODY()
-public:
+	GENERATED_UCLASS_BODY()
 
 	UNavigationTypes(const FObjectInitializer& ObjectInitializer) : Super(ObjectInitializer) { }
 };
@@ -697,4 +696,4 @@ struct ENGINE_API FNavHeightfieldSamples
 	FNavHeightfieldSamples();
 
 	FORCEINLINE bool IsEmpty() const { return Heights.Num() == 0; }
-};
+};
\ No newline at end of file
diff --git a/Engine/Source/Runtime/Engine/Classes/AI/Navigation/RecastNavMesh.h b/Engine/Source/Runtime/Engine/Classes/AI/Navigation/RecastNavMesh.h
--- a/Engine/Source/Runtime/Engine/Classes/AI/Navigation/RecastNavMesh.h
+++ b/Engine/Source/Runtime/Engine/Classes/AI/Navigation/RecastNavMesh.h
@@ -411,9 +411,7 @@ namespace FNavMeshConfig
 UCLASS(config=Engine, defaultconfig, hidecategories=(Input,Rendering,Tags,Transform,"Utilities|Transformation",Actor,Layers,Replication), notplaceable)
 class ENGINE_API ARecastNavMesh : public ANavigationData
 {
-	GENERATED_BODY()
-public:
-	ARecastNavMesh(const FObjectInitializer& ObjectInitializer = FObjectInitializer::Get());
+	GENERATED_UCLASS_BODY()
 
 	typedef uint16 FNavPolyFlags;
 
diff --git a/Engine/Source/Runtime/Engine/Classes/AI/Navigation/RecastNavMeshDataChunk.h b/Engine/Source/Runtime/Engine/Classes/AI/Navigation/RecastNavMeshDataChunk.h
--- a/Engine/Source/Runtime/Engine/Classes/AI/Navigation/RecastNavMeshDataChunk.h
+++ b/Engine/Source/Runtime/Engine/Classes/AI/Navigation/RecastNavMeshDataChunk.h
@@ -36,9 +36,7 @@ class dtNavMesh;
 UCLASS()
 class ENGINE_API URecastNavMeshDataChunk : public UNavigationDataChunk
 {
-	GENERATED_BODY()
-public:
-	URecastNavMeshDataChunk(const FObjectInitializer& ObjectInitializer = FObjectInitializer::Get());
+	GENERATED_UCLASS_BODY()
 
 	// Begin UObject Interface
 	virtual void Serialize(FArchive& Ar) override;
diff --git a/Engine/Source/Runtime/Engine/Classes/Animation/AimOffsetBlendSpace.h b/Engine/Source/Runtime/Engine/Classes/Animation/AimOffsetBlendSpace.h
--- a/Engine/Source/Runtime/Engine/Classes/Animation/AimOffsetBlendSpace.h
+++ b/Engine/Source/Runtime/Engine/Classes/Animation/AimOffsetBlendSpace.h
@@ -16,9 +16,7 @@ struct FBlendSample;
 UCLASS(config=Engine, hidecategories=Object, MinimalAPI, BlueprintType)
 class UAimOffsetBlendSpace : public UBlendSpace
 {
-	GENERATED_BODY()
-public:
-	ENGINE_API UAimOffsetBlendSpace(const FObjectInitializer& ObjectInitializer = FObjectInitializer::Get());
+	GENERATED_UCLASS_BODY()
 
 	virtual bool IsValidAdditive() const override;
 
diff --git a/Engine/Source/Runtime/Engine/Classes/Animation/AimOffsetBlendSpace1D.h b/Engine/Source/Runtime/Engine/Classes/Animation/AimOffsetBlendSpace1D.h
--- a/Engine/Source/Runtime/Engine/Classes/Animation/AimOffsetBlendSpace1D.h
+++ b/Engine/Source/Runtime/Engine/Classes/Animation/AimOffsetBlendSpace1D.h
@@ -15,9 +15,7 @@
 UCLASS(config=Engine, hidecategories=Object, MinimalAPI, BlueprintType)
 class UAimOffsetBlendSpace1D : public UBlendSpace1D
 {
-	GENERATED_BODY()
-public:
-	ENGINE_API UAimOffsetBlendSpace1D(const FObjectInitializer& ObjectInitializer = FObjectInitializer::Get());
+	GENERATED_UCLASS_BODY()
 
 	virtual bool IsValidAdditive() const override;
 
diff --git a/Engine/Source/Runtime/Engine/Classes/Animation/AnimBlueprint.h b/Engine/Source/Runtime/Engine/Classes/Animation/AnimBlueprint.h
--- a/Engine/Source/Runtime/Engine/Classes/Animation/AnimBlueprint.h
+++ b/Engine/Source/Runtime/Engine/Classes/Animation/AnimBlueprint.h
@@ -7,7 +7,7 @@
 USTRUCT()
 struct FAnimGroupInfo
 {
-	GENERATED_BODY()
+	GENERATED_USTRUCT_BODY()
 
 	UPROPERTY()
 	FName Name;
@@ -24,7 +24,7 @@ struct FAnimGroupInfo
 USTRUCT()
 struct FAnimParentNodeAssetOverride
 {
-	GENERATED_BODY()
+	GENERATED_USTRUCT_BODY()
 
 	UPROPERTY()
 	UAnimationAsset* NewAsset;
@@ -54,9 +54,7 @@ struct FAnimParentNodeAssetOverride
 UCLASS(BlueprintType)
 class ENGINE_API UAnimBlueprint : public UBlueprint
 {
-	GENERATED_BODY()
-public:
-	UAnimBlueprint(const FObjectInitializer& ObjectInitializer = FObjectInitializer::Get());
+	GENERATED_UCLASS_BODY()
 
 	/** The kind of skeleton that animation graphs compiled from the blueprint will animate */
 	UPROPERTY(AssetRegistrySearchable)
diff --git a/Engine/Source/Runtime/Engine/Classes/Animation/AnimComposite.h b/Engine/Source/Runtime/Engine/Classes/Animation/AnimComposite.h
--- a/Engine/Source/Runtime/Engine/Classes/Animation/AnimComposite.h
+++ b/Engine/Source/Runtime/Engine/Classes/Animation/AnimComposite.h
@@ -12,9 +12,7 @@
 UCLASS(config=Engine, hidecategories=UObject, MinimalAPI, BlueprintType)
 class UAnimComposite : public UAnimCompositeBase
 {
-	GENERATED_BODY()
-public:
-	ENGINE_API UAnimComposite(const FObjectInitializer& ObjectInitializer = FObjectInitializer::Get());
+	GENERATED_UCLASS_BODY()
 
 public:
 	/** Serializable data that stores section/anim pairing **/
diff --git a/Engine/Source/Runtime/Engine/Classes/Animation/AnimCompositeBase.h b/Engine/Source/Runtime/Engine/Classes/Animation/AnimCompositeBase.h
--- a/Engine/Source/Runtime/Engine/Classes/Animation/AnimCompositeBase.h
+++ b/Engine/Source/Runtime/Engine/Classes/Animation/AnimCompositeBase.h
@@ -20,7 +20,7 @@ class UAnimSequence;
 USTRUCT()
 struct FRootMotionExtractionStep
 {
-	GENERATED_BODY()
+	GENERATED_USTRUCT_BODY()
 
 	/** AnimSequence ref */
 	UPROPERTY()
@@ -53,7 +53,7 @@ struct FRootMotionExtractionStep
 USTRUCT()
 struct FAnimSegment
 {
-	GENERATED_BODY()
+	GENERATED_USTRUCT_BODY()
 
 	/** Anim Reference to play - only allow AnimSequence or AnimComposite **/
 	UPROPERTY(EditAnywhere, BlueprintReadWrite, Category=AnimSegment)
@@ -138,7 +138,7 @@ struct FAnimSegment
 USTRUCT()
 struct FAnimTrack
 {
-	GENERATED_BODY()
+	GENERATED_USTRUCT_BODY()
 
 	UPROPERTY(EditAnywhere, BlueprintReadWrite, Category=AnimTrack, EditFixedSize)
 	TArray<FAnimSegment>	AnimSegments;
@@ -186,9 +186,7 @@ struct FAnimTrack
 UCLASS(abstract, MinimalAPI)
 class UAnimCompositeBase : public UAnimSequenceBase
 {
-	GENERATED_BODY()
-public:
-	ENGINE_API UAnimCompositeBase(const FObjectInitializer& ObjectInitializer = FObjectInitializer::Get());
+	GENERATED_UCLASS_BODY()
 
 #if WITH_EDITOR
 	/** Set Sequence Length */
diff --git a/Engine/Source/Runtime/Engine/Classes/Animation/AnimCompress.h b/Engine/Source/Runtime/Engine/Classes/Animation/AnimCompress.h
--- a/Engine/Source/Runtime/Engine/Classes/Animation/AnimCompress.h
+++ b/Engine/Source/Runtime/Engine/Classes/Animation/AnimCompress.h
@@ -13,9 +13,7 @@
 UCLASS(abstract, hidecategories=Object, MinimalAPI)
 class UAnimCompress : public UObject
 {
-	GENERATED_BODY()
-public:
-	ENGINE_API UAnimCompress(const FObjectInitializer& ObjectInitializer = FObjectInitializer::Get());
+	GENERATED_UCLASS_BODY()
 
 	/** Name of Compression Scheme used for this asset*/
 	UPROPERTY(Category=Compression, VisibleAnywhere)
diff --git a/Engine/Source/Runtime/Engine/Classes/Animation/AnimCompress_Automatic.h b/Engine/Source/Runtime/Engine/Classes/Animation/AnimCompress_Automatic.h
--- a/Engine/Source/Runtime/Engine/Classes/Animation/AnimCompress_Automatic.h
+++ b/Engine/Source/Runtime/Engine/Classes/Animation/AnimCompress_Automatic.h
@@ -13,9 +13,7 @@
 UCLASS(hidecategories=AnimationCompressionAlgorithm)
 class UAnimCompress_Automatic : public UAnimCompress
 {
-	GENERATED_BODY()
-public:
-	UAnimCompress_Automatic(const FObjectInitializer& ObjectInitializer = FObjectInitializer::Get());
+	GENERATED_UCLASS_BODY()
 
 	/** Maximum amount of error that a compression technique can introduce in an end effector */
 	UPROPERTY(EditAnywhere, Category=AnimationCompressionAlgorithm_Automatic)
diff --git a/Engine/Source/Runtime/Engine/Classes/Animation/AnimCompress_BitwiseCompressOnly.h b/Engine/Source/Runtime/Engine/Classes/Animation/AnimCompress_BitwiseCompressOnly.h
--- a/Engine/Source/Runtime/Engine/Classes/Animation/AnimCompress_BitwiseCompressOnly.h
+++ b/Engine/Source/Runtime/Engine/Classes/Animation/AnimCompress_BitwiseCompressOnly.h
@@ -12,9 +12,7 @@
 UCLASS(MinimalAPI)
 class UAnimCompress_BitwiseCompressOnly : public UAnimCompress
 {
-	GENERATED_BODY()
-public:
-	ENGINE_API UAnimCompress_BitwiseCompressOnly(const FObjectInitializer& ObjectInitializer = FObjectInitializer::Get());
+	GENERATED_UCLASS_BODY()
 
 
 protected:
diff --git a/Engine/Source/Runtime/Engine/Classes/Animation/AnimCompress_LeastDestructive.h b/Engine/Source/Runtime/Engine/Classes/Animation/AnimCompress_LeastDestructive.h
--- a/Engine/Source/Runtime/Engine/Classes/Animation/AnimCompress_LeastDestructive.h
+++ b/Engine/Source/Runtime/Engine/Classes/Animation/AnimCompress_LeastDestructive.h
@@ -12,9 +12,7 @@
 UCLASS()
 class UAnimCompress_LeastDestructive : public UAnimCompress
 {
-	GENERATED_BODY()
-public:
-	UAnimCompress_LeastDestructive(const FObjectInitializer& ObjectInitializer = FObjectInitializer::Get());
+	GENERATED_UCLASS_BODY()
 
 
 protected:
diff --git a/Engine/Source/Runtime/Engine/Classes/Animation/AnimCompress_PerTrackCompression.h b/Engine/Source/Runtime/Engine/Classes/Animation/AnimCompress_PerTrackCompression.h
--- a/Engine/Source/Runtime/Engine/Classes/Animation/AnimCompress_PerTrackCompression.h
+++ b/Engine/Source/Runtime/Engine/Classes/Animation/AnimCompress_PerTrackCompression.h
@@ -13,9 +13,7 @@
 UCLASS(hidecategories=AnimCompress)
 class UAnimCompress_PerTrackCompression : public UAnimCompress_RemoveLinearKeys
 {
-	GENERATED_BODY()
-public:
-	UAnimCompress_PerTrackCompression(const FObjectInitializer& ObjectInitializer = FObjectInitializer::Get());
+	GENERATED_UCLASS_BODY()
 
 	/** Maximum threshold to use when replacing a component with zero. Lower values retain more keys, but yield less compression. */
 	UPROPERTY(EditAnywhere, Category=PerTrack)
diff --git a/Engine/Source/Runtime/Engine/Classes/Animation/AnimCompress_RemoveLinearKeys.h b/Engine/Source/Runtime/Engine/Classes/Animation/AnimCompress_RemoveLinearKeys.h
--- a/Engine/Source/Runtime/Engine/Classes/Animation/AnimCompress_RemoveLinearKeys.h
+++ b/Engine/Source/Runtime/Engine/Classes/Animation/AnimCompress_RemoveLinearKeys.h
@@ -12,9 +12,7 @@
 UCLASS(MinimalAPI)
 class UAnimCompress_RemoveLinearKeys : public UAnimCompress
 {
-	GENERATED_BODY()
-public:
-	ENGINE_API UAnimCompress_RemoveLinearKeys(const FObjectInitializer& ObjectInitializer = FObjectInitializer::Get());
+	GENERATED_UCLASS_BODY()
 
 	/** Maximum position difference to use when testing if an animation key may be removed. Lower values retain more keys, but yield less compression. */
 	UPROPERTY(EditAnywhere, Category=LinearKeyRemoval)
diff --git a/Engine/Source/Runtime/Engine/Classes/Animation/AnimCompress_RemoveTrivialKeys.h b/Engine/Source/Runtime/Engine/Classes/Animation/AnimCompress_RemoveTrivialKeys.h
--- a/Engine/Source/Runtime/Engine/Classes/Animation/AnimCompress_RemoveTrivialKeys.h
+++ b/Engine/Source/Runtime/Engine/Classes/Animation/AnimCompress_RemoveTrivialKeys.h
@@ -14,9 +14,7 @@
 UCLASS(MinimalAPI)
 class UAnimCompress_RemoveTrivialKeys : public UAnimCompress
 {
-	GENERATED_BODY()
-public:
-	ENGINE_API UAnimCompress_RemoveTrivialKeys(const FObjectInitializer& ObjectInitializer = FObjectInitializer::Get());
+	GENERATED_UCLASS_BODY()
 
 	UPROPERTY(EditAnywhere, Category=AnimationCompressionAlgorithm_RemoveTrivialKeys)
 	float MaxPosDiff;
diff --git a/Engine/Source/Runtime/Engine/Classes/Animation/AnimInstance.h b/Engine/Source/Runtime/Engine/Classes/Animation/AnimInstance.h
--- a/Engine/Source/Runtime/Engine/Classes/Animation/AnimInstance.h
+++ b/Engine/Source/Runtime/Engine/Classes/Animation/AnimInstance.h
@@ -83,7 +83,7 @@ enum EBoneRotationSource
 USTRUCT()
 struct FA2Pose
 {
-	GENERATED_BODY()
+	GENERATED_USTRUCT_BODY()
 
 	UPROPERTY()
 	TArray<FTransform> Bones;
@@ -95,7 +95,7 @@ struct FA2Pose
 USTRUCT()
 struct ENGINE_API FA2CSPose : public FA2Pose
 {
-	GENERATED_BODY()
+	GENERATED_USTRUCT_BODY()
 
 private:
 	/** Pointer to current BoneContainer. */
@@ -173,7 +173,7 @@ private:
 USTRUCT(BlueprintType)
 struct FBoneTransform
 {
-	GENERATED_BODY()
+	GENERATED_USTRUCT_BODY()
 
 	/** @todo anim: should be Skeleton bone index in the future, but right now it's Mesh BoneIndex **/
 	UPROPERTY()
@@ -196,7 +196,7 @@ struct FBoneTransform
 USTRUCT(BlueprintType)
 struct FPerBoneBlendWeight
 {
-	GENERATED_BODY()
+	GENERATED_USTRUCT_BODY()
 	UPROPERTY()
 	int32 SourceIndex; // source index of the buffer
 	UPROPERTY()
@@ -212,7 +212,7 @@ struct FPerBoneBlendWeight
 USTRUCT(BlueprintType)
 struct FPerBoneBlendWeights
 {
-	GENERATED_BODY()
+	GENERATED_USTRUCT_BODY()
 
 	UPROPERTY()
 	TArray<FPerBoneBlendWeight> BoneBlendWeights;
@@ -226,7 +226,7 @@ struct FPerBoneBlendWeights
 USTRUCT()
 struct FSlotEvaluationPose
 {
-	GENERATED_BODY()
+	GENERATED_USTRUCT_BODY()
 
 	/** Type of additive for pose */
 	UPROPERTY()
@@ -341,9 +341,7 @@ struct FNativeStateBinding
 UCLASS(transient, Blueprintable, hideCategories=AnimInstance, BlueprintType)
 class ENGINE_API UAnimInstance : public UObject
 {
-	GENERATED_BODY()
-public:
-	UAnimInstance(const FObjectInitializer& ObjectInitializer = FObjectInitializer::Get());
+	GENERATED_UCLASS_BODY()
 
 	/** DeltaTime **/
 	UPROPERTY()
diff --git a/Engine/Source/Runtime/Engine/Classes/Animation/AnimMontage.h b/Engine/Source/Runtime/Engine/Classes/Animation/AnimMontage.h
--- a/Engine/Source/Runtime/Engine/Classes/Animation/AnimMontage.h
+++ b/Engine/Source/Runtime/Engine/Classes/Animation/AnimMontage.h
@@ -21,7 +21,7 @@ class UAnimMontage;
 USTRUCT()
 struct FCompositeSection : public FAnimLinkableElement
 {
-	GENERATED_BODY()
+	GENERATED_USTRUCT_BODY()
 
 	/** Section Name */
 	UPROPERTY(EditAnywhere, Category=Section)
@@ -50,7 +50,7 @@ struct FCompositeSection : public FAnimLinkableElement
 USTRUCT()
 struct FSlotAnimationTrack
 {
-	GENERATED_BODY()
+	GENERATED_USTRUCT_BODY()
 
 	UPROPERTY(EditAnywhere, Category=Slot)
 	FName SlotName;
@@ -69,7 +69,7 @@ struct FSlotAnimationTrack
 USTRUCT()
 struct FBranchingPoint : public FAnimLinkableElement
 {
-	GENERATED_BODY()
+	GENERATED_USTRUCT_BODY()
 
 	UPROPERTY(EditAnywhere, Category=BranchingPoint)
 	FName EventName;
@@ -103,7 +103,7 @@ namespace EAnimNotifyEventType
 USTRUCT()
 struct FBranchingPointMarker
 {
-	GENERATED_BODY()
+	GENERATED_USTRUCT_BODY()
 
 	UPROPERTY()
 	int32 NotifyIndex;
@@ -147,7 +147,7 @@ DECLARE_DELEGATE_TwoParams( FOnMontageBlendingOutStarted, class UAnimMontage*, b
 USTRUCT()
 struct FAnimMontageInstance
 {
-	GENERATED_BODY()
+	GENERATED_USTRUCT_BODY()
 
 	// Montage reference
 	UPROPERTY()
@@ -309,9 +309,7 @@ private:
 UCLASS(config=Engine, hidecategories=(UObject, Length), MinimalAPI, BlueprintType)
 class UAnimMontage : public UAnimCompositeBase
 {
-	GENERATED_BODY()
-public:
-	ENGINE_API UAnimMontage(const FObjectInitializer& ObjectInitializer = FObjectInitializer::Get());
+	GENERATED_UCLASS_BODY()
 
 	/** Default blend in time. */
 	UPROPERTY(EditAnywhere, Category=Montage)
diff --git a/Engine/Source/Runtime/Engine/Classes/Animation/AnimNodeBase.h b/Engine/Source/Runtime/Engine/Classes/Animation/AnimNodeBase.h
--- a/Engine/Source/Runtime/Engine/Classes/Animation/AnimNodeBase.h
+++ b/Engine/Source/Runtime/Engine/Classes/Animation/AnimNodeBase.h
@@ -261,7 +261,7 @@ namespace EPinHidingMode
 USTRUCT()
 struct ENGINE_API FPoseLinkBase
 {
-	GENERATED_BODY()
+	GENERATED_USTRUCT_BODY()
 
 	// Serialized link ID, used to build the non-serialized pointer map
 	UPROPERTY()
@@ -305,7 +305,7 @@ public:
 USTRUCT()
 struct ENGINE_API FPoseLink : public FPoseLinkBase
 {
-	GENERATED_BODY()
+	GENERATED_USTRUCT_BODY()
 
 public:
 	// Interface
@@ -316,7 +316,7 @@ public:
 USTRUCT()
 struct ENGINE_API FComponentSpacePoseLink : public FPoseLinkBase
 {
-	GENERATED_BODY()
+	GENERATED_USTRUCT_BODY()
 
 public:
 	// Interface
@@ -327,7 +327,7 @@ public:
 USTRUCT()
 struct FExposedValueHandler
 {
-	GENERATED_BODY()
+	GENERATED_USTRUCT_BODY()
 
 	// The function to call to update associated properties (can be NULL)
 	UPROPERTY()
@@ -354,7 +354,7 @@ struct FExposedValueHandler
 USTRUCT()
 struct ENGINE_API FAnimNode_Base
 {
-	GENERATED_BODY()
+	GENERATED_USTRUCT_BODY()
 
 	// The default handler for graph-exposed inputs
 	UPROPERTY()
diff --git a/Engine/Source/Runtime/Engine/Classes/Animation/AnimNodeSpaceConversions.h b/Engine/Source/Runtime/Engine/Classes/Animation/AnimNodeSpaceConversions.h
--- a/Engine/Source/Runtime/Engine/Classes/Animation/AnimNodeSpaceConversions.h
+++ b/Engine/Source/Runtime/Engine/Classes/Animation/AnimNodeSpaceConversions.h
@@ -8,7 +8,7 @@
 USTRUCT()
 struct ENGINE_API FAnimNode_ConvertComponentToLocalSpace : public FAnimNode_Base
 {
-	GENERATED_BODY()
+	GENERATED_USTRUCT_BODY()
 
 public:
 	UPROPERTY(EditAnywhere, BlueprintReadWrite, Category=Links)
@@ -30,7 +30,7 @@ public:
 USTRUCT()
 struct ENGINE_API FAnimNode_ConvertLocalToComponentSpace : public FAnimNode_Base
 {
-	GENERATED_BODY()
+	GENERATED_USTRUCT_BODY()
 
 public:
 	UPROPERTY(EditAnywhere, BlueprintReadWrite, Category=Links)
diff --git a/Engine/Source/Runtime/Engine/Classes/Animation/AnimNotifies/AnimNotify.h b/Engine/Source/Runtime/Engine/Classes/Animation/AnimNotifies/AnimNotify.h
--- a/Engine/Source/Runtime/Engine/Classes/Animation/AnimNotifies/AnimNotify.h
+++ b/Engine/Source/Runtime/Engine/Classes/Animation/AnimNotifies/AnimNotify.h
@@ -12,16 +12,13 @@ struct FAnimNotifyEvent;
 UCLASS(abstract, Blueprintable, const, hidecategories=Object, collapsecategories)
 class ENGINE_API UAnimNotify : public UObject
 {
-	GENERATED_BODY()
-public:
-	UAnimNotify(const FObjectInitializer& ObjectInitializer = FObjectInitializer::Get());
+	GENERATED_UCLASS_BODY()
 
 	/** 
 	 * Implementable event to get a custom name for the notify
 	 */
 	UFUNCTION(BlueprintNativeEvent)
 	FString GetNotifyName() const;
-	virtual FString GetNotifyName_Implementation() const;
 
 	UFUNCTION(BlueprintImplementableEvent)
 	virtual bool Received_Notify(USkeletalMeshComponent* MeshComp, UAnimSequenceBase* Animation) const;
diff --git a/Engine/Source/Runtime/Engine/Classes/Animation/AnimNotifies/AnimNotifyState.h b/Engine/Source/Runtime/Engine/Classes/Animation/AnimNotifies/AnimNotifyState.h
--- a/Engine/Source/Runtime/Engine/Classes/Animation/AnimNotifies/AnimNotifyState.h
+++ b/Engine/Source/Runtime/Engine/Classes/Animation/AnimNotifies/AnimNotifyState.h
@@ -12,16 +12,13 @@ struct FAnimNotifyEvent;
 UCLASS(abstract, editinlinenew, Blueprintable, const, hidecategories=Object, collapsecategories, meta=(ShowWorldContextPin))
 class ENGINE_API UAnimNotifyState : public UObject
 {
-	GENERATED_BODY()
-public:
-	UAnimNotifyState(const FObjectInitializer& ObjectInitializer = FObjectInitializer::Get());
+	GENERATED_UCLASS_BODY()
 
 	/** 
 	 * Implementable event to get a custom name for the notify
 	 */
 	UFUNCTION(BlueprintNativeEvent)
 	FString GetNotifyName() const;
-	virtual FString GetNotifyName_Implementation() const;
 
 	UFUNCTION(BlueprintImplementableEvent)
 	virtual bool Received_NotifyBegin(USkeletalMeshComponent * MeshComp, UAnimSequenceBase * Animation, float TotalDuration) const;
diff --git a/Engine/Source/Runtime/Engine/Classes/Animation/AnimNotifies/AnimNotifyState_TimedParticleEffect.h b/Engine/Source/Runtime/Engine/Classes/Animation/AnimNotifies/AnimNotifyState_TimedParticleEffect.h
--- a/Engine/Source/Runtime/Engine/Classes/Animation/AnimNotifies/AnimNotifyState_TimedParticleEffect.h
+++ b/Engine/Source/Runtime/Engine/Classes/Animation/AnimNotifies/AnimNotifyState_TimedParticleEffect.h
@@ -10,9 +10,7 @@
 UCLASS(MinimalAPI, Blueprintable, meta = (DisplayName = "Timed Particle Effect"))
 class UAnimNotifyState_TimedParticleEffect : public UAnimNotifyState
 {
-	GENERATED_BODY()
-public:
-	ENGINE_API UAnimNotifyState_TimedParticleEffect(const FObjectInitializer& ObjectInitializer = FObjectInitializer::Get());
+	GENERATED_UCLASS_BODY()
 
 	// The particle system template to use when spawning the particle component
 	UPROPERTY(EditAnywhere, Category = ParticleSystem, meta = (ToolTip = "The particle system to spawn for the notify state"))
@@ -62,4 +60,4 @@ public:
 
 private:
 	bool ValidateParameters(USkeletalMeshComponent* MeshComp);
-};
+};
\ No newline at end of file
diff --git a/Engine/Source/Runtime/Engine/Classes/Animation/AnimNotifies/AnimNotifyState_Trail.h b/Engine/Source/Runtime/Engine/Classes/Animation/AnimNotifies/AnimNotifyState_Trail.h
--- a/Engine/Source/Runtime/Engine/Classes/Animation/AnimNotifies/AnimNotifyState_Trail.h
+++ b/Engine/Source/Runtime/Engine/Classes/Animation/AnimNotifies/AnimNotifyState_Trail.h
@@ -9,9 +9,7 @@ DECLARE_LOG_CATEGORY_EXTERN(LogAnimTrails, Log, All);
 UCLASS(editinlinenew, Blueprintable, const, hidecategories = Object, collapsecategories, MinimalAPI, meta = (ShowWorldContextPin, DisplayName = "Trail"))
 class UAnimNotifyState_Trail : public UAnimNotifyState
 {
-	GENERATED_BODY()
-public:
-	ENGINE_API UAnimNotifyState_Trail(const FObjectInitializer& ObjectInitializer = FObjectInitializer::Get());
+	GENERATED_UCLASS_BODY()
 	
 	/** The particle system to use for this trail. */
 	UPROPERTY(EditAnywhere, BlueprintReadWrite, Category = Trail)
diff --git a/Engine/Source/Runtime/Engine/Classes/Animation/AnimSequence.h b/Engine/Source/Runtime/Engine/Classes/Animation/AnimSequence.h
--- a/Engine/Source/Runtime/Engine/Classes/Animation/AnimSequence.h
+++ b/Engine/Source/Runtime/Engine/Classes/Animation/AnimSequence.h
@@ -84,7 +84,7 @@ enum EAdditiveBasePoseType
 USTRUCT()
 struct FRawAnimSequenceTrack
 {
-	GENERATED_BODY()
+	GENERATED_USTRUCT_BODY()
 
 	/** Position keys. */
 	UPROPERTY()
@@ -121,7 +121,7 @@ struct FRawAnimSequenceTrack
 USTRUCT()
 struct FAnimSequenceTrackContainer
 {
-	GENERATED_BODY()
+	GENERATED_USTRUCT_BODY()
 
 	UPROPERTY()
 	TArray<struct FRawAnimSequenceTrack> AnimationTracks;
@@ -159,7 +159,7 @@ struct FAnimSequenceTrackContainer
 USTRUCT()
 struct FTrackToSkeletonMap
 {
-	GENERATED_BODY()
+	GENERATED_USTRUCT_BODY()
 
 	UPROPERTY()
 	int32 SkeletonIndex_DEPRECATED;    // 0 as current Skeleton (the one above), and N is the N'th parent
@@ -185,7 +185,7 @@ struct FTrackToSkeletonMap
 USTRUCT()
 struct FTranslationTrack
 {
-	GENERATED_BODY()
+	GENERATED_USTRUCT_BODY()
 
 	UPROPERTY()
 	TArray<FVector> PosKeys;
@@ -200,7 +200,7 @@ struct FTranslationTrack
 USTRUCT()
 struct FRotationTrack
 {
-	GENERATED_BODY()
+	GENERATED_USTRUCT_BODY()
 
 	UPROPERTY()
 	TArray<FQuat> RotKeys;
@@ -215,7 +215,7 @@ struct FRotationTrack
 USTRUCT()
 struct FScaleTrack
 {
-	GENERATED_BODY()
+	GENERATED_USTRUCT_BODY()
 
 	UPROPERTY()
 	TArray<FVector> ScaleKeys;
@@ -233,7 +233,7 @@ struct FScaleTrack
 USTRUCT()
 struct FCurveTrack
 {
-	GENERATED_BODY()
+	GENERATED_USTRUCT_BODY()
 
 	UPROPERTY()
 	FName CurveName;
@@ -253,7 +253,7 @@ struct FCurveTrack
 USTRUCT()
 struct FCompressedTrack
 {
-	GENERATED_BODY()
+	GENERATED_USTRUCT_BODY()
 
 	UPROPERTY()
 	TArray<uint8> ByteStream;
@@ -285,7 +285,7 @@ struct FCompressedTrack
 USTRUCT()
 struct FCompressedOffsetData
 {
-	GENERATED_BODY()
+	GENERATED_USTRUCT_BODY()
 
 	UPROPERTY()
 	TArray<int32> OffsetData;
@@ -345,9 +345,7 @@ struct FCompressedOffsetData
 UCLASS(config=Engine, hidecategories=(UObject, Length), MinimalAPI, BlueprintType)
 class UAnimSequence : public UAnimSequenceBase
 {
-	GENERATED_BODY()
-public:
-	ENGINE_API UAnimSequence(const FObjectInitializer& ObjectInitializer = FObjectInitializer::Get());
+	GENERATED_UCLASS_BODY()
 
 	/** Number of raw frames in this sequence (not used by engine - just for informational purposes). */
 	UPROPERTY(AssetRegistrySearchable)
diff --git a/Engine/Source/Runtime/Engine/Classes/Animation/AnimSequenceBase.h b/Engine/Source/Runtime/Engine/Classes/Animation/AnimSequenceBase.h
--- a/Engine/Source/Runtime/Engine/Classes/Animation/AnimSequenceBase.h
+++ b/Engine/Source/Runtime/Engine/Classes/Animation/AnimSequenceBase.h
@@ -63,7 +63,7 @@ namespace ENotifyFilterType
 USTRUCT()
 struct FAnimNotifyEvent : public FAnimLinkableElement
 {
-	GENERATED_BODY()
+	GENERATED_USTRUCT_BODY()
 
 	/** The user requested time for this notify */
 	UPROPERTY()
@@ -193,7 +193,7 @@ struct FAnimNotifyEvent : public FAnimLinkableElement
 USTRUCT()
 struct FAnimNotifyTrack
 {
-	GENERATED_BODY()
+	GENERATED_USTRUCT_BODY()
 
 	UPROPERTY()
 	FName TrackName;
@@ -247,7 +247,7 @@ enum EAnimCurveFlags
 USTRUCT()
 struct FAnimCurveBase
 {
-	GENERATED_BODY()
+	GENERATED_USTRUCT_BODY()
 
 	// Last observed name of the curve. We store this so we can recover from situations that
 	// mean the skeleton doesn't have a mapped name for our UID (such as a user saving the an
@@ -316,7 +316,7 @@ public:
 USTRUCT()
 struct FFloatCurve : public FAnimCurveBase
 {
-	GENERATED_BODY()
+	GENERATED_USTRUCT_BODY()
 
 	/** Curve data for float. */
 	UPROPERTY()
@@ -348,7 +348,7 @@ struct FFloatCurve : public FAnimCurveBase
 USTRUCT()
 struct FVectorCurve : public FAnimCurveBase
 {
-	GENERATED_BODY()
+	GENERATED_USTRUCT_BODY()
 
 	enum EIndex
 	{
@@ -389,7 +389,7 @@ struct FVectorCurve : public FAnimCurveBase
 USTRUCT()
 struct FTransformCurve: public FAnimCurveBase
 {
-	GENERATED_BODY()
+	GENERATED_USTRUCT_BODY()
 
 	/** Curve data for each transform. */
 	UPROPERTY()
@@ -433,7 +433,7 @@ struct FTransformCurve: public FAnimCurveBase
 USTRUCT()
 struct FRawCurveTracks
 {
-	GENERATED_BODY()
+	GENERATED_USTRUCT_BODY()
 
 	enum ESupportedCurveType
 	{
@@ -564,9 +564,7 @@ enum ETypeAdvanceAnim
 UCLASS(abstract, MinimalAPI, BlueprintType)
 class UAnimSequenceBase : public UAnimationAsset
 {
-	GENERATED_BODY()
-public:
-	ENGINE_API UAnimSequenceBase(const FObjectInitializer& ObjectInitializer = FObjectInitializer::Get());
+	GENERATED_UCLASS_BODY()
 
 	/** Animation notifies, sorted by time (earliest notification first). */
 	UPROPERTY()
diff --git a/Engine/Source/Runtime/Engine/Classes/Animation/AnimSet.h b/Engine/Source/Runtime/Engine/Classes/Animation/AnimSet.h
--- a/Engine/Source/Runtime/Engine/Classes/Animation/AnimSet.h
+++ b/Engine/Source/Runtime/Engine/Classes/Animation/AnimSet.h
@@ -13,7 +13,7 @@
 USTRUCT()
 struct FAnimSetMeshLinkup
 {
-	GENERATED_BODY()
+	GENERATED_USTRUCT_BODY()
 
 	/** 
 	 * Mapping table. Size must be same as size of SkelMesh reference skeleton. 
@@ -33,9 +33,7 @@ struct FAnimSetMeshLinkup
 UCLASS(hidecategories=Object, MinimalAPI)
 class UAnimSet : public UObject
 {
-	GENERATED_BODY()
-public:
-	ENGINE_API UAnimSet(const FObjectInitializer& ObjectInitializer = FObjectInitializer::Get());
+	GENERATED_UCLASS_BODY()
 
 private:
 	/** 
diff --git a/Engine/Source/Runtime/Engine/Classes/Animation/AnimSingleNodeInstance.h b/Engine/Source/Runtime/Engine/Classes/Animation/AnimSingleNodeInstance.h
--- a/Engine/Source/Runtime/Engine/Classes/Animation/AnimSingleNodeInstance.h
+++ b/Engine/Source/Runtime/Engine/Classes/Animation/AnimSingleNodeInstance.h
@@ -15,9 +15,7 @@ DECLARE_DYNAMIC_DELEGATE(FPostEvaluateAnimEvent);
 UCLASS(transient, NotBlueprintable)
 class ENGINE_API UAnimSingleNodeInstance : public UAnimInstance
 {
-	GENERATED_BODY()
-public:
-	UAnimSingleNodeInstance(const FObjectInitializer& ObjectInitializer = FObjectInitializer::Get());
+	GENERATED_UCLASS_BODY()
 
 	/** Current Asset being played **/
 	UPROPERTY(transient)
diff --git a/Engine/Source/Runtime/Engine/Classes/Animation/AnimStateMachineTypes.h b/Engine/Source/Runtime/Engine/Classes/Animation/AnimStateMachineTypes.h
--- a/Engine/Source/Runtime/Engine/Classes/Animation/AnimStateMachineTypes.h
+++ b/Engine/Source/Runtime/Engine/Classes/Animation/AnimStateMachineTypes.h
@@ -30,7 +30,7 @@ namespace ETransitionLogicType
 USTRUCT()
 struct FAnimationTransitionRule
 {
-	GENERATED_BODY()
+	GENERATED_USTRUCT_BODY()
 
 	UPROPERTY()
 	FName RuleToExecute;
@@ -57,7 +57,7 @@ struct FAnimationTransitionRule
 USTRUCT()
 struct FAnimationStateBase
 {
-	GENERATED_BODY()
+	GENERATED_USTRUCT_BODY()
 
 	// The name of this state
 	UPROPERTY()
@@ -71,7 +71,7 @@ struct FAnimationStateBase
 USTRUCT()
 struct FAnimationState : public FAnimationStateBase
 {
-	GENERATED_BODY()
+	GENERATED_USTRUCT_BODY()
 
 	// Set of legal transitions out of this state; already in priority order
 	UPROPERTY()
@@ -105,7 +105,7 @@ struct FAnimationState : public FAnimationStateBase
 USTRUCT()
 struct FAnimationTransitionBetweenStates : public FAnimationStateBase
 {
-	GENERATED_BODY()
+	GENERATED_USTRUCT_BODY()
 
 	// Transition-only: State being transitioned from
 	UPROPERTY()
@@ -157,7 +157,7 @@ struct FAnimationTransitionBetweenStates : public FAnimationStateBase
 USTRUCT()
 struct FBakedStateExitTransition
 {
-	GENERATED_BODY()
+	GENERATED_USTRUCT_BODY()
 
 	// The node property index for this rule
 	UPROPERTY()
@@ -197,7 +197,7 @@ struct FBakedStateExitTransition
 USTRUCT()
 struct FBakedAnimationState
 {
-	GENERATED_BODY()
+	GENERATED_USTRUCT_BODY()
 
 	// The name of this state
 	UPROPERTY()
@@ -240,7 +240,7 @@ public:
 USTRUCT()
 struct FBakedAnimationStateMachine
 {
-	GENERATED_BODY()
+	GENERATED_USTRUCT_BODY()
 
 	// Name of this machine (primarily for debugging purposes)
 	UPROPERTY()
@@ -275,8 +275,6 @@ public:
 UCLASS()
 class UAnimStateMachineTypes : public UObject
 {
-	GENERATED_BODY()
-public:
-	UAnimStateMachineTypes(const FObjectInitializer& ObjectInitializer = FObjectInitializer::Get());
+	GENERATED_UCLASS_BODY()
 };
 
diff --git a/Engine/Source/Runtime/Engine/Classes/Animation/BlendSpace.h b/Engine/Source/Runtime/Engine/Classes/Animation/BlendSpace.h
--- a/Engine/Source/Runtime/Engine/Classes/Animation/BlendSpace.h
+++ b/Engine/Source/Runtime/Engine/Classes/Animation/BlendSpace.h
@@ -12,9 +12,7 @@
 UCLASS(config=Engine, hidecategories=Object, MinimalAPI, BlueprintType)
 class UBlendSpace : public UBlendSpaceBase
 {
-	GENERATED_BODY()
-public:
-	ENGINE_API UBlendSpace(const FObjectInitializer& ObjectInitializer = FObjectInitializer::Get());
+	GENERATED_UCLASS_BODY()
 
 public:
 
diff --git a/Engine/Source/Runtime/Engine/Classes/Animation/BlendSpaceBase.h b/Engine/Source/Runtime/Engine/Classes/Animation/BlendSpaceBase.h
--- a/Engine/Source/Runtime/Engine/Classes/Animation/BlendSpaceBase.h
+++ b/Engine/Source/Runtime/Engine/Classes/Animation/BlendSpaceBase.h
@@ -23,7 +23,7 @@ enum EBlendSpaceAxis
 USTRUCT()
 struct FInterpolationParameter
 {
-	GENERATED_BODY()
+	GENERATED_USTRUCT_BODY()
 
 	/** Interpolation Time for input, when it gets input, it will use this time to interpolate to target, used for smoother interpolation **/
 	UPROPERTY(EditAnywhere, Category=Parameter)
@@ -37,7 +37,7 @@ struct FInterpolationParameter
 USTRUCT()
 struct FBlendParameter
 {
-	GENERATED_BODY()
+	GENERATED_USTRUCT_BODY()
 
 	UPROPERTY(EditAnywhere, Category=BlendParameter)
 	FString DisplayName;
@@ -78,7 +78,7 @@ struct FBlendParameter
 USTRUCT()
 struct FBlendSample
 {
-	GENERATED_BODY()
+	GENERATED_USTRUCT_BODY()
 
 	UPROPERTY(EditAnywhere, Category=BlendSample)
 	class UAnimSequence* Animation;
@@ -115,7 +115,7 @@ struct FBlendSample
 USTRUCT()
 struct FEditorElement
 {
-	GENERATED_BODY()
+	GENERATED_USTRUCT_BODY()
 
 	UPROPERTY(EditAnywhere, Category=EditorElement)
 	int32 Indices[3];    /*MAX_VERTICES @fixmeconst*/
@@ -144,7 +144,7 @@ struct FEditorElement
 USTRUCT()
 struct FGridBlendSample
 {
-	GENERATED_BODY()
+	GENERATED_USTRUCT_BODY()
 
 	UPROPERTY()
 	struct FEditorElement GridElement;
@@ -163,7 +163,7 @@ struct FGridBlendSample
 USTRUCT()
 struct FPerBoneInterpolation
 {
-	GENERATED_BODY()
+	GENERATED_USTRUCT_BODY()
 
 	UPROPERTY(EditAnywhere, Category=FPerBoneInterpolation)
 	FBoneReference BoneReference;
@@ -198,9 +198,7 @@ namespace ENotifyTriggerMode
 UCLASS(config=Engine, hidecategories=Object, MinimalAPI, BlueprintType)
 class UBlendSpaceBase : public UAnimationAsset
 {
-	GENERATED_BODY()
-public:
-	ENGINE_API UBlendSpaceBase(const FObjectInitializer& ObjectInitializer = FObjectInitializer::Get());
+	GENERATED_UCLASS_BODY()
 
 protected:
 	/** Blend Parameters for each axis. **/
diff --git a/Engine/Source/Runtime/Engine/Classes/Animation/BoneControllers/AnimNode_Fabrik.h b/Engine/Source/Runtime/Engine/Classes/Animation/BoneControllers/AnimNode_Fabrik.h
--- a/Engine/Source/Runtime/Engine/Classes/Animation/BoneControllers/AnimNode_Fabrik.h
+++ b/Engine/Source/Runtime/Engine/Classes/Animation/BoneControllers/AnimNode_Fabrik.h
@@ -13,7 +13,7 @@
 USTRUCT()
 struct FABRIKChainLink
 {
-	GENERATED_BODY()
+	GENERATED_USTRUCT_BODY()
 
 public:
 	/** Position of bone in component space. */
@@ -53,7 +53,7 @@ public:
 USTRUCT()
 struct ENGINE_API FAnimNode_Fabrik : public FAnimNode_SkeletalControlBase
 {
-	GENERATED_BODY()
+	GENERATED_USTRUCT_BODY()
 
 	/** Coordinates for target location of tip bone - if EffectorLocationSpace is bone, this is the offset from Target Bone to use as target location*/
 	UPROPERTY(EditAnywhere, BlueprintReadWrite, Category = EndEffector, meta = (PinShownByDefault))
@@ -106,4 +106,4 @@ private:
 
 	// Convenience function to get current (pre-translation iteration) component space location of bone by bone index
 	FVector GetCurrentLocation(FA2CSPose & MeshBases, int32 const & BoneIndex);
-};
+};
\ No newline at end of file
diff --git a/Engine/Source/Runtime/Engine/Classes/Animation/BoneControllers/AnimNode_Trail.h b/Engine/Source/Runtime/Engine/Classes/Animation/BoneControllers/AnimNode_Trail.h
--- a/Engine/Source/Runtime/Engine/Classes/Animation/BoneControllers/AnimNode_Trail.h
+++ b/Engine/Source/Runtime/Engine/Classes/Animation/BoneControllers/AnimNode_Trail.h
@@ -12,7 +12,7 @@
 USTRUCT()
 struct ENGINE_API FAnimNode_Trail : public FAnimNode_SkeletalControlBase
 {
-	GENERATED_BODY()
+	GENERATED_USTRUCT_BODY()
 
 	/** Reference to the active bone in the hierarchy to modify. */
 	UPROPERTY(EditAnywhere, BlueprintReadWrite, Category=Trail)
@@ -81,4 +81,4 @@ private:
 	// End of FAnimNode_SkeletalControlBase interface
 
 	FVector GetAlignVector(EAxis::Type AxisOption, bool bInvert);
-};
+};
\ No newline at end of file
diff --git a/Engine/Source/Runtime/Engine/Classes/Animation/BoneControllers/AnimNode_TwoBoneIK.h b/Engine/Source/Runtime/Engine/Classes/Animation/BoneControllers/AnimNode_TwoBoneIK.h
--- a/Engine/Source/Runtime/Engine/Classes/Animation/BoneControllers/AnimNode_TwoBoneIK.h
+++ b/Engine/Source/Runtime/Engine/Classes/Animation/BoneControllers/AnimNode_TwoBoneIK.h
@@ -12,7 +12,7 @@
 USTRUCT()
 struct ENGINE_API FAnimNode_TwoBoneIK : public FAnimNode_SkeletalControlBase
 {
-	GENERATED_BODY()
+	GENERATED_USTRUCT_BODY()
 	
 	/** Name of bone to control. This is the main bone chain to modify from. **/
 	UPROPERTY(EditAnywhere, BlueprintReadWrite, Category=IK)
@@ -73,4 +73,4 @@ private:
 	// FAnimNode_SkeletalControlBase interface
 	virtual void InitializeBoneReferences(const FBoneContainer& RequiredBones) override;
 	// End of FAnimNode_SkeletalControlBase interface
-};
+};
\ No newline at end of file
diff --git a/Engine/Source/Runtime/Engine/Classes/Animation/Rig.h b/Engine/Source/Runtime/Engine/Classes/Animation/Rig.h
--- a/Engine/Source/Runtime/Engine/Classes/Animation/Rig.h
+++ b/Engine/Source/Runtime/Engine/Classes/Animation/Rig.h
@@ -17,7 +17,7 @@
 USTRUCT()
 struct FNode
 {
-	GENERATED_BODY()
+	GENERATED_USTRUCT_BODY()
 
 	/** Name of the original node. We don't allow to change this. This is used for identity.**/
 	UPROPERTY(VisibleAnywhere, Category="FNode")
@@ -86,7 +86,7 @@ namespace EConstraintTransform
 USTRUCT()
 struct FRigTransformConstraint
 {
-	GENERATED_BODY()
+	GENERATED_USTRUCT_BODY()
 
 	/** What transform type **/
 	UPROPERTY(/*EditAnywhere, Category="FTransformBaseConstraint"*/)
@@ -105,7 +105,7 @@ struct FRigTransformConstraint
 USTRUCT()
 struct FTransformBaseConstraint
 {
-	GENERATED_BODY()
+	GENERATED_USTRUCT_BODY()
 
 	/** What transform type **/
 	UPROPERTY(EditAnywhere, Category="FTransformBaseConstraint")
@@ -117,7 +117,7 @@ struct FTransformBaseConstraint
 USTRUCT()
 struct FTransformBase
 {
-	GENERATED_BODY()
+	GENERATED_USTRUCT_BODY()
 
 	UPROPERTY(EditAnywhere, Category="FTransformBase")
 	FName					Node;
@@ -136,9 +136,7 @@ DECLARE_DELEGATE_RetVal_OneParam(int32, FGetParentIndex, FName);
 UCLASS(hidecategories=Object, MinimalAPI)
 class URig : public UObject
 {
-	GENERATED_BODY()
-public:
-	ENGINE_API URig(const FObjectInitializer& ObjectInitializer = FObjectInitializer::Get());
+	GENERATED_UCLASS_BODY()
 
 private:
 
diff --git a/Engine/Source/Runtime/Engine/Classes/Animation/SkeletalMeshActor.h b/Engine/Source/Runtime/Engine/Classes/Animation/SkeletalMeshActor.h
--- a/Engine/Source/Runtime/Engine/Classes/Animation/SkeletalMeshActor.h
+++ b/Engine/Source/Runtime/Engine/Classes/Animation/SkeletalMeshActor.h
@@ -15,9 +15,7 @@
 UCLASS(ClassGroup=ISkeletalMeshes, Blueprintable, ComponentWrapperClass, ConversionRoot, meta=(ChildCanTick))
 class ENGINE_API ASkeletalMeshActor : public AActor, public IMatineeAnimInterface
 {
-	GENERATED_BODY()
-public:
-	ASkeletalMeshActor(const FObjectInitializer& ObjectInitializer = FObjectInitializer::Get());
+	GENERATED_UCLASS_BODY()
 
 	/** Whether or not this actor should respond to anim notifies - CURRENTLY ONLY AFFECTS PlayParticleEffect NOTIFIES**/
 	UPROPERTY(EditAnywhere, BlueprintReadWrite, Category=Animation, AdvancedDisplay)
diff --git a/Engine/Source/Runtime/Engine/Classes/Animation/Skeleton.h b/Engine/Source/Runtime/Engine/Classes/Animation/Skeleton.h
--- a/Engine/Source/Runtime/Engine/Classes/Animation/Skeleton.h
+++ b/Engine/Source/Runtime/Engine/Classes/Animation/Skeleton.h
@@ -17,7 +17,7 @@ class USkeletalMesh;
 USTRUCT()
 struct FSkeletonToMeshLinkup
 {
-	GENERATED_BODY()
+	GENERATED_USTRUCT_BODY()
 
 	/** 
 	 * Mapping table. Size must be same as size of bone tree (not Mesh Ref Pose). 
@@ -58,7 +58,7 @@ namespace EBoneTranslationRetargetingMode
 USTRUCT()
 struct FBoneNode
 {
-	GENERATED_BODY()
+	GENERATED_USTRUCT_BODY()
 
 	/** Name of bone, this is the search criteria to match with mesh bone. This will be NAME_None if deleted **/
 	UPROPERTY()
@@ -90,7 +90,7 @@ struct FBoneNode
 USTRUCT()
 struct FReferencePose
 {
-	GENERATED_BODY()
+	GENERATED_USTRUCT_BODY()
 
 	UPROPERTY()
 	FName	PoseName;
@@ -117,7 +117,7 @@ struct FReferencePose
 USTRUCT()
 struct FBoneReductionSetting
 {
-	GENERATED_BODY()
+	GENERATED_USTRUCT_BODY()
 
 	UPROPERTY()
 	TArray<FName> BonesToRemove;
@@ -147,7 +147,7 @@ struct FBoneReductionSetting
 USTRUCT()
 struct FNameMapping
 {
-	GENERATED_BODY()
+	GENERATED_USTRUCT_BODY()
 
 	UPROPERTY()
 	FName NodeName;
@@ -177,7 +177,7 @@ struct FNameMapping
 USTRUCT()
 struct FRigConfiguration
 {
-	GENERATED_BODY()
+	GENERATED_USTRUCT_BODY()
 
 	UPROPERTY()
 	class URig * Rig;
@@ -190,7 +190,7 @@ struct FRigConfiguration
 USTRUCT()
 struct FAnimSlotGroup
 {
-	GENERATED_BODY()
+	GENERATED_USTRUCT_BODY()
 
 public:
 	static ENGINE_API const FName DefaultGroupName;
@@ -223,9 +223,7 @@ public:
 UCLASS(hidecategories=Object, MinimalAPI)
 class USkeleton : public UObject
 {
-	GENERATED_BODY()
-public:
-	ENGINE_API USkeleton(const FObjectInitializer& ObjectInitializer = FObjectInitializer::Get());
+	GENERATED_UCLASS_BODY()
 
 private:
 	/** Skeleton bone tree - each contains name and parent index**/
diff --git a/Engine/Source/Runtime/Engine/Classes/Animation/SmartName.h b/Engine/Source/Runtime/Engine/Classes/Animation/SmartName.h
--- a/Engine/Source/Runtime/Engine/Classes/Animation/SmartName.h
+++ b/Engine/Source/Runtime/Engine/Classes/Animation/SmartName.h
@@ -7,7 +7,7 @@
 USTRUCT()
 struct ENGINE_API FSmartNameMapping
 {
-	GENERATED_BODY()
+	GENERATED_USTRUCT_BODY();
 
 	FSmartNameMapping();
 
@@ -83,7 +83,7 @@ private:
 USTRUCT()
 struct ENGINE_API FSmartNameContainer
 {
-	GENERATED_BODY()
+	GENERATED_USTRUCT_BODY();
 
 	// Add a new smartname container with the provided name
 	void AddContainer(FName NewContainerName);
@@ -99,4 +99,4 @@ struct ENGINE_API FSmartNameContainer
 
 private:
 	TMap<FName, FSmartNameMapping> NameMappings;	// List of smartname mappings
-};
+};
\ No newline at end of file
diff --git a/Engine/Source/Runtime/Engine/Classes/Animation/VertexAnim/MorphTarget.h b/Engine/Source/Runtime/Engine/Classes/Animation/VertexAnim/MorphTarget.h
--- a/Engine/Source/Runtime/Engine/Classes/Animation/VertexAnim/MorphTarget.h
+++ b/Engine/Source/Runtime/Engine/Classes/Animation/VertexAnim/MorphTarget.h
@@ -66,9 +66,7 @@ struct FMorphTargetLODModel
 UCLASS(hidecategories=Object, MinimalAPI)
 class UMorphTarget : public UVertexAnimBase
 {
-	GENERATED_BODY()
-public:
-	ENGINE_API UMorphTarget(const FObjectInitializer& ObjectInitializer = FObjectInitializer::Get());
+	GENERATED_UCLASS_BODY()
 
 public:
 	/** morph mesh vertex data for each LOD */
diff --git a/Engine/Source/Runtime/Engine/Classes/Animation/VertexAnim/VertexAnimBase.h b/Engine/Source/Runtime/Engine/Classes/Animation/VertexAnim/VertexAnimBase.h
--- a/Engine/Source/Runtime/Engine/Classes/Animation/VertexAnim/VertexAnimBase.h
+++ b/Engine/Source/Runtime/Engine/Classes/Animation/VertexAnim/VertexAnimBase.h
@@ -47,9 +47,7 @@ struct FVertexAnimEvalStateBase
 UCLASS(hidecategories=Object, MinimalAPI, abstract)
 class UVertexAnimBase : public UObject
 {
-	GENERATED_BODY()
-public:
-	ENGINE_API UVertexAnimBase(const FObjectInitializer& ObjectInitializer = FObjectInitializer::Get());
+	GENERATED_UCLASS_BODY()
 
 	/** USkeletalMesh that this vertex animation works on. */
 	UPROPERTY(AssetRegistrySearchable)
diff --git a/Engine/Source/Runtime/Engine/Classes/Animation/VertexAnim/VertexAnimation.h b/Engine/Source/Runtime/Engine/Classes/Animation/VertexAnim/VertexAnimation.h
--- a/Engine/Source/Runtime/Engine/Classes/Animation/VertexAnim/VertexAnimation.h
+++ b/Engine/Source/Runtime/Engine/Classes/Animation/VertexAnim/VertexAnimation.h
@@ -27,9 +27,7 @@ struct FVertexAnimFrame
 UCLASS(hidecategories=Object, MinimalAPI)
 class UVertexAnimation : public UVertexAnimBase
 {
-	GENERATED_BODY()
-public:
-	ENGINE_API UVertexAnimation(const FObjectInitializer& ObjectInitializer = FObjectInitializer::Get());
+	GENERATED_UCLASS_BODY()
 
 	/** Number of verts animated by this animation, should be size of Deltas array in each frame */
 	UPROPERTY()
diff --git a/Engine/Source/Runtime/Engine/Classes/Atmosphere/AtmosphericFog.h b/Engine/Source/Runtime/Engine/Classes/Atmosphere/AtmosphericFog.h
--- a/Engine/Source/Runtime/Engine/Classes/Atmosphere/AtmosphericFog.h
+++ b/Engine/Source/Runtime/Engine/Classes/Atmosphere/AtmosphericFog.h
@@ -12,9 +12,7 @@
 UCLASS(showcategories=(Movement, Rendering, "Utilities|Transformation", "Input|MouseInput", "Input|TouchInput"), ClassGroup=Fog, hidecategories=(Info,Object,Input), MinimalAPI)
 class AAtmosphericFog : public AInfo
 {
-	GENERATED_BODY()
-public:
-	ENGINE_API AAtmosphericFog(const FObjectInitializer& ObjectInitializer = FObjectInitializer::Get());
+	GENERATED_UCLASS_BODY()
 
 private_subobject:
 	/** Main fog component */
diff --git a/Engine/Source/Runtime/Engine/Classes/Camera/CameraActor.h b/Engine/Source/Runtime/Engine/Classes/Camera/CameraActor.h
--- a/Engine/Source/Runtime/Engine/Classes/Camera/CameraActor.h
+++ b/Engine/Source/Runtime/Engine/Classes/Camera/CameraActor.h
@@ -12,9 +12,7 @@
 UCLASS(ClassGroup=Common, hideCategories=(Input, Rendering), showcategories=("Input|MouseInput", "Input|TouchInput"), Blueprintable)
 class ENGINE_API ACameraActor : public AActor
 {
-	GENERATED_BODY()
-public:
-	ACameraActor(const FObjectInitializer& ObjectInitializer = FObjectInitializer::Get());
+	GENERATED_UCLASS_BODY()
 
 private:
 
diff --git a/Engine/Source/Runtime/Engine/Classes/Camera/CameraAnim.h b/Engine/Source/Runtime/Engine/Classes/Camera/CameraAnim.h
--- a/Engine/Source/Runtime/Engine/Classes/Camera/CameraAnim.h
+++ b/Engine/Source/Runtime/Engine/Classes/Camera/CameraAnim.h
@@ -10,9 +10,7 @@
 UCLASS(notplaceable, MinimalAPI)
 class UCameraAnim : public UObject
 {
-	GENERATED_BODY()
-public:
-	ENGINE_API UCameraAnim(const FObjectInitializer& ObjectInitializer = FObjectInitializer::Get());
+	GENERATED_UCLASS_BODY()
 
 	/** The UInterpGroup that holds our actual interpolation data. */
 	UPROPERTY()
diff --git a/Engine/Source/Runtime/Engine/Classes/Camera/CameraAnimInst.h b/Engine/Source/Runtime/Engine/Classes/Camera/CameraAnimInst.h
--- a/Engine/Source/Runtime/Engine/Classes/Camera/CameraAnimInst.h
+++ b/Engine/Source/Runtime/Engine/Classes/Camera/CameraAnimInst.h
@@ -13,9 +13,7 @@
 UCLASS(notplaceable, BlueprintType, transient)
 class ENGINE_API UCameraAnimInst : public UObject
 {
-	GENERATED_BODY()
-public:
-	UCameraAnimInst(const FObjectInitializer& ObjectInitializer = FObjectInitializer::Get());
+	GENERATED_UCLASS_BODY()
 
 	/** which CameraAnim this is an instance of */
 	UPROPERTY()
diff --git a/Engine/Source/Runtime/Engine/Classes/Camera/CameraComponent.h b/Engine/Source/Runtime/Engine/Classes/Camera/CameraComponent.h
--- a/Engine/Source/Runtime/Engine/Classes/Camera/CameraComponent.h
+++ b/Engine/Source/Runtime/Engine/Classes/Camera/CameraComponent.h
@@ -12,9 +12,7 @@
 UCLASS(HideCategories=(Mobility, Rendering, LOD), Blueprintable, ClassGroup=Camera, meta=(BlueprintSpawnableComponent), MinimalAPI)
 class UCameraComponent : public USceneComponent
 {
-	GENERATED_BODY()
-public:
-	ENGINE_API UCameraComponent(const FObjectInitializer& ObjectInitializer = FObjectInitializer::Get());
+	GENERATED_UCLASS_BODY()
 
 	/** The horizontal field of view (in degrees) in perspective mode (ignored in Orthographic mode) */
 	UPROPERTY(Interp, EditAnywhere, BlueprintReadWrite, Category=CameraSettings, meta=(UIMin = "5.0", UIMax = "170", ClampMin = "0.001", ClampMax = "360.0"))
diff --git a/Engine/Source/Runtime/Engine/Classes/Camera/CameraModifier.h b/Engine/Source/Runtime/Engine/Classes/Camera/CameraModifier.h
--- a/Engine/Source/Runtime/Engine/Classes/Camera/CameraModifier.h
+++ b/Engine/Source/Runtime/Engine/Classes/Camera/CameraModifier.h
@@ -13,9 +13,7 @@
 UCLASS()
 class ENGINE_API UCameraModifier : public UObject
 {
-	GENERATED_BODY()
-public:
-	UCameraModifier(const FObjectInitializer& ObjectInitializer = FObjectInitializer::Get());
+	GENERATED_UCLASS_BODY()
 
 protected:
 	/** If true, do not apply this modifier to the camera. */
diff --git a/Engine/Source/Runtime/Engine/Classes/Camera/CameraModifier_CameraShake.h b/Engine/Source/Runtime/Engine/Classes/Camera/CameraModifier_CameraShake.h
--- a/Engine/Source/Runtime/Engine/Classes/Camera/CameraModifier_CameraShake.h
+++ b/Engine/Source/Runtime/Engine/Classes/Camera/CameraModifier_CameraShake.h
@@ -16,7 +16,7 @@
 USTRUCT()
 struct ENGINE_API FCameraShakeInstance
 {
-	GENERATED_BODY()
+	GENERATED_USTRUCT_BODY()
 
 	/** The source camera shake data. */
 	UPROPERTY()
@@ -100,9 +100,7 @@ struct ENGINE_API FCameraShakeInstance
 UCLASS(config=Camera)
 class ENGINE_API UCameraModifier_CameraShake : public UCameraModifier
 {
-	GENERATED_BODY()
-public:
-	UCameraModifier_CameraShake(const FObjectInitializer& ObjectInitializer = FObjectInitializer::Get());
+	GENERATED_UCLASS_BODY()
 
 public:
 	/** List of active CameraShakes */
diff --git a/Engine/Source/Runtime/Engine/Classes/Camera/CameraShake.h b/Engine/Source/Runtime/Engine/Classes/Camera/CameraShake.h
--- a/Engine/Source/Runtime/Engine/Classes/Camera/CameraShake.h
+++ b/Engine/Source/Runtime/Engine/Classes/Camera/CameraShake.h
@@ -23,7 +23,7 @@ enum EInitialOscillatorOffset
 USTRUCT()
 struct FFOscillator
 {
-	GENERATED_BODY()
+	GENERATED_USTRUCT_BODY()
 
 	/** Amplitude of the sinusoidal oscillation. */
 	UPROPERTY(EditAnywhere, Category=FOscillator)
@@ -48,7 +48,7 @@ struct FFOscillator
 USTRUCT()
 struct FROscillator
 {
-	GENERATED_BODY()
+	GENERATED_USTRUCT_BODY()
 
 	/** Pitch oscillation. */
 	UPROPERTY(EditAnywhere, Category=ROscillator)
@@ -68,7 +68,7 @@ struct FROscillator
 USTRUCT()
 struct FVOscillator
 {
-	GENERATED_BODY()
+	GENERATED_USTRUCT_BODY()
 
 	/** Oscillation in the X axis. */
 	UPROPERTY(EditAnywhere, Category = VOscillator)
@@ -103,9 +103,7 @@ struct FVOscillator
 UCLASS(Blueprintable, editinlinenew)
 class ENGINE_API UCameraShake : public UObject
 {
-	GENERATED_BODY()
-public:
-	UCameraShake(const FObjectInitializer& ObjectInitializer = FObjectInitializer::Get());
+	GENERATED_UCLASS_BODY()
 
 	/** 
 	 *  If true to only allow a single instance of this shake to play at any given time. 
diff --git a/Engine/Source/Runtime/Engine/Classes/Camera/PlayerCameraManager.h b/Engine/Source/Runtime/Engine/Classes/Camera/PlayerCameraManager.h
--- a/Engine/Source/Runtime/Engine/Classes/Camera/PlayerCameraManager.h
+++ b/Engine/Source/Runtime/Engine/Classes/Camera/PlayerCameraManager.h
@@ -34,7 +34,7 @@ enum EViewTargetBlendFunction
 USTRUCT()
 struct FCameraCacheEntry
 {
-	GENERATED_BODY()
+	GENERATED_USTRUCT_BODY()
 public:
 
 	/** World time this entry was created. */
@@ -54,7 +54,7 @@ public:
 USTRUCT()
 struct ENGINE_API FTViewTarget
 {
-	GENERATED_BODY()
+	GENERATED_USTRUCT_BODY()
 public:
 
 	/** Target Actor used to compute POV */
@@ -93,7 +93,7 @@ public:
 USTRUCT()
 struct FViewTargetTransitionParams
 {
-	GENERATED_BODY()
+	GENERATED_USTRUCT_BODY()
 public:
 
 	/** Total duration of blend to pending view target. 0 means no blending. */
@@ -163,9 +163,7 @@ public:
 UCLASS(notplaceable, transient, BlueprintType, Blueprintable)
 class ENGINE_API APlayerCameraManager : public AActor
 {
-	GENERATED_BODY()
-public:
-	APlayerCameraManager(const FObjectInitializer& ObjectInitializer = FObjectInitializer::Get());
+	GENERATED_UCLASS_BODY()
 
 	/** PlayerController that owns this Camera actor */
 	UPROPERTY()
diff --git a/Engine/Source/Runtime/Engine/Classes/Commandlets/Commandlet.h b/Engine/Source/Runtime/Engine/Classes/Commandlets/Commandlet.h
--- a/Engine/Source/Runtime/Engine/Classes/Commandlets/Commandlet.h
+++ b/Engine/Source/Runtime/Engine/Classes/Commandlets/Commandlet.h
@@ -32,9 +32,7 @@
 UCLASS(abstract, transient,MinimalAPI)
 class UCommandlet : public UObject
 {
-	GENERATED_BODY()
-public:
-	ENGINE_API UCommandlet(const FObjectInitializer& ObjectInitializer = FObjectInitializer::Get());
+	GENERATED_UCLASS_BODY()
 
 	/** Description of the commandlet's purpose */
 	UPROPERTY()
diff --git a/Engine/Source/Runtime/Engine/Classes/Commandlets/PluginCommandlet.h b/Engine/Source/Runtime/Engine/Classes/Commandlets/PluginCommandlet.h
--- a/Engine/Source/Runtime/Engine/Classes/Commandlets/PluginCommandlet.h
+++ b/Engine/Source/Runtime/Engine/Classes/Commandlets/PluginCommandlet.h
@@ -7,9 +7,7 @@
 UCLASS()
 class UPluginCommandlet : public UCommandlet
 {
-	GENERATED_BODY()
-public:
-	UPluginCommandlet(const FObjectInitializer& ObjectInitializer = FObjectInitializer::Get());
+	GENERATED_UCLASS_BODY()
 
 	/** Parsed commandline tokens */
 	TArray<FString> CmdLineTokens;
diff --git a/Engine/Source/Runtime/Engine/Classes/Commandlets/SmokeTestCommandlet.h b/Engine/Source/Runtime/Engine/Classes/Commandlets/SmokeTestCommandlet.h
--- a/Engine/Source/Runtime/Engine/Classes/Commandlets/SmokeTestCommandlet.h
+++ b/Engine/Source/Runtime/Engine/Classes/Commandlets/SmokeTestCommandlet.h
@@ -7,9 +7,7 @@
 UCLASS()
 class USmokeTestCommandlet : public UCommandlet
 {
-	GENERATED_BODY()
-public:
-	USmokeTestCommandlet(const FObjectInitializer& ObjectInitializer = FObjectInitializer::Get());
+	GENERATED_UCLASS_BODY()
 
 
 	// Begin UCommandlet Interface
diff --git a/Engine/Source/Runtime/Engine/Classes/Components/ApplicationLifecycleComponent.h b/Engine/Source/Runtime/Engine/Classes/Components/ApplicationLifecycleComponent.h
--- a/Engine/Source/Runtime/Engine/Classes/Components/ApplicationLifecycleComponent.h
+++ b/Engine/Source/Runtime/Engine/Classes/Components/ApplicationLifecycleComponent.h
@@ -9,9 +9,7 @@
 UCLASS(ClassGroup=Utility, HideCategories=(Activation, "Components|Activation", Collision), meta=(BlueprintSpawnableComponent))
 class ENGINE_API UApplicationLifecycleComponent : public UActorComponent
 {
-	GENERATED_BODY()
-public:
-	UApplicationLifecycleComponent(const FObjectInitializer& ObjectInitializer = FObjectInitializer::Get());
+	GENERATED_UCLASS_BODY()
 
 	DECLARE_DYNAMIC_MULTICAST_DELEGATE(FApplicationLifetimeDelegate);
 
diff --git a/Engine/Source/Runtime/Engine/Classes/Components/ArrowComponent.h b/Engine/Source/Runtime/Engine/Classes/Components/ArrowComponent.h
--- a/Engine/Source/Runtime/Engine/Classes/Components/ArrowComponent.h
+++ b/Engine/Source/Runtime/Engine/Classes/Components/ArrowComponent.h
@@ -10,9 +10,7 @@
 UCLASS(ClassGroup=Utility, hidecategories=(Object,LOD,Physics,Lighting,TextureStreaming,Activation,"Components|Activation",Collision), editinlinenew, meta=(BlueprintSpawnableComponent), MinimalAPI)
 class UArrowComponent : public UPrimitiveComponent
 {
-	GENERATED_BODY()
-public:
-	ENGINE_API UArrowComponent(const FObjectInitializer& ObjectInitializer = FObjectInitializer::Get());
+	GENERATED_UCLASS_BODY()
 
 	UPROPERTY(EditAnywhere, BlueprintReadWrite, Category=ArrowComponent)
 	FColor ArrowColor;
diff --git a/Engine/Source/Runtime/Engine/Classes/Components/AudioComponent.h b/Engine/Source/Runtime/Engine/Classes/Components/AudioComponent.h
--- a/Engine/Source/Runtime/Engine/Classes/Components/AudioComponent.h
+++ b/Engine/Source/Runtime/Engine/Classes/Components/AudioComponent.h
@@ -21,7 +21,7 @@ DECLARE_DYNAMIC_DELEGATE_TwoParams( FOnQueueSubtitles, const TArray<struct FSubt
 USTRUCT()
 struct FAudioComponentParam
 {
-	GENERATED_BODY()
+	GENERATED_USTRUCT_BODY()
 
 	// Name of the parameter
 	UPROPERTY(EditAnywhere, BlueprintReadWrite, Category=AudioComponentParam)
@@ -63,9 +63,7 @@ struct FAudioComponentParam
 UCLASS(ClassGroup=(Audio, Common), hidecategories=(Object, ActorComponent, Physics, Rendering, Mobility, LOD), ShowCategories=Trigger, meta=(BlueprintSpawnableComponent))
 class ENGINE_API UAudioComponent : public USceneComponent
 {
-	GENERATED_BODY()
-public:
-	UAudioComponent(const FObjectInitializer& ObjectInitializer = FObjectInitializer::Get());
+	GENERATED_UCLASS_BODY()
 
 	/** The sound to be played */
 	UPROPERTY(EditAnywhere, BlueprintReadWrite, Category=Sound)
diff --git a/Engine/Source/Runtime/Engine/Classes/Components/BoxComponent.h b/Engine/Source/Runtime/Engine/Classes/Components/BoxComponent.h
--- a/Engine/Source/Runtime/Engine/Classes/Components/BoxComponent.h
+++ b/Engine/Source/Runtime/Engine/Classes/Components/BoxComponent.h
@@ -11,9 +11,7 @@
 UCLASS(ClassGroup="Collision", hidecategories=(Object,LOD,Lighting,TextureStreaming), editinlinenew, meta=(DisplayName="Box Collision", BlueprintSpawnableComponent))
 class ENGINE_API UBoxComponent : public UShapeComponent
 {
-	GENERATED_BODY()
-public:
-	UBoxComponent(const FObjectInitializer& ObjectInitializer = FObjectInitializer::Get());
+	GENERATED_UCLASS_BODY()
 
 protected:
 	/** The extents of the box **/
diff --git a/Engine/Source/Runtime/Engine/Classes/Components/BoxReflectionCaptureComponent.h b/Engine/Source/Runtime/Engine/Classes/Components/BoxReflectionCaptureComponent.h
--- a/Engine/Source/Runtime/Engine/Classes/Components/BoxReflectionCaptureComponent.h
+++ b/Engine/Source/Runtime/Engine/Classes/Components/BoxReflectionCaptureComponent.h
@@ -9,9 +9,7 @@
 UCLASS(hidecategories=(Collision, Object, Physics, SceneComponent), MinimalAPI)
 class UBoxReflectionCaptureComponent : public UReflectionCaptureComponent
 {
-	GENERATED_BODY()
-public:
-	ENGINE_API UBoxReflectionCaptureComponent(const FObjectInitializer& ObjectInitializer = FObjectInitializer::Get());
+	GENERATED_UCLASS_BODY()
 
 	/** Adjust capture transition distance */
 	UPROPERTY(EditAnywhere, BlueprintReadWrite, Category=ReflectionCapture, meta=(UIMin = "1", UIMax = "1000"))
diff --git a/Engine/Source/Runtime/Engine/Classes/Components/CapsuleComponent.h b/Engine/Source/Runtime/Engine/Classes/Components/CapsuleComponent.h
--- a/Engine/Source/Runtime/Engine/Classes/Components/CapsuleComponent.h
+++ b/Engine/Source/Runtime/Engine/Classes/Components/CapsuleComponent.h
@@ -10,9 +10,7 @@
 UCLASS(ClassGroup="Collision", editinlinenew, hidecategories=(Object,LOD,Lighting,TextureStreaming), meta=(DisplayName="Capsule Collision", BlueprintSpawnableComponent))
 class ENGINE_API UCapsuleComponent : public UShapeComponent
 {
-	GENERATED_BODY()
-public:
-	UCapsuleComponent(const FObjectInitializer& ObjectInitializer = FObjectInitializer::Get());
+	GENERATED_UCLASS_BODY()
 
 protected:
 	/** 
diff --git a/Engine/Source/Runtime/Engine/Classes/Components/ChildActorComponent.h b/Engine/Source/Runtime/Engine/Classes/Components/ChildActorComponent.h
--- a/Engine/Source/Runtime/Engine/Classes/Components/ChildActorComponent.h
+++ b/Engine/Source/Runtime/Engine/Classes/Components/ChildActorComponent.h
@@ -11,9 +11,7 @@ class FChildActorComponentInstanceData;
 UCLASS(ClassGroup=Utility, hidecategories=(Object,LOD,Physics,Lighting,TextureStreaming,Activation,"Components|Activation",Collision), meta=(BlueprintSpawnableComponent), MinimalAPI)
 class UChildActorComponent : public USceneComponent
 {
-	GENERATED_BODY()
-public:
-	ENGINE_API UChildActorComponent(const FObjectInitializer& ObjectInitializer = FObjectInitializer::Get());
+	GENERATED_UCLASS_BODY()
 
 	UFUNCTION(BlueprintCallable, Category=ChildActorComponent)
 	ENGINE_API void SetChildActorClass(TSubclassOf<AActor> Class);
diff --git a/Engine/Source/Runtime/Engine/Classes/Components/DecalComponent.h b/Engine/Source/Runtime/Engine/Classes/Components/DecalComponent.h
--- a/Engine/Source/Runtime/Engine/Classes/Components/DecalComponent.h
+++ b/Engine/Source/Runtime/Engine/Classes/Components/DecalComponent.h
@@ -13,9 +13,7 @@ class FDeferredDecalProxy;
 UCLASS(hidecategories=(Collision, Object, Physics, SceneComponent, Activation, "Components|Activation", Mobility), ClassGroup=Rendering, meta=(BlueprintSpawnableComponent))
 class ENGINE_API UDecalComponent : public USceneComponent
 {
-	GENERATED_BODY()
-public:
-	UDecalComponent(const FObjectInitializer& ObjectInitializer = FObjectInitializer::Get());
+	GENERATED_UCLASS_BODY()
 
 	/** Decal material. */
 	//editable(Decal) private{private} const MaterialInterface	DecalMaterial;
diff --git a/Engine/Source/Runtime/Engine/Classes/Components/DestructibleComponent.h b/Engine/Source/Runtime/Engine/Classes/Components/DestructibleComponent.h
--- a/Engine/Source/Runtime/Engine/Classes/Components/DestructibleComponent.h
+++ b/Engine/Source/Runtime/Engine/Classes/Components/DestructibleComponent.h
@@ -44,9 +44,7 @@ DECLARE_DYNAMIC_MULTICAST_DELEGATE_TwoParams(FComponentFractureSignature, const
 UCLASS(ClassGroup=Physics, hidecategories=(Object,Mesh,"Components|SkinnedMesh",Mirroring,Activation,"Components|Activation"), config=Engine, editinlinenew, meta=(BlueprintSpawnableComponent))
 class ENGINE_API UDestructibleComponent : public USkinnedMeshComponent
 {
-	GENERATED_BODY()
-public:
-	UDestructibleComponent(const FObjectInitializer& ObjectInitializer = FObjectInitializer::Get());
+	GENERATED_UCLASS_BODY()
 
 	/** If set, use this actor's fracture effects instead of the asset's fracture effects. */
 	UPROPERTY(EditAnywhere, BlueprintReadWrite, Category=DestructibleComponent)
diff --git a/Engine/Source/Runtime/Engine/Classes/Components/DirectionalLightComponent.h b/Engine/Source/Runtime/Engine/Classes/Components/DirectionalLightComponent.h
--- a/Engine/Source/Runtime/Engine/Classes/Components/DirectionalLightComponent.h
+++ b/Engine/Source/Runtime/Engine/Classes/Components/DirectionalLightComponent.h
@@ -13,9 +13,7 @@
 UCLASS(ClassGroup=Lights, hidecategories=(Object, LightProfiles), editinlinenew, meta=(BlueprintSpawnableComponent), MinimalAPI)
 class UDirectionalLightComponent : public ULightComponent
 {
-	GENERATED_BODY()
-public:
-	ENGINE_API UDirectionalLightComponent(const FObjectInitializer& ObjectInitializer = FObjectInitializer::Get());
+	GENERATED_UCLASS_BODY()
 
 	/** Whether to occlude fog and atmosphere inscattering with screenspace blurred occlusion from this light. */
 	UPROPERTY(EditAnywhere, BlueprintReadOnly, Category=LightShafts, meta=(DisplayName = "Light Shaft Occlusion"))
diff --git a/Engine/Source/Runtime/Engine/Classes/Components/DrawFrustumComponent.h b/Engine/Source/Runtime/Engine/Classes/Components/DrawFrustumComponent.h
--- a/Engine/Source/Runtime/Engine/Classes/Components/DrawFrustumComponent.h
+++ b/Engine/Source/Runtime/Engine/Classes/Components/DrawFrustumComponent.h
@@ -11,9 +11,7 @@
 UCLASS(collapsecategories, hidecategories=Object, editinlinenew, MinimalAPI)
 class UDrawFrustumComponent : public UPrimitiveComponent
 {
-	GENERATED_BODY()
-public:
-	ENGINE_API UDrawFrustumComponent(const FObjectInitializer& ObjectInitializer = FObjectInitializer::Get());
+	GENERATED_UCLASS_BODY()
 	
 	/** Color to draw the wireframe frustum. */
 	UPROPERTY(EditAnywhere, BlueprintReadWrite, Category=DrawFrustumComponent)
diff --git a/Engine/Source/Runtime/Engine/Classes/Components/ExponentialHeightFogComponent.h b/Engine/Source/Runtime/Engine/Classes/Components/ExponentialHeightFogComponent.h
--- a/Engine/Source/Runtime/Engine/Classes/Components/ExponentialHeightFogComponent.h
+++ b/Engine/Source/Runtime/Engine/Classes/Components/ExponentialHeightFogComponent.h
@@ -10,9 +10,7 @@
 UCLASS(ClassGroup=Rendering, collapsecategories, hidecategories=(Object, Mobility), editinlinenew, meta=(BlueprintSpawnableComponent))
 class ENGINE_API UExponentialHeightFogComponent : public USceneComponent
 {
-	GENERATED_BODY()
-public:
-	UExponentialHeightFogComponent(const FObjectInitializer& ObjectInitializer = FObjectInitializer::Get());
+	GENERATED_UCLASS_BODY()
 
 	/** Global density factor. */
 	UPROPERTY(BlueprintReadOnly, interp, Category=ExponentialHeightFogComponent, meta=(UIMin = "0", UIMax = ".05"))
diff --git a/Engine/Source/Runtime/Engine/Classes/Components/HierarchicalInstancedStaticMeshComponent.h b/Engine/Source/Runtime/Engine/Classes/Components/HierarchicalInstancedStaticMeshComponent.h
--- a/Engine/Source/Runtime/Engine/Classes/Components/HierarchicalInstancedStaticMeshComponent.h
+++ b/Engine/Source/Runtime/Engine/Classes/Components/HierarchicalInstancedStaticMeshComponent.h
@@ -12,7 +12,7 @@ class FClusterBuilder;
 USTRUCT()
 struct FClusterNode
 {
-	GENERATED_BODY()
+	GENERATED_USTRUCT_BODY()
 
 	UPROPERTY()
 	FVector BoundMin;
@@ -76,9 +76,7 @@ public:
 UCLASS(ClassGroup=Rendering, meta=(BlueprintSpawnableComponent))
 class ENGINE_API UHierarchicalInstancedStaticMeshComponent : public UInstancedStaticMeshComponent
 {
-	GENERATED_BODY()
-public:
-	UHierarchicalInstancedStaticMeshComponent(const FObjectInitializer& ObjectInitializer = FObjectInitializer::Get());
+	GENERATED_UCLASS_BODY()
 
 	~UHierarchicalInstancedStaticMeshComponent();
 
diff --git a/Engine/Source/Runtime/Engine/Classes/Components/InputComponent.h b/Engine/Source/Runtime/Engine/Classes/Components/InputComponent.h
--- a/Engine/Source/Runtime/Engine/Classes/Components/InputComponent.h
+++ b/Engine/Source/Runtime/Engine/Classes/Components/InputComponent.h
@@ -80,7 +80,7 @@ protected:
 USTRUCT()
 struct FInputChord
 {
-	GENERATED_BODY()
+	GENERATED_USTRUCT_BODY()
 
 	/** The Key is the core of the chord. */
 	UPROPERTY()
@@ -530,9 +530,7 @@ UCLASS(transient, config=Input, hidecategories=(Activation, "Components|Activati
 class ENGINE_API UInputComponent
 	: public UActorComponent
 {
-	GENERATED_BODY()
-public:
-	UInputComponent(const FObjectInitializer& ObjectInitializer = FObjectInitializer::Get());
+	GENERATED_UCLASS_BODY()
 
 	/** The collection of key bindings. */
 	TArray<struct FInputKeyBinding> KeyBindings;
diff --git a/Engine/Source/Runtime/Engine/Classes/Components/InstancedStaticMeshComponent.h b/Engine/Source/Runtime/Engine/Classes/Components/InstancedStaticMeshComponent.h
--- a/Engine/Source/Runtime/Engine/Classes/Components/InstancedStaticMeshComponent.h
+++ b/Engine/Source/Runtime/Engine/Classes/Components/InstancedStaticMeshComponent.h
@@ -16,7 +16,7 @@ struct FPerInstanceRenderData;
 USTRUCT()
 struct FInstancedStaticMeshInstanceData
 {
-	GENERATED_BODY()
+	GENERATED_USTRUCT_BODY()
 
 	UPROPERTY(EditAnywhere, Category=Instances)
 	FMatrix Transform;
@@ -49,7 +49,7 @@ struct FInstancedStaticMeshInstanceData
 USTRUCT()
 struct FInstancedStaticMeshMappingInfo
 {
-	GENERATED_BODY()
+	GENERATED_USTRUCT_BODY()
 
 	FStaticLightingTextureMapping_InstancedStaticMesh* Mapping;
 
@@ -63,9 +63,7 @@ struct FInstancedStaticMeshMappingInfo
 UCLASS(ClassGroup=Rendering, meta=(BlueprintSpawnableComponent))
 class ENGINE_API UInstancedStaticMeshComponent : public UStaticMeshComponent
 {
-	GENERATED_BODY()
-public:
-	UInstancedStaticMeshComponent(const FObjectInitializer& ObjectInitializer = FObjectInitializer::Get());
+	GENERATED_UCLASS_BODY()
 
 	/** Array of instances, bulk serialized */
 	UPROPERTY(EditAnywhere, Transient, DuplicateTransient, DisplayName="Instances", Category=Instances, meta=(MakeEditWidget=true))
diff --git a/Engine/Source/Runtime/Engine/Classes/Components/LightComponent.h b/Engine/Source/Runtime/Engine/Classes/Components/LightComponent.h
--- a/Engine/Source/Runtime/Engine/Classes/Components/LightComponent.h
+++ b/Engine/Source/Runtime/Engine/Classes/Components/LightComponent.h
@@ -42,9 +42,7 @@ public:
 UCLASS(abstract, HideCategories=(Trigger,Activation,"Components|Activation",Physics), ShowCategories=(Mobility))
 class ENGINE_API ULightComponent : public ULightComponentBase
 {
-	GENERATED_BODY()
-public:
-	ULightComponent(const FObjectInitializer& ObjectInitializer = FObjectInitializer::Get());
+	GENERATED_UCLASS_BODY()
 
 	/** 
 	 * Shadow map channel which is used to match up with the appropriate static shadowing during a deferred shading pass.
diff --git a/Engine/Source/Runtime/Engine/Classes/Components/LightComponentBase.h b/Engine/Source/Runtime/Engine/Classes/Components/LightComponentBase.h
--- a/Engine/Source/Runtime/Engine/Classes/Components/LightComponentBase.h
+++ b/Engine/Source/Runtime/Engine/Classes/Components/LightComponentBase.h
@@ -6,9 +6,7 @@
 UCLASS(abstract, HideCategories=(Trigger,Activation,"Components|Activation",Physics), ShowCategories=(Mobility))
 class ENGINE_API ULightComponentBase : public USceneComponent
 {
-	GENERATED_BODY()
-public:
-	ULightComponentBase(const FObjectInitializer& ObjectInitializer = FObjectInitializer::Get());
+	GENERATED_UCLASS_BODY()
 
 	/**
 	 * GUID used to associate a light component with precomputed shadowing information across levels.
diff --git a/Engine/Source/Runtime/Engine/Classes/Components/LineBatchComponent.h b/Engine/Source/Runtime/Engine/Classes/Components/LineBatchComponent.h
--- a/Engine/Source/Runtime/Engine/Classes/Components/LineBatchComponent.h
+++ b/Engine/Source/Runtime/Engine/Classes/Components/LineBatchComponent.h
@@ -7,7 +7,7 @@
 USTRUCT()
 struct FBatchedLine
 {
-	GENERATED_BODY()
+	GENERATED_USTRUCT_BODY()
 
 	UPROPERTY()
 	FVector Start;
@@ -48,7 +48,7 @@ struct FBatchedLine
 USTRUCT()
 struct FBatchedPoint
 {
-	GENERATED_BODY()
+	GENERATED_USTRUCT_BODY()
 
 	UPROPERTY()
 	FVector Position;
@@ -111,9 +111,7 @@ public:
 UCLASS(MinimalAPI)
 class ULineBatchComponent : public UPrimitiveComponent
 {
-	GENERATED_BODY()
-public:
-	ENGINE_API ULineBatchComponent(const FObjectInitializer& ObjectInitializer = FObjectInitializer::Get());
+	GENERATED_UCLASS_BODY()
 
 	/** Buffer of lines to draw */
 	TArray<struct FBatchedLine> BatchedLines;
diff --git a/Engine/Source/Runtime/Engine/Classes/Components/MeshComponent.h b/Engine/Source/Runtime/Engine/Classes/Components/MeshComponent.h
--- a/Engine/Source/Runtime/Engine/Classes/Components/MeshComponent.h
+++ b/Engine/Source/Runtime/Engine/Classes/Components/MeshComponent.h
@@ -17,9 +17,7 @@ struct FMaterialRelevance;
 UCLASS(abstract)
 class ENGINE_API UMeshComponent : public UPrimitiveComponent
 {
-	GENERATED_BODY()
-public:
-	UMeshComponent(const FObjectInitializer& ObjectInitializer = FObjectInitializer::Get());
+	GENERATED_UCLASS_BODY()
 
 	/** Per-Component material overrides.  These must NOT be set directly or a race condition can occur between GC and the rendering thread. */
 	UPROPERTY(EditAnywhere, Category=Rendering, Meta=(ToolTip="Material overrides."))
diff --git a/Engine/Source/Runtime/Engine/Classes/Components/ModelComponent.h b/Engine/Source/Runtime/Engine/Classes/Components/ModelComponent.h
--- a/Engine/Source/Runtime/Engine/Classes/Components/ModelComponent.h
+++ b/Engine/Source/Runtime/Engine/Classes/Components/ModelComponent.h
@@ -18,9 +18,7 @@ class FModelElement;
 UCLASS(MinimalAPI)
 class UModelComponent : public UPrimitiveComponent, public IInterface_CollisionDataProvider
 {
-	GENERATED_BODY()
-public:
-	ENGINE_API UModelComponent(const FObjectInitializer& ObjectInitializer = FObjectInitializer::Get());
+	GENERATED_UCLASS_BODY()
 
 private:
 	/** The BSP tree. */
diff --git a/Engine/Source/Runtime/Engine/Classes/Components/NiagaraComponent.h b/Engine/Source/Runtime/Engine/Classes/Components/NiagaraComponent.h
--- a/Engine/Source/Runtime/Engine/Classes/Components/NiagaraComponent.h
+++ b/Engine/Source/Runtime/Engine/Classes/Components/NiagaraComponent.h
@@ -93,9 +93,7 @@ private:
 UCLASS()
 class ENGINE_API UNiagaraComponent : public UPrimitiveComponent
 {
-	GENERATED_BODY()
-public:
-	UNiagaraComponent(const FObjectInitializer& ObjectInitializer = FObjectInitializer::Get());
+	GENERATED_UCLASS_BODY()
 private:
 	UPROPERTY()
 	//class UNiagaraEffect *Effect;
diff --git a/Engine/Source/Runtime/Engine/Classes/Components/PawnNoiseEmitterComponent.h b/Engine/Source/Runtime/Engine/Classes/Components/PawnNoiseEmitterComponent.h
--- a/Engine/Source/Runtime/Engine/Classes/Components/PawnNoiseEmitterComponent.h
+++ b/Engine/Source/Runtime/Engine/Classes/Components/PawnNoiseEmitterComponent.h
@@ -12,9 +12,7 @@
 UCLASS(ClassGroup=AI, meta=(BlueprintSpawnableComponent))
 class ENGINE_API UPawnNoiseEmitterComponent : public UActorComponent
 {
-	GENERATED_BODY()
-public:
-	UPawnNoiseEmitterComponent(const FObjectInitializer& ObjectInitializer = FObjectInitializer::Get());
+	GENERATED_UCLASS_BODY()
 
 public:
 	// Most recent noise made by this pawn not at its own location
diff --git a/Engine/Source/Runtime/Engine/Classes/Components/PlaneReflectionCaptureComponent.h b/Engine/Source/Runtime/Engine/Classes/Components/PlaneReflectionCaptureComponent.h
--- a/Engine/Source/Runtime/Engine/Classes/Components/PlaneReflectionCaptureComponent.h
+++ b/Engine/Source/Runtime/Engine/Classes/Components/PlaneReflectionCaptureComponent.h
@@ -9,9 +9,7 @@
 UCLASS(hidecategories=(Collision, Object, Physics, SceneComponent), MinimalAPI)
 class UPlaneReflectionCaptureComponent : public UReflectionCaptureComponent
 {
-	GENERATED_BODY()
-public:
-	ENGINE_API UPlaneReflectionCaptureComponent(const FObjectInitializer& ObjectInitializer = FObjectInitializer::Get());
+	GENERATED_UCLASS_BODY()
 
 	/** Radius of the area that can receive reflections from this capture. */
 	UPROPERTY(EditAnywhere, BlueprintReadWrite, Category=ReflectionCapture, meta=(UIMin = "8.0", UIMax = "16384.0"))
diff --git a/Engine/Source/Runtime/Engine/Classes/Components/PointLightComponent.h b/Engine/Source/Runtime/Engine/Classes/Components/PointLightComponent.h
--- a/Engine/Source/Runtime/Engine/Classes/Components/PointLightComponent.h
+++ b/Engine/Source/Runtime/Engine/Classes/Components/PointLightComponent.h
@@ -11,9 +11,7 @@
 UCLASS(ClassGroup=(Lights,Common), hidecategories=(Object, LightShafts), editinlinenew, meta=(BlueprintSpawnableComponent))
 class ENGINE_API UPointLightComponent : public ULightComponent
 {
-	GENERATED_BODY()
-public:
-	UPointLightComponent(const FObjectInitializer& ObjectInitializer = FObjectInitializer::Get());
+	GENERATED_UCLASS_BODY()
 
 	UPROPERTY()
 	float Radius_DEPRECATED;
diff --git a/Engine/Source/Runtime/Engine/Classes/Components/PoseableMeshComponent.h b/Engine/Source/Runtime/Engine/Classes/Components/PoseableMeshComponent.h
--- a/Engine/Source/Runtime/Engine/Classes/Components/PoseableMeshComponent.h
+++ b/Engine/Source/Runtime/Engine/Classes/Components/PoseableMeshComponent.h
@@ -10,9 +10,7 @@
 UCLASS(ClassGroup=Rendering, hidecategories=(Object,Physics), config=Engine, editinlinenew, meta=(BlueprintSpawnableComponent))
 class ENGINE_API UPoseableMeshComponent : public USkinnedMeshComponent
 {
-	GENERATED_BODY()
-public:
-	UPoseableMeshComponent(const FObjectInitializer& ObjectInitializer = FObjectInitializer::Get());
+	GENERATED_UCLASS_BODY()
 
 	/** Temporary array of local-space (ie relative to parent bone) rotation/translation/scale for each bone. */
 	TArray<FTransform> LocalAtoms;
diff --git a/Engine/Source/Runtime/Engine/Classes/Components/PostProcessComponent.h b/Engine/Source/Runtime/Engine/Classes/Components/PostProcessComponent.h
--- a/Engine/Source/Runtime/Engine/Classes/Components/PostProcessComponent.h
+++ b/Engine/Source/Runtime/Engine/Classes/Components/PostProcessComponent.h
@@ -11,9 +11,7 @@
 UCLASS(ClassGroup = Rendering, collapsecategories, hidecategories = (Object), editinlinenew, meta = (BlueprintSpawnableComponent), MinimalAPI)
 class UPostProcessComponent : public USceneComponent, public IInterface_PostProcessVolume
 {
-	GENERATED_BODY()
-public:
-	ENGINE_API UPostProcessComponent(const FObjectInitializer& ObjectInitializer = FObjectInitializer::Get());
+	GENERATED_UCLASS_BODY()
 
 	/** Post process settings to use for this volume. */
 	UPROPERTY(interp, Category = PostProcessVolume, meta = (ShowOnlyInnerProperties))
diff --git a/Engine/Source/Runtime/Engine/Classes/Components/ReflectionCaptureComponent.h b/Engine/Source/Runtime/Engine/Classes/Components/ReflectionCaptureComponent.h
--- a/Engine/Source/Runtime/Engine/Classes/Components/ReflectionCaptureComponent.h
+++ b/Engine/Source/Runtime/Engine/Classes/Components/ReflectionCaptureComponent.h
@@ -71,9 +71,7 @@ private:
 UCLASS(abstract, hidecategories=(Collision, Object, Physics, SceneComponent, Activation, "Components|Activation", Mobility), MinimalAPI)
 class UReflectionCaptureComponent : public USceneComponent
 {
-	GENERATED_BODY()
-public:
-	ENGINE_API UReflectionCaptureComponent(const FObjectInitializer& ObjectInitializer = FObjectInitializer::Get());
+	GENERATED_UCLASS_BODY()
 
 	/** A brightness control to scale the captured scene's reflection intensity. */
 	UPROPERTY(EditAnywhere, BlueprintReadOnly, Category=ReflectionCapture, meta=(UIMin = ".5", UIMax = "4"))
diff --git a/Engine/Source/Runtime/Engine/Classes/Components/SceneCaptureComponent.h b/Engine/Source/Runtime/Engine/Classes/Components/SceneCaptureComponent.h
--- a/Engine/Source/Runtime/Engine/Classes/Components/SceneCaptureComponent.h
+++ b/Engine/Source/Runtime/Engine/Classes/Components/SceneCaptureComponent.h
@@ -8,7 +8,7 @@
 USTRUCT()
 struct FEngineShowFlagsSetting
 {
-	GENERATED_BODY()
+	GENERATED_USTRUCT_BODY()
 
 	UPROPERTY(EditAnywhere, BlueprintReadWrite, Category = SceneCapture)
 	FString ShowFlagName;
@@ -21,9 +21,7 @@ struct FEngineShowFlagsSetting
 UCLASS(hidecategories=(abstract, Collision, Object, Physics, SceneComponent, Mobility), MinimalAPI)
 class USceneCaptureComponent : public USceneComponent
 {
-	GENERATED_BODY()
-public:
-	ENGINE_API USceneCaptureComponent(const FObjectInitializer& ObjectInitializer = FObjectInitializer::Get());
+	GENERATED_UCLASS_BODY()
 
 	/** The components won't rendered by current component.*/
  	UPROPERTY()
diff --git a/Engine/Source/Runtime/Engine/Classes/Components/SceneCaptureComponent2D.h b/Engine/Source/Runtime/Engine/Classes/Components/SceneCaptureComponent2D.h
--- a/Engine/Source/Runtime/Engine/Classes/Components/SceneCaptureComponent2D.h
+++ b/Engine/Source/Runtime/Engine/Classes/Components/SceneCaptureComponent2D.h
@@ -20,9 +20,7 @@ enum ESceneCaptureSource
 UCLASS(hidecategories=(Collision, Object, Physics, SceneComponent), ClassGroup=Rendering, MinimalAPI, editinlinenew, meta=(BlueprintSpawnableComponent))
 class USceneCaptureComponent2D : public USceneCaptureComponent
 {
-	GENERATED_BODY()
-public:
-	ENGINE_API USceneCaptureComponent2D(const FObjectInitializer& ObjectInitializer = FObjectInitializer::Get());
+	GENERATED_UCLASS_BODY()
 
 	/** Camera field of view (in degrees) */
 	UPROPERTY(interp, Category=SceneCapture, meta=(DisplayName = "Field of View", UIMin = "5.0", UIMax = "170", ClampMin = "0.001", ClampMax = "360.0"))
@@ -66,4 +64,4 @@ public:
 	ENGINE_API void UpdateContent();
 
 	ENGINE_API static void UpdateDeferredCaptures( FSceneInterface* Scene );
-};
+};
\ No newline at end of file
diff --git a/Engine/Source/Runtime/Engine/Classes/Components/SceneCaptureComponentCube.h b/Engine/Source/Runtime/Engine/Classes/Components/SceneCaptureComponentCube.h
--- a/Engine/Source/Runtime/Engine/Classes/Components/SceneCaptureComponentCube.h
+++ b/Engine/Source/Runtime/Engine/Classes/Components/SceneCaptureComponentCube.h
@@ -10,9 +10,7 @@
 UCLASS(hidecategories = (Collision, Object, Physics, SceneComponent), ClassGroup=Rendering, MinimalAPI, editinlinenew, meta = (BlueprintSpawnableComponent))
 class USceneCaptureComponentCube : public USceneCaptureComponent
 {
-	GENERATED_BODY()
-public:
-	ENGINE_API USceneCaptureComponentCube(const FObjectInitializer& ObjectInitializer = FObjectInitializer::Get());
+	GENERATED_UCLASS_BODY()
 
 	/** Temporary render target that can be used by the editor. */
 	UPROPERTY(EditAnywhere, BlueprintReadWrite, Category = SceneCapture)
@@ -34,4 +32,4 @@ public:
 	void UpdateContent();
 
 	ENGINE_API static void UpdateDeferredCaptures(FSceneInterface* Scene);
-};
+};
\ No newline at end of file
diff --git a/Engine/Source/Runtime/Engine/Classes/Components/ShapeComponent.h b/Engine/Source/Runtime/Engine/Classes/Components/ShapeComponent.h
--- a/Engine/Source/Runtime/Engine/Classes/Components/ShapeComponent.h
+++ b/Engine/Source/Runtime/Engine/Classes/Components/ShapeComponent.h
@@ -10,9 +10,7 @@
 UCLASS(abstract, hidecategories=(Object,LOD,Lighting,TextureStreaming,Activation,"Components|Activation"), editinlinenew, meta=(BlueprintSpawnableComponent), showcategories=(Mobility))
 class ENGINE_API UShapeComponent : public UPrimitiveComponent
 {
-	GENERATED_BODY()
-public:
-	UShapeComponent(const FObjectInitializer& ObjectInitializer = FObjectInitializer::Get());
+	GENERATED_UCLASS_BODY()
 
 	/** Color used to draw the shape. */
 	UPROPERTY(EditAnywhere, AdvancedDisplay, BlueprintReadOnly, Category=Shape)
diff --git a/Engine/Source/Runtime/Engine/Classes/Components/SkinnedMeshComponent.h b/Engine/Source/Runtime/Engine/Classes/Components/SkinnedMeshComponent.h
--- a/Engine/Source/Runtime/Engine/Classes/Components/SkinnedMeshComponent.h
+++ b/Engine/Source/Runtime/Engine/Classes/Components/SkinnedMeshComponent.h
@@ -76,7 +76,7 @@ namespace EBoneSpaces
 USTRUCT()
 struct FActiveVertexAnim
 {
-	GENERATED_BODY()
+	GENERATED_USTRUCT_BODY()
 
 	/** The anim that we want to apply. */
 	UPROPERTY()
@@ -122,7 +122,7 @@ struct FActiveVertexAnim
 USTRUCT()
 struct FSkelMeshComponentLODInfo
 {
-	GENERATED_BODY()
+	GENERATED_USTRUCT_BODY()
 
 	/** Material corresponds to section. To show/hide each section, use this **/
 	UPROPERTY()
@@ -147,9 +147,7 @@ struct FSkelMeshComponentLODInfo
 UCLASS(hidecategories=Object, config=Engine, editinlinenew, abstract)
 class ENGINE_API USkinnedMeshComponent : public UMeshComponent
 {
-	GENERATED_BODY()
-public:
-	USkinnedMeshComponent(const FObjectInitializer& ObjectInitializer = FObjectInitializer::Get());
+	GENERATED_UCLASS_BODY()
 
 	/** The skeletal mesh used by this component. */
 	UPROPERTY(EditAnywhere, BlueprintReadOnly, Category="Mesh")
diff --git a/Engine/Source/Runtime/Engine/Classes/Components/SphereReflectionCaptureComponent.h b/Engine/Source/Runtime/Engine/Classes/Components/SphereReflectionCaptureComponent.h
--- a/Engine/Source/Runtime/Engine/Classes/Components/SphereReflectionCaptureComponent.h
+++ b/Engine/Source/Runtime/Engine/Classes/Components/SphereReflectionCaptureComponent.h
@@ -9,9 +9,7 @@
 UCLASS(hidecategories=(Collision, Object, Physics, SceneComponent), MinimalAPI)
 class USphereReflectionCaptureComponent : public UReflectionCaptureComponent
 {
-	GENERATED_BODY()
-public:
-	ENGINE_API USphereReflectionCaptureComponent(const FObjectInitializer& ObjectInitializer = FObjectInitializer::Get());
+	GENERATED_UCLASS_BODY()
 
 	/** Radius of the area that can receive reflections from this capture. */
 	UPROPERTY(EditAnywhere, BlueprintReadWrite, Category=ReflectionCapture, meta=(UIMin = "8.0", UIMax = "16384.0"))
diff --git a/Engine/Source/Runtime/Engine/Classes/Components/SplineComponent.h b/Engine/Source/Runtime/Engine/Classes/Components/SplineComponent.h
--- a/Engine/Source/Runtime/Engine/Classes/Components/SplineComponent.h
+++ b/Engine/Source/Runtime/Engine/Classes/Components/SplineComponent.h
@@ -27,9 +27,7 @@ namespace ESplinePointType
 UCLASS(ClassGroup=Utility, meta=(BlueprintSpawnableComponent))
 class ENGINE_API USplineComponent : public UPrimitiveComponent
 {
-	GENERATED_BODY()
-public:
-	USplineComponent(const FObjectInitializer& ObjectInitializer = FObjectInitializer::Get());
+	GENERATED_UCLASS_BODY()
 
 	/** Actual data for spline. Locations and tangents are in world space. */
 	UPROPERTY()
diff --git a/Engine/Source/Runtime/Engine/Classes/Components/SplineMeshComponent.h b/Engine/Source/Runtime/Engine/Classes/Components/SplineMeshComponent.h
--- a/Engine/Source/Runtime/Engine/Classes/Components/SplineMeshComponent.h
+++ b/Engine/Source/Runtime/Engine/Classes/Components/SplineMeshComponent.h
@@ -25,7 +25,7 @@ namespace ESplineMeshAxis
 USTRUCT()
 struct FSplineMeshParams
 {
-	GENERATED_BODY()
+	GENERATED_USTRUCT_BODY()
 
 	/** Start location of spline, in component space */
 	UPROPERTY(EditAnywhere, Category=SplineMesh)
@@ -91,9 +91,7 @@ struct FSplineMeshParams
 UCLASS(ClassGroup=Rendering, hidecategories=(Physics), meta=(BlueprintSpawnableComponent))
 class ENGINE_API USplineMeshComponent : public UStaticMeshComponent, public IInterface_CollisionDataProvider
 {
-	GENERATED_BODY()
-public:
-	USplineMeshComponent(const FObjectInitializer& ObjectInitializer = FObjectInitializer::Get());
+	GENERATED_UCLASS_BODY()
 
 	/** Spline that is used to deform mesh */
 	UPROPERTY(EditAnywhere, Category=SplineMesh, meta=(ShowOnlyInnerProperties))
diff --git a/Engine/Source/Runtime/Engine/Classes/Components/SpotLightComponent.h b/Engine/Source/Runtime/Engine/Classes/Components/SpotLightComponent.h
--- a/Engine/Source/Runtime/Engine/Classes/Components/SpotLightComponent.h
+++ b/Engine/Source/Runtime/Engine/Classes/Components/SpotLightComponent.h
@@ -10,9 +10,7 @@
 UCLASS(ClassGroup=Lights, hidecategories=Object, editinlinenew, meta=(BlueprintSpawnableComponent))
 class ENGINE_API USpotLightComponent : public UPointLightComponent
 {
-	GENERATED_BODY()
-public:
-	USpotLightComponent(const FObjectInitializer& ObjectInitializer = FObjectInitializer::Get());
+	GENERATED_UCLASS_BODY()
 
 	/** Degrees. */
 	UPROPERTY(EditAnywhere, BlueprintReadOnly, Category=Light, meta=(UIMin = "1.0", UIMax = "80.0"))
diff --git a/Engine/Source/Runtime/Engine/Classes/Components/StaticMeshComponent.h b/Engine/Source/Runtime/Engine/Classes/Components/StaticMeshComponent.h
--- a/Engine/Source/Runtime/Engine/Classes/Components/StaticMeshComponent.h
+++ b/Engine/Source/Runtime/Engine/Classes/Components/StaticMeshComponent.h
@@ -19,7 +19,7 @@ struct FConvexVolume;
 USTRUCT()
 struct FPaintedVertex
 {
-	GENERATED_BODY()
+	GENERATED_USTRUCT_BODY()
 
 	UPROPERTY()
 	FVector Position;
@@ -51,7 +51,7 @@ struct FPaintedVertex
 USTRUCT()
 struct FStaticMeshComponentLODInfo
 {
-	GENERATED_BODY()
+	GENERATED_USTRUCT_BODY()
 
 	FLightMapRef LightMap;
 
@@ -120,9 +120,7 @@ struct TStructOpsTypeTraits<FStaticMeshComponentLODInfo> : public TStructOpsType
 UCLASS(ClassGroup=(Rendering, Common), hidecategories=(Object,Activation,"Components|Activation"), ShowCategories=(Mobility), editinlinenew, meta=(BlueprintSpawnableComponent))
 class ENGINE_API UStaticMeshComponent : public UMeshComponent
 {
-	GENERATED_BODY()
-public:
-	UStaticMeshComponent(const FObjectInitializer& ObjectInitializer = FObjectInitializer::Get());
+	GENERATED_UCLASS_BODY()
 
 	/** If 0, auto-select LOD level. if >0, force to (ForcedLodModel-1). */
 	UPROPERTY(EditAnywhere, AdvancedDisplay, BlueprintReadOnly, Category=LOD)
diff --git a/Engine/Source/Runtime/Engine/Classes/Components/TextRenderComponent.h b/Engine/Source/Runtime/Engine/Classes/Components/TextRenderComponent.h
--- a/Engine/Source/Runtime/Engine/Classes/Components/TextRenderComponent.h
+++ b/Engine/Source/Runtime/Engine/Classes/Components/TextRenderComponent.h
@@ -35,9 +35,7 @@ enum EVerticalTextAligment
 UCLASS(ClassGroup=Rendering, hidecategories=(Object,LOD,Physics,TextureStreaming,Activation,"Components|Activation",Collision), editinlinenew, meta=(BlueprintSpawnableComponent = ""))
 class ENGINE_API UTextRenderComponent : public UPrimitiveComponent
 {
-	GENERATED_BODY()
-public:
-	UTextRenderComponent(const FObjectInitializer& ObjectInitializer = FObjectInitializer::Get());
+	GENERATED_UCLASS_BODY()
 
 	/** Text content, can be multi line using <br> as line separator */
 	UPROPERTY(EditAnywhere, BlueprintReadOnly, Category=Text, meta=(MultiLine=true))
diff --git a/Engine/Source/Runtime/Engine/Classes/Components/TimelineComponent.h b/Engine/Source/Runtime/Engine/Classes/Components/TimelineComponent.h
--- a/Engine/Source/Runtime/Engine/Classes/Components/TimelineComponent.h
+++ b/Engine/Source/Runtime/Engine/Classes/Components/TimelineComponent.h
@@ -45,7 +45,7 @@ namespace ETimelineDirection
 USTRUCT()
 struct FTimelineEventEntry
 {
-	GENERATED_BODY()
+	GENERATED_USTRUCT_BODY()
 
 	/** Time at which event should fire */
 	UPROPERTY()
@@ -67,7 +67,7 @@ struct FTimelineEventEntry
 USTRUCT()
 struct FTimelineVectorTrack
 {
-	GENERATED_BODY()
+	GENERATED_USTRUCT_BODY()
 
 	/** Vector curve to be evaluated */
 	UPROPERTY()
@@ -101,7 +101,7 @@ struct FTimelineVectorTrack
 USTRUCT()
 struct FTimelineFloatTrack
 {
-	GENERATED_BODY()
+	GENERATED_USTRUCT_BODY()
 
 	/** Float curve to be evaluated */
 	UPROPERTY()
@@ -136,7 +136,7 @@ struct FTimelineFloatTrack
 USTRUCT()
 struct FTimelineLinearColorTrack
 {
-	GENERATED_BODY()
+	GENERATED_USTRUCT_BODY()
 
 	/** Float curve to be evaluated */
 	UPROPERTY()
@@ -169,7 +169,7 @@ struct FTimelineLinearColorTrack
 USTRUCT()
 struct FTimeline
 {
-	GENERATED_BODY()
+	GENERATED_USTRUCT_BODY()
 		
 private:
 	/** Specified how the timeline determines its own length (e.g. specified length, last keyframe) */
@@ -361,9 +361,7 @@ private:
 UCLASS(MinimalAPI)
 class UTimelineComponent : public UActorComponent
 {
-	GENERATED_BODY()
-public:
-	ENGINE_API UTimelineComponent(const FObjectInitializer& ObjectInitializer = FObjectInitializer::Get());
+	GENERATED_UCLASS_BODY()
 
 private:
 	/** The actual timeline structure */
