Back out changelist 2481333

diff --git a/Engine/Source/Runtime/Engine/Classes/Components/VectorFieldComponent.h b/Engine/Source/Runtime/Engine/Classes/Components/VectorFieldComponent.h
--- a/Engine/Source/Runtime/Engine/Classes/Components/VectorFieldComponent.h
+++ b/Engine/Source/Runtime/Engine/Classes/Components/VectorFieldComponent.h
@@ -12,9 +12,7 @@
 UCLASS(ClassGroup=Rendering, hidecategories=Object, meta=(BlueprintSpawnableComponent), MinimalAPI)
 class UVectorFieldComponent : public UPrimitiveComponent
 {
-	GENERATED_BODY()
-public:
-	ENGINE_API UVectorFieldComponent(const FObjectInitializer& ObjectInitializer = FObjectInitializer::Get());
+	GENERATED_UCLASS_BODY()
 
 	/** The vector field asset. */
 	UPROPERTY(EditAnywhere, BlueprintReadWrite, Category=VectorFieldComponent)
diff --git a/Engine/Source/Runtime/Engine/Classes/Curves/CurveBase.h b/Engine/Source/Runtime/Engine/Classes/Curves/CurveBase.h
--- a/Engine/Source/Runtime/Engine/Classes/Curves/CurveBase.h
+++ b/Engine/Source/Runtime/Engine/Classes/Curves/CurveBase.h
@@ -42,7 +42,7 @@ private:
 USTRUCT()
 struct FKeyHandleMap
 {
-	GENERATED_BODY()
+	GENERATED_USTRUCT_BODY()
 public:
 	FKeyHandleMap() {}
 
@@ -92,7 +92,7 @@ struct TStructOpsTypeTraits< FKeyHandleMap > : public TStructOpsTypeTraitsBase
 USTRUCT()
 struct ENGINE_API FIndexedCurve
 {
-	GENERATED_BODY()
+	GENERATED_USTRUCT_BODY()
 public:
 	FIndexedCurve() {}
 
@@ -160,7 +160,7 @@ enum ERichCurveTangentWeightMode
 USTRUCT()
 struct ENGINE_API FRichCurveKey
 {
-	GENERATED_BODY()
+	GENERATED_USTRUCT_BODY()
 
 	/** Interpolation mode between this key and the next */
 	UPROPERTY()
@@ -268,7 +268,7 @@ struct TStructOpsTypeTraits< FRichCurveKey > : public TStructOpsTypeTraitsBase
 USTRUCT()
 struct ENGINE_API FRichCurve : public FIndexedCurve
 {
-	GENERATED_BODY()
+	GENERATED_USTRUCT_BODY()
 
 public:
 	virtual ~FRichCurve()
@@ -462,9 +462,7 @@ public:
 UCLASS(abstract)
 class ENGINE_API UCurveBase : public UObject, public FCurveOwnerInterface
 {
-	GENERATED_BODY()
-public:
-	UCurveBase(const FObjectInitializer& ObjectInitializer = FObjectInitializer::Get());
+	GENERATED_UCLASS_BODY()
 
 	/** The filename imported to create this object. Relative to this object's package, BaseDir() or absolute */
 	UPROPERTY()
@@ -526,7 +524,7 @@ public:
 USTRUCT()
 struct FIntegralKey
 {
-	GENERATED_BODY()
+	GENERATED_USTRUCT_BODY()
 public:
 	FIntegralKey(float InTime = 0.f, int32 InValue = 0)
 		: Time(InTime)
@@ -545,7 +543,7 @@ public:
 USTRUCT()
 struct ENGINE_API FIntegralCurve : public FIndexedCurve
 {
-	GENERATED_BODY()
+	GENERATED_USTRUCT_BODY()
 public:
 	
 	/** Get number of keys in curve. */
diff --git a/Engine/Source/Runtime/Engine/Classes/Curves/CurveEdPresetCurve.h b/Engine/Source/Runtime/Engine/Classes/Curves/CurveEdPresetCurve.h
--- a/Engine/Source/Runtime/Engine/Classes/Curves/CurveEdPresetCurve.h
+++ b/Engine/Source/Runtime/Engine/Classes/Curves/CurveEdPresetCurve.h
@@ -11,8 +11,6 @@
 UCLASS(deprecated)
 class UDEPRECATED_CurveEdPresetCurve : public UObject
 {
-	GENERATED_BODY()
-public:
-	UDEPRECATED_CurveEdPresetCurve(const FObjectInitializer& ObjectInitializer = FObjectInitializer::Get());
+	GENERATED_UCLASS_BODY()
 };
 
diff --git a/Engine/Source/Runtime/Engine/Classes/Curves/CurveFloat.h b/Engine/Source/Runtime/Engine/Classes/Curves/CurveFloat.h
--- a/Engine/Source/Runtime/Engine/Classes/Curves/CurveFloat.h
+++ b/Engine/Source/Runtime/Engine/Classes/Curves/CurveFloat.h
@@ -8,7 +8,7 @@
 USTRUCT()
 struct ENGINE_API FRuntimeFloatCurve
 {
-	GENERATED_BODY()
+	GENERATED_USTRUCT_BODY()
 
 	UPROPERTY()
 	FRichCurve EditorCurveData;
@@ -26,9 +26,7 @@ struct ENGINE_API FRuntimeFloatCurve
 UCLASS(BlueprintType)
 class ENGINE_API UCurveFloat : public UCurveBase
 {
-	GENERATED_BODY()
-public:
-	UCurveFloat(const FObjectInitializer& ObjectInitializer = FObjectInitializer::Get());
+	GENERATED_UCLASS_BODY()
 
 	/** Keyframe data */
 	UPROPERTY()
diff --git a/Engine/Source/Runtime/Engine/Classes/Curves/CurveLinearColor.h b/Engine/Source/Runtime/Engine/Classes/Curves/CurveLinearColor.h
--- a/Engine/Source/Runtime/Engine/Classes/Curves/CurveLinearColor.h
+++ b/Engine/Source/Runtime/Engine/Classes/Curves/CurveLinearColor.h
@@ -7,7 +7,7 @@
 USTRUCT()
 struct ENGINE_API FRuntimeCurveLinearColor
 {
-	GENERATED_BODY()
+	GENERATED_USTRUCT_BODY()
 public:
 	UPROPERTY()
 	FRichCurve ColorCurves[4];
@@ -21,9 +21,7 @@ public:
 UCLASS(BlueprintType)
 class ENGINE_API UCurveLinearColor : public UCurveBase
 {
-	GENERATED_BODY()
-public:
-	UCurveLinearColor(const FObjectInitializer& ObjectInitializer = FObjectInitializer::Get());
+	GENERATED_UCLASS_BODY()
 
 	/** Keyframe data, one curve for red, green, blue, and alpha */
 	UPROPERTY()
diff --git a/Engine/Source/Runtime/Engine/Classes/Curves/CurveVector.h b/Engine/Source/Runtime/Engine/Classes/Curves/CurveVector.h
--- a/Engine/Source/Runtime/Engine/Classes/Curves/CurveVector.h
+++ b/Engine/Source/Runtime/Engine/Classes/Curves/CurveVector.h
@@ -7,9 +7,7 @@
 UCLASS(BlueprintType, MinimalAPI)
 class UCurveVector : public UCurveBase
 {
-	GENERATED_BODY()
-public:
-	ENGINE_API UCurveVector(const FObjectInitializer& ObjectInitializer = FObjectInitializer::Get());
+	GENERATED_UCLASS_BODY()
 
 	/** Keyframe data, one curve for X, Y and Z */
 	UPROPERTY()
diff --git a/Engine/Source/Runtime/Engine/Classes/Debug/DebugDrawService.h b/Engine/Source/Runtime/Engine/Classes/Debug/DebugDrawService.h
--- a/Engine/Source/Runtime/Engine/Classes/Debug/DebugDrawService.h
+++ b/Engine/Source/Runtime/Engine/Classes/Debug/DebugDrawService.h
@@ -15,9 +15,7 @@ DECLARE_DELEGATE_TwoParams(FDebugDrawDelegate, class UCanvas*, class APlayerCont
 UCLASS(config=Engine)
 class ENGINE_API UDebugDrawService : public UBlueprintFunctionLibrary
 {
-	GENERATED_BODY()
-public:
-	UDebugDrawService(const FObjectInitializer& ObjectInitializer = FObjectInitializer::Get());
+	GENERATED_UCLASS_BODY()
 
 	//void Register
 
diff --git a/Engine/Source/Runtime/Engine/Classes/Debug/LogVisualizerCameraController.h b/Engine/Source/Runtime/Engine/Classes/Debug/LogVisualizerCameraController.h
--- a/Engine/Source/Runtime/Engine/Classes/Debug/LogVisualizerCameraController.h
+++ b/Engine/Source/Runtime/Engine/Classes/Debug/LogVisualizerCameraController.h
@@ -8,13 +8,11 @@
 UCLASS(config=Input, hidedropdown)
 class ALogVisualizerCameraController : public ADebugCameraController
 {
-	GENERATED_BODY()
-public:
-	ALogVisualizerCameraController(const FObjectInitializer& ObjectInitializer = FObjectInitializer::Get());
 public:
 	DECLARE_DELEGATE_OneParam(FActorSelectedDelegate, class AActor*);
 	DECLARE_DELEGATE_OneParam(FLogEntryIterationDelegate, int32);
 
+	GENERATED_UCLASS_BODY()
 
 	UPROPERTY()
 	class AActor* PickedActor;
diff --git a/Engine/Source/Runtime/Engine/Classes/Debug/ReporterBase.h b/Engine/Source/Runtime/Engine/Classes/Debug/ReporterBase.h
--- a/Engine/Source/Runtime/Engine/Classes/Debug/ReporterBase.h
+++ b/Engine/Source/Runtime/Engine/Classes/Debug/ReporterBase.h
@@ -17,9 +17,7 @@ namespace EReporterLineStyle
 UCLASS(Abstract)
 class UReporterBase : public UObject
 {
-	GENERATED_BODY()
-public:
-	UReporterBase(const FObjectInitializer& ObjectInitializer = FObjectInitializer::Get());
+	GENERATED_UCLASS_BODY()
 
 public:
 	
@@ -55,4 +53,4 @@ protected:
 	 * @param Color The color to draw the triangle
 	 */
 	void DrawTriangle(UCanvas* Canvas, const FVector2D& Vertex1, const FVector2D& Vertex2, const FVector2D& Vertex3, const FLinearColor& Color);
-};
+};
\ No newline at end of file
diff --git a/Engine/Source/Runtime/Engine/Classes/Debug/ReporterGraph.h b/Engine/Source/Runtime/Engine/Classes/Debug/ReporterGraph.h
--- a/Engine/Source/Runtime/Engine/Classes/Debug/ReporterGraph.h
+++ b/Engine/Source/Runtime/Engine/Classes/Debug/ReporterGraph.h
@@ -96,9 +96,7 @@ struct FGraphLine
 UCLASS()
 class ENGINE_API UReporterGraph : public UReporterBase
 {
-	GENERATED_BODY()
-public:
-	UReporterGraph(const FObjectInitializer& ObjectInitializer = FObjectInitializer::Get());
+	GENERATED_UCLASS_BODY()
 
 public:
 
diff --git a/Engine/Source/Runtime/Engine/Classes/DeviceProfiles/DeviceProfile.h b/Engine/Source/Runtime/Engine/Classes/DeviceProfiles/DeviceProfile.h
--- a/Engine/Source/Runtime/Engine/Classes/DeviceProfiles/DeviceProfile.h
+++ b/Engine/Source/Runtime/Engine/Classes/DeviceProfiles/DeviceProfile.h
@@ -15,9 +15,7 @@ DECLARE_DELEGATE(FOnCVarsUpdated);
 UCLASS(config=DeviceProfiles, perobjectconfig, Blueprintable)
 class ENGINE_API UDeviceProfile : public UObject
 {
-	GENERATED_BODY()
-public:
-	UDeviceProfile(const FObjectInitializer& ObjectInitializer = FObjectInitializer::Get());
+	GENERATED_UCLASS_BODY()
 
 	/** The type of this profile, I.e. IOS, Windows, PS4 etc */
 	UPROPERTY(VisibleAnywhere, config, Category=DeviceSettings)
diff --git a/Engine/Source/Runtime/Engine/Classes/Distributions/DistributionFloatConstant.h b/Engine/Source/Runtime/Engine/Classes/Distributions/DistributionFloatConstant.h
--- a/Engine/Source/Runtime/Engine/Classes/Distributions/DistributionFloatConstant.h
+++ b/Engine/Source/Runtime/Engine/Classes/Distributions/DistributionFloatConstant.h
@@ -8,9 +8,7 @@
 UCLASS(collapsecategories, hidecategories=Object, editinlinenew,MinimalAPI)
 class UDistributionFloatConstant : public UDistributionFloat
 {
-	GENERATED_BODY()
-public:
-	ENGINE_API UDistributionFloatConstant(const FObjectInitializer& ObjectInitializer = FObjectInitializer::Get());
+	GENERATED_UCLASS_BODY()
 
 	/** This float will be returned for all values of time. */
 	UPROPERTY(EditAnywhere, Category=DistributionFloatConstant)
diff --git a/Engine/Source/Runtime/Engine/Classes/Distributions/DistributionFloatConstantCurve.h b/Engine/Source/Runtime/Engine/Classes/Distributions/DistributionFloatConstantCurve.h
--- a/Engine/Source/Runtime/Engine/Classes/Distributions/DistributionFloatConstantCurve.h
+++ b/Engine/Source/Runtime/Engine/Classes/Distributions/DistributionFloatConstantCurve.h
@@ -8,9 +8,7 @@
 UCLASS(collapsecategories, hidecategories=Object, editinlinenew,MinimalAPI)
 class UDistributionFloatConstantCurve : public UDistributionFloat
 {
-	GENERATED_BODY()
-public:
-	ENGINE_API UDistributionFloatConstantCurve(const FObjectInitializer& ObjectInitializer = FObjectInitializer::Get());
+	GENERATED_UCLASS_BODY()
 
 	/** Keyframe data for how output constant varies over time. */
 	UPROPERTY(EditAnywhere, Category=DistributionFloatConstantCurve)
diff --git a/Engine/Source/Runtime/Engine/Classes/Distributions/DistributionFloatParameterBase.h b/Engine/Source/Runtime/Engine/Classes/Distributions/DistributionFloatParameterBase.h
--- a/Engine/Source/Runtime/Engine/Classes/Distributions/DistributionFloatParameterBase.h
+++ b/Engine/Source/Runtime/Engine/Classes/Distributions/DistributionFloatParameterBase.h
@@ -8,9 +8,7 @@
 UCLASS(abstract, collapsecategories, hidecategories=Object, editinlinenew)
 class UDistributionFloatParameterBase : public UDistributionFloatConstant
 {
-	GENERATED_BODY()
-public:
-	UDistributionFloatParameterBase(const FObjectInitializer& ObjectInitializer = FObjectInitializer::Get());
+	GENERATED_UCLASS_BODY()
 
 	/** todo document */
 	UPROPERTY(EditAnywhere, Category=DistributionFloatParameterBase)
diff --git a/Engine/Source/Runtime/Engine/Classes/Distributions/DistributionFloatParticleParameter.h b/Engine/Source/Runtime/Engine/Classes/Distributions/DistributionFloatParticleParameter.h
--- a/Engine/Source/Runtime/Engine/Classes/Distributions/DistributionFloatParticleParameter.h
+++ b/Engine/Source/Runtime/Engine/Classes/Distributions/DistributionFloatParticleParameter.h
@@ -8,9 +8,7 @@
 UCLASS(collapsecategories, hidecategories=Object, editinlinenew, MinimalAPI)
 class UDistributionFloatParticleParameter : public UDistributionFloatParameterBase
 {
-	GENERATED_BODY()
-public:
-	ENGINE_API UDistributionFloatParticleParameter(const FObjectInitializer& ObjectInitializer = FObjectInitializer::Get());
+	GENERATED_UCLASS_BODY()
 
 
 	// Begin UDistributionFloatParameterBase Interface
diff --git a/Engine/Source/Runtime/Engine/Classes/Distributions/DistributionFloatUniform.h b/Engine/Source/Runtime/Engine/Classes/Distributions/DistributionFloatUniform.h
--- a/Engine/Source/Runtime/Engine/Classes/Distributions/DistributionFloatUniform.h
+++ b/Engine/Source/Runtime/Engine/Classes/Distributions/DistributionFloatUniform.h
@@ -9,9 +9,7 @@
 UCLASS(collapsecategories, hidecategories=Object, editinlinenew, MinimalAPI)
 class UDistributionFloatUniform : public UDistributionFloat
 {
-	GENERATED_BODY()
-public:
-	ENGINE_API UDistributionFloatUniform(const FObjectInitializer& ObjectInitializer = FObjectInitializer::Get());
+	GENERATED_UCLASS_BODY()
 
 	/** Low end of output float distribution. */
 	UPROPERTY(EditAnywhere, Category=DistributionFloatUniform)
diff --git a/Engine/Source/Runtime/Engine/Classes/Distributions/DistributionFloatUniformCurve.h b/Engine/Source/Runtime/Engine/Classes/Distributions/DistributionFloatUniformCurve.h
--- a/Engine/Source/Runtime/Engine/Classes/Distributions/DistributionFloatUniformCurve.h
+++ b/Engine/Source/Runtime/Engine/Classes/Distributions/DistributionFloatUniformCurve.h
@@ -9,9 +9,7 @@
 UCLASS(collapsecategories, hidecategories=Object, editinlinenew, MinimalAPI)
 class UDistributionFloatUniformCurve : public UDistributionFloat
 {
-	GENERATED_BODY()
-public:
-	ENGINE_API UDistributionFloatUniformCurve(const FObjectInitializer& ObjectInitializer = FObjectInitializer::Get());
+	GENERATED_UCLASS_BODY()
 
 	/** Keyframe data for how output constant varies over time. */
 	UPROPERTY(EditAnywhere, Category=DistributionFloatUniformCurve)
diff --git a/Engine/Source/Runtime/Engine/Classes/Distributions/DistributionVector.h b/Engine/Source/Runtime/Engine/Classes/Distributions/DistributionVector.h
--- a/Engine/Source/Runtime/Engine/Classes/Distributions/DistributionVector.h
+++ b/Engine/Source/Runtime/Engine/Classes/Distributions/DistributionVector.h
@@ -50,7 +50,7 @@ struct FVector4Distribution
 USTRUCT()
 struct FRawDistributionVector : public FRawDistribution
 {
-	GENERATED_BODY()
+	GENERATED_USTRUCT_BODY()
 
 private:
 	UPROPERTY()
@@ -101,11 +101,10 @@ public:
 	
 };
 
-UCLASS(abstract)
+UCLASS(abstract, customconstructor)
 class ENGINE_API UDistributionVector : public UDistribution
 {
-	GENERATED_BODY()
-public:
+	GENERATED_UCLASS_BODY()
 
 	/** Can this variable be baked out to a FRawDistribution? Should be true 99% of the time*/
 	UPROPERTY(EditAnywhere, Category=Baked)
diff --git a/Engine/Source/Runtime/Engine/Classes/Distributions/DistributionVectorConstant.h b/Engine/Source/Runtime/Engine/Classes/Distributions/DistributionVectorConstant.h
--- a/Engine/Source/Runtime/Engine/Classes/Distributions/DistributionVectorConstant.h
+++ b/Engine/Source/Runtime/Engine/Classes/Distributions/DistributionVectorConstant.h
@@ -8,9 +8,7 @@
 UCLASS(collapsecategories, hidecategories=Object, editinlinenew)
 class ENGINE_API UDistributionVectorConstant : public UDistributionVector
 {
-	GENERATED_BODY()
-public:
-	UDistributionVectorConstant(const FObjectInitializer& ObjectInitializer = FObjectInitializer::Get());
+	GENERATED_UCLASS_BODY()
 
 	/** This FVector will be returned for all input times. */
 	UPROPERTY(EditAnywhere, Category=DistributionVectorConstant)
diff --git a/Engine/Source/Runtime/Engine/Classes/Distributions/DistributionVectorConstantCurve.h b/Engine/Source/Runtime/Engine/Classes/Distributions/DistributionVectorConstantCurve.h
--- a/Engine/Source/Runtime/Engine/Classes/Distributions/DistributionVectorConstantCurve.h
+++ b/Engine/Source/Runtime/Engine/Classes/Distributions/DistributionVectorConstantCurve.h
@@ -8,9 +8,7 @@
 UCLASS(collapsecategories, hidecategories=Object, editinlinenew,MinimalAPI)
 class UDistributionVectorConstantCurve : public UDistributionVector
 {
-	GENERATED_BODY()
-public:
-	ENGINE_API UDistributionVectorConstantCurve(const FObjectInitializer& ObjectInitializer = FObjectInitializer::Get());
+	GENERATED_UCLASS_BODY()
 
 	/** Keyframe data for each component (X,Y,Z) over time. */
 	UPROPERTY(EditAnywhere, Category=DistributionVectorConstantCurve)
diff --git a/Engine/Source/Runtime/Engine/Classes/Distributions/DistributionVectorParticleParameter.h b/Engine/Source/Runtime/Engine/Classes/Distributions/DistributionVectorParticleParameter.h
--- a/Engine/Source/Runtime/Engine/Classes/Distributions/DistributionVectorParticleParameter.h
+++ b/Engine/Source/Runtime/Engine/Classes/Distributions/DistributionVectorParticleParameter.h
@@ -8,9 +8,7 @@
 UCLASS(collapsecategories, hidecategories=Object, editinlinenew, MinimalAPI)
 class UDistributionVectorParticleParameter : public UDistributionVectorParameterBase
 {
-	GENERATED_BODY()
-public:
-	ENGINE_API UDistributionVectorParticleParameter(const FObjectInitializer& ObjectInitializer = FObjectInitializer::Get());
+	GENERATED_UCLASS_BODY()
 
 
 	// Begin UDistributionVectorParameterBase Interface
diff --git a/Engine/Source/Runtime/Engine/Classes/Distributions/DistributionVectorUniform.h b/Engine/Source/Runtime/Engine/Classes/Distributions/DistributionVectorUniform.h
--- a/Engine/Source/Runtime/Engine/Classes/Distributions/DistributionVectorUniform.h
+++ b/Engine/Source/Runtime/Engine/Classes/Distributions/DistributionVectorUniform.h
@@ -8,9 +8,7 @@
 UCLASS(collapsecategories, hidecategories=Object, editinlinenew, MinimalAPI)
 class UDistributionVectorUniform : public UDistributionVector
 {
-	GENERATED_BODY()
-public:
-	ENGINE_API UDistributionVectorUniform(const FObjectInitializer& ObjectInitializer = FObjectInitializer::Get());
+	GENERATED_UCLASS_BODY()
 
 	/** Upper end of FVector magnitude range. */
 	UPROPERTY(EditAnywhere, Category=DistributionVectorUniform)
diff --git a/Engine/Source/Runtime/Engine/Classes/Distributions/DistributionVectorUniformCurve.h b/Engine/Source/Runtime/Engine/Classes/Distributions/DistributionVectorUniformCurve.h
--- a/Engine/Source/Runtime/Engine/Classes/Distributions/DistributionVectorUniformCurve.h
+++ b/Engine/Source/Runtime/Engine/Classes/Distributions/DistributionVectorUniformCurve.h
@@ -8,9 +8,7 @@
 UCLASS(collapsecategories, hidecategories=Object, editinlinenew, MinimalAPI)
 class UDistributionVectorUniformCurve : public UDistributionVector
 {
-	GENERATED_BODY()
-public:
-	ENGINE_API UDistributionVectorUniformCurve(const FObjectInitializer& ObjectInitializer = FObjectInitializer::Get());
+	GENERATED_UCLASS_BODY()
 
 	/** Keyframe data for how output constant varies over time. */
 	UPROPERTY(EditAnywhere, Category=DistributionVectorUniformCurve)
diff --git a/Engine/Source/Runtime/Engine/Classes/EdGraph/EdGraph.h b/Engine/Source/Runtime/Engine/Classes/EdGraph/EdGraph.h
--- a/Engine/Source/Runtime/Engine/Classes/EdGraph/EdGraph.h
+++ b/Engine/Source/Runtime/Engine/Classes/EdGraph/EdGraph.h
@@ -9,7 +9,7 @@ struct FEdGraphEditAction;
 USTRUCT()
 struct ENGINE_API FGraphReference
 {
-	GENERATED_BODY()
+	GENERATED_USTRUCT_BODY()
 protected:
 	// Reference to the actual graph
 	UPROPERTY()
@@ -55,9 +55,7 @@ struct TStructOpsTypeTraits<FGraphReference> : public TStructOpsTypeTraitsBase
 UCLASS()
 class ENGINE_API UEdGraph : public UObject
 {
-	GENERATED_BODY()
-public:
-	UEdGraph(const FObjectInitializer& ObjectInitializer = FObjectInitializer::Get());
+	GENERATED_UCLASS_BODY()
 
 	/** The schema that this graph obeys */
 	UPROPERTY()
diff --git a/Engine/Source/Runtime/Engine/Classes/EdGraph/EdGraphNode.h b/Engine/Source/Runtime/Engine/Classes/EdGraph/EdGraphNode.h
--- a/Engine/Source/Runtime/Engine/Classes/EdGraph/EdGraphNode.h
+++ b/Engine/Source/Runtime/Engine/Classes/EdGraph/EdGraphNode.h
@@ -85,9 +85,7 @@ struct FGraphNodeContextMenuBuilder
 UCLASS()
 class ENGINE_API UEdGraphNode : public UObject
 {
-	GENERATED_BODY()
-public:
-	UEdGraphNode(const FObjectInitializer& ObjectInitializer = FObjectInitializer::Get());
+	GENERATED_UCLASS_BODY()
 
 	/** List of connector pins */
 	UPROPERTY()
diff --git a/Engine/Source/Runtime/Engine/Classes/EdGraph/EdGraphNode_Documentation.h b/Engine/Source/Runtime/Engine/Classes/EdGraph/EdGraphNode_Documentation.h
--- a/Engine/Source/Runtime/Engine/Classes/EdGraph/EdGraphNode_Documentation.h
+++ b/Engine/Source/Runtime/Engine/Classes/EdGraph/EdGraphNode_Documentation.h
@@ -6,9 +6,7 @@
 UCLASS(MinimalAPI)
 class UEdGraphNode_Documentation : public UEdGraphNode
 {
-	GENERATED_BODY()
-public:
-	ENGINE_API UEdGraphNode_Documentation(const FObjectInitializer& ObjectInitializer = FObjectInitializer::Get());
+	GENERATED_UCLASS_BODY()
 
 	/** Documentation Link */
 	UPROPERTY(meta=(FriendlyName="Documentation Link"))
diff --git a/Engine/Source/Runtime/Engine/Classes/EdGraph/EdGraphPin.h b/Engine/Source/Runtime/Engine/Classes/EdGraph/EdGraphPin.h
--- a/Engine/Source/Runtime/Engine/Classes/EdGraph/EdGraphPin.h
+++ b/Engine/Source/Runtime/Engine/Classes/EdGraph/EdGraphPin.h
@@ -8,7 +8,7 @@
 USTRUCT()
 struct FSimpleMemberReference
 {
-	GENERATED_BODY()
+	GENERATED_USTRUCT_BODY()
 
 	/** Class that this member is defined in. */
 	UPROPERTY()
@@ -48,7 +48,7 @@ FORCEINLINE FArchive& operator<<(FArchive& Ar, FSimpleMemberReference& Data)
 USTRUCT()
 struct FEdGraphPinType
 {
-	GENERATED_BODY()
+	GENERATED_USTRUCT_BODY()
 
 	/** Category of pin type */
 	UPROPERTY()
@@ -152,9 +152,7 @@ enum EBlueprintPinStyleType
 UCLASS(MinimalAPI)
 class UEdGraphPin : public UObject
 {
-	GENERATED_BODY()
-public:
-	ENGINE_API UEdGraphPin(const FObjectInitializer& ObjectInitializer = FObjectInitializer::Get());
+	GENERATED_UCLASS_BODY()
 
 	/** Name of this pin */
 	UPROPERTY()
diff --git a/Engine/Source/Runtime/Engine/Classes/EdGraph/EdGraphSchema.h b/Engine/Source/Runtime/Engine/Classes/EdGraph/EdGraphSchema.h
--- a/Engine/Source/Runtime/Engine/Classes/EdGraph/EdGraphSchema.h
+++ b/Engine/Source/Runtime/Engine/Classes/EdGraph/EdGraphSchema.h
@@ -50,7 +50,7 @@ enum ECanCreateConnectionResponse
 USTRUCT()
 struct ENGINE_API FEdGraphSchemaAction
 {
-	GENERATED_BODY()
+	GENERATED_USTRUCT_BODY()
 
 	// Simple type info
 	static FName StaticGetTypeId() {static FName Type("FEdGraphSchemaAction"); return Type;}
@@ -146,7 +146,7 @@ struct ENGINE_API FEdGraphSchemaAction
 USTRUCT()
 struct ENGINE_API FEdGraphSchemaAction_NewNode : public FEdGraphSchemaAction
 {
-	GENERATED_BODY()
+	GENERATED_USTRUCT_BODY()
 
 	// Simple type info
 	static FName StaticGetTypeId() {static FName Type("FEdGraphSchemaAction_NewNode"); return Type;}
@@ -398,9 +398,7 @@ public:
 UCLASS(abstract)
 class ENGINE_API UEdGraphSchema : public UObject
 {
-	GENERATED_BODY()
-public:
-	UEdGraphSchema(const FObjectInitializer& ObjectInitializer = FObjectInitializer::Get());
+	GENERATED_UCLASS_BODY()
 
 
 	/**
diff --git a/Engine/Source/Runtime/Engine/Classes/EditorFramework/AssetImportData.h b/Engine/Source/Runtime/Engine/Classes/EditorFramework/AssetImportData.h
--- a/Engine/Source/Runtime/Engine/Classes/EditorFramework/AssetImportData.h
+++ b/Engine/Source/Runtime/Engine/Classes/EditorFramework/AssetImportData.h
@@ -9,9 +9,7 @@
 UCLASS(MinimalAPI)
 class UAssetImportData : public UObject
 {
-	GENERATED_BODY()
-public:
-	ENGINE_API UAssetImportData(const FObjectInitializer& ObjectInitializer = FObjectInitializer::Get());
+	GENERATED_UCLASS_BODY()
 
 	/** Path to the resource used to construct this static mesh. Relative to the object's package, BaseDir() or absolute */
 	UPROPERTY(EditAnywhere, Category=ImportSettings)
diff --git a/Engine/Source/Runtime/Engine/Classes/EditorFramework/ThumbnailInfo.h b/Engine/Source/Runtime/Engine/Classes/EditorFramework/ThumbnailInfo.h
--- a/Engine/Source/Runtime/Engine/Classes/EditorFramework/ThumbnailInfo.h
+++ b/Engine/Source/Runtime/Engine/Classes/EditorFramework/ThumbnailInfo.h
@@ -10,9 +10,7 @@
 UCLASS(MinimalAPI)
 class UThumbnailInfo : public UObject
 {
-	GENERATED_BODY()
-public:
-	ENGINE_API UThumbnailInfo(const FObjectInitializer& ObjectInitializer = FObjectInitializer::Get());
+	GENERATED_UCLASS_BODY()
 };
 
 
diff --git a/Engine/Source/Runtime/Engine/Classes/Engine/ActorChannel.h b/Engine/Source/Runtime/Engine/Classes/Engine/ActorChannel.h
--- a/Engine/Source/Runtime/Engine/Classes/Engine/ActorChannel.h
+++ b/Engine/Source/Runtime/Engine/Classes/Engine/ActorChannel.h
@@ -36,12 +36,11 @@
  *		| </End Tag>           |                                                                           |
  *		|----------------------|---------------------------------------------------------------------------|
  */
-UCLASS(transient)
+UCLASS(transient, customConstructor)
 class ENGINE_API UActorChannel
 	: public UChannel
 {
-	GENERATED_BODY()
-public:
+	GENERATED_UCLASS_BODY()
 
 	friend class FObjectReplicator;
 
diff --git a/Engine/Source/Runtime/Engine/Classes/Engine/AssetUserData.h b/Engine/Source/Runtime/Engine/Classes/Engine/AssetUserData.h
--- a/Engine/Source/Runtime/Engine/Classes/Engine/AssetUserData.h
+++ b/Engine/Source/Runtime/Engine/Classes/Engine/AssetUserData.h
@@ -12,9 +12,7 @@ UCLASS(DefaultToInstanced)
 class ENGINE_API UAssetUserData
 	: public UObject
 {
-	GENERATED_BODY()
-public:
-	UAssetUserData(const FObjectInitializer& ObjectInitializer = FObjectInitializer::Get());
+	GENERATED_UCLASS_BODY()
 
 	/** used for debugging UAssetUserData data in editor */
 	virtual void Draw(class FPrimitiveDrawInterface* PDI, const class FSceneView* View) const {}
diff --git a/Engine/Source/Runtime/Engine/Classes/Engine/BlockingVolume.h b/Engine/Source/Runtime/Engine/Classes/Engine/BlockingVolume.h
--- a/Engine/Source/Runtime/Engine/Classes/Engine/BlockingVolume.h
+++ b/Engine/Source/Runtime/Engine/Classes/Engine/BlockingVolume.h
@@ -11,9 +11,7 @@ UCLASS(MinimalAPI)
 class ABlockingVolume
 	: public AVolume
 {
-	GENERATED_BODY()
-public:
-	ENGINE_API ABlockingVolume(const FObjectInitializer& ObjectInitializer = FObjectInitializer::Get());
+	GENERATED_UCLASS_BODY()
 
 public:
 
diff --git a/Engine/Source/Runtime/Engine/Classes/Engine/Blueprint.h b/Engine/Source/Runtime/Engine/Classes/Engine/Blueprint.h
--- a/Engine/Source/Runtime/Engine/Classes/Engine/Blueprint.h
+++ b/Engine/Source/Runtime/Engine/Classes/Engine/Blueprint.h
@@ -113,7 +113,7 @@ public:
 USTRUCT()
 struct FBPVariableMetaDataEntry
 {
-	GENERATED_BODY()
+	GENERATED_USTRUCT_BODY()
 
 	/** Name of metadata key */
 	UPROPERTY(EditAnywhere, Category=BPVariableMetaDataEntry)
@@ -138,7 +138,7 @@ struct FBPVariableMetaDataEntry
 USTRUCT()
 struct FBPVariableDescription
 {
-	GENERATED_BODY()
+	GENERATED_USTRUCT_BODY()
 
 	/** Name of the variable */
 	UPROPERTY(EditAnywhere, Category=BPVariableDescription)
@@ -198,7 +198,7 @@ struct FBPVariableDescription
 USTRUCT()
 struct FBPInterfaceDescription
 {
-	GENERATED_BODY()
+	GENERATED_USTRUCT_BODY()
 
 	/** Reference to the interface class we're adding to this blueprint */
 	UPROPERTY()
@@ -218,7 +218,7 @@ struct FBPInterfaceDescription
 USTRUCT()
 struct FEditedDocumentInfo
 {
-	GENERATED_BODY()
+	GENERATED_USTRUCT_BODY()
 
 	UPROPERTY()
 	UObject* EditedObject;
@@ -264,9 +264,7 @@ struct FEditedDocumentInfo
 UCLASS(config=Engine, BlueprintType)
 class ENGINE_API UBlueprint : public UBlueprintCore
 {
-	GENERATED_BODY()
-public:
-	UBlueprint(const FObjectInitializer& ObjectInitializer = FObjectInitializer::Get());
+	GENERATED_UCLASS_BODY()
 
 	/** Whether or not this blueprint should recompile itself on load */
 	UPROPERTY(config)
diff --git a/Engine/Source/Runtime/Engine/Classes/Engine/BlueprintCore.h b/Engine/Source/Runtime/Engine/Classes/Engine/BlueprintCore.h
--- a/Engine/Source/Runtime/Engine/Classes/Engine/BlueprintCore.h
+++ b/Engine/Source/Runtime/Engine/Classes/Engine/BlueprintCore.h
@@ -9,9 +9,7 @@ UCLASS()
 class ENGINE_API UBlueprintCore
 	: public UObject
 {
-	GENERATED_BODY()
-public:
-	UBlueprintCore(const FObjectInitializer& ObjectInitializer = FObjectInitializer::Get());
+	GENERATED_UCLASS_BODY()
 
 	/** Pointer to the skeleton class; this is regenerated any time a member variable or function is added but  
 	is usually incomplete (no code or hidden autogenerated variables are added to it) */
diff --git a/Engine/Source/Runtime/Engine/Classes/Engine/BlueprintGeneratedClass.h b/Engine/Source/Runtime/Engine/Classes/Engine/BlueprintGeneratedClass.h
--- a/Engine/Source/Runtime/Engine/Classes/Engine/BlueprintGeneratedClass.h
+++ b/Engine/Source/Runtime/Engine/Classes/Engine/BlueprintGeneratedClass.h
@@ -12,7 +12,7 @@ class UEdGraphPin;
 USTRUCT()
 struct FNodeToCodeAssociation
 {
-	GENERATED_BODY()
+	GENERATED_USTRUCT_BODY()
 
 public:
 	TWeakObjectPtr<UEdGraphNode> Node;
@@ -38,7 +38,7 @@ public:
 USTRUCT()
 struct FDebuggingInfoForSingleFunction
 {
-	GENERATED_BODY()
+	GENERATED_USTRUCT_BODY()
 
 public:
 	// Reverse map from code offset to source node
@@ -63,7 +63,7 @@ public:
 USTRUCT()
 struct FPointerToUberGraphFrame
 {
-	GENERATED_BODY()
+	GENERATED_USTRUCT_BODY()
 
 public:
 	uint8* RawPointer;
@@ -180,7 +180,7 @@ public:
 USTRUCT()
 struct ENGINE_API FBlueprintDebugData
 {
-	GENERATED_BODY()
+	GENERATED_USTRUCT_BODY()
 	FBlueprintDebugData()
 	{
 	}
@@ -420,7 +420,7 @@ public:
 USTRUCT()
 struct ENGINE_API FEventGraphFastCallPair
 {
-	GENERATED_BODY()
+	GENERATED_USTRUCT_BODY()
 
 	UPROPERTY()
 	UFunction* FunctionToPatch;
@@ -432,9 +432,7 @@ struct ENGINE_API FEventGraphFastCallPair
 UCLASS()
 class ENGINE_API UBlueprintGeneratedClass : public UClass
 {
-	GENERATED_BODY()
-public:
-	UBlueprintGeneratedClass(const FObjectInitializer& ObjectInitializer = FObjectInitializer::Get());
+	GENERATED_UCLASS_BODY()
 
 public:
 
diff --git a/Engine/Source/Runtime/Engine/Classes/Engine/BookMark.h b/Engine/Source/Runtime/Engine/Classes/Engine/BookMark.h
--- a/Engine/Source/Runtime/Engine/Classes/Engine/BookMark.h
+++ b/Engine/Source/Runtime/Engine/Classes/Engine/BookMark.h
@@ -11,9 +11,7 @@
 UCLASS(hidecategories=Object, MinimalAPI)
 class UBookMark : public UObject
 {
-	GENERATED_BODY()
-public:
-	ENGINE_API UBookMark(const FObjectInitializer& ObjectInitializer = FObjectInitializer::Get());
+	GENERATED_UCLASS_BODY()
 
 	/** Camera position */
 	UPROPERTY(EditAnywhere, Category=BookMark)
diff --git a/Engine/Source/Runtime/Engine/Classes/Engine/BookMark2D.h b/Engine/Source/Runtime/Engine/Classes/Engine/BookMark2D.h
--- a/Engine/Source/Runtime/Engine/Classes/Engine/BookMark2D.h
+++ b/Engine/Source/Runtime/Engine/Classes/Engine/BookMark2D.h
@@ -12,9 +12,7 @@ UCLASS(hidecategories=Object)
 class UBookMark2D
 	: public UObject
 {
-	GENERATED_BODY()
-public:
-	UBookMark2D(const FObjectInitializer& ObjectInitializer = FObjectInitializer::Get());
+	GENERATED_UCLASS_BODY()
 
 	/** Zoom of the camera */
 	UPROPERTY(EditAnywhere, Category=BookMark2D)
diff --git a/Engine/Source/Runtime/Engine/Classes/Engine/BoxReflectionCapture.h b/Engine/Source/Runtime/Engine/Classes/Engine/BoxReflectionCapture.h
--- a/Engine/Source/Runtime/Engine/Classes/Engine/BoxReflectionCapture.h
+++ b/Engine/Source/Runtime/Engine/Classes/Engine/BoxReflectionCapture.h
@@ -14,7 +14,5 @@ UCLASS(hidecategories=(Collision, Attachment, Actor), MinimalAPI)
 class ABoxReflectionCapture
 	: public AReflectionCapture
 {
-	GENERATED_BODY()
-public:
-	ENGINE_API ABoxReflectionCapture(const FObjectInitializer& ObjectInitializer = FObjectInitializer::Get());
+	GENERATED_UCLASS_BODY()
 };
diff --git a/Engine/Source/Runtime/Engine/Classes/Engine/Breakpoint.h b/Engine/Source/Runtime/Engine/Classes/Engine/Breakpoint.h
--- a/Engine/Source/Runtime/Engine/Classes/Engine/Breakpoint.h
+++ b/Engine/Source/Runtime/Engine/Classes/Engine/Breakpoint.h
@@ -9,9 +9,7 @@ UCLASS()
 class ENGINE_API UBreakpoint
 	: public UObject
 {
-	GENERATED_BODY()
-public:
-	UBreakpoint(const FObjectInitializer& ObjectInitializer = FObjectInitializer::Get());
+	GENERATED_UCLASS_BODY()
 
 private:
 	// Is the breakpoint currently enabled?
diff --git a/Engine/Source/Runtime/Engine/Classes/Engine/Brush.h b/Engine/Source/Runtime/Engine/Classes/Engine/Brush.h
--- a/Engine/Source/Runtime/Engine/Classes/Engine/Brush.h
+++ b/Engine/Source/Runtime/Engine/Classes/Engine/Brush.h
@@ -45,7 +45,7 @@ enum EBrushType
 USTRUCT()
 struct FGeomSelection
 {
-	GENERATED_BODY()
+	GENERATED_USTRUCT_BODY()
 
 	UPROPERTY()
 	int32 Type;    // EGeometrySelectionType_
@@ -74,9 +74,7 @@ UCLASS(hidecategories=(Object, Collision, Display, Rendering, Physics, Input, Bl
 class ENGINE_API ABrush
 	: public AActor
 {
-	GENERATED_BODY()
-public:
-	ABrush(const FObjectInitializer& ObjectInitializer = FObjectInitializer::Get());
+	GENERATED_UCLASS_BODY()
 
 	/** Type of brush */
 	UPROPERTY(EditAnywhere, Category=Brush)
diff --git a/Engine/Source/Runtime/Engine/Classes/Engine/BrushBuilder.h b/Engine/Source/Runtime/Engine/Classes/Engine/BrushBuilder.h
--- a/Engine/Source/Runtime/Engine/Classes/Engine/BrushBuilder.h
+++ b/Engine/Source/Runtime/Engine/Classes/Engine/BrushBuilder.h
@@ -9,7 +9,7 @@
 USTRUCT()
 struct FBuilderPoly
 {
-	GENERATED_BODY()
+	GENERATED_USTRUCT_BODY()
 
 	UPROPERTY()
 	TArray<int32> VertexIndices;
@@ -48,9 +48,7 @@ UCLASS(abstract, hidecategories=(Object), MinimalAPI)
 class UBrushBuilder
 	: public UObject
 {
-	GENERATED_BODY()
-public:
-	ENGINE_API UBrushBuilder(const FObjectInitializer& ObjectInitializer = FObjectInitializer::Get());
+	GENERATED_UCLASS_BODY()
 
 	UPROPERTY()
 	FString BitmapFilename;
diff --git a/Engine/Source/Runtime/Engine/Classes/Engine/BrushShape.h b/Engine/Source/Runtime/Engine/Classes/Engine/BrushShape.h
--- a/Engine/Source/Runtime/Engine/Classes/Engine/BrushShape.h
+++ b/Engine/Source/Runtime/Engine/Classes/Engine/BrushShape.h
@@ -12,9 +12,7 @@ UCLASS(MinimalAPI)
 class ABrushShape
 	: public ABrush
 {
-	GENERATED_BODY()
-public:
-	ENGINE_API ABrushShape(const FObjectInitializer& ObjectInitializer = FObjectInitializer::Get());
+	GENERATED_UCLASS_BODY()
 
 public:
 
diff --git a/Engine/Source/Runtime/Engine/Classes/Engine/Canvas.h b/Engine/Source/Runtime/Engine/Classes/Engine/Canvas.h
--- a/Engine/Source/Runtime/Engine/Classes/Engine/Canvas.h
+++ b/Engine/Source/Runtime/Engine/Classes/Engine/Canvas.h
@@ -16,7 +16,7 @@ class UFont;
 USTRUCT()
 struct FCanvasIcon
 {
-	GENERATED_BODY()
+	GENERATED_USTRUCT_BODY()
 
 	/** Source texture */
 	UPROPERTY(EditAnywhere, BlueprintReadWrite, Category=CanvasIcon)
@@ -53,9 +53,7 @@ UCLASS(transient)
 class ENGINE_API UCanvas
 	: public UObject
 {
-	GENERATED_BODY()
-public:
-	UCanvas(const FObjectInitializer& ObjectInitializer = FObjectInitializer::Get());
+	GENERATED_UCLASS_BODY()
 
 	// Modifiable properties.
 	UPROPERTY()
diff --git a/Engine/Source/Runtime/Engine/Classes/Engine/CanvasRenderTarget2D.h b/Engine/Source/Runtime/Engine/Classes/Engine/CanvasRenderTarget2D.h
--- a/Engine/Source/Runtime/Engine/Classes/Engine/CanvasRenderTarget2D.h
+++ b/Engine/Source/Runtime/Engine/Classes/Engine/CanvasRenderTarget2D.h
@@ -24,9 +24,7 @@ UCLASS(BlueprintType, Blueprintable)
 class ENGINE_API UCanvasRenderTarget2D
 	: public UTextureRenderTarget2D
 {
-	GENERATED_BODY()
-public:
-	UCanvasRenderTarget2D(const FObjectInitializer& ObjectInitializer = FObjectInitializer::Get());
+	GENERATED_UCLASS_BODY()
 
 public:
 
diff --git a/Engine/Source/Runtime/Engine/Classes/Engine/Channel.h b/Engine/Source/Runtime/Engine/Classes/Engine/Channel.h
--- a/Engine/Source/Runtime/Engine/Classes/Engine/Channel.h
+++ b/Engine/Source/Runtime/Engine/Classes/Engine/Channel.h
@@ -33,9 +33,7 @@ UCLASS(transient)
 class ENGINE_API UChannel
 	: public UObject
 {
-	GENERATED_BODY()
-public:
-	UChannel(const FObjectInitializer& ObjectInitializer = FObjectInitializer::Get());
+	GENERATED_UCLASS_BODY()
 	
 	UPROPERTY()
 	class UNetConnection*	Connection;		// Owner connection.
diff --git a/Engine/Source/Runtime/Engine/Classes/Engine/ChildConnection.h b/Engine/Source/Runtime/Engine/Classes/Engine/ChildConnection.h
--- a/Engine/Source/Runtime/Engine/Classes/Engine/ChildConnection.h
+++ b/Engine/Source/Runtime/Engine/Classes/Engine/ChildConnection.h
@@ -12,9 +12,7 @@ UCLASS(MinimalAPI,transient,config=Engine)
 class UChildConnection
 	: public UNetConnection
 {
-	GENERATED_BODY()
-public:
-	ENGINE_API UChildConnection(const FObjectInitializer& ObjectInitializer = FObjectInitializer::Get());
+	GENERATED_UCLASS_BODY()
 
 	UPROPERTY(transient)
 	class UNetConnection* Parent;
diff --git a/Engine/Source/Runtime/Engine/Classes/Engine/CloudStorageBase.h b/Engine/Source/Runtime/Engine/Classes/Engine/CloudStorageBase.h
--- a/Engine/Source/Runtime/Engine/Classes/Engine/CloudStorageBase.h
+++ b/Engine/Source/Runtime/Engine/Classes/Engine/CloudStorageBase.h
@@ -33,9 +33,7 @@ UCLASS()
 class UCloudStorageBase
 	: public UPlatformInterfaceBase
 {
-	GENERATED_BODY()
-public:
-	UCloudStorageBase(const FObjectInitializer& ObjectInitializer = FObjectInitializer::Get());
+	GENERATED_UCLASS_BODY()
 
 	/** When using local storage (aka "cloud emulation"), this maintains a list of the file paths */
 	UPROPERTY()
diff --git a/Engine/Source/Runtime/Engine/Classes/Engine/ComponentDelegateBinding.h b/Engine/Source/Runtime/Engine/Classes/Engine/ComponentDelegateBinding.h
--- a/Engine/Source/Runtime/Engine/Classes/Engine/ComponentDelegateBinding.h
+++ b/Engine/Source/Runtime/Engine/Classes/Engine/ComponentDelegateBinding.h
@@ -10,7 +10,7 @@
 USTRUCT()
 struct ENGINE_API FBlueprintComponentDelegateBinding
 {
-	GENERATED_BODY()
+	GENERATED_USTRUCT_BODY()
 
 	/** Name of component property that contains delegate we want to assign to. */
 	UPROPERTY()
@@ -36,9 +36,7 @@ UCLASS()
 class ENGINE_API UComponentDelegateBinding
 	: public UDynamicBlueprintBinding
 {
-	GENERATED_BODY()
-public:
-	UComponentDelegateBinding(const FObjectInitializer& ObjectInitializer = FObjectInitializer::Get());
+	GENERATED_UCLASS_BODY()
 
 	UPROPERTY()
 	TArray<FBlueprintComponentDelegateBinding> ComponentDelegateBindings;
diff --git a/Engine/Source/Runtime/Engine/Classes/Engine/ControlChannel.h b/Engine/Source/Runtime/Engine/Classes/Engine/ControlChannel.h
--- a/Engine/Source/Runtime/Engine/Classes/Engine/ControlChannel.h
+++ b/Engine/Source/Runtime/Engine/Classes/Engine/ControlChannel.h
@@ -8,12 +8,11 @@
 /**
  * A channel for exchanging connection control messages.
  */
-UCLASS(transient)
+UCLASS(transient, customConstructor)
 class UControlChannel
 	: public UChannel
 {
-	GENERATED_BODY()
-public:
+	GENERATED_UCLASS_BODY()
 
 	/**
 	 * Used to interrogate the first packet received to determine endianess
diff --git a/Engine/Source/Runtime/Engine/Classes/Engine/CullDistanceVolume.h b/Engine/Source/Runtime/Engine/Classes/Engine/CullDistanceVolume.h
--- a/Engine/Source/Runtime/Engine/Classes/Engine/CullDistanceVolume.h
+++ b/Engine/Source/Runtime/Engine/Classes/Engine/CullDistanceVolume.h
@@ -12,7 +12,7 @@
 USTRUCT()
 struct FCullDistanceSizePair
 {
-	GENERATED_BODY()
+	GENERATED_USTRUCT_BODY()
 
 	/** Size to associate with cull distance. */
 	UPROPERTY(EditAnywhere, BlueprintReadOnly, Category=CullDistanceSizePair)
@@ -39,9 +39,7 @@ UCLASS(hidecategories=(Advanced, Attachment, Collision, Volume))
 class ACullDistanceVolume
 	: public AVolume
 {
-	GENERATED_BODY()
-public:
-	ACullDistanceVolume(const FObjectInitializer& ObjectInitializer = FObjectInitializer::Get());
+	GENERATED_UCLASS_BODY()
 
 	/**
 	 * Array of size and cull distance pairs. The code will calculate the sphere diameter of a primitive's BB and look for a best
diff --git a/Engine/Source/Runtime/Engine/Classes/Engine/CurveTable.h b/Engine/Source/Runtime/Engine/Classes/Engine/CurveTable.h
--- a/Engine/Source/Runtime/Engine/Classes/Engine/CurveTable.h
+++ b/Engine/Source/Runtime/Engine/Classes/Engine/CurveTable.h
@@ -22,9 +22,7 @@ UCLASS(MinimalAPI)
 class UCurveTable
 	: public UObject
 {
-	GENERATED_BODY()
-public:
-	ENGINE_API UCurveTable(const FObjectInitializer& ObjectInitializer = FObjectInitializer::Get());
+	GENERATED_UCLASS_BODY()
 
 	/** The filename that was used to create this object. Relative to the object's package, BaseDir() or absolute */
 	UPROPERTY()
@@ -101,7 +99,7 @@ protected:
 USTRUCT(BlueprintType)
 struct ENGINE_API FCurveTableRowHandle
 {
-	GENERATED_BODY()
+	GENERATED_USTRUCT_BODY()
 
 	FCurveTableRowHandle()
 		: CurveTable(nullptr)
diff --git a/Engine/Source/Runtime/Engine/Classes/Engine/DataAsset.h b/Engine/Source/Runtime/Engine/Classes/Engine/DataAsset.h
--- a/Engine/Source/Runtime/Engine/Classes/Engine/DataAsset.h
+++ b/Engine/Source/Runtime/Engine/Classes/Engine/DataAsset.h
@@ -12,9 +12,7 @@ UCLASS(abstract,MinimalAPI)
 class UDataAsset
 	: public UObject
 {
-	GENERATED_BODY()
-public:
-	ENGINE_API UDataAsset(const FObjectInitializer& ObjectInitializer = FObjectInitializer::Get());
+	GENERATED_UCLASS_BODY()
 
 public:
 
diff --git a/Engine/Source/Runtime/Engine/Classes/Engine/DataTable.h b/Engine/Source/Runtime/Engine/Classes/Engine/DataTable.h
--- a/Engine/Source/Runtime/Engine/Classes/Engine/DataTable.h
+++ b/Engine/Source/Runtime/Engine/Classes/Engine/DataTable.h
@@ -21,7 +21,7 @@ class TJsonWriter;
 USTRUCT()
 struct FTableRowBase
 {
-	GENERATED_BODY()
+	GENERATED_USTRUCT_BODY()
 
 	FTableRowBase() { }
 };
@@ -34,9 +34,7 @@ UCLASS(MinimalAPI)
 class UDataTable
 	: public UObject
 {
-	GENERATED_BODY()
-public:
-	ENGINE_API UDataTable(const FObjectInitializer& ObjectInitializer = FObjectInitializer::Get());
+	GENERATED_UCLASS_BODY()
 
 	/** Structure to use for each row of the table, must inherit from FTableRowBase */
 	UPROPERTY()
@@ -203,7 +201,7 @@ private:
 USTRUCT(BlueprintType)
 struct ENGINE_API FDataTableRowHandle
 {
-	GENERATED_BODY()
+	GENERATED_USTRUCT_BODY()
 
 	FDataTableRowHandle()
 		: DataTable(nullptr)
@@ -244,7 +242,7 @@ struct ENGINE_API FDataTableRowHandle
 USTRUCT()
 struct ENGINE_API FDataTableCategoryHandle
 {
-	GENERATED_BODY()
+	GENERATED_USTRUCT_BODY()
 
 	/** Pointer to table we want a row from */
 	UPROPERTY(EditAnywhere, BlueprintReadWrite, Category=DataTableCategoryHandle)
diff --git a/Engine/Source/Runtime/Engine/Classes/Engine/DebugCameraController.h b/Engine/Source/Runtime/Engine/Classes/Engine/DebugCameraController.h
--- a/Engine/Source/Runtime/Engine/Classes/Engine/DebugCameraController.h
+++ b/Engine/Source/Runtime/Engine/Classes/Engine/DebugCameraController.h
@@ -16,9 +16,7 @@ UCLASS(config=Game, hidedropdown)
 class ENGINE_API ADebugCameraController
 	: public APlayerController
 {
-	GENERATED_BODY()
-public:
-	ADebugCameraController(const FObjectInitializer& ObjectInitializer = FObjectInitializer::Get());
+	GENERATED_UCLASS_BODY()
 
 	/** Whether to show information about the selected actor on the debug camera HUD. */
 	UPROPERTY(globalconfig)
diff --git a/Engine/Source/Runtime/Engine/Classes/Engine/DebugCameraHUD.h b/Engine/Source/Runtime/Engine/Classes/Engine/DebugCameraHUD.h
--- a/Engine/Source/Runtime/Engine/Classes/Engine/DebugCameraHUD.h
+++ b/Engine/Source/Runtime/Engine/Classes/Engine/DebugCameraHUD.h
@@ -13,9 +13,7 @@ UCLASS(config=Game)
 class ENGINE_API ADebugCameraHUD
 	: public AHUD
 {
-	GENERATED_BODY()
-public:
-	ADebugCameraHUD(const FObjectInitializer& ObjectInitializer = FObjectInitializer::Get());
+	GENERATED_UCLASS_BODY()
 
 	/** @todo document */
 	virtual bool DisplayMaterials( float X, float& Y, float DY, class UMeshComponent* MeshComp );
diff --git a/Engine/Source/Runtime/Engine/Classes/Engine/DecalActor.h b/Engine/Source/Runtime/Engine/Classes/Engine/DecalActor.h
--- a/Engine/Source/Runtime/Engine/Classes/Engine/DecalActor.h
+++ b/Engine/Source/Runtime/Engine/Classes/Engine/DecalActor.h
@@ -15,9 +15,7 @@ UCLASS(hideCategories=(Collision, Attachment, Actor, Input, Replication), showCa
 class ADecalActor
 	: public AActor
 {
-	GENERATED_BODY()
-public:
-	ENGINE_API ADecalActor(const FObjectInitializer& ObjectInitializer = FObjectInitializer::Get());
+	GENERATED_UCLASS_BODY()
 
 private_subobject:
 	/** The decal component for this decal actor */
diff --git a/Engine/Source/Runtime/Engine/Classes/Engine/DemoNetConnection.h b/Engine/Source/Runtime/Engine/Classes/Engine/DemoNetConnection.h
--- a/Engine/Source/Runtime/Engine/Classes/Engine/DemoNetConnection.h
+++ b/Engine/Source/Runtime/Engine/Classes/Engine/DemoNetConnection.h
@@ -18,9 +18,7 @@ UCLASS(transient, config=Engine)
 class UDemoNetConnection
 	: public UNetConnection
 {
-	GENERATED_BODY()
-public:
-	UDemoNetConnection(const FObjectInitializer& ObjectInitializer = FObjectInitializer::Get());
+	GENERATED_UCLASS_BODY()
 
 public:
 
diff --git a/Engine/Source/Runtime/Engine/Classes/Engine/DemoNetDriver.h b/Engine/Source/Runtime/Engine/Classes/Engine/DemoNetDriver.h
--- a/Engine/Source/Runtime/Engine/Classes/Engine/DemoNetDriver.h
+++ b/Engine/Source/Runtime/Engine/Classes/Engine/DemoNetDriver.h
@@ -11,9 +11,7 @@ UCLASS(transient, config=Engine)
 class UDemoNetDriver
 	: public UNetDriver
 {
-	GENERATED_BODY()
-public:
-	UDemoNetDriver(const FObjectInitializer& ObjectInitializer = FObjectInitializer::Get());
+	GENERATED_UCLASS_BODY()
 
 	/** Name of the file to read/write from */
 	FString DemoFilename;
diff --git a/Engine/Source/Runtime/Engine/Classes/Engine/DemoPendingNetGame.h b/Engine/Source/Runtime/Engine/Classes/Engine/DemoPendingNetGame.h
--- a/Engine/Source/Runtime/Engine/Classes/Engine/DemoPendingNetGame.h
+++ b/Engine/Source/Runtime/Engine/Classes/Engine/DemoPendingNetGame.h
@@ -9,7 +9,5 @@ UCLASS(transient, config=Engine)
 class UDemoPendingNetGame
 	: public UPendingNetGame
 {
-	GENERATED_BODY()
-public:
-	UDemoPendingNetGame(const FObjectInitializer& ObjectInitializer = FObjectInitializer::Get());
+	GENERATED_UCLASS_BODY()
 };
diff --git a/Engine/Source/Runtime/Engine/Classes/Engine/DestructibleFractureSettings.h b/Engine/Source/Runtime/Engine/Classes/Engine/DestructibleFractureSettings.h
--- a/Engine/Source/Runtime/Engine/Classes/Engine/DestructibleFractureSettings.h
+++ b/Engine/Source/Runtime/Engine/Classes/Engine/DestructibleFractureSettings.h
@@ -44,7 +44,7 @@ namespace EDestructibleImportOptions
 USTRUCT()
 struct FFractureMaterial
 {
-	GENERATED_BODY()
+	GENERATED_USTRUCT_BODY()
 
 	/**
 	 * The UV scale (geometric distance/unit texture distance) for interior materials.
@@ -102,7 +102,7 @@ struct FFractureMaterial
 USTRUCT()
 struct FDestructibleChunkParameters
 {
-	GENERATED_BODY()
+	GENERATED_USTRUCT_BODY()
 
 	/**
 		Defines the chunk to be environmentally supported, if the appropriate NxDestructibleParametersFlag flags
@@ -153,9 +153,7 @@ UCLASS(MinimalAPI)
 class UDestructibleFractureSettings
 	: public UObject
 {
-	GENERATED_BODY()
-public:
-	ENGINE_API UDestructibleFractureSettings(const FObjectInitializer& ObjectInitializer = FObjectInitializer::Get());
+	GENERATED_UCLASS_BODY()
 
 #if WITH_HOT_RELOAD_CTORS
 	/** DO NOT USE. This constructor is for internal usage only for hot-reload purposes. */
diff --git a/Engine/Source/Runtime/Engine/Classes/Engine/DestructibleMesh.h b/Engine/Source/Runtime/Engine/Classes/Engine/DestructibleMesh.h
--- a/Engine/Source/Runtime/Engine/Classes/Engine/DestructibleMesh.h
+++ b/Engine/Source/Runtime/Engine/Classes/Engine/DestructibleMesh.h
@@ -39,7 +39,7 @@ enum EImpactDamageOverride
 USTRUCT()
 struct FDestructibleDepthParameters
 {
-	GENERATED_BODY()
+	GENERATED_USTRUCT_BODY()
 
 	/** Chunks up to the depth DefaultImpactDamageDepth will take impact damage, unless one of the override options (see EImpactDamageOverride) is chosen. */
 	UPROPERTY(EditAnywhere, Category=DestructibleDepthParameters)
@@ -62,7 +62,7 @@ struct FDestructibleDepthParameters
 USTRUCT()
 struct FDestructibleParametersFlag
 {
-	GENERATED_BODY()
+	GENERATED_USTRUCT_BODY()
 
 	/**
 		If set, chunks will "remember" damage applied to them, so that many applications of a damage amount
@@ -168,7 +168,7 @@ struct FDestructibleParametersFlag
 USTRUCT()
 struct FDestructibleDamageParameters
 {
-	GENERATED_BODY()
+	GENERATED_USTRUCT_BODY()
 
 	/**
 		The damage amount which will cause a chunk to fracture (break free).
@@ -238,7 +238,7 @@ struct FDestructibleDamageParameters
 USTRUCT()
 struct FDestructibleDebrisParameters
 {
-	GENERATED_BODY()
+	GENERATED_USTRUCT_BODY()
 
 	/**
 		"Debris chunks" (see debrisDepth, above) will be destroyed after a time (in seconds)
@@ -299,7 +299,7 @@ struct FDestructibleDebrisParameters
 USTRUCT()
 struct FDestructibleAdvancedParameters
 {
-	GENERATED_BODY()
+	GENERATED_USTRUCT_BODY()
 
 	/**
 		Limits the amount of damage applied to a chunk.  This is useful for preventing the entire destructible
@@ -352,7 +352,7 @@ struct FDestructibleAdvancedParameters
 USTRUCT()
 struct FDestructibleSpecialHierarchyDepths
 {
-	GENERATED_BODY()
+	GENERATED_USTRUCT_BODY()
 
 	// todo dynamicChunksDominanceGroup
 	// todo dynamicChunksGroupsMask
@@ -413,7 +413,7 @@ struct FDestructibleSpecialHierarchyDepths
 USTRUCT()
 struct FDestructibleParameters
 {
-	GENERATED_BODY()
+	GENERATED_USTRUCT_BODY()
 
 	/** Parameters that pertain to chunk damage.  See DestructibleDamageParameters. */
 	UPROPERTY(EditAnywhere, Category=DestructibleParameters)
@@ -452,9 +452,7 @@ UCLASS(hidecategories=(Object, Mesh, LevelOfDetail, Mirroring, Physics, Reimport
 class UDestructibleMesh
 	: public USkeletalMesh
 {
-	GENERATED_BODY()
-public:
-	ENGINE_API UDestructibleMesh(const FObjectInitializer& ObjectInitializer = FObjectInitializer::Get());
+	GENERATED_UCLASS_BODY()
 
 	/** Parameters controlling the destruction behavior. */
 	UPROPERTY(EditAnywhere, Category=DestructibleMesh)
diff --git a/Engine/Source/Runtime/Engine/Classes/Engine/DirectionalLight.h b/Engine/Source/Runtime/Engine/Classes/Engine/DirectionalLight.h
--- a/Engine/Source/Runtime/Engine/Classes/Engine/DirectionalLight.h
+++ b/Engine/Source/Runtime/Engine/Classes/Engine/DirectionalLight.h
@@ -16,9 +16,7 @@ UCLASS(ClassGroup=(Lights, DirectionalLights), MinimalAPI, meta=(ChildCanTick))
 class ADirectionalLight
 	: public ALight
 {
-	GENERATED_BODY()
-public:
-	ENGINE_API ADirectionalLight(const FObjectInitializer& ObjectInitializer = FObjectInitializer::Get());
+	GENERATED_UCLASS_BODY()
 
 #if WITH_EDITORONLY_DATA
 	// Reference to editor visualization arrow
diff --git a/Engine/Source/Runtime/Engine/Classes/Engine/DocumentationActor.h b/Engine/Source/Runtime/Engine/Classes/Engine/DocumentationActor.h
--- a/Engine/Source/Runtime/Engine/Classes/Engine/DocumentationActor.h
+++ b/Engine/Source/Runtime/Engine/Classes/Engine/DocumentationActor.h
@@ -25,9 +25,7 @@ UCLASS(hidecategories = (Sprite, MaterialSprite, Actor, Transform, Tags, Materia
 class ENGINE_API ADocumentationActor
 	: public AActor
 {
-	GENERATED_BODY()
-public:
-	ADocumentationActor(const FObjectInitializer& ObjectInitializer = FObjectInitializer::Get());
+	GENERATED_UCLASS_BODY()
 
 	/* Open the document link this actor relates to */
 	bool OpenDocumentLink() const;
diff --git a/Engine/Source/Runtime/Engine/Classes/Engine/DynamicBlueprintBinding.h b/Engine/Source/Runtime/Engine/Classes/Engine/DynamicBlueprintBinding.h
--- a/Engine/Source/Runtime/Engine/Classes/Engine/DynamicBlueprintBinding.h
+++ b/Engine/Source/Runtime/Engine/Classes/Engine/DynamicBlueprintBinding.h
@@ -9,9 +9,7 @@ UCLASS(abstract)
 class ENGINE_API UDynamicBlueprintBinding
 	: public UObject
 {
-	GENERATED_BODY()
-public:
-	UDynamicBlueprintBinding(const FObjectInitializer& ObjectInitializer = FObjectInitializer::Get());
+	GENERATED_UCLASS_BODY()
 
 	virtual void BindDynamicDelegates(UObject* InInstance) const { }
 };
diff --git a/Engine/Source/Runtime/Engine/Classes/Engine/Engine.h b/Engine/Source/Runtime/Engine/Classes/Engine/Engine.h
--- a/Engine/Source/Runtime/Engine/Classes/Engine/Engine.h
+++ b/Engine/Source/Runtime/Engine/Classes/Engine/Engine.h
@@ -63,7 +63,7 @@ enum EConsoleType
 USTRUCT()
 struct FFullyLoadedPackagesInfo
 {
-	GENERATED_BODY()
+	GENERATED_USTRUCT_BODY()
 
 	/** When to load these packages */
 	UPROPERTY()
@@ -94,7 +94,7 @@ struct FFullyLoadedPackagesInfo
 USTRUCT()
 struct FLevelStreamingStatus
 {
-	GENERATED_BODY()
+	GENERATED_USTRUCT_BODY()
 
 	UPROPERTY()
 	FName PackageName;
@@ -130,7 +130,7 @@ struct FLevelStreamingStatus
 USTRUCT()
 struct FNetDriverDefinition
 {
-	GENERATED_BODY()
+	GENERATED_USTRUCT_BODY()
 
 	/** Unique name of this net driver definition */
 	UPROPERTY()
@@ -160,7 +160,7 @@ struct FNetDriverDefinition
 USTRUCT()
 struct FNamedNetDriver
 {
-	GENERATED_BODY()
+	GENERATED_USTRUCT_BODY()
 
 	/** Instantiation of named net driver */
 	UPROPERTY(transient)
@@ -208,7 +208,7 @@ struct FNamedNetDriver
 USTRUCT()
 struct FWorldContext
 {
-	GENERATED_BODY()
+	GENERATED_USTRUCT_BODY()
 
 	/**************************************************************/
 	
@@ -341,7 +341,7 @@ private:
 USTRUCT()
 struct FStatColorMapEntry
 {
-	GENERATED_BODY()
+	GENERATED_USTRUCT_BODY()
 
 	UPROPERTY(globalconfig)
 	float In;
@@ -362,7 +362,7 @@ struct FStatColorMapEntry
 USTRUCT()
 struct FStatColorMapping
 {
-	GENERATED_BODY()
+	GENERATED_USTRUCT_BODY()
 
 	UPROPERTY(globalconfig)
 	FString StatName;
@@ -386,7 +386,7 @@ struct FStatColorMapping
 USTRUCT()
 struct FDropNoteInfo
 {
-	GENERATED_BODY()
+	GENERATED_USTRUCT_BODY()
 
 	/** Location to create Note actor in edited level. */
 	UPROPERTY()
@@ -415,7 +415,7 @@ struct FDropNoteInfo
 USTRUCT()
 struct FScreenMessageString
 {
-	GENERATED_BODY()
+	GENERATED_USTRUCT_BODY()
 
 	/** The 'key' for this message. */
 	UPROPERTY(transient)
@@ -466,7 +466,7 @@ namespace EMatineeCaptureType
 USTRUCT()
 struct FGameNameRedirect
 {
-	GENERATED_BODY()
+	GENERATED_USTRUCT_BODY()
 
 	UPROPERTY()
 	FString OldGameName;
@@ -479,7 +479,7 @@ struct FGameNameRedirect
 USTRUCT()
 struct FClassRedirect
 {
-	GENERATED_BODY()
+	GENERATED_USTRUCT_BODY()
 
 	UPROPERTY()
 	FString ObjectName;
@@ -504,7 +504,7 @@ struct FClassRedirect
 USTRUCT()
 struct FStructRedirect
 {
-	GENERATED_BODY()
+	GENERATED_USTRUCT_BODY()
 
 	UPROPERTY()
 	FString OldStructName;
@@ -517,7 +517,7 @@ struct FStructRedirect
 USTRUCT()
 struct FPluginRedirect
 {
-	GENERATED_BODY()
+	GENERATED_USTRUCT_BODY()
 
 	UPROPERTY()
 	FString OldPluginName;
@@ -535,7 +535,7 @@ DECLARE_DELEGATE(FEndStreamingPauseDelegate);
 USTRUCT()
 struct FMatineeScreenshotOptions
 {
-	GENERATED_BODY()
+	GENERATED_USTRUCT_BODY()
 
 	/** determines if we should start the matinee capture as soon as the game loads */
 	UPROPERTY(transient)
@@ -578,9 +578,7 @@ class ENGINE_API UEngine
 	: public UObject
 	, public FExec
 {
-	GENERATED_BODY()
-public:
-	UEngine(const FObjectInitializer& ObjectInitializer = FObjectInitializer::Get());
+	GENERATED_UCLASS_BODY()
 
 private:
 	// Fonts.
diff --git a/Engine/Source/Runtime/Engine/Classes/Engine/EngineBaseTypes.h b/Engine/Source/Runtime/Engine/Classes/Engine/EngineBaseTypes.h
--- a/Engine/Source/Runtime/Engine/Classes/Engine/EngineBaseTypes.h
+++ b/Engine/Source/Runtime/Engine/Classes/Engine/EngineBaseTypes.h
@@ -76,7 +76,7 @@ enum ETickingGroup
 USTRUCT()
 struct FTickPrerequisite
 {
-	GENERATED_BODY()
+	GENERATED_USTRUCT_BODY()
 
 	/** Tick functions live inside of UObjects, so we need a separate weak pointer to the UObject solely for the purpose of determining if PrerequisiteTickFunction is still valid **/
 	TWeakObjectPtr<class UObject> PrerequisiteObject;
@@ -126,7 +126,7 @@ struct FTickPrerequisite
 USTRUCT()
 struct ENGINE_API FTickFunction
 {
-	GENERATED_BODY()
+	GENERATED_USTRUCT_BODY()
 
 public:
 	// The following UPROPERTYs are for configuration and inherited from the CDO/archetype/blueprint etc
@@ -308,7 +308,7 @@ struct TStructOpsTypeTraits<FTickFunction> : public TStructOpsTypeTraitsBase
 USTRUCT()
 struct FActorTickFunction : public FTickFunction
 {
-	GENERATED_BODY()
+	GENERATED_USTRUCT_BODY()
 
 	/**  AActor  that is the target of this tick **/
 	class AActor*	Target;
@@ -340,7 +340,7 @@ struct TStructOpsTypeTraits<FActorTickFunction> : public TStructOpsTypeTraitsBas
 USTRUCT()
 struct FActorComponentTickFunction : public FTickFunction
 {
-	GENERATED_BODY()
+	GENERATED_USTRUCT_BODY()
 
 	/**  AActor  component that is the target of this tick **/
 	class UActorComponent*	Target;
@@ -372,7 +372,7 @@ struct TStructOpsTypeTraits<FActorComponentTickFunction> : public TStructOpsType
 USTRUCT()
 struct FPrimitiveComponentPostPhysicsTickFunction : public FTickFunction
 {
-	GENERATED_BODY()
+	GENERATED_USTRUCT_BODY()
 
 	/** PrimitiveComponent component that is the target of this tick **/
 	class UPrimitiveComponent*	Target;
@@ -575,7 +575,7 @@ namespace EDemoPlayFailure
 USTRUCT()
 struct ENGINE_API FURL
 {
-	GENERATED_BODY()
+	GENERATED_USTRUCT_BODY()
 
 	UPROPERTY()
 	FString Protocol;    // Protocol, i.e. "unreal" or "http".
@@ -704,7 +704,7 @@ enum EViewModeIndex
 USTRUCT()
 struct FExposureSettings
 {
-	GENERATED_BODY()
+	GENERATED_USTRUCT_BODY()
 
 	FExposureSettings()
 		: LogOffset(0), bFixed(false)
@@ -745,9 +745,7 @@ struct FExposureSettings
 UCLASS(abstract, config=Engine)
 class UEngineBaseTypes : public UObject
 {
-	GENERATED_BODY()
-public:
-	UEngineBaseTypes(const FObjectInitializer& ObjectInitializer = FObjectInitializer::Get());
+	GENERATED_UCLASS_BODY()
 
 };
 
diff --git a/Engine/Source/Runtime/Engine/Classes/Engine/ExponentialHeightFog.h b/Engine/Source/Runtime/Engine/Classes/Engine/ExponentialHeightFog.h
--- a/Engine/Source/Runtime/Engine/Classes/Engine/ExponentialHeightFog.h
+++ b/Engine/Source/Runtime/Engine/Classes/Engine/ExponentialHeightFog.h
@@ -12,9 +12,7 @@ UCLASS(showcategories=(Movement, Rendering, "Utilities|Transformation"), ClassGr
 class AExponentialHeightFog
 	: public AInfo
 {
-	GENERATED_BODY()
-public:
-	ENGINE_API AExponentialHeightFog(const FObjectInitializer& ObjectInitializer = FObjectInitializer::Get());
+	GENERATED_UCLASS_BODY()
 
 private_subobject:
 	/** @todo document */
diff --git a/Engine/Source/Runtime/Engine/Classes/Engine/Font.h b/Engine/Source/Runtime/Engine/Classes/Engine/Font.h
--- a/Engine/Source/Runtime/Engine/Classes/Engine/Font.h
+++ b/Engine/Source/Runtime/Engine/Classes/Engine/Font.h
@@ -1,4 +1,4 @@
-// Copyright 1998-2015 Epic Games, Inc. All Rights Reserved.
+﻿// Copyright 1998-2015 Epic Games, Inc. All Rights Reserved.
 
 #pragma once
 
@@ -26,7 +26,7 @@ enum class EFontCacheType
 USTRUCT()
 struct FFontCharacter
 {
-	GENERATED_BODY()
+	GENERATED_USTRUCT_BODY()
 
 	UPROPERTY(EditAnywhere, Category=FontCharacter)
 	int32 StartU;
@@ -97,9 +97,7 @@ struct TStructOpsTypeTraits<FFontCharacter> : public TStructOpsTypeTraitsBase
 UCLASS(hidecategories=Object, autoexpandcategories=Font, MinimalAPI, BlueprintType)
 class UFont : public UObject, public IFontProviderInterface
 {
-	GENERATED_BODY()
-public:
-	ENGINE_API UFont(const FObjectInitializer& ObjectInitializer = FObjectInitializer::Get());
+	GENERATED_UCLASS_BODY()
 
 	~UFont();
 
diff --git a/Engine/Source/Runtime/Engine/Classes/Engine/FontImportOptions.h b/Engine/Source/Runtime/Engine/Classes/Engine/FontImportOptions.h
--- a/Engine/Source/Runtime/Engine/Classes/Engine/FontImportOptions.h
+++ b/Engine/Source/Runtime/Engine/Classes/Engine/FontImportOptions.h
@@ -20,7 +20,7 @@ enum EFontImportCharacterSet
 USTRUCT()
 struct FFontImportOptionsData
 {
-	GENERATED_BODY()
+	GENERATED_USTRUCT_BODY()
 
 	/** Name of the typeface for the font to import */
 	UPROPERTY(EditAnywhere, Category=FontImportOptionsData)
@@ -179,9 +179,7 @@ UCLASS(hidecategories=Object, transient, MinimalAPI)
 class UFontImportOptions
 	: public UObject
 {
-	GENERATED_BODY()
-public:
-	ENGINE_API UFontImportOptions(const FObjectInitializer& ObjectInitializer = FObjectInitializer::Get());
+	GENERATED_UCLASS_BODY()
 
 	/** The actual data for this object.  We wrap it in a struct so that we can copy it around between objects. */
 	UPROPERTY(EditAnywhere, Category=FontImportOptions, meta=(FullyExpand = "true"))
diff --git a/Engine/Source/Runtime/Engine/Classes/Engine/GameEngine.h b/Engine/Source/Runtime/Engine/Classes/Engine/GameEngine.h
--- a/Engine/Source/Runtime/Engine/Classes/Engine/GameEngine.h
+++ b/Engine/Source/Runtime/Engine/Classes/Engine/GameEngine.h
@@ -13,9 +13,7 @@ UCLASS(config=Engine, transient)
 class ENGINE_API UGameEngine
 	: public UEngine
 {
-	GENERATED_BODY()
-public:
-	UGameEngine(const FObjectInitializer& ObjectInitializer = FObjectInitializer::Get());
+	GENERATED_UCLASS_BODY()
 
 	/** Check to see if we need to start a movie capture (used on the first tick when we want to record a Matinee) */
 	UPROPERTY(transient)
diff --git a/Engine/Source/Runtime/Engine/Classes/Engine/GameInstance.h b/Engine/Source/Runtime/Engine/Classes/Engine/GameInstance.h
--- a/Engine/Source/Runtime/Engine/Classes/Engine/GameInstance.h
+++ b/Engine/Source/Runtime/Engine/Classes/Engine/GameInstance.h
@@ -33,9 +33,7 @@ class FOnlineSessionSearchResult;
 UCLASS(config=Game, transient, BlueprintType, Blueprintable)
 class ENGINE_API UGameInstance : public UObject, public FExec
 {
-	GENERATED_BODY()
-public:
-	UGameInstance(const FObjectInitializer& ObjectInitializer = FObjectInitializer::Get());
+	GENERATED_UCLASS_BODY()
 
 protected:
 	struct FWorldContext* WorldContext;
@@ -159,4 +157,4 @@ public:
 	virtual void			HandleDemoPlaybackFailure( EDemoPlayFailure::Type FailureType, const FString& ErrorString = TEXT("") ) { }
 
 	static void				AddReferencedObjects(UObject* InThis, FReferenceCollector& Collector);
-};
+};
\ No newline at end of file
diff --git a/Engine/Source/Runtime/Engine/Classes/Engine/GeneratedMeshAreaLight.h b/Engine/Source/Runtime/Engine/Classes/Engine/GeneratedMeshAreaLight.h
--- a/Engine/Source/Runtime/Engine/Classes/Engine/GeneratedMeshAreaLight.h
+++ b/Engine/Source/Runtime/Engine/Classes/Engine/GeneratedMeshAreaLight.h
@@ -12,7 +12,5 @@ UCLASS(notplaceable, MinimalAPI)
 class AGeneratedMeshAreaLight
 	: public ASpotLight
 {
-	GENERATED_BODY()
-public:
-	ENGINE_API AGeneratedMeshAreaLight(const FObjectInitializer& ObjectInitializer = FObjectInitializer::Get());
+	GENERATED_UCLASS_BODY()
 };
diff --git a/Engine/Source/Runtime/Engine/Classes/Engine/InGameAdManager.h b/Engine/Source/Runtime/Engine/Classes/Engine/InGameAdManager.h
--- a/Engine/Source/Runtime/Engine/Classes/Engine/InGameAdManager.h
+++ b/Engine/Source/Runtime/Engine/Classes/Engine/InGameAdManager.h
@@ -34,9 +34,7 @@ DECLARE_DYNAMIC_DELEGATE( FOnUserClosedAdvertisement );
 UCLASS()
 class UInGameAdManager : public UPlatformInterfaceBase
 {
-	GENERATED_BODY()
-public:
-	UInGameAdManager(const FObjectInitializer& ObjectInitializer = FObjectInitializer::Get());
+	GENERATED_UCLASS_BODY()
 
 	/** If true, the game will pause when the user clicks on the ad, which could take over the screen */
 	UPROPERTY()
diff --git a/Engine/Source/Runtime/Engine/Classes/Engine/InputActionDelegateBinding.h b/Engine/Source/Runtime/Engine/Classes/Engine/InputActionDelegateBinding.h
--- a/Engine/Source/Runtime/Engine/Classes/Engine/InputActionDelegateBinding.h
+++ b/Engine/Source/Runtime/Engine/Classes/Engine/InputActionDelegateBinding.h
@@ -7,7 +7,7 @@
 USTRUCT()
 struct ENGINE_API FBlueprintInputActionDelegateBinding : public FBlueprintInputDelegateBinding
 {
-	GENERATED_BODY()
+	GENERATED_USTRUCT_BODY()
 
 	UPROPERTY()
 	FName InputActionName;
@@ -30,9 +30,7 @@ struct ENGINE_API FBlueprintInputActionDelegateBinding : public FBlueprintInputD
 UCLASS()
 class ENGINE_API UInputActionDelegateBinding : public UInputDelegateBinding
 {
-	GENERATED_BODY()
-public:
-	UInputActionDelegateBinding(const FObjectInitializer& ObjectInitializer = FObjectInitializer::Get());
+	GENERATED_UCLASS_BODY()
 
 	UPROPERTY()
 	TArray<FBlueprintInputActionDelegateBinding> InputActionDelegateBindings;
diff --git a/Engine/Source/Runtime/Engine/Classes/Engine/InputAxisKeyDelegateBinding.h b/Engine/Source/Runtime/Engine/Classes/Engine/InputAxisKeyDelegateBinding.h
--- a/Engine/Source/Runtime/Engine/Classes/Engine/InputAxisKeyDelegateBinding.h
+++ b/Engine/Source/Runtime/Engine/Classes/Engine/InputAxisKeyDelegateBinding.h
@@ -7,7 +7,7 @@
 USTRUCT()
 struct ENGINE_API FBlueprintInputAxisKeyDelegateBinding : public FBlueprintInputDelegateBinding
 {
-	GENERATED_BODY()
+	GENERATED_USTRUCT_BODY()
 
 	UPROPERTY()
 	FKey AxisKey;
@@ -25,9 +25,7 @@ struct ENGINE_API FBlueprintInputAxisKeyDelegateBinding : public FBlueprintInput
 UCLASS()
 class ENGINE_API UInputAxisKeyDelegateBinding : public UInputDelegateBinding
 {
-	GENERATED_BODY()
-public:
-	UInputAxisKeyDelegateBinding(const FObjectInitializer& ObjectInitializer = FObjectInitializer::Get());
+	GENERATED_UCLASS_BODY()
 
 	UPROPERTY()
 	TArray<FBlueprintInputAxisKeyDelegateBinding> InputAxisKeyDelegateBindings;
diff --git a/Engine/Source/Runtime/Engine/Classes/Engine/InputDelegateBinding.h b/Engine/Source/Runtime/Engine/Classes/Engine/InputDelegateBinding.h
--- a/Engine/Source/Runtime/Engine/Classes/Engine/InputDelegateBinding.h
+++ b/Engine/Source/Runtime/Engine/Classes/Engine/InputDelegateBinding.h
@@ -7,7 +7,7 @@
 USTRUCT()
 struct ENGINE_API FBlueprintInputDelegateBinding
 {
-	GENERATED_BODY()
+	GENERATED_USTRUCT_BODY()
 
 	UPROPERTY()
 	uint32 bConsumeInput:1;
@@ -29,9 +29,7 @@ struct ENGINE_API FBlueprintInputDelegateBinding
 UCLASS(abstract)
 class ENGINE_API UInputDelegateBinding : public UDynamicBlueprintBinding
 {
-	GENERATED_BODY()
-public:
-	UInputDelegateBinding(const FObjectInitializer& ObjectInitializer = FObjectInitializer::Get());
+	GENERATED_UCLASS_BODY()
 
 	virtual void BindToInputComponent(UInputComponent* InputComponent) const { };
 
diff --git a/Engine/Source/Runtime/Engine/Classes/Engine/InputKeyDelegateBinding.h b/Engine/Source/Runtime/Engine/Classes/Engine/InputKeyDelegateBinding.h
--- a/Engine/Source/Runtime/Engine/Classes/Engine/InputKeyDelegateBinding.h
+++ b/Engine/Source/Runtime/Engine/Classes/Engine/InputKeyDelegateBinding.h
@@ -8,7 +8,7 @@
 USTRUCT()
 struct ENGINE_API FBlueprintInputKeyDelegateBinding : public FBlueprintInputDelegateBinding
 {
-	GENERATED_BODY()
+	GENERATED_USTRUCT_BODY()
 
 	UPROPERTY()
 	FInputChord InputChord;
@@ -30,9 +30,7 @@ struct ENGINE_API FBlueprintInputKeyDelegateBinding : public FBlueprintInputDele
 UCLASS()
 class ENGINE_API UInputKeyDelegateBinding : public UInputDelegateBinding
 {
-	GENERATED_BODY()
-public:
-	UInputKeyDelegateBinding(const FObjectInitializer& ObjectInitializer = FObjectInitializer::Get());
+	GENERATED_UCLASS_BODY()
 
 	UPROPERTY()
 	TArray<FBlueprintInputKeyDelegateBinding> InputKeyDelegateBindings;
diff --git a/Engine/Source/Runtime/Engine/Classes/Engine/InputTouchDelegateBinding.h b/Engine/Source/Runtime/Engine/Classes/Engine/InputTouchDelegateBinding.h
--- a/Engine/Source/Runtime/Engine/Classes/Engine/InputTouchDelegateBinding.h
+++ b/Engine/Source/Runtime/Engine/Classes/Engine/InputTouchDelegateBinding.h
@@ -7,7 +7,7 @@
 USTRUCT()
 struct ENGINE_API FBlueprintInputTouchDelegateBinding : public FBlueprintInputDelegateBinding
 {
-	GENERATED_BODY()
+	GENERATED_USTRUCT_BODY()
 
 	UPROPERTY()
 	TEnumAsByte<EInputEvent> InputKeyEvent;
@@ -27,9 +27,7 @@ struct ENGINE_API FBlueprintInputTouchDelegateBinding : public FBlueprintInputDe
 UCLASS()
 class ENGINE_API UInputTouchDelegateBinding : public UInputDelegateBinding
 {
-	GENERATED_BODY()
-public:
-	UInputTouchDelegateBinding(const FObjectInitializer& ObjectInitializer = FObjectInitializer::Get());
+	GENERATED_UCLASS_BODY()
 
 	UPROPERTY()
 	TArray<FBlueprintInputTouchDelegateBinding> InputTouchDelegateBindings;
diff --git a/Engine/Source/Runtime/Engine/Classes/Engine/InputVectorAxisDelegateBinding.h b/Engine/Source/Runtime/Engine/Classes/Engine/InputVectorAxisDelegateBinding.h
--- a/Engine/Source/Runtime/Engine/Classes/Engine/InputVectorAxisDelegateBinding.h
+++ b/Engine/Source/Runtime/Engine/Classes/Engine/InputVectorAxisDelegateBinding.h
@@ -7,9 +7,7 @@
 UCLASS()
 class ENGINE_API UInputVectorAxisDelegateBinding : public UInputAxisKeyDelegateBinding
 {
-	GENERATED_BODY()
-public:
-	UInputVectorAxisDelegateBinding(const FObjectInitializer& ObjectInitializer = FObjectInitializer::Get());
+	GENERATED_UCLASS_BODY()
 
 	// Begin UInputDelegateBinding interface
 	virtual void BindToInputComponent(UInputComponent* InputComponent) const override;
diff --git a/Engine/Source/Runtime/Engine/Classes/Engine/IntSerialization.h b/Engine/Source/Runtime/Engine/Classes/Engine/IntSerialization.h
--- a/Engine/Source/Runtime/Engine/Classes/Engine/IntSerialization.h
+++ b/Engine/Source/Runtime/Engine/Classes/Engine/IntSerialization.h
@@ -7,9 +7,7 @@
 UCLASS()
 class ENGINE_API UIntSerialization	: public UObject
 {
-	GENERATED_BODY()
-public:
-	UIntSerialization(const FObjectInitializer& ObjectInitializer = FObjectInitializer::Get());
+	GENERATED_UCLASS_BODY()
 
 public:
 //	ENGINE_API UIntSerialization(const FObjectInitializer& ObjectInitializer, uint8 InUnsignedInt8Variable, uint16 InUnsignedInt16Variable, uint32 InUnsignedInt32Variable, uint64 InUnsignedInt64Variable, int8 InSignedInt8Variable, int16 InSignedInt16Variable, int InSignedInt32Variable, int64 InSignedInt64Variable);
@@ -57,4 +55,4 @@ public:
 /*FArchive& operator<< (FArchive& Ar, FIntSerilizationTest& IntSerilizationTest)
 {
 	return Ar << IntSerilizationTest.UnsignedInt16Variable;
-}*/
+}*/
\ No newline at end of file
diff --git a/Engine/Source/Runtime/Engine/Classes/Engine/InterpCurveEdSetup.h b/Engine/Source/Runtime/Engine/Classes/Engine/InterpCurveEdSetup.h
--- a/Engine/Source/Runtime/Engine/Classes/Engine/InterpCurveEdSetup.h
+++ b/Engine/Source/Runtime/Engine/Classes/Engine/InterpCurveEdSetup.h
@@ -9,7 +9,7 @@
 USTRUCT()
 struct FCurveEdEntry
 {
-	GENERATED_BODY()
+	GENERATED_USTRUCT_BODY()
 
 	UPROPERTY()
 	class UObject* CurveObject;
@@ -56,7 +56,7 @@ struct FCurveEdEntry
 USTRUCT()
 struct FCurveEdTab
 {
-	GENERATED_BODY()
+	GENERATED_USTRUCT_BODY()
 
 	UPROPERTY()
 	FString TabName;
@@ -99,9 +99,7 @@ struct FCurveEdTab
 UCLASS(MinimalAPI)
 class UInterpCurveEdSetup : public UObject
 {
-	GENERATED_BODY()
-public:
-	ENGINE_API UInterpCurveEdSetup(const FObjectInitializer& ObjectInitializer = FObjectInitializer::Get());
+	GENERATED_UCLASS_BODY()
 
 	UPROPERTY()
 	TArray<struct FCurveEdTab> Tabs;
diff --git a/Engine/Source/Runtime/Engine/Classes/Engine/LODActor.h b/Engine/Source/Runtime/Engine/Classes/Engine/LODActor.h
--- a/Engine/Source/Runtime/Engine/Classes/Engine/LODActor.h
+++ b/Engine/Source/Runtime/Engine/Classes/Engine/LODActor.h
@@ -13,9 +13,7 @@
 UCLASS()
 class ENGINE_API ALODActor : public AActor
 {
-	GENERATED_BODY()
-public:
-	ALODActor(const FObjectInitializer& ObjectInitializer = FObjectInitializer::Get());
+	GENERATED_UCLASS_BODY()
 
 private_subobject:
 	UPROPERTY(Category=LODActor, VisibleAnywhere)
diff --git a/Engine/Source/Runtime/Engine/Classes/Engine/LevelBounds.h b/Engine/Source/Runtime/Engine/Classes/Engine/LevelBounds.h
--- a/Engine/Source/Runtime/Engine/Classes/Engine/LevelBounds.h
+++ b/Engine/Source/Runtime/Engine/Classes/Engine/LevelBounds.h
@@ -12,9 +12,7 @@
 UCLASS(hidecategories=(Advanced, Collision, Display, Rendering, Physics, Input), showcategories=("Input|MouseInput", "Input|TouchInput"), MinimalAPI)
 class ALevelBounds : public AActor
 {
-	GENERATED_BODY()
-public:
-	ENGINE_API ALevelBounds(const FObjectInitializer& ObjectInitializer = FObjectInitializer::Get());
+	GENERATED_UCLASS_BODY()
 		
 	/** Whether to automatically update actor bounds based on all relevant actors bounds belonging to the same level */
 	UPROPERTY(EditAnywhere, Category=LevelBounds)
diff --git a/Engine/Source/Runtime/Engine/Classes/Engine/LevelScriptActor.h b/Engine/Source/Runtime/Engine/Classes/Engine/LevelScriptActor.h
--- a/Engine/Source/Runtime/Engine/Classes/Engine/LevelScriptActor.h
+++ b/Engine/Source/Runtime/Engine/Classes/Engine/LevelScriptActor.h
@@ -21,9 +21,7 @@
 UCLASS(notplaceable, meta=(KismetHideOverrides = "ReceiveAnyDamage,ReceivePointDamage,ReceiveRadialDamage,ReceiveActorBeginOverlap,ReceiveActorEndOverlap,ReceiveHit,ReceiveDestroyed,ReceiveActorBeginCursorOver,ReceiveActorEndCursorOver,ReceiveActorOnClicked,ReceiveActorOnReleased,ReceiveActorOnInputTouchBegin,ReceiveActorOnInputTouchEnd,ReceiveActorOnInputTouchEnter,ReceiveActorOnInputTouchLeave"), HideCategories=(Collision,Rendering,"Utilities|Transformation"))
 class ENGINE_API ALevelScriptActor : public AActor
 {
-	GENERATED_BODY()
-public:
-	ALevelScriptActor(const FObjectInitializer& ObjectInitializer = FObjectInitializer::Get());
+	GENERATED_UCLASS_BODY()
 
 	// --- Utility Functions ----------------------------
 	
diff --git a/Engine/Source/Runtime/Engine/Classes/Engine/LevelScriptBlueprint.h b/Engine/Source/Runtime/Engine/Classes/Engine/LevelScriptBlueprint.h
--- a/Engine/Source/Runtime/Engine/Classes/Engine/LevelScriptBlueprint.h
+++ b/Engine/Source/Runtime/Engine/Classes/Engine/LevelScriptBlueprint.h
@@ -20,9 +20,7 @@
 UCLASS(MinimalAPI, NotBlueprintType)
 class ULevelScriptBlueprint : public UBlueprint
 {
-	GENERATED_BODY()
-public:
-	ENGINE_API ULevelScriptBlueprint(const FObjectInitializer& ObjectInitializer = FObjectInitializer::Get());
+	GENERATED_UCLASS_BODY()
 
 #if WITH_EDITORONLY_DATA
 	/** The friendly name to use for UI */
diff --git a/Engine/Source/Runtime/Engine/Classes/Engine/LevelStreaming.h b/Engine/Source/Runtime/Engine/Classes/Engine/LevelStreaming.h
--- a/Engine/Source/Runtime/Engine/Classes/Engine/LevelStreaming.h
+++ b/Engine/Source/Runtime/Engine/Classes/Engine/LevelStreaming.h
@@ -69,9 +69,7 @@ DECLARE_DYNAMIC_MULTICAST_DELEGATE( FLevelStreamingVisibilityStatus );
 UCLASS(abstract, editinlinenew,MinimalAPI)
 class ULevelStreaming : public UObject
 {
-	GENERATED_BODY()
-public:
-	ENGINE_API ULevelStreaming(const FObjectInitializer& ObjectInitializer = FObjectInitializer::Get());
+	GENERATED_UCLASS_BODY()
 
 	/** Deprecated name of the package containing the level to load. Use WorldAsset or GetWorldAssetPackageFName instead.		*/
 	UPROPERTY()
diff --git a/Engine/Source/Runtime/Engine/Classes/Engine/LevelStreamingKismet.h b/Engine/Source/Runtime/Engine/Classes/Engine/LevelStreamingKismet.h
--- a/Engine/Source/Runtime/Engine/Classes/Engine/LevelStreamingKismet.h
+++ b/Engine/Source/Runtime/Engine/Classes/Engine/LevelStreamingKismet.h
@@ -14,9 +14,7 @@
 UCLASS(MinimalAPI, BlueprintType)
 class ULevelStreamingKismet : public ULevelStreaming
 {
-	GENERATED_BODY()
-public:
-	ENGINE_API ULevelStreamingKismet(const FObjectInitializer& ObjectInitializer = FObjectInitializer::Get());
+	GENERATED_UCLASS_BODY()
 
 	/** Whether the level should be loaded at startup																			*/
 	UPROPERTY(Category=LevelStreaming, EditAnywhere)
diff --git a/Engine/Source/Runtime/Engine/Classes/Engine/LevelStreamingPersistent.h b/Engine/Source/Runtime/Engine/Classes/Engine/LevelStreamingPersistent.h
--- a/Engine/Source/Runtime/Engine/Classes/Engine/LevelStreamingPersistent.h
+++ b/Engine/Source/Runtime/Engine/Classes/Engine/LevelStreamingPersistent.h
@@ -13,9 +13,7 @@
 UCLASS(transient)
 class ULevelStreamingPersistent : public ULevelStreaming
 {
-	GENERATED_BODY()
-public:
-	ULevelStreamingPersistent(const FObjectInitializer& ObjectInitializer = FObjectInitializer::Get());
+	GENERATED_UCLASS_BODY()
 
 
 	// Begin ULevelStreaming Interface
diff --git a/Engine/Source/Runtime/Engine/Classes/Engine/LevelStreamingVolume.h b/Engine/Source/Runtime/Engine/Classes/Engine/LevelStreamingVolume.h
--- a/Engine/Source/Runtime/Engine/Classes/Engine/LevelStreamingVolume.h
+++ b/Engine/Source/Runtime/Engine/Classes/Engine/LevelStreamingVolume.h
@@ -25,9 +25,7 @@ enum EStreamingVolumeUsage
 UCLASS(hidecategories=(Advanced, Attachment, Collision, Volume), MinimalAPI)
 class ALevelStreamingVolume : public AVolume
 {
-	GENERATED_BODY()
-public:
-	ENGINE_API ALevelStreamingVolume(const FObjectInitializer& ObjectInitializer = FObjectInitializer::Get());
+	GENERATED_UCLASS_BODY()
 
 	/** Levels names affected by this level streaming volume. */
 	UPROPERTY(Category=LevelStreamingVolume, VisibleAnywhere, BlueprintReadOnly, meta=(DisplayName = "Streaming Levels"))
diff --git a/Engine/Source/Runtime/Engine/Classes/Engine/Light.h b/Engine/Source/Runtime/Engine/Classes/Engine/Light.h
--- a/Engine/Source/Runtime/Engine/Classes/Engine/Light.h
+++ b/Engine/Source/Runtime/Engine/Classes/Engine/Light.h
@@ -8,9 +8,7 @@
 UCLASS(Abstract, ClassGroup=Lights, hideCategories=(Input,Collision,Replication), showCategories=("Input|MouseInput", "Input|TouchInput"), ComponentWrapperClass, ConversionRoot, meta=(ChildCanTick))
 class ENGINE_API ALight : public AActor
 {
-	GENERATED_BODY()
-public:
-	ALight(const FObjectInitializer& ObjectInitializer = FObjectInitializer::Get());
+	GENERATED_UCLASS_BODY()
 
 private_subobject:
 	/** @todo document */
diff --git a/Engine/Source/Runtime/Engine/Classes/Engine/LightMapTexture2D.h b/Engine/Source/Runtime/Engine/Classes/Engine/LightMapTexture2D.h
--- a/Engine/Source/Runtime/Engine/Classes/Engine/LightMapTexture2D.h
+++ b/Engine/Source/Runtime/Engine/Classes/Engine/LightMapTexture2D.h
@@ -21,9 +21,7 @@ enum ELightMapFlags
 UCLASS(MinimalAPI)
 class ULightMapTexture2D : public UTexture2D
 {
-	GENERATED_BODY()
-public:
-	ENGINE_API ULightMapTexture2D(const FObjectInitializer& ObjectInitializer = FObjectInitializer::Get());
+	GENERATED_UCLASS_BODY()
 
 	// Begin UObject interface.
 	virtual void Serialize( FArchive& Ar ) override;
diff --git a/Engine/Source/Runtime/Engine/Classes/Engine/LocalPlayer.h b/Engine/Source/Runtime/Engine/Classes/Engine/LocalPlayer.h
--- a/Engine/Source/Runtime/Engine/Classes/Engine/LocalPlayer.h
+++ b/Engine/Source/Runtime/Engine/Classes/Engine/LocalPlayer.h
@@ -138,9 +138,7 @@ private:
 UCLASS(Within=Engine, config=Engine, transient)
 class ENGINE_API ULocalPlayer : public UPlayer
 {
-	GENERATED_BODY()
-public:
-	ULocalPlayer(const FObjectInitializer& ObjectInitializer = FObjectInitializer::Get());
+	GENERATED_UCLASS_BODY()
 
 #if WITH_HOT_RELOAD_CTORS
 	/** DO NOT USE. This constructor is for internal usage only for hot-reload purposes. */
diff --git a/Engine/Source/Runtime/Engine/Classes/Engine/MicroTransactionBase.h b/Engine/Source/Runtime/Engine/Classes/Engine/MicroTransactionBase.h
--- a/Engine/Source/Runtime/Engine/Classes/Engine/MicroTransactionBase.h
+++ b/Engine/Source/Runtime/Engine/Classes/Engine/MicroTransactionBase.h
@@ -49,7 +49,7 @@ enum EMicroTransactionResult
 USTRUCT()
 struct FPurchaseInfo
 {
-	GENERATED_BODY()
+	GENERATED_USTRUCT_BODY()
 
 // 	/** What kind of microtransaction purchase is this? */
 // 	var EPurchaseType Type;
@@ -75,9 +75,7 @@ struct FPurchaseInfo
 UCLASS()
 class UMicroTransactionBase : public UPlatformInterfaceBase
 {
-	GENERATED_BODY()
-public:
-	UMicroTransactionBase(const FObjectInitializer& ObjectInitializer = FObjectInitializer::Get());
+	GENERATED_UCLASS_BODY()
 
 	/** The list of products available to purchase, filled out by the time a MTD_PurchaseQueryComplete is fired */
 	UPROPERTY()
diff --git a/Engine/Source/Runtime/Engine/Classes/Engine/NavigationObjectBase.h b/Engine/Source/Runtime/Engine/Classes/Engine/NavigationObjectBase.h
--- a/Engine/Source/Runtime/Engine/Classes/Engine/NavigationObjectBase.h
+++ b/Engine/Source/Runtime/Engine/Classes/Engine/NavigationObjectBase.h
@@ -8,9 +8,7 @@
 UCLASS(hidecategories=(Lighting, LightColor, Force), ClassGroup=Navigation, NotBlueprintable, abstract)
 class ENGINE_API ANavigationObjectBase : public AActor, public INavAgentInterface
 {
-	GENERATED_BODY()
-public:
-	ANavigationObjectBase(const FObjectInitializer& ObjectInitializer = FObjectInitializer::Get());
+	GENERATED_UCLASS_BODY()
 
 private_subobject:
 	DEPRECATED_FORGAME(4.6, "CapsuleComponent should not be accessed directly, please use GetCapsuleComponent() function instead. CapsuleComponent will soon be private and your code will not compile.")
diff --git a/Engine/Source/Runtime/Engine/Classes/Engine/NetConnection.h b/Engine/Source/Runtime/Engine/Classes/Engine/NetConnection.h
--- a/Engine/Source/Runtime/Engine/Classes/Engine/NetConnection.h
+++ b/Engine/Source/Runtime/Engine/Classes/Engine/NetConnection.h
@@ -97,11 +97,10 @@ struct DelayedPacket
 #define PING_ACK_DELAY 0.5
 
 
-UCLASS(Abstract, MinimalAPI, transient, config=Engine)
+UCLASS(customConstructor, Abstract, MinimalAPI, transient, config=Engine)
 class UNetConnection : public UPlayer
 {
-	GENERATED_BODY()
-public:
+	GENERATED_UCLASS_BODY()
 
 	/** child connections for secondary viewports */
 	UPROPERTY(transient)
diff --git a/Engine/Source/Runtime/Engine/Classes/Engine/NetDriver.h b/Engine/Source/Runtime/Engine/Classes/Engine/NetDriver.h
--- a/Engine/Source/Runtime/Engine/Classes/Engine/NetDriver.h
+++ b/Engine/Source/Runtime/Engine/Classes/Engine/NetDriver.h
@@ -92,11 +92,10 @@ struct FActorDestructionInfo
 };
 
 
-UCLASS(Abstract, transient, MinimalAPI, config=Engine)
+UCLASS(Abstract, customConstructor, transient, MinimalAPI, config=Engine)
 class UNetDriver : public UObject, public FExec
 {
-	GENERATED_BODY()
-public:
+	GENERATED_UCLASS_BODY()
 
 protected:
 
diff --git a/Engine/Source/Runtime/Engine/Classes/Engine/NetSerialization.h b/Engine/Source/Runtime/Engine/Classes/Engine/NetSerialization.h
--- a/Engine/Source/Runtime/Engine/Classes/Engine/NetSerialization.h
+++ b/Engine/Source/Runtime/Engine/Classes/Engine/NetSerialization.h
@@ -26,7 +26,7 @@
 USTRUCT()
 struct FExampleStruct
 {
-	GENERATED_BODY()
+	GENERATED_USTRUCT_BODY()
 	
 	/**
 	 * @param Ar			FArchive to read or write from.
@@ -90,7 +90,7 @@ struct TStructOpsTypeTraits< FExampleStruct > : public TStructOpsTypeTraitsBase
 USTRUCT()
 struct FExampleItemEntry : public FFastArraySerializerItem
 {
-	GENERATED_BODY()
+	GENERATED_USTRUCT_BODY()
 
 	// Your data:
 	UPROPERTY()
@@ -116,7 +116,7 @@ struct FExampleItemEntry : public FFastArraySerializerItem
 USTRUCT()
 struct FExampleArray: public FFastArraySerializer
 {
-	GENERATED_BODY()
+	GENERATED_USTRUCT_BODY()
 
 	UPROPERTY()
 	TArray<FExampleItemEntry>	Items;	/** Step 3: You MUST have a TArray named Items of the struct you made in step 1. */
@@ -274,7 +274,7 @@ public:
 USTRUCT()
 struct FFastArraySerializerItem
 {
-	GENERATED_BODY()
+	GENERATED_USTRUCT_BODY()
 
 	FFastArraySerializerItem()
 		: ReplicationID(INDEX_NONE), ReplicationKey(INDEX_NONE)
@@ -341,7 +341,7 @@ struct FFastArraySerializerUnmappedItem
 USTRUCT()
 struct FFastArraySerializer
 {
-	GENERATED_BODY()
+	GENERATED_USTRUCT_BODY()
 
 	FFastArraySerializer() : IDCounter(0), ArrayReplicationKey(0) { }
 
@@ -1060,7 +1060,7 @@ bool SerializeFixedVector(FVector &Vector, FArchive& Ar)
 USTRUCT()
 struct FVector_NetQuantize : public FVector
 {
-	GENERATED_BODY()
+	GENERATED_USTRUCT_BODY()
 
 	FORCEINLINE FVector_NetQuantize()
 	{}
@@ -1105,7 +1105,7 @@ struct TStructOpsTypeTraits< FVector_NetQuantize > : public TStructOpsTypeTraits
 USTRUCT()
 struct FVector_NetQuantize10 : public FVector
 {
-	GENERATED_BODY()
+	GENERATED_USTRUCT_BODY()
 
 	FORCEINLINE FVector_NetQuantize10()
 	{}
@@ -1150,7 +1150,7 @@ struct TStructOpsTypeTraits< FVector_NetQuantize10 > : public TStructOpsTypeTrai
 USTRUCT()
 struct FVector_NetQuantize100 : public FVector
 {
-	GENERATED_BODY()
+	GENERATED_USTRUCT_BODY()
 
 	FORCEINLINE FVector_NetQuantize100()
 	{}
@@ -1193,7 +1193,7 @@ struct TStructOpsTypeTraits< FVector_NetQuantize100 > : public TStructOpsTypeTra
 USTRUCT()
 struct FVector_NetQuantizeNormal : public FVector
 {
-	GENERATED_BODY()
+	GENERATED_USTRUCT_BODY()
 
 	FORCEINLINE FVector_NetQuantizeNormal()
 	{}
diff --git a/Engine/Source/Runtime/Engine/Classes/Engine/NetworkSettings.h b/Engine/Source/Runtime/Engine/Classes/Engine/NetworkSettings.h
--- a/Engine/Source/Runtime/Engine/Classes/Engine/NetworkSettings.h
+++ b/Engine/Source/Runtime/Engine/Classes/Engine/NetworkSettings.h
@@ -12,9 +12,7 @@ UCLASS(config=Engine, defaultconfig)
 class ENGINE_API UNetworkSettings
 	: public UObject
 {
-	GENERATED_BODY()
-public:
-	UNetworkSettings(const FObjectInitializer& ObjectInitializer = FObjectInitializer::Get());
+	GENERATED_UCLASS_BODY()
 
 	UPROPERTY(config, EditAnywhere, Category=libcurl, meta=(
 		ConsoleVariable="n.VerifyPeer",DisplayName="Verify Peer",
diff --git a/Engine/Source/Runtime/Engine/Classes/Engine/NiagaraActor.h b/Engine/Source/Runtime/Engine/Classes/Engine/NiagaraActor.h
--- a/Engine/Source/Runtime/Engine/Classes/Engine/NiagaraActor.h
+++ b/Engine/Source/Runtime/Engine/Classes/Engine/NiagaraActor.h
@@ -6,9 +6,7 @@
 UCLASS()
 class ENGINE_API ANiagaraActor : public AActor
 {
-	GENERATED_BODY()
-public:
-	ANiagaraActor(const FObjectInitializer& ObjectInitializer = FObjectInitializer::Get());
+	GENERATED_UCLASS_BODY()
 
 private_subobject:
 	/** Pointer to effect component */
diff --git a/Engine/Source/Runtime/Engine/Classes/Engine/NiagaraEffect.h b/Engine/Source/Runtime/Engine/Classes/Engine/NiagaraEffect.h
--- a/Engine/Source/Runtime/Engine/Classes/Engine/NiagaraEffect.h
+++ b/Engine/Source/Runtime/Engine/Classes/Engine/NiagaraEffect.h
@@ -9,9 +9,7 @@
 UCLASS(MinimalAPI)
 class UNiagaraEffect : public UObject
 {
-	GENERATED_BODY()
-public:
-	ENGINE_API UNiagaraEffect(const FObjectInitializer& ObjectInitializer = FObjectInitializer::Get());
+	GENERATED_UCLASS_BODY()
 
 public:
 	ENGINE_API FNiagaraEmitterProperties *AddEmitterProperties();
@@ -172,4 +170,4 @@ private:
 	TArray< TSharedPtr<FNiagaraSimulation> > Emitters;
 
 
-};
+};
\ No newline at end of file
diff --git a/Engine/Source/Runtime/Engine/Classes/Engine/NiagaraEffectRendererProperties.h b/Engine/Source/Runtime/Engine/Classes/Engine/NiagaraEffectRendererProperties.h
--- a/Engine/Source/Runtime/Engine/Classes/Engine/NiagaraEffectRendererProperties.h
+++ b/Engine/Source/Runtime/Engine/Classes/Engine/NiagaraEffectRendererProperties.h
@@ -12,9 +12,7 @@
 UCLASS()
 class UNiagaraEffectRendererProperties : public UObject
 {
-	GENERATED_BODY()
-public:
-	UNiagaraEffectRendererProperties(const FObjectInitializer& ObjectInitializer);
+	GENERATED_UCLASS_BODY()
 	UNiagaraEffectRendererProperties()
 	{}
 
diff --git a/Engine/Source/Runtime/Engine/Classes/Engine/NiagaraRibbonRendererProperties.h b/Engine/Source/Runtime/Engine/Classes/Engine/NiagaraRibbonRendererProperties.h
--- a/Engine/Source/Runtime/Engine/Classes/Engine/NiagaraRibbonRendererProperties.h
+++ b/Engine/Source/Runtime/Engine/Classes/Engine/NiagaraRibbonRendererProperties.h
@@ -8,10 +8,8 @@
 UCLASS()
 class UNiagaraRibbonRendererProperties : public UNiagaraEffectRendererProperties
 {
-	GENERATED_BODY()
-public:
-	UNiagaraRibbonRendererProperties(const FObjectInitializer& ObjectInitializer);
 public:
+	GENERATED_UCLASS_BODY()
 
 	UNiagaraRibbonRendererProperties()
 	{
diff --git a/Engine/Source/Runtime/Engine/Classes/Engine/NiagaraScript.h b/Engine/Source/Runtime/Engine/Classes/Engine/NiagaraScript.h
--- a/Engine/Source/Runtime/Engine/Classes/Engine/NiagaraScript.h
+++ b/Engine/Source/Runtime/Engine/Classes/Engine/NiagaraScript.h
@@ -10,9 +10,7 @@
 UCLASS(MinimalAPI)
 class UNiagaraScript : public UObject
 {
-	GENERATED_BODY()
-public:
-	ENGINE_API UNiagaraScript(const FObjectInitializer& ObjectInitializer = FObjectInitializer::Get());
+	GENERATED_UCLASS_BODY()
 
 	/** Byte code to execute for this system */
 	UPROPERTY()
@@ -31,4 +29,4 @@ public:
 	UPROPERTY()
 	class UNiagaraScriptSourceBase*	Source;
 #endif
-};
+};
\ No newline at end of file
diff --git a/Engine/Source/Runtime/Engine/Classes/Engine/NiagaraScriptSourceBase.h b/Engine/Source/Runtime/Engine/Classes/Engine/NiagaraScriptSourceBase.h
--- a/Engine/Source/Runtime/Engine/Classes/Engine/NiagaraScriptSourceBase.h
+++ b/Engine/Source/Runtime/Engine/Classes/Engine/NiagaraScriptSourceBase.h
@@ -20,11 +20,9 @@ struct EditorExposedVectorCurveConstant
 UCLASS(MinimalAPI)
 class UNiagaraScriptSourceBase : public UObject
 {
-	GENERATED_BODY()
-public:
-	ENGINE_API UNiagaraScriptSourceBase(const FObjectInitializer& ObjectInitializer = FObjectInitializer::Get());
+	GENERATED_UCLASS_BODY()
 
 	TArray<TSharedPtr<EditorExposedVectorConstant> > ExposedVectorConstants;
 	TArray<TSharedPtr<EditorExposedVectorCurveConstant> > ExposedVectorCurveConstants;
 	virtual void Compile() {};
-};
+};
\ No newline at end of file
diff --git a/Engine/Source/Runtime/Engine/Classes/Engine/Note.h b/Engine/Source/Runtime/Engine/Classes/Engine/Note.h
--- a/Engine/Source/Runtime/Engine/Classes/Engine/Note.h
+++ b/Engine/Source/Runtime/Engine/Classes/Engine/Note.h
@@ -11,9 +11,7 @@
 UCLASS(MinimalAPI, hidecategories = (Input), showcategories=("Input|MouseInput", "Input|TouchInput"))
 class ANote : public AActor
 {
-	GENERATED_BODY()
-public:
-	ENGINE_API ANote(const FObjectInitializer& ObjectInitializer = FObjectInitializer::Get());
+	GENERATED_UCLASS_BODY()
 
 #if WITH_EDITORONLY_DATA
 	UPROPERTY(EditAnywhere, BlueprintReadWrite, Category=Note)
diff --git a/Engine/Source/Runtime/Engine/Classes/Engine/ObjectLibrary.h b/Engine/Source/Runtime/Engine/Classes/Engine/ObjectLibrary.h
--- a/Engine/Source/Runtime/Engine/Classes/Engine/ObjectLibrary.h
+++ b/Engine/Source/Runtime/Engine/Classes/Engine/ObjectLibrary.h
@@ -8,9 +8,7 @@
 UCLASS(MinimalAPI)
 class UObjectLibrary : public UObject
 {
-	GENERATED_BODY()
-public:
-	ENGINE_API UObjectLibrary(const FObjectInitializer& ObjectInitializer = FObjectInitializer::Get());
+	GENERATED_UCLASS_BODY()
 
 	/** Class that Objects must be of. If ContainsBlueprints is true, this is the native class that the blueprints are instances of and not UClass  */
 	UPROPERTY(EditAnywhere, Category=ObjectLibrary, meta=(AllowAbstract = ""))
diff --git a/Engine/Source/Runtime/Engine/Classes/Engine/ObjectReferencer.h b/Engine/Source/Runtime/Engine/Classes/Engine/ObjectReferencer.h
--- a/Engine/Source/Runtime/Engine/Classes/Engine/ObjectReferencer.h
+++ b/Engine/Source/Runtime/Engine/Classes/Engine/ObjectReferencer.h
@@ -11,9 +11,7 @@
 UCLASS()
 class UObjectReferencer : public UObject
 {
-	GENERATED_BODY()
-public:
-	UObjectReferencer(const FObjectInitializer& ObjectInitializer = FObjectInitializer::Get());
+	GENERATED_UCLASS_BODY()
 
 	/** Array of objects being referenced. */
 	UPROPERTY(EditAnywhere, Category=ObjectReferencer)
diff --git a/Engine/Source/Runtime/Engine/Classes/Engine/PendingNetGame.h b/Engine/Source/Runtime/Engine/Classes/Engine/PendingNetGame.h
--- a/Engine/Source/Runtime/Engine/Classes/Engine/PendingNetGame.h
+++ b/Engine/Source/Runtime/Engine/Classes/Engine/PendingNetGame.h
@@ -61,13 +61,12 @@ public:
 	virtual void NotifyControlMessage(UNetConnection* Connection, uint8 MessageType, class FInBunch& Bunch) PURE_VIRTUAL(FNetworkNotify::NotifyReceivedText,);
 };
 
-UCLASS(transient)
+UCLASS(customConstructor, transient)
 class UPendingNetGame :
 	public UObject,
 	public FNetworkNotify
 {
-	GENERATED_BODY()
-public:
+	GENERATED_UCLASS_BODY()
 
 	/** 
 	 * Net driver created for contacting the new server
diff --git a/Engine/Source/Runtime/Engine/Classes/Engine/PlaneReflectionCapture.h b/Engine/Source/Runtime/Engine/Classes/Engine/PlaneReflectionCapture.h
--- a/Engine/Source/Runtime/Engine/Classes/Engine/PlaneReflectionCapture.h
+++ b/Engine/Source/Runtime/Engine/Classes/Engine/PlaneReflectionCapture.h
@@ -10,9 +10,7 @@
 UCLASS(abstract, hidecategories=(Collision, Attachment, Actor), MinimalAPI)
 class APlaneReflectionCapture : public AReflectionCapture
 {
-	GENERATED_BODY()
-public:
-	ENGINE_API APlaneReflectionCapture(const FObjectInitializer& ObjectInitializer = FObjectInitializer::Get());
+	GENERATED_UCLASS_BODY()
 
 public:
 
diff --git a/Engine/Source/Runtime/Engine/Classes/Engine/PlatformInterfaceBase.h b/Engine/Source/Runtime/Engine/Classes/Engine/PlatformInterfaceBase.h
--- a/Engine/Source/Runtime/Engine/Classes/Engine/PlatformInterfaceBase.h
+++ b/Engine/Source/Runtime/Engine/Classes/Engine/PlatformInterfaceBase.h
@@ -34,7 +34,7 @@ enum EPlatformInterfaceDataType
 USTRUCT()
 struct FPlatformInterfaceData
 {
-	GENERATED_BODY()
+	GENERATED_USTRUCT_BODY()
 
 	/** An optional tag for this data */
 	UPROPERTY()
@@ -72,7 +72,7 @@ struct FPlatformInterfaceData
 USTRUCT()
 struct FPlatformInterfaceDelegateResult
 {
-	GENERATED_BODY()
+	GENERATED_USTRUCT_BODY()
 
 	/** This is always usable, no matter the type */
 	UPROPERTY()
@@ -100,7 +100,7 @@ DECLARE_DYNAMIC_DELEGATE_OneParam( FPlatformInterfaceDelegate, const struct FPla
 USTRUCT()
 struct FDelegateArray
 {
-	GENERATED_BODY()
+	GENERATED_USTRUCT_BODY()
 
 	UPROPERTY()
 	TArray<FPlatformInterfaceDelegate> Delegates;
@@ -110,9 +110,7 @@ struct FDelegateArray
 UCLASS(transient,MinimalAPI)
 class UPlatformInterfaceBase : public UObject
 {
-	GENERATED_BODY()
-public:
-	ENGINE_API UPlatformInterfaceBase(const FObjectInitializer& ObjectInitializer = FObjectInitializer::Get());
+	GENERATED_UCLASS_BODY()
 
 	/** Array of delegate arrays. Only add and remove via helper functions, and call via the helper delegate call function */
 	UPROPERTY()
diff --git a/Engine/Source/Runtime/Engine/Classes/Engine/PlatformInterfaceWebResponse.h b/Engine/Source/Runtime/Engine/Classes/Engine/PlatformInterfaceWebResponse.h
--- a/Engine/Source/Runtime/Engine/Classes/Engine/PlatformInterfaceWebResponse.h
+++ b/Engine/Source/Runtime/Engine/Classes/Engine/PlatformInterfaceWebResponse.h
@@ -13,9 +13,7 @@
 UCLASS(transient,MinimalAPI)
 class UPlatformInterfaceWebResponse : public UObject
 {
-	GENERATED_BODY()
-public:
-	ENGINE_API UPlatformInterfaceWebResponse(const FObjectInitializer& ObjectInitializer = FObjectInitializer::Get());
+	GENERATED_UCLASS_BODY()
 
 	/** This holds the original requested URL */
 	UPROPERTY()
diff --git a/Engine/Source/Runtime/Engine/Classes/Engine/Player.h b/Engine/Source/Runtime/Engine/Classes/Engine/Player.h
--- a/Engine/Source/Runtime/Engine/Classes/Engine/Player.h
+++ b/Engine/Source/Runtime/Engine/Classes/Engine/Player.h
@@ -10,9 +10,7 @@
 UCLASS(transient, config=Engine)
 class UPlayer : public UObject, public FExec
 {
-	GENERATED_BODY()
-public:
-	UPlayer(const FObjectInitializer& ObjectInitializer = FObjectInitializer::Get());
+	GENERATED_UCLASS_BODY()
 
 	/** The actor this player controls. */
 	UPROPERTY(transient)
diff --git a/Engine/Source/Runtime/Engine/Classes/Engine/PlayerStartPIE.h b/Engine/Source/Runtime/Engine/Classes/Engine/PlayerStartPIE.h
--- a/Engine/Source/Runtime/Engine/Classes/Engine/PlayerStartPIE.h
+++ b/Engine/Source/Runtime/Engine/Classes/Engine/PlayerStartPIE.h
@@ -11,9 +11,7 @@
 UCLASS(notplaceable,MinimalAPI)
 class APlayerStartPIE : public APlayerStart
 {
-	GENERATED_BODY()
-public:
-	ENGINE_API APlayerStartPIE(const FObjectInitializer& ObjectInitializer = FObjectInitializer::Get());
+	GENERATED_UCLASS_BODY()
 
 };
 
diff --git a/Engine/Source/Runtime/Engine/Classes/Engine/PointLight.h b/Engine/Source/Runtime/Engine/Classes/Engine/PointLight.h
--- a/Engine/Source/Runtime/Engine/Classes/Engine/PointLight.h
+++ b/Engine/Source/Runtime/Engine/Classes/Engine/PointLight.h
@@ -7,9 +7,7 @@
 UCLASS(ClassGroup=(Lights, PointLights), ComponentWrapperClass, MinimalAPI, meta=(ChildCanTick))
 class APointLight : public ALight
 {
-	GENERATED_BODY()
-public:
-	ENGINE_API APointLight(const FObjectInitializer& ObjectInitializer = FObjectInitializer::Get());
+	GENERATED_UCLASS_BODY()
 
 	UPROPERTY(BlueprintReadOnly, Category="Light", meta=(ExposeFunctionCategories="PointLight,Rendering|Lighting"))
 	class UPointLightComponent* PointLightComponent;
diff --git a/Engine/Source/Runtime/Engine/Classes/Engine/Polys.h b/Engine/Source/Runtime/Engine/Classes/Engine/Polys.h
--- a/Engine/Source/Runtime/Engine/Classes/Engine/Polys.h
+++ b/Engine/Source/Runtime/Engine/Classes/Engine/Polys.h
@@ -578,11 +578,10 @@ public:
 	}
 };
 
-UCLASS(MinimalAPI)
+UCLASS(customConstructor, MinimalAPI)
 class UPolys : public UObject
 {
-	GENERATED_BODY()
-public:
+	GENERATED_UCLASS_BODY()
 	// Elements.
 	TTransArray<FPoly> Element;
 
diff --git a/Engine/Source/Runtime/Engine/Classes/Engine/PostProcessVolume.h b/Engine/Source/Runtime/Engine/Classes/Engine/PostProcessVolume.h
--- a/Engine/Source/Runtime/Engine/Classes/Engine/PostProcessVolume.h
+++ b/Engine/Source/Runtime/Engine/Classes/Engine/PostProcessVolume.h
@@ -17,9 +17,7 @@
 UCLASS(autoexpandcategories=PostProcessVolume, hidecategories=(Advanced, Collision, Volume, Brush, Attachment), MinimalAPI)
 class APostProcessVolume : public AVolume, public IInterface_PostProcessVolume
 {
-	GENERATED_BODY()
-public:
-	ENGINE_API APostProcessVolume(const FObjectInitializer& ObjectInitializer = FObjectInitializer::Get());
+	GENERATED_UCLASS_BODY()
 
 	/** Post process settings to use for this volume. */
 	UPROPERTY(interp, Category=PostProcessVolume, meta=(ShowOnlyInnerProperties))
diff --git a/Engine/Source/Runtime/Engine/Classes/Engine/ReflectionCapture.h b/Engine/Source/Runtime/Engine/Classes/Engine/ReflectionCapture.h
--- a/Engine/Source/Runtime/Engine/Classes/Engine/ReflectionCapture.h
+++ b/Engine/Source/Runtime/Engine/Classes/Engine/ReflectionCapture.h
@@ -6,9 +6,7 @@
 UCLASS(abstract, hidecategories=(Collision, Attachment, Actor), MinimalAPI)
 class AReflectionCapture : public AActor
 {
-	GENERATED_BODY()
-public:
-	ENGINE_API AReflectionCapture(const FObjectInitializer& ObjectInitializer = FObjectInitializer::Get());
+	GENERATED_UCLASS_BODY()
 
 private_subobject:
 	/** Reflection capture component. */
diff --git a/Engine/Source/Runtime/Engine/Classes/Engine/SCS_Node.h b/Engine/Source/Runtime/Engine/Classes/Engine/SCS_Node.h
--- a/Engine/Source/Runtime/Engine/Classes/Engine/SCS_Node.h
+++ b/Engine/Source/Runtime/Engine/Classes/Engine/SCS_Node.h
@@ -10,9 +10,7 @@ DECLARE_DELEGATE_OneParam( FSCSNodeNameChanged, const FName& );
 UCLASS(MinimalAPI)
 class USCS_Node : public UObject
 {
-	GENERATED_BODY()
-public:
-	ENGINE_API USCS_Node(const FObjectInitializer& ObjectInitializer = FObjectInitializer::Get());
+	GENERATED_UCLASS_BODY()
 
 	/** Template for the component to create */
 	UPROPERTY()
diff --git a/Engine/Source/Runtime/Engine/Classes/Engine/SceneCapture.h b/Engine/Source/Runtime/Engine/Classes/Engine/SceneCapture.h
--- a/Engine/Source/Runtime/Engine/Classes/Engine/SceneCapture.h
+++ b/Engine/Source/Runtime/Engine/Classes/Engine/SceneCapture.h
@@ -11,9 +11,7 @@
 UCLASS(abstract, hidecategories=(Collision, Attachment, Actor), MinimalAPI)
 class ASceneCapture : public AActor
 {
-	GENERATED_BODY()
-public:
-	ENGINE_API ASceneCapture(const FObjectInitializer& ObjectInitializer = FObjectInitializer::Get());
+	GENERATED_UCLASS_BODY()
 
 private_subobject:
 	/** To display the 3d camera in the editor. */
diff --git a/Engine/Source/Runtime/Engine/Classes/Engine/SceneCapture2D.h b/Engine/Source/Runtime/Engine/Classes/Engine/SceneCapture2D.h
--- a/Engine/Source/Runtime/Engine/Classes/Engine/SceneCapture2D.h
+++ b/Engine/Source/Runtime/Engine/Classes/Engine/SceneCapture2D.h
@@ -12,9 +12,7 @@
 UCLASS(hidecategories=(Collision, Material, Attachment, Actor), MinimalAPI)
 class ASceneCapture2D : public ASceneCapture
 {
-	GENERATED_BODY()
-public:
-	ENGINE_API ASceneCapture2D(const FObjectInitializer& ObjectInitializer = FObjectInitializer::Get());
+	GENERATED_UCLASS_BODY()
 
 private_subobject:
 	/** Scene capture component. */
diff --git a/Engine/Source/Runtime/Engine/Classes/Engine/SceneCaptureCube.h b/Engine/Source/Runtime/Engine/Classes/Engine/SceneCaptureCube.h
--- a/Engine/Source/Runtime/Engine/Classes/Engine/SceneCaptureCube.h
+++ b/Engine/Source/Runtime/Engine/Classes/Engine/SceneCaptureCube.h
@@ -12,9 +12,7 @@
 UCLASS(hidecategories = (Collision, Material, Attachment, Actor), MinimalAPI)
 class ASceneCaptureCube : public ASceneCapture
 {
-	GENERATED_BODY()
-public:
-	ENGINE_API ASceneCaptureCube(const FObjectInitializer& ObjectInitializer = FObjectInitializer::Get());
+	GENERATED_UCLASS_BODY()
 
 private_subobject:
 	/** Scene capture component. */
diff --git a/Engine/Source/Runtime/Engine/Classes/Engine/ScriptViewportClient.h b/Engine/Source/Runtime/Engine/Classes/Engine/ScriptViewportClient.h
--- a/Engine/Source/Runtime/Engine/Classes/Engine/ScriptViewportClient.h
+++ b/Engine/Source/Runtime/Engine/Classes/Engine/ScriptViewportClient.h
@@ -12,9 +12,7 @@
 UCLASS(transient)
 class UScriptViewportClient : public UObject, public FCommonViewportClient
 {
-	GENERATED_BODY()
-public:
-	UScriptViewportClient(const FObjectInitializer& ObjectInitializer = FObjectInitializer::Get());
+	GENERATED_UCLASS_BODY()
 
 };
 
