Back out changelist 2481333

diff --git a/Engine/Source/Runtime/Engine/Classes/Engine/Selection.h b/Engine/Source/Runtime/Engine/Classes/Engine/Selection.h
--- a/Engine/Source/Runtime/Engine/Classes/Engine/Selection.h
+++ b/Engine/Source/Runtime/Engine/Classes/Engine/Selection.h
@@ -13,11 +13,10 @@
 
 #include "Selection.generated.h"
 
-UCLASS(transient)
+UCLASS(customConstructor, transient)
 class ENGINE_API USelection : public UObject
 {
-	GENERATED_BODY()
-public:
+	GENERATED_UCLASS_BODY()
 private:
 	typedef TArray<TWeakObjectPtr<UObject> >	ObjectArray;
 	typedef TMRUArray<UClass*>	ClassArray;
diff --git a/Engine/Source/Runtime/Engine/Classes/Engine/ShadowMapTexture2D.h b/Engine/Source/Runtime/Engine/Classes/Engine/ShadowMapTexture2D.h
--- a/Engine/Source/Runtime/Engine/Classes/Engine/ShadowMapTexture2D.h
+++ b/Engine/Source/Runtime/Engine/Classes/Engine/ShadowMapTexture2D.h
@@ -7,9 +7,7 @@
 UCLASS(MinimalAPI)
 class UShadowMapTexture2D : public UTexture2D
 {
-	GENERATED_BODY()
-public:
-	ENGINE_API UShadowMapTexture2D(const FObjectInitializer& ObjectInitializer = FObjectInitializer::Get());
+	GENERATED_UCLASS_BODY()
 	
 	/** Bit-field with shadowmap flags. */
 	UPROPERTY()
diff --git a/Engine/Source/Runtime/Engine/Classes/Engine/SimpleConstructionScript.h b/Engine/Source/Runtime/Engine/Classes/Engine/SimpleConstructionScript.h
--- a/Engine/Source/Runtime/Engine/Classes/Engine/SimpleConstructionScript.h
+++ b/Engine/Source/Runtime/Engine/Classes/Engine/SimpleConstructionScript.h
@@ -9,9 +9,7 @@ class USCS_Node;
 UCLASS(MinimalAPI)
 class USimpleConstructionScript : public UObject
 {
-	GENERATED_BODY()
-public:
-	ENGINE_API USimpleConstructionScript(const FObjectInitializer& ObjectInitializer = FObjectInitializer::Get());
+	GENERATED_UCLASS_BODY()
 
 	// Begin UObject Interface
 	virtual void Serialize(FArchive& Ar) override;
diff --git a/Engine/Source/Runtime/Engine/Classes/Engine/SkeletalMesh.h b/Engine/Source/Runtime/Engine/Classes/Engine/SkeletalMesh.h
--- a/Engine/Source/Runtime/Engine/Classes/Engine/SkeletalMesh.h
+++ b/Engine/Source/Runtime/Engine/Classes/Engine/SkeletalMesh.h
@@ -56,7 +56,7 @@ enum SkeletalMeshOptimizationType
 USTRUCT()
 struct FBoneMirrorInfo
 {
-	GENERATED_BODY()
+	GENERATED_USTRUCT_BODY()
 
 	/** The bone to mirror. */
 	UPROPERTY(EditAnywhere, Category=BoneMirrorInfo, meta=(ArrayClamp = "RefSkeleton"))
@@ -78,7 +78,7 @@ struct FBoneMirrorInfo
 USTRUCT()
 struct FBoneMirrorExport
 {
-	GENERATED_BODY()
+	GENERATED_USTRUCT_BODY()
 
 	UPROPERTY(EditAnywhere, Category=BoneMirrorExport)
 	FName BoneName;
@@ -101,7 +101,7 @@ struct FBoneMirrorExport
 USTRUCT()
 struct FTriangleSortSettings
 {
-	GENERATED_BODY()
+	GENERATED_USTRUCT_BODY()
 
 	UPROPERTY(EditAnywhere, Category=TriangleSortSettings)
 	TEnumAsByte<enum ETriangleSortOption> TriangleSorting;
@@ -125,7 +125,7 @@ struct FTriangleSortSettings
 USTRUCT()
 struct FBoneReference
 {
-	GENERATED_BODY()
+	GENERATED_USTRUCT_BODY()
 
 	/** Name of bone to control. This is the main bone chain to modify from. **/
 	UPROPERTY(EditAnywhere, Category = BoneReference)
@@ -160,7 +160,7 @@ struct FBoneReference
 USTRUCT()
 struct FSkeletalMeshOptimizationSettings
 {
-	GENERATED_BODY()
+	GENERATED_USTRUCT_BODY()
 
 		/** The method to use when optimizing the skeletal mesh LOD */
 		UPROPERTY()
@@ -281,7 +281,7 @@ struct FSkeletalMeshOptimizationSettings
 USTRUCT()
 struct FSkeletalMeshLODInfo
 {
-	GENERATED_BODY()
+	GENERATED_USTRUCT_BODY()
 
 	/**	Indicates when to use this LOD. A smaller number means use this LOD when further away. */
 	UPROPERTY(EditAnywhere, Category=SkeletalMeshLODInfo)
@@ -322,7 +322,7 @@ struct FSkeletalMeshLODInfo
 USTRUCT()
 struct FMorphTargetMap
 {
-	GENERATED_BODY()
+	GENERATED_USTRUCT_BODY()
 
 	/** The bone to mirror. */
 	UPROPERTY(EditAnywhere, Category=MorphTargetMap)
@@ -403,7 +403,7 @@ struct FClothBonePlane
 USTRUCT()
 struct FClothPhysicsProperties
 {
-	GENERATED_BODY()
+	GENERATED_USTRUCT_BODY()
 
 	// Bending stiffness of the cloth in the range [0, 1]. 
 	UPROPERTY(EditAnywhere, Category = Stiffness, meta = (ClampMin = "0.0", ClampMax = "1.0", UIMin = "0.0", UIMax = "1.0"))
@@ -443,7 +443,7 @@ struct FClothPhysicsProperties
 USTRUCT()
 struct FClothingAssetData
 {
-	GENERATED_BODY()
+	GENERATED_USTRUCT_BODY()
 
 	/* User-defined asset name */
 	UPROPERTY(EditAnywhere, Category=ClothingAssetData)
@@ -491,7 +491,7 @@ struct FClothingAssetData
 USTRUCT()
 struct FSkeletalMaterial
 {
-	GENERATED_BODY()
+	GENERATED_USTRUCT_BODY()
 
 	FSkeletalMaterial()
 		: MaterialInterface( NULL )
@@ -531,9 +531,7 @@ class FSkeletalMeshResource;
 UCLASS(hidecategories=Object, MinimalAPI, BlueprintType)
 class USkeletalMesh : public UObject, public IInterface_CollisionDataProvider, public IInterface_AssetUserData
 {
-	GENERATED_BODY()
-public:
-	ENGINE_API USkeletalMesh(const FObjectInitializer& ObjectInitializer = FObjectInitializer::Get());
+	GENERATED_UCLASS_BODY()
 
 private:
 	/** Rendering resources created at import time. */
diff --git a/Engine/Source/Runtime/Engine/Classes/Engine/SkeletalMeshSocket.h b/Engine/Source/Runtime/Engine/Classes/Engine/SkeletalMeshSocket.h
--- a/Engine/Source/Runtime/Engine/Classes/Engine/SkeletalMeshSocket.h
+++ b/Engine/Source/Runtime/Engine/Classes/Engine/SkeletalMeshSocket.h
@@ -7,9 +7,7 @@
 UCLASS(hidecategories=Object, hidecategories=Actor, MinimalAPI)
 class USkeletalMeshSocket : public UObject
 {
-	GENERATED_BODY()
-public:
-	ENGINE_API USkeletalMeshSocket(const FObjectInitializer& ObjectInitializer = FObjectInitializer::Get());
+	GENERATED_UCLASS_BODY()
 
 	/** 
 	 *	Defines a named attachment location on the USkeletalMesh. 
diff --git a/Engine/Source/Runtime/Engine/Classes/Engine/SkyLight.h b/Engine/Source/Runtime/Engine/Classes/Engine/SkyLight.h
--- a/Engine/Source/Runtime/Engine/Classes/Engine/SkyLight.h
+++ b/Engine/Source/Runtime/Engine/Classes/Engine/SkyLight.h
@@ -7,9 +7,7 @@
 UCLASS(ClassGroup=Lights, hidecategories=(Input,Collision,Replication,Info), showcategories=("Rendering", "Input|MouseInput", "Input|TouchInput"), ComponentWrapperClass, ConversionRoot, Blueprintable)
 class ENGINE_API ASkyLight : public AInfo
 {
-	GENERATED_BODY()
-public:
-	ASkyLight(const FObjectInitializer& ObjectInitializer = FObjectInitializer::Get());
+	GENERATED_UCLASS_BODY()
 
 private_subobject:
 	/** @todo document */
diff --git a/Engine/Source/Runtime/Engine/Classes/Engine/SphereReflectionCapture.h b/Engine/Source/Runtime/Engine/Classes/Engine/SphereReflectionCapture.h
--- a/Engine/Source/Runtime/Engine/Classes/Engine/SphereReflectionCapture.h
+++ b/Engine/Source/Runtime/Engine/Classes/Engine/SphereReflectionCapture.h
@@ -13,9 +13,7 @@ class UDrawSphereComponent;
 UCLASS(hidecategories = (Collision, Attachment, Actor), MinimalAPI)
 class ASphereReflectionCapture : public AReflectionCapture
 {
-	GENERATED_BODY()
-public:
-	ENGINE_API ASphereReflectionCapture(const FObjectInitializer& ObjectInitializer = FObjectInitializer::Get());
+	GENERATED_UCLASS_BODY()
 
 private_subobject:
 	/** Sphere component used to visualize the capture radius */
diff --git a/Engine/Source/Runtime/Engine/Classes/Engine/SplineMeshActor.h b/Engine/Source/Runtime/Engine/Classes/Engine/SplineMeshActor.h
--- a/Engine/Source/Runtime/Engine/Classes/Engine/SplineMeshActor.h
+++ b/Engine/Source/Runtime/Engine/Classes/Engine/SplineMeshActor.h
@@ -11,9 +11,7 @@
 UCLASS(hideCategories = (Input), showCategories = ("Input|MouseInput", "Input|TouchInput"), ConversionRoot, ComponentWrapperClass, meta = (ChildCanTick))
 class ENGINE_API ASplineMeshActor : public AActor
 {
-	GENERATED_BODY()
-public:
-	ASplineMeshActor(const FObjectInitializer& ObjectInitializer = FObjectInitializer::Get());
+	GENERATED_UCLASS_BODY()
 
 private_subobject:
 	UPROPERTY(Category = SplineMeshActor, VisibleAnywhere, BlueprintReadOnly, meta = (ExposeFunctionCategories = "Mesh,Rendering,Physics,Components|StaticMesh,Components|SplineMesh", AllowPrivateAccess = "true"))
diff --git a/Engine/Source/Runtime/Engine/Classes/Engine/SpotLight.h b/Engine/Source/Runtime/Engine/Classes/Engine/SpotLight.h
--- a/Engine/Source/Runtime/Engine/Classes/Engine/SpotLight.h
+++ b/Engine/Source/Runtime/Engine/Classes/Engine/SpotLight.h
@@ -8,9 +8,7 @@
 UCLASS(ClassGroup=(Lights, SpotLights), MinimalAPI, ComponentWrapperClass, meta=(ChildCanTick))
 class ASpotLight : public ALight
 {
-	GENERATED_BODY()
-public:
-	ENGINE_API ASpotLight(const FObjectInitializer& ObjectInitializer = FObjectInitializer::Get());
+	GENERATED_UCLASS_BODY()
 
 	UPROPERTY(BlueprintReadOnly, Category="Light", meta=(ExposeFunctionCategories="SpotLight,Rendering|Lighting"))
 	class USpotLightComponent* SpotLightComponent;
diff --git a/Engine/Source/Runtime/Engine/Classes/Engine/StaticMesh.h b/Engine/Source/Runtime/Engine/Classes/Engine/StaticMesh.h
--- a/Engine/Source/Runtime/Engine/Classes/Engine/StaticMesh.h
+++ b/Engine/Source/Runtime/Engine/Classes/Engine/StaticMesh.h
@@ -53,7 +53,7 @@ enum EOptimizationType
 USTRUCT()
 struct FStaticMeshOptimizationSettings
 {
-	GENERATED_BODY()
+	GENERATED_USTRUCT_BODY()
 
 	/** The method to use when optimizing the skeletal mesh LOD */
 	UPROPERTY()
@@ -134,7 +134,7 @@ struct FStaticMeshOptimizationSettings
 USTRUCT()
 struct FStaticMeshSourceModel
 {
-	GENERATED_BODY()
+	GENERATED_USTRUCT_BODY()
 
 #if WITH_EDITOR
 	/** Imported raw mesh data. Optional for all but the first LOD. */
@@ -174,7 +174,7 @@ struct FStaticMeshSourceModel
 USTRUCT()
 struct FMeshSectionInfo
 {
-	GENERATED_BODY()
+	GENERATED_USTRUCT_BODY()
 
 	/** Index in to the Materials array on UStaticMesh. */
 	UPROPERTY()
@@ -215,7 +215,7 @@ bool operator!=(const FMeshSectionInfo& A, const FMeshSectionInfo& B);
 USTRUCT()
 struct FMeshSectionInfoMap
 {
-	GENERATED_BODY()
+	GENERATED_USTRUCT_BODY()
 
 	/** Maps an LOD+Section to the material it should render with. */
 	TMap<uint32,FMeshSectionInfo> Map;
@@ -245,7 +245,7 @@ struct FMeshSectionInfoMap
 USTRUCT()
 struct FAssetEditorOrbitCameraPosition
 {
-	GENERATED_BODY()
+	GENERATED_USTRUCT_BODY()
 
 	FAssetEditorOrbitCameraPosition()
 		: bIsSet(false)
@@ -285,11 +285,10 @@ struct FAssetEditorOrbitCameraPosition
  * @see https://docs.unrealengine.com/latest/INT/Engine/Content/Types/StaticMeshes/
  * @see AStaticMeshActor, UStaticMeshComponent
  */
-UCLASS(collapsecategories, hidecategories=Object, MinimalAPI, BlueprintType, config=Engine)
+UCLASS(collapsecategories, hidecategories=Object, customconstructor, MinimalAPI, BlueprintType, config=Engine)
 class UStaticMesh : public UObject, public IInterface_CollisionDataProvider, public IInterface_AssetUserData
 {
-	GENERATED_BODY()
-public:
+	GENERATED_UCLASS_BODY()
 
 	/** Pointer to the data used to render this static mesh. */
 	TScopedPointer<class FStaticMeshRenderData> RenderData;
diff --git a/Engine/Source/Runtime/Engine/Classes/Engine/StaticMeshActor.h b/Engine/Source/Runtime/Engine/Classes/Engine/StaticMeshActor.h
--- a/Engine/Source/Runtime/Engine/Classes/Engine/StaticMeshActor.h
+++ b/Engine/Source/Runtime/Engine/Classes/Engine/StaticMeshActor.h
@@ -15,9 +15,7 @@
 UCLASS(hidecategories=(Input), showcategories=("Input|MouseInput", "Input|TouchInput"), ConversionRoot, ComponentWrapperClass, meta=(ChildCanTick))
 class ENGINE_API AStaticMeshActor : public AActor
 {
-	GENERATED_BODY()
-public:
-	AStaticMeshActor(const FObjectInitializer& ObjectInitializer = FObjectInitializer::Get());
+	GENERATED_UCLASS_BODY()
 
 private_subobject:
 	DEPRECATED_FORGAME(4.6, "StaticMeshComponent should not be accessed directly, please use GetStaticMeshComponent() function instead. StaticMeshComponent will soon be private and your code will not compile.")
diff --git a/Engine/Source/Runtime/Engine/Classes/Engine/StaticMeshSocket.h b/Engine/Source/Runtime/Engine/Classes/Engine/StaticMeshSocket.h
--- a/Engine/Source/Runtime/Engine/Classes/Engine/StaticMeshSocket.h
+++ b/Engine/Source/Runtime/Engine/Classes/Engine/StaticMeshSocket.h
@@ -7,9 +7,7 @@
 UCLASS(hidecategories=Object, hidecategories=Actor, MinimalAPI)
 class UStaticMeshSocket : public UObject
 {
-	GENERATED_BODY()
-public:
-	ENGINE_API UStaticMeshSocket(const FObjectInitializer& ObjectInitializer = FObjectInitializer::Get());
+	GENERATED_UCLASS_BODY()
 
 	/** 
 	 *	Defines a named attachment location on the UStaticMesh. 
diff --git a/Engine/Source/Runtime/Engine/Classes/Engine/StreamableManager.h b/Engine/Source/Runtime/Engine/Classes/Engine/StreamableManager.h
--- a/Engine/Source/Runtime/Engine/Classes/Engine/StreamableManager.h
+++ b/Engine/Source/Runtime/Engine/Classes/Engine/StreamableManager.h
@@ -9,7 +9,7 @@ DECLARE_DELEGATE( FStreamableDelegate );
 USTRUCT()
 struct ENGINE_API FStreamableManager
 {
-	GENERATED_BODY()
+	GENERATED_USTRUCT_BODY();
 
 	FStreamableManager();
 	~FStreamableManager();
diff --git a/Engine/Source/Runtime/Engine/Classes/Engine/SubsurfaceProfile.h b/Engine/Source/Runtime/Engine/Classes/Engine/SubsurfaceProfile.h
--- a/Engine/Source/Runtime/Engine/Classes/Engine/SubsurfaceProfile.h
+++ b/Engine/Source/Runtime/Engine/Classes/Engine/SubsurfaceProfile.h
@@ -8,7 +8,7 @@
 USTRUCT(BlueprintType)
 struct FSubsurfaceProfileStruct
 {
-	GENERATED_BODY()
+	GENERATED_USTRUCT_BODY()
 	
 	/** in world/unreal units (cm) */
 	UPROPERTY(Category = "SubsurfaceProfileStruct", EditAnywhere, BlueprintReadOnly, meta = (ClampMin = "0.1", UIMax = "50.0", ClampMax = "1000.0"))
@@ -55,9 +55,7 @@ struct FSubsurfaceProfileStruct
 UCLASS(autoexpandcategories = SubsurfaceProfile, MinimalAPI)
 class USubsurfaceProfile : public UObject
 {
-	GENERATED_BODY()
-public:
-	ENGINE_API USubsurfaceProfile(const FObjectInitializer& ObjectInitializer = FObjectInitializer::Get());
+	GENERATED_UCLASS_BODY()
 
 	UPROPERTY(Category = USubsurfaceProfile, EditAnywhere, meta = (ShowOnlyInnerProperties))
 	struct FSubsurfaceProfileStruct Settings;
diff --git a/Engine/Source/Runtime/Engine/Classes/Engine/TargetPoint.h b/Engine/Source/Runtime/Engine/Classes/Engine/TargetPoint.h
--- a/Engine/Source/Runtime/Engine/Classes/Engine/TargetPoint.h
+++ b/Engine/Source/Runtime/Engine/Classes/Engine/TargetPoint.h
@@ -10,9 +10,7 @@
 UCLASS(MinimalAPI)
 class ATargetPoint : public AActor
 {
-	GENERATED_BODY()
-public:
-	ENGINE_API ATargetPoint(const FObjectInitializer& ObjectInitializer = FObjectInitializer::Get());
+	GENERATED_UCLASS_BODY()
 
 #if WITH_EDITORONLY_DATA
 private_subobject:
diff --git a/Engine/Source/Runtime/Engine/Classes/Engine/Texture.h b/Engine/Source/Runtime/Engine/Classes/Engine/Texture.h
--- a/Engine/Source/Runtime/Engine/Classes/Engine/Texture.h
+++ b/Engine/Source/Runtime/Engine/Classes/Engine/Texture.h
@@ -111,7 +111,7 @@ enum ETextureSourceFormat
 USTRUCT()
 struct FTextureSource
 {
-	GENERATED_BODY()
+	GENERATED_USTRUCT_BODY()
 
 	/** Default constructor. */
 	ENGINE_API FTextureSource();
@@ -277,7 +277,7 @@ private:
 USTRUCT()
 struct FTexturePlatformData
 {
-	GENERATED_BODY()
+	GENERATED_USTRUCT_BODY()
 
 	/** Width of the texture. */
 	int32 SizeX;
@@ -341,9 +341,7 @@ struct FTexturePlatformData
 UCLASS(abstract, MinimalAPI, BlueprintType)
 class UTexture : public UObject
 {
-	GENERATED_BODY()
-public:
-	ENGINE_API UTexture(const FObjectInitializer& ObjectInitializer = FObjectInitializer::Get());
+	GENERATED_UCLASS_BODY()
 
 	/*--------------------------------------------------------------------------
 		Editor only properties used to build the runtime texture data.
diff --git a/Engine/Source/Runtime/Engine/Classes/Engine/Texture2D.h b/Engine/Source/Runtime/Engine/Classes/Engine/Texture2D.h
--- a/Engine/Source/Runtime/Engine/Classes/Engine/Texture2D.h
+++ b/Engine/Source/Runtime/Engine/Classes/Engine/Texture2D.h
@@ -8,9 +8,7 @@
 UCLASS(hidecategories=Object, MinimalAPI, BlueprintType)
 class UTexture2D : public UTexture
 {
-	GENERATED_BODY()
-public:
-	ENGINE_API UTexture2D(const FObjectInitializer& ObjectInitializer = FObjectInitializer::Get());
+	GENERATED_UCLASS_BODY()
 
 public:
 	/** Number of miplevels the texture should have resident.					*/
diff --git a/Engine/Source/Runtime/Engine/Classes/Engine/Texture2DDynamic.h b/Engine/Source/Runtime/Engine/Classes/Engine/Texture2DDynamic.h
--- a/Engine/Source/Runtime/Engine/Classes/Engine/Texture2DDynamic.h
+++ b/Engine/Source/Runtime/Engine/Classes/Engine/Texture2DDynamic.h
@@ -7,9 +7,7 @@
 UCLASS(hidecategories=Object, MinimalAPI)
 class UTexture2DDynamic : public UTexture
 {
-	GENERATED_BODY()
-public:
-	ENGINE_API UTexture2DDynamic(const FObjectInitializer& ObjectInitializer = FObjectInitializer::Get());
+	GENERATED_UCLASS_BODY()
 
 	/** The width of the texture. */
 	int32 SizeX;
diff --git a/Engine/Source/Runtime/Engine/Classes/Engine/TextureCube.h b/Engine/Source/Runtime/Engine/Classes/Engine/TextureCube.h
--- a/Engine/Source/Runtime/Engine/Classes/Engine/TextureCube.h
+++ b/Engine/Source/Runtime/Engine/Classes/Engine/TextureCube.h
@@ -7,9 +7,7 @@
 UCLASS(hidecategories=Object, MinimalAPI)
 class UTextureCube : public UTexture
 {
-	GENERATED_BODY()
-public:
-	ENGINE_API UTextureCube(const FObjectInitializer& ObjectInitializer = FObjectInitializer::Get());
+	GENERATED_UCLASS_BODY()
 
 public:
 	/** Platform data. */
diff --git a/Engine/Source/Runtime/Engine/Classes/Engine/TextureLightProfile.h b/Engine/Source/Runtime/Engine/Classes/Engine/TextureLightProfile.h
--- a/Engine/Source/Runtime/Engine/Classes/Engine/TextureLightProfile.h
+++ b/Engine/Source/Runtime/Engine/Classes/Engine/TextureLightProfile.h
@@ -8,9 +8,7 @@
 UCLASS(hidecategories=(Object,CompositeTexture,Texture2D), MinimalAPI, BlueprintType)
 class UTextureLightProfile : public UTexture2D
 {
-	GENERATED_BODY()
-public:
-	ENGINE_API UTextureLightProfile(const FObjectInitializer& ObjectInitializer = FObjectInitializer::Get());
+	GENERATED_UCLASS_BODY()
 
 	/** Light brightness in Lumens, imported from IES profile, <= 0 if the profile is used for masking only. Use with InverseSquareFalloff. */
 	UPROPERTY(EditAnywhere, Category=TextureLightProfile, AssetRegistrySearchable)
diff --git a/Engine/Source/Runtime/Engine/Classes/Engine/TextureRenderTarget.h b/Engine/Source/Runtime/Engine/Classes/Engine/TextureRenderTarget.h
--- a/Engine/Source/Runtime/Engine/Classes/Engine/TextureRenderTarget.h
+++ b/Engine/Source/Runtime/Engine/Classes/Engine/TextureRenderTarget.h
@@ -13,9 +13,7 @@
 UCLASS(abstract, MinimalAPI)
 class UTextureRenderTarget : public UTexture
 {
-	GENERATED_BODY()
-public:
-	ENGINE_API UTextureRenderTarget(const FObjectInitializer& ObjectInitializer = FObjectInitializer::Get());
+	GENERATED_UCLASS_BODY()	
 
 	/** Will override FTextureRenderTarget2DResource::GetDisplayGamma if > 0. */
 	UPROPERTY(EditAnywhere, BlueprintReadWrite, Category=TextureRenderTarget)
diff --git a/Engine/Source/Runtime/Engine/Classes/Engine/TextureRenderTarget2D.h b/Engine/Source/Runtime/Engine/Classes/Engine/TextureRenderTarget2D.h
--- a/Engine/Source/Runtime/Engine/Classes/Engine/TextureRenderTarget2D.h
+++ b/Engine/Source/Runtime/Engine/Classes/Engine/TextureRenderTarget2D.h
@@ -15,9 +15,7 @@
 UCLASS(hidecategories=Object, hidecategories=Texture, MinimalAPI)
 class UTextureRenderTarget2D : public UTextureRenderTarget
 {
-	GENERATED_BODY()
-public:
-	ENGINE_API UTextureRenderTarget2D(const FObjectInitializer& ObjectInitializer = FObjectInitializer::Get());
+	GENERATED_UCLASS_BODY()
 
 	/** The width of the texture.												*/
 	UPROPERTY(EditAnywhere, BlueprintReadOnly, Category=TextureRenderTarget2D, AssetRegistrySearchable)
diff --git a/Engine/Source/Runtime/Engine/Classes/Engine/TextureRenderTargetCube.h b/Engine/Source/Runtime/Engine/Classes/Engine/TextureRenderTargetCube.h
--- a/Engine/Source/Runtime/Engine/Classes/Engine/TextureRenderTargetCube.h
+++ b/Engine/Source/Runtime/Engine/Classes/Engine/TextureRenderTargetCube.h
@@ -14,9 +14,7 @@
 UCLASS(hidecategories=Object, hidecategories=Texture, MinimalAPI)
 class UTextureRenderTargetCube : public UTextureRenderTarget
 {
-	GENERATED_BODY()
-public:
-	ENGINE_API UTextureRenderTargetCube(const FObjectInitializer& ObjectInitializer = FObjectInitializer::Get());
+	GENERATED_UCLASS_BODY()
 
 	/** The width of the texture.												*/
 	UPROPERTY(EditAnywhere, BlueprintReadWrite, Category=TextureRenderTargetCube, AssetRegistrySearchable)
diff --git a/Engine/Source/Runtime/Engine/Classes/Engine/TimelineTemplate.h b/Engine/Source/Runtime/Engine/Classes/Engine/TimelineTemplate.h
--- a/Engine/Source/Runtime/Engine/Classes/Engine/TimelineTemplate.h
+++ b/Engine/Source/Runtime/Engine/Classes/Engine/TimelineTemplate.h
@@ -7,7 +7,7 @@
 USTRUCT()
 struct FTTTrackBase
 {
-	GENERATED_BODY()
+	GENERATED_USTRUCT_BODY()
 
 	/** Name of this track */
 	UPROPERTY()
@@ -30,7 +30,7 @@ struct FTTTrackBase
 USTRUCT()
 struct FTTEventTrack : public FTTTrackBase
 {
-	GENERATED_BODY()
+	GENERATED_USTRUCT_BODY()
 
 	/** Curve object used to store keys */
 	UPROPERTY()
@@ -50,7 +50,7 @@ struct FTTEventTrack : public FTTTrackBase
 USTRUCT()
 struct FTTFloatTrack : public FTTTrackBase
 {
-	GENERATED_BODY()
+	GENERATED_USTRUCT_BODY()
 
 	/** Curve object used to define float value over time */
 	UPROPERTY()
@@ -70,7 +70,7 @@ struct FTTFloatTrack : public FTTTrackBase
 USTRUCT()
 struct FTTVectorTrack : public FTTTrackBase
 {
-	GENERATED_BODY()
+	GENERATED_USTRUCT_BODY()
 
 	/** Curve object used to define vector value over time */
 	UPROPERTY()
@@ -90,7 +90,7 @@ struct FTTVectorTrack : public FTTTrackBase
 USTRUCT()
 struct FTTLinearColorTrack : public FTTTrackBase
 {
-	GENERATED_BODY()
+	GENERATED_USTRUCT_BODY()
 
 	/** Curve object used to define color value over time */
 	UPROPERTY()
@@ -109,9 +109,7 @@ struct FTTLinearColorTrack : public FTTTrackBase
 UCLASS(MinimalAPI)
 class UTimelineTemplate : public UObject
 {
-	GENERATED_BODY()
-public:
-	ENGINE_API UTimelineTemplate(const FObjectInitializer& ObjectInitializer = FObjectInitializer::Get());
+	GENERATED_UCLASS_BODY()
 
 	/** Length of this timeline */
 	UPROPERTY(EditAnywhere, Category=TimelineTemplate)
diff --git a/Engine/Source/Runtime/Engine/Classes/Engine/TriggerBase.h b/Engine/Source/Runtime/Engine/Classes/Engine/TriggerBase.h
--- a/Engine/Source/Runtime/Engine/Classes/Engine/TriggerBase.h
+++ b/Engine/Source/Runtime/Engine/Classes/Engine/TriggerBase.h
@@ -11,9 +11,7 @@ class UBillboardComponent;
 UCLASS(ClassGroup=Common, abstract, ConversionRoot, MinimalAPI)
 class ATriggerBase : public AActor
 {
-	GENERATED_BODY()
-public:
-	ENGINE_API ATriggerBase(const FObjectInitializer& ObjectInitializer = FObjectInitializer::Get());
+	GENERATED_UCLASS_BODY()
 
 private_subobject:
 	/** Shape component used for collision */
diff --git a/Engine/Source/Runtime/Engine/Classes/Engine/TriggerBox.h b/Engine/Source/Runtime/Engine/Classes/Engine/TriggerBox.h
--- a/Engine/Source/Runtime/Engine/Classes/Engine/TriggerBox.h
+++ b/Engine/Source/Runtime/Engine/Classes/Engine/TriggerBox.h
@@ -9,9 +9,7 @@
 UCLASS(MinimalAPI)
 class ATriggerBox : public ATriggerBase
 {
-	GENERATED_BODY()
-public:
-	ENGINE_API ATriggerBox(const FObjectInitializer& ObjectInitializer = FObjectInitializer::Get());
+	GENERATED_UCLASS_BODY()
 
 
 #if WITH_EDITOR
diff --git a/Engine/Source/Runtime/Engine/Classes/Engine/TriggerCapsule.h b/Engine/Source/Runtime/Engine/Classes/Engine/TriggerCapsule.h
--- a/Engine/Source/Runtime/Engine/Classes/Engine/TriggerCapsule.h
+++ b/Engine/Source/Runtime/Engine/Classes/Engine/TriggerCapsule.h
@@ -9,9 +9,7 @@
 UCLASS(MinimalAPI)
 class ATriggerCapsule : public ATriggerBase
 {
-	GENERATED_BODY()
-public:
-	ENGINE_API ATriggerCapsule(const FObjectInitializer& ObjectInitializer = FObjectInitializer::Get());
+	GENERATED_UCLASS_BODY()
 
 
 #if WITH_EDITOR
diff --git a/Engine/Source/Runtime/Engine/Classes/Engine/TriggerSphere.h b/Engine/Source/Runtime/Engine/Classes/Engine/TriggerSphere.h
--- a/Engine/Source/Runtime/Engine/Classes/Engine/TriggerSphere.h
+++ b/Engine/Source/Runtime/Engine/Classes/Engine/TriggerSphere.h
@@ -9,9 +9,7 @@
 UCLASS(MinimalAPI)
 class ATriggerSphere : public ATriggerBase
 {
-	GENERATED_BODY()
-public:
-	ENGINE_API ATriggerSphere(const FObjectInitializer& ObjectInitializer = FObjectInitializer::Get());
+	GENERATED_UCLASS_BODY()
 
 
 #if WITH_EDITOR
diff --git a/Engine/Source/Runtime/Engine/Classes/Engine/TriggerVolume.h b/Engine/Source/Runtime/Engine/Classes/Engine/TriggerVolume.h
--- a/Engine/Source/Runtime/Engine/Classes/Engine/TriggerVolume.h
+++ b/Engine/Source/Runtime/Engine/Classes/Engine/TriggerVolume.h
@@ -8,9 +8,7 @@
 UCLASS()
 class ATriggerVolume : public AVolume
 {
-	GENERATED_BODY()
-public:
-	ATriggerVolume(const FObjectInitializer& ObjectInitializer = FObjectInitializer::Get());
+	GENERATED_UCLASS_BODY()
 
 	// Begin UObject interface.
 #if WITH_EDITOR
diff --git a/Engine/Source/Runtime/Engine/Classes/Engine/TwitterIntegrationBase.h b/Engine/Source/Runtime/Engine/Classes/Engine/TwitterIntegrationBase.h
--- a/Engine/Source/Runtime/Engine/Classes/Engine/TwitterIntegrationBase.h
+++ b/Engine/Source/Runtime/Engine/Classes/Engine/TwitterIntegrationBase.h
@@ -31,9 +31,7 @@ enum ETwitterIntegrationDelegate
 UCLASS(config=Engine)
 class UTwitterIntegrationBase : public UPlatformInterfaceBase
 {
-	GENERATED_BODY()
-public:
-	UTwitterIntegrationBase(const FObjectInitializer& ObjectInitializer = FObjectInitializer::Get());
+	GENERATED_UCLASS_BODY()
 
 	/**
 	 * Perform any needed initialization
diff --git a/Engine/Source/Runtime/Engine/Classes/Engine/UserDefinedEnum.h b/Engine/Source/Runtime/Engine/Classes/Engine/UserDefinedEnum.h
--- a/Engine/Source/Runtime/Engine/Classes/Engine/UserDefinedEnum.h
+++ b/Engine/Source/Runtime/Engine/Classes/Engine/UserDefinedEnum.h
@@ -17,9 +17,7 @@
 UCLASS()
 class ENGINE_API UUserDefinedEnum : public UEnum
 {
-	GENERATED_BODY()
-public:
-	UUserDefinedEnum(const FObjectInitializer& ObjectInitializer = FObjectInitializer::Get());
+	GENERATED_UCLASS_BODY()
 
 #if WITH_EDITORONLY_DATA
 	UPROPERTY()
diff --git a/Engine/Source/Runtime/Engine/Classes/Engine/UserDefinedStruct.h b/Engine/Source/Runtime/Engine/Classes/Engine/UserDefinedStruct.h
--- a/Engine/Source/Runtime/Engine/Classes/Engine/UserDefinedStruct.h
+++ b/Engine/Source/Runtime/Engine/Classes/Engine/UserDefinedStruct.h
@@ -22,9 +22,7 @@ enum EUserDefinedStructureStatus
 UCLASS()
 class ENGINE_API UUserDefinedStruct : public UScriptStruct
 {
-	GENERATED_BODY()
-public:
-	UUserDefinedStruct(const FObjectInitializer& ObjectInitializer = FObjectInitializer::Get());
+	GENERATED_UCLASS_BODY()
 
 public:
 #if WITH_EDITORONLY_DATA
diff --git a/Engine/Source/Runtime/Engine/Classes/Engine/WindDirectionalSource.h b/Engine/Source/Runtime/Engine/Classes/Engine/WindDirectionalSource.h
--- a/Engine/Source/Runtime/Engine/Classes/Engine/WindDirectionalSource.h
+++ b/Engine/Source/Runtime/Engine/Classes/Engine/WindDirectionalSource.h
@@ -6,9 +6,7 @@
 UCLASS(ClassGroup=Wind, showcategories=(Rendering, "Utilities|Transformation"))
 class ENGINE_API AWindDirectionalSource : public AInfo
 {
-	GENERATED_BODY()
-public:
-	AWindDirectionalSource(const FObjectInitializer& ObjectInitializer = FObjectInitializer::Get());
+	GENERATED_UCLASS_BODY()
 
 private_subobject:
 	DEPRECATED_FORGAME(4.6, "Component should not be accessed directly, please use GetComponent() function instead. Component will soon be private and your code will not compile.")
diff --git a/Engine/Source/Runtime/Engine/Classes/Engine/World.h b/Engine/Source/Runtime/Engine/Classes/Engine/World.h
--- a/Engine/Source/Runtime/Engine/Classes/Engine/World.h
+++ b/Engine/Source/Runtime/Engine/Classes/Engine/World.h
@@ -324,7 +324,7 @@ struct ENGINE_API FLevelViewportInfo
 USTRUCT()
 struct FStartPhysicsTickFunction : public FTickFunction
 {
-	GENERATED_BODY()
+	GENERATED_USTRUCT_BODY()
 
 	/** World this tick function belongs to **/
 	class UWorld*	Target;
@@ -347,7 +347,7 @@ struct FStartPhysicsTickFunction : public FTickFunction
 USTRUCT()
 struct FEndPhysicsTickFunction : public FTickFunction
 {
-	GENERATED_BODY()
+	GENERATED_USTRUCT_BODY()
 
 	/** World this tick function belongs to **/
 	class UWorld*	Target;
@@ -370,7 +370,7 @@ struct FEndPhysicsTickFunction : public FTickFunction
 USTRUCT()
 struct FStartClothSimulationFunction : public FTickFunction
 {
-	GENERATED_BODY()
+	GENERATED_USTRUCT_BODY()
 
 	/** World this tick function belongs to **/
 	class UWorld*	Target;
@@ -393,7 +393,7 @@ struct FStartClothSimulationFunction : public FTickFunction
 USTRUCT()
 struct FEndClothSimulationFunction : public FTickFunction
 {
-	GENERATED_BODY()
+	GENERATED_USTRUCT_BODY()
 
 	/** World this tick function belongs to **/
 	class UWorld*	Target;
@@ -537,11 +537,10 @@ public:
  *
  */
 
-UCLASS(config=Engine)
+UCLASS(customConstructor, config=Engine)
 class ENGINE_API UWorld : public UObject, public FNetworkNotify
 {
-	GENERATED_BODY()
-public:
+	GENERATED_UCLASS_BODY()
 
 	~UWorld();
 
diff --git a/Engine/Source/Runtime/Engine/Classes/Engine/WorldComposition.h b/Engine/Source/Runtime/Engine/Classes/Engine/WorldComposition.h
--- a/Engine/Source/Runtime/Engine/Classes/Engine/WorldComposition.h
+++ b/Engine/Source/Runtime/Engine/Classes/Engine/WorldComposition.h
@@ -64,9 +64,7 @@ struct FDistanceVisibleLevel
 UCLASS(config=Engine)
 class ENGINE_API UWorldComposition : public UObject
 {
-	GENERATED_BODY()
-public:
-	UWorldComposition(const FObjectInitializer& ObjectInitializer = FObjectInitializer::Get());
+	GENERATED_UCLASS_BODY()
 
 	typedef TArray<FWorldCompositionTile> FTilesList;
 
@@ -223,4 +221,4 @@ public:
 	UPROPERTY(config)
 	bool						bLoadAllTilesDuringCinematic;
 
-};
+};
\ No newline at end of file
diff --git a/Engine/Source/Runtime/Engine/Classes/Exporters/Exporter.h b/Engine/Source/Runtime/Engine/Classes/Exporters/Exporter.h
--- a/Engine/Source/Runtime/Engine/Classes/Exporters/Exporter.h
+++ b/Engine/Source/Runtime/Engine/Classes/Exporters/Exporter.h
@@ -14,9 +14,7 @@ class UActorComponent;
 UCLASS(abstract, transient, MinimalAPI)
 class UExporter : public UObject
 {
-	GENERATED_BODY()
-public:
-	ENGINE_API UExporter(const FObjectInitializer& ObjectInitializer = FObjectInitializer::Get());
+	GENERATED_UCLASS_BODY()
 
 	/** Supported class of this exporter */
 	UPROPERTY()
diff --git a/Engine/Source/Runtime/Engine/Classes/GameFramework/CameraBlockingVolume.h b/Engine/Source/Runtime/Engine/Classes/GameFramework/CameraBlockingVolume.h
--- a/Engine/Source/Runtime/Engine/Classes/GameFramework/CameraBlockingVolume.h
+++ b/Engine/Source/Runtime/Engine/Classes/GameFramework/CameraBlockingVolume.h
@@ -9,9 +9,7 @@
 UCLASS()
 class ENGINE_API ACameraBlockingVolume : public AVolume
 {
-	GENERATED_BODY()
-public:
-	ACameraBlockingVolume(const FObjectInitializer& ObjectInitializer = FObjectInitializer::Get());
+	GENERATED_UCLASS_BODY()
 };
 
 
diff --git a/Engine/Source/Runtime/Engine/Classes/GameFramework/CheatManager.h b/Engine/Source/Runtime/Engine/Classes/GameFramework/CheatManager.h
--- a/Engine/Source/Runtime/Engine/Classes/GameFramework/CheatManager.h
+++ b/Engine/Source/Runtime/Engine/Classes/GameFramework/CheatManager.h
@@ -57,9 +57,7 @@ struct FDebugTraceInfo
 UCLASS(Within=PlayerController)
 class ENGINE_API UCheatManager : public UObject
 {
-	GENERATED_BODY()
-public:
-	UCheatManager(const FObjectInitializer& ObjectInitializer = FObjectInitializer::Get());
+	GENERATED_UCLASS_BODY()
 
 	/** Debug camera - used to have independent camera without stopping gameplay */
 	UPROPERTY()
@@ -206,10 +204,8 @@ public:
 	UFUNCTION(exec)
 	virtual void ToggleAILogging();
 
-	UFUNCTION(reliable, server="ServerToggleAILogging_Implementation", WithValidation="ServerToggleAILogging_Validate")
+	UFUNCTION(reliable, server, WithValidation)
 	virtual void ServerToggleAILogging();
-	virtual void ServerToggleAILogging_Implementation();
-	virtual bool ServerToggleAILogging_Validate();
 
 	/** Toggle capsule trace debugging. Will trace a capsule from current view point and show where it hits the world */
 	UFUNCTION(exec)
diff --git a/Engine/Source/Runtime/Engine/Classes/GameFramework/DamageType.h b/Engine/Source/Runtime/Engine/Classes/GameFramework/DamageType.h
--- a/Engine/Source/Runtime/Engine/Classes/GameFramework/DamageType.h
+++ b/Engine/Source/Runtime/Engine/Classes/GameFramework/DamageType.h
@@ -15,9 +15,7 @@
 UCLASS(MinimalAPI, const, Blueprintable, BlueprintType)
 class UDamageType : public UObject
 {
-	GENERATED_BODY()
-public:
-	ENGINE_API UDamageType(const FObjectInitializer& ObjectInitializer = FObjectInitializer::Get());
+	GENERATED_UCLASS_BODY()
 
 	/** True if this damagetype is caused by the world (falling off level, into lava, etc). */
 	UPROPERTY(EditAnywhere, BlueprintReadWrite, Category=DamageType)
diff --git a/Engine/Source/Runtime/Engine/Classes/GameFramework/DefaultPhysicsVolume.h b/Engine/Source/Runtime/Engine/Classes/GameFramework/DefaultPhysicsVolume.h
--- a/Engine/Source/Runtime/Engine/Classes/GameFramework/DefaultPhysicsVolume.h
+++ b/Engine/Source/Runtime/Engine/Classes/GameFramework/DefaultPhysicsVolume.h
@@ -17,9 +17,7 @@
 UCLASS(notplaceable, transient, MinimalAPI)
 class ADefaultPhysicsVolume : public APhysicsVolume
 {
-	GENERATED_BODY()
-public:
-	ENGINE_API ADefaultPhysicsVolume(const FObjectInitializer& ObjectInitializer = FObjectInitializer::Get());
+	GENERATED_UCLASS_BODY()
 
 	// Begin AActor interface
 	ENGINE_API virtual void Destroyed() override;
diff --git a/Engine/Source/Runtime/Engine/Classes/GameFramework/EngineMessage.h b/Engine/Source/Runtime/Engine/Classes/GameFramework/EngineMessage.h
--- a/Engine/Source/Runtime/Engine/Classes/GameFramework/EngineMessage.h
+++ b/Engine/Source/Runtime/Engine/Classes/GameFramework/EngineMessage.h
@@ -7,9 +7,7 @@
 UCLASS(abstract, MinimalAPI)
 class  UEngineMessage : public ULocalMessage
 {
-	GENERATED_BODY()
-public:
-	ENGINE_API UEngineMessage(const FObjectInitializer& ObjectInitializer = FObjectInitializer::Get());
+	GENERATED_UCLASS_BODY()
 
 	/** Message displayed in message dialog when player pawn fails to spawn because no playerstart was available. */
 	UPROPERTY()
diff --git a/Engine/Source/Runtime/Engine/Classes/GameFramework/FloatingPawnMovement.h b/Engine/Source/Runtime/Engine/Classes/GameFramework/FloatingPawnMovement.h
--- a/Engine/Source/Runtime/Engine/Classes/GameFramework/FloatingPawnMovement.h
+++ b/Engine/Source/Runtime/Engine/Classes/GameFramework/FloatingPawnMovement.h
@@ -15,9 +15,7 @@
 UCLASS(ClassGroup = Movement, meta = (BlueprintSpawnableComponent))
 class ENGINE_API UFloatingPawnMovement : public UPawnMovementComponent
 {
-	GENERATED_BODY()
-public:
-	UFloatingPawnMovement(const FObjectInitializer& ObjectInitializer = FObjectInitializer::Get());
+	GENERATED_UCLASS_BODY()
 
 	//Begin UActorComponent Interface
 	virtual void TickComponent(float DeltaTime, enum ELevelTick TickType, FActorComponentTickFunction *ThisTickFunction) override;
diff --git a/Engine/Source/Runtime/Engine/Classes/GameFramework/ForceFeedbackEffect.h b/Engine/Source/Runtime/Engine/Classes/GameFramework/ForceFeedbackEffect.h
--- a/Engine/Source/Runtime/Engine/Classes/GameFramework/ForceFeedbackEffect.h
+++ b/Engine/Source/Runtime/Engine/Classes/GameFramework/ForceFeedbackEffect.h
@@ -6,7 +6,7 @@
 USTRUCT()
 struct FForceFeedbackChannelDetails
 {
-	GENERATED_BODY()
+	GENERATED_USTRUCT_BODY()
 
 	UPROPERTY(EditAnywhere, Category="ChannelDetails")
 	uint32 bAffectsLeftLarge:1;
@@ -35,7 +35,7 @@ struct FForceFeedbackChannelDetails
 USTRUCT()
 struct FActiveForceFeedbackEffect
 {
-	GENERATED_BODY()
+	GENERATED_USTRUCT_BODY()
 
 	UPROPERTY()
 	class UForceFeedbackEffect* ForceFeedbackEffect;
@@ -70,9 +70,7 @@ struct FActiveForceFeedbackEffect
 UCLASS(BlueprintType, MinimalAPI)
 class UForceFeedbackEffect : public UObject
 {
-	GENERATED_BODY()
-public:
-	ENGINE_API UForceFeedbackEffect(const FObjectInitializer& ObjectInitializer = FObjectInitializer::Get());
+	GENERATED_UCLASS_BODY()
 
 	UPROPERTY(EditAnywhere, Category="ForceFeedbackEffect")
 	TArray<FForceFeedbackChannelDetails> ChannelDetails;
@@ -88,4 +86,4 @@ public:
 	float GetDuration();
 
 	void GetValues(const float EvalTime, FForceFeedbackValues& Values) const;
-};
+};
\ No newline at end of file
diff --git a/Engine/Source/Runtime/Engine/Classes/GameFramework/GameMode.h b/Engine/Source/Runtime/Engine/Classes/GameFramework/GameMode.h
--- a/Engine/Source/Runtime/Engine/Classes/GameFramework/GameMode.h
+++ b/Engine/Source/Runtime/Engine/Classes/GameFramework/GameMode.h
@@ -33,7 +33,7 @@ DECLARE_DELEGATE_RetVal( bool, FCanUnpause );
 USTRUCT()
 struct FGameClassShortName
 {
-	GENERATED_BODY()
+	GENERATED_USTRUCT_BODY()
 
 	/** Abbreviation that can be used as an alias for the class name */
 	UPROPERTY()
@@ -62,9 +62,7 @@ struct FGameClassShortName
 UCLASS(config=Game, notplaceable, BlueprintType, Blueprintable, Transient, hideCategories=(Info, Rendering, MovementReplication, Replication, Actor), meta=(ShortTooltip="Game Mode defines the game being played, its rules, scoring, and other facets of the game type."))
 class ENGINE_API AGameMode : public AInfo
 {
-	GENERATED_BODY()
-public:
-	AGameMode(const FObjectInitializer& ObjectInitializer = FObjectInitializer::Get());
+	GENERATED_UCLASS_BODY()
 
 	// Code to deal with the match state machine
 
diff --git a/Engine/Source/Runtime/Engine/Classes/GameFramework/GameNetworkManager.h b/Engine/Source/Runtime/Engine/Classes/GameFramework/GameNetworkManager.h
--- a/Engine/Source/Runtime/Engine/Classes/GameFramework/GameNetworkManager.h
+++ b/Engine/Source/Runtime/Engine/Classes/GameFramework/GameNetworkManager.h
@@ -22,9 +22,7 @@ enum EStandbyType
 UCLASS(config=Game, notplaceable)
 class ENGINE_API AGameNetworkManager : public AInfo
 {
-	GENERATED_BODY()
-public:
-	AGameNetworkManager(const FObjectInitializer& ObjectInitializer = FObjectInitializer::Get());
+	GENERATED_UCLASS_BODY()
 
 	//======================================================================================================================
 	// Listen server dynamic netspeed adjustment
diff --git a/Engine/Source/Runtime/Engine/Classes/GameFramework/GameSession.h b/Engine/Source/Runtime/Engine/Classes/GameFramework/GameSession.h
--- a/Engine/Source/Runtime/Engine/Classes/GameFramework/GameSession.h
+++ b/Engine/Source/Runtime/Engine/Classes/GameFramework/GameSession.h
@@ -14,9 +14,7 @@ class APlayerController;
 UCLASS(config=Game, notplaceable)
 class ENGINE_API AGameSession : public AInfo
 {
-	GENERATED_BODY()
-public:
-	AGameSession(const FObjectInitializer& ObjectInitializer = FObjectInitializer::Get());
+	GENERATED_UCLASS_BODY()
 
 	/** Maximum number of spectators allowed by this server. */
 	UPROPERTY(globalconfig)
diff --git a/Engine/Source/Runtime/Engine/Classes/GameFramework/GameState.h b/Engine/Source/Runtime/Engine/Classes/GameFramework/GameState.h
--- a/Engine/Source/Runtime/Engine/Classes/GameFramework/GameState.h
+++ b/Engine/Source/Runtime/Engine/Classes/GameFramework/GameState.h
@@ -11,9 +11,7 @@
 UCLASS(config=Game, notplaceable, BlueprintType, Blueprintable)
 class ENGINE_API AGameState : public AInfo
 {
-	GENERATED_BODY()
-public:
-	AGameState(const FObjectInitializer& ObjectInitializer = FObjectInitializer::Get());
+	GENERATED_UCLASS_BODY()
 
 	/** Class of the server's game mode, assigned by GameMode. */
 	UPROPERTY(replicatedUsing=OnRep_GameModeClass)
diff --git a/Engine/Source/Runtime/Engine/Classes/GameFramework/GameUserSettings.h b/Engine/Source/Runtime/Engine/Classes/GameFramework/GameUserSettings.h
--- a/Engine/Source/Runtime/Engine/Classes/GameFramework/GameUserSettings.h
+++ b/Engine/Source/Runtime/Engine/Classes/GameFramework/GameUserSettings.h
@@ -10,9 +10,7 @@
 UCLASS(config=GameUserSettings, configdonotcheckdefaults)
 class ENGINE_API UGameUserSettings : public UObject
 {
-	GENERATED_BODY()
-public:
-	UGameUserSettings(const FObjectInitializer& ObjectInitializer = FObjectInitializer::Get());
+	GENERATED_UCLASS_BODY()
 
 public:
 
diff --git a/Engine/Source/Runtime/Engine/Classes/GameFramework/HUD.h b/Engine/Source/Runtime/Engine/Classes/GameFramework/HUD.h
--- a/Engine/Source/Runtime/Engine/Classes/GameFramework/HUD.h
+++ b/Engine/Source/Runtime/Engine/Classes/GameFramework/HUD.h
@@ -25,9 +25,7 @@ class UFont;
 UCLASS(config=Game, hidecategories=(Rendering,Actor,Input,Replication), showcategories=("Input|MouseInput", "Input|TouchInput"), notplaceable, transient, BlueprintType, Blueprintable)
 class ENGINE_API AHUD : public AActor
 {
-	GENERATED_BODY()
-public:
-	AHUD(const FObjectInitializer& ObjectInitializer = FObjectInitializer::Get());
+	GENERATED_UCLASS_BODY()
 
 	/** Pre-defined FColors for convenience. */
 	UPROPERTY()
@@ -153,16 +151,14 @@ public:
 	 * @param InFont 				font to use
 	 * @param FontScale 			scale
 	 */
-	UFUNCTION(reliable, client="AddDebugText_Implementation", SealedEvent)
+	UFUNCTION(reliable, client, SealedEvent)
 	void AddDebugText(const FString& DebugText, AActor* SrcActor = NULL, float Duration = 0, FVector Offset = FVector(ForceInit), FVector DesiredOffset = FVector(ForceInit), FColor TextColor = FColor(ForceInit), bool bSkipOverwriteCheck = false, bool bAbsoluteLocation = false, bool bKeepAttachedToActor = false, UFont* InFont = NULL, float FontScale = 1.0);
-	void AddDebugText_Implementation(const FString& DebugText, AActor* SrcActor, float Duration, FVector Offset, FVector DesiredOffset, FColor TextColor, bool bSkipOverwriteCheck, bool bAbsoluteLocation, bool bKeepAttachedToActor, UFont* InFont, float FontScale);
 
 	/**
 	 * Remove all debug strings added via AddDebugText
 	 */
-	UFUNCTION(reliable, client="RemoveAllDebugStrings_Implementation", SealedEvent)
+	UFUNCTION(reliable, client, SealedEvent)
 	void RemoveAllDebugStrings();
-	void RemoveAllDebugStrings_Implementation();
 
 	/**
 	 * Remove debug strings for the given actor
@@ -170,9 +166,8 @@ public:
 	 * @param	SrcActor			Actor whose string you wish to remove
 	 * @param	bLeaveDurationText	when true text that has a finite duration will be removed, otherwise all will be removed for given actor
 	 */
-	UFUNCTION(reliable, client="RemoveDebugText_Implementation", SealedEvent)
+	UFUNCTION(reliable, client, SealedEvent)
 	void RemoveDebugText(AActor* SrcActor, bool bLeaveDurationText = false);
-	void RemoveDebugText_Implementation(AActor* SrcActor, bool bLeaveDurationText);
 
 	/** Hook to allow blueprints to do custom HUD drawing. @see bSuppressNativeHUD to control HUD drawing in base class. */
 	UFUNCTION(BlueprintImplementableEvent, BlueprintCosmetic)
diff --git a/Engine/Source/Runtime/Engine/Classes/GameFramework/Info.h b/Engine/Source/Runtime/Engine/Classes/GameFramework/Info.h
--- a/Engine/Source/Runtime/Engine/Classes/GameFramework/Info.h
+++ b/Engine/Source/Runtime/Engine/Classes/GameFramework/Info.h
@@ -16,9 +16,7 @@
 UCLASS(abstract, hidecategories=(Input, Movement, Collision, Rendering, "Utilities|Transformation"), showcategories=("Input|MouseInput", "Input|TouchInput"), MinimalAPI, NotBlueprintable)
 class AInfo : public AActor
 {
-	GENERATED_BODY()
-public:
-	ENGINE_API AInfo(const FObjectInitializer& ObjectInitializer = FObjectInitializer::Get());
+	GENERATED_UCLASS_BODY()
 
 #if WITH_EDITORONLY_DATA
 private_subobject:
diff --git a/Engine/Source/Runtime/Engine/Classes/GameFramework/InputSettings.h b/Engine/Source/Runtime/Engine/Classes/GameFramework/InputSettings.h
--- a/Engine/Source/Runtime/Engine/Classes/GameFramework/InputSettings.h
+++ b/Engine/Source/Runtime/Engine/Classes/GameFramework/InputSettings.h
@@ -14,9 +14,7 @@ UCLASS(config=Input, defaultconfig)
 class ENGINE_API UInputSettings
 	: public UObject
 {
-	GENERATED_BODY()
-public:
-	UInputSettings(const FObjectInitializer& ObjectInitializer = FObjectInitializer::Get());
+	GENERATED_UCLASS_BODY()
 
 	/** Properties of Axis controls */
 	UPROPERTY(config, EditAnywhere, EditFixedSize, Category="Bindings", meta=(ToolTip="List of Axis Properties"), AdvancedDisplay)
diff --git a/Engine/Source/Runtime/Engine/Classes/GameFramework/KillZVolume.h b/Engine/Source/Runtime/Engine/Classes/GameFramework/KillZVolume.h
--- a/Engine/Source/Runtime/Engine/Classes/GameFramework/KillZVolume.h
+++ b/Engine/Source/Runtime/Engine/Classes/GameFramework/KillZVolume.h
@@ -14,9 +14,7 @@
 UCLASS()
 class ENGINE_API AKillZVolume : public APhysicsVolume
 {
-	GENERATED_BODY()
-public:
-	AKillZVolume(const FObjectInitializer& ObjectInitializer = FObjectInitializer::Get());
+	GENERATED_UCLASS_BODY()
 	
 	//Begin PhysicsVolume Interface
 	virtual void ActorEnteredVolume(class AActor* Other) override;
diff --git a/Engine/Source/Runtime/Engine/Classes/GameFramework/LocalMessage.h b/Engine/Source/Runtime/Engine/Classes/GameFramework/LocalMessage.h
--- a/Engine/Source/Runtime/Engine/Classes/GameFramework/LocalMessage.h
+++ b/Engine/Source/Runtime/Engine/Classes/GameFramework/LocalMessage.h
@@ -19,8 +19,8 @@
 USTRUCT()
 struct ENGINE_API FClientReceiveData
 {
-	GENERATED_BODY()
 	//always need to be here
+	GENERATED_USTRUCT_BODY()
 
 	UPROPERTY()
 	APlayerController* LocalPC;
@@ -49,9 +49,7 @@ struct ENGINE_API FClientReceiveData
 UCLASS(abstract)
 class ENGINE_API ULocalMessage : public UObject
 {
-	GENERATED_BODY()
-public:
-	ULocalMessage(const FObjectInitializer& ObjectInitializer = FObjectInitializer::Get());
+	GENERATED_UCLASS_BODY()
 	/** send message to client */
 	virtual void ClientReceive(const FClientReceiveData& ClientData) const PURE_VIRTUAL(ULocalMessage::ClientReceive,);
 };
diff --git a/Engine/Source/Runtime/Engine/Classes/GameFramework/MovementComponent.h b/Engine/Source/Runtime/Engine/Classes/GameFramework/MovementComponent.h
--- a/Engine/Source/Runtime/Engine/Classes/GameFramework/MovementComponent.h
+++ b/Engine/Source/Runtime/Engine/Classes/GameFramework/MovementComponent.h
@@ -46,9 +46,7 @@ enum class EPlaneConstraintAxisSetting : uint8
 UCLASS(ClassGroup=Movement, abstract, BlueprintType)
 class ENGINE_API UMovementComponent : public UActorComponent
 {
-	GENERATED_BODY()
-public:
-	UMovementComponent(const FObjectInitializer& ObjectInitializer = FObjectInitializer::Get());
+	GENERATED_UCLASS_BODY()
 
 	/**
 	 * The component we move and update.
diff --git a/Engine/Source/Runtime/Engine/Classes/GameFramework/NavMovementComponent.h b/Engine/Source/Runtime/Engine/Classes/GameFramework/NavMovementComponent.h
--- a/Engine/Source/Runtime/Engine/Classes/GameFramework/NavMovementComponent.h
+++ b/Engine/Source/Runtime/Engine/Classes/GameFramework/NavMovementComponent.h
@@ -22,9 +22,7 @@ class UCapsuleComponent;
 UCLASS(abstract)
 class ENGINE_API UNavMovementComponent : public UMovementComponent
 {
-	GENERATED_BODY()
-public:
-	UNavMovementComponent(const FObjectInitializer& ObjectInitializer = FObjectInitializer::Get());
+	GENERATED_UCLASS_BODY()
 
 	/** Properties that define how the component can move. */
 	UPROPERTY(EditAnywhere, BlueprintReadWrite, Category=MovementComponent, meta=(DisplayName="Movement Capabilities", Keywords="Nav Agent"))
diff --git a/Engine/Source/Runtime/Engine/Classes/GameFramework/OnlineSession.h b/Engine/Source/Runtime/Engine/Classes/GameFramework/OnlineSession.h
--- a/Engine/Source/Runtime/Engine/Classes/GameFramework/OnlineSession.h
+++ b/Engine/Source/Runtime/Engine/Classes/GameFramework/OnlineSession.h
@@ -10,9 +10,7 @@
 UCLASS(config=Game)
 class ENGINE_API UOnlineSession : public UObject
 {
-	GENERATED_BODY()
-public:
-	UOnlineSession(const FObjectInitializer& ObjectInitializer = FObjectInitializer::Get());
+	GENERATED_UCLASS_BODY()
 
 public:
 
diff --git a/Engine/Source/Runtime/Engine/Classes/GameFramework/PhysicsVolume.h b/Engine/Source/Runtime/Engine/Classes/GameFramework/PhysicsVolume.h
--- a/Engine/Source/Runtime/Engine/Classes/GameFramework/PhysicsVolume.h
+++ b/Engine/Source/Runtime/Engine/Classes/GameFramework/PhysicsVolume.h
@@ -12,9 +12,7 @@
 UCLASS()
 class ENGINE_API APhysicsVolume : public AVolume
 {
-	GENERATED_BODY()
-public:
-	APhysicsVolume(const FObjectInitializer& ObjectInitializer = FObjectInitializer::Get());
+	GENERATED_UCLASS_BODY()
 
 	// Begin UObject interface.
 #if WITH_EDITOR
diff --git a/Engine/Source/Runtime/Engine/Classes/GameFramework/PlayerController.h b/Engine/Source/Runtime/Engine/Classes/GameFramework/PlayerController.h
--- a/Engine/Source/Runtime/Engine/Classes/GameFramework/PlayerController.h
+++ b/Engine/Source/Runtime/Engine/Classes/GameFramework/PlayerController.h
@@ -146,9 +146,7 @@ protected:
 UCLASS(config=Game, BlueprintType, Blueprintable, meta=(ShortTooltip="A Player Controller is an actor responsible for controlling a Pawn used by the player."))
 class ENGINE_API APlayerController : public AController
 {
-	GENERATED_BODY()
-public:
-	APlayerController(const FObjectInitializer& ObjectInitializer = FObjectInitializer::Get());
+	GENERATED_UCLASS_BODY()
 	
 	/** UPlayer associated with this PlayerController.  Could be a local player or a net connection. */
 	UPROPERTY()
@@ -339,14 +337,12 @@ public:
 	virtual void LocalTravel(const FString& URL);
 
 	/** Return the client to the main menu gracefully */
-	UFUNCTION(Reliable, Client="ClientReturnToMainMenu_Implementation")
+	UFUNCTION(Reliable, Client)
 	virtual void ClientReturnToMainMenu(const FString& ReturnReason);
-	virtual void ClientReturnToMainMenu_Implementation(const FString& ReturnReason);
 
 	/** Development RPC for testing object reference replication */
-	UFUNCTION(Reliable, Client="ClientRepObjRef_Implementation")
+	UFUNCTION(Reliable, Client)
 	virtual void ClientRepObjRef(UObject* Object);
-	virtual void ClientRepObjRef_Implementation(UObject* Object);
 
 	/**
 	 * Locally try to pause game (call serverpause to pause network game); returns success indicator.  Calls GameMode's SetPause().
@@ -511,9 +507,8 @@ public:
 	 * Tell the client to enable or disable voice chat (not muting)
 	 * @param bEnable enable or disable voice chat
 	 */
-	UFUNCTION(Reliable, Client="ClientEnableNetworkVoice_Implementation")
+	UFUNCTION(Reliable, Client)
 	virtual void ClientEnableNetworkVoice(bool bEnable);
-	virtual void ClientEnableNetworkVoice_Implementation(bool bEnable);
 
 	/** Enable voice chat transmission */
 	void StartTalking();
@@ -535,43 +530,36 @@ public:
 	 *
 	 * NOTE: This is done as an RPC instead of variable replication because ordering matters
 	 */
-	UFUNCTION(Reliable, Client="ClientVoiceHandshakeComplete_Implementation")
+	UFUNCTION(Reliable, Client)
 	virtual void ClientVoiceHandshakeComplete();
-	virtual void ClientVoiceHandshakeComplete_Implementation();
 
 	/**
 	 * Tell the server to mute a player for this controller
 	 * @param PlayerId player id to mute
 	 */
-	UFUNCTION(server="ServerMutePlayer_Implementation", reliable, WithValidation="ServerMutePlayer_Validate")
+	UFUNCTION(server, reliable, WithValidation)
 	virtual void ServerMutePlayer(FUniqueNetIdRepl PlayerId);
-	virtual void ServerMutePlayer_Implementation(FUniqueNetIdRepl PlayerId);
-	virtual bool ServerMutePlayer_Validate(FUniqueNetIdRepl );
 
 	/**
 	 * Tell the server to unmute a player for this controller
 	 * @param PlayerId player id to unmute
 	 */
-	UFUNCTION(server="ServerUnmutePlayer_Implementation", reliable, WithValidation="ServerUnmutePlayer_Validate" )
+	UFUNCTION(server, reliable, WithValidation )
 	virtual void ServerUnmutePlayer(FUniqueNetIdRepl PlayerId);
-	virtual void ServerUnmutePlayer_Implementation(FUniqueNetIdRepl PlayerId);
-	virtual bool ServerUnmutePlayer_Validate(FUniqueNetIdRepl );
 
 	/**
 	 * Tell the client to mute a player for this controller
 	 * @param PlayerId player id to mute
 	 */
-	UFUNCTION(Reliable, Client="ClientMutePlayer_Implementation")
+	UFUNCTION(Reliable, Client)
 	virtual void ClientMutePlayer(FUniqueNetIdRepl PlayerId);
-	virtual void ClientMutePlayer_Implementation(FUniqueNetIdRepl PlayerId);
 
 	/**
 	 * Tell the client to unmute a player for this controller
 	 * @param PlayerId player id to unmute
 	 */
-	UFUNCTION(Reliable, Client="ClientUnmutePlayer_Implementation")
+	UFUNCTION(Reliable, Client)
 	virtual void ClientUnmutePlayer(FUniqueNetIdRepl PlayerId);
-	virtual void ClientUnmutePlayer_Implementation(FUniqueNetIdRepl PlayerId);
 
 	/**
 	 * Mutes a remote player on the server and then tells the client to mute
@@ -608,64 +596,54 @@ public:
 	virtual void SendToConsole(const FString& Command);
 
 	/** Adds a location to the texture streaming system for the specified duration. */
-	UFUNCTION(reliable, client="ClientAddTextureStreamingLoc_Implementation", SealedEvent)
+	UFUNCTION(reliable, client, SealedEvent)
 	void ClientAddTextureStreamingLoc(FVector InLoc, float Duration, bool bOverrideLocation);
-	void ClientAddTextureStreamingLoc_Implementation(FVector InLoc, float Duration, bool bOverrideLocation);
 
 	/** Tells client to cancel any pending map change. */
-	UFUNCTION(Reliable, Client="ClientCancelPendingMapChange_Implementation")
+	UFUNCTION(Reliable, Client)
 	void ClientCancelPendingMapChange();
-	virtual void ClientCancelPendingMapChange_Implementation();
 
 	/** Set CurrentNetSpeed to the lower of its current value and Cap. */
-	UFUNCTION(Reliable, Client="ClientCapBandwidth_Implementation")
+	UFUNCTION(Reliable, Client)
 	void ClientCapBandwidth(int32 Cap);
-	virtual void ClientCapBandwidth_Implementation(int32 Cap);
 
 	/** Actually performs the level transition prepared by PrepareMapChange(). */
-	UFUNCTION(Reliable, Client="ClientCommitMapChange_Implementation")	
+	UFUNCTION(Reliable, Client)	
 	void ClientCommitMapChange();
-	virtual void ClientCommitMapChange_Implementation();
 
 	/** Tells the client to block until all pending level streaming actions are complete
 	 * happens at the end of the tick
 	 * primarily used to force update the client ASAP at join time
 	 */
-	UFUNCTION(reliable, client="ClientFlushLevelStreaming_Implementation", SealedEvent)
+	UFUNCTION(reliable, client, SealedEvent)
 	void ClientFlushLevelStreaming();
-	void ClientFlushLevelStreaming_Implementation();
 
 	/** Forces GC at the end of the tick on the client */
-	UFUNCTION(Reliable, Client="ClientForceGarbageCollection_Implementation")
+	UFUNCTION(Reliable, Client)
 	void ClientForceGarbageCollection();
-	virtual void ClientForceGarbageCollection_Implementation();
 
 	/** 
 	 * Replicated function called by GameHasEnded().
 	 * @param	EndGameFocus - actor to view with camera
 	 * @param	bIsWinner - true if this controller is on winning team
 	 */
-	UFUNCTION(Reliable, Client="ClientGameEnded_Implementation")
+	UFUNCTION(Reliable, Client)
 	void ClientGameEnded(class AActor* EndGameFocus, bool bIsWinner);
-	virtual void ClientGameEnded_Implementation(AActor* EndGameFocus, bool bIsWinner);
 
 	/** 
 	 * Server uses this to force client into NewState .
 	 * @Note ALL STATE NAMES NEED TO BE DEFINED IN name table in UnrealNames.h to be correctly replicated (so they are mapped to the same thing on client and server).
 	 */
-	UFUNCTION(Reliable, Client="ClientGotoState_Implementation")
+	UFUNCTION(Reliable, Client)
 	void ClientGotoState(FName NewState);
-	virtual void ClientGotoState_Implementation(FName NewState);
 
 	/** calls IgnoreLookInput on client */
-	UFUNCTION(Reliable, Client="ClientIgnoreLookInput_Implementation")
+	UFUNCTION(Reliable, Client)
 	void ClientIgnoreLookInput(bool bIgnore);
-	virtual void ClientIgnoreLookInput_Implementation(bool bIgnore);
 
 	/** calls IgnoreMoveInput on client */
-	UFUNCTION(Reliable, Client="ClientIgnoreMoveInput_Implementation")
+	UFUNCTION(Reliable, Client)
 	void ClientIgnoreMoveInput(bool bIgnore);
-	virtual void ClientIgnoreMoveInput_Implementation(bool bIgnore);
 
 	/**
 	 * Outputs a message to HUD
@@ -673,9 +651,8 @@ public:
 	 * @param Type - @todo document
 	 * @param MsgLifeTime - Optional length of time to display 0 = default time
 	 */
-	UFUNCTION(Reliable, Client="ClientMessage_Implementation")
+	UFUNCTION(Reliable, Client)
 	void ClientMessage(const FString& S, FName Type = NAME_None, float MsgLifeTime = 0);
-	virtual void ClientMessage_Implementation(const FString& S, FName Type, float MsgLifeTime);
 
 	/** Play the indicated CameraAnim on this camera.
 	 * @param AnimToPlay - Camera animation to play
@@ -688,9 +665,8 @@ public:
 	 * @param Space - Animation play area
 	 * @param CustomPlaySpace - Matrix used when Space = CAPS_UserDefined
 	 */
-	UFUNCTION(unreliable, client="ClientPlayCameraAnim_Implementation", BlueprintCallable, Category = "Game|Feedback")
+	UFUNCTION(unreliable, client, BlueprintCallable, Category = "Game|Feedback")
 	void ClientPlayCameraAnim(class UCameraAnim* AnimToPlay, float Scale=0, float Rate=0, float BlendInTime=0, float BlendOutTime=0, bool bLoop=false, bool bRandomStartTime=false, ECameraAnimPlaySpace::Type Space=ECameraAnimPlaySpace::CameraLocal, FRotator CustomPlaySpace=FRotator::ZeroRotator);
-	virtual void ClientPlayCameraAnim_Implementation(UCameraAnim* AnimToPlay, float Scale, float Rate, float BlendInTime, float BlendOutTime, bool bLoop, bool bRandomStartTime, ECameraAnimPlaySpace::Type Space, FRotator CustomPlaySpace);
 
 	/** 
 	 * Play Camera Shake 
@@ -699,9 +675,8 @@ public:
 	 * @param PlaySpace - Which coordinate system to play the shake in (used for CameraAnims within the shake).
 	 * @param UserPlaySpaceRot - Matrix used when PlaySpace = CAPS_UserDefined
 	 */
-	UFUNCTION(unreliable, client="ClientPlayCameraShake_Implementation", BlueprintCallable, Category="Game|Feedback")
+	UFUNCTION(unreliable, client, BlueprintCallable, Category="Game|Feedback")
 	void ClientPlayCameraShake(TSubclassOf<class UCameraShake> Shake, float Scale = 0, ECameraAnimPlaySpace::Type PlaySpace = ECameraAnimPlaySpace::CameraLocal, FRotator UserPlaySpaceRot = FRotator::ZeroRotator);
-	virtual void ClientPlayCameraShake_Implementation(TSubclassOf<UCameraShake>  Shake, float Scale, ECameraAnimPlaySpace::Type PlaySpace, FRotator UserPlaySpaceRot);
 
 	/**
 	 * Play sound client-side (so only the client will hear it)
@@ -709,9 +684,8 @@ public:
 	 * @param VolumeMultiplier - Volume multiplier to apply to the sound
 	 * @param PitchMultiplier - Pitch multiplier to apply to the sound
 	 */
-	UFUNCTION(unreliable, client="ClientPlaySound_Implementation")
+	UFUNCTION(unreliable, client)
 	void ClientPlaySound(class USoundBase* Sound, float VolumeMultiplier = 1.f, float PitchMultiplier = 1.f );
-	virtual void ClientPlaySound_Implementation(USoundBase* Sound, float VolumeMultiplier, float PitchMultiplier);
 
 	/**
 	 * Play sound client-side at the specified location
@@ -720,9 +694,8 @@ public:
 	 * @param VolumeMultiplier - Volume multiplier to apply to the sound
 	 * @param PitchMultiplier - Pitch multiplier to apply to the sound
 	 */
-	UFUNCTION(unreliable, client="ClientPlaySoundAtLocation_Implementation")
+	UFUNCTION(unreliable, client)
 	void ClientPlaySoundAtLocation(class USoundBase* Sound, FVector Location, float VolumeMultiplier = 1.f, float PitchMultiplier = 1.f);
-	virtual void ClientPlaySoundAtLocation_Implementation(USoundBase* Sound, FVector Location, float VolumeMultiplier, float PitchMultiplier);
 
 	/** Asynchronously loads the given level in preparation for a streaming map transition.
 	 * the server sends one function per level name since dynamic arrays can't be replicated
@@ -730,9 +703,8 @@ public:
 	 * @param bFirst - whether this is the first item in the list (so clear the list first)
 	 * @param bLast - whether this is the last item in the list (so start preparing the change after receiving it)
 	 */
-	UFUNCTION(Reliable, Client="ClientPrepareMapChange_Implementation")
+	UFUNCTION(Reliable, Client)
 	void ClientPrepareMapChange(FName LevelName, bool bFirst, bool bLast);
-	virtual void ClientPrepareMapChange_Implementation(FName LevelName, bool bFirst, bool bLast);
 
 	/**
 	 * Forces the streaming system to disregard the normal logic for the specified duration and
@@ -742,27 +714,23 @@ public:
 	 * @param bEnableStreaming	- Whether to start (true) or stop (false) streaming
 	 * @param CinematicTextureGroups	- Bitfield indicating which texture groups that use extra high-resolution mips
 	 */
-	UFUNCTION(Reliable, Client="ClientPrestreamTextures_Implementation")
+	UFUNCTION(Reliable, Client)
 	void ClientPrestreamTextures(class AActor* ForcedActor, float ForceDuration, bool bEnableStreaming, int32 CinematicTextureGroups = 0);
-	virtual void ClientPrestreamTextures_Implementation(AActor* ForcedActor, float ForceDuration, bool bEnableStreaming, int32 CinematicTextureGroups);
 
 	/** Tell client to reset the PlayerController */
-	UFUNCTION(Reliable, Client="ClientReset_Implementation")
+	UFUNCTION(Reliable, Client)
 	void ClientReset();
-	virtual void ClientReset_Implementation();
 
 	/** Tell client to restart the level */
-	UFUNCTION(Reliable, Client="ClientRestart_Implementation")
+	UFUNCTION(Reliable, Client)
 	void ClientRestart(class APawn* NewPawn);
-	virtual void ClientRestart_Implementation(APawn* NewPawn);
 
 	/** 
 	 * Tells the client to block until all pending level streaming actions are complete.
 	 * Happens at the end of the tick primarily used to force update the client ASAP at join time.
 	 */
-	UFUNCTION(Reliable, Client="ClientSetBlockOnAsyncLoading_Implementation")
+	UFUNCTION(Reliable, Client)
 	void ClientSetBlockOnAsyncLoading();
-	virtual void ClientSetBlockOnAsyncLoading_Implementation();
 
 	/** Tell client to fade camera
 	 * @Param bEnableFading - true if we should apply FadeColor/FadeAmount to the screen
@@ -771,22 +739,19 @@ public:
 	 * @Param FadeTime - length of time for fade to occur over
 	 * @Param bFadeAudio - true to apply fading of audio alongside the video
 	 */
-	UFUNCTION(Reliable, Client="ClientSetCameraFade_Implementation")
+	UFUNCTION(Reliable, Client)
 	void ClientSetCameraFade(bool bEnableFading, FColor FadeColor = FColor(ForceInit), FVector2D FadeAlpha = FVector2D(ForceInit), float FadeTime = 0, bool bFadeAudio = false);
-	virtual void ClientSetCameraFade_Implementation(bool bEnableFading, FColor FadeColor, FVector2D FadeAlpha, float FadeTime, bool bFadeAudio);
 
 	/**
 	 * Replicated function to set camera style on client
 	 * @param	NewCamMode, name defining the new camera mode
 	 */
-	UFUNCTION(Reliable, Client="ClientSetCameraMode_Implementation")
+	UFUNCTION(Reliable, Client)
 	void ClientSetCameraMode(FName NewCamMode);
-	virtual void ClientSetCameraMode_Implementation(FName NewCamMode);
 
 	/** Called by the server to synchronize cinematic transitions with the client */
-	UFUNCTION(Reliable, Client="ClientSetCinematicMode_Implementation")
+	UFUNCTION(Reliable, Client)
 	void ClientSetCinematicMode(bool bInCinematicMode, bool bAffectsMovement, bool bAffectsTurning, bool bAffectsHUD);
-	virtual void ClientSetCinematicMode_Implementation(bool bInCinematicMode, bool bAffectsMovement, bool bAffectsTurning, bool bAffectsHUD);
 
 	/**
 	 * Forces the streaming system to disregard the normal logic for the specified duration and
@@ -796,14 +761,12 @@ public:
 	 * @param ForceDuration	- Number of seconds to keep all mip-levels in memory, disregarding the normal priority logic.
 	 * @param CinematicTextureGroups	- Bitfield indicating which texture groups that use extra high-resolution mips
 	 */
-	UFUNCTION(Reliable, Client="ClientSetForceMipLevelsToBeResident_Implementation")
+	UFUNCTION(Reliable, Client)
 	void ClientSetForceMipLevelsToBeResident(class UMaterialInterface* Material, float ForceDuration, int32 CinematicTextureGroups = 0);
-	virtual void ClientSetForceMipLevelsToBeResident_Implementation(UMaterialInterface* Material, float ForceDuration, int32 CinematicTextureGroups);
 
 	/** Set the client's class of HUD and spawns a new instance of it. If there was already a HUD active, it is destroyed. */
-	UFUNCTION(BlueprintCallable, Category="HUD", Reliable, Client="ClientSetHUD_Implementation")
+	UFUNCTION(BlueprintCallable, Category="HUD", Reliable, Client)
 	void ClientSetHUD(TSubclassOf<class AHUD> NewHUDClass);
-	virtual void ClientSetHUD_Implementation(TSubclassOf<AHUD>  NewHUDClass);
 
 	/** Helper to get the size of the HUD canvas for this player controller.  Returns 0 if there is no HUD */
 	UFUNCTION(BlueprintCallable, Category="HUD")
@@ -817,29 +780,24 @@ public:
 	 * @param A - new actor to set as view target
 	 * @param TransitionParams - parameters to use for controlling the transition
 	 */
-	UFUNCTION(Reliable, Client="ClientSetViewTarget_Implementation")
+	UFUNCTION(Reliable, Client)
 	void ClientSetViewTarget(class AActor* A, struct FViewTargetTransitionParams TransitionParams = FViewTargetTransitionParams());
-	virtual void ClientSetViewTarget_Implementation(AActor* A, FViewTargetTransitionParams TransitionParams);
 
 	/** Spawn a camera lens effect (e.g. blood). */
-	UFUNCTION(unreliable, client="ClientSpawnCameraLensEffect_Implementation", BlueprintCallable, Category="Game|Feedback")
+	UFUNCTION(unreliable, client, BlueprintCallable, Category="Game|Feedback")
 	void ClientSpawnCameraLensEffect(TSubclassOf<class AEmitterCameraLensEffectBase>  LensEffectEmitterClass);
-	virtual void ClientSpawnCameraLensEffect_Implementation(TSubclassOf<AEmitterCameraLensEffectBase>  LensEffectEmitterClass);
 
 	/** Removes all Camera Lens Effects. */
-	UFUNCTION(reliable, client="ClientClearCameraLensEffects_Implementation", BlueprintCallable, Category="Game|Feedback")
+	UFUNCTION(reliable, client, BlueprintCallable, Category="Game|Feedback")
 	virtual void ClientClearCameraLensEffects();
-	virtual void ClientClearCameraLensEffects_Implementation();
 
 	/** Stop camera animation on client. */
-	UFUNCTION(Reliable, Client="ClientStopCameraAnim_Implementation")
+	UFUNCTION(Reliable, Client)
 	void ClientStopCameraAnim(class UCameraAnim* AnimToStop);
-	virtual void ClientStopCameraAnim_Implementation(UCameraAnim* AnimToStop);
 
 	/** Stop camera shake on client.  */
-	UFUNCTION(reliable, client="ClientStopCameraShake_Implementation", BlueprintCallable, Category="Game|Feedback")
+	UFUNCTION(reliable, client, BlueprintCallable, Category="Game|Feedback")
 	void ClientStopCameraShake(TSubclassOf<class UCameraShake> Shake);
-	virtual void ClientStopCameraShake_Implementation(TSubclassOf<UCameraShake>  Shake);
 
 	/** 
 	 * Play a force feedback pattern on the player's controller
@@ -847,18 +805,16 @@ public:
 	 * @param	bLooping				Whether the pattern should be played repeatedly or be a single one shot
 	 * @param	Tag						A tag that allows stopping of an effect.  If another effect with this Tag is playing, it will be stopped and replaced
 	 */
-	UFUNCTION(unreliable, client="ClientPlayForceFeedback_Implementation", BlueprintCallable, Category="Game|Feedback")
+	UFUNCTION(unreliable, client, BlueprintCallable, Category="Game|Feedback")
 	void ClientPlayForceFeedback(class UForceFeedbackEffect* ForceFeedbackEffect, bool bLooping, FName Tag);
-	virtual void ClientPlayForceFeedback_Implementation(UForceFeedbackEffect* ForceFeedbackEffect, bool bLooping, FName Tag);
 
 	/** 
 	 * Stops a playing force feedback pattern
 	 * @param	ForceFeedbackEffect		If set only patterns from that effect will be stopped
 	 * @param	Tag						If not none only the pattern with this tag will be stopped
 	 */
-	UFUNCTION(reliable, client="ClientStopForceFeedback_Implementation", BlueprintCallable, Category="Game|Feedback")
+	UFUNCTION(reliable, client, BlueprintCallable, Category="Game|Feedback")
 	void ClientStopForceFeedback(class UForceFeedbackEffect* ForceFeedbackEffect, FName Tag);
-	virtual void ClientStopForceFeedback_Implementation(UForceFeedbackEffect* ForceFeedbackEffect, FName Tag);
 
 	/** 
 	 * Latent action that controls the playing of force feedback 
@@ -899,9 +855,8 @@ public:
 	 * @param MapPackageGuid	The GUID of the map package to travel to - this is used to find the file when it has been autodownloaded,
 	 * 							so it is only needed for clients
 	 */
-	UFUNCTION(Reliable, Client="ClientTravelInternal_Implementation")
+	UFUNCTION(Reliable, Client)
 	void ClientTravelInternal(const FString& URL, enum ETravelType TravelType, bool bSeamless = false, FGuid MapPackageGuid = FGuid());
-	virtual void ClientTravelInternal_Implementation(const FString& URL, ETravelType TravelType, bool bSeamless, FGuid MapPackageGuid);
 
 	/**
 	 * Replicated Update streaming status
@@ -911,97 +866,73 @@ public:
 	 * @param bNewShouldBlockOnLoad - Whether we want to force a blocking load
 	 * @param LODIndex				- Current LOD index for a streaming level
 	 */
-	UFUNCTION(Reliable, Client="ClientUpdateLevelStreamingStatus_Implementation")
+	UFUNCTION(Reliable, Client)
 	void ClientUpdateLevelStreamingStatus(FName PackageName, bool bNewShouldBeLoaded, bool bNewShouldBeVisible, bool bNewShouldBlockOnLoad, int32 LODIndex);
-	virtual void ClientUpdateLevelStreamingStatus_Implementation(FName PackageName, bool bNewShouldBeLoaded, bool bNewShouldBeVisible, bool bNewShouldBlockOnLoad, int32 LODIndex);
 
 	/** Notify client they were kicked from the server */
-	UFUNCTION(Reliable, Client="ClientWasKicked_Implementation")
+	UFUNCTION(Reliable, Client)
 	void ClientWasKicked(const FText& KickReason);
-	virtual void ClientWasKicked_Implementation(const FText& KickReason);
 
 	/** Assign Pawn to player, but avoid calling ClientRestart if we have already accepted this pawn */
-	UFUNCTION(Reliable, Client="ClientRetryClientRestart_Implementation")
+	UFUNCTION(Reliable, Client)
 	void ClientRetryClientRestart(class APawn* NewPawn);
-	virtual void ClientRetryClientRestart_Implementation(APawn* NewPawn);
 
 	/** Call ClientRetryClientRestart, but only if the current pawn is not the currently acknowledged pawn (and throttled to avoid saturating the network). */
 	virtual void SafeRetryClientRestart();
 
 	/** send client localized message id */
-	UFUNCTION(Reliable, Client="ClientReceiveLocalizedMessage_Implementation")
+	UFUNCTION(Reliable, Client)
 	void ClientReceiveLocalizedMessage(TSubclassOf<ULocalMessage> Message, int32 Switch = 0, class APlayerState* RelatedPlayerState_1 = NULL, class APlayerState* RelatedPlayerState_2 = NULL, class UObject* OptionalObject = NULL);
-	virtual void ClientReceiveLocalizedMessage_Implementation(TSubclassOf<ULocalMessage>  Message, int32 Switch, APlayerState* RelatedPlayerState_1, APlayerState* RelatedPlayerState_2, UObject* OptionalObject);
 
 	/** acknowledge possession of pawn */
-	UFUNCTION(reliable, server="ServerAcknowledgePossession_Implementation", WithValidation="ServerAcknowledgePossession_Validate")
+	UFUNCTION(reliable, server, WithValidation)
 	void ServerAcknowledgePossession(class APawn* P);
-	virtual void ServerAcknowledgePossession_Implementation(APawn* P);
-	virtual bool ServerAcknowledgePossession_Validate(APawn* );
 
 	/** change mode of camera */
-	UFUNCTION(reliable, server="ServerCamera_Implementation", WithValidation="ServerCamera_Validate")
+	UFUNCTION(reliable, server, WithValidation)
 	void ServerCamera(FName NewMode);
-	virtual void ServerCamera_Implementation(FName NewMode);
-	virtual bool ServerCamera_Validate(FName );
 
 	/** Change name of server */
-	UFUNCTION(reliable, server="ServerChangeName_Implementation", WithValidation="ServerChangeName_Validate")
+	UFUNCTION(reliable, server, WithValidation)
 	void ServerChangeName(const FString& S);
-	virtual void ServerChangeName_Implementation(const FString& S);
-	virtual bool ServerChangeName_Validate(const FString& );
 
 	/** 
 	 * Called to notify the server when the client has loaded a new world via seamless traveling
 	 * @param WorldPackageName the name of the world package that was loaded
 	 */
-	UFUNCTION(reliable, server="ServerNotifyLoadedWorld_Implementation", WithValidation="ServerNotifyLoadedWorld_Validate", SealedEvent)
+	UFUNCTION(reliable, server, WithValidation, SealedEvent)
 	void ServerNotifyLoadedWorld(FName WorldPackageName);
-	void ServerNotifyLoadedWorld_Implementation(FName WorldPackageName);
-	bool ServerNotifyLoadedWorld_Validate(FName );
 
 	/** Replicate pause request to the server */
-	UFUNCTION(reliable, server="ServerPause_Implementation", WithValidation="ServerPause_Validate")
+	UFUNCTION(reliable, server, WithValidation)
 	void ServerPause();
-	virtual void ServerPause_Implementation();
-	virtual bool ServerPause_Validate();
 
 	/** Attempts to restart this player, generally called from the client upon respawn request. */
-	UFUNCTION(reliable, server="ServerRestartPlayer_Implementation", WithValidation="ServerRestartPlayer_Validate")
+	UFUNCTION(reliable, server, WithValidation)
 	void ServerRestartPlayer();
-	virtual void ServerRestartPlayer_Implementation();
-	virtual bool ServerRestartPlayer_Validate();
 
 	/** When spectating, pings the server to make sure spectating should continue. */
-	UFUNCTION(unreliable, server="ServerSetSpectatorLocation_Implementation", WithValidation="ServerSetSpectatorLocation_Validate")
+	UFUNCTION(unreliable, server, WithValidation)
 	void ServerSetSpectatorLocation(FVector NewLoc);
-	virtual void ServerSetSpectatorLocation_Implementation(FVector NewLoc);
-	virtual bool ServerSetSpectatorLocation_Validate(FVector );
 
 	/** Calls ServerSetSpectatorLocation but throttles it to reduce bandwidth and only calls it when necessary. */
 	void SafeServerUpdateSpectatorState();
 
 	/** Tells the server to make sure the possessed pawn is in sync with the client. */
-	UFUNCTION(unreliable, server="ServerCheckClientPossession_Implementation", WithValidation="ServerCheckClientPossession_Validate")
+	UFUNCTION(unreliable, server, WithValidation)
 	void ServerCheckClientPossession();
-	virtual void ServerCheckClientPossession_Implementation();
-	virtual bool ServerCheckClientPossession_Validate();
 
 	/** Call ServerCheckClientPossession on the server, but only if the current pawn is not the acknowledged pawn (and throttled to avoid saturating the network). */
 	virtual void SafeServerCheckClientPossession();
 
 	/** Notifies the server that the client has ticked gameplay code, and should no longer get the extended "still loading" timeout grace period */
-	UFUNCTION(reliable, server="ServerShortTimeout_Implementation", WithValidation="ServerShortTimeout_Validate")
+	UFUNCTION(reliable, server, WithValidation)
 	void ServerShortTimeout();
-	virtual void ServerShortTimeout_Implementation();
-	virtual bool ServerShortTimeout_Validate();
 
 	/** If PlayerCamera.bUseClientSideCameraUpdates is set, client will replicate camera positions to the server. */
 	// @TODO - combine pitch/yaw into one int, maybe also send location compressed
-	UFUNCTION(unreliable, server="ServerUpdateCamera_Implementation", WithValidation="ServerUpdateCamera_Validate")
+	UFUNCTION(unreliable, server, WithValidation)
 	void ServerUpdateCamera(FVector_NetQuantize CamLoc, int32 CamPitchAndYaw);
-	virtual void ServerUpdateCamera_Implementation(FVector_NetQuantize CamLoc, int32 CamPitchAndYaw);
-	virtual bool ServerUpdateCamera_Validate(FVector_NetQuantize , int32 );
 
 	/** 
 	 * Called when the client adds/removes a streamed level
@@ -1009,45 +940,32 @@ public:
 	 * Actors the client has not initialized
 	 * @param PackageName the name of the package for the level whose status changed
 	 */
-	UFUNCTION(reliable, server="ServerUpdateLevelVisibility_Implementation", WithValidation="ServerUpdateLevelVisibility_Validate", SealedEvent)
+	UFUNCTION(reliable, server, WithValidation, SealedEvent)
 	void ServerUpdateLevelVisibility(FName PackageName, bool bIsVisible);
-	void ServerUpdateLevelVisibility_Implementation(FName PackageName, bool bIsVisible);
-	bool ServerUpdateLevelVisibility_Validate(FName , bool );
 
 	/** Used by client to request server to confirm current viewtarget (server will respond with ClientSetViewTarget() ). */
-	UFUNCTION(reliable, server="ServerVerifyViewTarget_Implementation", WithValidation="ServerVerifyViewTarget_Validate")
+	UFUNCTION(reliable, server, WithValidation)
 	void ServerVerifyViewTarget();
-	virtual void ServerVerifyViewTarget_Implementation();
-	virtual bool ServerVerifyViewTarget_Validate();
 
 	/** Move camera to next player on round ended or spectating*/
-	UFUNCTION(unreliable, server="ServerViewNextPlayer_Implementation", WithValidation="ServerViewNextPlayer_Validate")
+	UFUNCTION(unreliable, server, WithValidation)
 	void ServerViewNextPlayer();
-	virtual void ServerViewNextPlayer_Implementation();
-	virtual bool ServerViewNextPlayer_Validate();
 
 	/** Move camera to previous player on round ended or spectating */
-	UFUNCTION(unreliable, server="ServerViewPrevPlayer_Implementation", WithValidation="ServerViewPrevPlayer_Validate")
+	UFUNCTION(unreliable, server, WithValidation)
 	void ServerViewPrevPlayer();
-	virtual void ServerViewPrevPlayer_Implementation();
-	virtual bool ServerViewPrevPlayer_Validate();
 
 	/** Move camera to current user */
-	UFUNCTION(unreliable, server="ServerViewSelf_Implementation", WithValidation="ServerViewSelf_Validate")
+	UFUNCTION(unreliable, server, WithValidation)
 	void ServerViewSelf(struct FViewTargetTransitionParams TransitionParams = FViewTargetTransitionParams());
-	virtual void ServerViewSelf_Implementation(FViewTargetTransitionParams TransitionParams);
-	virtual bool ServerViewSelf_Validate(FViewTargetTransitionParams );
 
 	/** @todo document */
-	UFUNCTION(Reliable, Client="ClientTeamMessage_Implementation")
+	UFUNCTION(Reliable, Client)
 	void ClientTeamMessage(class APlayerState* SenderPlayerState, const FString& S, FName Type, float MsgLifeTime = 0);
-	virtual void ClientTeamMessage_Implementation(APlayerState* SenderPlayerState, const FString& S, FName Type, float MsgLifeTime);
 
 	/** Used by UGameplayDebuggingControllerComponent to replicate messages for AI debugging in network games. */
-	UFUNCTION(reliable, server="ServerToggleAILogging_Implementation", WithValidation="ServerToggleAILogging_Validate")
+	UFUNCTION(reliable, server, WithValidation)
 	void ServerToggleAILogging();
-	virtual void ServerToggleAILogging_Implementation();
-	virtual bool ServerToggleAILogging_Validate();
 
 	/** Add Pitch (look up) input */
 	UFUNCTION(BlueprintCallable, Category="Game|Player", meta=(Keywords="up down"))
@@ -1494,9 +1412,8 @@ public:
 	/**
 	 * Notify from server that Visual Logger is recording, to show that information on client about possible performance issues 
 	 */
-	UFUNCTION(Reliable, Client="OnServerStartedVisualLogger_Implementation")
+	UFUNCTION(Reliable, Client)
 	void OnServerStartedVisualLogger(bool bIsLogging);
-	virtual void OnServerStartedVisualLogger_Implementation(bool bIsLogging);
 
 protected:
 
diff --git a/Engine/Source/Runtime/Engine/Classes/GameFramework/PlayerInput.h b/Engine/Source/Runtime/Engine/Classes/GameFramework/PlayerInput.h
--- a/Engine/Source/Runtime/Engine/Classes/GameFramework/PlayerInput.h
+++ b/Engine/Source/Runtime/Engine/Classes/GameFramework/PlayerInput.h
@@ -19,7 +19,7 @@ ENGINE_API DECLARE_LOG_CATEGORY_EXTERN(LogInput, Log, All);
 USTRUCT()
 struct FKeyBind
 {
-	GENERATED_BODY()
+	GENERATED_USTRUCT_BODY()
 
 	/** The key to be bound to the command */
 	UPROPERTY(config)
@@ -70,7 +70,7 @@ struct FKeyBind
 USTRUCT()
 struct FInputAxisProperties
 {
-	GENERATED_BODY()
+	GENERATED_USTRUCT_BODY()
 
 	/** What the dead zone of the axis is.  For control axes such as analog sticks. */
 	UPROPERTY(EditAnywhere, Category="Input")
@@ -101,7 +101,7 @@ struct FInputAxisProperties
 USTRUCT()
 struct FInputAxisConfigEntry
 {
-	GENERATED_BODY()
+	GENERATED_USTRUCT_BODY()
 
 	/** Axis Key these properties apply to */
 	UPROPERTY(VisibleAnywhere, Category="Input")
@@ -120,7 +120,7 @@ struct FInputAxisConfigEntry
 USTRUCT()
 struct FInputActionKeyMapping
 {
-	GENERATED_BODY()
+	GENERATED_USTRUCT_BODY()
 
 	/** Friendly name of action, e.g "jump" */
 	UPROPERTY(EditAnywhere, Category="Input")
@@ -188,7 +188,7 @@ struct FInputActionKeyMapping
 USTRUCT()
 struct FInputAxisKeyMapping
 {
-	GENERATED_BODY()
+	GENERATED_USTRUCT_BODY()
 
 	/** Friendly name of axis, e.g "MoveForward" */
 	UPROPERTY(EditAnywhere, Category="Input")
@@ -247,9 +247,7 @@ struct FInputAxisKeyMapping
 UCLASS(Within=PlayerController, config=Input, transient)
 class ENGINE_API UPlayerInput : public UObject
 {
-	GENERATED_BODY()
-public:
-	UPlayerInput(const FObjectInitializer& ObjectInitializer = FObjectInitializer::Get());
+	GENERATED_UCLASS_BODY()
 
 public:
 	// NOTE: These touch vectors are calculated and set directly, they do not go through the .ini Bindings
diff --git a/Engine/Source/Runtime/Engine/Classes/GameFramework/PlayerStart.h b/Engine/Source/Runtime/Engine/Classes/GameFramework/PlayerStart.h
--- a/Engine/Source/Runtime/Engine/Classes/GameFramework/PlayerStart.h
+++ b/Engine/Source/Runtime/Engine/Classes/GameFramework/PlayerStart.h
@@ -12,11 +12,9 @@
 UCLASS(ClassGroup=Common, hidecategories=Collision)
 class ENGINE_API APlayerStart : public ANavigationObjectBase
 {
-	GENERATED_BODY()
-public:
-	APlayerStart(const FObjectInitializer& ObjectInitializer = FObjectInitializer::Get());
 	/** To take more control over PlayerStart selection, you can override the virtual AGameMode::FindPlayerStart and AGameMode::ChoosePlayerStart function */
 
+	GENERATED_UCLASS_BODY()
 
 	/** Used when searching for which playerstart to use. */
 	UPROPERTY(EditAnywhere, BlueprintReadWrite, Category=Object)
diff --git a/Engine/Source/Runtime/Engine/Classes/GameFramework/PlayerState.h b/Engine/Source/Runtime/Engine/Classes/GameFramework/PlayerState.h
--- a/Engine/Source/Runtime/Engine/Classes/GameFramework/PlayerState.h
+++ b/Engine/Source/Runtime/Engine/Classes/GameFramework/PlayerState.h
@@ -31,9 +31,7 @@ struct PingAvgData
 UCLASS(BlueprintType, Blueprintable, notplaceable)
 class ENGINE_API APlayerState : public AInfo
 {
-	GENERATED_BODY()
-public:
-	APlayerState(const FObjectInitializer& ObjectInitializer = FObjectInitializer::Get());
+	GENERATED_UCLASS_BODY()
 
 	/** Player's current score. */
 	UPROPERTY(replicatedUsing=OnRep_Score, BlueprintReadOnly, Category=PlayerState)
diff --git a/Engine/Source/Runtime/Engine/Classes/GameFramework/ProjectileMovementComponent.h b/Engine/Source/Runtime/Engine/Classes/GameFramework/ProjectileMovementComponent.h
--- a/Engine/Source/Runtime/Engine/Classes/GameFramework/ProjectileMovementComponent.h
+++ b/Engine/Source/Runtime/Engine/Classes/GameFramework/ProjectileMovementComponent.h
@@ -18,9 +18,7 @@
 UCLASS(ClassGroup=Movement, meta=(BlueprintSpawnableComponent), ShowCategories=(Velocity))
 class ENGINE_API UProjectileMovementComponent : public UMovementComponent
 {
-	GENERATED_BODY()
-public:
-	UProjectileMovementComponent(const FObjectInitializer& ObjectInitializer = FObjectInitializer::Get());
+	GENERATED_UCLASS_BODY()
 
 	DECLARE_DYNAMIC_MULTICAST_DELEGATE_TwoParams( FOnProjectileBounceDelegate, const FHitResult&, ImpactResult, const FVector&, ImpactVelocity );
 	DECLARE_DYNAMIC_MULTICAST_DELEGATE_OneParam( FOnProjectileStopDelegate, const FHitResult&, ImpactResult );
diff --git a/Engine/Source/Runtime/Engine/Classes/GameFramework/RotatingMovementComponent.h b/Engine/Source/Runtime/Engine/Classes/GameFramework/RotatingMovementComponent.h
--- a/Engine/Source/Runtime/Engine/Classes/GameFramework/RotatingMovementComponent.h
+++ b/Engine/Source/Runtime/Engine/Classes/GameFramework/RotatingMovementComponent.h
@@ -12,9 +12,7 @@
 UCLASS(ClassGroup=Movement, meta=(BlueprintSpawnableComponent), HideCategories=(Velocity))
 class ENGINE_API URotatingMovementComponent : public UMovementComponent
 {
-	GENERATED_BODY()
-public:
-	URotatingMovementComponent(const FObjectInitializer& ObjectInitializer = FObjectInitializer::Get());
+	GENERATED_UCLASS_BODY()
 
 	/** How fast to update roll/pitch/yaw of the component we update. */
 	UPROPERTY(EditAnywhere, BlueprintReadWrite, Category=RotatingComponent)
diff --git a/Engine/Source/Runtime/Engine/Classes/GameFramework/SaveGame.h b/Engine/Source/Runtime/Engine/Classes/GameFramework/SaveGame.h
--- a/Engine/Source/Runtime/Engine/Classes/GameFramework/SaveGame.h
+++ b/Engine/Source/Runtime/Engine/Classes/GameFramework/SaveGame.h
@@ -14,9 +14,6 @@
 UCLASS(abstract, Blueprintable, BlueprintType)
 class ENGINE_API USaveGame : public UObject
 {
-	GENERATED_BODY()
-public:
-	USaveGame(const FObjectInitializer& ObjectInitializer = FObjectInitializer::Get());
 	/**
 	 *	@see UGameplayStatics::CreateSaveGameObject
 	 *	@see UGameplayStatics::SaveGameToSlot
@@ -25,4 +22,5 @@ public:
 	 *	@see UGameplayStatics::DeleteGameInSlot
 	 */
 
+	GENERATED_UCLASS_BODY()
 };
diff --git a/Engine/Source/Runtime/Engine/Classes/GameFramework/SpectatorPawnMovement.h b/Engine/Source/Runtime/Engine/Classes/GameFramework/SpectatorPawnMovement.h
--- a/Engine/Source/Runtime/Engine/Classes/GameFramework/SpectatorPawnMovement.h
+++ b/Engine/Source/Runtime/Engine/Classes/GameFramework/SpectatorPawnMovement.h
@@ -11,9 +11,7 @@
 UCLASS()
 class ENGINE_API USpectatorPawnMovement : public UFloatingPawnMovement
 {
-	GENERATED_BODY()
-public:
-	USpectatorPawnMovement(const FObjectInitializer& ObjectInitializer = FObjectInitializer::Get());
+	GENERATED_UCLASS_BODY()
 
 	/** If true, component moves at full speed no matter the time dilation. Default is false. */
 	UPROPERTY()
diff --git a/Engine/Source/Runtime/Engine/Classes/GameFramework/SpringArmComponent.h b/Engine/Source/Runtime/Engine/Classes/GameFramework/SpringArmComponent.h
--- a/Engine/Source/Runtime/Engine/Classes/GameFramework/SpringArmComponent.h
+++ b/Engine/Source/Runtime/Engine/Classes/GameFramework/SpringArmComponent.h
@@ -14,9 +14,7 @@
 UCLASS(ClassGroup=Camera, meta=(BlueprintSpawnableComponent), hideCategories=(Mobility))
 class ENGINE_API USpringArmComponent : public USceneComponent
 {
-	GENERATED_BODY()
-public:
-	USpringArmComponent(const FObjectInitializer& ObjectInitializer = FObjectInitializer::Get());
+	GENERATED_UCLASS_BODY()
 
 	/** Natural length of the spring arm when there are no collisions */
 	UPROPERTY(EditAnywhere, BlueprintReadWrite, Category=Camera)
diff --git a/Engine/Source/Runtime/Engine/Classes/GameFramework/TouchInterface.h b/Engine/Source/Runtime/Engine/Classes/GameFramework/TouchInterface.h
--- a/Engine/Source/Runtime/Engine/Classes/GameFramework/TouchInterface.h
+++ b/Engine/Source/Runtime/Engine/Classes/GameFramework/TouchInterface.h
@@ -9,7 +9,7 @@ class SVirtualJoystick;
 USTRUCT()
 struct FTouchInputControl
 {
-	GENERATED_BODY()
+	GENERATED_USTRUCT_BODY()
 
 	// basically mirroring SVirtualJoystick::FControlInfo but as an editable class
 	UPROPERTY(EditAnywhere, Category="Control", meta=(ToolTip="For sticks, this is the Thumb"))
@@ -44,9 +44,7 @@ struct FTouchInputControl
 UCLASS(Blueprintable, BlueprintType)
 class ENGINE_API UTouchInterface : public UObject
 {
-	GENERATED_BODY()
-public:
-	UTouchInterface(const FObjectInitializer& ObjectInitializer = FObjectInitializer::Get());
+	GENERATED_UCLASS_BODY()
 
 	UPROPERTY(EditAnywhere, Category="TouchInterface")
 	TArray<FTouchInputControl> Controls;
diff --git a/Engine/Source/Runtime/Engine/Classes/GameFramework/Volume.h b/Engine/Source/Runtime/Engine/Classes/GameFramework/Volume.h
--- a/Engine/Source/Runtime/Engine/Classes/GameFramework/Volume.h
+++ b/Engine/Source/Runtime/Engine/Classes/GameFramework/Volume.h
@@ -13,9 +13,7 @@ DECLARE_LOG_CATEGORY_EXTERN(LogVolume, Log, All);
 UCLASS(showcategories=Collision, hidecategories=(Brush, Physics), abstract, ConversionRoot)
 class ENGINE_API AVolume : public ABrush
 {
-	GENERATED_BODY()
-public:
-	AVolume(const FObjectInitializer& ObjectInitializer = FObjectInitializer::Get());
+	GENERATED_UCLASS_BODY()
 
 #if WITH_EDITOR
 	/** Delegate used for notifications when a volumes initial shape changes */
diff --git a/Engine/Source/Runtime/Engine/Classes/GameFramework/WheeledVehicle.h b/Engine/Source/Runtime/Engine/Classes/GameFramework/WheeledVehicle.h
--- a/Engine/Source/Runtime/Engine/Classes/GameFramework/WheeledVehicle.h
+++ b/Engine/Source/Runtime/Engine/Classes/GameFramework/WheeledVehicle.h
@@ -19,9 +19,7 @@ class FDebugDisplayInfo;
 UCLASS(abstract, config=Game, BlueprintType)
 class ENGINE_API AWheeledVehicle : public APawn
 {
-	GENERATED_BODY()
-public:
-	AWheeledVehicle(const FObjectInitializer& ObjectInitializer = FObjectInitializer::Get());
+	GENERATED_UCLASS_BODY()
 
 private_subobject:
 	/**  The main skeletal mesh associated with this Vehicle */
diff --git a/Engine/Source/Runtime/Engine/Classes/GameFramework/WorldSettings.h b/Engine/Source/Runtime/Engine/Classes/GameFramework/WorldSettings.h
--- a/Engine/Source/Runtime/Engine/Classes/GameFramework/WorldSettings.h
+++ b/Engine/Source/Runtime/Engine/Classes/GameFramework/WorldSettings.h
@@ -18,7 +18,7 @@ enum EVisibilityAggressiveness
 USTRUCT()
 struct FGameModePrefix
 {
-	GENERATED_BODY()
+	GENERATED_USTRUCT_BODY()
 
 	/** map prefix, e.g. "DM" */
 	UPROPERTY()
@@ -32,7 +32,7 @@ struct FGameModePrefix
 USTRUCT()
 struct FLightmassWorldInfoSettings
 {
-	GENERATED_BODY()
+	GENERATED_USTRUCT_BODY()
 
 	/** 
 	 * Warning: Setting this to less than 1 will greatly increase build times!
@@ -162,7 +162,7 @@ struct FLightmassWorldInfoSettings
 USTRUCT()
 struct ENGINE_API FNetViewer
 {
-	GENERATED_BODY()
+	GENERATED_USTRUCT_BODY()
 
 	/** The player controller corresponding to this view */
 	UPROPERTY()
@@ -195,7 +195,7 @@ struct ENGINE_API FNetViewer
 USTRUCT()
 struct ENGINE_API FHierarchicalSimplification
 {
-	GENERATED_BODY()
+	GENERATED_USTRUCT_BODY()
 
 	/** If this is true, it will simplify mesh. However it is slower. If false, it will just merge actors but not simplify */
 	UPROPERTY(Category=FHierarchicalSimplification, EditAnywhere)
@@ -233,9 +233,7 @@ struct ENGINE_API FHierarchicalSimplification
 UCLASS(config=game, hidecategories=(Actor, Advanced, Display, Events, Object, Attachment, Info, Input, Blueprint, Layers, Tags, Replication), showcategories=("Input|MouseInput", "Input|TouchInput"), notplaceable)
 class ENGINE_API AWorldSettings : public AInfo, public IInterface_AssetUserData
 {
-	GENERATED_BODY()
-public:
-	AWorldSettings(const FObjectInitializer& ObjectInitializer = FObjectInitializer::Get());
+	GENERATED_UCLASS_BODY()
 
 	/** DEFAULT BASIC PHYSICS SETTINGS **/
 
diff --git a/Engine/Source/Runtime/Engine/Classes/Interfaces/Interface_CollisionDataProvider.h b/Engine/Source/Runtime/Engine/Classes/Interfaces/Interface_CollisionDataProvider.h
--- a/Engine/Source/Runtime/Engine/Classes/Interfaces/Interface_CollisionDataProvider.h
+++ b/Engine/Source/Runtime/Engine/Classes/Interfaces/Interface_CollisionDataProvider.h
@@ -12,7 +12,7 @@
 USTRUCT()
 struct FTriIndices
 {
-	GENERATED_BODY()
+	GENERATED_USTRUCT_BODY()
 
 	UPROPERTY()
 	int32 v0;
@@ -37,7 +37,7 @@ struct FTriIndices
 USTRUCT()
 struct FTriMeshCollisionData
 {
-	GENERATED_BODY()
+	GENERATED_USTRUCT_BODY()
 
 	/** Array of vertices included in the triangle mesh */
 	UPROPERTY(transient)
@@ -64,15 +64,12 @@ struct FTriMeshCollisionData
 UINTERFACE(meta=(CannotImplementInterfaceInBlueprint))
 class UInterface_CollisionDataProvider : public UInterface
 {
-	GENERATED_BODY()
-public:
-	UInterface_CollisionDataProvider(const FObjectInitializer& ObjectInitializer = FObjectInitializer::Get());
+	GENERATED_UINTERFACE_BODY()
 };
 
 class IInterface_CollisionDataProvider
 {
-	GENERATED_BODY()
-public:
+	GENERATED_IINTERFACE_BODY()
 
 
 	/**	 Interface for retrieving triangle mesh collision data from the implementing object 
diff --git a/Engine/Source/Runtime/Engine/Classes/Interfaces/Interface_PostProcessVolume.h b/Engine/Source/Runtime/Engine/Classes/Interfaces/Interface_PostProcessVolume.h
--- a/Engine/Source/Runtime/Engine/Classes/Interfaces/Interface_PostProcessVolume.h
+++ b/Engine/Source/Runtime/Engine/Classes/Interfaces/Interface_PostProcessVolume.h
@@ -7,9 +7,7 @@
 UINTERFACE()
 class UInterface_PostProcessVolume : public UInterface
 {
-	GENERATED_BODY()
-public:
-	UInterface_PostProcessVolume(const FObjectInitializer& ObjectInitializer = FObjectInitializer::Get());
+	GENERATED_UINTERFACE_BODY()
 };
 
 struct FPostProcessVolumeProperties
@@ -24,8 +22,7 @@ struct FPostProcessVolumeProperties
 
 class IInterface_PostProcessVolume
 {
-	GENERATED_BODY()
-public:
+	GENERATED_IINTERFACE_BODY()
 
 	ENGINE_API virtual bool EncompassesPoint(FVector Point, float SphereRadius/*=0.f*/, float* OutDistanceToPoint) = 0;
 	ENGINE_API virtual FPostProcessVolumeProperties GetProperties() const = 0;
diff --git a/Engine/Source/Runtime/Engine/Classes/Interfaces/NetworkPredictionInterface.h b/Engine/Source/Runtime/Engine/Classes/Interfaces/NetworkPredictionInterface.h
--- a/Engine/Source/Runtime/Engine/Classes/Interfaces/NetworkPredictionInterface.h
+++ b/Engine/Source/Runtime/Engine/Classes/Interfaces/NetworkPredictionInterface.h
@@ -12,15 +12,12 @@
 UINTERFACE(MinimalAPI, meta=(CannotImplementInterfaceInBlueprint))
 class UNetworkPredictionInterface : public UInterface
 {
-	GENERATED_BODY()
-public:
-	ENGINE_API UNetworkPredictionInterface(const FObjectInitializer& ObjectInitializer = FObjectInitializer::Get());
+	GENERATED_UINTERFACE_BODY()
 };
 
 class ENGINE_API INetworkPredictionInterface
 {
-	GENERATED_BODY()
-public:
+	GENERATED_IINTERFACE_BODY()
 
 	//--------------------------------
 	// Server hooks
diff --git a/Engine/Source/Runtime/Engine/Classes/Intrinsic/Model.h b/Engine/Source/Runtime/Engine/Classes/Intrinsic/Model.h
--- a/Engine/Source/Runtime/Engine/Classes/Intrinsic/Model.h
+++ b/Engine/Source/Runtime/Engine/Classes/Intrinsic/Model.h
@@ -5,9 +5,8 @@
 UCLASS(noexport, Intrinsic)
 class UModel : public UObject
 {
-	GENERATED_BODY()
-public:
-	UModel(const FObjectInitializer& ObjectInitializer = FObjectInitializer::Get());
+	GENERATED_UCLASS_BODY()
+
 };
 
 #endif
diff --git a/Engine/Source/Runtime/Engine/Classes/Kismet/BlueprintAsyncActionBase.h b/Engine/Source/Runtime/Engine/Classes/Kismet/BlueprintAsyncActionBase.h
--- a/Engine/Source/Runtime/Engine/Classes/Kismet/BlueprintAsyncActionBase.h
+++ b/Engine/Source/Runtime/Engine/Classes/Kismet/BlueprintAsyncActionBase.h
@@ -6,9 +6,7 @@
 UCLASS()
 class ENGINE_API UBlueprintAsyncActionBase : public UObject
 {
-	GENERATED_BODY()
-public:
-	UBlueprintAsyncActionBase(const FObjectInitializer& ObjectInitializer = FObjectInitializer::Get());
+	GENERATED_UCLASS_BODY()
 
 	// Called to trigger the action once the delegates have been bound
 	UFUNCTION(BlueprintCallable, meta=(BlueprintInternalUseOnly="true"))
diff --git a/Engine/Source/Runtime/Engine/Classes/Kismet/BlueprintFunctionLibrary.h b/Engine/Source/Runtime/Engine/Classes/Kismet/BlueprintFunctionLibrary.h
--- a/Engine/Source/Runtime/Engine/Classes/Kismet/BlueprintFunctionLibrary.h
+++ b/Engine/Source/Runtime/Engine/Classes/Kismet/BlueprintFunctionLibrary.h
@@ -8,9 +8,7 @@
 UCLASS(Abstract, MinimalAPI)
 class UBlueprintFunctionLibrary : public UObject
 {
-	GENERATED_BODY()
-public:
-	ENGINE_API UBlueprintFunctionLibrary(const FObjectInitializer& ObjectInitializer = FObjectInitializer::Get());
+	GENERATED_UCLASS_BODY()
 
 	// UObject interface
 	ENGINE_API virtual int32 GetFunctionCallspace(UFunction* Function, void* Parms, FFrame* Stack) override;
diff --git a/Engine/Source/Runtime/Engine/Classes/Kismet/DataTableFunctionLibrary.h b/Engine/Source/Runtime/Engine/Classes/Kismet/DataTableFunctionLibrary.h
--- a/Engine/Source/Runtime/Engine/Classes/Kismet/DataTableFunctionLibrary.h
+++ b/Engine/Source/Runtime/Engine/Classes/Kismet/DataTableFunctionLibrary.h
@@ -23,9 +23,7 @@ namespace EEvaluateCurveTableResult
 UCLASS(MinimalAPI)
 class UDataTableFunctionLibrary : public UBlueprintFunctionLibrary
 {
-	GENERATED_BODY()
-public:
-	ENGINE_API UDataTableFunctionLibrary(const FObjectInitializer& ObjectInitializer = FObjectInitializer::Get());
+	GENERATED_UCLASS_BODY()
 
 	UFUNCTION(BlueprintCallable, Category = "DataTable", meta = (ExpandEnumAsExecs="OutResult", DataTablePin="CurveTable"))
 	static void EvaluateCurveTableRow(UCurveTable* CurveTable, FName RowName, float InXY, TEnumAsByte<EEvaluateCurveTableResult::Type>& OutResult, float& OutXY);
diff --git a/Engine/Source/Runtime/Engine/Classes/Kismet/GameplayStatics.h b/Engine/Source/Runtime/Engine/Classes/Kismet/GameplayStatics.h
--- a/Engine/Source/Runtime/Engine/Classes/Kismet/GameplayStatics.h
+++ b/Engine/Source/Runtime/Engine/Classes/Kismet/GameplayStatics.h
@@ -1,4 +1,4 @@
-// Copyright 1998-2015 Epic Games, Inc. All Rights Reserved.
+﻿// Copyright 1998-2015 Epic Games, Inc. All Rights Reserved.
 
 #pragma once
 #include "Engine/LatentActionManager.h"
@@ -28,9 +28,7 @@ namespace ESuggestProjVelocityTraceOption
 UCLASS()
 class ENGINE_API UGameplayStatics : public UBlueprintFunctionLibrary
 {
-	GENERATED_BODY()
-public:
-	UGameplayStatics(const FObjectInitializer& ObjectInitializer = FObjectInitializer::Get());
+	GENERATED_UCLASS_BODY()
 
 	// --- Spawning functions ------------------------------
 
diff --git a/Engine/Source/Runtime/Engine/Classes/Kismet/HeadMountedDisplayFunctionLibrary.h b/Engine/Source/Runtime/Engine/Classes/Kismet/HeadMountedDisplayFunctionLibrary.h
--- a/Engine/Source/Runtime/Engine/Classes/Kismet/HeadMountedDisplayFunctionLibrary.h
+++ b/Engine/Source/Runtime/Engine/Classes/Kismet/HeadMountedDisplayFunctionLibrary.h
@@ -17,9 +17,7 @@ namespace EOrientPositionSelector
 UCLASS()
 class UHeadMountedDisplayFunctionLibrary : public UBlueprintFunctionLibrary
 {
-	GENERATED_BODY()
-public:
-	UHeadMountedDisplayFunctionLibrary(const FObjectInitializer& ObjectInitializer = FObjectInitializer::Get());
+	GENERATED_UCLASS_BODY()
 
 	/**
 	 * Returns whether or not we are currently using the head mounted display.
diff --git a/Engine/Source/Runtime/Engine/Classes/Kismet/KismetArrayLibrary.h b/Engine/Source/Runtime/Engine/Classes/Kismet/KismetArrayLibrary.h
--- a/Engine/Source/Runtime/Engine/Classes/Kismet/KismetArrayLibrary.h
+++ b/Engine/Source/Runtime/Engine/Classes/Kismet/KismetArrayLibrary.h
@@ -6,9 +6,7 @@
 UCLASS(MinimalAPI)
 class UKismetArrayLibrary : public UBlueprintFunctionLibrary
 {
-	GENERATED_BODY()
-public:
-	ENGINE_API UKismetArrayLibrary(const FObjectInitializer& ObjectInitializer = FObjectInitializer::Get());
+	GENERATED_UCLASS_BODY()
 	/** 
 	 *Add item to array
 	 *
diff --git a/Engine/Source/Runtime/Engine/Classes/Kismet/KismetGuidLibrary.h b/Engine/Source/Runtime/Engine/Classes/Kismet/KismetGuidLibrary.h
--- a/Engine/Source/Runtime/Engine/Classes/Kismet/KismetGuidLibrary.h
+++ b/Engine/Source/Runtime/Engine/Classes/Kismet/KismetGuidLibrary.h
@@ -9,9 +9,7 @@ UCLASS(MinimalAPI)
 class UKismetGuidLibrary
 	: public UBlueprintFunctionLibrary
 {
-	GENERATED_BODY()
-public:
-	ENGINE_API UKismetGuidLibrary(const FObjectInitializer& ObjectInitializer = FObjectInitializer::Get());
+	GENERATED_UCLASS_BODY()
 
 	/** Returns true if the values are equal (A == B) */
 	UFUNCTION(BlueprintPure, meta=(FriendlyName="Equal (Guid)", CompactNodeTitle="=="), Category="Guid")
diff --git a/Engine/Source/Runtime/Engine/Classes/Kismet/KismetInputLibrary.h b/Engine/Source/Runtime/Engine/Classes/Kismet/KismetInputLibrary.h
--- a/Engine/Source/Runtime/Engine/Classes/Kismet/KismetInputLibrary.h
+++ b/Engine/Source/Runtime/Engine/Classes/Kismet/KismetInputLibrary.h
@@ -9,9 +9,7 @@ typedef FKeyEvent FKeyEvent;
 UCLASS(MinimalAPI)
 class UKismetInputLibrary : public UBlueprintFunctionLibrary
 {
-	GENERATED_BODY()
-public:
-	ENGINE_API UKismetInputLibrary(const FObjectInitializer& ObjectInitializer = FObjectInitializer::Get());
+	GENERATED_UCLASS_BODY()
 
 	/** Calibrate the tilt for the input device */
 	UFUNCTION(BlueprintCallable, Category="Input")
diff --git a/Engine/Source/Runtime/Engine/Classes/Kismet/KismetMaterialLibrary.h b/Engine/Source/Runtime/Engine/Classes/Kismet/KismetMaterialLibrary.h
--- a/Engine/Source/Runtime/Engine/Classes/Kismet/KismetMaterialLibrary.h
+++ b/Engine/Source/Runtime/Engine/Classes/Kismet/KismetMaterialLibrary.h
@@ -7,9 +7,7 @@
 UCLASS(MinimalAPI)
 class UKismetMaterialLibrary : public UBlueprintFunctionLibrary
 {
-	GENERATED_BODY()
-public:
-	ENGINE_API UKismetMaterialLibrary(const FObjectInitializer& ObjectInitializer = FObjectInitializer::Get());
+	GENERATED_UCLASS_BODY()
 
 	/** Sets a scalar parameter value on the material collection instance. Logs if ParameterName is invalid. */
 	UFUNCTION(BlueprintCallable, Category="Rendering|Material", meta=(Keywords="SetFloatParameterValue", WorldContext="WorldContextObject", MaterialParameterCollectionFunction = "true"))
