Back out changelist 2481333

diff --git a/Engine/Source/Runtime/Engine/Classes/Kismet/KismetMathLibrary.h b/Engine/Source/Runtime/Engine/Classes/Kismet/KismetMathLibrary.h
--- a/Engine/Source/Runtime/Engine/Classes/Kismet/KismetMathLibrary.h
+++ b/Engine/Source/Runtime/Engine/Classes/Kismet/KismetMathLibrary.h
@@ -58,9 +58,7 @@ namespace EEasingFunc
 UCLASS(MinimalAPI)
 class UKismetMathLibrary : public UBlueprintFunctionLibrary
 {
-	GENERATED_BODY()
-public:
-	ENGINE_API UKismetMathLibrary(const FObjectInitializer& ObjectInitializer = FObjectInitializer::Get());
+	GENERATED_UCLASS_BODY()
 
 	//
 	// Boolean functions.
diff --git a/Engine/Source/Runtime/Engine/Classes/Kismet/KismetNodeHelperLibrary.h b/Engine/Source/Runtime/Engine/Classes/Kismet/KismetNodeHelperLibrary.h
--- a/Engine/Source/Runtime/Engine/Classes/Kismet/KismetNodeHelperLibrary.h
+++ b/Engine/Source/Runtime/Engine/Classes/Kismet/KismetNodeHelperLibrary.h
@@ -12,9 +12,7 @@
 UCLASS(MinimalAPI)
 class UKismetNodeHelperLibrary : public UBlueprintFunctionLibrary
 {
-	GENERATED_BODY()
-public:
-	ENGINE_API UKismetNodeHelperLibrary(const FObjectInitializer& ObjectInitializer = FObjectInitializer::Get());
+	GENERATED_UCLASS_BODY()
 
 	// --- Functions for handling bitmasking an integer as an array of bools ------------------------------
 
diff --git a/Engine/Source/Runtime/Engine/Classes/Kismet/KismetStringLibrary.h b/Engine/Source/Runtime/Engine/Classes/Kismet/KismetStringLibrary.h
--- a/Engine/Source/Runtime/Engine/Classes/Kismet/KismetStringLibrary.h
+++ b/Engine/Source/Runtime/Engine/Classes/Kismet/KismetStringLibrary.h
@@ -7,9 +7,7 @@
 UCLASS(MinimalAPI)
 class UKismetStringLibrary : public UBlueprintFunctionLibrary
 {
-	GENERATED_BODY()
-public:
-	ENGINE_API UKismetStringLibrary(const FObjectInitializer& ObjectInitializer = FObjectInitializer::Get());
+	GENERATED_UCLASS_BODY()
 
 	/** Converts a float value to a string */
 	UFUNCTION(BlueprintPure, meta=(FriendlyName = "ToString (float)", CompactNodeTitle = "->"), Category="Utilities|String")
diff --git a/Engine/Source/Runtime/Engine/Classes/Kismet/KismetTextLibrary.h b/Engine/Source/Runtime/Engine/Classes/Kismet/KismetTextLibrary.h
--- a/Engine/Source/Runtime/Engine/Classes/Kismet/KismetTextLibrary.h
+++ b/Engine/Source/Runtime/Engine/Classes/Kismet/KismetTextLibrary.h
@@ -31,7 +31,7 @@ enum ERoundingMode
 USTRUCT()
 struct FFormatTextArgument
 {
-	GENERATED_BODY()
+	GENERATED_USTRUCT_BODY()
 
 	UPROPERTY(EditInstanceOnly, BlueprintReadWrite, Category=ArgumentValue)
 	FText ArgumentName;
@@ -43,9 +43,7 @@ struct FFormatTextArgument
 UCLASS(MinimalAPI)
 class UKismetTextLibrary : public UBlueprintFunctionLibrary
 {
-	GENERATED_BODY()
-public:
-	ENGINE_API UKismetTextLibrary(const FObjectInitializer& ObjectInitializer = FObjectInitializer::Get());
+	GENERATED_UCLASS_BODY()
 
 	/** Converts localizable text to the string */
 	UFUNCTION(BlueprintPure, meta=(FriendlyName = "ToString (text)", CompactNodeTitle = "->"), Category="Utilities|String")
diff --git a/Engine/Source/Runtime/Engine/Classes/Layers/Layer.h b/Engine/Source/Runtime/Engine/Classes/Layers/Layer.h
--- a/Engine/Source/Runtime/Engine/Classes/Layers/Layer.h
+++ b/Engine/Source/Runtime/Engine/Classes/Layers/Layer.h
@@ -6,7 +6,7 @@
 USTRUCT()
 struct FLayerActorStats
 {
-	GENERATED_BODY()
+	GENERATED_USTRUCT_BODY()
 
 	/** A Type of Actor currently associated with the Layer */
 	UPROPERTY()
@@ -26,9 +26,7 @@ struct FLayerActorStats
 UCLASS( MinimalAPI )
 class ULayer : public UObject
 {
-	GENERATED_BODY()
-public:
-	ENGINE_API ULayer(const FObjectInitializer& ObjectInitializer = FObjectInitializer::Get());
+	GENERATED_UCLASS_BODY()
 
 	/** The display name of the layer */
 	UPROPERTY()
@@ -43,4 +41,4 @@ public:
 	 */
 	UPROPERTY(transient)
 	TArray< FLayerActorStats > ActorStats;
-};
+};
\ No newline at end of file
diff --git a/Engine/Source/Runtime/Engine/Classes/Lightmass/LightmappedSurfaceCollection.h b/Engine/Source/Runtime/Engine/Classes/Lightmass/LightmappedSurfaceCollection.h
--- a/Engine/Source/Runtime/Engine/Classes/Lightmass/LightmappedSurfaceCollection.h
+++ b/Engine/Source/Runtime/Engine/Classes/Lightmass/LightmappedSurfaceCollection.h
@@ -11,9 +11,7 @@
 UCLASS(hidecategories=Object, editinlinenew, MinimalAPI)
 class ULightmappedSurfaceCollection : public UObject
 {
-	GENERATED_BODY()
-public:
-	ENGINE_API ULightmappedSurfaceCollection(const FObjectInitializer& ObjectInitializer = FObjectInitializer::Get());
+	GENERATED_UCLASS_BODY()
 
 	/** The UModel these surfaces come from. */
 	UPROPERTY(EditAnywhere, Category=LightmappedSurfaceCollection)
diff --git a/Engine/Source/Runtime/Engine/Classes/Lightmass/LightmassCharacterIndirectDetailVolume.h b/Engine/Source/Runtime/Engine/Classes/Lightmass/LightmassCharacterIndirectDetailVolume.h
--- a/Engine/Source/Runtime/Engine/Classes/Lightmass/LightmassCharacterIndirectDetailVolume.h
+++ b/Engine/Source/Runtime/Engine/Classes/Lightmass/LightmassCharacterIndirectDetailVolume.h
@@ -11,9 +11,7 @@
 UCLASS(hidecategories=(Collision, Brush, Attachment, Physics, Volume), MinimalAPI)
 class ALightmassCharacterIndirectDetailVolume : public AVolume
 {
-	GENERATED_BODY()
-public:
-	ENGINE_API ALightmassCharacterIndirectDetailVolume(const FObjectInitializer& ObjectInitializer = FObjectInitializer::Get());
+	GENERATED_UCLASS_BODY()
 
 };
 
diff --git a/Engine/Source/Runtime/Engine/Classes/Lightmass/LightmassImportanceVolume.h b/Engine/Source/Runtime/Engine/Classes/Lightmass/LightmassImportanceVolume.h
--- a/Engine/Source/Runtime/Engine/Classes/Lightmass/LightmassImportanceVolume.h
+++ b/Engine/Source/Runtime/Engine/Classes/Lightmass/LightmassImportanceVolume.h
@@ -12,9 +12,7 @@
 UCLASS(hidecategories=(Collision, Brush, Attachment, Physics, Volume), MinimalAPI)
 class ALightmassImportanceVolume : public AVolume
 {
-	GENERATED_BODY()
-public:
-	ENGINE_API ALightmassImportanceVolume(const FObjectInitializer& ObjectInitializer = FObjectInitializer::Get());
+	GENERATED_UCLASS_BODY()
 
 };
 
diff --git a/Engine/Source/Runtime/Engine/Classes/Lightmass/LightmassPrimitiveSettingsObject.h b/Engine/Source/Runtime/Engine/Classes/Lightmass/LightmassPrimitiveSettingsObject.h
--- a/Engine/Source/Runtime/Engine/Classes/Lightmass/LightmassPrimitiveSettingsObject.h
+++ b/Engine/Source/Runtime/Engine/Classes/Lightmass/LightmassPrimitiveSettingsObject.h
@@ -13,9 +13,7 @@
 UCLASS(hidecategories=Object, editinlinenew, MinimalAPI,collapsecategories)
 class ULightmassPrimitiveSettingsObject : public UObject
 {
-	GENERATED_BODY()
-public:
-	ENGINE_API ULightmassPrimitiveSettingsObject(const FObjectInitializer& ObjectInitializer = FObjectInitializer::Get());
+	GENERATED_UCLASS_BODY()
 
 	UPROPERTY(EditAnywhere, Category=Lightmass)
 	struct FLightmassPrimitiveSettings LightmassSettings;
diff --git a/Engine/Source/Runtime/Engine/Classes/Lightmass/PrecomputedVisibilityOverrideVolume.h b/Engine/Source/Runtime/Engine/Classes/Lightmass/PrecomputedVisibilityOverrideVolume.h
--- a/Engine/Source/Runtime/Engine/Classes/Lightmass/PrecomputedVisibilityOverrideVolume.h
+++ b/Engine/Source/Runtime/Engine/Classes/Lightmass/PrecomputedVisibilityOverrideVolume.h
@@ -11,9 +11,7 @@
 UCLASS(hidecategories=(Collision, Brush, Attachment, Physics, Volume), MinimalAPI)
 class APrecomputedVisibilityOverrideVolume : public AVolume
 {
-	GENERATED_BODY()
-public:
-	ENGINE_API APrecomputedVisibilityOverrideVolume(const FObjectInitializer& ObjectInitializer = FObjectInitializer::Get());
+	GENERATED_UCLASS_BODY()
 
 	/** Array of actors that will always be considered visible by Precomputed Visibility when viewed from inside this volume. */
 	UPROPERTY(EditAnywhere, BlueprintReadWrite, Category=PrecomputedVisibilityOverrideVolume)
diff --git a/Engine/Source/Runtime/Engine/Classes/Lightmass/PrecomputedVisibilityVolume.h b/Engine/Source/Runtime/Engine/Classes/Lightmass/PrecomputedVisibilityVolume.h
--- a/Engine/Source/Runtime/Engine/Classes/Lightmass/PrecomputedVisibilityVolume.h
+++ b/Engine/Source/Runtime/Engine/Classes/Lightmass/PrecomputedVisibilityVolume.h
@@ -11,9 +11,7 @@
 UCLASS(hidecategories=(Collision, Brush, Attachment, Physics, Volume), MinimalAPI)
 class APrecomputedVisibilityVolume : public AVolume
 {
-	GENERATED_BODY()
-public:
-	ENGINE_API APrecomputedVisibilityVolume(const FObjectInitializer& ObjectInitializer = FObjectInitializer::Get());
+	GENERATED_UCLASS_BODY()
 
 };
 
diff --git a/Engine/Source/Runtime/Engine/Classes/Materials/Material.h b/Engine/Source/Runtime/Engine/Classes/Materials/Material.h
--- a/Engine/Source/Runtime/Engine/Classes/Materials/Material.h
+++ b/Engine/Source/Runtime/Engine/Classes/Materials/Material.h
@@ -217,7 +217,7 @@ struct FVector2MaterialInput : public FMaterialInput
 USTRUCT()
 struct FMaterialFunctionInfo
 {
-	GENERATED_BODY()
+	GENERATED_USTRUCT_BODY()
 
 	/** Id that the function had when this material was last compiled. */
 	UPROPERTY()
@@ -236,7 +236,7 @@ struct FMaterialFunctionInfo
 USTRUCT()
 struct FMaterialParameterCollectionInfo
 {
-	GENERATED_BODY()
+	GENERATED_USTRUCT_BODY()
 
 	/** Id that the collection had when this material was last compiled. */
 	UPROPERTY()
@@ -264,9 +264,7 @@ struct FMaterialParameterCollectionInfo
 UCLASS(hidecategories=Object, MinimalAPI, BlueprintType)
 class UMaterial : public UMaterialInterface
 {
-	GENERATED_BODY()
-public:
-	ENGINE_API UMaterial(const FObjectInitializer& ObjectInitializer = FObjectInitializer::Get());
+	GENERATED_UCLASS_BODY()
 
 	// Physics.
 	
diff --git a/Engine/Source/Runtime/Engine/Classes/Materials/MaterialExpression.h b/Engine/Source/Runtime/Engine/Classes/Materials/MaterialExpression.h
--- a/Engine/Source/Runtime/Engine/Classes/Materials/MaterialExpression.h
+++ b/Engine/Source/Runtime/Engine/Classes/Materials/MaterialExpression.h
@@ -85,9 +85,7 @@ struct FExpressionOutput
 UCLASS(abstract, hidecategories=Object)
 class ENGINE_API UMaterialExpression : public UObject
 {
-	GENERATED_BODY()
-public:
-	UMaterialExpression(const FObjectInitializer& ObjectInitializer = FObjectInitializer::Get());
+	GENERATED_UCLASS_BODY()
 
 #if WITH_EDITORONLY_DATA
 	UPROPERTY()
diff --git a/Engine/Source/Runtime/Engine/Classes/Materials/MaterialExpressionAbs.h b/Engine/Source/Runtime/Engine/Classes/Materials/MaterialExpressionAbs.h
--- a/Engine/Source/Runtime/Engine/Classes/Materials/MaterialExpressionAbs.h
+++ b/Engine/Source/Runtime/Engine/Classes/Materials/MaterialExpressionAbs.h
@@ -12,9 +12,7 @@
 UCLASS(collapsecategories, hidecategories=Object)
 class UMaterialExpressionAbs : public UMaterialExpression
 {
-	GENERATED_BODY()
-public:
-	UMaterialExpressionAbs(const FObjectInitializer& ObjectInitializer = FObjectInitializer::Get());
+	GENERATED_UCLASS_BODY()
 
 	/** Link to the input expression to be evaluated */
 	UPROPERTY()
diff --git a/Engine/Source/Runtime/Engine/Classes/Materials/MaterialExpressionActorPositionWS.h b/Engine/Source/Runtime/Engine/Classes/Materials/MaterialExpressionActorPositionWS.h
--- a/Engine/Source/Runtime/Engine/Classes/Materials/MaterialExpressionActorPositionWS.h
+++ b/Engine/Source/Runtime/Engine/Classes/Materials/MaterialExpressionActorPositionWS.h
@@ -8,9 +8,7 @@
 UCLASS(collapsecategories, hidecategories=Object)
 class UMaterialExpressionActorPositionWS : public UMaterialExpression
 {
-	GENERATED_BODY()
-public:
-	UMaterialExpressionActorPositionWS(const FObjectInitializer& ObjectInitializer = FObjectInitializer::Get());
+	GENERATED_UCLASS_BODY()
 
 
 	// Begin UMaterialExpression Interface
diff --git a/Engine/Source/Runtime/Engine/Classes/Materials/MaterialExpressionAdd.h b/Engine/Source/Runtime/Engine/Classes/Materials/MaterialExpressionAdd.h
--- a/Engine/Source/Runtime/Engine/Classes/Materials/MaterialExpressionAdd.h
+++ b/Engine/Source/Runtime/Engine/Classes/Materials/MaterialExpressionAdd.h
@@ -8,9 +8,7 @@
 UCLASS(MinimalAPI)
 class UMaterialExpressionAdd : public UMaterialExpression
 {
-	GENERATED_BODY()
-public:
-	ENGINE_API UMaterialExpressionAdd(const FObjectInitializer& ObjectInitializer = FObjectInitializer::Get());
+	GENERATED_UCLASS_BODY()
 
 	UPROPERTY(meta = (RequiredInput = "false", ToolTip = "Defaults to 'ConstA' if not specified"))
 	FExpressionInput A;
diff --git a/Engine/Source/Runtime/Engine/Classes/Materials/MaterialExpressionAntialiasedTextureMask.h b/Engine/Source/Runtime/Engine/Classes/Materials/MaterialExpressionAntialiasedTextureMask.h
--- a/Engine/Source/Runtime/Engine/Classes/Materials/MaterialExpressionAntialiasedTextureMask.h
+++ b/Engine/Source/Runtime/Engine/Classes/Materials/MaterialExpressionAntialiasedTextureMask.h
@@ -18,9 +18,7 @@ enum ETextureColorChannel
 UCLASS(MinimalAPI)
 class UMaterialExpressionAntialiasedTextureMask : public UMaterialExpressionTextureSampleParameter2D
 {
-	GENERATED_BODY()
-public:
-	ENGINE_API UMaterialExpressionAntialiasedTextureMask(const FObjectInitializer& ObjectInitializer = FObjectInitializer::Get());
+	GENERATED_UCLASS_BODY()
 
 	UPROPERTY(EditAnywhere, Category=MaterialExpressionAntialiasedTextureMask, meta=(UIMin = "0.0", UIMax = "1.0", ClampMin = "0.0", ClampMax = "1.0"))
 	float Threshold;
diff --git a/Engine/Source/Runtime/Engine/Classes/Materials/MaterialExpressionAtmosphericFogColor.h b/Engine/Source/Runtime/Engine/Classes/Materials/MaterialExpressionAtmosphericFogColor.h
--- a/Engine/Source/Runtime/Engine/Classes/Materials/MaterialExpressionAtmosphericFogColor.h
+++ b/Engine/Source/Runtime/Engine/Classes/Materials/MaterialExpressionAtmosphericFogColor.h
@@ -7,9 +7,7 @@
 UCLASS(collapsecategories, hidecategories=Object)
 class UMaterialExpressionAtmosphericFogColor : public UMaterialExpression
 {
-	GENERATED_BODY()
-public:
-	UMaterialExpressionAtmosphericFogColor(const FObjectInitializer& ObjectInitializer = FObjectInitializer::Get());
+	GENERATED_UCLASS_BODY()
 
 	UPROPERTY()
 	FExpressionInput WorldPosition;
diff --git a/Engine/Source/Runtime/Engine/Classes/Materials/MaterialExpressionBlackBody.h b/Engine/Source/Runtime/Engine/Classes/Materials/MaterialExpressionBlackBody.h
--- a/Engine/Source/Runtime/Engine/Classes/Materials/MaterialExpressionBlackBody.h
+++ b/Engine/Source/Runtime/Engine/Classes/Materials/MaterialExpressionBlackBody.h
@@ -8,9 +8,7 @@
 UCLASS()
 class UMaterialExpressionBlackBody : public UMaterialExpression
 {
-	GENERATED_BODY()
-public:
-	UMaterialExpressionBlackBody(const FObjectInitializer& ObjectInitializer = FObjectInitializer::Get());
+	GENERATED_UCLASS_BODY()
 
 	/** Temperature */
 	UPROPERTY()
diff --git a/Engine/Source/Runtime/Engine/Classes/Materials/MaterialExpressionBreakMaterialAttributes.h b/Engine/Source/Runtime/Engine/Classes/Materials/MaterialExpressionBreakMaterialAttributes.h
--- a/Engine/Source/Runtime/Engine/Classes/Materials/MaterialExpressionBreakMaterialAttributes.h
+++ b/Engine/Source/Runtime/Engine/Classes/Materials/MaterialExpressionBreakMaterialAttributes.h
@@ -8,9 +8,7 @@
 UCLASS(collapsecategories, hidecategories=Object, MinimalAPI)
 class UMaterialExpressionBreakMaterialAttributes : public UMaterialExpression
 {
-	GENERATED_BODY()
-public:
-	ENGINE_API UMaterialExpressionBreakMaterialAttributes(const FObjectInitializer& ObjectInitializer = FObjectInitializer::Get());
+	GENERATED_UCLASS_BODY()
 
 	UPROPERTY()
 	FExpressionInput Struct;
diff --git a/Engine/Source/Runtime/Engine/Classes/Materials/MaterialExpressionBumpOffset.h b/Engine/Source/Runtime/Engine/Classes/Materials/MaterialExpressionBumpOffset.h
--- a/Engine/Source/Runtime/Engine/Classes/Materials/MaterialExpressionBumpOffset.h
+++ b/Engine/Source/Runtime/Engine/Classes/Materials/MaterialExpressionBumpOffset.h
@@ -8,9 +8,7 @@
 UCLASS(collapsecategories, hidecategories=Object, MinimalAPI)
 class UMaterialExpressionBumpOffset : public UMaterialExpression
 {
-	GENERATED_BODY()
-public:
-	ENGINE_API UMaterialExpressionBumpOffset(const FObjectInitializer& ObjectInitializer = FObjectInitializer::Get());
+	GENERATED_UCLASS_BODY()
 
 	// Outputs: Coordinate + Eye.xy * (Height - ReferencePlane) * HeightRatio
 	UPROPERTY(meta = (RequiredInput = "false", ToolTip = "Defaults to 'ConstCoordinate' if not specified"))
diff --git a/Engine/Source/Runtime/Engine/Classes/Materials/MaterialExpressionCameraPositionWS.h b/Engine/Source/Runtime/Engine/Classes/Materials/MaterialExpressionCameraPositionWS.h
--- a/Engine/Source/Runtime/Engine/Classes/Materials/MaterialExpressionCameraPositionWS.h
+++ b/Engine/Source/Runtime/Engine/Classes/Materials/MaterialExpressionCameraPositionWS.h
@@ -8,9 +8,7 @@
 UCLASS(collapsecategories, hidecategories=Object)
 class UMaterialExpressionCameraPositionWS : public UMaterialExpression
 {
-	GENERATED_BODY()
-public:
-	UMaterialExpressionCameraPositionWS(const FObjectInitializer& ObjectInitializer = FObjectInitializer::Get());
+	GENERATED_UCLASS_BODY()
 
 
 	// Begin UMaterialExpression Interface
diff --git a/Engine/Source/Runtime/Engine/Classes/Materials/MaterialExpressionCameraVectorWS.h b/Engine/Source/Runtime/Engine/Classes/Materials/MaterialExpressionCameraVectorWS.h
--- a/Engine/Source/Runtime/Engine/Classes/Materials/MaterialExpressionCameraVectorWS.h
+++ b/Engine/Source/Runtime/Engine/Classes/Materials/MaterialExpressionCameraVectorWS.h
@@ -8,9 +8,7 @@
 UCLASS(collapsecategories, hidecategories=Object)
 class UMaterialExpressionCameraVectorWS : public UMaterialExpression
 {
-	GENERATED_BODY()
-public:
-	UMaterialExpressionCameraVectorWS(const FObjectInitializer& ObjectInitializer = FObjectInitializer::Get());
+	GENERATED_UCLASS_BODY()
 
 
 	// Begin UMaterialExpression Interface
diff --git a/Engine/Source/Runtime/Engine/Classes/Materials/MaterialExpressionCeil.h b/Engine/Source/Runtime/Engine/Classes/Materials/MaterialExpressionCeil.h
--- a/Engine/Source/Runtime/Engine/Classes/Materials/MaterialExpressionCeil.h
+++ b/Engine/Source/Runtime/Engine/Classes/Materials/MaterialExpressionCeil.h
@@ -9,9 +9,7 @@
 UCLASS(collapsecategories, hidecategories=Object)
 class UMaterialExpressionCeil : public UMaterialExpression
 {
-	GENERATED_BODY()
-public:
-	UMaterialExpressionCeil(const FObjectInitializer& ObjectInitializer = FObjectInitializer::Get());
+	GENERATED_UCLASS_BODY()
 
 	UPROPERTY()
 	FExpressionInput Input;
diff --git a/Engine/Source/Runtime/Engine/Classes/Materials/MaterialExpressionClamp.h b/Engine/Source/Runtime/Engine/Classes/Materials/MaterialExpressionClamp.h
--- a/Engine/Source/Runtime/Engine/Classes/Materials/MaterialExpressionClamp.h
+++ b/Engine/Source/Runtime/Engine/Classes/Materials/MaterialExpressionClamp.h
@@ -17,9 +17,7 @@ enum EClampMode
 UCLASS(MinimalAPI)
 class UMaterialExpressionClamp : public UMaterialExpression
 {
-	GENERATED_BODY()
-public:
-	ENGINE_API UMaterialExpressionClamp(const FObjectInitializer& ObjectInitializer = FObjectInitializer::Get());
+	GENERATED_UCLASS_BODY()
 
 	UPROPERTY()
 	FExpressionInput Input;
diff --git a/Engine/Source/Runtime/Engine/Classes/Materials/MaterialExpressionCollectionParameter.h b/Engine/Source/Runtime/Engine/Classes/Materials/MaterialExpressionCollectionParameter.h
--- a/Engine/Source/Runtime/Engine/Classes/Materials/MaterialExpressionCollectionParameter.h
+++ b/Engine/Source/Runtime/Engine/Classes/Materials/MaterialExpressionCollectionParameter.h
@@ -11,9 +11,7 @@
 UCLASS(hidecategories=object, MinimalAPI)
 class UMaterialExpressionCollectionParameter : public UMaterialExpression
 {
-	GENERATED_BODY()
-public:
-	ENGINE_API UMaterialExpressionCollectionParameter(const FObjectInitializer& ObjectInitializer = FObjectInitializer::Get());
+	GENERATED_UCLASS_BODY()
 
 	/** The Parameter Collection to use. */
 	UPROPERTY(EditAnywhere, Category=MaterialExpressionCollectionParameter)
diff --git a/Engine/Source/Runtime/Engine/Classes/Materials/MaterialExpressionComment.h b/Engine/Source/Runtime/Engine/Classes/Materials/MaterialExpressionComment.h
--- a/Engine/Source/Runtime/Engine/Classes/Materials/MaterialExpressionComment.h
+++ b/Engine/Source/Runtime/Engine/Classes/Materials/MaterialExpressionComment.h
@@ -8,9 +8,7 @@
 UCLASS(MinimalAPI)
 class UMaterialExpressionComment : public UMaterialExpression
 {
-	GENERATED_BODY()
-public:
-	ENGINE_API UMaterialExpressionComment(const FObjectInitializer& ObjectInitializer = FObjectInitializer::Get());
+	GENERATED_UCLASS_BODY()
 
 	UPROPERTY()
 	int32 SizeX;
diff --git a/Engine/Source/Runtime/Engine/Classes/Materials/MaterialExpressionConstant3Vector.h b/Engine/Source/Runtime/Engine/Classes/Materials/MaterialExpressionConstant3Vector.h
--- a/Engine/Source/Runtime/Engine/Classes/Materials/MaterialExpressionConstant3Vector.h
+++ b/Engine/Source/Runtime/Engine/Classes/Materials/MaterialExpressionConstant3Vector.h
@@ -8,9 +8,7 @@
 UCLASS(collapsecategories, hidecategories=Object, MinimalAPI)
 class UMaterialExpressionConstant3Vector : public UMaterialExpression
 {
-	GENERATED_BODY()
-public:
-	ENGINE_API UMaterialExpressionConstant3Vector(const FObjectInitializer& ObjectInitializer = FObjectInitializer::Get());
+	GENERATED_UCLASS_BODY()
 
  	UPROPERTY(EditAnywhere, Category=MaterialExpressionConstant3Vector, meta=(HideAlphaChannel))
 	FLinearColor Constant;
diff --git a/Engine/Source/Runtime/Engine/Classes/Materials/MaterialExpressionConstant4Vector.h b/Engine/Source/Runtime/Engine/Classes/Materials/MaterialExpressionConstant4Vector.h
--- a/Engine/Source/Runtime/Engine/Classes/Materials/MaterialExpressionConstant4Vector.h
+++ b/Engine/Source/Runtime/Engine/Classes/Materials/MaterialExpressionConstant4Vector.h
@@ -8,9 +8,7 @@
 UCLASS(collapsecategories, hidecategories=Object, MinimalAPI)
 class UMaterialExpressionConstant4Vector : public UMaterialExpression
 {
-	GENERATED_BODY()
-public:
-	ENGINE_API UMaterialExpressionConstant4Vector(const FObjectInitializer& ObjectInitializer = FObjectInitializer::Get());
+	GENERATED_UCLASS_BODY()
 
 	UPROPERTY(EditAnywhere, Category=MaterialExpressionConstant4Vector)
 	FLinearColor Constant;
diff --git a/Engine/Source/Runtime/Engine/Classes/Materials/MaterialExpressionConstantBiasScale.h b/Engine/Source/Runtime/Engine/Classes/Materials/MaterialExpressionConstantBiasScale.h
--- a/Engine/Source/Runtime/Engine/Classes/Materials/MaterialExpressionConstantBiasScale.h
+++ b/Engine/Source/Runtime/Engine/Classes/Materials/MaterialExpressionConstantBiasScale.h
@@ -8,9 +8,7 @@
 UCLASS(MinimalAPI)
 class UMaterialExpressionConstantBiasScale : public UMaterialExpression
 {
-	GENERATED_BODY()
-public:
-	ENGINE_API UMaterialExpressionConstantBiasScale(const FObjectInitializer& ObjectInitializer = FObjectInitializer::Get());
+	GENERATED_UCLASS_BODY()
 
 	UPROPERTY()
 	FExpressionInput Input;
diff --git a/Engine/Source/Runtime/Engine/Classes/Materials/MaterialExpressionCosine.h b/Engine/Source/Runtime/Engine/Classes/Materials/MaterialExpressionCosine.h
--- a/Engine/Source/Runtime/Engine/Classes/Materials/MaterialExpressionCosine.h
+++ b/Engine/Source/Runtime/Engine/Classes/Materials/MaterialExpressionCosine.h
@@ -8,9 +8,7 @@
 UCLASS(collapsecategories, hidecategories=Object)
 class UMaterialExpressionCosine : public UMaterialExpression
 {
-	GENERATED_BODY()
-public:
-	UMaterialExpressionCosine(const FObjectInitializer& ObjectInitializer = FObjectInitializer::Get());
+	GENERATED_UCLASS_BODY()
 
 	UPROPERTY()
 	FExpressionInput Input;
diff --git a/Engine/Source/Runtime/Engine/Classes/Materials/MaterialExpressionCrossProduct.h b/Engine/Source/Runtime/Engine/Classes/Materials/MaterialExpressionCrossProduct.h
--- a/Engine/Source/Runtime/Engine/Classes/Materials/MaterialExpressionCrossProduct.h
+++ b/Engine/Source/Runtime/Engine/Classes/Materials/MaterialExpressionCrossProduct.h
@@ -8,9 +8,7 @@
 UCLASS()
 class UMaterialExpressionCrossProduct : public UMaterialExpression
 {
-	GENERATED_BODY()
-public:
-	UMaterialExpressionCrossProduct(const FObjectInitializer& ObjectInitializer = FObjectInitializer::Get());
+	GENERATED_UCLASS_BODY()
 
 	UPROPERTY()
 	FExpressionInput A;
diff --git a/Engine/Source/Runtime/Engine/Classes/Materials/MaterialExpressionCustom.h b/Engine/Source/Runtime/Engine/Classes/Materials/MaterialExpressionCustom.h
--- a/Engine/Source/Runtime/Engine/Classes/Materials/MaterialExpressionCustom.h
+++ b/Engine/Source/Runtime/Engine/Classes/Materials/MaterialExpressionCustom.h
@@ -18,7 +18,7 @@ enum ECustomMaterialOutputType
 USTRUCT()
 struct FCustomInput
 {
-	GENERATED_BODY()
+	GENERATED_USTRUCT_BODY()
 
 	UPROPERTY(EditAnywhere, Category=CustomInput)
 	FString InputName;
@@ -31,9 +31,7 @@ struct FCustomInput
 UCLASS(collapsecategories, hidecategories=Object, MinimalAPI)
 class UMaterialExpressionCustom : public UMaterialExpression
 {
-	GENERATED_BODY()
-public:
-	ENGINE_API UMaterialExpressionCustom(const FObjectInitializer& ObjectInitializer = FObjectInitializer::Get());
+	GENERATED_UCLASS_BODY()
 
 	UPROPERTY(EditAnywhere, Category=MaterialExpressionCustom, meta=(MultiLine=true))
 	FString Code;
diff --git a/Engine/Source/Runtime/Engine/Classes/Materials/MaterialExpressionCustomOutput.h b/Engine/Source/Runtime/Engine/Classes/Materials/MaterialExpressionCustomOutput.h
--- a/Engine/Source/Runtime/Engine/Classes/Materials/MaterialExpressionCustomOutput.h
+++ b/Engine/Source/Runtime/Engine/Classes/Materials/MaterialExpressionCustomOutput.h
@@ -8,9 +8,7 @@
 UCLASS(collapsecategories, hidecategories=Object, MinimalAPI)
 class UMaterialExpressionCustomOutput : public UMaterialExpression
 {
-	GENERATED_BODY()
-public:
-	ENGINE_API UMaterialExpressionCustomOutput(const FObjectInitializer& ObjectInitializer = FObjectInitializer::Get());
+	GENERATED_UCLASS_BODY()
 
 	// Override to enable multiple outputs
 	virtual int32 GetNumOutputs() const { return 1; };
diff --git a/Engine/Source/Runtime/Engine/Classes/Materials/MaterialExpressionDDX.h b/Engine/Source/Runtime/Engine/Classes/Materials/MaterialExpressionDDX.h
--- a/Engine/Source/Runtime/Engine/Classes/Materials/MaterialExpressionDDX.h
+++ b/Engine/Source/Runtime/Engine/Classes/Materials/MaterialExpressionDDX.h
@@ -9,9 +9,7 @@
 UCLASS()
 class UMaterialExpressionDDX : public UMaterialExpression
 {
-	GENERATED_BODY()
-public:
-	UMaterialExpressionDDX(const FObjectInitializer& ObjectInitializer = FObjectInitializer::Get());
+	GENERATED_UCLASS_BODY()
 
 	/** The value we want to compute ddx/ddy from */
 	UPROPERTY()
diff --git a/Engine/Source/Runtime/Engine/Classes/Materials/MaterialExpressionDDY.h b/Engine/Source/Runtime/Engine/Classes/Materials/MaterialExpressionDDY.h
--- a/Engine/Source/Runtime/Engine/Classes/Materials/MaterialExpressionDDY.h
+++ b/Engine/Source/Runtime/Engine/Classes/Materials/MaterialExpressionDDY.h
@@ -9,9 +9,7 @@
 UCLASS()
 class UMaterialExpressionDDY : public UMaterialExpression
 {
-	GENERATED_BODY()
-public:
-	UMaterialExpressionDDY(const FObjectInitializer& ObjectInitializer = FObjectInitializer::Get());
+	GENERATED_UCLASS_BODY()
 
 	/** The value we want to compute ddx/ddy from */
 	UPROPERTY()
diff --git a/Engine/Source/Runtime/Engine/Classes/Materials/MaterialExpressionDepthOfFieldFunction.h b/Engine/Source/Runtime/Engine/Classes/Materials/MaterialExpressionDepthOfFieldFunction.h
--- a/Engine/Source/Runtime/Engine/Classes/Materials/MaterialExpressionDepthOfFieldFunction.h
+++ b/Engine/Source/Runtime/Engine/Classes/Materials/MaterialExpressionDepthOfFieldFunction.h
@@ -20,9 +20,7 @@ enum EDepthOfFieldFunctionValue
 UCLASS()
 class UMaterialExpressionDepthOfFieldFunction : public UMaterialExpression
 {
-	GENERATED_BODY()
-public:
-	UMaterialExpressionDepthOfFieldFunction(const FObjectInitializer& ObjectInitializer = FObjectInitializer::Get());
+	GENERATED_UCLASS_BODY()
 
 	/** Determines the mapping place to use on the terrain. */
 	UPROPERTY(EditAnywhere, Category=MaterialExpressionDepthOfFieldFunction)
diff --git a/Engine/Source/Runtime/Engine/Classes/Materials/MaterialExpressionDesaturation.h b/Engine/Source/Runtime/Engine/Classes/Materials/MaterialExpressionDesaturation.h
--- a/Engine/Source/Runtime/Engine/Classes/Materials/MaterialExpressionDesaturation.h
+++ b/Engine/Source/Runtime/Engine/Classes/Materials/MaterialExpressionDesaturation.h
@@ -8,9 +8,7 @@
 UCLASS(collapsecategories, hidecategories=Object)
 class UMaterialExpressionDesaturation : public UMaterialExpression
 {
-	GENERATED_BODY()
-public:
-	UMaterialExpressionDesaturation(const FObjectInitializer& ObjectInitializer = FObjectInitializer::Get());
+	GENERATED_UCLASS_BODY()
 
 	// Outputs: Lerp(Input,dot(Input,LuminanceFactors)),Fraction)
 	UPROPERTY()
diff --git a/Engine/Source/Runtime/Engine/Classes/Materials/MaterialExpressionDistance.h b/Engine/Source/Runtime/Engine/Classes/Materials/MaterialExpressionDistance.h
--- a/Engine/Source/Runtime/Engine/Classes/Materials/MaterialExpressionDistance.h
+++ b/Engine/Source/Runtime/Engine/Classes/Materials/MaterialExpressionDistance.h
@@ -8,9 +8,7 @@
 UCLASS()
 class UMaterialExpressionDistance : public UMaterialExpression
 {
-	GENERATED_BODY()
-public:
-	UMaterialExpressionDistance(const FObjectInitializer& ObjectInitializer = FObjectInitializer::Get());
+	GENERATED_UCLASS_BODY()
 
 	UPROPERTY()
 	FExpressionInput A;
diff --git a/Engine/Source/Runtime/Engine/Classes/Materials/MaterialExpressionDistanceCullFade.h b/Engine/Source/Runtime/Engine/Classes/Materials/MaterialExpressionDistanceCullFade.h
--- a/Engine/Source/Runtime/Engine/Classes/Materials/MaterialExpressionDistanceCullFade.h
+++ b/Engine/Source/Runtime/Engine/Classes/Materials/MaterialExpressionDistanceCullFade.h
@@ -8,9 +8,7 @@
 UCLASS(collapsecategories, hidecategories=Object)
 class UMaterialExpressionDistanceCullFade : public UMaterialExpression
 {
-	GENERATED_BODY()
-public:
-	UMaterialExpressionDistanceCullFade(const FObjectInitializer& ObjectInitializer = FObjectInitializer::Get());
+	GENERATED_UCLASS_BODY()
 
 
 	// Begin UMaterialExpression Interface
diff --git a/Engine/Source/Runtime/Engine/Classes/Materials/MaterialExpressionDivide.h b/Engine/Source/Runtime/Engine/Classes/Materials/MaterialExpressionDivide.h
--- a/Engine/Source/Runtime/Engine/Classes/Materials/MaterialExpressionDivide.h
+++ b/Engine/Source/Runtime/Engine/Classes/Materials/MaterialExpressionDivide.h
@@ -8,9 +8,7 @@
 UCLASS(MinimalAPI)
 class UMaterialExpressionDivide : public UMaterialExpression
 {
-	GENERATED_BODY()
-public:
-	ENGINE_API UMaterialExpressionDivide(const FObjectInitializer& ObjectInitializer = FObjectInitializer::Get());
+	GENERATED_UCLASS_BODY()
 
 	UPROPERTY(meta = (RequiredInput = "false", ToolTip = "Defaults to 'ConstA' if not specified"))
 	FExpressionInput A;
diff --git a/Engine/Source/Runtime/Engine/Classes/Materials/MaterialExpressionDynamicParameter.h b/Engine/Source/Runtime/Engine/Classes/Materials/MaterialExpressionDynamicParameter.h
--- a/Engine/Source/Runtime/Engine/Classes/Materials/MaterialExpressionDynamicParameter.h
+++ b/Engine/Source/Runtime/Engine/Classes/Materials/MaterialExpressionDynamicParameter.h
@@ -13,9 +13,7 @@
 UCLASS(collapsecategories, hidecategories=Object, MinimalAPI)
 class UMaterialExpressionDynamicParameter : public UMaterialExpression
 {
-	GENERATED_BODY()
-public:
-	ENGINE_API UMaterialExpressionDynamicParameter(const FObjectInitializer& ObjectInitializer = FObjectInitializer::Get());
+	GENERATED_UCLASS_BODY()
 
 	/** 
 	 *	The names of the parameters.
diff --git a/Engine/Source/Runtime/Engine/Classes/Materials/MaterialExpressionEyeAdaptation.h b/Engine/Source/Runtime/Engine/Classes/Materials/MaterialExpressionEyeAdaptation.h
--- a/Engine/Source/Runtime/Engine/Classes/Materials/MaterialExpressionEyeAdaptation.h
+++ b/Engine/Source/Runtime/Engine/Classes/Materials/MaterialExpressionEyeAdaptation.h
@@ -12,9 +12,7 @@
 UCLASS(collapsecategories, hidecategories=Object)
 class UMaterialExpressionEyeAdaptation : public UMaterialExpression
 {
-	GENERATED_BODY()
-public:
-	UMaterialExpressionEyeAdaptation(const FObjectInitializer& ObjectInitializer = FObjectInitializer::Get());
+	GENERATED_UCLASS_BODY()
 
 	// Begin UMaterialExpression Interface
 	virtual int32 Compile(class FMaterialCompiler* Compiler, int32 OutputIndex, int32 MultiplexIndex) override;
diff --git a/Engine/Source/Runtime/Engine/Classes/Materials/MaterialExpressionFeatureLevelSwitch.h b/Engine/Source/Runtime/Engine/Classes/Materials/MaterialExpressionFeatureLevelSwitch.h
--- a/Engine/Source/Runtime/Engine/Classes/Materials/MaterialExpressionFeatureLevelSwitch.h
+++ b/Engine/Source/Runtime/Engine/Classes/Materials/MaterialExpressionFeatureLevelSwitch.h
@@ -8,9 +8,7 @@
 UCLASS(collapsecategories, hidecategories=Object, MinimalAPI)
 class UMaterialExpressionFeatureLevelSwitch : public UMaterialExpression
 {
-	GENERATED_BODY()
-public:
-	ENGINE_API UMaterialExpressionFeatureLevelSwitch(const FObjectInitializer& ObjectInitializer = FObjectInitializer::Get());
+	GENERATED_UCLASS_BODY()
 
 	/** Default connection, used when a certain feature level doesn't have an override. */
 	UPROPERTY()
diff --git a/Engine/Source/Runtime/Engine/Classes/Materials/MaterialExpressionFloor.h b/Engine/Source/Runtime/Engine/Classes/Materials/MaterialExpressionFloor.h
--- a/Engine/Source/Runtime/Engine/Classes/Materials/MaterialExpressionFloor.h
+++ b/Engine/Source/Runtime/Engine/Classes/Materials/MaterialExpressionFloor.h
@@ -8,9 +8,7 @@
 UCLASS(collapsecategories, hidecategories=Object)
 class UMaterialExpressionFloor : public UMaterialExpression
 {
-	GENERATED_BODY()
-public:
-	UMaterialExpressionFloor(const FObjectInitializer& ObjectInitializer = FObjectInitializer::Get());
+	GENERATED_UCLASS_BODY()
 
 	UPROPERTY()
 	FExpressionInput Input;
diff --git a/Engine/Source/Runtime/Engine/Classes/Materials/MaterialExpressionFmod.h b/Engine/Source/Runtime/Engine/Classes/Materials/MaterialExpressionFmod.h
--- a/Engine/Source/Runtime/Engine/Classes/Materials/MaterialExpressionFmod.h
+++ b/Engine/Source/Runtime/Engine/Classes/Materials/MaterialExpressionFmod.h
@@ -8,9 +8,7 @@
 UCLASS(collapsecategories, hidecategories=Object)
 class UMaterialExpressionFmod : public UMaterialExpression
 {
-	GENERATED_BODY()
-public:
-	UMaterialExpressionFmod(const FObjectInitializer& ObjectInitializer = FObjectInitializer::Get());
+	GENERATED_UCLASS_BODY()
 
 	UPROPERTY()
 	FExpressionInput A;
diff --git a/Engine/Source/Runtime/Engine/Classes/Materials/MaterialExpressionFontSample.h b/Engine/Source/Runtime/Engine/Classes/Materials/MaterialExpressionFontSample.h
--- a/Engine/Source/Runtime/Engine/Classes/Materials/MaterialExpressionFontSample.h
+++ b/Engine/Source/Runtime/Engine/Classes/Materials/MaterialExpressionFontSample.h
@@ -8,9 +8,7 @@
 UCLASS(collapsecategories, hidecategories=Object, MinimalAPI)
 class UMaterialExpressionFontSample : public UMaterialExpression
 {
-	GENERATED_BODY()
-public:
-	ENGINE_API UMaterialExpressionFontSample(const FObjectInitializer& ObjectInitializer = FObjectInitializer::Get());
+	GENERATED_UCLASS_BODY()
 
 	/** font resource that will be sampled */
 	UPROPERTY(EditAnywhere, Category=MaterialExpressionFontSample)
diff --git a/Engine/Source/Runtime/Engine/Classes/Materials/MaterialExpressionFontSampleParameter.h b/Engine/Source/Runtime/Engine/Classes/Materials/MaterialExpressionFontSampleParameter.h
--- a/Engine/Source/Runtime/Engine/Classes/Materials/MaterialExpressionFontSampleParameter.h
+++ b/Engine/Source/Runtime/Engine/Classes/Materials/MaterialExpressionFontSampleParameter.h
@@ -8,9 +8,7 @@
 UCLASS(collapsecategories, hidecategories=Object, MinimalAPI)
 class UMaterialExpressionFontSampleParameter : public UMaterialExpressionFontSample
 {
-	GENERATED_BODY()
-public:
-	ENGINE_API UMaterialExpressionFontSampleParameter(const FObjectInitializer& ObjectInitializer = FObjectInitializer::Get());
+	GENERATED_UCLASS_BODY()
 
 	/** name to be referenced when we want to find and set thsi parameter */
 	UPROPERTY(EditAnywhere, Category=MaterialExpressionFontSampleParameter)
diff --git a/Engine/Source/Runtime/Engine/Classes/Materials/MaterialExpressionFrac.h b/Engine/Source/Runtime/Engine/Classes/Materials/MaterialExpressionFrac.h
--- a/Engine/Source/Runtime/Engine/Classes/Materials/MaterialExpressionFrac.h
+++ b/Engine/Source/Runtime/Engine/Classes/Materials/MaterialExpressionFrac.h
@@ -8,9 +8,7 @@
 UCLASS(collapsecategories, hidecategories=Object)
 class UMaterialExpressionFrac : public UMaterialExpression
 {
-	GENERATED_BODY()
-public:
-	UMaterialExpressionFrac(const FObjectInitializer& ObjectInitializer = FObjectInitializer::Get());
+	GENERATED_UCLASS_BODY()
 
 	UPROPERTY()
 	FExpressionInput Input;
diff --git a/Engine/Source/Runtime/Engine/Classes/Materials/MaterialExpressionFresnel.h b/Engine/Source/Runtime/Engine/Classes/Materials/MaterialExpressionFresnel.h
--- a/Engine/Source/Runtime/Engine/Classes/Materials/MaterialExpressionFresnel.h
+++ b/Engine/Source/Runtime/Engine/Classes/Materials/MaterialExpressionFresnel.h
@@ -14,9 +14,7 @@
 UCLASS(collapsecategories, hidecategories=Object)
 class UMaterialExpressionFresnel : public UMaterialExpression
 {
-	GENERATED_BODY()
-public:
-	UMaterialExpressionFresnel(const FObjectInitializer& ObjectInitializer = FObjectInitializer::Get());
+	GENERATED_UCLASS_BODY()
 
 	UPROPERTY(meta = (RequiredInput = "false", ToolTip = "Defaults to 'Exponent' if not specified"))
 	FExpressionInput ExponentIn;
diff --git a/Engine/Source/Runtime/Engine/Classes/Materials/MaterialExpressionFunctionInput.h b/Engine/Source/Runtime/Engine/Classes/Materials/MaterialExpressionFunctionInput.h
--- a/Engine/Source/Runtime/Engine/Classes/Materials/MaterialExpressionFunctionInput.h
+++ b/Engine/Source/Runtime/Engine/Classes/Materials/MaterialExpressionFunctionInput.h
@@ -23,9 +23,7 @@ enum EFunctionInputType
 UCLASS(hidecategories=object, MinimalAPI)
 class UMaterialExpressionFunctionInput : public UMaterialExpression
 {
-	GENERATED_BODY()
-public:
-	ENGINE_API UMaterialExpressionFunctionInput(const FObjectInitializer& ObjectInitializer = FObjectInitializer::Get());
+	GENERATED_UCLASS_BODY()
 
 	/** Used for previewing when editing the function, also temporarily used to connect to the outside material when compiling that material. */
 	UPROPERTY(meta=(RequiredInput = "false"))
diff --git a/Engine/Source/Runtime/Engine/Classes/Materials/MaterialExpressionFunctionOutput.h b/Engine/Source/Runtime/Engine/Classes/Materials/MaterialExpressionFunctionOutput.h
--- a/Engine/Source/Runtime/Engine/Classes/Materials/MaterialExpressionFunctionOutput.h
+++ b/Engine/Source/Runtime/Engine/Classes/Materials/MaterialExpressionFunctionOutput.h
@@ -8,9 +8,7 @@
 UCLASS(hidecategories=object, MinimalAPI)
 class UMaterialExpressionFunctionOutput : public UMaterialExpression
 {
-	GENERATED_BODY()
-public:
-	ENGINE_API UMaterialExpressionFunctionOutput(const FObjectInitializer& ObjectInitializer = FObjectInitializer::Get());
+	GENERATED_UCLASS_BODY()
 
 	/** The output's name, which will be drawn on the connector in function call expressions that use this function. */
 	UPROPERTY(EditAnywhere, Category=MaterialExpressionFunctionOutput)
diff --git a/Engine/Source/Runtime/Engine/Classes/Materials/MaterialExpressionGIReplace.h b/Engine/Source/Runtime/Engine/Classes/Materials/MaterialExpressionGIReplace.h
--- a/Engine/Source/Runtime/Engine/Classes/Materials/MaterialExpressionGIReplace.h
+++ b/Engine/Source/Runtime/Engine/Classes/Materials/MaterialExpressionGIReplace.h
@@ -8,9 +8,7 @@
 UCLASS()
 class UMaterialExpressionGIReplace : public UMaterialExpression
 {
-	GENERATED_BODY()
-public:
-	UMaterialExpressionGIReplace(const FObjectInitializer& ObjectInitializer = FObjectInitializer::Get());
+	GENERATED_UCLASS_BODY()
 
 	/** Used for direct lighting computations e.g. real-time shaders */
 	UPROPERTY()
diff --git a/Engine/Source/Runtime/Engine/Classes/Materials/MaterialExpressionIf.h b/Engine/Source/Runtime/Engine/Classes/Materials/MaterialExpressionIf.h
--- a/Engine/Source/Runtime/Engine/Classes/Materials/MaterialExpressionIf.h
+++ b/Engine/Source/Runtime/Engine/Classes/Materials/MaterialExpressionIf.h
@@ -8,9 +8,7 @@
 UCLASS(collapsecategories, hidecategories=Object, MinimalAPI)
 class UMaterialExpressionIf : public UMaterialExpression
 {
-	GENERATED_BODY()
-public:
-	ENGINE_API UMaterialExpressionIf(const FObjectInitializer& ObjectInitializer = FObjectInitializer::Get());
+	GENERATED_UCLASS_BODY()
 
 	UPROPERTY()
 	FExpressionInput A;
diff --git a/Engine/Source/Runtime/Engine/Classes/Materials/MaterialExpressionLightVector.h b/Engine/Source/Runtime/Engine/Classes/Materials/MaterialExpressionLightVector.h
--- a/Engine/Source/Runtime/Engine/Classes/Materials/MaterialExpressionLightVector.h
+++ b/Engine/Source/Runtime/Engine/Classes/Materials/MaterialExpressionLightVector.h
@@ -8,9 +8,7 @@
 UCLASS(collapsecategories, hidecategories=Object)
 class UMaterialExpressionLightVector : public UMaterialExpression
 {
-	GENERATED_BODY()
-public:
-	UMaterialExpressionLightVector(const FObjectInitializer& ObjectInitializer = FObjectInitializer::Get());
+	GENERATED_UCLASS_BODY()
 
 
 	// Begin UMaterialExpression Interface
diff --git a/Engine/Source/Runtime/Engine/Classes/Materials/MaterialExpressionLightmapUVs.h b/Engine/Source/Runtime/Engine/Classes/Materials/MaterialExpressionLightmapUVs.h
--- a/Engine/Source/Runtime/Engine/Classes/Materials/MaterialExpressionLightmapUVs.h
+++ b/Engine/Source/Runtime/Engine/Classes/Materials/MaterialExpressionLightmapUVs.h
@@ -12,9 +12,7 @@
 UCLASS(collapsecategories, hidecategories=Object)
 class UMaterialExpressionLightmapUVs : public UMaterialExpression
 {
-	GENERATED_BODY()
-public:
-	UMaterialExpressionLightmapUVs(const FObjectInitializer& ObjectInitializer = FObjectInitializer::Get());
+	GENERATED_UCLASS_BODY()
 
 
 	// Begin UMaterialExpression Interface
diff --git a/Engine/Source/Runtime/Engine/Classes/Materials/MaterialExpressionLightmassReplace.h b/Engine/Source/Runtime/Engine/Classes/Materials/MaterialExpressionLightmassReplace.h
--- a/Engine/Source/Runtime/Engine/Classes/Materials/MaterialExpressionLightmassReplace.h
+++ b/Engine/Source/Runtime/Engine/Classes/Materials/MaterialExpressionLightmassReplace.h
@@ -8,9 +8,7 @@
 UCLASS()
 class UMaterialExpressionLightmassReplace : public UMaterialExpression
 {
-	GENERATED_BODY()
-public:
-	UMaterialExpressionLightmassReplace(const FObjectInitializer& ObjectInitializer = FObjectInitializer::Get());
+	GENERATED_UCLASS_BODY()
 
 	UPROPERTY()
 	FExpressionInput Realtime;
diff --git a/Engine/Source/Runtime/Engine/Classes/Materials/MaterialExpressionLinearInterpolate.h b/Engine/Source/Runtime/Engine/Classes/Materials/MaterialExpressionLinearInterpolate.h
--- a/Engine/Source/Runtime/Engine/Classes/Materials/MaterialExpressionLinearInterpolate.h
+++ b/Engine/Source/Runtime/Engine/Classes/Materials/MaterialExpressionLinearInterpolate.h
@@ -8,9 +8,7 @@
 UCLASS(MinimalAPI)
 class UMaterialExpressionLinearInterpolate : public UMaterialExpression
 {
-	GENERATED_BODY()
-public:
-	ENGINE_API UMaterialExpressionLinearInterpolate(const FObjectInitializer& ObjectInitializer = FObjectInitializer::Get());
+	GENERATED_UCLASS_BODY()
 
 	UPROPERTY(meta = (RequiredInput = "false", ToolTip = "Defaults to 'ConstA' if not specified"))
 	FExpressionInput A;
diff --git a/Engine/Source/Runtime/Engine/Classes/Materials/MaterialExpressionMakeMaterialAttributes.h b/Engine/Source/Runtime/Engine/Classes/Materials/MaterialExpressionMakeMaterialAttributes.h
--- a/Engine/Source/Runtime/Engine/Classes/Materials/MaterialExpressionMakeMaterialAttributes.h
+++ b/Engine/Source/Runtime/Engine/Classes/Materials/MaterialExpressionMakeMaterialAttributes.h
@@ -8,9 +8,7 @@
 UCLASS(collapsecategories, hidecategories=Object)
 class UMaterialExpressionMakeMaterialAttributes : public UMaterialExpression
 {
-	GENERATED_BODY()
-public:
-	UMaterialExpressionMakeMaterialAttributes(const FObjectInitializer& ObjectInitializer = FObjectInitializer::Get());
+	GENERATED_UCLASS_BODY()
 
 	UPROPERTY()
 	FExpressionInput BaseColor;
diff --git a/Engine/Source/Runtime/Engine/Classes/Materials/MaterialExpressionMaterialFunctionCall.h b/Engine/Source/Runtime/Engine/Classes/Materials/MaterialExpressionMaterialFunctionCall.h
--- a/Engine/Source/Runtime/Engine/Classes/Materials/MaterialExpressionMaterialFunctionCall.h
+++ b/Engine/Source/Runtime/Engine/Classes/Materials/MaterialExpressionMaterialFunctionCall.h
@@ -12,7 +12,7 @@
 USTRUCT()
 struct FFunctionExpressionInput
 {
-	GENERATED_BODY()
+	GENERATED_USTRUCT_BODY()
 
 	/** 
 	 * Reference to the FunctionInput in the material function.  
@@ -41,7 +41,7 @@ struct FFunctionExpressionInput
 USTRUCT()
 struct FFunctionExpressionOutput
 {
-	GENERATED_BODY()
+	GENERATED_USTRUCT_BODY()
 
 	/** 
 	 * Reference to the FunctionOutput in the material function.  
@@ -69,9 +69,7 @@ struct FFunctionExpressionOutput
 UCLASS(hidecategories=object, MinimalAPI)
 class UMaterialExpressionMaterialFunctionCall : public UMaterialExpression
 {
-	GENERATED_BODY()
-public:
-	ENGINE_API UMaterialExpressionMaterialFunctionCall(const FObjectInitializer& ObjectInitializer = FObjectInitializer::Get());
+	GENERATED_UCLASS_BODY()
 
 	/** The function to call. */
 	UPROPERTY(EditAnywhere, Category=MaterialExpressionMaterialFunctionCall)
diff --git a/Engine/Source/Runtime/Engine/Classes/Materials/MaterialExpressionMax.h b/Engine/Source/Runtime/Engine/Classes/Materials/MaterialExpressionMax.h
--- a/Engine/Source/Runtime/Engine/Classes/Materials/MaterialExpressionMax.h
+++ b/Engine/Source/Runtime/Engine/Classes/Materials/MaterialExpressionMax.h
@@ -8,9 +8,7 @@
 UCLASS(MinimalAPI)
 class UMaterialExpressionMax : public UMaterialExpression
 {
-	GENERATED_BODY()
-public:
-	ENGINE_API UMaterialExpressionMax(const FObjectInitializer& ObjectInitializer = FObjectInitializer::Get());
+	GENERATED_UCLASS_BODY()
 
 	UPROPERTY(meta = (RequiredInput = "false", ToolTip = "Defaults to 'ConstA' if not specified"))
 	FExpressionInput A;
diff --git a/Engine/Source/Runtime/Engine/Classes/Materials/MaterialExpressionMin.h b/Engine/Source/Runtime/Engine/Classes/Materials/MaterialExpressionMin.h
--- a/Engine/Source/Runtime/Engine/Classes/Materials/MaterialExpressionMin.h
+++ b/Engine/Source/Runtime/Engine/Classes/Materials/MaterialExpressionMin.h
@@ -8,9 +8,7 @@
 UCLASS(MinimalAPI)
 class UMaterialExpressionMin : public UMaterialExpression
 {
-	GENERATED_BODY()
-public:
-	ENGINE_API UMaterialExpressionMin(const FObjectInitializer& ObjectInitializer = FObjectInitializer::Get());
+	GENERATED_UCLASS_BODY()
 
 	UPROPERTY(meta = (RequiredInput = "false", ToolTip = "Defaults to 'ConstA' if not specified"))
 	FExpressionInput A;
diff --git a/Engine/Source/Runtime/Engine/Classes/Materials/MaterialExpressionMultiply.h b/Engine/Source/Runtime/Engine/Classes/Materials/MaterialExpressionMultiply.h
--- a/Engine/Source/Runtime/Engine/Classes/Materials/MaterialExpressionMultiply.h
+++ b/Engine/Source/Runtime/Engine/Classes/Materials/MaterialExpressionMultiply.h
@@ -8,9 +8,7 @@
 UCLASS(MinimalAPI)
 class UMaterialExpressionMultiply : public UMaterialExpression
 {
-	GENERATED_BODY()
-public:
-	ENGINE_API UMaterialExpressionMultiply(const FObjectInitializer& ObjectInitializer = FObjectInitializer::Get());
+	GENERATED_UCLASS_BODY()
 
 	UPROPERTY(meta = (RequiredInput = "false", ToolTip = "Defaults to 'ConstA' if not specified"))
 	FExpressionInput A;
diff --git a/Engine/Source/Runtime/Engine/Classes/Materials/MaterialExpressionNoise.h b/Engine/Source/Runtime/Engine/Classes/Materials/MaterialExpressionNoise.h
--- a/Engine/Source/Runtime/Engine/Classes/Materials/MaterialExpressionNoise.h
+++ b/Engine/Source/Runtime/Engine/Classes/Materials/MaterialExpressionNoise.h
@@ -22,9 +22,7 @@ enum ENoiseFunction
 UCLASS()
 class UMaterialExpressionNoise : public UMaterialExpression
 {
-	GENERATED_BODY()
-public:
-	UMaterialExpressionNoise(const FObjectInitializer& ObjectInitializer = FObjectInitializer::Get());
+	GENERATED_UCLASS_BODY()
 
 	/** 2 to 3 dimensional vector */
 	UPROPERTY()
diff --git a/Engine/Source/Runtime/Engine/Classes/Materials/MaterialExpressionNormalize.h b/Engine/Source/Runtime/Engine/Classes/Materials/MaterialExpressionNormalize.h
--- a/Engine/Source/Runtime/Engine/Classes/Materials/MaterialExpressionNormalize.h
+++ b/Engine/Source/Runtime/Engine/Classes/Materials/MaterialExpressionNormalize.h
@@ -8,9 +8,7 @@
 UCLASS(MinimalAPI)
 class UMaterialExpressionNormalize : public UMaterialExpression
 {
-	GENERATED_BODY()
-public:
-	ENGINE_API UMaterialExpressionNormalize(const FObjectInitializer& ObjectInitializer = FObjectInitializer::Get());
+	GENERATED_UCLASS_BODY()
 
 	UPROPERTY()
 	FExpressionInput VectorInput;
diff --git a/Engine/Source/Runtime/Engine/Classes/Materials/MaterialExpressionObjectBounds.h b/Engine/Source/Runtime/Engine/Classes/Materials/MaterialExpressionObjectBounds.h
--- a/Engine/Source/Runtime/Engine/Classes/Materials/MaterialExpressionObjectBounds.h
+++ b/Engine/Source/Runtime/Engine/Classes/Materials/MaterialExpressionObjectBounds.h
@@ -8,9 +8,7 @@
 UCLASS(collapsecategories, hidecategories=Object)
 class UMaterialExpressionObjectBounds : public UMaterialExpression
 {
-	GENERATED_BODY()
-public:
-	UMaterialExpressionObjectBounds(const FObjectInitializer& ObjectInitializer = FObjectInitializer::Get());
+	GENERATED_UCLASS_BODY()
 
 
 	// Begin UMaterialExpression Interface
diff --git a/Engine/Source/Runtime/Engine/Classes/Materials/MaterialExpressionObjectOrientation.h b/Engine/Source/Runtime/Engine/Classes/Materials/MaterialExpressionObjectOrientation.h
--- a/Engine/Source/Runtime/Engine/Classes/Materials/MaterialExpressionObjectOrientation.h
+++ b/Engine/Source/Runtime/Engine/Classes/Materials/MaterialExpressionObjectOrientation.h
@@ -8,9 +8,7 @@
 UCLASS(collapsecategories, hidecategories=Object)
 class UMaterialExpressionObjectOrientation : public UMaterialExpression
 {
-	GENERATED_BODY()
-public:
-	UMaterialExpressionObjectOrientation(const FObjectInitializer& ObjectInitializer = FObjectInitializer::Get());
+	GENERATED_UCLASS_BODY()
 
 
 	// Begin UMaterialExpression Interface
diff --git a/Engine/Source/Runtime/Engine/Classes/Materials/MaterialExpressionObjectRadius.h b/Engine/Source/Runtime/Engine/Classes/Materials/MaterialExpressionObjectRadius.h
--- a/Engine/Source/Runtime/Engine/Classes/Materials/MaterialExpressionObjectRadius.h
+++ b/Engine/Source/Runtime/Engine/Classes/Materials/MaterialExpressionObjectRadius.h
@@ -8,9 +8,7 @@
 UCLASS(collapsecategories, hidecategories=Object)
 class UMaterialExpressionObjectRadius : public UMaterialExpression
 {
-	GENERATED_BODY()
-public:
-	UMaterialExpressionObjectRadius(const FObjectInitializer& ObjectInitializer = FObjectInitializer::Get());
+	GENERATED_UCLASS_BODY()
 
 
 	// Begin UMaterialExpression Interface
diff --git a/Engine/Source/Runtime/Engine/Classes/Materials/MaterialExpressionOneMinus.h b/Engine/Source/Runtime/Engine/Classes/Materials/MaterialExpressionOneMinus.h
--- a/Engine/Source/Runtime/Engine/Classes/Materials/MaterialExpressionOneMinus.h
+++ b/Engine/Source/Runtime/Engine/Classes/Materials/MaterialExpressionOneMinus.h
@@ -8,9 +8,7 @@
 UCLASS(collapsecategories, hidecategories=Object, MinimalAPI)
 class UMaterialExpressionOneMinus : public UMaterialExpression
 {
-	GENERATED_BODY()
-public:
-	ENGINE_API UMaterialExpressionOneMinus(const FObjectInitializer& ObjectInitializer = FObjectInitializer::Get());
+	GENERATED_UCLASS_BODY()
 
 	UPROPERTY()
 	FExpressionInput Input;
diff --git a/Engine/Source/Runtime/Engine/Classes/Materials/MaterialExpressionPanner.h b/Engine/Source/Runtime/Engine/Classes/Materials/MaterialExpressionPanner.h
--- a/Engine/Source/Runtime/Engine/Classes/Materials/MaterialExpressionPanner.h
+++ b/Engine/Source/Runtime/Engine/Classes/Materials/MaterialExpressionPanner.h
@@ -8,9 +8,7 @@
 UCLASS(collapsecategories, hidecategories=Object, MinimalAPI)
 class UMaterialExpressionPanner : public UMaterialExpression
 {
-	GENERATED_BODY()
-public:
-	ENGINE_API UMaterialExpressionPanner(const FObjectInitializer& ObjectInitializer = FObjectInitializer::Get());
+	GENERATED_UCLASS_BODY()
 
 	UPROPERTY(meta = (RequiredInput = "false", ToolTip = "Defaults to 'ConstCoordinate' if not specified"))
 	FExpressionInput Coordinate;
diff --git a/Engine/Source/Runtime/Engine/Classes/Materials/MaterialExpressionParameter.h b/Engine/Source/Runtime/Engine/Classes/Materials/MaterialExpressionParameter.h
--- a/Engine/Source/Runtime/Engine/Classes/Materials/MaterialExpressionParameter.h
+++ b/Engine/Source/Runtime/Engine/Classes/Materials/MaterialExpressionParameter.h
@@ -8,9 +8,7 @@
 UCLASS(collapsecategories, hidecategories=Object, MinimalAPI)
 class UMaterialExpressionParameter : public UMaterialExpression
 {
-	GENERATED_BODY()
-public:
-	ENGINE_API UMaterialExpressionParameter(const FObjectInitializer& ObjectInitializer = FObjectInitializer::Get());
+	GENERATED_UCLASS_BODY()
 
 	/** The name of the parameter */
 	UPROPERTY(EditAnywhere, Category=MaterialExpressionParameter)
diff --git a/Engine/Source/Runtime/Engine/Classes/Materials/MaterialExpressionParticleColor.h b/Engine/Source/Runtime/Engine/Classes/Materials/MaterialExpressionParticleColor.h
--- a/Engine/Source/Runtime/Engine/Classes/Materials/MaterialExpressionParticleColor.h
+++ b/Engine/Source/Runtime/Engine/Classes/Materials/MaterialExpressionParticleColor.h
@@ -8,9 +8,7 @@
 UCLASS(collapsecategories, hidecategories=Object)
 class UMaterialExpressionParticleColor : public UMaterialExpression
 {
-	GENERATED_BODY()
-public:
-	UMaterialExpressionParticleColor(const FObjectInitializer& ObjectInitializer = FObjectInitializer::Get());
+	GENERATED_UCLASS_BODY()
 
 
 	// Begin UMaterialExpression Interface
diff --git a/Engine/Source/Runtime/Engine/Classes/Materials/MaterialExpressionParticleDirection.h b/Engine/Source/Runtime/Engine/Classes/Materials/MaterialExpressionParticleDirection.h
--- a/Engine/Source/Runtime/Engine/Classes/Materials/MaterialExpressionParticleDirection.h
+++ b/Engine/Source/Runtime/Engine/Classes/Materials/MaterialExpressionParticleDirection.h
@@ -12,9 +12,7 @@
 UCLASS(collapsecategories, hidecategories=Object)
 class UMaterialExpressionParticleDirection : public UMaterialExpression
 {
-	GENERATED_BODY()
-public:
-	UMaterialExpressionParticleDirection(const FObjectInitializer& ObjectInitializer = FObjectInitializer::Get());
+	GENERATED_UCLASS_BODY()
 
 
 	// Begin UMaterialExpression Interface
diff --git a/Engine/Source/Runtime/Engine/Classes/Materials/MaterialExpressionParticleMacroUV.h b/Engine/Source/Runtime/Engine/Classes/Materials/MaterialExpressionParticleMacroUV.h
--- a/Engine/Source/Runtime/Engine/Classes/Materials/MaterialExpressionParticleMacroUV.h
+++ b/Engine/Source/Runtime/Engine/Classes/Materials/MaterialExpressionParticleMacroUV.h
@@ -12,9 +12,7 @@
 UCLASS()
 class UMaterialExpressionParticleMacroUV : public UMaterialExpression
 {
-	GENERATED_BODY()
-public:
-	UMaterialExpressionParticleMacroUV(const FObjectInitializer& ObjectInitializer = FObjectInitializer::Get());
+	GENERATED_UCLASS_BODY()
 
 
 	// Begin UMaterialExpression Interface
diff --git a/Engine/Source/Runtime/Engine/Classes/Materials/MaterialExpressionParticleMotionBlurFade.h b/Engine/Source/Runtime/Engine/Classes/Materials/MaterialExpressionParticleMotionBlurFade.h
--- a/Engine/Source/Runtime/Engine/Classes/Materials/MaterialExpressionParticleMotionBlurFade.h
+++ b/Engine/Source/Runtime/Engine/Classes/Materials/MaterialExpressionParticleMotionBlurFade.h
@@ -12,9 +12,7 @@
 UCLASS(collapsecategories, hidecategories=Object)
 class UMaterialExpressionParticleMotionBlurFade : public UMaterialExpression
 {
-	GENERATED_BODY()
-public:
-	UMaterialExpressionParticleMotionBlurFade(const FObjectInitializer& ObjectInitializer = FObjectInitializer::Get());
+	GENERATED_UCLASS_BODY()
 
 
 	// Begin UMaterialExpression Interface
diff --git a/Engine/Source/Runtime/Engine/Classes/Materials/MaterialExpressionParticlePositionWS.h b/Engine/Source/Runtime/Engine/Classes/Materials/MaterialExpressionParticlePositionWS.h
--- a/Engine/Source/Runtime/Engine/Classes/Materials/MaterialExpressionParticlePositionWS.h
+++ b/Engine/Source/Runtime/Engine/Classes/Materials/MaterialExpressionParticlePositionWS.h
@@ -8,9 +8,7 @@
 UCLASS(collapsecategories, hidecategories=Object)
 class UMaterialExpressionParticlePositionWS : public UMaterialExpression
 {
-	GENERATED_BODY()
-public:
-	UMaterialExpressionParticlePositionWS(const FObjectInitializer& ObjectInitializer = FObjectInitializer::Get());
+	GENERATED_UCLASS_BODY()
 
 
 	// Begin UMaterialExpression Interface
diff --git a/Engine/Source/Runtime/Engine/Classes/Materials/MaterialExpressionParticleRadius.h b/Engine/Source/Runtime/Engine/Classes/Materials/MaterialExpressionParticleRadius.h
--- a/Engine/Source/Runtime/Engine/Classes/Materials/MaterialExpressionParticleRadius.h
+++ b/Engine/Source/Runtime/Engine/Classes/Materials/MaterialExpressionParticleRadius.h
@@ -8,9 +8,7 @@
 UCLASS(collapsecategories, hidecategories=Object)
 class UMaterialExpressionParticleRadius : public UMaterialExpression
 {
-	GENERATED_BODY()
-public:
-	UMaterialExpressionParticleRadius(const FObjectInitializer& ObjectInitializer = FObjectInitializer::Get());
+	GENERATED_UCLASS_BODY()
 
 
 	// Begin UMaterialExpression Interface
diff --git a/Engine/Source/Runtime/Engine/Classes/Materials/MaterialExpressionParticleRelativeTime.h b/Engine/Source/Runtime/Engine/Classes/Materials/MaterialExpressionParticleRelativeTime.h
--- a/Engine/Source/Runtime/Engine/Classes/Materials/MaterialExpressionParticleRelativeTime.h
+++ b/Engine/Source/Runtime/Engine/Classes/Materials/MaterialExpressionParticleRelativeTime.h
@@ -12,9 +12,7 @@
 UCLASS(collapsecategories, hidecategories=Object)
 class UMaterialExpressionParticleRelativeTime : public UMaterialExpression
 {
-	GENERATED_BODY()
-public:
-	UMaterialExpressionParticleRelativeTime(const FObjectInitializer& ObjectInitializer = FObjectInitializer::Get());
+	GENERATED_UCLASS_BODY()
 
 
 	// Begin UMaterialExpression Interface
diff --git a/Engine/Source/Runtime/Engine/Classes/Materials/MaterialExpressionParticleSize.h b/Engine/Source/Runtime/Engine/Classes/Materials/MaterialExpressionParticleSize.h
--- a/Engine/Source/Runtime/Engine/Classes/Materials/MaterialExpressionParticleSize.h
+++ b/Engine/Source/Runtime/Engine/Classes/Materials/MaterialExpressionParticleSize.h
@@ -12,9 +12,7 @@
 UCLASS(collapsecategories, hidecategories=Object)
 class UMaterialExpressionParticleSize : public UMaterialExpression
 {
-	GENERATED_BODY()
-public:
-	UMaterialExpressionParticleSize(const FObjectInitializer& ObjectInitializer = FObjectInitializer::Get());
+	GENERATED_UCLASS_BODY()
 
 
 	// Begin UMaterialExpression Interface
diff --git a/Engine/Source/Runtime/Engine/Classes/Materials/MaterialExpressionParticleSpeed.h b/Engine/Source/Runtime/Engine/Classes/Materials/MaterialExpressionParticleSpeed.h
--- a/Engine/Source/Runtime/Engine/Classes/Materials/MaterialExpressionParticleSpeed.h
+++ b/Engine/Source/Runtime/Engine/Classes/Materials/MaterialExpressionParticleSpeed.h
@@ -12,9 +12,7 @@
 UCLASS(collapsecategories, hidecategories=Object)
 class UMaterialExpressionParticleSpeed : public UMaterialExpression
 {
-	GENERATED_BODY()
-public:
-	UMaterialExpressionParticleSpeed(const FObjectInitializer& ObjectInitializer = FObjectInitializer::Get());
+	GENERATED_UCLASS_BODY()
 
 
 	// Begin UMaterialExpression Interface
diff --git a/Engine/Source/Runtime/Engine/Classes/Materials/MaterialExpressionParticleSubUV.h b/Engine/Source/Runtime/Engine/Classes/Materials/MaterialExpressionParticleSubUV.h
--- a/Engine/Source/Runtime/Engine/Classes/Materials/MaterialExpressionParticleSubUV.h
+++ b/Engine/Source/Runtime/Engine/Classes/Materials/MaterialExpressionParticleSubUV.h
@@ -8,9 +8,7 @@
 UCLASS(MinimalAPI)
 class UMaterialExpressionParticleSubUV : public UMaterialExpressionTextureSample
 {
-	GENERATED_BODY()
-public:
-	ENGINE_API UMaterialExpressionParticleSubUV(const FObjectInitializer& ObjectInitializer = FObjectInitializer::Get());
+	GENERATED_UCLASS_BODY()
 
 	UPROPERTY(EditAnywhere, Category=MaterialExpressionParticleSubUV)
 	uint32 bBlend:1;
diff --git a/Engine/Source/Runtime/Engine/Classes/Materials/MaterialExpressionPerInstanceFadeAmount.h b/Engine/Source/Runtime/Engine/Classes/Materials/MaterialExpressionPerInstanceFadeAmount.h
--- a/Engine/Source/Runtime/Engine/Classes/Materials/MaterialExpressionPerInstanceFadeAmount.h
+++ b/Engine/Source/Runtime/Engine/Classes/Materials/MaterialExpressionPerInstanceFadeAmount.h
@@ -8,9 +8,7 @@
 UCLASS(collapsecategories, hidecategories=Object)
 class UMaterialExpressionPerInstanceFadeAmount : public UMaterialExpression
 {
-	GENERATED_BODY()
-public:
-	UMaterialExpressionPerInstanceFadeAmount(const FObjectInitializer& ObjectInitializer = FObjectInitializer::Get());
+	GENERATED_UCLASS_BODY()
 
 
 	// Begin UMaterialExpression Interface
diff --git a/Engine/Source/Runtime/Engine/Classes/Materials/MaterialExpressionPerInstanceRandom.h b/Engine/Source/Runtime/Engine/Classes/Materials/MaterialExpressionPerInstanceRandom.h
--- a/Engine/Source/Runtime/Engine/Classes/Materials/MaterialExpressionPerInstanceRandom.h
+++ b/Engine/Source/Runtime/Engine/Classes/Materials/MaterialExpressionPerInstanceRandom.h
@@ -8,9 +8,7 @@
 UCLASS(collapsecategories, hidecategories=Object)
 class UMaterialExpressionPerInstanceRandom : public UMaterialExpression
 {
-	GENERATED_BODY()
-public:
-	UMaterialExpressionPerInstanceRandom(const FObjectInitializer& ObjectInitializer = FObjectInitializer::Get());
+	GENERATED_UCLASS_BODY()
 
 
 	// Begin UMaterialExpression Interface
diff --git a/Engine/Source/Runtime/Engine/Classes/Materials/MaterialExpressionPixelDepth.h b/Engine/Source/Runtime/Engine/Classes/Materials/MaterialExpressionPixelDepth.h
--- a/Engine/Source/Runtime/Engine/Classes/Materials/MaterialExpressionPixelDepth.h
+++ b/Engine/Source/Runtime/Engine/Classes/Materials/MaterialExpressionPixelDepth.h
@@ -11,9 +11,7 @@
 UCLASS(collapsecategories, hidecategories=Object)
 class UMaterialExpressionPixelDepth : public UMaterialExpression
 {
-	GENERATED_BODY()
-public:
-	UMaterialExpressionPixelDepth(const FObjectInitializer& ObjectInitializer = FObjectInitializer::Get());
+	GENERATED_UCLASS_BODY()
 
 
 	// Begin UMaterialExpression Interface
diff --git a/Engine/Source/Runtime/Engine/Classes/Materials/MaterialExpressionPixelNormalWS.h b/Engine/Source/Runtime/Engine/Classes/Materials/MaterialExpressionPixelNormalWS.h
--- a/Engine/Source/Runtime/Engine/Classes/Materials/MaterialExpressionPixelNormalWS.h
+++ b/Engine/Source/Runtime/Engine/Classes/Materials/MaterialExpressionPixelNormalWS.h
@@ -8,9 +8,7 @@
 UCLASS(collapsecategories, hidecategories=Object)
 class UMaterialExpressionPixelNormalWS : public UMaterialExpression
 {
-	GENERATED_BODY()
-public:
-	UMaterialExpressionPixelNormalWS(const FObjectInitializer& ObjectInitializer = FObjectInitializer::Get());
+	GENERATED_UCLASS_BODY()
 
 
 	// Begin UMaterialExpression Interface
diff --git a/Engine/Source/Runtime/Engine/Classes/Materials/MaterialExpressionPower.h b/Engine/Source/Runtime/Engine/Classes/Materials/MaterialExpressionPower.h
--- a/Engine/Source/Runtime/Engine/Classes/Materials/MaterialExpressionPower.h
+++ b/Engine/Source/Runtime/Engine/Classes/Materials/MaterialExpressionPower.h
@@ -8,9 +8,7 @@
 UCLASS(collapsecategories, hidecategories=Object, MinimalAPI)
 class UMaterialExpressionPower : public UMaterialExpression
 {
-	GENERATED_BODY()
-public:
-	ENGINE_API UMaterialExpressionPower(const FObjectInitializer& ObjectInitializer = FObjectInitializer::Get());
+	GENERATED_UCLASS_BODY()
 
 	UPROPERTY()
 	FExpressionInput Base;
diff --git a/Engine/Source/Runtime/Engine/Classes/Materials/MaterialExpressionQualitySwitch.h b/Engine/Source/Runtime/Engine/Classes/Materials/MaterialExpressionQualitySwitch.h
--- a/Engine/Source/Runtime/Engine/Classes/Materials/MaterialExpressionQualitySwitch.h
+++ b/Engine/Source/Runtime/Engine/Classes/Materials/MaterialExpressionQualitySwitch.h
@@ -8,9 +8,7 @@
 UCLASS(collapsecategories, hidecategories=Object, MinimalAPI)
 class UMaterialExpressionQualitySwitch : public UMaterialExpression
 {
-	GENERATED_BODY()
-public:
-	ENGINE_API UMaterialExpressionQualitySwitch(const FObjectInitializer& ObjectInitializer = FObjectInitializer::Get());
+	GENERATED_UCLASS_BODY()
 
 	/** Default connection, used when a specific quality level input is missing. */
 	UPROPERTY()
diff --git a/Engine/Source/Runtime/Engine/Classes/Materials/MaterialExpressionReflectionVectorWS.h b/Engine/Source/Runtime/Engine/Classes/Materials/MaterialExpressionReflectionVectorWS.h
--- a/Engine/Source/Runtime/Engine/Classes/Materials/MaterialExpressionReflectionVectorWS.h
+++ b/Engine/Source/Runtime/Engine/Classes/Materials/MaterialExpressionReflectionVectorWS.h
@@ -8,9 +8,7 @@
 UCLASS(collapsecategories, hidecategories=Object, MinimalAPI)
 class UMaterialExpressionReflectionVectorWS : public UMaterialExpression
 {
-	GENERATED_BODY()
-public:
-	ENGINE_API UMaterialExpressionReflectionVectorWS(const FObjectInitializer& ObjectInitializer = FObjectInitializer::Get());
+	GENERATED_UCLASS_BODY()
 
 	/** Optional world normal to reflect the camera view vector about. If unconnected, pixel normal is used */
 	UPROPERTY(meta = (RequiredInput = "false", ToolTip = "Defaults to Reflection Vector if not specified"))
diff --git a/Engine/Source/Runtime/Engine/Classes/Materials/MaterialExpressionRotateAboutAxis.h b/Engine/Source/Runtime/Engine/Classes/Materials/MaterialExpressionRotateAboutAxis.h
--- a/Engine/Source/Runtime/Engine/Classes/Materials/MaterialExpressionRotateAboutAxis.h
+++ b/Engine/Source/Runtime/Engine/Classes/Materials/MaterialExpressionRotateAboutAxis.h
@@ -8,9 +8,7 @@
 UCLASS(MinimalAPI)
 class UMaterialExpressionRotateAboutAxis : public UMaterialExpression
 {
-	GENERATED_BODY()
-public:
-	ENGINE_API UMaterialExpressionRotateAboutAxis(const FObjectInitializer& ObjectInitializer = FObjectInitializer::Get());
+	GENERATED_UCLASS_BODY()
 
 	UPROPERTY()
 	FExpressionInput NormalizedRotationAxis;
diff --git a/Engine/Source/Runtime/Engine/Classes/Materials/MaterialExpressionRotator.h b/Engine/Source/Runtime/Engine/Classes/Materials/MaterialExpressionRotator.h
--- a/Engine/Source/Runtime/Engine/Classes/Materials/MaterialExpressionRotator.h
+++ b/Engine/Source/Runtime/Engine/Classes/Materials/MaterialExpressionRotator.h
@@ -8,9 +8,7 @@
 UCLASS(collapsecategories, hidecategories=Object)
 class UMaterialExpressionRotator : public UMaterialExpression
 {
-	GENERATED_BODY()
-public:
-	UMaterialExpressionRotator(const FObjectInitializer& ObjectInitializer = FObjectInitializer::Get());
+	GENERATED_UCLASS_BODY()
 
 	UPROPERTY(meta = (RequiredInput = "false", ToolTip = "Defaults to 'ConstCoordinate' if not specified"))
 	FExpressionInput Coordinate;
diff --git a/Engine/Source/Runtime/Engine/Classes/Materials/MaterialExpressionScalarParameter.h b/Engine/Source/Runtime/Engine/Classes/Materials/MaterialExpressionScalarParameter.h
--- a/Engine/Source/Runtime/Engine/Classes/Materials/MaterialExpressionScalarParameter.h
+++ b/Engine/Source/Runtime/Engine/Classes/Materials/MaterialExpressionScalarParameter.h
@@ -8,9 +8,7 @@
 UCLASS(collapsecategories, hidecategories=Object, MinimalAPI)
 class UMaterialExpressionScalarParameter : public UMaterialExpressionParameter
 {
-	GENERATED_BODY()
-public:
-	ENGINE_API UMaterialExpressionScalarParameter(const FObjectInitializer& ObjectInitializer = FObjectInitializer::Get());
+	GENERATED_UCLASS_BODY()
 
 	UPROPERTY(EditAnywhere, Category=MaterialExpressionScalarParameter)
 	float DefaultValue;
diff --git a/Engine/Source/Runtime/Engine/Classes/Materials/MaterialExpressionSceneColor.h b/Engine/Source/Runtime/Engine/Classes/Materials/MaterialExpressionSceneColor.h
--- a/Engine/Source/Runtime/Engine/Classes/Materials/MaterialExpressionSceneColor.h
+++ b/Engine/Source/Runtime/Engine/Classes/Materials/MaterialExpressionSceneColor.h
@@ -18,9 +18,7 @@ namespace EMaterialSceneAttributeInputMode
 UCLASS(collapsecategories, hidecategories=Object)
 class UMaterialExpressionSceneColor : public UMaterialExpression
 {
-	GENERATED_BODY()
-public:
-	UMaterialExpressionSceneColor(const FObjectInitializer& ObjectInitializer = FObjectInitializer::Get());
+	GENERATED_UCLASS_BODY()
 	
 	/**
 	* Coordinates - UV coordinates to apply to the scene color lookup.
diff --git a/Engine/Source/Runtime/Engine/Classes/Materials/MaterialExpressionSceneDepth.h b/Engine/Source/Runtime/Engine/Classes/Materials/MaterialExpressionSceneDepth.h
--- a/Engine/Source/Runtime/Engine/Classes/Materials/MaterialExpressionSceneDepth.h
+++ b/Engine/Source/Runtime/Engine/Classes/Materials/MaterialExpressionSceneDepth.h
@@ -8,9 +8,7 @@
 UCLASS(collapsecategories, hidecategories=Object)
 class UMaterialExpressionSceneDepth : public UMaterialExpression
 {
-	GENERATED_BODY()
-public:
-	UMaterialExpressionSceneDepth(const FObjectInitializer& ObjectInitializer = FObjectInitializer::Get());
+	GENERATED_UCLASS_BODY()
 
 	/** 
 	* Coordinates - UV coordinates to apply to the scene depth lookup.
diff --git a/Engine/Source/Runtime/Engine/Classes/Materials/MaterialExpressionSceneTexelSize.h b/Engine/Source/Runtime/Engine/Classes/Materials/MaterialExpressionSceneTexelSize.h
--- a/Engine/Source/Runtime/Engine/Classes/Materials/MaterialExpressionSceneTexelSize.h
+++ b/Engine/Source/Runtime/Engine/Classes/Materials/MaterialExpressionSceneTexelSize.h
@@ -8,9 +8,7 @@
 UCLASS(collapsecategories, hidecategories=Object)
 class UMaterialExpressionSceneTexelSize : public UMaterialExpression
 {
-	GENERATED_BODY()
-public:
-	UMaterialExpressionSceneTexelSize(const FObjectInitializer& ObjectInitializer = FObjectInitializer::Get());
+	GENERATED_UCLASS_BODY()
 
 	// Begin UMaterialExpression Interface
 	virtual int32 Compile(class FMaterialCompiler* Compiler, int32 OutputIndex, int32 MultiplexIndex) override;
diff --git a/Engine/Source/Runtime/Engine/Classes/Materials/MaterialExpressionScreenPosition.h b/Engine/Source/Runtime/Engine/Classes/Materials/MaterialExpressionScreenPosition.h
--- a/Engine/Source/Runtime/Engine/Classes/Materials/MaterialExpressionScreenPosition.h
+++ b/Engine/Source/Runtime/Engine/Classes/Materials/MaterialExpressionScreenPosition.h
@@ -8,9 +8,7 @@
 UCLASS(collapsecategories, hidecategories=Object)
 class UMaterialExpressionScreenPosition : public UMaterialExpression
 {
-	GENERATED_BODY()
-public:
-	UMaterialExpressionScreenPosition(const FObjectInitializer& ObjectInitializer = FObjectInitializer::Get());
+	GENERATED_UCLASS_BODY()
 
 	// Begin UMaterialExpression Interface
 	virtual int32 Compile(class FMaterialCompiler* Compiler, int32 OutputIndex, int32 MultiplexIndex) override;
diff --git a/Engine/Source/Runtime/Engine/Classes/Materials/MaterialExpressionSine.h b/Engine/Source/Runtime/Engine/Classes/Materials/MaterialExpressionSine.h
--- a/Engine/Source/Runtime/Engine/Classes/Materials/MaterialExpressionSine.h
+++ b/Engine/Source/Runtime/Engine/Classes/Materials/MaterialExpressionSine.h
@@ -8,9 +8,7 @@
 UCLASS(collapsecategories, hidecategories=Object)
 class UMaterialExpressionSine : public UMaterialExpression
 {
-	GENERATED_BODY()
-public:
-	UMaterialExpressionSine(const FObjectInitializer& ObjectInitializer = FObjectInitializer::Get());
+	GENERATED_UCLASS_BODY()
 
 	UPROPERTY()
 	FExpressionInput Input;
diff --git a/Engine/Source/Runtime/Engine/Classes/Materials/MaterialExpressionSpeedTree.h b/Engine/Source/Runtime/Engine/Classes/Materials/MaterialExpressionSpeedTree.h
--- a/Engine/Source/Runtime/Engine/Classes/Materials/MaterialExpressionSpeedTree.h
+++ b/Engine/Source/Runtime/Engine/Classes/Materials/MaterialExpressionSpeedTree.h
@@ -37,9 +37,7 @@ enum ESpeedTreeLODType
 UCLASS(collapsecategories, hidecategories=Object, MinimalAPI)
 class UMaterialExpressionSpeedTree : public UMaterialExpression
 {
-	GENERATED_BODY()
-public:
-	ENGINE_API UMaterialExpressionSpeedTree(const FObjectInitializer& ObjectInitializer = FObjectInitializer::Get());
+	GENERATED_UCLASS_BODY()	
 
 	UPROPERTY(EditAnywhere, Category=MaterialExpressionSpeedTree, meta=(DisplayName = "Geometry Type", ToolTip="The type of SpeedTree geometry on which this material will be used"))
 	TEnumAsByte<enum ESpeedTreeGeometryType> GeometryType;
diff --git a/Engine/Source/Runtime/Engine/Classes/Materials/MaterialExpressionSphereMask.h b/Engine/Source/Runtime/Engine/Classes/Materials/MaterialExpressionSphereMask.h
--- a/Engine/Source/Runtime/Engine/Classes/Materials/MaterialExpressionSphereMask.h
+++ b/Engine/Source/Runtime/Engine/Classes/Materials/MaterialExpressionSphereMask.h
@@ -8,9 +8,7 @@
 UCLASS()
 class UMaterialExpressionSphereMask : public UMaterialExpression
 {
-	GENERATED_BODY()
-public:
-	UMaterialExpressionSphereMask(const FObjectInitializer& ObjectInitializer = FObjectInitializer::Get());
+	GENERATED_UCLASS_BODY()
 
 	/** 1 to 4 dimensional vector, should be the same type as B */
 	UPROPERTY()
diff --git a/Engine/Source/Runtime/Engine/Classes/Materials/MaterialExpressionSphericalParticleOpacity.h b/Engine/Source/Runtime/Engine/Classes/Materials/MaterialExpressionSphericalParticleOpacity.h
--- a/Engine/Source/Runtime/Engine/Classes/Materials/MaterialExpressionSphericalParticleOpacity.h
+++ b/Engine/Source/Runtime/Engine/Classes/Materials/MaterialExpressionSphericalParticleOpacity.h
@@ -8,9 +8,7 @@
 UCLASS(collapsecategories, hidecategories=Object)
 class UMaterialExpressionSphericalParticleOpacity : public UMaterialExpression
 {
-	GENERATED_BODY()
-public:
-	UMaterialExpressionSphericalParticleOpacity(const FObjectInitializer& ObjectInitializer = FObjectInitializer::Get());
+	GENERATED_UCLASS_BODY()
 
 	/** Density of the particle sphere. */
 	UPROPERTY(meta = (RequiredInput = "false", ToolTip = "Defaults to 'ConstantDensity' if not specified"))
diff --git a/Engine/Source/Runtime/Engine/Classes/Materials/MaterialExpressionSquareRoot.h b/Engine/Source/Runtime/Engine/Classes/Materials/MaterialExpressionSquareRoot.h
--- a/Engine/Source/Runtime/Engine/Classes/Materials/MaterialExpressionSquareRoot.h
+++ b/Engine/Source/Runtime/Engine/Classes/Materials/MaterialExpressionSquareRoot.h
@@ -8,9 +8,7 @@
 UCLASS(collapsecategories, hidecategories=Object)
 class UMaterialExpressionSquareRoot : public UMaterialExpression
 {
-	GENERATED_BODY()
-public:
-	UMaterialExpressionSquareRoot(const FObjectInitializer& ObjectInitializer = FObjectInitializer::Get());
+	GENERATED_UCLASS_BODY()
 
 	UPROPERTY()
 	FExpressionInput Input;
diff --git a/Engine/Source/Runtime/Engine/Classes/Materials/MaterialExpressionStaticBool.h b/Engine/Source/Runtime/Engine/Classes/Materials/MaterialExpressionStaticBool.h
--- a/Engine/Source/Runtime/Engine/Classes/Materials/MaterialExpressionStaticBool.h
+++ b/Engine/Source/Runtime/Engine/Classes/Materials/MaterialExpressionStaticBool.h
@@ -8,9 +8,7 @@
 UCLASS(collapsecategories, hidecategories=Object, MinimalAPI)
 class UMaterialExpressionStaticBool : public UMaterialExpression
 {
-	GENERATED_BODY()
-public:
-	ENGINE_API UMaterialExpressionStaticBool(const FObjectInitializer& ObjectInitializer = FObjectInitializer::Get());
+	GENERATED_UCLASS_BODY()
 
 	UPROPERTY(EditAnywhere, Category=MaterialExpressionStaticBool)
 	uint32 Value:1;
diff --git a/Engine/Source/Runtime/Engine/Classes/Materials/MaterialExpressionStaticComponentMaskParameter.h b/Engine/Source/Runtime/Engine/Classes/Materials/MaterialExpressionStaticComponentMaskParameter.h
--- a/Engine/Source/Runtime/Engine/Classes/Materials/MaterialExpressionStaticComponentMaskParameter.h
+++ b/Engine/Source/Runtime/Engine/Classes/Materials/MaterialExpressionStaticComponentMaskParameter.h
@@ -8,9 +8,7 @@
 UCLASS(collapsecategories, hidecategories=Object, MinimalAPI)
 class UMaterialExpressionStaticComponentMaskParameter : public UMaterialExpressionParameter
 {
-	GENERATED_BODY()
-public:
-	ENGINE_API UMaterialExpressionStaticComponentMaskParameter(const FObjectInitializer& ObjectInitializer = FObjectInitializer::Get());
+	GENERATED_UCLASS_BODY()
 
 	UPROPERTY()
 	FExpressionInput Input;
diff --git a/Engine/Source/Runtime/Engine/Classes/Materials/MaterialExpressionStaticSwitch.h b/Engine/Source/Runtime/Engine/Classes/Materials/MaterialExpressionStaticSwitch.h
--- a/Engine/Source/Runtime/Engine/Classes/Materials/MaterialExpressionStaticSwitch.h
+++ b/Engine/Source/Runtime/Engine/Classes/Materials/MaterialExpressionStaticSwitch.h
@@ -8,9 +8,7 @@
 UCLASS(collapsecategories, hidecategories=Object, MinimalAPI)
 class UMaterialExpressionStaticSwitch : public UMaterialExpression
 {
-	GENERATED_BODY()
-public:
-	ENGINE_API UMaterialExpressionStaticSwitch(const FObjectInitializer& ObjectInitializer = FObjectInitializer::Get());
+	GENERATED_UCLASS_BODY()
 
 	UPROPERTY(EditAnywhere, Category=MaterialExpressionStaticSwitch, meta=(OverridingInputProperty = "Value"))
 	uint32 DefaultValue:1;
diff --git a/Engine/Source/Runtime/Engine/Classes/Materials/MaterialExpressionStaticSwitchParameter.h b/Engine/Source/Runtime/Engine/Classes/Materials/MaterialExpressionStaticSwitchParameter.h
--- a/Engine/Source/Runtime/Engine/Classes/Materials/MaterialExpressionStaticSwitchParameter.h
+++ b/Engine/Source/Runtime/Engine/Classes/Materials/MaterialExpressionStaticSwitchParameter.h
@@ -8,9 +8,7 @@
 UCLASS(collapsecategories, hidecategories=Object, MinimalAPI)
 class UMaterialExpressionStaticSwitchParameter : public UMaterialExpressionStaticBoolParameter
 {
-	GENERATED_BODY()
-public:
-	ENGINE_API UMaterialExpressionStaticSwitchParameter(const FObjectInitializer& ObjectInitializer = FObjectInitializer::Get());
+	GENERATED_UCLASS_BODY()
 
 	UPROPERTY()
 	FExpressionInput A;
diff --git a/Engine/Source/Runtime/Engine/Classes/Materials/MaterialExpressionSubtract.h b/Engine/Source/Runtime/Engine/Classes/Materials/MaterialExpressionSubtract.h
--- a/Engine/Source/Runtime/Engine/Classes/Materials/MaterialExpressionSubtract.h
+++ b/Engine/Source/Runtime/Engine/Classes/Materials/MaterialExpressionSubtract.h
@@ -8,9 +8,7 @@
 UCLASS()
 class UMaterialExpressionSubtract : public UMaterialExpression
 {
-	GENERATED_BODY()
-public:
-	UMaterialExpressionSubtract(const FObjectInitializer& ObjectInitializer = FObjectInitializer::Get());
+	GENERATED_UCLASS_BODY()
 
 	UPROPERTY(meta = (RequiredInput = "false", ToolTip = "Defaults to 'ConstA' if not specified"))
 	FExpressionInput A;
diff --git a/Engine/Source/Runtime/Engine/Classes/Materials/MaterialExpressionTextureBase.h b/Engine/Source/Runtime/Engine/Classes/Materials/MaterialExpressionTextureBase.h
--- a/Engine/Source/Runtime/Engine/Classes/Materials/MaterialExpressionTextureBase.h
+++ b/Engine/Source/Runtime/Engine/Classes/Materials/MaterialExpressionTextureBase.h
@@ -12,9 +12,7 @@
 UCLASS(abstract, hidecategories=Object)
 class ENGINE_API UMaterialExpressionTextureBase : public UMaterialExpression 
 {
-	GENERATED_BODY()
-public:
-	UMaterialExpressionTextureBase(const FObjectInitializer& ObjectInitializer = FObjectInitializer::Get());
+	GENERATED_UCLASS_BODY()
 
 	UPROPERTY(EditAnywhere, Category=MaterialExpressionTextureBase)
 	class UTexture* Texture;
diff --git a/Engine/Source/Runtime/Engine/Classes/Materials/MaterialExpressionTextureCoordinate.h b/Engine/Source/Runtime/Engine/Classes/Materials/MaterialExpressionTextureCoordinate.h
--- a/Engine/Source/Runtime/Engine/Classes/Materials/MaterialExpressionTextureCoordinate.h
+++ b/Engine/Source/Runtime/Engine/Classes/Materials/MaterialExpressionTextureCoordinate.h
@@ -8,9 +8,7 @@
 UCLASS(collapsecategories, hidecategories=Object, MinimalAPI)
 class UMaterialExpressionTextureCoordinate : public UMaterialExpression
 {
-	GENERATED_BODY()
-public:
-	ENGINE_API UMaterialExpressionTextureCoordinate(const FObjectInitializer& ObjectInitializer = FObjectInitializer::Get());
+	GENERATED_UCLASS_BODY()
 
 	/** Texture coordinate index */
 	UPROPERTY(EditAnywhere, Category=MaterialExpressionTextureCoordinate)
diff --git a/Engine/Source/Runtime/Engine/Classes/Materials/MaterialExpressionTextureObject.h b/Engine/Source/Runtime/Engine/Classes/Materials/MaterialExpressionTextureObject.h
--- a/Engine/Source/Runtime/Engine/Classes/Materials/MaterialExpressionTextureObject.h
+++ b/Engine/Source/Runtime/Engine/Classes/Materials/MaterialExpressionTextureObject.h
@@ -12,9 +12,7 @@
 UCLASS(collapsecategories, hidecategories=Object, MinimalAPI)
 class UMaterialExpressionTextureObject : public UMaterialExpressionTextureBase
 {
-	GENERATED_BODY()
-public:
-	ENGINE_API UMaterialExpressionTextureObject(const FObjectInitializer& ObjectInitializer = FObjectInitializer::Get());
+	GENERATED_UCLASS_BODY()
 
 #if WITH_EDITOR
 	virtual void PostEditChangeProperty(FPropertyChangedEvent& PropertyChangedEvent) override;
diff --git a/Engine/Source/Runtime/Engine/Classes/Materials/MaterialExpressionTextureObjectParameter.h b/Engine/Source/Runtime/Engine/Classes/Materials/MaterialExpressionTextureObjectParameter.h
--- a/Engine/Source/Runtime/Engine/Classes/Materials/MaterialExpressionTextureObjectParameter.h
+++ b/Engine/Source/Runtime/Engine/Classes/Materials/MaterialExpressionTextureObjectParameter.h
@@ -12,9 +12,7 @@
 UCLASS(collapsecategories, hidecategories=Object)
 class UMaterialExpressionTextureObjectParameter : public UMaterialExpressionTextureSampleParameter
 {
-	GENERATED_BODY()
-public:
-	UMaterialExpressionTextureObjectParameter(const FObjectInitializer& ObjectInitializer = FObjectInitializer::Get());
+	GENERATED_UCLASS_BODY()
 
 
 	// Begin UMaterialExpression Interface
diff --git a/Engine/Source/Runtime/Engine/Classes/Materials/MaterialExpressionTextureSample.h b/Engine/Source/Runtime/Engine/Classes/Materials/MaterialExpressionTextureSample.h
--- a/Engine/Source/Runtime/Engine/Classes/Materials/MaterialExpressionTextureSample.h
+++ b/Engine/Source/Runtime/Engine/Classes/Materials/MaterialExpressionTextureSample.h
@@ -19,9 +19,7 @@ enum ETextureMipValueMode
 UCLASS(collapsecategories, hidecategories=Object)
 class ENGINE_API UMaterialExpressionTextureSample : public UMaterialExpressionTextureBase
 {
-	GENERATED_BODY()
-public:
-	UMaterialExpressionTextureSample(const FObjectInitializer& ObjectInitializer = FObjectInitializer::Get());
+	GENERATED_UCLASS_BODY()
 
 	UPROPERTY(meta = (RequiredInput = "false", ToolTip = "Defaults to 'ConstCoordinate' if not specified"))
 	FExpressionInput Coordinates;
diff --git a/Engine/Source/Runtime/Engine/Classes/Materials/MaterialExpressionTextureSampleParameter.h b/Engine/Source/Runtime/Engine/Classes/Materials/MaterialExpressionTextureSampleParameter.h
--- a/Engine/Source/Runtime/Engine/Classes/Materials/MaterialExpressionTextureSampleParameter.h
+++ b/Engine/Source/Runtime/Engine/Classes/Materials/MaterialExpressionTextureSampleParameter.h
@@ -8,9 +8,7 @@
 UCLASS(collapsecategories, abstract, hidecategories=Object)
 class ENGINE_API UMaterialExpressionTextureSampleParameter : public UMaterialExpressionTextureSample
 {
-	GENERATED_BODY()
-public:
-	UMaterialExpressionTextureSampleParameter(const FObjectInitializer& ObjectInitializer = FObjectInitializer::Get());
+	GENERATED_UCLASS_BODY()
 
 	UPROPERTY(EditAnywhere, Category=MaterialExpressionTextureSampleParameter)
 	FName ParameterName;
diff --git a/Engine/Source/Runtime/Engine/Classes/Materials/MaterialExpressionTextureSampleParameterCube.h b/Engine/Source/Runtime/Engine/Classes/Materials/MaterialExpressionTextureSampleParameterCube.h
--- a/Engine/Source/Runtime/Engine/Classes/Materials/MaterialExpressionTextureSampleParameterCube.h
+++ b/Engine/Source/Runtime/Engine/Classes/Materials/MaterialExpressionTextureSampleParameterCube.h
@@ -8,9 +8,7 @@
 UCLASS(collapsecategories, hidecategories=Object, MinimalAPI)
 class UMaterialExpressionTextureSampleParameterCube : public UMaterialExpressionTextureSampleParameter
 {
-	GENERATED_BODY()
-public:
-	ENGINE_API UMaterialExpressionTextureSampleParameterCube(const FObjectInitializer& ObjectInitializer = FObjectInitializer::Get());
+	GENERATED_UCLASS_BODY()
 
 
 
diff --git a/Engine/Source/Runtime/Engine/Classes/Materials/MaterialExpressionTextureSampleParameterSubUV.h b/Engine/Source/Runtime/Engine/Classes/Materials/MaterialExpressionTextureSampleParameterSubUV.h
--- a/Engine/Source/Runtime/Engine/Classes/Materials/MaterialExpressionTextureSampleParameterSubUV.h
+++ b/Engine/Source/Runtime/Engine/Classes/Materials/MaterialExpressionTextureSampleParameterSubUV.h
@@ -8,9 +8,7 @@
 UCLASS(collapsecategories, hidecategories=Object, MinimalAPI)
 class UMaterialExpressionTextureSampleParameterSubUV : public UMaterialExpressionTextureSampleParameter2D
 {
-	GENERATED_BODY()
-public:
-	ENGINE_API UMaterialExpressionTextureSampleParameterSubUV(const FObjectInitializer& ObjectInitializer = FObjectInitializer::Get());
+	GENERATED_UCLASS_BODY()
 
 	UPROPERTY(EditAnywhere, Category=MaterialExpressionTextureSampleParameterSubUV)
 	uint32 bBlend:1;
diff --git a/Engine/Source/Runtime/Engine/Classes/Materials/MaterialExpressionTime.h b/Engine/Source/Runtime/Engine/Classes/Materials/MaterialExpressionTime.h
--- a/Engine/Source/Runtime/Engine/Classes/Materials/MaterialExpressionTime.h
+++ b/Engine/Source/Runtime/Engine/Classes/Materials/MaterialExpressionTime.h
@@ -9,9 +9,7 @@
 UCLASS(collapsecategories, hidecategories=Object)
 class UMaterialExpressionTime : public UMaterialExpression
 {
-	GENERATED_BODY()
-public:
-	UMaterialExpressionTime(const FObjectInitializer& ObjectInitializer = FObjectInitializer::Get());
+	GENERATED_UCLASS_BODY()
 
 	/** This time continues advancing regardless of whether the game is paused. */
 	UPROPERTY(EditAnywhere, Category=MaterialExpressionTime)
diff --git a/Engine/Source/Runtime/Engine/Classes/Materials/MaterialExpressionTransformPosition.h b/Engine/Source/Runtime/Engine/Classes/Materials/MaterialExpressionTransformPosition.h
--- a/Engine/Source/Runtime/Engine/Classes/Materials/MaterialExpressionTransformPosition.h
+++ b/Engine/Source/Runtime/Engine/Classes/Materials/MaterialExpressionTransformPosition.h
@@ -16,9 +16,7 @@ enum EMaterialPositionTransformSource
 UCLASS(collapsecategories, hidecategories=Object, MinimalAPI)
 class UMaterialExpressionTransformPosition : public UMaterialExpression
 {
-	GENERATED_BODY()
-public:
-	ENGINE_API UMaterialExpressionTransformPosition(const FObjectInitializer& ObjectInitializer = FObjectInitializer::Get());
+	GENERATED_UCLASS_BODY()
 
 	/** input expression for this transform */
 	UPROPERTY()
diff --git a/Engine/Source/Runtime/Engine/Classes/Materials/MaterialExpressionTwoSidedSign.h b/Engine/Source/Runtime/Engine/Classes/Materials/MaterialExpressionTwoSidedSign.h
--- a/Engine/Source/Runtime/Engine/Classes/Materials/MaterialExpressionTwoSidedSign.h
+++ b/Engine/Source/Runtime/Engine/Classes/Materials/MaterialExpressionTwoSidedSign.h
@@ -8,9 +8,7 @@
 UCLASS(collapsecategories, hidecategories=Object, MinimalAPI)
 class UMaterialExpressionTwoSidedSign : public UMaterialExpression
 {
-	GENERATED_BODY()
-public:
-	ENGINE_API UMaterialExpressionTwoSidedSign(const FObjectInitializer& ObjectInitializer = FObjectInitializer::Get());
+	GENERATED_UCLASS_BODY()
 
 
 	// Begin UMaterialExpression Interface
diff --git a/Engine/Source/Runtime/Engine/Classes/Materials/MaterialExpressionVectorParameter.h b/Engine/Source/Runtime/Engine/Classes/Materials/MaterialExpressionVectorParameter.h
--- a/Engine/Source/Runtime/Engine/Classes/Materials/MaterialExpressionVectorParameter.h
+++ b/Engine/Source/Runtime/Engine/Classes/Materials/MaterialExpressionVectorParameter.h
@@ -8,9 +8,7 @@
 UCLASS(collapsecategories, hidecategories=Object, MinimalAPI)
 class UMaterialExpressionVectorParameter : public UMaterialExpressionParameter
 {
-	GENERATED_BODY()
-public:
-	ENGINE_API UMaterialExpressionVectorParameter(const FObjectInitializer& ObjectInitializer = FObjectInitializer::Get());
+	GENERATED_UCLASS_BODY()
 
 	UPROPERTY(EditAnywhere, Category=MaterialExpressionVectorParameter)
 	FLinearColor DefaultValue;
diff --git a/Engine/Source/Runtime/Engine/Classes/Materials/MaterialExpressionVertexColor.h b/Engine/Source/Runtime/Engine/Classes/Materials/MaterialExpressionVertexColor.h
--- a/Engine/Source/Runtime/Engine/Classes/Materials/MaterialExpressionVertexColor.h
+++ b/Engine/Source/Runtime/Engine/Classes/Materials/MaterialExpressionVertexColor.h
@@ -8,9 +8,7 @@
 UCLASS(collapsecategories, hidecategories=Object, MinimalAPI)
 class UMaterialExpressionVertexColor : public UMaterialExpression
 {
-	GENERATED_BODY()
-public:
-	ENGINE_API UMaterialExpressionVertexColor(const FObjectInitializer& ObjectInitializer = FObjectInitializer::Get());
+	GENERATED_UCLASS_BODY()
 
 
 	// Begin UMaterialExpression Interface
diff --git a/Engine/Source/Runtime/Engine/Classes/Materials/MaterialExpressionVertexNormalWS.h b/Engine/Source/Runtime/Engine/Classes/Materials/MaterialExpressionVertexNormalWS.h
--- a/Engine/Source/Runtime/Engine/Classes/Materials/MaterialExpressionVertexNormalWS.h
+++ b/Engine/Source/Runtime/Engine/Classes/Materials/MaterialExpressionVertexNormalWS.h
@@ -8,9 +8,7 @@
 UCLASS(collapsecategories, hidecategories=Object)
 class UMaterialExpressionVertexNormalWS : public UMaterialExpression
 {
-	GENERATED_BODY()
-public:
-	UMaterialExpressionVertexNormalWS(const FObjectInitializer& ObjectInitializer = FObjectInitializer::Get());
+	GENERATED_UCLASS_BODY()
 
 
 	// Begin UMaterialExpression Interface
diff --git a/Engine/Source/Runtime/Engine/Classes/Materials/MaterialExpressionViewSize.h b/Engine/Source/Runtime/Engine/Classes/Materials/MaterialExpressionViewSize.h
--- a/Engine/Source/Runtime/Engine/Classes/Materials/MaterialExpressionViewSize.h
+++ b/Engine/Source/Runtime/Engine/Classes/Materials/MaterialExpressionViewSize.h
@@ -8,9 +8,7 @@
 UCLASS(collapsecategories, hidecategories=Object)
 class UMaterialExpressionViewSize : public UMaterialExpression
 {
-	GENERATED_BODY()
-public:
-	UMaterialExpressionViewSize(const FObjectInitializer& ObjectInitializer = FObjectInitializer::Get());
+	GENERATED_UCLASS_BODY()
 
 	// Begin UMaterialExpression Interface
 	virtual int32 Compile(class FMaterialCompiler* Compiler, int32 OutputIndex, int32 MultiplexIndex) override;
diff --git a/Engine/Source/Runtime/Engine/Classes/Materials/MaterialExpressionWorldPosition.h b/Engine/Source/Runtime/Engine/Classes/Materials/MaterialExpressionWorldPosition.h
--- a/Engine/Source/Runtime/Engine/Classes/Materials/MaterialExpressionWorldPosition.h
+++ b/Engine/Source/Runtime/Engine/Classes/Materials/MaterialExpressionWorldPosition.h
@@ -27,9 +27,7 @@ enum EWorldPositionIncludedOffsets
 UCLASS(collapsecategories, hidecategories=Object, MinimalAPI)
 class UMaterialExpressionWorldPosition : public UMaterialExpression
 {
-	GENERATED_BODY()
-public:
-	ENGINE_API UMaterialExpressionWorldPosition(const FObjectInitializer& ObjectInitializer = FObjectInitializer::Get());
+	GENERATED_UCLASS_BODY()
 
 	UPROPERTY(EditAnywhere, Category=UMaterialExpressionWorldPosition, meta=(DisplayName = "Shader Offsets"))
 	TEnumAsByte<EWorldPositionIncludedOffsets> WorldPositionShaderOffset;
diff --git a/Engine/Source/Runtime/Engine/Classes/Materials/MaterialFunction.h b/Engine/Source/Runtime/Engine/Classes/Materials/MaterialFunction.h
--- a/Engine/Source/Runtime/Engine/Classes/Materials/MaterialFunction.h
+++ b/Engine/Source/Runtime/Engine/Classes/Materials/MaterialFunction.h
@@ -9,9 +9,7 @@
 UCLASS(hidecategories=object, MinimalAPI)
 class UMaterialFunction : public UObject
 {
-	GENERATED_BODY()
-public:
-	ENGINE_API UMaterialFunction(const FObjectInitializer& ObjectInitializer = FObjectInitializer::Get());
+	GENERATED_UCLASS_BODY()
 
 	/** Used by materials using this function to know when to recompile. */
 	UPROPERTY(duplicatetransient)
diff --git a/Engine/Source/Runtime/Engine/Classes/Materials/MaterialInstance.h b/Engine/Source/Runtime/Engine/Classes/Materials/MaterialInstance.h
--- a/Engine/Source/Runtime/Engine/Classes/Materials/MaterialInstance.h
+++ b/Engine/Source/Runtime/Engine/Classes/Materials/MaterialInstance.h
@@ -18,7 +18,7 @@ class FSHAHash;
 USTRUCT()
 struct FFontParameterValue
 {
-	GENERATED_BODY()
+	GENERATED_USTRUCT_BODY()
 
 	UPROPERTY(EditAnywhere, BlueprintReadWrite, Category=FontParameterValue)
 	FName ParameterName;
@@ -48,7 +48,7 @@ struct FFontParameterValue
 USTRUCT()
 struct FScalarParameterValue
 {
-	GENERATED_BODY()
+	GENERATED_USTRUCT_BODY()
 
 	UPROPERTY(EditAnywhere, BlueprintReadWrite, Category=ScalarParameterValue)
 	FName ParameterName;
@@ -74,7 +74,7 @@ struct FScalarParameterValue
 USTRUCT()
 struct FTextureParameterValue
 {
-	GENERATED_BODY()
+	GENERATED_USTRUCT_BODY()
 
 	UPROPERTY(EditAnywhere, BlueprintReadWrite, Category=TextureParameterValue)
 	FName ParameterName;
@@ -100,7 +100,7 @@ struct FTextureParameterValue
 USTRUCT()
 struct FVectorParameterValue
 {
-	GENERATED_BODY()
+	GENERATED_USTRUCT_BODY()
 
 	UPROPERTY(EditAnywhere, BlueprintReadWrite, Category=VectorParameterValue)
 	FName ParameterName;
@@ -125,9 +125,7 @@ struct FVectorParameterValue
 UCLASS(abstract, BlueprintType,MinimalAPI)
 class UMaterialInstance : public UMaterialInterface
 {
-	GENERATED_BODY()
-public:
-	ENGINE_API UMaterialInstance(const FObjectInitializer& ObjectInitializer = FObjectInitializer::Get());
+	GENERATED_UCLASS_BODY()
 
 	/** Physical material to use for this graphics material. Used for sounds, effects etc.*/
 	UPROPERTY(EditAnywhere, BlueprintReadWrite, Category=MaterialInstance)
diff --git a/Engine/Source/Runtime/Engine/Classes/Materials/MaterialInstanceActor.h b/Engine/Source/Runtime/Engine/Classes/Materials/MaterialInstanceActor.h
--- a/Engine/Source/Runtime/Engine/Classes/Materials/MaterialInstanceActor.h
+++ b/Engine/Source/Runtime/Engine/Classes/Materials/MaterialInstanceActor.h
@@ -11,9 +11,7 @@
 UCLASS(hidecategories=Movement, hidecategories=Advanced, hidecategories=Collision, hidecategories=Display, hidecategories=Actor, hidecategories=Attachment, MinimalAPI)
 class AMaterialInstanceActor : public AActor
 {
-	GENERATED_BODY()
-public:
-	ENGINE_API AMaterialInstanceActor(const FObjectInitializer& ObjectInitializer = FObjectInitializer::Get());
+	GENERATED_UCLASS_BODY()
 
 	/** Pointer to actors that we want to control paramters of using Matinee. */
 	UPROPERTY(EditAnywhere, BlueprintReadWrite, Category=MaterialInstanceActor)
diff --git a/Engine/Source/Runtime/Engine/Classes/Materials/MaterialInstanceConstant.h b/Engine/Source/Runtime/Engine/Classes/Materials/MaterialInstanceConstant.h
--- a/Engine/Source/Runtime/Engine/Classes/Materials/MaterialInstanceConstant.h
+++ b/Engine/Source/Runtime/Engine/Classes/Materials/MaterialInstanceConstant.h
@@ -14,9 +14,7 @@
 UCLASS(hidecategories=Object, collapsecategories, BlueprintType,MinimalAPI)
 class UMaterialInstanceConstant : public UMaterialInstance
 {
-	GENERATED_BODY()
-public:
-	ENGINE_API UMaterialInstanceConstant(const FObjectInitializer& ObjectInitializer = FObjectInitializer::Get());
+	GENERATED_UCLASS_BODY()
 
 #if WITH_EDITOR
 	/** For constructing new MICs. */
diff --git a/Engine/Source/Runtime/Engine/Classes/Materials/MaterialInstanceDynamic.h b/Engine/Source/Runtime/Engine/Classes/Materials/MaterialInstanceDynamic.h
--- a/Engine/Source/Runtime/Engine/Classes/Materials/MaterialInstanceDynamic.h
+++ b/Engine/Source/Runtime/Engine/Classes/Materials/MaterialInstanceDynamic.h
@@ -8,9 +8,7 @@
 UCLASS(hidecategories=Object, collapsecategories, BlueprintType, MinimalAPI)
 class UMaterialInstanceDynamic : public UMaterialInstance
 {
-	GENERATED_BODY()
-public:
-	ENGINE_API UMaterialInstanceDynamic(const FObjectInitializer& ObjectInitializer = FObjectInitializer::Get());
+	GENERATED_UCLASS_BODY()
 
 	/** Set a MID scalar (float) parameter value */
 	UFUNCTION(BlueprintCallable, meta=(Keywords = "SetFloatParameterValue"), Category="Rendering|Material")
diff --git a/Engine/Source/Runtime/Engine/Classes/Materials/MaterialInterface.h b/Engine/Source/Runtime/Engine/Classes/Materials/MaterialInterface.h
--- a/Engine/Source/Runtime/Engine/Classes/Materials/MaterialInterface.h
+++ b/Engine/Source/Runtime/Engine/Classes/Materials/MaterialInterface.h
@@ -34,7 +34,7 @@ enum EMaterialUsage
 USTRUCT()
 struct ENGINE_API FMaterialRelevance
 {
-	GENERATED_BODY()
+	GENERATED_USTRUCT_BODY()
 
 	UPROPERTY()
 	uint32 bOpaque : 1;
@@ -104,7 +104,7 @@ struct ENGINE_API FMaterialRelevance
 USTRUCT()
 struct FLightmassMaterialInterfaceSettings
 {
-	GENERATED_BODY()
+	GENERATED_USTRUCT_BODY()
 
 	/** If true, forces translucency to cast static shadows as if the material were masked. */
 	UPROPERTY(EditAnywhere, BlueprintReadWrite, Category=Material)
@@ -171,9 +171,7 @@ struct FLightmassMaterialInterfaceSettings
 UCLASS(abstract, BlueprintType,MinimalAPI)
 class UMaterialInterface : public UObject, public IBlendableInterface
 {
-	GENERATED_BODY()
-public:
-	ENGINE_API UMaterialInterface(const FObjectInitializer& ObjectInitializer = FObjectInitializer::Get());
+	GENERATED_UCLASS_BODY()
 
 	/** SubsurfaceProfile, for Screen Space Subsurface Scattering */
 	UPROPERTY(EditAnywhere, BlueprintReadOnly, Category = Material, meta = (DisplayName = "Subsurface Profile"))
@@ -624,4 +622,4 @@ private:
 };
 
 /** Helper function to serialize inline shader maps for the given material resources. */
-extern void SerializeInlineShaderMaps(const TMap<const class ITargetPlatform*,TArray<FMaterialResource*>>* PlatformMaterialResourcesToSave, FArchive& Ar, FMaterialResource* OutMaterialResourcesLoaded[][ERHIFeatureLevel::Num]);
+extern void SerializeInlineShaderMaps(const TMap<const class ITargetPlatform*,TArray<FMaterialResource*>>* PlatformMaterialResourcesToSave, FArchive& Ar, FMaterialResource* OutMaterialResourcesLoaded[][ERHIFeatureLevel::Num]);
\ No newline at end of file
diff --git a/Engine/Source/Runtime/Engine/Classes/Materials/MaterialParameterCollection.h b/Engine/Source/Runtime/Engine/Classes/Materials/MaterialParameterCollection.h
--- a/Engine/Source/Runtime/Engine/Classes/Materials/MaterialParameterCollection.h
+++ b/Engine/Source/Runtime/Engine/Classes/Materials/MaterialParameterCollection.h
@@ -11,7 +11,7 @@
 USTRUCT()
 struct FCollectionParameterBase
 {
-	GENERATED_BODY()
+	GENERATED_USTRUCT_BODY()
 	
 	FCollectionParameterBase()
 	{
@@ -31,7 +31,7 @@ struct FCollectionParameterBase
 USTRUCT()
 struct FCollectionScalarParameter : public FCollectionParameterBase
 {
-	GENERATED_BODY()
+	GENERATED_USTRUCT_BODY()
 	
 	UPROPERTY(EditAnywhere, Category=Parameter)
 	float DefaultValue;
@@ -41,7 +41,7 @@ struct FCollectionScalarParameter : public FCollectionParameterBase
 USTRUCT()
 struct FCollectionVectorParameter : public FCollectionParameterBase
 {
-	GENERATED_BODY()
+	GENERATED_USTRUCT_BODY()
 	
 	UPROPERTY(EditAnywhere, Category=Parameter)
 	FLinearColor DefaultValue;
@@ -54,9 +54,7 @@ struct FCollectionVectorParameter : public FCollectionParameterBase
 UCLASS(hidecategories=object, MinimalAPI)
 class UMaterialParameterCollection : public UObject
 {
-	GENERATED_BODY()
-public:
-	ENGINE_API UMaterialParameterCollection(const FObjectInitializer& ObjectInitializer = FObjectInitializer::Get());
+	GENERATED_UCLASS_BODY()
 
 	/** Used by materials using this collection to know when to recompile. */
 	UPROPERTY(duplicatetransient)
diff --git a/Engine/Source/Runtime/Engine/Classes/Materials/MaterialParameterCollectionInstance.h b/Engine/Source/Runtime/Engine/Classes/Materials/MaterialParameterCollectionInstance.h
--- a/Engine/Source/Runtime/Engine/Classes/Materials/MaterialParameterCollectionInstance.h
+++ b/Engine/Source/Runtime/Engine/Classes/Materials/MaterialParameterCollectionInstance.h
@@ -14,9 +14,7 @@
 UCLASS(hidecategories=object, MinimalAPI)
 class UMaterialParameterCollectionInstance : public UObject
 {
-	GENERATED_BODY()
-public:
-	ENGINE_API UMaterialParameterCollectionInstance(const FObjectInitializer& ObjectInitializer = FObjectInitializer::Get());
+	GENERATED_UCLASS_BODY()
 
 	// Begin UObject interface.
 	ENGINE_API virtual void PostInitProperties();	
