Back out changelist 2481333

diff --git a/Engine/Source/Runtime/Engine/Classes/Matinee/InterpData.h b/Engine/Source/Runtime/Engine/Classes/Matinee/InterpData.h
--- a/Engine/Source/Runtime/Engine/Classes/Matinee/InterpData.h
+++ b/Engine/Source/Runtime/Engine/Classes/Matinee/InterpData.h
@@ -13,9 +13,7 @@
 UCLASS(hidecategories=Object, MinimalAPI)
 class UInterpData : public UObject
 {
-	GENERATED_BODY()
-public:
-	ENGINE_API UInterpData(const FObjectInitializer& ObjectInitializer = FObjectInitializer::Get());
+	GENERATED_UCLASS_BODY()
 
 	/** Duration of interpolation sequence - in seconds. */
 	UPROPERTY()
diff --git a/Engine/Source/Runtime/Engine/Classes/Matinee/InterpFilter.h b/Engine/Source/Runtime/Engine/Classes/Matinee/InterpFilter.h
--- a/Engine/Source/Runtime/Engine/Classes/Matinee/InterpFilter.h
+++ b/Engine/Source/Runtime/Engine/Classes/Matinee/InterpFilter.h
@@ -12,9 +12,7 @@
 UCLASS()
 class UInterpFilter : public UObject
 {
-	GENERATED_BODY()
-public:
-	UInterpFilter(const FObjectInitializer& ObjectInitializer = FObjectInitializer::Get());
+	GENERATED_UCLASS_BODY()
 
 	/** Caption for this filter. */
 	UPROPERTY()
diff --git a/Engine/Source/Runtime/Engine/Classes/Matinee/InterpFilter_Classes.h b/Engine/Source/Runtime/Engine/Classes/Matinee/InterpFilter_Classes.h
--- a/Engine/Source/Runtime/Engine/Classes/Matinee/InterpFilter_Classes.h
+++ b/Engine/Source/Runtime/Engine/Classes/Matinee/InterpFilter_Classes.h
@@ -13,9 +13,7 @@
 UCLASS()
 class UInterpFilter_Classes : public UInterpFilter
 {
-	GENERATED_BODY()
-public:
-	UInterpFilter_Classes(const FObjectInitializer& ObjectInitializer = FObjectInitializer::Get());
+	GENERATED_UCLASS_BODY()
 
 #if WITH_EDITORONLY_DATA
 	/** Which class to filter groups by. */
diff --git a/Engine/Source/Runtime/Engine/Classes/Matinee/InterpFilter_Custom.h b/Engine/Source/Runtime/Engine/Classes/Matinee/InterpFilter_Custom.h
--- a/Engine/Source/Runtime/Engine/Classes/Matinee/InterpFilter_Custom.h
+++ b/Engine/Source/Runtime/Engine/Classes/Matinee/InterpFilter_Custom.h
@@ -13,9 +13,7 @@
 UCLASS(MinimalAPI)
 class UInterpFilter_Custom : public UInterpFilter
 {
-	GENERATED_BODY()
-public:
-	ENGINE_API UInterpFilter_Custom(const FObjectInitializer& ObjectInitializer = FObjectInitializer::Get());
+	GENERATED_UCLASS_BODY()
 
 #if WITH_EDITORONLY_DATA
 	/** Which groups are included in this filter. */
diff --git a/Engine/Source/Runtime/Engine/Classes/Matinee/InterpGroup.h b/Engine/Source/Runtime/Engine/Classes/Matinee/InterpGroup.h
--- a/Engine/Source/Runtime/Engine/Classes/Matinee/InterpGroup.h
+++ b/Engine/Source/Runtime/Engine/Classes/Matinee/InterpGroup.h
@@ -13,7 +13,7 @@
 USTRUCT()
 struct FInterpEdSelKey
 {
-	GENERATED_BODY()
+	GENERATED_USTRUCT_BODY()
 
 	UPROPERTY()
 	class UInterpGroup* Group;
@@ -77,9 +77,7 @@ struct FInterpEdSelKey
 UCLASS(collapsecategories, hidecategories=Object, MinimalAPI)
 class UInterpGroup : public UObject, public FInterpEdInputInterface
 {
-	GENERATED_BODY()
-public:
-	ENGINE_API UInterpGroup(const FObjectInitializer& ObjectInitializer = FObjectInitializer::Get());
+	GENERATED_UCLASS_BODY()
 
 	UPROPERTY(export)
 	TArray<class UInterpTrack*> InterpTracks;
diff --git a/Engine/Source/Runtime/Engine/Classes/Matinee/InterpGroupCamera.h b/Engine/Source/Runtime/Engine/Classes/Matinee/InterpGroupCamera.h
--- a/Engine/Source/Runtime/Engine/Classes/Matinee/InterpGroupCamera.h
+++ b/Engine/Source/Runtime/Engine/Classes/Matinee/InterpGroupCamera.h
@@ -15,7 +15,7 @@
 USTRUCT()
 struct FCameraPreviewInfo
 {
-	GENERATED_BODY()
+	GENERATED_USTRUCT_BODY()
 
 	UPROPERTY(EditAnywhere, Category=CameraPreviewInfo)
 	TSubclassOf<class APawn>  PawnClass;
@@ -50,9 +50,7 @@ struct FCameraPreviewInfo
 UCLASS(collapsecategories, hidecategories=Object, MinimalAPI)
 class UInterpGroupCamera : public UInterpGroup
 {
-	GENERATED_BODY()
-public:
-	ENGINE_API UInterpGroupCamera(const FObjectInitializer& ObjectInitializer = FObjectInitializer::Get());
+	GENERATED_UCLASS_BODY()
 
 	UPROPERTY(transient)
 	class UCameraAnim* CameraAnimInst;
diff --git a/Engine/Source/Runtime/Engine/Classes/Matinee/InterpGroupDirector.h b/Engine/Source/Runtime/Engine/Classes/Matinee/InterpGroupDirector.h
--- a/Engine/Source/Runtime/Engine/Classes/Matinee/InterpGroupDirector.h
+++ b/Engine/Source/Runtime/Engine/Classes/Matinee/InterpGroupDirector.h
@@ -12,9 +12,7 @@
 UCLASS(collapsecategories, hidecategories=Object, MinimalAPI)
 class UInterpGroupDirector : public UInterpGroup
 {
-	GENERATED_BODY()
-public:
-	ENGINE_API UInterpGroupDirector(const FObjectInitializer& ObjectInitializer = FObjectInitializer::Get());
+	GENERATED_UCLASS_BODY()
 
 	// Begin UInterpGroup Interface
 	virtual AActor* SelectGroupActor( class UInterpGroupInst* GrInst, bool bDeselectActors ) override;
diff --git a/Engine/Source/Runtime/Engine/Classes/Matinee/InterpGroupInst.h b/Engine/Source/Runtime/Engine/Classes/Matinee/InterpGroupInst.h
--- a/Engine/Source/Runtime/Engine/Classes/Matinee/InterpGroupInst.h
+++ b/Engine/Source/Runtime/Engine/Classes/Matinee/InterpGroupInst.h
@@ -6,9 +6,7 @@
 UCLASS(MinimalAPI)
 class UInterpGroupInst : public UObject
 {
-	GENERATED_BODY()
-public:
-	ENGINE_API UInterpGroupInst(const FObjectInitializer& ObjectInitializer = FObjectInitializer::Get());
+	GENERATED_UCLASS_BODY()
 
 	/**
 	 *
diff --git a/Engine/Source/Runtime/Engine/Classes/Matinee/InterpGroupInstCamera.h b/Engine/Source/Runtime/Engine/Classes/Matinee/InterpGroupInstCamera.h
--- a/Engine/Source/Runtime/Engine/Classes/Matinee/InterpGroupInstCamera.h
+++ b/Engine/Source/Runtime/Engine/Classes/Matinee/InterpGroupInstCamera.h
@@ -7,9 +7,7 @@
 UCLASS()
 class UInterpGroupInstCamera : public UInterpGroupInst
 {
-	GENERATED_BODY()
-public:
-	UInterpGroupInstCamera(const FObjectInitializer& ObjectInitializer = FObjectInitializer::Get());
+	GENERATED_UCLASS_BODY()
 
 };
 
diff --git a/Engine/Source/Runtime/Engine/Classes/Matinee/InterpGroupInstDirector.h b/Engine/Source/Runtime/Engine/Classes/Matinee/InterpGroupInstDirector.h
--- a/Engine/Source/Runtime/Engine/Classes/Matinee/InterpGroupInstDirector.h
+++ b/Engine/Source/Runtime/Engine/Classes/Matinee/InterpGroupInstDirector.h
@@ -8,9 +8,7 @@
 UCLASS(MinimalAPI)
 class UInterpGroupInstDirector : public UInterpGroupInst
 {
-	GENERATED_BODY()
-public:
-	ENGINE_API UInterpGroupInstDirector(const FObjectInitializer& ObjectInitializer = FObjectInitializer::Get());
+	GENERATED_UCLASS_BODY()
 
 };
 
diff --git a/Engine/Source/Runtime/Engine/Classes/Matinee/InterpTrack.h b/Engine/Source/Runtime/Engine/Classes/Matinee/InterpTrack.h
--- a/Engine/Source/Runtime/Engine/Classes/Matinee/InterpTrack.h
+++ b/Engine/Source/Runtime/Engine/Classes/Matinee/InterpTrack.h
@@ -33,7 +33,7 @@ enum ETrackActiveCondition
 USTRUCT()
 struct FSupportedSubTrackInfo
 {
-	GENERATED_BODY()
+	GENERATED_USTRUCT_BODY()
 
 	/** The sub track class which is supported by this track */
 	UPROPERTY()
@@ -60,7 +60,7 @@ struct FSupportedSubTrackInfo
 USTRUCT()
 struct FSubTrackGroup
 {
-	GENERATED_BODY()
+	GENERATED_USTRUCT_BODY()
 
 	/** Name of the subtrack  group */
 	UPROPERTY()
@@ -90,9 +90,7 @@ struct FSubTrackGroup
 UCLASS(collapsecategories, hidecategories=Object, abstract, MinimalAPI)
 class UInterpTrack : public UObject, public FCurveEdInterface, public FInterpEdInputInterface
 {
-	GENERATED_BODY()
-public:
-	ENGINE_API UInterpTrack(const FObjectInitializer& ObjectInitializer = FObjectInitializer::Get());
+	GENERATED_UCLASS_BODY()
 
 	/** A list of subtracks that belong to this track */
 	UPROPERTY()
diff --git a/Engine/Source/Runtime/Engine/Classes/Matinee/InterpTrackAnimControl.h b/Engine/Source/Runtime/Engine/Classes/Matinee/InterpTrackAnimControl.h
--- a/Engine/Source/Runtime/Engine/Classes/Matinee/InterpTrackAnimControl.h
+++ b/Engine/Source/Runtime/Engine/Classes/Matinee/InterpTrackAnimControl.h
@@ -8,7 +8,7 @@
 USTRUCT()
 struct FAnimControlTrackKey
 {
-	GENERATED_BODY()
+	GENERATED_USTRUCT_BODY()
 
 	/** Position in the Matinee sequence to start playing this animation. */
 	UPROPERTY()
@@ -55,9 +55,7 @@ struct FAnimControlTrackKey
 UCLASS(MinimalAPI, meta=( DisplayName = "Animation Track" ))
 class UInterpTrackAnimControl : public UInterpTrackFloatBase
 {
-	GENERATED_BODY()
-public:
-	ENGINE_API UInterpTrackAnimControl(const FObjectInitializer& ObjectInitializer = FObjectInitializer::Get());
+	GENERATED_UCLASS_BODY()
 	
 	/** 
 	 *	Name of slot to use when playing animation. Passed to Actor. 
diff --git a/Engine/Source/Runtime/Engine/Classes/Matinee/InterpTrackAudioMaster.h b/Engine/Source/Runtime/Engine/Classes/Matinee/InterpTrackAudioMaster.h
--- a/Engine/Source/Runtime/Engine/Classes/Matinee/InterpTrackAudioMaster.h
+++ b/Engine/Source/Runtime/Engine/Classes/Matinee/InterpTrackAudioMaster.h
@@ -8,9 +8,7 @@
 UCLASS(MinimalAPI, meta=( DisplayName = "Audio Master Track" ) )
 class UInterpTrackAudioMaster : public UInterpTrackVectorBase
 {
-	GENERATED_BODY()
-public:
-	ENGINE_API UInterpTrackAudioMaster(const FObjectInitializer& ObjectInitializer = FObjectInitializer::Get());
+	GENERATED_UCLASS_BODY()
 
 
 	// Begin UInterpTrack interface.
diff --git a/Engine/Source/Runtime/Engine/Classes/Matinee/InterpTrackBoolProp.h b/Engine/Source/Runtime/Engine/Classes/Matinee/InterpTrackBoolProp.h
--- a/Engine/Source/Runtime/Engine/Classes/Matinee/InterpTrackBoolProp.h
+++ b/Engine/Source/Runtime/Engine/Classes/Matinee/InterpTrackBoolProp.h
@@ -8,7 +8,7 @@
 USTRUCT()
 struct FBoolTrackKey
 {
-	GENERATED_BODY()
+	GENERATED_USTRUCT_BODY()
 
 	UPROPERTY()
 	float Time;
@@ -28,9 +28,7 @@ struct FBoolTrackKey
 UCLASS(MinimalAPI, meta=( DisplayName = "Bool Property Track" ) )
 class UInterpTrackBoolProp : public UInterpTrack
 {
-	GENERATED_BODY()
-public:
-	ENGINE_API UInterpTrackBoolProp(const FObjectInitializer& ObjectInitializer = FObjectInitializer::Get());
+	GENERATED_UCLASS_BODY()
 
 	/** Array of booleans to set. */
 	UPROPERTY()
diff --git a/Engine/Source/Runtime/Engine/Classes/Matinee/InterpTrackColorProp.h b/Engine/Source/Runtime/Engine/Classes/Matinee/InterpTrackColorProp.h
--- a/Engine/Source/Runtime/Engine/Classes/Matinee/InterpTrackColorProp.h
+++ b/Engine/Source/Runtime/Engine/Classes/Matinee/InterpTrackColorProp.h
@@ -8,9 +8,7 @@
 UCLASS(MinimalAPI, meta=( DisplayName = "Color Property Track" ) )
 class UInterpTrackColorProp : public UInterpTrackVectorBase
 {
-	GENERATED_BODY()
-public:
-	ENGINE_API UInterpTrackColorProp(const FObjectInitializer& ObjectInitializer = FObjectInitializer::Get());
+	GENERATED_UCLASS_BODY()
 
 	/** Name of property in Group  AActor  which this track mill modify over time. */
 	UPROPERTY(Category=InterpTrackColorProp, VisibleAnywhere)
diff --git a/Engine/Source/Runtime/Engine/Classes/Matinee/InterpTrackColorScale.h b/Engine/Source/Runtime/Engine/Classes/Matinee/InterpTrackColorScale.h
--- a/Engine/Source/Runtime/Engine/Classes/Matinee/InterpTrackColorScale.h
+++ b/Engine/Source/Runtime/Engine/Classes/Matinee/InterpTrackColorScale.h
@@ -12,9 +12,7 @@
 UCLASS(meta=( DisplayName = "Color Scale Track" ) )
 class UInterpTrackColorScale : public UInterpTrackVectorBase
 {
-	GENERATED_BODY()
-public:
-	UInterpTrackColorScale(const FObjectInitializer& ObjectInitializer = FObjectInitializer::Get());
+	GENERATED_UCLASS_BODY()
 
 
 	// Begin UInterpTrack interface.
diff --git a/Engine/Source/Runtime/Engine/Classes/Matinee/InterpTrackDirector.h b/Engine/Source/Runtime/Engine/Classes/Matinee/InterpTrackDirector.h
--- a/Engine/Source/Runtime/Engine/Classes/Matinee/InterpTrackDirector.h
+++ b/Engine/Source/Runtime/Engine/Classes/Matinee/InterpTrackDirector.h
@@ -14,7 +14,7 @@
 USTRUCT()
 struct FDirectorTrackCut
 {
-	GENERATED_BODY()
+	GENERATED_USTRUCT_BODY()
 
 	/** Time to perform the cut. */
 	UPROPERTY()
@@ -45,9 +45,7 @@ struct FDirectorTrackCut
 UCLASS(MinimalAPI, meta=( DisplayName = "Director Track" ) )
 class UInterpTrackDirector : public UInterpTrack
 {
-	GENERATED_BODY()
-public:
-	ENGINE_API UInterpTrackDirector(const FObjectInitializer& ObjectInitializer = FObjectInitializer::Get());
+	GENERATED_UCLASS_BODY()
 
 	/** Array of cuts between cameras. */
 	UPROPERTY()
diff --git a/Engine/Source/Runtime/Engine/Classes/Matinee/InterpTrackEvent.h b/Engine/Source/Runtime/Engine/Classes/Matinee/InterpTrackEvent.h
--- a/Engine/Source/Runtime/Engine/Classes/Matinee/InterpTrackEvent.h
+++ b/Engine/Source/Runtime/Engine/Classes/Matinee/InterpTrackEvent.h
@@ -15,7 +15,7 @@
 USTRUCT()
 struct FEventTrackKey
 {
-	GENERATED_BODY()
+	GENERATED_USTRUCT_BODY()
 
 	UPROPERTY()
 	float Time;
@@ -34,9 +34,7 @@ struct FEventTrackKey
 UCLASS(MinimalAPI, meta=( DisplayName = "Event Track" ) )
 class UInterpTrackEvent : public UInterpTrack
 {
-	GENERATED_BODY()
-public:
-	ENGINE_API UInterpTrackEvent(const FObjectInitializer& ObjectInitializer = FObjectInitializer::Get());
+	GENERATED_UCLASS_BODY()
 
 	/** Array of events to fire off. */
 	UPROPERTY()
diff --git a/Engine/Source/Runtime/Engine/Classes/Matinee/InterpTrackFade.h b/Engine/Source/Runtime/Engine/Classes/Matinee/InterpTrackFade.h
--- a/Engine/Source/Runtime/Engine/Classes/Matinee/InterpTrackFade.h
+++ b/Engine/Source/Runtime/Engine/Classes/Matinee/InterpTrackFade.h
@@ -7,9 +7,7 @@
 UCLASS(meta=( DisplayName = "Fade Track" ) )
 class UInterpTrackFade : public UInterpTrackFloatBase
 {
-	GENERATED_BODY()
-public:
-	UInterpTrackFade(const FObjectInitializer& ObjectInitializer = FObjectInitializer::Get());
+	GENERATED_UCLASS_BODY()
 
 	/** 
 	 * InterpTrackFade
diff --git a/Engine/Source/Runtime/Engine/Classes/Matinee/InterpTrackFloatAnimBPParam.h b/Engine/Source/Runtime/Engine/Classes/Matinee/InterpTrackFloatAnimBPParam.h
--- a/Engine/Source/Runtime/Engine/Classes/Matinee/InterpTrackFloatAnimBPParam.h
+++ b/Engine/Source/Runtime/Engine/Classes/Matinee/InterpTrackFloatAnimBPParam.h
@@ -6,9 +6,7 @@
 UCLASS(meta=( DisplayName = "Float Anim BP Parameter Track" ) )
 class UInterpTrackFloatAnimBPParam : public UInterpTrackFloatBase
 {
-	GENERATED_BODY()
-public:
-	UInterpTrackFloatAnimBPParam(const FObjectInitializer& ObjectInitializer = FObjectInitializer::Get());
+	GENERATED_UCLASS_BODY()
 	
 	/** Materials whose parameters we want to change and the references to those materials. */
 	UPROPERTY(EditAnywhere, Category=InterpTrackFloatAnimBPParam)
diff --git a/Engine/Source/Runtime/Engine/Classes/Matinee/InterpTrackFloatBase.h b/Engine/Source/Runtime/Engine/Classes/Matinee/InterpTrackFloatBase.h
--- a/Engine/Source/Runtime/Engine/Classes/Matinee/InterpTrackFloatBase.h
+++ b/Engine/Source/Runtime/Engine/Classes/Matinee/InterpTrackFloatBase.h
@@ -7,9 +7,7 @@
 UCLASS(abstract, MinimalAPI)
 class UInterpTrackFloatBase : public UInterpTrack
 {
-	GENERATED_BODY()
-public:
-	ENGINE_API UInterpTrackFloatBase(const FObjectInitializer& ObjectInitializer = FObjectInitializer::Get());
+	GENERATED_UCLASS_BODY()
 	
 	/** Actually track data containing keyframes of float as it varies over time. */
 	UPROPERTY()
diff --git a/Engine/Source/Runtime/Engine/Classes/Matinee/InterpTrackFloatMaterialParam.h b/Engine/Source/Runtime/Engine/Classes/Matinee/InterpTrackFloatMaterialParam.h
--- a/Engine/Source/Runtime/Engine/Classes/Matinee/InterpTrackFloatMaterialParam.h
+++ b/Engine/Source/Runtime/Engine/Classes/Matinee/InterpTrackFloatMaterialParam.h
@@ -6,9 +6,7 @@
 UCLASS(meta=( DisplayName = "Float Material Parameter Track" ) )
 class UInterpTrackFloatMaterialParam : public UInterpTrackFloatBase
 {
-	GENERATED_BODY()
-public:
-	UInterpTrackFloatMaterialParam(const FObjectInitializer& ObjectInitializer = FObjectInitializer::Get());
+	GENERATED_UCLASS_BODY()
 	
 	/** Materials whose parameters we want to change and the references to those materials. */
 	UPROPERTY(EditAnywhere, Category=InterpTrackFloatMaterialParam)
diff --git a/Engine/Source/Runtime/Engine/Classes/Matinee/InterpTrackFloatParticleParam.h b/Engine/Source/Runtime/Engine/Classes/Matinee/InterpTrackFloatParticleParam.h
--- a/Engine/Source/Runtime/Engine/Classes/Matinee/InterpTrackFloatParticleParam.h
+++ b/Engine/Source/Runtime/Engine/Classes/Matinee/InterpTrackFloatParticleParam.h
@@ -6,9 +6,7 @@
 UCLASS(meta=( DisplayName = "Float Particle Param Track" ) )
 class UInterpTrackFloatParticleParam : public UInterpTrackFloatBase
 {
-	GENERATED_BODY()
-public:
-	UInterpTrackFloatParticleParam(const FObjectInitializer& ObjectInitializer = FObjectInitializer::Get());
+	GENERATED_UCLASS_BODY()
 	
 	/** Name of property in the Emitter which this track mill modify over time. */
 	UPROPERTY(EditAnywhere, Category=InterpTrackFloatParticleParam)
diff --git a/Engine/Source/Runtime/Engine/Classes/Matinee/InterpTrackFloatProp.h b/Engine/Source/Runtime/Engine/Classes/Matinee/InterpTrackFloatProp.h
--- a/Engine/Source/Runtime/Engine/Classes/Matinee/InterpTrackFloatProp.h
+++ b/Engine/Source/Runtime/Engine/Classes/Matinee/InterpTrackFloatProp.h
@@ -7,9 +7,7 @@
 UCLASS(MinimalAPI, meta=( DisplayName = "Float Property Track" ) )
 class UInterpTrackFloatProp : public UInterpTrackFloatBase
 {
-	GENERATED_BODY()
-public:
-	ENGINE_API UInterpTrackFloatProp(const FObjectInitializer& ObjectInitializer = FObjectInitializer::Get());
+	GENERATED_UCLASS_BODY()
 	
 	/** Name of property in Group  AActor  which this track mill modify over time. */
 	UPROPERTY(Category=InterpTrackFloatProp, VisibleAnywhere)
diff --git a/Engine/Source/Runtime/Engine/Classes/Matinee/InterpTrackInst.h b/Engine/Source/Runtime/Engine/Classes/Matinee/InterpTrackInst.h
--- a/Engine/Source/Runtime/Engine/Classes/Matinee/InterpTrackInst.h
+++ b/Engine/Source/Runtime/Engine/Classes/Matinee/InterpTrackInst.h
@@ -14,9 +14,7 @@ class UInterpTrack;
 UCLASS(MinimalAPI)
 class UInterpTrackInst : public UObject
 {
-	GENERATED_BODY()
-public:
-	ENGINE_API UInterpTrackInst(const FObjectInitializer& ObjectInitializer = FObjectInitializer::Get());
+	GENERATED_UCLASS_BODY()
 
 
 	/** 
diff --git a/Engine/Source/Runtime/Engine/Classes/Matinee/InterpTrackInstAnimControl.h b/Engine/Source/Runtime/Engine/Classes/Matinee/InterpTrackInstAnimControl.h
--- a/Engine/Source/Runtime/Engine/Classes/Matinee/InterpTrackInstAnimControl.h
+++ b/Engine/Source/Runtime/Engine/Classes/Matinee/InterpTrackInstAnimControl.h
@@ -6,9 +6,7 @@
 UCLASS(MinimalAPI)
 class UInterpTrackInstAnimControl : public UInterpTrackInst
 {
-	GENERATED_BODY()
-public:
-	ENGINE_API UInterpTrackInstAnimControl(const FObjectInitializer& ObjectInitializer = FObjectInitializer::Get());
+	GENERATED_UCLASS_BODY()
 
 	/**
 	 */
diff --git a/Engine/Source/Runtime/Engine/Classes/Matinee/InterpTrackInstAudioMaster.h b/Engine/Source/Runtime/Engine/Classes/Matinee/InterpTrackInstAudioMaster.h
--- a/Engine/Source/Runtime/Engine/Classes/Matinee/InterpTrackInstAudioMaster.h
+++ b/Engine/Source/Runtime/Engine/Classes/Matinee/InterpTrackInstAudioMaster.h
@@ -7,9 +7,7 @@
 UCLASS()
 class UInterpTrackInstAudioMaster : public UInterpTrackInst
 {
-	GENERATED_BODY()
-public:
-	UInterpTrackInstAudioMaster(const FObjectInitializer& ObjectInitializer = FObjectInitializer::Get());
+	GENERATED_UCLASS_BODY()
 
 
 	// Begin UInterpTrackInst Instance
diff --git a/Engine/Source/Runtime/Engine/Classes/Matinee/InterpTrackInstBoolProp.h b/Engine/Source/Runtime/Engine/Classes/Matinee/InterpTrackInstBoolProp.h
--- a/Engine/Source/Runtime/Engine/Classes/Matinee/InterpTrackInstBoolProp.h
+++ b/Engine/Source/Runtime/Engine/Classes/Matinee/InterpTrackInstBoolProp.h
@@ -7,9 +7,7 @@
 UCLASS()
 class UInterpTrackInstBoolProp : public UInterpTrackInstProperty
 {
-	GENERATED_BODY()
-public:
-	UInterpTrackInstBoolProp(const FObjectInitializer& ObjectInitializer = FObjectInitializer::Get());
+	GENERATED_UCLASS_BODY()
 
 	/** Pointer to boolean property in TrackObject. */
 	void* BoolPropertyAddress;
diff --git a/Engine/Source/Runtime/Engine/Classes/Matinee/InterpTrackInstColorProp.h b/Engine/Source/Runtime/Engine/Classes/Matinee/InterpTrackInstColorProp.h
--- a/Engine/Source/Runtime/Engine/Classes/Matinee/InterpTrackInstColorProp.h
+++ b/Engine/Source/Runtime/Engine/Classes/Matinee/InterpTrackInstColorProp.h
@@ -7,9 +7,7 @@
 UCLASS()
 class UInterpTrackInstColorProp : public UInterpTrackInstProperty
 {
-	GENERATED_BODY()
-public:
-	UInterpTrackInstColorProp(const FObjectInitializer& ObjectInitializer = FObjectInitializer::Get());
+	GENERATED_UCLASS_BODY()
 
 	/** Pointer to color property in TrackObject. */
 	FColor* ColorProp;
diff --git a/Engine/Source/Runtime/Engine/Classes/Matinee/InterpTrackInstColorScale.h b/Engine/Source/Runtime/Engine/Classes/Matinee/InterpTrackInstColorScale.h
--- a/Engine/Source/Runtime/Engine/Classes/Matinee/InterpTrackInstColorScale.h
+++ b/Engine/Source/Runtime/Engine/Classes/Matinee/InterpTrackInstColorScale.h
@@ -6,9 +6,7 @@
 UCLASS()
 class UInterpTrackInstColorScale : public UInterpTrackInst
 {
-	GENERATED_BODY()
-public:
-	UInterpTrackInstColorScale(const FObjectInitializer& ObjectInitializer = FObjectInitializer::Get());
+	GENERATED_UCLASS_BODY()
 
 
 	// Begin UInterpTrackInst Instance
diff --git a/Engine/Source/Runtime/Engine/Classes/Matinee/InterpTrackInstDirector.h b/Engine/Source/Runtime/Engine/Classes/Matinee/InterpTrackInstDirector.h
--- a/Engine/Source/Runtime/Engine/Classes/Matinee/InterpTrackInstDirector.h
+++ b/Engine/Source/Runtime/Engine/Classes/Matinee/InterpTrackInstDirector.h
@@ -8,9 +8,7 @@
 UCLASS(MinimalAPI)
 class UInterpTrackInstDirector : public UInterpTrackInst
 {
-	GENERATED_BODY()
-public:
-	ENGINE_API UInterpTrackInstDirector(const FObjectInitializer& ObjectInitializer = FObjectInitializer::Get());
+	GENERATED_UCLASS_BODY()
 
 	UPROPERTY()
 	class AActor* OldViewTarget;
diff --git a/Engine/Source/Runtime/Engine/Classes/Matinee/InterpTrackInstEvent.h b/Engine/Source/Runtime/Engine/Classes/Matinee/InterpTrackInstEvent.h
--- a/Engine/Source/Runtime/Engine/Classes/Matinee/InterpTrackInstEvent.h
+++ b/Engine/Source/Runtime/Engine/Classes/Matinee/InterpTrackInstEvent.h
@@ -7,9 +7,7 @@
 UCLASS(MinimalAPI)
 class UInterpTrackInstEvent : public UInterpTrackInst
 {
-	GENERATED_BODY()
-public:
-	ENGINE_API UInterpTrackInstEvent(const FObjectInitializer& ObjectInitializer = FObjectInitializer::Get());
+	GENERATED_UCLASS_BODY()
 
 	/** 
 	 *	Position we were in last time we evaluated Events. 
diff --git a/Engine/Source/Runtime/Engine/Classes/Matinee/InterpTrackInstFade.h b/Engine/Source/Runtime/Engine/Classes/Matinee/InterpTrackInstFade.h
--- a/Engine/Source/Runtime/Engine/Classes/Matinee/InterpTrackInstFade.h
+++ b/Engine/Source/Runtime/Engine/Classes/Matinee/InterpTrackInstFade.h
@@ -7,9 +7,7 @@
 UCLASS()
 class UInterpTrackInstFade : public UInterpTrackInst
 {
-	GENERATED_BODY()
-public:
-	UInterpTrackInstFade(const FObjectInitializer& ObjectInitializer = FObjectInitializer::Get());
+	GENERATED_UCLASS_BODY()
 
 
 	// Begin UInterpTrackInst Instance
diff --git a/Engine/Source/Runtime/Engine/Classes/Matinee/InterpTrackInstFloatAnimBPParam.h b/Engine/Source/Runtime/Engine/Classes/Matinee/InterpTrackInstFloatAnimBPParam.h
--- a/Engine/Source/Runtime/Engine/Classes/Matinee/InterpTrackInstFloatAnimBPParam.h
+++ b/Engine/Source/Runtime/Engine/Classes/Matinee/InterpTrackInstFloatAnimBPParam.h
@@ -7,9 +7,7 @@
 UCLASS()
 class UInterpTrackInstFloatAnimBPParam : public UInterpTrackInst
 {
-	GENERATED_BODY()
-public:
-	UInterpTrackInstFloatAnimBPParam(const FObjectInitializer& ObjectInitializer = FObjectInitializer::Get());
+	GENERATED_UCLASS_BODY()
 
 	/** MIDs we're using to set the desired parameter. */
 	UPROPERTY(transient)
diff --git a/Engine/Source/Runtime/Engine/Classes/Matinee/InterpTrackInstFloatMaterialParam.h b/Engine/Source/Runtime/Engine/Classes/Matinee/InterpTrackInstFloatMaterialParam.h
--- a/Engine/Source/Runtime/Engine/Classes/Matinee/InterpTrackInstFloatMaterialParam.h
+++ b/Engine/Source/Runtime/Engine/Classes/Matinee/InterpTrackInstFloatMaterialParam.h
@@ -7,9 +7,7 @@
 UCLASS()
 class UInterpTrackInstFloatMaterialParam : public UInterpTrackInst
 {
-	GENERATED_BODY()
-public:
-	UInterpTrackInstFloatMaterialParam(const FObjectInitializer& ObjectInitializer = FObjectInitializer::Get());
+	GENERATED_UCLASS_BODY()
 
 	/** MIDs we're using to set the desired parameter. */
 	UPROPERTY()
diff --git a/Engine/Source/Runtime/Engine/Classes/Matinee/InterpTrackInstFloatParticleParam.h b/Engine/Source/Runtime/Engine/Classes/Matinee/InterpTrackInstFloatParticleParam.h
--- a/Engine/Source/Runtime/Engine/Classes/Matinee/InterpTrackInstFloatParticleParam.h
+++ b/Engine/Source/Runtime/Engine/Classes/Matinee/InterpTrackInstFloatParticleParam.h
@@ -7,9 +7,7 @@
 UCLASS()
 class UInterpTrackInstFloatParticleParam : public UInterpTrackInst
 {
-	GENERATED_BODY()
-public:
-	UInterpTrackInstFloatParticleParam(const FObjectInitializer& ObjectInitializer = FObjectInitializer::Get());
+	GENERATED_UCLASS_BODY()
 
 	/** Saved value for restoring state when exiting Matinee. */
 	UPROPERTY()
diff --git a/Engine/Source/Runtime/Engine/Classes/Matinee/InterpTrackInstMove.h b/Engine/Source/Runtime/Engine/Classes/Matinee/InterpTrackInstMove.h
--- a/Engine/Source/Runtime/Engine/Classes/Matinee/InterpTrackInstMove.h
+++ b/Engine/Source/Runtime/Engine/Classes/Matinee/InterpTrackInstMove.h
@@ -8,9 +8,7 @@
 UCLASS(MinimalAPI)
 class UInterpTrackInstMove : public UInterpTrackInst
 {
-	GENERATED_BODY()
-public:
-	ENGINE_API UInterpTrackInstMove(const FObjectInitializer& ObjectInitializer = FObjectInitializer::Get());
+	GENERATED_UCLASS_BODY()
 
 	/** Saved position. Used in editor for resetting when quitting Matinee. */
 	UPROPERTY()
diff --git a/Engine/Source/Runtime/Engine/Classes/Matinee/InterpTrackInstParticleReplay.h b/Engine/Source/Runtime/Engine/Classes/Matinee/InterpTrackInstParticleReplay.h
--- a/Engine/Source/Runtime/Engine/Classes/Matinee/InterpTrackInstParticleReplay.h
+++ b/Engine/Source/Runtime/Engine/Classes/Matinee/InterpTrackInstParticleReplay.h
@@ -7,9 +7,7 @@
 UCLASS()
 class UInterpTrackInstParticleReplay : public UInterpTrackInst
 {
-	GENERATED_BODY()
-public:
-	UInterpTrackInstParticleReplay(const FObjectInitializer& ObjectInitializer = FObjectInitializer::Get());
+	GENERATED_UCLASS_BODY()
 
 	/** 
 	 *	Position we were in last time we evaluated.
diff --git a/Engine/Source/Runtime/Engine/Classes/Matinee/InterpTrackInstProperty.h b/Engine/Source/Runtime/Engine/Classes/Matinee/InterpTrackInstProperty.h
--- a/Engine/Source/Runtime/Engine/Classes/Matinee/InterpTrackInstProperty.h
+++ b/Engine/Source/Runtime/Engine/Classes/Matinee/InterpTrackInstProperty.h
@@ -8,9 +8,7 @@
 UCLASS()
 class UInterpTrackInstProperty : public UInterpTrackInst
 {
-	GENERATED_BODY()
-public:
-	UInterpTrackInstProperty(const FObjectInitializer& ObjectInitializer = FObjectInitializer::Get());
+	GENERATED_UCLASS_BODY()
 
 	/** Function to call after updating the value of the color property. */
 	UPROPERTY()
diff --git a/Engine/Source/Runtime/Engine/Classes/Matinee/InterpTrackInstSlomo.h b/Engine/Source/Runtime/Engine/Classes/Matinee/InterpTrackInstSlomo.h
--- a/Engine/Source/Runtime/Engine/Classes/Matinee/InterpTrackInstSlomo.h
+++ b/Engine/Source/Runtime/Engine/Classes/Matinee/InterpTrackInstSlomo.h
@@ -7,9 +7,7 @@
 UCLASS()
 class UInterpTrackInstSlomo : public UInterpTrackInst
 {
-	GENERATED_BODY()
-public:
-	UInterpTrackInstSlomo(const FObjectInitializer& ObjectInitializer = FObjectInitializer::Get());
+	GENERATED_UCLASS_BODY()
 
 	/** Backup of initial LevelInfo MatineeTimeDilation setting when interpolation started. */
 	UPROPERTY()
diff --git a/Engine/Source/Runtime/Engine/Classes/Matinee/InterpTrackInstSound.h b/Engine/Source/Runtime/Engine/Classes/Matinee/InterpTrackInstSound.h
--- a/Engine/Source/Runtime/Engine/Classes/Matinee/InterpTrackInstSound.h
+++ b/Engine/Source/Runtime/Engine/Classes/Matinee/InterpTrackInstSound.h
@@ -7,9 +7,7 @@
 UCLASS()
 class UInterpTrackInstSound : public UInterpTrackInst
 {
-	GENERATED_BODY()
-public:
-	UInterpTrackInstSound(const FObjectInitializer& ObjectInitializer = FObjectInitializer::Get());
+	GENERATED_UCLASS_BODY()
 
 	UPROPERTY()
 	float LastUpdatePosition;
diff --git a/Engine/Source/Runtime/Engine/Classes/Matinee/InterpTrackInstToggle.h b/Engine/Source/Runtime/Engine/Classes/Matinee/InterpTrackInstToggle.h
--- a/Engine/Source/Runtime/Engine/Classes/Matinee/InterpTrackInstToggle.h
+++ b/Engine/Source/Runtime/Engine/Classes/Matinee/InterpTrackInstToggle.h
@@ -10,9 +10,7 @@
 UCLASS()
 class UInterpTrackInstToggle : public UInterpTrackInst
 {
-	GENERATED_BODY()
-public:
-	UInterpTrackInstToggle(const FObjectInitializer& ObjectInitializer = FObjectInitializer::Get());
+	GENERATED_UCLASS_BODY()
 
 	UPROPERTY(EditAnywhere, Category=InterpTrackInstToggle)
 	TEnumAsByte<enum ETrackToggleAction> Action;
diff --git a/Engine/Source/Runtime/Engine/Classes/Matinee/InterpTrackInstVectorMaterialParam.h b/Engine/Source/Runtime/Engine/Classes/Matinee/InterpTrackInstVectorMaterialParam.h
--- a/Engine/Source/Runtime/Engine/Classes/Matinee/InterpTrackInstVectorMaterialParam.h
+++ b/Engine/Source/Runtime/Engine/Classes/Matinee/InterpTrackInstVectorMaterialParam.h
@@ -7,9 +7,7 @@
 UCLASS()
 class UInterpTrackInstVectorMaterialParam : public UInterpTrackInst
 {
-	GENERATED_BODY()
-public:
-	UInterpTrackInstVectorMaterialParam(const FObjectInitializer& ObjectInitializer = FObjectInitializer::Get());
+	GENERATED_UCLASS_BODY()
 
 	/** MIDs we're using to set the desired parameter. */
 	UPROPERTY()
diff --git a/Engine/Source/Runtime/Engine/Classes/Matinee/InterpTrackInstVectorProp.h b/Engine/Source/Runtime/Engine/Classes/Matinee/InterpTrackInstVectorProp.h
--- a/Engine/Source/Runtime/Engine/Classes/Matinee/InterpTrackInstVectorProp.h
+++ b/Engine/Source/Runtime/Engine/Classes/Matinee/InterpTrackInstVectorProp.h
@@ -7,9 +7,7 @@
 UCLASS()
 class UInterpTrackInstVectorProp : public UInterpTrackInstProperty
 {
-	GENERATED_BODY()
-public:
-	UInterpTrackInstVectorProp(const FObjectInitializer& ObjectInitializer = FObjectInitializer::Get());
+	GENERATED_UCLASS_BODY()
 
 	/** Pointer to FVector property in TrackObject. */
 	FVector* VectorProp;
diff --git a/Engine/Source/Runtime/Engine/Classes/Matinee/InterpTrackLinearColorBase.h b/Engine/Source/Runtime/Engine/Classes/Matinee/InterpTrackLinearColorBase.h
--- a/Engine/Source/Runtime/Engine/Classes/Matinee/InterpTrackLinearColorBase.h
+++ b/Engine/Source/Runtime/Engine/Classes/Matinee/InterpTrackLinearColorBase.h
@@ -7,9 +7,7 @@
 UCLASS(abstract, MinimalAPI)
 class UInterpTrackLinearColorBase : public UInterpTrack
 {
-	GENERATED_BODY()
-public:
-	ENGINE_API UInterpTrackLinearColorBase(const FObjectInitializer& ObjectInitializer = FObjectInitializer::Get());
+	GENERATED_UCLASS_BODY()
 
 	/** Actually track data containing keyframes of a FVector as it varies over time. */
 	UPROPERTY()
diff --git a/Engine/Source/Runtime/Engine/Classes/Matinee/InterpTrackLinearColorProp.h b/Engine/Source/Runtime/Engine/Classes/Matinee/InterpTrackLinearColorProp.h
--- a/Engine/Source/Runtime/Engine/Classes/Matinee/InterpTrackLinearColorProp.h
+++ b/Engine/Source/Runtime/Engine/Classes/Matinee/InterpTrackLinearColorProp.h
@@ -8,9 +8,7 @@
 UCLASS(MinimalAPI, meta=( DisplayName = "LinearColor Property Track" ) )
 class UInterpTrackLinearColorProp : public UInterpTrackLinearColorBase
 {
-	GENERATED_BODY()
-public:
-	ENGINE_API UInterpTrackLinearColorProp(const FObjectInitializer& ObjectInitializer = FObjectInitializer::Get());
+	GENERATED_UCLASS_BODY()
 
 	/** Name of property in Group  AActor  which this track mill modify over time. */
 	UPROPERTY(Category=InterpTrackLinearColorProp, VisibleAnywhere)
diff --git a/Engine/Source/Runtime/Engine/Classes/Matinee/InterpTrackMove.h b/Engine/Source/Runtime/Engine/Classes/Matinee/InterpTrackMove.h
--- a/Engine/Source/Runtime/Engine/Classes/Matinee/InterpTrackMove.h
+++ b/Engine/Source/Runtime/Engine/Classes/Matinee/InterpTrackMove.h
@@ -29,7 +29,7 @@ enum EInterpTrackMoveRotMode
 USTRUCT()
 struct FInterpLookupPoint
 {
-	GENERATED_BODY()
+	GENERATED_USTRUCT_BODY()
 
 	UPROPERTY()
 	FName GroupName;
@@ -48,7 +48,7 @@ struct FInterpLookupPoint
 USTRUCT()
 struct FInterpLookupTrack
 {
-	GENERATED_BODY()
+	GENERATED_USTRUCT_BODY()
 
 	UPROPERTY()
 	TArray<struct FInterpLookupPoint> Points;
@@ -92,9 +92,7 @@ struct FInterpLookupTrack
 UCLASS(MinimalAPI, meta=( DisplayName = "Movement Track" ) )
 class UInterpTrackMove : public UInterpTrack
 {
-	GENERATED_BODY()
-public:
-	ENGINE_API UInterpTrackMove(const FObjectInitializer& ObjectInitializer = FObjectInitializer::Get());
+	GENERATED_UCLASS_BODY()
 	
 	/** Actual position keyframe data. */
 	UPROPERTY()
diff --git a/Engine/Source/Runtime/Engine/Classes/Matinee/InterpTrackMoveAxis.h b/Engine/Source/Runtime/Engine/Classes/Matinee/InterpTrackMoveAxis.h
--- a/Engine/Source/Runtime/Engine/Classes/Matinee/InterpTrackMoveAxis.h
+++ b/Engine/Source/Runtime/Engine/Classes/Matinee/InterpTrackMoveAxis.h
@@ -26,9 +26,7 @@ enum EInterpMoveAxis
 UCLASS(MinimalAPI, meta=( DisplayName = "Move Axis Track" ) )
 class UInterpTrackMoveAxis : public UInterpTrackFloatBase
 {
-	GENERATED_BODY()
-public:
-	ENGINE_API UInterpTrackMoveAxis(const FObjectInitializer& ObjectInitializer = FObjectInitializer::Get());
+	GENERATED_UCLASS_BODY()
 
 	/** The axis which this track will use when transforming an actor */
 	UPROPERTY()
diff --git a/Engine/Source/Runtime/Engine/Classes/Matinee/InterpTrackParticleReplay.h b/Engine/Source/Runtime/Engine/Classes/Matinee/InterpTrackParticleReplay.h
--- a/Engine/Source/Runtime/Engine/Classes/Matinee/InterpTrackParticleReplay.h
+++ b/Engine/Source/Runtime/Engine/Classes/Matinee/InterpTrackParticleReplay.h
@@ -12,7 +12,7 @@
 USTRUCT()
 struct FParticleReplayTrackKey
 {
-	GENERATED_BODY()
+	GENERATED_USTRUCT_BODY()
 
 	/** Position along timeline */
 	UPROPERTY()
@@ -39,9 +39,7 @@ struct FParticleReplayTrackKey
 UCLASS(MinimalAPI, meta=( DisplayName = "Particle Replay Track" ) )
 class UInterpTrackParticleReplay : public UInterpTrack
 {
-	GENERATED_BODY()
-public:
-	ENGINE_API UInterpTrackParticleReplay(const FObjectInitializer& ObjectInitializer = FObjectInitializer::Get());
+	GENERATED_UCLASS_BODY()
 
 	/** Array of keys */
 	UPROPERTY()
diff --git a/Engine/Source/Runtime/Engine/Classes/Matinee/InterpTrackSlomo.h b/Engine/Source/Runtime/Engine/Classes/Matinee/InterpTrackSlomo.h
--- a/Engine/Source/Runtime/Engine/Classes/Matinee/InterpTrackSlomo.h
+++ b/Engine/Source/Runtime/Engine/Classes/Matinee/InterpTrackSlomo.h
@@ -7,9 +7,7 @@
 UCLASS(meta=( DisplayName = "Slomo Track" ) )
 class UInterpTrackSlomo : public UInterpTrackFloatBase
 {
-	GENERATED_BODY()
-public:
-	UInterpTrackSlomo(const FObjectInitializer& ObjectInitializer = FObjectInitializer::Get());
+	GENERATED_UCLASS_BODY()
 
 
 	// Begin InterpTrack interface.
diff --git a/Engine/Source/Runtime/Engine/Classes/Matinee/InterpTrackSound.h b/Engine/Source/Runtime/Engine/Classes/Matinee/InterpTrackSound.h
--- a/Engine/Source/Runtime/Engine/Classes/Matinee/InterpTrackSound.h
+++ b/Engine/Source/Runtime/Engine/Classes/Matinee/InterpTrackSound.h
@@ -13,7 +13,7 @@
 USTRUCT()
 struct FSoundTrackKey
 {
-	GENERATED_BODY()
+	GENERATED_USTRUCT_BODY()
 
 	UPROPERTY()
 	float Time;
@@ -42,9 +42,7 @@ struct FSoundTrackKey
 UCLASS(MinimalAPI, meta=( DisplayName = "Sound Track" ) )
 class UInterpTrackSound : public UInterpTrackVectorBase
 {
-	GENERATED_BODY()
-public:
-	ENGINE_API UInterpTrackSound(const FObjectInitializer& ObjectInitializer = FObjectInitializer::Get());
+	GENERATED_UCLASS_BODY()
 
 	/** Array of sounds to play at specific times. */
 	UPROPERTY()
diff --git a/Engine/Source/Runtime/Engine/Classes/Matinee/InterpTrackToggle.h b/Engine/Source/Runtime/Engine/Classes/Matinee/InterpTrackToggle.h
--- a/Engine/Source/Runtime/Engine/Classes/Matinee/InterpTrackToggle.h
+++ b/Engine/Source/Runtime/Engine/Classes/Matinee/InterpTrackToggle.h
@@ -26,7 +26,7 @@ enum ETrackToggleAction
 USTRUCT()
 struct FToggleTrackKey
 {
-	GENERATED_BODY()
+	GENERATED_USTRUCT_BODY()
 
 	UPROPERTY()
 	float Time;
@@ -46,9 +46,7 @@ struct FToggleTrackKey
 UCLASS(MinimalAPI, meta=( DisplayName = "Toggle Track" ) )
 class UInterpTrackToggle : public UInterpTrack
 {
-	GENERATED_BODY()
-public:
-	ENGINE_API UInterpTrackToggle(const FObjectInitializer& ObjectInitializer = FObjectInitializer::Get());
+	GENERATED_UCLASS_BODY()
 
 	/** Array of events to fire off. */
 	UPROPERTY()
diff --git a/Engine/Source/Runtime/Engine/Classes/Matinee/InterpTrackVectorBase.h b/Engine/Source/Runtime/Engine/Classes/Matinee/InterpTrackVectorBase.h
--- a/Engine/Source/Runtime/Engine/Classes/Matinee/InterpTrackVectorBase.h
+++ b/Engine/Source/Runtime/Engine/Classes/Matinee/InterpTrackVectorBase.h
@@ -7,9 +7,7 @@
 UCLASS(abstract, MinimalAPI)
 class UInterpTrackVectorBase : public UInterpTrack
 {
-	GENERATED_BODY()
-public:
-	ENGINE_API UInterpTrackVectorBase(const FObjectInitializer& ObjectInitializer = FObjectInitializer::Get());
+	GENERATED_UCLASS_BODY()
 
 	/** Actually track data containing keyframes of a FVector as it varies over time. */
 	UPROPERTY()
diff --git a/Engine/Source/Runtime/Engine/Classes/Matinee/InterpTrackVectorMaterialParam.h b/Engine/Source/Runtime/Engine/Classes/Matinee/InterpTrackVectorMaterialParam.h
--- a/Engine/Source/Runtime/Engine/Classes/Matinee/InterpTrackVectorMaterialParam.h
+++ b/Engine/Source/Runtime/Engine/Classes/Matinee/InterpTrackVectorMaterialParam.h
@@ -6,9 +6,7 @@
 UCLASS(MinimalAPI, meta=( DisplayName = "Vector Material Parameter Track" ))
 class UInterpTrackVectorMaterialParam : public UInterpTrackVectorBase
 {
-	GENERATED_BODY()
-public:
-	ENGINE_API UInterpTrackVectorMaterialParam(const FObjectInitializer& ObjectInitializer = FObjectInitializer::Get());
+	GENERATED_UCLASS_BODY()
 
 	/** Materials whose parameters we want to change and the references to those materials. */
 	UPROPERTY(EditAnywhere, Category=InterpTrackVectorMaterialParam)
diff --git a/Engine/Source/Runtime/Engine/Classes/Matinee/InterpTrackVectorProp.h b/Engine/Source/Runtime/Engine/Classes/Matinee/InterpTrackVectorProp.h
--- a/Engine/Source/Runtime/Engine/Classes/Matinee/InterpTrackVectorProp.h
+++ b/Engine/Source/Runtime/Engine/Classes/Matinee/InterpTrackVectorProp.h
@@ -8,9 +8,7 @@
 UCLASS(MinimalAPI, meta=( DisplayName = "Vector Property Track" ) )
 class UInterpTrackVectorProp : public UInterpTrackVectorBase
 {
-	GENERATED_BODY()
-public:
-	ENGINE_API UInterpTrackVectorProp(const FObjectInitializer& ObjectInitializer = FObjectInitializer::Get());
+	GENERATED_UCLASS_BODY()
 
 	/** Name of property in Group  AActor  which this track mill modify over time. */
 	UPROPERTY(Category=InterpTrackVectorProp, VisibleAnywhere)
diff --git a/Engine/Source/Runtime/Engine/Classes/Matinee/InterpTrackVisibility.h b/Engine/Source/Runtime/Engine/Classes/Matinee/InterpTrackVisibility.h
--- a/Engine/Source/Runtime/Engine/Classes/Matinee/InterpTrackVisibility.h
+++ b/Engine/Source/Runtime/Engine/Classes/Matinee/InterpTrackVisibility.h
@@ -39,7 +39,7 @@ enum EVisibilityTrackCondition
 USTRUCT()
 struct FVisibilityTrackKey
 {
-	GENERATED_BODY()
+	GENERATED_USTRUCT_BODY()
 
 	UPROPERTY()
 	float Time;
@@ -64,9 +64,7 @@ struct FVisibilityTrackKey
 UCLASS(MinimalAPI, meta=( DisplayName = "Visibility Track" ) )
 class UInterpTrackVisibility : public UInterpTrack
 {
-	GENERATED_BODY()
-public:
-	ENGINE_API UInterpTrackVisibility(const FObjectInitializer& ObjectInitializer = FObjectInitializer::Get());
+	GENERATED_UCLASS_BODY()
 
 	/** Array of events to fire off. */
 	UPROPERTY()
diff --git a/Engine/Source/Runtime/Engine/Classes/Matinee/MatineeActor.h b/Engine/Source/Runtime/Engine/Classes/Matinee/MatineeActor.h
--- a/Engine/Source/Runtime/Engine/Classes/Matinee/MatineeActor.h
+++ b/Engine/Source/Runtime/Engine/Classes/Matinee/MatineeActor.h
@@ -16,7 +16,7 @@ DECLARE_DYNAMIC_DELEGATE( FOnMatineeEvent );
 USTRUCT()
 struct FCameraCutInfo
 {
-	GENERATED_BODY()
+	GENERATED_USTRUCT_BODY()
 
 	UPROPERTY()
 	FVector Location;
@@ -39,7 +39,7 @@ struct FCameraCutInfo
 USTRUCT()
 struct FInterpGroupActorInfo
 {
-	GENERATED_BODY()
+	GENERATED_USTRUCT_BODY()
 
 	UPROPERTY(EditAnywhere, BlueprintReadWrite, Category=InterpGroupActorInfo)
 	FName ObjectName;
@@ -52,9 +52,7 @@ struct FInterpGroupActorInfo
 UCLASS(MinimalAPI, NotBlueprintable, hidecategories=(Collision, Game, Input), showcategories=("Input|MouseInput", "Input|TouchInput", "Game|Damage"))
 class AMatineeActor : public AActor
 {
-	GENERATED_BODY()
-public:
-	ENGINE_API AMatineeActor(const FObjectInitializer& ObjectInitializer = FObjectInitializer::Get());
+	GENERATED_UCLASS_BODY()
 
 	/** The matinee data used by this actor*/
 	UPROPERTY(EditAnywhere, BlueprintReadWrite, Category=MatineeActor, replicated, meta=(ForceRebuildProperty = "GroupActorInfos"))
diff --git a/Engine/Source/Runtime/Engine/Classes/Matinee/MatineeActorCameraAnim.h b/Engine/Source/Runtime/Engine/Classes/Matinee/MatineeActorCameraAnim.h
--- a/Engine/Source/Runtime/Engine/Classes/Matinee/MatineeActorCameraAnim.h
+++ b/Engine/Source/Runtime/Engine/Classes/Matinee/MatineeActorCameraAnim.h
@@ -12,9 +12,7 @@ class UCameraAnim;
 UCLASS(notplaceable, MinimalAPI, NotBlueprintable)
 class AMatineeActorCameraAnim : public AMatineeActor
 {
-	GENERATED_BODY()
-public:
-	ENGINE_API AMatineeActorCameraAnim(const FObjectInitializer& ObjectInitializer = FObjectInitializer::Get());
+	GENERATED_UCLASS_BODY()
 
 	/** The camera anim we are editing */
 	UPROPERTY(Transient)
@@ -31,4 +29,4 @@ public:
 		return false;
 	}
 	// End UObject interface
-};
+};
\ No newline at end of file
diff --git a/Engine/Source/Runtime/Engine/Classes/Matinee/MatineeAnimInterface.h b/Engine/Source/Runtime/Engine/Classes/Matinee/MatineeAnimInterface.h
--- a/Engine/Source/Runtime/Engine/Classes/Matinee/MatineeAnimInterface.h
+++ b/Engine/Source/Runtime/Engine/Classes/Matinee/MatineeAnimInterface.h
@@ -13,15 +13,12 @@
 UINTERFACE(MinimalAPI, meta=(CannotImplementInterfaceInBlueprint))
 class UMatineeAnimInterface : public UInterface
 {
-	GENERATED_BODY()
-public:
-	ENGINE_API UMatineeAnimInterface(const FObjectInitializer& ObjectInitializer = FObjectInitializer::Get());
+	GENERATED_UINTERFACE_BODY()
 };
 
 class IMatineeAnimInterface
 {
-	GENERATED_BODY()
-public:
+	GENERATED_IINTERFACE_BODY()
 
 
 	
diff --git a/Engine/Source/Runtime/Engine/Classes/Matinee/MatineeInterface.h b/Engine/Source/Runtime/Engine/Classes/Matinee/MatineeInterface.h
--- a/Engine/Source/Runtime/Engine/Classes/Matinee/MatineeInterface.h
+++ b/Engine/Source/Runtime/Engine/Classes/Matinee/MatineeInterface.h
@@ -12,15 +12,12 @@
 UINTERFACE(meta=(CannotImplementInterfaceInBlueprint))
 class UMatineeInterface : public UInterface
 {
-	GENERATED_BODY()
-public:
-	UMatineeInterface(const FObjectInitializer& ObjectInitializer = FObjectInitializer::Get());
+	GENERATED_UINTERFACE_BODY()
 };
 
 class IMatineeInterface
 {
-	GENERATED_BODY()
-public:
+	GENERATED_IINTERFACE_BODY()
 
 
 	/** Called when a matinee actor starts interpolating this  AActor via matinee.
diff --git a/Engine/Source/Runtime/Engine/Classes/MovieScene/RuntimeMovieScenePlayerInterface.h b/Engine/Source/Runtime/Engine/Classes/MovieScene/RuntimeMovieScenePlayerInterface.h
--- a/Engine/Source/Runtime/Engine/Classes/MovieScene/RuntimeMovieScenePlayerInterface.h
+++ b/Engine/Source/Runtime/Engine/Classes/MovieScene/RuntimeMovieScenePlayerInterface.h
@@ -7,15 +7,12 @@
 UINTERFACE(MinimalAPI, meta=(CannotImplementInterfaceInBlueprint))
 class URuntimeMovieScenePlayerInterface : public UInterface
 {
-	GENERATED_BODY()
-public:
-	ENGINE_API URuntimeMovieScenePlayerInterface(const FObjectInitializer& ObjectInitializer = FObjectInitializer::Get());
+	GENERATED_UINTERFACE_BODY()
 };
 
 class IRuntimeMovieScenePlayerInterface
 {
-	GENERATED_BODY()
-public:
+	GENERATED_IINTERFACE_BODY()
 
 	/**
 	 * Ticks this MovieScenePlayer.  Should be called every frame.
diff --git a/Engine/Source/Runtime/Engine/Classes/Particles/Acceleration/ParticleModuleAcceleration.h b/Engine/Source/Runtime/Engine/Classes/Particles/Acceleration/ParticleModuleAcceleration.h
--- a/Engine/Source/Runtime/Engine/Classes/Particles/Acceleration/ParticleModuleAcceleration.h
+++ b/Engine/Source/Runtime/Engine/Classes/Particles/Acceleration/ParticleModuleAcceleration.h
@@ -8,9 +8,7 @@
 UCLASS(editinlinenew, hidecategories=Object, meta=(DisplayName = "Acceleration"), MinimalAPI)
 class UParticleModuleAcceleration : public UParticleModuleAccelerationBase
 {
-	GENERATED_BODY()
-public:
-	ENGINE_API UParticleModuleAcceleration(const FObjectInitializer& ObjectInitializer = FObjectInitializer::Get());
+	GENERATED_UCLASS_BODY()
 
 	/**
 	 *	The initial acceleration of the particle.
diff --git a/Engine/Source/Runtime/Engine/Classes/Particles/Acceleration/ParticleModuleAccelerationBase.h b/Engine/Source/Runtime/Engine/Classes/Particles/Acceleration/ParticleModuleAccelerationBase.h
--- a/Engine/Source/Runtime/Engine/Classes/Particles/Acceleration/ParticleModuleAccelerationBase.h
+++ b/Engine/Source/Runtime/Engine/Classes/Particles/Acceleration/ParticleModuleAccelerationBase.h
@@ -8,9 +8,7 @@
 UCLASS(editinlinenew, hidecategories=Object, meta=(DisplayName = "Acceleration"), abstract)
 class UParticleModuleAccelerationBase : public UParticleModule
 {
-	GENERATED_BODY()
-public:
-	UParticleModuleAccelerationBase(const FObjectInitializer& ObjectInitializer = FObjectInitializer::Get());
+	GENERATED_UCLASS_BODY()
 
 	/**
 	 *	If true, then treat the acceleration as world-space
diff --git a/Engine/Source/Runtime/Engine/Classes/Particles/Acceleration/ParticleModuleAccelerationConstant.h b/Engine/Source/Runtime/Engine/Classes/Particles/Acceleration/ParticleModuleAccelerationConstant.h
--- a/Engine/Source/Runtime/Engine/Classes/Particles/Acceleration/ParticleModuleAccelerationConstant.h
+++ b/Engine/Source/Runtime/Engine/Classes/Particles/Acceleration/ParticleModuleAccelerationConstant.h
@@ -11,9 +11,7 @@
 UCLASS(editinlinenew, hidecategories=(Object, Acceleration), meta=(DisplayName = "Const Acceleration"))
 class UParticleModuleAccelerationConstant : public UParticleModuleAccelerationBase
 {
-	GENERATED_BODY()
-public:
-	UParticleModuleAccelerationConstant(const FObjectInitializer& ObjectInitializer = FObjectInitializer::Get());
+	GENERATED_UCLASS_BODY()
 
 	/** Constant acceleration for particles in this system. */
 	UPROPERTY(EditAnywhere, Category=ParticleModuleAccelerationConstant)
diff --git a/Engine/Source/Runtime/Engine/Classes/Particles/Acceleration/ParticleModuleAccelerationDrag.h b/Engine/Source/Runtime/Engine/Classes/Particles/Acceleration/ParticleModuleAccelerationDrag.h
--- a/Engine/Source/Runtime/Engine/Classes/Particles/Acceleration/ParticleModuleAccelerationDrag.h
+++ b/Engine/Source/Runtime/Engine/Classes/Particles/Acceleration/ParticleModuleAccelerationDrag.h
@@ -11,9 +11,7 @@
 UCLASS(editinlinenew, hidecategories=(Object, Acceleration), MinimalAPI, meta=(DisplayName = "Drag"))
 class UParticleModuleAccelerationDrag : public UParticleModuleAccelerationBase
 {
-	GENERATED_BODY()
-public:
-	ENGINE_API UParticleModuleAccelerationDrag(const FObjectInitializer& ObjectInitializer = FObjectInitializer::Get());
+	GENERATED_UCLASS_BODY()
 
 	/** Per-particle drag coefficient. Evaluted using emitter time. */
 	UPROPERTY(EditAnywhere, Category=Drag)
diff --git a/Engine/Source/Runtime/Engine/Classes/Particles/Acceleration/ParticleModuleAccelerationDragScaleOverLife.h b/Engine/Source/Runtime/Engine/Classes/Particles/Acceleration/ParticleModuleAccelerationDragScaleOverLife.h
--- a/Engine/Source/Runtime/Engine/Classes/Particles/Acceleration/ParticleModuleAccelerationDragScaleOverLife.h
+++ b/Engine/Source/Runtime/Engine/Classes/Particles/Acceleration/ParticleModuleAccelerationDragScaleOverLife.h
@@ -11,9 +11,7 @@
 UCLASS(editinlinenew, hidecategories=(UObject, Acceleration), MinimalAPI, meta=(DisplayName = "Drag Scale/Life"))
 class UParticleModuleAccelerationDragScaleOverLife : public UParticleModuleAccelerationBase
 {
-	GENERATED_BODY()
-public:
-	ENGINE_API UParticleModuleAccelerationDragScaleOverLife(const FObjectInitializer& ObjectInitializer = FObjectInitializer::Get());
+	GENERATED_UCLASS_BODY()
 
 	/** Per-particle drag scale. Evaluted using particle relative time. */
 	UPROPERTY(EditAnywhere, Category=Drag)
diff --git a/Engine/Source/Runtime/Engine/Classes/Particles/Acceleration/ParticleModuleAccelerationOverLifetime.h b/Engine/Source/Runtime/Engine/Classes/Particles/Acceleration/ParticleModuleAccelerationOverLifetime.h
--- a/Engine/Source/Runtime/Engine/Classes/Particles/Acceleration/ParticleModuleAccelerationOverLifetime.h
+++ b/Engine/Source/Runtime/Engine/Classes/Particles/Acceleration/ParticleModuleAccelerationOverLifetime.h
@@ -8,9 +8,7 @@
 UCLASS(editinlinenew, hidecategories=Object, meta=(DisplayName = "Acceleration/Life"), MinimalAPI)
 class UParticleModuleAccelerationOverLifetime : public UParticleModuleAccelerationBase
 {
-	GENERATED_BODY()
-public:
-	ENGINE_API UParticleModuleAccelerationOverLifetime(const FObjectInitializer& ObjectInitializer = FObjectInitializer::Get());
+	GENERATED_UCLASS_BODY()
 
 	/** 
 	 *	The acceleration of the particle over its lifetime.
diff --git a/Engine/Source/Runtime/Engine/Classes/Particles/Attractor/ParticleModuleAttractorBase.h b/Engine/Source/Runtime/Engine/Classes/Particles/Attractor/ParticleModuleAttractorBase.h
--- a/Engine/Source/Runtime/Engine/Classes/Particles/Attractor/ParticleModuleAttractorBase.h
+++ b/Engine/Source/Runtime/Engine/Classes/Particles/Attractor/ParticleModuleAttractorBase.h
@@ -8,9 +8,7 @@
 UCLASS(editinlinenew, hidecategories=Object, abstract, meta=(DisplayName = "Attraction"))
 class UParticleModuleAttractorBase : public UParticleModule
 {
-	GENERATED_BODY()
-public:
-	UParticleModuleAttractorBase(const FObjectInitializer& ObjectInitializer = FObjectInitializer::Get());
+	GENERATED_UCLASS_BODY()
 
 };
 
diff --git a/Engine/Source/Runtime/Engine/Classes/Particles/Attractor/ParticleModuleAttractorLine.h b/Engine/Source/Runtime/Engine/Classes/Particles/Attractor/ParticleModuleAttractorLine.h
--- a/Engine/Source/Runtime/Engine/Classes/Particles/Attractor/ParticleModuleAttractorLine.h
+++ b/Engine/Source/Runtime/Engine/Classes/Particles/Attractor/ParticleModuleAttractorLine.h
@@ -8,9 +8,7 @@
 UCLASS(editinlinenew, hidecategories=Object, MinimalAPI, meta=(DisplayName = "Line Attractor"))
 class UParticleModuleAttractorLine : public UParticleModuleAttractorBase
 {
-	GENERATED_BODY()
-public:
-	ENGINE_API UParticleModuleAttractorLine(const FObjectInitializer& ObjectInitializer = FObjectInitializer::Get());
+	GENERATED_UCLASS_BODY()
 
 	/** The first endpoint of the line. */
 	UPROPERTY(EditAnywhere, Category=Attractor)
diff --git a/Engine/Source/Runtime/Engine/Classes/Particles/Attractor/ParticleModuleAttractorParticle.h b/Engine/Source/Runtime/Engine/Classes/Particles/Attractor/ParticleModuleAttractorParticle.h
--- a/Engine/Source/Runtime/Engine/Classes/Particles/Attractor/ParticleModuleAttractorParticle.h
+++ b/Engine/Source/Runtime/Engine/Classes/Particles/Attractor/ParticleModuleAttractorParticle.h
@@ -16,9 +16,7 @@ enum EAttractorParticleSelectionMethod
 UCLASS(editinlinenew, hidecategories=Object, MinimalAPI, meta=(DisplayName = "Particle Attractor"))
 class UParticleModuleAttractorParticle : public UParticleModuleAttractorBase
 {
-	GENERATED_BODY()
-public:
-	ENGINE_API UParticleModuleAttractorParticle(const FObjectInitializer& ObjectInitializer = FObjectInitializer::Get());
+	GENERATED_UCLASS_BODY()
 
 	/**
 	 *	The source emitter for attractors
diff --git a/Engine/Source/Runtime/Engine/Classes/Particles/Attractor/ParticleModuleAttractorPoint.h b/Engine/Source/Runtime/Engine/Classes/Particles/Attractor/ParticleModuleAttractorPoint.h
--- a/Engine/Source/Runtime/Engine/Classes/Particles/Attractor/ParticleModuleAttractorPoint.h
+++ b/Engine/Source/Runtime/Engine/Classes/Particles/Attractor/ParticleModuleAttractorPoint.h
@@ -8,9 +8,7 @@
 UCLASS(editinlinenew, hidecategories=Object, MinimalAPI, meta=(DisplayName = "Point Attractor"))
 class UParticleModuleAttractorPoint : public UParticleModuleAttractorBase
 {
-	GENERATED_BODY()
-public:
-	ENGINE_API UParticleModuleAttractorPoint(const FObjectInitializer& ObjectInitializer = FObjectInitializer::Get());
+	GENERATED_UCLASS_BODY()
 
 	/**	The position of the point attractor from the source of the emitter.		*/
 	UPROPERTY(EditAnywhere, Category=Attractor)
diff --git a/Engine/Source/Runtime/Engine/Classes/Particles/Attractor/ParticleModuleAttractorPointGravity.h b/Engine/Source/Runtime/Engine/Classes/Particles/Attractor/ParticleModuleAttractorPointGravity.h
--- a/Engine/Source/Runtime/Engine/Classes/Particles/Attractor/ParticleModuleAttractorPointGravity.h
+++ b/Engine/Source/Runtime/Engine/Classes/Particles/Attractor/ParticleModuleAttractorPointGravity.h
@@ -12,9 +12,7 @@
 UCLASS(editinlinenew, hidecategories=Object, meta=(DisplayName = "Point Gravity"))
 class UParticleModuleAttractorPointGravity : public UParticleModuleAttractorBase
 {
-	GENERATED_BODY()
-public:
-	UParticleModuleAttractorPointGravity(const FObjectInitializer& ObjectInitializer = FObjectInitializer::Get());
+	GENERATED_UCLASS_BODY()
 
 	/** The position of the point gravity source. */
 	UPROPERTY(EditAnywhere, Category=PointGravitySource)
diff --git a/Engine/Source/Runtime/Engine/Classes/Particles/Beam/ParticleModuleBeamBase.h b/Engine/Source/Runtime/Engine/Classes/Particles/Beam/ParticleModuleBeamBase.h
--- a/Engine/Source/Runtime/Engine/Classes/Particles/Beam/ParticleModuleBeamBase.h
+++ b/Engine/Source/Runtime/Engine/Classes/Particles/Beam/ParticleModuleBeamBase.h
@@ -55,9 +55,7 @@ enum Beam2SourceTargetTangentMethod
 UCLASS(editinlinenew, hidecategories=Object, abstract, meta=(DisplayName = "Beam"))
 class UParticleModuleBeamBase : public UParticleModule
 {
-	GENERATED_BODY()
-public:
-	UParticleModuleBeamBase(const FObjectInitializer& ObjectInitializer = FObjectInitializer::Get());
+	GENERATED_UCLASS_BODY()
 
 
 	//Begin UParticleModule Interface
diff --git a/Engine/Source/Runtime/Engine/Classes/Particles/Beam/ParticleModuleBeamModifier.h b/Engine/Source/Runtime/Engine/Classes/Particles/Beam/ParticleModuleBeamModifier.h
--- a/Engine/Source/Runtime/Engine/Classes/Particles/Beam/ParticleModuleBeamModifier.h
+++ b/Engine/Source/Runtime/Engine/Classes/Particles/Beam/ParticleModuleBeamModifier.h
@@ -25,7 +25,7 @@ enum BeamModifierType
 USTRUCT()
 struct FBeamModifierOptions
 {
-	GENERATED_BODY()
+	GENERATED_USTRUCT_BODY()
 
 	/** If true, modify the value associated with this grouping.	*/
 	UPROPERTY(EditAnywhere, Category=BeamModifierOptions)
@@ -52,9 +52,7 @@ struct FBeamModifierOptions
 UCLASS(editinlinenew, hidecategories=Object, meta=(DisplayName = "Beam Modifier"))
 class UParticleModuleBeamModifier : public UParticleModuleBeamBase
 {
-	GENERATED_BODY()
-public:
-	UParticleModuleBeamModifier(const FObjectInitializer& ObjectInitializer = FObjectInitializer::Get());
+	GENERATED_UCLASS_BODY()
 
 	/**	Whether this module modifies the Source or the Target. */
 	UPROPERTY(EditAnywhere, Category=Modifier)
diff --git a/Engine/Source/Runtime/Engine/Classes/Particles/Beam/ParticleModuleBeamNoise.h b/Engine/Source/Runtime/Engine/Classes/Particles/Beam/ParticleModuleBeamNoise.h
--- a/Engine/Source/Runtime/Engine/Classes/Particles/Beam/ParticleModuleBeamNoise.h
+++ b/Engine/Source/Runtime/Engine/Classes/Particles/Beam/ParticleModuleBeamNoise.h
@@ -14,9 +14,7 @@
 UCLASS(editinlinenew, hidecategories=Object, MinimalAPI, meta=(DisplayName = "Noise"))
 class UParticleModuleBeamNoise : public UParticleModuleBeamBase
 {
-	GENERATED_BODY()
-public:
-	ENGINE_API UParticleModuleBeamNoise(const FObjectInitializer& ObjectInitializer = FObjectInitializer::Get());
+	GENERATED_UCLASS_BODY()
 
 	/** Is low frequency noise enabled. */
 	UPROPERTY(EditAnywhere, Category=LowFreq)
diff --git a/Engine/Source/Runtime/Engine/Classes/Particles/Beam/ParticleModuleBeamSource.h b/Engine/Source/Runtime/Engine/Classes/Particles/Beam/ParticleModuleBeamSource.h
--- a/Engine/Source/Runtime/Engine/Classes/Particles/Beam/ParticleModuleBeamSource.h
+++ b/Engine/Source/Runtime/Engine/Classes/Particles/Beam/ParticleModuleBeamSource.h
@@ -14,9 +14,7 @@
 UCLASS(editinlinenew, hidecategories=Object, MinimalAPI, meta=(DisplayName = "Source"))
 class UParticleModuleBeamSource : public UParticleModuleBeamBase
 {
-	GENERATED_BODY()
-public:
-	ENGINE_API UParticleModuleBeamSource(const FObjectInitializer& ObjectInitializer = FObjectInitializer::Get());
+	GENERATED_UCLASS_BODY()
 
 	/** The method flag. */
 	UPROPERTY(EditAnywhere, Category=Source)
diff --git a/Engine/Source/Runtime/Engine/Classes/Particles/Beam/ParticleModuleBeamTarget.h b/Engine/Source/Runtime/Engine/Classes/Particles/Beam/ParticleModuleBeamTarget.h
--- a/Engine/Source/Runtime/Engine/Classes/Particles/Beam/ParticleModuleBeamTarget.h
+++ b/Engine/Source/Runtime/Engine/Classes/Particles/Beam/ParticleModuleBeamTarget.h
@@ -14,9 +14,7 @@
 UCLASS(editinlinenew, hidecategories=Object, MinimalAPI, meta=(DisplayName = "Target"))
 class UParticleModuleBeamTarget : public UParticleModuleBeamBase
 {
-	GENERATED_BODY()
-public:
-	ENGINE_API UParticleModuleBeamTarget(const FObjectInitializer& ObjectInitializer = FObjectInitializer::Get());
+	GENERATED_UCLASS_BODY()
 
 	/** The method flag. */
 	UPROPERTY(EditAnywhere, Category=Target)
diff --git a/Engine/Source/Runtime/Engine/Classes/Particles/Camera/ParticleModuleCameraBase.h b/Engine/Source/Runtime/Engine/Classes/Particles/Camera/ParticleModuleCameraBase.h
--- a/Engine/Source/Runtime/Engine/Classes/Particles/Camera/ParticleModuleCameraBase.h
+++ b/Engine/Source/Runtime/Engine/Classes/Particles/Camera/ParticleModuleCameraBase.h
@@ -8,9 +8,7 @@
 UCLASS(editinlinenew, hidecategories=Object, abstract, meta=(DisplayName = "Camera"))
 class UParticleModuleCameraBase : public UParticleModule
 {
-	GENERATED_BODY()
-public:
-	UParticleModuleCameraBase(const FObjectInitializer& ObjectInitializer = FObjectInitializer::Get());
+	GENERATED_UCLASS_BODY()
 
 	virtual bool CanTickInAnyThread() override
 	{
diff --git a/Engine/Source/Runtime/Engine/Classes/Particles/Camera/ParticleModuleCameraOffset.h b/Engine/Source/Runtime/Engine/Classes/Particles/Camera/ParticleModuleCameraOffset.h
--- a/Engine/Source/Runtime/Engine/Classes/Particles/Camera/ParticleModuleCameraOffset.h
+++ b/Engine/Source/Runtime/Engine/Classes/Particles/Camera/ParticleModuleCameraOffset.h
@@ -20,9 +20,7 @@ enum EParticleCameraOffsetUpdateMethod
 UCLASS(editinlinenew, hidecategories=Object, meta=(DisplayName = "Camera Offset"))
 class UParticleModuleCameraOffset : public UParticleModuleCameraBase
 {
-	GENERATED_BODY()
-public:
-	UParticleModuleCameraOffset(const FObjectInitializer& ObjectInitializer = FObjectInitializer::Get());
+	GENERATED_UCLASS_BODY()
 
 	/** 
 	 *	The camera-relative offset to apply to sprite location
diff --git a/Engine/Source/Runtime/Engine/Classes/Particles/Collision/ParticleModuleCollision.h b/Engine/Source/Runtime/Engine/Classes/Particles/Collision/ParticleModuleCollision.h
--- a/Engine/Source/Runtime/Engine/Classes/Particles/Collision/ParticleModuleCollision.h
+++ b/Engine/Source/Runtime/Engine/Classes/Particles/Collision/ParticleModuleCollision.h
@@ -8,9 +8,7 @@
 UCLASS(editinlinenew, hidecategories=Object, MinimalAPI, meta=(DisplayName = "Collision"))
 class UParticleModuleCollision : public UParticleModuleCollisionBase
 {
-	GENERATED_BODY()
-public:
-	ENGINE_API UParticleModuleCollision(const FObjectInitializer& ObjectInitializer = FObjectInitializer::Get());
+	GENERATED_UCLASS_BODY()
 
 	/**
 	 *	How much to `slow' the velocity of the particle after a collision.
diff --git a/Engine/Source/Runtime/Engine/Classes/Particles/Collision/ParticleModuleCollisionBase.h b/Engine/Source/Runtime/Engine/Classes/Particles/Collision/ParticleModuleCollisionBase.h
--- a/Engine/Source/Runtime/Engine/Classes/Particles/Collision/ParticleModuleCollisionBase.h
+++ b/Engine/Source/Runtime/Engine/Classes/Particles/Collision/ParticleModuleCollisionBase.h
@@ -29,9 +29,7 @@ enum EParticleCollisionComplete
 UCLASS(editinlinenew, hidecategories=Object, abstract, meta=(DisplayName = "Collision"))
 class UParticleModuleCollisionBase : public UParticleModule
 {
-	GENERATED_BODY()
-public:
-	UParticleModuleCollisionBase(const FObjectInitializer& ObjectInitializer = FObjectInitializer::Get());
+	GENERATED_UCLASS_BODY()
 
 };
 
diff --git a/Engine/Source/Runtime/Engine/Classes/Particles/Collision/ParticleModuleCollisionGPU.h b/Engine/Source/Runtime/Engine/Classes/Particles/Collision/ParticleModuleCollisionGPU.h
--- a/Engine/Source/Runtime/Engine/Classes/Particles/Collision/ParticleModuleCollisionGPU.h
+++ b/Engine/Source/Runtime/Engine/Classes/Particles/Collision/ParticleModuleCollisionGPU.h
@@ -24,9 +24,7 @@ namespace EParticleCollisionResponse
 UCLASS(editinlinenew, hidecategories=Object, MinimalAPI, meta=(DisplayName = "Collision (Scene Depth)"))
 class UParticleModuleCollisionGPU : public UParticleModuleCollisionBase
 {
-	GENERATED_BODY()
-public:
-	ENGINE_API UParticleModuleCollisionGPU(const FObjectInitializer& ObjectInitializer = FObjectInitializer::Get());
+	GENERATED_UCLASS_BODY()
 
 	/**
 	 * Dampens the velocity of a particle in the direction normal to the
diff --git a/Engine/Source/Runtime/Engine/Classes/Particles/Color/ParticleModuleColor.h b/Engine/Source/Runtime/Engine/Classes/Particles/Color/ParticleModuleColor.h
--- a/Engine/Source/Runtime/Engine/Classes/Particles/Color/ParticleModuleColor.h
+++ b/Engine/Source/Runtime/Engine/Classes/Particles/Color/ParticleModuleColor.h
@@ -8,9 +8,7 @@
 UCLASS(editinlinenew, hidecategories=Object, MinimalAPI, meta=(DisplayName = "Initial Color"))
 class UParticleModuleColor : public UParticleModuleColorBase
 {
-	GENERATED_BODY()
-public:
-	ENGINE_API UParticleModuleColor(const FObjectInitializer& ObjectInitializer = FObjectInitializer::Get());
+	GENERATED_UCLASS_BODY()
 
 	/** Initial color for a particle as a function of Emitter time. */
 	UPROPERTY(EditAnywhere, Category = Color, meta = (TreatAsColor))
diff --git a/Engine/Source/Runtime/Engine/Classes/Particles/Color/ParticleModuleColorBase.h b/Engine/Source/Runtime/Engine/Classes/Particles/Color/ParticleModuleColorBase.h
--- a/Engine/Source/Runtime/Engine/Classes/Particles/Color/ParticleModuleColorBase.h
+++ b/Engine/Source/Runtime/Engine/Classes/Particles/Color/ParticleModuleColorBase.h
@@ -8,9 +8,7 @@
 UCLASS(editinlinenew, hidecategories=Object, abstract, meta=(DisplayName = "Color"))
 class UParticleModuleColorBase : public UParticleModule
 {
-	GENERATED_BODY()
-public:
-	UParticleModuleColorBase(const FObjectInitializer& ObjectInitializer = FObjectInitializer::Get());
+	GENERATED_UCLASS_BODY()
 
 };
 
diff --git a/Engine/Source/Runtime/Engine/Classes/Particles/Color/ParticleModuleColorOverLife.h b/Engine/Source/Runtime/Engine/Classes/Particles/Color/ParticleModuleColorOverLife.h
--- a/Engine/Source/Runtime/Engine/Classes/Particles/Color/ParticleModuleColorOverLife.h
+++ b/Engine/Source/Runtime/Engine/Classes/Particles/Color/ParticleModuleColorOverLife.h
@@ -8,9 +8,7 @@
 UCLASS(editinlinenew, hidecategories=Object, MinimalAPI, meta=(DisplayName = "Color Over Life"))
 class UParticleModuleColorOverLife : public UParticleModuleColorBase
 {
-	GENERATED_BODY()
-public:
-	ENGINE_API UParticleModuleColorOverLife(const FObjectInitializer& ObjectInitializer = FObjectInitializer::Get());
+	GENERATED_UCLASS_BODY()
 
 	/** The color to apply to the particle, as a function of the particle RelativeTime. */
 	UPROPERTY(EditAnywhere, Category=Color, meta=(TreatAsColor))
diff --git a/Engine/Source/Runtime/Engine/Classes/Particles/Color/ParticleModuleColorScaleOverLife.h b/Engine/Source/Runtime/Engine/Classes/Particles/Color/ParticleModuleColorScaleOverLife.h
--- a/Engine/Source/Runtime/Engine/Classes/Particles/Color/ParticleModuleColorScaleOverLife.h
+++ b/Engine/Source/Runtime/Engine/Classes/Particles/Color/ParticleModuleColorScaleOverLife.h
@@ -14,9 +14,7 @@
 UCLASS(editinlinenew, hidecategories=Object, MinimalAPI, meta=(DisplayName = "Scale Color / Life"))
 class UParticleModuleColorScaleOverLife : public UParticleModuleColorBase
 {
-	GENERATED_BODY()
-public:
-	ENGINE_API UParticleModuleColorScaleOverLife(const FObjectInitializer& ObjectInitializer = FObjectInitializer::Get());
+	GENERATED_UCLASS_BODY()
 
 	/** The scale factor for the color.													*/
 	UPROPERTY(EditAnywhere, Category = Color, meta = (TreatAsColor))
diff --git a/Engine/Source/Runtime/Engine/Classes/Particles/Color/ParticleModuleColor_Seeded.h b/Engine/Source/Runtime/Engine/Classes/Particles/Color/ParticleModuleColor_Seeded.h
--- a/Engine/Source/Runtime/Engine/Classes/Particles/Color/ParticleModuleColor_Seeded.h
+++ b/Engine/Source/Runtime/Engine/Classes/Particles/Color/ParticleModuleColor_Seeded.h
@@ -8,9 +8,7 @@
 UCLASS(editinlinenew, hidecategories=Object, meta=(DisplayName = "Init Color (Seed)"))
 class UParticleModuleColor_Seeded : public UParticleModuleColor
 {
-	GENERATED_BODY()
-public:
-	UParticleModuleColor_Seeded(const FObjectInitializer& ObjectInitializer = FObjectInitializer::Get());
+	GENERATED_UCLASS_BODY()
 
 	/** The random seed(s) to use for looking up values in StartLocation */
 	UPROPERTY(EditAnywhere, Category=RandomSeed)
diff --git a/Engine/Source/Runtime/Engine/Classes/Particles/EmitterCameraLensEffectBase.h b/Engine/Source/Runtime/Engine/Classes/Particles/EmitterCameraLensEffectBase.h
--- a/Engine/Source/Runtime/Engine/Classes/Particles/EmitterCameraLensEffectBase.h
+++ b/Engine/Source/Runtime/Engine/Classes/Particles/EmitterCameraLensEffectBase.h
@@ -13,9 +13,7 @@
 UCLASS(abstract, Blueprintable, MinimalAPI)
 class AEmitterCameraLensEffectBase : public AEmitter
 {
-	GENERATED_BODY()
-public:
-	ENGINE_API AEmitterCameraLensEffectBase(const FObjectInitializer& ObjectInitializer = FObjectInitializer::Get());
+	GENERATED_UCLASS_BODY()
 
 protected:
 	/** Particle System to use */
diff --git a/Engine/Source/Runtime/Engine/Classes/Particles/Event/ParticleModuleEventBase.h b/Engine/Source/Runtime/Engine/Classes/Particles/Event/ParticleModuleEventBase.h
--- a/Engine/Source/Runtime/Engine/Classes/Particles/Event/ParticleModuleEventBase.h
+++ b/Engine/Source/Runtime/Engine/Classes/Particles/Event/ParticleModuleEventBase.h
@@ -8,9 +8,7 @@
 UCLASS(editinlinenew, hidecategories=Object, abstract, meta=(DisplayName = "Event"))
 class UParticleModuleEventBase : public UParticleModule
 {
-	GENERATED_BODY()
-public:
-	UParticleModuleEventBase(const FObjectInitializer& ObjectInitializer = FObjectInitializer::Get());
+	GENERATED_UCLASS_BODY()
 
 
 	// Begin UParticleModule Interface
diff --git a/Engine/Source/Runtime/Engine/Classes/Particles/Event/ParticleModuleEventGenerator.h b/Engine/Source/Runtime/Engine/Classes/Particles/Event/ParticleModuleEventGenerator.h
--- a/Engine/Source/Runtime/Engine/Classes/Particles/Event/ParticleModuleEventGenerator.h
+++ b/Engine/Source/Runtime/Engine/Classes/Particles/Event/ParticleModuleEventGenerator.h
@@ -10,7 +10,7 @@
 USTRUCT()
 struct FParticleEvent_GenerateInfo
 {
-	GENERATED_BODY()
+	GENERATED_USTRUCT_BODY()
 
 	/** The type of event. */
 	UPROPERTY(EditAnywhere, Category=ParticleEvent_GenerateInfo)
@@ -61,9 +61,7 @@ struct FParticleEvent_GenerateInfo
 UCLASS(editinlinenew, hidecategories=Object, MinimalAPI, meta=(DisplayName = "Event Generator"))
 class UParticleModuleEventGenerator : public UParticleModuleEventBase
 {
-	GENERATED_BODY()
-public:
-	ENGINE_API UParticleModuleEventGenerator(const FObjectInitializer& ObjectInitializer = FObjectInitializer::Get());
+	GENERATED_UCLASS_BODY()
 
 	/**
 	 */
diff --git a/Engine/Source/Runtime/Engine/Classes/Particles/Event/ParticleModuleEventReceiverBase.h b/Engine/Source/Runtime/Engine/Classes/Particles/Event/ParticleModuleEventReceiverBase.h
--- a/Engine/Source/Runtime/Engine/Classes/Particles/Event/ParticleModuleEventReceiverBase.h
+++ b/Engine/Source/Runtime/Engine/Classes/Particles/Event/ParticleModuleEventReceiverBase.h
@@ -9,9 +9,7 @@
 UCLASS(editinlinenew, hidecategories=Object, abstract)
 class UParticleModuleEventReceiverBase : public UParticleModuleEventBase
 {
-	GENERATED_BODY()
-public:
-	UParticleModuleEventReceiverBase(const FObjectInitializer& ObjectInitializer = FObjectInitializer::Get());
+	GENERATED_UCLASS_BODY()
 
 	/** The type of event that will generate the kill. */
 	UPROPERTY(EditAnywhere, Category=Source)
diff --git a/Engine/Source/Runtime/Engine/Classes/Particles/Event/ParticleModuleEventReceiverKillParticles.h b/Engine/Source/Runtime/Engine/Classes/Particles/Event/ParticleModuleEventReceiverKillParticles.h
--- a/Engine/Source/Runtime/Engine/Classes/Particles/Event/ParticleModuleEventReceiverKillParticles.h
+++ b/Engine/Source/Runtime/Engine/Classes/Particles/Event/ParticleModuleEventReceiverKillParticles.h
@@ -8,9 +8,7 @@
 UCLASS(editinlinenew, hidecategories=Object, meta=(DisplayName = "EventReceiver Kill All"))
 class UParticleModuleEventReceiverKillParticles : public UParticleModuleEventReceiverBase
 {
-	GENERATED_BODY()
-public:
-	UParticleModuleEventReceiverKillParticles(const FObjectInitializer& ObjectInitializer = FObjectInitializer::Get());
+	GENERATED_UCLASS_BODY()
 
 	/** If true, stop this emitter from spawning as well. */
 	UPROPERTY(EditAnywhere, Category=ParticleModuleEventReceiverKillParticles)
diff --git a/Engine/Source/Runtime/Engine/Classes/Particles/Event/ParticleModuleEventReceiverSpawn.h b/Engine/Source/Runtime/Engine/Classes/Particles/Event/ParticleModuleEventReceiverSpawn.h
--- a/Engine/Source/Runtime/Engine/Classes/Particles/Event/ParticleModuleEventReceiverSpawn.h
+++ b/Engine/Source/Runtime/Engine/Classes/Particles/Event/ParticleModuleEventReceiverSpawn.h
@@ -8,9 +8,7 @@
 UCLASS(editinlinenew, hidecategories=Object, meta=(DisplayName = "EventReceiver Spawn"))
 class UParticleModuleEventReceiverSpawn : public UParticleModuleEventReceiverBase
 {
-	GENERATED_BODY()
-public:
-	UParticleModuleEventReceiverSpawn(const FObjectInitializer& ObjectInitializer = FObjectInitializer::Get());
+	GENERATED_UCLASS_BODY()
 
 	/** The number of particles to spawn. */
 	UPROPERTY(EditAnywhere, Category=Spawn)
diff --git a/Engine/Source/Runtime/Engine/Classes/Particles/Kill/ParticleModuleKillBase.h b/Engine/Source/Runtime/Engine/Classes/Particles/Kill/ParticleModuleKillBase.h
--- a/Engine/Source/Runtime/Engine/Classes/Particles/Kill/ParticleModuleKillBase.h
+++ b/Engine/Source/Runtime/Engine/Classes/Particles/Kill/ParticleModuleKillBase.h
@@ -8,9 +8,7 @@
 UCLASS(editinlinenew, hidecategories=Object, abstract, meta=(DisplayName = "Kill"))
 class UParticleModuleKillBase : public UParticleModule
 {
-	GENERATED_BODY()
-public:
-	UParticleModuleKillBase(const FObjectInitializer& ObjectInitializer = FObjectInitializer::Get());
+	GENERATED_UCLASS_BODY()
 
 };
 
diff --git a/Engine/Source/Runtime/Engine/Classes/Particles/Kill/ParticleModuleKillBox.h b/Engine/Source/Runtime/Engine/Classes/Particles/Kill/ParticleModuleKillBox.h
--- a/Engine/Source/Runtime/Engine/Classes/Particles/Kill/ParticleModuleKillBox.h
+++ b/Engine/Source/Runtime/Engine/Classes/Particles/Kill/ParticleModuleKillBox.h
@@ -8,9 +8,7 @@
 UCLASS(editinlinenew, hidecategories=Object, meta=(DisplayName = "Kill Box"))
 class UParticleModuleKillBox : public UParticleModuleKillBase
 {
-	GENERATED_BODY()
-public:
-	UParticleModuleKillBox(const FObjectInitializer& ObjectInitializer = FObjectInitializer::Get());
+	GENERATED_UCLASS_BODY()
 
 	/** The lower left corner of the box. */
 	UPROPERTY(EditAnywhere, Category=Kill)
diff --git a/Engine/Source/Runtime/Engine/Classes/Particles/Kill/ParticleModuleKillHeight.h b/Engine/Source/Runtime/Engine/Classes/Particles/Kill/ParticleModuleKillHeight.h
--- a/Engine/Source/Runtime/Engine/Classes/Particles/Kill/ParticleModuleKillHeight.h
+++ b/Engine/Source/Runtime/Engine/Classes/Particles/Kill/ParticleModuleKillHeight.h
@@ -8,9 +8,7 @@
 UCLASS(editinlinenew, hidecategories=Object, meta=(DisplayName = "Kill Height"))
 class UParticleModuleKillHeight : public UParticleModuleKillBase
 {
-	GENERATED_BODY()
-public:
-	UParticleModuleKillHeight(const FObjectInitializer& ObjectInitializer = FObjectInitializer::Get());
+	GENERATED_UCLASS_BODY()
 
 	/** The height at which to kill the particle. */
 	UPROPERTY(EditAnywhere, Category=Kill)
diff --git a/Engine/Source/Runtime/Engine/Classes/Particles/Lifetime/ParticleModuleLifetime.h b/Engine/Source/Runtime/Engine/Classes/Particles/Lifetime/ParticleModuleLifetime.h
--- a/Engine/Source/Runtime/Engine/Classes/Particles/Lifetime/ParticleModuleLifetime.h
+++ b/Engine/Source/Runtime/Engine/Classes/Particles/Lifetime/ParticleModuleLifetime.h
@@ -8,9 +8,7 @@
 UCLASS(editinlinenew, hidecategories=Object, MinimalAPI, meta=(DisplayName = "Lifetime"))
 class UParticleModuleLifetime : public UParticleModuleLifetimeBase
 {
-	GENERATED_BODY()
-public:
-	ENGINE_API UParticleModuleLifetime(const FObjectInitializer& ObjectInitializer = FObjectInitializer::Get());
+	GENERATED_UCLASS_BODY()
 
 	/** The lifetime of the particle, in seconds. Retrieved using the EmitterTime at the spawn of the particle. */
 	UPROPERTY(EditAnywhere, Category=Lifetime)
diff --git a/Engine/Source/Runtime/Engine/Classes/Particles/Lifetime/ParticleModuleLifetimeBase.h b/Engine/Source/Runtime/Engine/Classes/Particles/Lifetime/ParticleModuleLifetimeBase.h
--- a/Engine/Source/Runtime/Engine/Classes/Particles/Lifetime/ParticleModuleLifetimeBase.h
+++ b/Engine/Source/Runtime/Engine/Classes/Particles/Lifetime/ParticleModuleLifetimeBase.h
@@ -8,9 +8,7 @@
 UCLASS(editinlinenew, hidecategories=Object, abstract, meta=(DisplayName = "Lifetime"))
 class UParticleModuleLifetimeBase : public UParticleModule
 {
-	GENERATED_BODY()
-public:
-	UParticleModuleLifetimeBase(const FObjectInitializer& ObjectInitializer = FObjectInitializer::Get());
+	GENERATED_UCLASS_BODY()
 
 
 	/** Return the maximum lifetime this module would return. */
diff --git a/Engine/Source/Runtime/Engine/Classes/Particles/Lifetime/ParticleModuleLifetime_Seeded.h b/Engine/Source/Runtime/Engine/Classes/Particles/Lifetime/ParticleModuleLifetime_Seeded.h
--- a/Engine/Source/Runtime/Engine/Classes/Particles/Lifetime/ParticleModuleLifetime_Seeded.h
+++ b/Engine/Source/Runtime/Engine/Classes/Particles/Lifetime/ParticleModuleLifetime_Seeded.h
@@ -8,9 +8,7 @@
 UCLASS(editinlinenew, hidecategories=Object, meta=(DisplayName = "Lifetime (Seed)"))
 class UParticleModuleLifetime_Seeded : public UParticleModuleLifetime
 {
-	GENERATED_BODY()
-public:
-	UParticleModuleLifetime_Seeded(const FObjectInitializer& ObjectInitializer = FObjectInitializer::Get());
+	GENERATED_UCLASS_BODY()
 
 	/** The random seed(s) to use for looking up values in StartLocation */
 	UPROPERTY(EditAnywhere, Category=RandomSeed)
diff --git a/Engine/Source/Runtime/Engine/Classes/Particles/Light/ParticleModuleLight.h b/Engine/Source/Runtime/Engine/Classes/Particles/Light/ParticleModuleLight.h
--- a/Engine/Source/Runtime/Engine/Classes/Particles/Light/ParticleModuleLight.h
+++ b/Engine/Source/Runtime/Engine/Classes/Particles/Light/ParticleModuleLight.h
@@ -9,9 +9,7 @@
 UCLASS(editinlinenew, hidecategories=Object, MinimalAPI, meta=(DisplayName = "Light"))
 class UParticleModuleLight : public UParticleModuleLightBase
 {
-	GENERATED_BODY()
-public:
-	ENGINE_API UParticleModuleLight(const FObjectInitializer& ObjectInitializer = FObjectInitializer::Get());
+	GENERATED_UCLASS_BODY()
 
 	/** Whether to use physically based inverse squared falloff from the light.  If unchecked, the LightExponent distribution will be used instead. */
 	UPROPERTY(EditAnywhere, Category=Light)
diff --git a/Engine/Source/Runtime/Engine/Classes/Particles/Light/ParticleModuleLightBase.h b/Engine/Source/Runtime/Engine/Classes/Particles/Light/ParticleModuleLightBase.h
--- a/Engine/Source/Runtime/Engine/Classes/Particles/Light/ParticleModuleLightBase.h
+++ b/Engine/Source/Runtime/Engine/Classes/Particles/Light/ParticleModuleLightBase.h
@@ -8,9 +8,7 @@
 UCLASS(editinlinenew, hidecategories=Object, abstract, meta=(DisplayName = "Light"))
 class UParticleModuleLightBase : public UParticleModule
 {
-	GENERATED_BODY()
-public:
-	UParticleModuleLightBase(const FObjectInitializer& ObjectInitializer = FObjectInitializer::Get());
+	GENERATED_UCLASS_BODY()
 	virtual bool CanTickInAnyThread() override
 	{
 		return false;
diff --git a/Engine/Source/Runtime/Engine/Classes/Particles/Light/ParticleModuleLight_Seeded.h b/Engine/Source/Runtime/Engine/Classes/Particles/Light/ParticleModuleLight_Seeded.h
--- a/Engine/Source/Runtime/Engine/Classes/Particles/Light/ParticleModuleLight_Seeded.h
+++ b/Engine/Source/Runtime/Engine/Classes/Particles/Light/ParticleModuleLight_Seeded.h
@@ -8,9 +8,7 @@
 UCLASS(editinlinenew, hidecategories=Object, MinimalAPI, meta=(DisplayName = "Light (Seed)"))
 class UParticleModuleLight_Seeded : public UParticleModuleLight
 {
-	GENERATED_BODY()
-public:
-	ENGINE_API UParticleModuleLight_Seeded(const FObjectInitializer& ObjectInitializer = FObjectInitializer::Get());
+	GENERATED_UCLASS_BODY()
 
 	/** The random seed(s) to use for looking up values in StartLocation */
 	UPROPERTY(EditAnywhere, Category=RandomSeed)
diff --git a/Engine/Source/Runtime/Engine/Classes/Particles/Location/ParticleModuleLocation.h b/Engine/Source/Runtime/Engine/Classes/Particles/Location/ParticleModuleLocation.h
--- a/Engine/Source/Runtime/Engine/Classes/Particles/Location/ParticleModuleLocation.h
+++ b/Engine/Source/Runtime/Engine/Classes/Particles/Location/ParticleModuleLocation.h
@@ -8,9 +8,7 @@
 UCLASS(editinlinenew, hidecategories=Object, MinimalAPI, meta=(DisplayName = "Initial Location"))
 class UParticleModuleLocation : public UParticleModuleLocationBase
 {
-	GENERATED_BODY()
-public:
-	ENGINE_API UParticleModuleLocation(const FObjectInitializer& ObjectInitializer = FObjectInitializer::Get());
+	GENERATED_UCLASS_BODY()
 
 	/** 
 	 *	The location the particle should be emitted.
diff --git a/Engine/Source/Runtime/Engine/Classes/Particles/Location/ParticleModuleLocationBase.h b/Engine/Source/Runtime/Engine/Classes/Particles/Location/ParticleModuleLocationBase.h
--- a/Engine/Source/Runtime/Engine/Classes/Particles/Location/ParticleModuleLocationBase.h
+++ b/Engine/Source/Runtime/Engine/Classes/Particles/Location/ParticleModuleLocationBase.h
@@ -8,9 +8,7 @@
 UCLASS(editinlinenew, hidecategories=Object, abstract, meta=(DisplayName = "Location"))
 class UParticleModuleLocationBase : public UParticleModule
 {
-	GENERATED_BODY()
-public:
-	UParticleModuleLocationBase(const FObjectInitializer& ObjectInitializer = FObjectInitializer::Get());
+	GENERATED_UCLASS_BODY()
 
 };
 
diff --git a/Engine/Source/Runtime/Engine/Classes/Particles/Location/ParticleModuleLocationBoneSocket.h b/Engine/Source/Runtime/Engine/Classes/Particles/Location/ParticleModuleLocationBoneSocket.h
--- a/Engine/Source/Runtime/Engine/Classes/Particles/Location/ParticleModuleLocationBoneSocket.h
+++ b/Engine/Source/Runtime/Engine/Classes/Particles/Location/ParticleModuleLocationBoneSocket.h
@@ -24,7 +24,7 @@ enum ELocationBoneSocketSelectionMethod
 USTRUCT()
 struct FLocationBoneSocketInfo
 {
-	GENERATED_BODY()
+	GENERATED_USTRUCT_BODY()
 
 	/** The name of the bone/socket on the skeletal mesh */
 	UPROPERTY(EditAnywhere, Category=LocationBoneSocketInfo)
@@ -45,9 +45,7 @@ struct FLocationBoneSocketInfo
 UCLASS(editinlinenew, hidecategories=Object, meta=(DisplayName = "Bone/Socket Location"))
 class UParticleModuleLocationBoneSocket : public UParticleModuleLocationBase
 {
-	GENERATED_BODY()
-public:
-	UParticleModuleLocationBoneSocket(const FObjectInitializer& ObjectInitializer = FObjectInitializer::Get());
+	GENERATED_UCLASS_BODY()
 
 	/**
 	 *	Whether the module uses Bones or Sockets for locations.
diff --git a/Engine/Source/Runtime/Engine/Classes/Particles/Location/ParticleModuleLocationDirect.h b/Engine/Source/Runtime/Engine/Classes/Particles/Location/ParticleModuleLocationDirect.h
--- a/Engine/Source/Runtime/Engine/Classes/Particles/Location/ParticleModuleLocationDirect.h
+++ b/Engine/Source/Runtime/Engine/Classes/Particles/Location/ParticleModuleLocationDirect.h
@@ -14,9 +14,7 @@
 UCLASS(editinlinenew, hidecategories=Object, MinimalAPI, meta=(DisplayName = "Direct Location"))
 class UParticleModuleLocationDirect : public UParticleModuleLocationBase
 {
-	GENERATED_BODY()
-public:
-	ENGINE_API UParticleModuleLocationDirect(const FObjectInitializer& ObjectInitializer = FObjectInitializer::Get());
+	GENERATED_UCLASS_BODY()
 
 	/** 
 	 *	The location of the particle at a give time. Retrieved using the particle RelativeTime. 
diff --git a/Engine/Source/Runtime/Engine/Classes/Particles/Location/ParticleModuleLocationEmitter.h b/Engine/Source/Runtime/Engine/Classes/Particles/Location/ParticleModuleLocationEmitter.h
--- a/Engine/Source/Runtime/Engine/Classes/Particles/Location/ParticleModuleLocationEmitter.h
+++ b/Engine/Source/Runtime/Engine/Classes/Particles/Location/ParticleModuleLocationEmitter.h
@@ -23,9 +23,7 @@ enum ELocationEmitterSelectionMethod
 UCLASS(editinlinenew, hidecategories=Object, MinimalAPI, meta=(DisplayName = "Emitter Initial Location"))
 class UParticleModuleLocationEmitter : public UParticleModuleLocationBase
 {
-	GENERATED_BODY()
-public:
-	ENGINE_API UParticleModuleLocationEmitter(const FObjectInitializer& ObjectInitializer = FObjectInitializer::Get());
+	GENERATED_UCLASS_BODY()
 
 	//=============================================================================
 	// Variables
diff --git a/Engine/Source/Runtime/Engine/Classes/Particles/Location/ParticleModuleLocationPrimitiveBase.h b/Engine/Source/Runtime/Engine/Classes/Particles/Location/ParticleModuleLocationPrimitiveBase.h
--- a/Engine/Source/Runtime/Engine/Classes/Particles/Location/ParticleModuleLocationPrimitiveBase.h
+++ b/Engine/Source/Runtime/Engine/Classes/Particles/Location/ParticleModuleLocationPrimitiveBase.h
@@ -12,9 +12,7 @@
 UCLASS(editinlinenew, hidecategories=Object)
 class UParticleModuleLocationPrimitiveBase : public UParticleModuleLocationBase
 {
-	GENERATED_BODY()
-public:
-	UParticleModuleLocationPrimitiveBase(const FObjectInitializer& ObjectInitializer = FObjectInitializer::Get());
+	GENERATED_UCLASS_BODY()
 
 	/** Whether the positive X axis is valid for spawning. */
 	UPROPERTY(EditAnywhere, Category=Location)
diff --git a/Engine/Source/Runtime/Engine/Classes/Particles/Location/ParticleModuleLocationPrimitiveCylinder.h b/Engine/Source/Runtime/Engine/Classes/Particles/Location/ParticleModuleLocationPrimitiveCylinder.h
--- a/Engine/Source/Runtime/Engine/Classes/Particles/Location/ParticleModuleLocationPrimitiveCylinder.h
+++ b/Engine/Source/Runtime/Engine/Classes/Particles/Location/ParticleModuleLocationPrimitiveCylinder.h
@@ -21,9 +21,7 @@ enum CylinderHeightAxis
 UCLASS(editinlinenew, hidecategories = Object, MinimalAPI, meta = (DisplayName = "Cylinder"))
 class UParticleModuleLocationPrimitiveCylinder : public UParticleModuleLocationPrimitiveBase
 {
-	GENERATED_BODY()
-public:
-	ENGINE_API UParticleModuleLocationPrimitiveCylinder(const FObjectInitializer& ObjectInitializer = FObjectInitializer::Get());
+	GENERATED_UCLASS_BODY()
 
 	/** If true, get the particle velocity form the radial distance inside the primitive. */
 	UPROPERTY(EditAnywhere, Category=Location)
diff --git a/Engine/Source/Runtime/Engine/Classes/Particles/Location/ParticleModuleLocationPrimitiveCylinder_Seeded.h b/Engine/Source/Runtime/Engine/Classes/Particles/Location/ParticleModuleLocationPrimitiveCylinder_Seeded.h
--- a/Engine/Source/Runtime/Engine/Classes/Particles/Location/ParticleModuleLocationPrimitiveCylinder_Seeded.h
+++ b/Engine/Source/Runtime/Engine/Classes/Particles/Location/ParticleModuleLocationPrimitiveCylinder_Seeded.h
@@ -8,9 +8,7 @@
 UCLASS(editinlinenew, hidecategories=Object, meta=(DisplayName = "Cylinder (Seed)"))
 class UParticleModuleLocationPrimitiveCylinder_Seeded : public UParticleModuleLocationPrimitiveCylinder
 {
-	GENERATED_BODY()
-public:
-	UParticleModuleLocationPrimitiveCylinder_Seeded(const FObjectInitializer& ObjectInitializer = FObjectInitializer::Get());
+	GENERATED_UCLASS_BODY()
 
 	/** The random seed(s) to use for looking up values in StartLocation */
 	UPROPERTY(EditAnywhere, Category=RandomSeed)
diff --git a/Engine/Source/Runtime/Engine/Classes/Particles/Location/ParticleModuleLocationPrimitiveSphere.h b/Engine/Source/Runtime/Engine/Classes/Particles/Location/ParticleModuleLocationPrimitiveSphere.h
--- a/Engine/Source/Runtime/Engine/Classes/Particles/Location/ParticleModuleLocationPrimitiveSphere.h
+++ b/Engine/Source/Runtime/Engine/Classes/Particles/Location/ParticleModuleLocationPrimitiveSphere.h
@@ -12,9 +12,7 @@
 UCLASS(editinlinenew, hidecategories=Object, MinimalAPI, meta=(DisplayName = "Sphere"))
 class UParticleModuleLocationPrimitiveSphere : public UParticleModuleLocationPrimitiveBase
 {
-	GENERATED_BODY()
-public:
-	ENGINE_API UParticleModuleLocationPrimitiveSphere(const FObjectInitializer& ObjectInitializer = FObjectInitializer::Get());
+	GENERATED_UCLASS_BODY()
 
 	/** The radius of the sphere. Retrieved using EmitterTime. */
 	UPROPERTY(EditAnywhere, Category=Location)
diff --git a/Engine/Source/Runtime/Engine/Classes/Particles/Location/ParticleModuleLocationPrimitiveSphere_Seeded.h b/Engine/Source/Runtime/Engine/Classes/Particles/Location/ParticleModuleLocationPrimitiveSphere_Seeded.h
--- a/Engine/Source/Runtime/Engine/Classes/Particles/Location/ParticleModuleLocationPrimitiveSphere_Seeded.h
+++ b/Engine/Source/Runtime/Engine/Classes/Particles/Location/ParticleModuleLocationPrimitiveSphere_Seeded.h
@@ -8,9 +8,7 @@
 UCLASS(editinlinenew, hidecategories=Object, meta=(DisplayName = "Sphere (Seed)"))
 class UParticleModuleLocationPrimitiveSphere_Seeded : public UParticleModuleLocationPrimitiveSphere
 {
-	GENERATED_BODY()
-public:
-	UParticleModuleLocationPrimitiveSphere_Seeded(const FObjectInitializer& ObjectInitializer = FObjectInitializer::Get());
+	GENERATED_UCLASS_BODY()
 
 	/** The random seed(s) to use for looking up values in StartLocation */
 	UPROPERTY(EditAnywhere, Category=RandomSeed)
diff --git a/Engine/Source/Runtime/Engine/Classes/Particles/Location/ParticleModuleLocationPrimitiveTriangle.h b/Engine/Source/Runtime/Engine/Classes/Particles/Location/ParticleModuleLocationPrimitiveTriangle.h
--- a/Engine/Source/Runtime/Engine/Classes/Particles/Location/ParticleModuleLocationPrimitiveTriangle.h
+++ b/Engine/Source/Runtime/Engine/Classes/Particles/Location/ParticleModuleLocationPrimitiveTriangle.h
@@ -12,9 +12,7 @@
 UCLASS(editinlinenew, hidecategories=Object, MinimalAPI, meta=(DisplayName = "Triangle"))
 class UParticleModuleLocationPrimitiveTriangle : public UParticleModuleLocationBase
 {
-	GENERATED_BODY()
-public:
-	ENGINE_API UParticleModuleLocationPrimitiveTriangle(const FObjectInitializer& ObjectInitializer = FObjectInitializer::Get());
+	GENERATED_UCLASS_BODY()
 
 
 	UPROPERTY(EditAnywhere, Category=Location)
diff --git a/Engine/Source/Runtime/Engine/Classes/Particles/Location/ParticleModuleLocationSkelVertSurface.h b/Engine/Source/Runtime/Engine/Classes/Particles/Location/ParticleModuleLocationSkelVertSurface.h
--- a/Engine/Source/Runtime/Engine/Classes/Particles/Location/ParticleModuleLocationSkelVertSurface.h
+++ b/Engine/Source/Runtime/Engine/Classes/Particles/Location/ParticleModuleLocationSkelVertSurface.h
@@ -16,9 +16,7 @@ enum ELocationSkelVertSurfaceSource
 UCLASS(editinlinenew, hidecategories=Object, meta=(DisplayName = "Skel Vert/Surf Location"))
 class UParticleModuleLocationSkelVertSurface : public UParticleModuleLocationBase
 {
-	GENERATED_BODY()
-public:
-	UParticleModuleLocationSkelVertSurface(const FObjectInitializer& ObjectInitializer = FObjectInitializer::Get());
+	GENERATED_UCLASS_BODY()
 
 	/**
 	 *	Whether the module uses Verts or Surfaces for locations.
diff --git a/Engine/Source/Runtime/Engine/Classes/Particles/Location/ParticleModuleLocationWorldOffset.h b/Engine/Source/Runtime/Engine/Classes/Particles/Location/ParticleModuleLocationWorldOffset.h
--- a/Engine/Source/Runtime/Engine/Classes/Particles/Location/ParticleModuleLocationWorldOffset.h
+++ b/Engine/Source/Runtime/Engine/Classes/Particles/Location/ParticleModuleLocationWorldOffset.h
@@ -8,9 +8,7 @@
 UCLASS(editinlinenew, meta=(DisplayName = "World Offset"))
 class UParticleModuleLocationWorldOffset : public UParticleModuleLocation
 {
-	GENERATED_BODY()
-public:
-	UParticleModuleLocationWorldOffset(const FObjectInitializer& ObjectInitializer = FObjectInitializer::Get());
+	GENERATED_UCLASS_BODY()
 
 
 protected:
diff --git a/Engine/Source/Runtime/Engine/Classes/Particles/Location/ParticleModuleLocationWorldOffset_Seeded.h b/Engine/Source/Runtime/Engine/Classes/Particles/Location/ParticleModuleLocationWorldOffset_Seeded.h
--- a/Engine/Source/Runtime/Engine/Classes/Particles/Location/ParticleModuleLocationWorldOffset_Seeded.h
+++ b/Engine/Source/Runtime/Engine/Classes/Particles/Location/ParticleModuleLocationWorldOffset_Seeded.h
@@ -8,9 +8,7 @@
 UCLASS(editinlinenew, meta=(DisplayName = "World Offset (Seed)"))
 class UParticleModuleLocationWorldOffset_Seeded : public UParticleModuleLocationWorldOffset
 {
-	GENERATED_BODY()
-public:
-	UParticleModuleLocationWorldOffset_Seeded(const FObjectInitializer& ObjectInitializer = FObjectInitializer::Get());
+	GENERATED_UCLASS_BODY()
 
 	/** The random seed(s) to use for looking up values in StartLocation */
 	UPROPERTY(EditAnywhere, Category=RandomSeed)
diff --git a/Engine/Source/Runtime/Engine/Classes/Particles/Location/ParticleModuleLocation_Seeded.h b/Engine/Source/Runtime/Engine/Classes/Particles/Location/ParticleModuleLocation_Seeded.h
--- a/Engine/Source/Runtime/Engine/Classes/Particles/Location/ParticleModuleLocation_Seeded.h
+++ b/Engine/Source/Runtime/Engine/Classes/Particles/Location/ParticleModuleLocation_Seeded.h
@@ -8,9 +8,7 @@
 UCLASS(editinlinenew, hidecategories=Object, meta=(DisplayName = "Initial Location (Seed)"))
 class UParticleModuleLocation_Seeded : public UParticleModuleLocation
 {
-	GENERATED_BODY()
-public:
-	UParticleModuleLocation_Seeded(const FObjectInitializer& ObjectInitializer = FObjectInitializer::Get());
+	GENERATED_UCLASS_BODY()
 
 	/** The random seed(s) to use for looking up values in StartLocation */
 	UPROPERTY(EditAnywhere, Category=RandomSeed)
diff --git a/Engine/Source/Runtime/Engine/Classes/Particles/Location/ParticleModuleSourceMovement.h b/Engine/Source/Runtime/Engine/Classes/Particles/Location/ParticleModuleSourceMovement.h
--- a/Engine/Source/Runtime/Engine/Classes/Particles/Location/ParticleModuleSourceMovement.h
+++ b/Engine/Source/Runtime/Engine/Classes/Particles/Location/ParticleModuleSourceMovement.h
@@ -8,9 +8,7 @@
 UCLASS(editinlinenew, hidecategories=Object, meta=(DisplayName = "Source Movement"))
 class UParticleModuleSourceMovement : public UParticleModuleLocationBase
 {
-	GENERATED_BODY()
-public:
-	UParticleModuleSourceMovement(const FObjectInitializer& ObjectInitializer = FObjectInitializer::Get());
+	GENERATED_UCLASS_BODY()
 
 	/** 
 	 *	The scale factor to apply to the source movement before adding to the particle location.
diff --git a/Engine/Source/Runtime/Engine/Classes/Particles/Material/ParticleModuleMaterialBase.h b/Engine/Source/Runtime/Engine/Classes/Particles/Material/ParticleModuleMaterialBase.h
--- a/Engine/Source/Runtime/Engine/Classes/Particles/Material/ParticleModuleMaterialBase.h
+++ b/Engine/Source/Runtime/Engine/Classes/Particles/Material/ParticleModuleMaterialBase.h
@@ -8,9 +8,7 @@
 UCLASS(editinlinenew, hidecategories=Object, abstract, meta=(DisplayName = "Material"))
 class UParticleModuleMaterialBase : public UParticleModule
 {
-	GENERATED_BODY()
-public:
-	UParticleModuleMaterialBase(const FObjectInitializer& ObjectInitializer = FObjectInitializer::Get());
+	GENERATED_UCLASS_BODY()
 
 	virtual bool CanTickInAnyThread() override
 	{
diff --git a/Engine/Source/Runtime/Engine/Classes/Particles/Material/ParticleModuleMeshMaterial.h b/Engine/Source/Runtime/Engine/Classes/Particles/Material/ParticleModuleMeshMaterial.h
--- a/Engine/Source/Runtime/Engine/Classes/Particles/Material/ParticleModuleMeshMaterial.h
+++ b/Engine/Source/Runtime/Engine/Classes/Particles/Material/ParticleModuleMeshMaterial.h
@@ -8,9 +8,7 @@
 UCLASS(editinlinenew, hidecategories=Object, MinimalAPI, meta=(DisplayName = "Mesh Material"))
 class UParticleModuleMeshMaterial : public UParticleModuleMaterialBase
 {
-	GENERATED_BODY()
-public:
-	ENGINE_API UParticleModuleMeshMaterial(const FObjectInitializer& ObjectInitializer = FObjectInitializer::Get());
+	GENERATED_UCLASS_BODY()
 
 	/** The array of materials to apply to the mesh particles. */
 	UPROPERTY(EditAnywhere, Category=MeshMaterials)
diff --git a/Engine/Source/Runtime/Engine/Classes/Particles/Modules/Location/ParticleModulePivotOffset.h b/Engine/Source/Runtime/Engine/Classes/Particles/Modules/Location/ParticleModulePivotOffset.h
--- a/Engine/Source/Runtime/Engine/Classes/Particles/Modules/Location/ParticleModulePivotOffset.h
+++ b/Engine/Source/Runtime/Engine/Classes/Particles/Modules/Location/ParticleModulePivotOffset.h
@@ -7,9 +7,7 @@
 UCLASS(editinlinenew, hidecategories=Object, MinimalAPI, DisplayName="Pivot Offset")
 class UParticleModulePivotOffset : public UParticleModuleLocationBase
 {
-	GENERATED_BODY()
-public:
-	ENGINE_API UParticleModulePivotOffset(const FObjectInitializer& ObjectInitializer = FObjectInitializer::Get());
+	GENERATED_UCLASS_BODY()
 
 	/** Offset applied in UV space to the particle vertex positions. Defaults to (0.5,0.5) putting the pivot in the centre of the partilce. */
 	UPROPERTY(EditAnywhere, Category=PivotOffset)
diff --git a/Engine/Source/Runtime/Engine/Classes/Particles/Orbit/ParticleModuleOrbit.h b/Engine/Source/Runtime/Engine/Classes/Particles/Orbit/ParticleModuleOrbit.h
--- a/Engine/Source/Runtime/Engine/Classes/Particles/Orbit/ParticleModuleOrbit.h
+++ b/Engine/Source/Runtime/Engine/Classes/Particles/Orbit/ParticleModuleOrbit.h
@@ -23,7 +23,7 @@ enum EOrbitChainMode
 USTRUCT()
 struct FOrbitOptions
 {
-	GENERATED_BODY()
+	GENERATED_USTRUCT_BODY()
 
 	/**
 	 *	Whether to process the data during spawning.
@@ -57,9 +57,7 @@ struct FOrbitOptions
 UCLASS(editinlinenew, hidecategories=(Object, Orbit), MinimalAPI, meta=(DisplayName = "Orbit"))
 class UParticleModuleOrbit : public UParticleModuleOrbitBase
 {
-	GENERATED_BODY()
-public:
-	ENGINE_API UParticleModuleOrbit(const FObjectInitializer& ObjectInitializer = FObjectInitializer::Get());
+	GENERATED_UCLASS_BODY()
 
 	/**
 	 *	Orbit modules will chain together in the order they appear in the module stack.
diff --git a/Engine/Source/Runtime/Engine/Classes/Particles/Orbit/ParticleModuleOrbitBase.h b/Engine/Source/Runtime/Engine/Classes/Particles/Orbit/ParticleModuleOrbitBase.h
--- a/Engine/Source/Runtime/Engine/Classes/Particles/Orbit/ParticleModuleOrbitBase.h
+++ b/Engine/Source/Runtime/Engine/Classes/Particles/Orbit/ParticleModuleOrbitBase.h
@@ -8,9 +8,7 @@
 UCLASS(editinlinenew, hidecategories=Object, abstract, meta=(DisplayName = "Orbit"))
 class UParticleModuleOrbitBase : public UParticleModule
 {
-	GENERATED_BODY()
-public:
-	UParticleModuleOrbitBase(const FObjectInitializer& ObjectInitializer = FObjectInitializer::Get());
+	GENERATED_UCLASS_BODY()
 
 	/** 
 	 *	If true, distribution values will be retrieved using the EmitterTime.
diff --git a/Engine/Source/Runtime/Engine/Classes/Particles/Orientation/ParticleModuleOrientationBase.h b/Engine/Source/Runtime/Engine/Classes/Particles/Orientation/ParticleModuleOrientationBase.h
--- a/Engine/Source/Runtime/Engine/Classes/Particles/Orientation/ParticleModuleOrientationBase.h
+++ b/Engine/Source/Runtime/Engine/Classes/Particles/Orientation/ParticleModuleOrientationBase.h
@@ -8,9 +8,7 @@
 UCLASS(editinlinenew, hidecategories=Object, abstract, meta=(DisplayName = "Orientation"))
 class UParticleModuleOrientationBase : public UParticleModule
 {
-	GENERATED_BODY()
-public:
-	UParticleModuleOrientationBase(const FObjectInitializer& ObjectInitializer = FObjectInitializer::Get());
+	GENERATED_UCLASS_BODY()
 
 };
 
diff --git a/Engine/Source/Runtime/Engine/Classes/Particles/Parameter/ParticleModuleParameterBase.h b/Engine/Source/Runtime/Engine/Classes/Particles/Parameter/ParticleModuleParameterBase.h
--- a/Engine/Source/Runtime/Engine/Classes/Particles/Parameter/ParticleModuleParameterBase.h
+++ b/Engine/Source/Runtime/Engine/Classes/Particles/Parameter/ParticleModuleParameterBase.h
@@ -8,9 +8,7 @@
 UCLASS(editinlinenew, hidecategories=Object, abstract, meta=(DisplayName = "Parameter"))
 class UParticleModuleParameterBase : public UParticleModule
 {
-	GENERATED_BODY()
-public:
-	UParticleModuleParameterBase(const FObjectInitializer& ObjectInitializer = FObjectInitializer::Get());
+	GENERATED_UCLASS_BODY()
 	virtual bool CanTickInAnyThread() override
 	{
 		return false;
diff --git a/Engine/Source/Runtime/Engine/Classes/Particles/Parameter/ParticleModuleParameterDynamic_Seeded.h b/Engine/Source/Runtime/Engine/Classes/Particles/Parameter/ParticleModuleParameterDynamic_Seeded.h
--- a/Engine/Source/Runtime/Engine/Classes/Particles/Parameter/ParticleModuleParameterDynamic_Seeded.h
+++ b/Engine/Source/Runtime/Engine/Classes/Particles/Parameter/ParticleModuleParameterDynamic_Seeded.h
@@ -8,9 +8,7 @@
 UCLASS(editinlinenew, hidecategories=Object, meta=(DisplayName = "Dynamic (Seed)"))
 class UParticleModuleParameterDynamic_Seeded : public UParticleModuleParameterDynamic
 {
-	GENERATED_BODY()
-public:
-	UParticleModuleParameterDynamic_Seeded(const FObjectInitializer& ObjectInitializer = FObjectInitializer::Get());
+	GENERATED_UCLASS_BODY()
 
 	/** The random seed(s) to use for looking up values in StartLocation */
 	UPROPERTY(EditAnywhere, Category=RandomSeed)
diff --git a/Engine/Source/Runtime/Engine/Classes/Particles/ParticleEventManager.h b/Engine/Source/Runtime/Engine/Classes/Particles/ParticleEventManager.h
--- a/Engine/Source/Runtime/Engine/Classes/Particles/ParticleEventManager.h
+++ b/Engine/Source/Runtime/Engine/Classes/Particles/ParticleEventManager.h
@@ -16,9 +16,7 @@ class UParticleSystemComponent;
 UCLASS(config=Game, NotBlueprintable, notplaceable)
 class AParticleEventManager : public AActor
 {
-	GENERATED_BODY()
-public:
-	AParticleEventManager(const FObjectInitializer& ObjectInitializer = FObjectInitializer::Get());
+	GENERATED_UCLASS_BODY()
 
 	virtual void HandleParticleSpawnEvents( UParticleSystemComponent* Component, const TArray<FParticleEventSpawnData>& SpawnEvents );
 	virtual void HandleParticleDeathEvents( UParticleSystemComponent* Component, const TArray<FParticleEventDeathData>& DeathEvents );
diff --git a/Engine/Source/Runtime/Engine/Classes/Particles/ParticleLODLevel.h b/Engine/Source/Runtime/Engine/Classes/Particles/ParticleLODLevel.h
--- a/Engine/Source/Runtime/Engine/Classes/Particles/ParticleLODLevel.h
+++ b/Engine/Source/Runtime/Engine/Classes/Particles/ParticleLODLevel.h
@@ -13,9 +13,7 @@ class UInterpCurveEdSetup;
 UCLASS(collapsecategories, hidecategories=Object, editinlinenew, MinimalAPI)
 class UParticleLODLevel : public UObject
 {
-	GENERATED_BODY()
-public:
-	ENGINE_API UParticleLODLevel(const FObjectInitializer& ObjectInitializer = FObjectInitializer::Get());
+	GENERATED_UCLASS_BODY()
 
 	/** The index value of the LOD level												*/
 	UPROPERTY()
