Back out changelist 2481333

diff --git a/Engine/Source/Runtime/Engine/Classes/Particles/ParticleModule.h b/Engine/Source/Runtime/Engine/Classes/Particles/ParticleModule.h
--- a/Engine/Source/Runtime/Engine/Classes/Particles/ParticleModule.h
+++ b/Engine/Source/Runtime/Engine/Classes/Particles/ParticleModule.h
@@ -56,7 +56,7 @@ enum EParticleSourceSelectionMethod
 USTRUCT()
 struct FParticleCurvePair
 {
-	GENERATED_BODY()
+	GENERATED_USTRUCT_BODY()
 
 	UPROPERTY()
 	FString CurveName;
@@ -75,7 +75,7 @@ struct FParticleCurvePair
 USTRUCT()
 struct FParticleRandomSeedInfo
 {
-	GENERATED_BODY()
+	GENERATED_USTRUCT_BODY()
 
 	/** The name to expose to the placed instances for setting this seed */
 	UPROPERTY(EditAnywhere, Category=ParticleRandomSeedInfo)
@@ -136,9 +136,7 @@ struct FParticleRandomSeedInfo
 UCLASS(editinlinenew, hidecategories=Object, abstract, MinimalAPI)
 class UParticleModule : public UObject
 {
-	GENERATED_BODY()
-public:
-	ENGINE_API UParticleModule(const FObjectInitializer& ObjectInitializer = FObjectInitializer::Get());
+	GENERATED_UCLASS_BODY()
 
 	/** If true, the module performs operations on particles during Spawning		*/
 	UPROPERTY()
diff --git a/Engine/Source/Runtime/Engine/Classes/Particles/ParticleModuleRequired.h b/Engine/Source/Runtime/Engine/Classes/Particles/ParticleModuleRequired.h
--- a/Engine/Source/Runtime/Engine/Classes/Particles/ParticleModuleRequired.h
+++ b/Engine/Source/Runtime/Engine/Classes/Particles/ParticleModuleRequired.h
@@ -77,9 +77,7 @@ enum EEmitterNormalsMode
 UCLASS(editinlinenew, hidecategories=(Object, Cascade), meta=(DisplayName = "Required"), MinimalAPI)
 class UParticleModuleRequired : public UParticleModule
 {
-	GENERATED_BODY()
-public:
-	ENGINE_API UParticleModuleRequired(const FObjectInitializer& ObjectInitializer = FObjectInitializer::Get());
+	GENERATED_UCLASS_BODY()
 
 	//
 	// General.
diff --git a/Engine/Source/Runtime/Engine/Classes/Particles/ParticleSpriteEmitter.h b/Engine/Source/Runtime/Engine/Classes/Particles/ParticleSpriteEmitter.h
--- a/Engine/Source/Runtime/Engine/Classes/Particles/ParticleSpriteEmitter.h
+++ b/Engine/Source/Runtime/Engine/Classes/Particles/ParticleSpriteEmitter.h
@@ -20,9 +20,7 @@ enum EParticleScreenAlignment
 UCLASS(collapsecategories, hidecategories=Object, editinlinenew, MinimalAPI)
 class UParticleSpriteEmitter : public UParticleEmitter
 {
-	GENERATED_BODY()
-public:
-	ENGINE_API UParticleSpriteEmitter(const FObjectInitializer& ObjectInitializer = FObjectInitializer::Get());
+	GENERATED_UCLASS_BODY()
 
 
 	// Begin UObject Interface
diff --git a/Engine/Source/Runtime/Engine/Classes/Particles/ParticleSystem.h b/Engine/Source/Runtime/Engine/Classes/Particles/ParticleSystem.h
--- a/Engine/Source/Runtime/Engine/Classes/Particles/ParticleSystem.h
+++ b/Engine/Source/Runtime/Engine/Classes/Particles/ParticleSystem.h
@@ -47,7 +47,7 @@ enum EParticleSystemOcclusionBoundsMethod
 USTRUCT()
 struct FParticleSystemLOD
 {
-	GENERATED_BODY()
+	GENERATED_USTRUCT_BODY()
 
 	FParticleSystemLOD()
 	{
@@ -67,7 +67,7 @@ struct FParticleSystemLOD
 USTRUCT()
 struct FLODSoloTrack
 {
-	GENERATED_BODY()
+	GENERATED_USTRUCT_BODY()
 
 	UPROPERTY(transient)
 	TArray<uint8> SoloEnableSetting;
@@ -77,7 +77,7 @@ struct FLODSoloTrack
 USTRUCT()
 struct FNamedEmitterMaterial
 {
-	GENERATED_BODY()
+	GENERATED_USTRUCT_BODY()
 
 	FNamedEmitterMaterial()
 	: Name(NAME_None)
@@ -100,9 +100,7 @@ struct FNamedEmitterMaterial
 UCLASS(hidecategories=Object, MinimalAPI, BlueprintType)
 class UParticleSystem : public UObject
 {
-	GENERATED_BODY()
-public:
-	ENGINE_API UParticleSystem(const FObjectInitializer& ObjectInitializer = FObjectInitializer::Get());
+	GENERATED_UCLASS_BODY()
 
 	UPROPERTY(EditAnywhere, Category=ParticleSystem, AssetRegistrySearchable)
 	TEnumAsByte<enum EParticleSystemUpdateMode> SystemUpdateMode;
diff --git a/Engine/Source/Runtime/Engine/Classes/Particles/ParticleSystemComponent.h b/Engine/Source/Runtime/Engine/Classes/Particles/ParticleSystemComponent.h
--- a/Engine/Source/Runtime/Engine/Classes/Particles/ParticleSystemComponent.h
+++ b/Engine/Source/Runtime/Engine/Classes/Particles/ParticleSystemComponent.h
@@ -78,7 +78,7 @@ DECLARE_DYNAMIC_MULTICAST_DELEGATE_OneParam( FOnSystemFinished, class UParticleS
 USTRUCT(BlueprintType)
 struct FParticleSysParam
 {
-	GENERATED_BODY()
+	GENERATED_USTRUCT_BODY()
 
 	/** The name of the parameter */
 	UPROPERTY(EditAnywhere, BlueprintReadWrite, Category=ParticleSysParam)
@@ -250,9 +250,7 @@ struct FParticleEventKismetData : public FParticleEventData
 UCLASS(ClassGroup=(Rendering, Common), hidecategories=Object, hidecategories=Physics, hidecategories=Collision, showcategories=Trigger, editinlinenew, meta=(BlueprintSpawnableComponent))
 class ENGINE_API UParticleSystemComponent : public UPrimitiveComponent
 {
-	GENERATED_BODY()
-public:
-	UParticleSystemComponent(const FObjectInitializer& ObjectInitializer = FObjectInitializer::Get());
+	GENERATED_UCLASS_BODY()
 
 	UPROPERTY(EditAnywhere, BlueprintReadOnly, Category=Particles)
 	class UParticleSystem* Template;
diff --git a/Engine/Source/Runtime/Engine/Classes/Particles/ParticleSystemReplay.h b/Engine/Source/Runtime/Engine/Classes/Particles/ParticleSystemReplay.h
--- a/Engine/Source/Runtime/Engine/Classes/Particles/ParticleSystemReplay.h
+++ b/Engine/Source/Runtime/Engine/Classes/Particles/ParticleSystemReplay.h
@@ -7,7 +7,7 @@
 USTRUCT()
 struct FParticleEmitterReplayFrame
 {
-	GENERATED_BODY()
+	GENERATED_USTRUCT_BODY()
 
 	/** Type of emitter (EDynamicEmitterType) */
 	int32 EmitterType;
@@ -46,7 +46,7 @@ struct FParticleEmitterReplayFrame
 USTRUCT()
 struct FParticleSystemReplayFrame
 {
-	GENERATED_BODY()
+	GENERATED_USTRUCT_BODY()
 
 	/** Emitter frame state data */
 	TArray<struct FParticleEmitterReplayFrame> Emitters;
@@ -61,9 +61,7 @@ struct FParticleSystemReplayFrame
 UCLASS(hidecategories=Object, AutoExpandCategories=ParticleSystemReplay)
 class UParticleSystemReplay : public UObject
 {
-	GENERATED_BODY()
-public:
-	UParticleSystemReplay(const FObjectInitializer& ObjectInitializer = FObjectInitializer::Get());
+	GENERATED_UCLASS_BODY()
 
 	/** Unique ID number for this replay clip */
 	UPROPERTY(EditAnywhere, transient, duplicatetransient, Category=ParticleSystemReplay)
diff --git a/Engine/Source/Runtime/Engine/Classes/Particles/Rotation/ParticleModuleMeshRotation.h b/Engine/Source/Runtime/Engine/Classes/Particles/Rotation/ParticleModuleMeshRotation.h
--- a/Engine/Source/Runtime/Engine/Classes/Particles/Rotation/ParticleModuleMeshRotation.h
+++ b/Engine/Source/Runtime/Engine/Classes/Particles/Rotation/ParticleModuleMeshRotation.h
@@ -8,9 +8,7 @@
 UCLASS(editinlinenew, hidecategories=Object, MinimalAPI, meta=(DisplayName = "Init Mesh Rotation"))
 class UParticleModuleMeshRotation : public UParticleModuleRotationBase
 {
-	GENERATED_BODY()
-public:
-	ENGINE_API UParticleModuleMeshRotation(const FObjectInitializer& ObjectInitializer = FObjectInitializer::Get());
+	GENERATED_UCLASS_BODY()
 
 	/**
 	 *	Initial rotation in ROTATIONS PER SECOND (1 = 360 degrees).
diff --git a/Engine/Source/Runtime/Engine/Classes/Particles/Rotation/ParticleModuleMeshRotation_Seeded.h b/Engine/Source/Runtime/Engine/Classes/Particles/Rotation/ParticleModuleMeshRotation_Seeded.h
--- a/Engine/Source/Runtime/Engine/Classes/Particles/Rotation/ParticleModuleMeshRotation_Seeded.h
+++ b/Engine/Source/Runtime/Engine/Classes/Particles/Rotation/ParticleModuleMeshRotation_Seeded.h
@@ -8,9 +8,7 @@
 UCLASS(editinlinenew, hidecategories=Object, meta=(DisplayName = "Init Mesh Rot (Seed)"))
 class UParticleModuleMeshRotation_Seeded : public UParticleModuleMeshRotation
 {
-	GENERATED_BODY()
-public:
-	UParticleModuleMeshRotation_Seeded(const FObjectInitializer& ObjectInitializer = FObjectInitializer::Get());
+	GENERATED_UCLASS_BODY()
 
 	/** The random seed(s) to use for looking up values in StartLocation */
 	UPROPERTY(EditAnywhere, Category=RandomSeed)
diff --git a/Engine/Source/Runtime/Engine/Classes/Particles/Rotation/ParticleModuleRotation.h b/Engine/Source/Runtime/Engine/Classes/Particles/Rotation/ParticleModuleRotation.h
--- a/Engine/Source/Runtime/Engine/Classes/Particles/Rotation/ParticleModuleRotation.h
+++ b/Engine/Source/Runtime/Engine/Classes/Particles/Rotation/ParticleModuleRotation.h
@@ -8,9 +8,7 @@
 UCLASS(editinlinenew, hidecategories=Object, MinimalAPI, meta=(DisplayName = "Initial Rotation"))
 class UParticleModuleRotation : public UParticleModuleRotationBase
 {
-	GENERATED_BODY()
-public:
-	ENGINE_API UParticleModuleRotation(const FObjectInitializer& ObjectInitializer = FObjectInitializer::Get());
+	GENERATED_UCLASS_BODY()
 
 	/**
 	 *	Initial rotation of the particle (1 = 360 degrees).
diff --git a/Engine/Source/Runtime/Engine/Classes/Particles/Rotation/ParticleModuleRotationBase.h b/Engine/Source/Runtime/Engine/Classes/Particles/Rotation/ParticleModuleRotationBase.h
--- a/Engine/Source/Runtime/Engine/Classes/Particles/Rotation/ParticleModuleRotationBase.h
+++ b/Engine/Source/Runtime/Engine/Classes/Particles/Rotation/ParticleModuleRotationBase.h
@@ -8,9 +8,7 @@
 UCLASS(editinlinenew, hidecategories=Object, abstract, meta=(DisplayName = "Rotation"))
 class UParticleModuleRotationBase : public UParticleModule
 {
-	GENERATED_BODY()
-public:
-	UParticleModuleRotationBase(const FObjectInitializer& ObjectInitializer = FObjectInitializer::Get());
+	GENERATED_UCLASS_BODY()
 
 };
 
diff --git a/Engine/Source/Runtime/Engine/Classes/Particles/Rotation/ParticleModuleRotation_Seeded.h b/Engine/Source/Runtime/Engine/Classes/Particles/Rotation/ParticleModuleRotation_Seeded.h
--- a/Engine/Source/Runtime/Engine/Classes/Particles/Rotation/ParticleModuleRotation_Seeded.h
+++ b/Engine/Source/Runtime/Engine/Classes/Particles/Rotation/ParticleModuleRotation_Seeded.h
@@ -8,9 +8,7 @@
 UCLASS(editinlinenew, hidecategories=Object, meta=(DisplayName = "Init Rotation (Seed)"))
 class UParticleModuleRotation_Seeded : public UParticleModuleRotation
 {
-	GENERATED_BODY()
-public:
-	UParticleModuleRotation_Seeded(const FObjectInitializer& ObjectInitializer = FObjectInitializer::Get());
+	GENERATED_UCLASS_BODY()
 
 	/** The random seed(s) to use for looking up values in StartLocation */
 	UPROPERTY(EditAnywhere, Category=RandomSeed)
diff --git a/Engine/Source/Runtime/Engine/Classes/Particles/RotationRate/ParticleModuleMeshRotationRate.h b/Engine/Source/Runtime/Engine/Classes/Particles/RotationRate/ParticleModuleMeshRotationRate.h
--- a/Engine/Source/Runtime/Engine/Classes/Particles/RotationRate/ParticleModuleMeshRotationRate.h
+++ b/Engine/Source/Runtime/Engine/Classes/Particles/RotationRate/ParticleModuleMeshRotationRate.h
@@ -8,9 +8,7 @@
 UCLASS(editinlinenew, hidecategories=Object, MinimalAPI, meta=(DisplayName = "Init Mesh Rotation Rate"))
 class UParticleModuleMeshRotationRate : public UParticleModuleRotationRateBase
 {
-	GENERATED_BODY()
-public:
-	ENGINE_API UParticleModuleMeshRotationRate(const FObjectInitializer& ObjectInitializer = FObjectInitializer::Get());
+	GENERATED_UCLASS_BODY()
 
 	/**
 	 *	Initial rotation rate, in rotations per second.
diff --git a/Engine/Source/Runtime/Engine/Classes/Particles/RotationRate/ParticleModuleMeshRotationRateMultiplyLife.h b/Engine/Source/Runtime/Engine/Classes/Particles/RotationRate/ParticleModuleMeshRotationRateMultiplyLife.h
--- a/Engine/Source/Runtime/Engine/Classes/Particles/RotationRate/ParticleModuleMeshRotationRateMultiplyLife.h
+++ b/Engine/Source/Runtime/Engine/Classes/Particles/RotationRate/ParticleModuleMeshRotationRateMultiplyLife.h
@@ -8,9 +8,7 @@
 UCLASS(editinlinenew, hidecategories=Object, meta=(DisplayName = "Mesh Rotation Rate * Life"))
 class UParticleModuleMeshRotationRateMultiplyLife : public UParticleModuleRotationRateBase
 {
-	GENERATED_BODY()
-public:
-	UParticleModuleMeshRotationRateMultiplyLife(const FObjectInitializer& ObjectInitializer = FObjectInitializer::Get());
+	GENERATED_UCLASS_BODY()
 
 	/**
 	 *	The scale factor that should be applied to the rotation rate.
diff --git a/Engine/Source/Runtime/Engine/Classes/Particles/RotationRate/ParticleModuleMeshRotationRateOverLife.h b/Engine/Source/Runtime/Engine/Classes/Particles/RotationRate/ParticleModuleMeshRotationRateOverLife.h
--- a/Engine/Source/Runtime/Engine/Classes/Particles/RotationRate/ParticleModuleMeshRotationRateOverLife.h
+++ b/Engine/Source/Runtime/Engine/Classes/Particles/RotationRate/ParticleModuleMeshRotationRateOverLife.h
@@ -8,9 +8,7 @@
 UCLASS(editinlinenew, hidecategories=Object, meta=(DisplayName = "Mesh Rotation Rate over Life"))
 class UParticleModuleMeshRotationRateOverLife : public UParticleModuleRotationRateBase
 {
-	GENERATED_BODY()
-public:
-	UParticleModuleMeshRotationRateOverLife(const FObjectInitializer& ObjectInitializer = FObjectInitializer::Get());
+	GENERATED_UCLASS_BODY()
 
 	/**
 	 *	The rotation rate desired.
diff --git a/Engine/Source/Runtime/Engine/Classes/Particles/RotationRate/ParticleModuleMeshRotationRate_Seeded.h b/Engine/Source/Runtime/Engine/Classes/Particles/RotationRate/ParticleModuleMeshRotationRate_Seeded.h
--- a/Engine/Source/Runtime/Engine/Classes/Particles/RotationRate/ParticleModuleMeshRotationRate_Seeded.h
+++ b/Engine/Source/Runtime/Engine/Classes/Particles/RotationRate/ParticleModuleMeshRotationRate_Seeded.h
@@ -8,9 +8,7 @@
 UCLASS(editinlinenew, hidecategories=Object, meta=(DisplayName = "Init Mesh Rotation Rate (Seed)"))
 class UParticleModuleMeshRotationRate_Seeded : public UParticleModuleMeshRotationRate
 {
-	GENERATED_BODY()
-public:
-	UParticleModuleMeshRotationRate_Seeded(const FObjectInitializer& ObjectInitializer = FObjectInitializer::Get());
+	GENERATED_UCLASS_BODY()
 
 	/** The random seed(s) to use for looking up values in StartLocation */
 	UPROPERTY(EditAnywhere, Category=RandomSeed)
diff --git a/Engine/Source/Runtime/Engine/Classes/Particles/RotationRate/ParticleModuleRotationRate.h b/Engine/Source/Runtime/Engine/Classes/Particles/RotationRate/ParticleModuleRotationRate.h
--- a/Engine/Source/Runtime/Engine/Classes/Particles/RotationRate/ParticleModuleRotationRate.h
+++ b/Engine/Source/Runtime/Engine/Classes/Particles/RotationRate/ParticleModuleRotationRate.h
@@ -8,9 +8,7 @@
 UCLASS(editinlinenew, hidecategories=Object, MinimalAPI, meta=(DisplayName = "Initial Rotation Rate"))
 class UParticleModuleRotationRate : public UParticleModuleRotationRateBase
 {
-	GENERATED_BODY()
-public:
-	ENGINE_API UParticleModuleRotationRate(const FObjectInitializer& ObjectInitializer = FObjectInitializer::Get());
+	GENERATED_UCLASS_BODY()
 
 	/** 
 	 *	Initial rotation rate, in rotations per second.
diff --git a/Engine/Source/Runtime/Engine/Classes/Particles/RotationRate/ParticleModuleRotationRateBase.h b/Engine/Source/Runtime/Engine/Classes/Particles/RotationRate/ParticleModuleRotationRateBase.h
--- a/Engine/Source/Runtime/Engine/Classes/Particles/RotationRate/ParticleModuleRotationRateBase.h
+++ b/Engine/Source/Runtime/Engine/Classes/Particles/RotationRate/ParticleModuleRotationRateBase.h
@@ -8,9 +8,7 @@
 UCLASS(editinlinenew, hidecategories=Object, abstract, meta=(DisplayName = "Rotation Rate"))
 class UParticleModuleRotationRateBase : public UParticleModule
 {
-	GENERATED_BODY()
-public:
-	UParticleModuleRotationRateBase(const FObjectInitializer& ObjectInitializer = FObjectInitializer::Get());
+	GENERATED_UCLASS_BODY()
 
 };
 
diff --git a/Engine/Source/Runtime/Engine/Classes/Particles/RotationRate/ParticleModuleRotationRateMultiplyLife.h b/Engine/Source/Runtime/Engine/Classes/Particles/RotationRate/ParticleModuleRotationRateMultiplyLife.h
--- a/Engine/Source/Runtime/Engine/Classes/Particles/RotationRate/ParticleModuleRotationRateMultiplyLife.h
+++ b/Engine/Source/Runtime/Engine/Classes/Particles/RotationRate/ParticleModuleRotationRateMultiplyLife.h
@@ -8,9 +8,7 @@
 UCLASS(editinlinenew, hidecategories=Object, MinimalAPI, meta=(DisplayName = "Rotation Rate * Life"))
 class UParticleModuleRotationRateMultiplyLife : public UParticleModuleRotationRateBase
 {
-	GENERATED_BODY()
-public:
-	ENGINE_API UParticleModuleRotationRateMultiplyLife(const FObjectInitializer& ObjectInitializer = FObjectInitializer::Get());
+	GENERATED_UCLASS_BODY()
 
 	/**
 	 *	The scale factor that should be applied to the rotation rate.
diff --git a/Engine/Source/Runtime/Engine/Classes/Particles/RotationRate/ParticleModuleRotationRate_Seeded.h b/Engine/Source/Runtime/Engine/Classes/Particles/RotationRate/ParticleModuleRotationRate_Seeded.h
--- a/Engine/Source/Runtime/Engine/Classes/Particles/RotationRate/ParticleModuleRotationRate_Seeded.h
+++ b/Engine/Source/Runtime/Engine/Classes/Particles/RotationRate/ParticleModuleRotationRate_Seeded.h
@@ -8,9 +8,7 @@
 UCLASS(editinlinenew, hidecategories=Object, meta=(DisplayName = "Init Rotation Rate (Seed)"))
 class UParticleModuleRotationRate_Seeded : public UParticleModuleRotationRate
 {
-	GENERATED_BODY()
-public:
-	UParticleModuleRotationRate_Seeded(const FObjectInitializer& ObjectInitializer = FObjectInitializer::Get());
+	GENERATED_UCLASS_BODY()
 
 	/** The random seed(s) to use for looking up values in StartLocation */
 	UPROPERTY(EditAnywhere, Category=RandomSeed)
diff --git a/Engine/Source/Runtime/Engine/Classes/Particles/Size/ParticleModuleSize.h b/Engine/Source/Runtime/Engine/Classes/Particles/Size/ParticleModuleSize.h
--- a/Engine/Source/Runtime/Engine/Classes/Particles/Size/ParticleModuleSize.h
+++ b/Engine/Source/Runtime/Engine/Classes/Particles/Size/ParticleModuleSize.h
@@ -8,9 +8,7 @@
 UCLASS(editinlinenew, hidecategories=Object, MinimalAPI, meta=(DisplayName = "Initial Size"))
 class UParticleModuleSize : public UParticleModuleSizeBase
 {
-	GENERATED_BODY()
-public:
-	ENGINE_API UParticleModuleSize(const FObjectInitializer& ObjectInitializer = FObjectInitializer::Get());
+	GENERATED_UCLASS_BODY()
 
 	/**
 	 *	The initial size that should be used for a particle.
diff --git a/Engine/Source/Runtime/Engine/Classes/Particles/Size/ParticleModuleSizeBase.h b/Engine/Source/Runtime/Engine/Classes/Particles/Size/ParticleModuleSizeBase.h
--- a/Engine/Source/Runtime/Engine/Classes/Particles/Size/ParticleModuleSizeBase.h
+++ b/Engine/Source/Runtime/Engine/Classes/Particles/Size/ParticleModuleSizeBase.h
@@ -8,9 +8,7 @@
 UCLASS(editinlinenew, hidecategories=Object, abstract, meta=(DisplayName = "Size"))
 class UParticleModuleSizeBase : public UParticleModule
 {
-	GENERATED_BODY()
-public:
-	UParticleModuleSizeBase(const FObjectInitializer& ObjectInitializer = FObjectInitializer::Get());
+	GENERATED_UCLASS_BODY()
 
 };
 
diff --git a/Engine/Source/Runtime/Engine/Classes/Particles/Size/ParticleModuleSizeMultiplyLife.h b/Engine/Source/Runtime/Engine/Classes/Particles/Size/ParticleModuleSizeMultiplyLife.h
--- a/Engine/Source/Runtime/Engine/Classes/Particles/Size/ParticleModuleSizeMultiplyLife.h
+++ b/Engine/Source/Runtime/Engine/Classes/Particles/Size/ParticleModuleSizeMultiplyLife.h
@@ -8,9 +8,7 @@
 UCLASS(editinlinenew, hidecategories=Object, MinimalAPI, meta=(DisplayName = "Size By Life"))
 class UParticleModuleSizeMultiplyLife : public UParticleModuleSizeBase
 {
-	GENERATED_BODY()
-public:
-	ENGINE_API UParticleModuleSizeMultiplyLife(const FObjectInitializer& ObjectInitializer = FObjectInitializer::Get());
+	GENERATED_UCLASS_BODY()
 
 	/**
 	 *	The scale factor for the size that should be used for a particle.
diff --git a/Engine/Source/Runtime/Engine/Classes/Particles/Size/ParticleModuleSizeScale.h b/Engine/Source/Runtime/Engine/Classes/Particles/Size/ParticleModuleSizeScale.h
--- a/Engine/Source/Runtime/Engine/Classes/Particles/Size/ParticleModuleSizeScale.h
+++ b/Engine/Source/Runtime/Engine/Classes/Particles/Size/ParticleModuleSizeScale.h
@@ -8,9 +8,7 @@
 UCLASS(editinlinenew, hidecategories=Object, MinimalAPI, meta=(DisplayName = "Size Scale"))
 class UParticleModuleSizeScale : public UParticleModuleSizeBase
 {
-	GENERATED_BODY()
-public:
-	ENGINE_API UParticleModuleSizeScale(const FObjectInitializer& ObjectInitializer = FObjectInitializer::Get());
+	GENERATED_UCLASS_BODY()
 
 	/**
 	 *	The amount the BaseSize should be scaled before being used as the size of the particle. 
diff --git a/Engine/Source/Runtime/Engine/Classes/Particles/Size/ParticleModuleSizeScaleBySpeed.h b/Engine/Source/Runtime/Engine/Classes/Particles/Size/ParticleModuleSizeScaleBySpeed.h
--- a/Engine/Source/Runtime/Engine/Classes/Particles/Size/ParticleModuleSizeScaleBySpeed.h
+++ b/Engine/Source/Runtime/Engine/Classes/Particles/Size/ParticleModuleSizeScaleBySpeed.h
@@ -11,9 +11,7 @@
 UCLASS(editinlinenew, hidecategories=Object, meta=(DisplayName = "Size By Speed"))
 class UParticleModuleSizeScaleBySpeed : public UParticleModuleSizeBase
 {
-	GENERATED_BODY()
-public:
-	UParticleModuleSizeScaleBySpeed(const FObjectInitializer& ObjectInitializer = FObjectInitializer::Get());
+	GENERATED_UCLASS_BODY()
 
 	/** By how much speed affects the size of the particle in each dimension. */
 	UPROPERTY(EditAnywhere, Category=ParticleModuleSizeScaleBySpeed)
diff --git a/Engine/Source/Runtime/Engine/Classes/Particles/Size/ParticleModuleSize_Seeded.h b/Engine/Source/Runtime/Engine/Classes/Particles/Size/ParticleModuleSize_Seeded.h
--- a/Engine/Source/Runtime/Engine/Classes/Particles/Size/ParticleModuleSize_Seeded.h
+++ b/Engine/Source/Runtime/Engine/Classes/Particles/Size/ParticleModuleSize_Seeded.h
@@ -8,9 +8,7 @@
 UCLASS(editinlinenew, hidecategories=Object, meta=(DisplayName = "Initial Size (Seed)"))
 class UParticleModuleSize_Seeded : public UParticleModuleSize
 {
-	GENERATED_BODY()
-public:
-	UParticleModuleSize_Seeded(const FObjectInitializer& ObjectInitializer = FObjectInitializer::Get());
+	GENERATED_UCLASS_BODY()
 
 	/** The random seed(s) to use for looking up values in StartLocation */
 	UPROPERTY(EditAnywhere, Category=RandomSeed)
diff --git a/Engine/Source/Runtime/Engine/Classes/Particles/Spawn/ParticleModuleSpawn.h b/Engine/Source/Runtime/Engine/Classes/Particles/Spawn/ParticleModuleSpawn.h
--- a/Engine/Source/Runtime/Engine/Classes/Particles/Spawn/ParticleModuleSpawn.h
+++ b/Engine/Source/Runtime/Engine/Classes/Particles/Spawn/ParticleModuleSpawn.h
@@ -11,9 +11,7 @@ class UParticleLODLevel;
 UCLASS(editinlinenew, hidecategories=Object, hidecategories=ParticleModuleSpawnBase, MinimalAPI, meta=(DisplayName = "Spawn"))
 class UParticleModuleSpawn : public UParticleModuleSpawnBase
 {
-	GENERATED_BODY()
-public:
-	ENGINE_API UParticleModuleSpawn(const FObjectInitializer& ObjectInitializer = FObjectInitializer::Get());
+	GENERATED_UCLASS_BODY()
 
 	/** The rate at which to spawn particles. */
 	UPROPERTY(EditAnywhere, Category=Spawn)
diff --git a/Engine/Source/Runtime/Engine/Classes/Particles/Spawn/ParticleModuleSpawnBase.h b/Engine/Source/Runtime/Engine/Classes/Particles/Spawn/ParticleModuleSpawnBase.h
--- a/Engine/Source/Runtime/Engine/Classes/Particles/Spawn/ParticleModuleSpawnBase.h
+++ b/Engine/Source/Runtime/Engine/Classes/Particles/Spawn/ParticleModuleSpawnBase.h
@@ -10,9 +10,7 @@ struct FParticleEmitterInstance;
 UCLASS(editinlinenew, hidecategories=Object, abstract, meta=(DisplayName = "Spawn"))
 class UParticleModuleSpawnBase : public UParticleModule
 {
-	GENERATED_BODY()
-public:
-	UParticleModuleSpawnBase(const FObjectInitializer& ObjectInitializer = FObjectInitializer::Get());
+	GENERATED_UCLASS_BODY()
 
 	/** 
 	 *	If true, the SpawnRate of the SpawnModule of the emitter will be processed.
diff --git a/Engine/Source/Runtime/Engine/Classes/Particles/SubUV/ParticleModuleSubUV.h b/Engine/Source/Runtime/Engine/Classes/Particles/SubUV/ParticleModuleSubUV.h
--- a/Engine/Source/Runtime/Engine/Classes/Particles/SubUV/ParticleModuleSubUV.h
+++ b/Engine/Source/Runtime/Engine/Classes/Particles/SubUV/ParticleModuleSubUV.h
@@ -9,9 +9,7 @@
 UCLASS(editinlinenew, hidecategories=Object, MinimalAPI, meta=(DisplayName = "SubImage Index"))
 class UParticleModuleSubUV : public UParticleModuleSubUVBase
 {
-	GENERATED_BODY()
-public:
-	ENGINE_API UParticleModuleSubUV(const FObjectInitializer& ObjectInitializer = FObjectInitializer::Get());
+	GENERATED_UCLASS_BODY()
 
 	/**
 	 *	The index of the sub-image that should be used for the particle.
diff --git a/Engine/Source/Runtime/Engine/Classes/Particles/SubUV/ParticleModuleSubUVBase.h b/Engine/Source/Runtime/Engine/Classes/Particles/SubUV/ParticleModuleSubUVBase.h
--- a/Engine/Source/Runtime/Engine/Classes/Particles/SubUV/ParticleModuleSubUVBase.h
+++ b/Engine/Source/Runtime/Engine/Classes/Particles/SubUV/ParticleModuleSubUVBase.h
@@ -8,9 +8,7 @@
 UCLASS(editinlinenew, hidecategories=Object, abstract, meta=(DisplayName = "SubUV"))
 class UParticleModuleSubUVBase : public UParticleModule
 {
-	GENERATED_BODY()
-public:
-	UParticleModuleSubUVBase(const FObjectInitializer& ObjectInitializer = FObjectInitializer::Get());
+	GENERATED_UCLASS_BODY()
 
 };
 
diff --git a/Engine/Source/Runtime/Engine/Classes/Particles/Trail/ParticleModuleTrailBase.h b/Engine/Source/Runtime/Engine/Classes/Particles/Trail/ParticleModuleTrailBase.h
--- a/Engine/Source/Runtime/Engine/Classes/Particles/Trail/ParticleModuleTrailBase.h
+++ b/Engine/Source/Runtime/Engine/Classes/Particles/Trail/ParticleModuleTrailBase.h
@@ -13,9 +13,7 @@
 UCLASS(editinlinenew, hidecategories=Object, abstract, meta=(DisplayName = "Trail"))
 class UParticleModuleTrailBase : public UParticleModule
 {
-	GENERATED_BODY()
-public:
-	UParticleModuleTrailBase(const FObjectInitializer& ObjectInitializer = FObjectInitializer::Get());
+	GENERATED_UCLASS_BODY()
 
 
 	// Begin UParticleModule Interface
diff --git a/Engine/Source/Runtime/Engine/Classes/Particles/Trail/ParticleModuleTrailSource.h b/Engine/Source/Runtime/Engine/Classes/Particles/Trail/ParticleModuleTrailSource.h
--- a/Engine/Source/Runtime/Engine/Classes/Particles/Trail/ParticleModuleTrailSource.h
+++ b/Engine/Source/Runtime/Engine/Classes/Particles/Trail/ParticleModuleTrailSource.h
@@ -32,9 +32,7 @@ enum ETrail2SourceMethod
 UCLASS(editinlinenew, hidecategories=Object, MinimalAPI, meta=(DisplayName = "Source"))
 class UParticleModuleTrailSource : public UParticleModuleTrailBase
 {
-	GENERATED_BODY()
-public:
-	ENGINE_API UParticleModuleTrailSource(const FObjectInitializer& ObjectInitializer = FObjectInitializer::Get());
+	GENERATED_UCLASS_BODY()
 
 	/** The source method for the trail. */
 	UPROPERTY(EditAnywhere, Category=Source)
diff --git a/Engine/Source/Runtime/Engine/Classes/Particles/TypeData/ParticleModuleTypeDataAnimTrail.h b/Engine/Source/Runtime/Engine/Classes/Particles/TypeData/ParticleModuleTypeDataAnimTrail.h
--- a/Engine/Source/Runtime/Engine/Classes/Particles/TypeData/ParticleModuleTypeDataAnimTrail.h
+++ b/Engine/Source/Runtime/Engine/Classes/Particles/TypeData/ParticleModuleTypeDataAnimTrail.h
@@ -13,9 +13,7 @@
 UCLASS(editinlinenew, hidecategories=Object, meta=(DisplayName = "AnimTrail Data"))
 class UParticleModuleTypeDataAnimTrail : public UParticleModuleTypeDataBase
 {
-	GENERATED_BODY()
-public:
-	UParticleModuleTypeDataAnimTrail(const FObjectInitializer& ObjectInitializer = FObjectInitializer::Get());
+	GENERATED_UCLASS_BODY()
 
 	//*****************************************************************************
 	// General Trail Variables
diff --git a/Engine/Source/Runtime/Engine/Classes/Particles/TypeData/ParticleModuleTypeDataBase.h b/Engine/Source/Runtime/Engine/Classes/Particles/TypeData/ParticleModuleTypeDataBase.h
--- a/Engine/Source/Runtime/Engine/Classes/Particles/TypeData/ParticleModuleTypeDataBase.h
+++ b/Engine/Source/Runtime/Engine/Classes/Particles/TypeData/ParticleModuleTypeDataBase.h
@@ -10,9 +10,7 @@ struct FParticleEmitterInstance;
 UCLASS(editinlinenew, hidecategories=Object, abstract, MinimalAPI)
 class UParticleModuleTypeDataBase : public UParticleModule
 {
-	GENERATED_BODY()
-public:
-	ENGINE_API UParticleModuleTypeDataBase(const FObjectInitializer& ObjectInitializer = FObjectInitializer::Get());
+	GENERATED_UCLASS_BODY()
 
 
 	// Begin UParticleModule Interface
diff --git a/Engine/Source/Runtime/Engine/Classes/Particles/TypeData/ParticleModuleTypeDataBeam2.h b/Engine/Source/Runtime/Engine/Classes/Particles/TypeData/ParticleModuleTypeDataBeam2.h
--- a/Engine/Source/Runtime/Engine/Classes/Particles/TypeData/ParticleModuleTypeDataBeam2.h
+++ b/Engine/Source/Runtime/Engine/Classes/Particles/TypeData/ParticleModuleTypeDataBeam2.h
@@ -34,7 +34,7 @@ enum EBeamTaperMethod
 USTRUCT()
 struct FBeamTargetData
 {
-	GENERATED_BODY()
+	GENERATED_USTRUCT_BODY()
 
 	/** Name of the target.																	*/
 	UPROPERTY(EditAnywhere, Category=BeamTargetData)
@@ -55,9 +55,7 @@ struct FBeamTargetData
 UCLASS(editinlinenew, dontcollapsecategories, hidecategories=Object, MinimalAPI, meta=(DisplayName = "Beam Data"))
 class UParticleModuleTypeDataBeam2 : public UParticleModuleTypeDataBase
 {
-	GENERATED_BODY()
-public:
-	ENGINE_API UParticleModuleTypeDataBeam2(const FObjectInitializer& ObjectInitializer = FObjectInitializer::Get());
+	GENERATED_UCLASS_BODY()
 
 	//
 	// General Beam Variables.
diff --git a/Engine/Source/Runtime/Engine/Classes/Particles/TypeData/ParticleModuleTypeDataGpu.h b/Engine/Source/Runtime/Engine/Classes/Particles/TypeData/ParticleModuleTypeDataGpu.h
--- a/Engine/Source/Runtime/Engine/Classes/Particles/TypeData/ParticleModuleTypeDataGpu.h
+++ b/Engine/Source/Runtime/Engine/Classes/Particles/TypeData/ParticleModuleTypeDataGpu.h
@@ -18,7 +18,7 @@
 USTRUCT()
 struct FGPUSpriteLocalVectorFieldInfo
 {
-	GENERATED_BODY()
+	GENERATED_USTRUCT_BODY()
 
 	/** Local vector field to apply to this emitter. */
 	UPROPERTY()
@@ -84,7 +84,7 @@ struct FGPUSpriteLocalVectorFieldInfo
 USTRUCT()
 struct FGPUSpriteEmitterInfo
 {
-	GENERATED_BODY()
+	GENERATED_USTRUCT_BODY()
 
 	/** The required module. Needed for now, but should be divorced from the runtime. */
 	UPROPERTY()
@@ -215,7 +215,7 @@ struct FGPUSpriteEmitterInfo
 USTRUCT()
 struct FGPUSpriteResourceData
 {
-	GENERATED_BODY()
+	GENERATED_USTRUCT_BODY()
 
 	/** Quantized color samples. */
 	UPROPERTY()
@@ -392,9 +392,7 @@ struct FGPUSpriteResourceData
 UCLASS(editinlinenew, hidecategories=Object, MinimalAPI, meta=(DisplayName = "GPU Sprites"))
 class UParticleModuleTypeDataGpu : public UParticleModuleTypeDataBase
 {
-	GENERATED_BODY()
-public:
-	ENGINE_API UParticleModuleTypeDataGpu(const FObjectInitializer& ObjectInitializer = FObjectInitializer::Get());
+	GENERATED_UCLASS_BODY()
 
 	/** Information for runtime simulation. */
 	UPROPERTY(transient)
diff --git a/Engine/Source/Runtime/Engine/Classes/Particles/TypeData/ParticleModuleTypeDataMesh.h b/Engine/Source/Runtime/Engine/Classes/Particles/TypeData/ParticleModuleTypeDataMesh.h
--- a/Engine/Source/Runtime/Engine/Classes/Particles/TypeData/ParticleModuleTypeDataMesh.h
+++ b/Engine/Source/Runtime/Engine/Classes/Particles/TypeData/ParticleModuleTypeDataMesh.h
@@ -53,9 +53,7 @@ enum EMeshCameraFacingOptions
 UCLASS(editinlinenew, hidecategories=Object, MinimalAPI, meta=(DisplayName = "Mesh Data"))
 class UParticleModuleTypeDataMesh : public UParticleModuleTypeDataBase
 {
-	GENERATED_BODY()
-public:
-	ENGINE_API UParticleModuleTypeDataMesh(const FObjectInitializer& ObjectInitializer = FObjectInitializer::Get());
+	GENERATED_UCLASS_BODY()
 
 	/** The static mesh to render at the particle positions */
 	UPROPERTY(EditAnywhere, Category=Mesh)
diff --git a/Engine/Source/Runtime/Engine/Classes/Particles/TypeData/ParticleModuleTypeDataRibbon.h b/Engine/Source/Runtime/Engine/Classes/Particles/TypeData/ParticleModuleTypeDataRibbon.h
--- a/Engine/Source/Runtime/Engine/Classes/Particles/TypeData/ParticleModuleTypeDataRibbon.h
+++ b/Engine/Source/Runtime/Engine/Classes/Particles/TypeData/ParticleModuleTypeDataRibbon.h
@@ -22,9 +22,7 @@ enum ETrailsRenderAxisOption
 UCLASS(editinlinenew, hidecategories=Object, meta=(DisplayName = "Ribbon Data"))
 class UParticleModuleTypeDataRibbon : public UParticleModuleTypeDataBase
 {
-	GENERATED_BODY()
-public:
-	UParticleModuleTypeDataRibbon(const FObjectInitializer& ObjectInitializer = FObjectInitializer::Get());
+	GENERATED_UCLASS_BODY()
 
 	//
 	// General Trail Variables.
diff --git a/Engine/Source/Runtime/Engine/Classes/Particles/VectorField/ParticleModuleVectorFieldBase.h b/Engine/Source/Runtime/Engine/Classes/Particles/VectorField/ParticleModuleVectorFieldBase.h
--- a/Engine/Source/Runtime/Engine/Classes/Particles/VectorField/ParticleModuleVectorFieldBase.h
+++ b/Engine/Source/Runtime/Engine/Classes/Particles/VectorField/ParticleModuleVectorFieldBase.h
@@ -12,9 +12,7 @@
 UCLASS(editinlinenew, hidecategories=Object, abstract, meta=(DisplayName = "Vector Field"))
 class UParticleModuleVectorFieldBase : public UParticleModule
 {
-	GENERATED_BODY()
-public:
-	UParticleModuleVectorFieldBase(const FObjectInitializer& ObjectInitializer = FObjectInitializer::Get());
+	GENERATED_UCLASS_BODY()
 
 };
 
diff --git a/Engine/Source/Runtime/Engine/Classes/Particles/VectorField/ParticleModuleVectorFieldGlobal.h b/Engine/Source/Runtime/Engine/Classes/Particles/VectorField/ParticleModuleVectorFieldGlobal.h
--- a/Engine/Source/Runtime/Engine/Classes/Particles/VectorField/ParticleModuleVectorFieldGlobal.h
+++ b/Engine/Source/Runtime/Engine/Classes/Particles/VectorField/ParticleModuleVectorFieldGlobal.h
@@ -11,9 +11,7 @@
 UCLASS(editinlinenew, hidecategories=Object, meta=(DisplayName = "Global Vector Fields"))
 class UParticleModuleVectorFieldGlobal : public UParticleModuleVectorFieldBase
 {
-	GENERATED_BODY()
-public:
-	UParticleModuleVectorFieldGlobal(const FObjectInitializer& ObjectInitializer = FObjectInitializer::Get());
+	GENERATED_UCLASS_BODY()
 
 	/** Property override value for global vector field tightness.  */
 	UPROPERTY()
diff --git a/Engine/Source/Runtime/Engine/Classes/Particles/VectorField/ParticleModuleVectorFieldLocal.h b/Engine/Source/Runtime/Engine/Classes/Particles/VectorField/ParticleModuleVectorFieldLocal.h
--- a/Engine/Source/Runtime/Engine/Classes/Particles/VectorField/ParticleModuleVectorFieldLocal.h
+++ b/Engine/Source/Runtime/Engine/Classes/Particles/VectorField/ParticleModuleVectorFieldLocal.h
@@ -11,9 +11,7 @@
 UCLASS(editinlinenew, hidecategories=Object, MinimalAPI, meta=(DisplayName = "Local Vector Field"))
 class UParticleModuleVectorFieldLocal : public UParticleModuleVectorFieldBase
 {
-	GENERATED_BODY()
-public:
-	ENGINE_API UParticleModuleVectorFieldLocal(const FObjectInitializer& ObjectInitializer = FObjectInitializer::Get());
+	GENERATED_UCLASS_BODY()
 
 	/** Vector field asset to use. */
 	UPROPERTY(EditAnywhere, Category=VectorField)
diff --git a/Engine/Source/Runtime/Engine/Classes/Particles/VectorField/ParticleModuleVectorFieldRotation.h b/Engine/Source/Runtime/Engine/Classes/Particles/VectorField/ParticleModuleVectorFieldRotation.h
--- a/Engine/Source/Runtime/Engine/Classes/Particles/VectorField/ParticleModuleVectorFieldRotation.h
+++ b/Engine/Source/Runtime/Engine/Classes/Particles/VectorField/ParticleModuleVectorFieldRotation.h
@@ -12,9 +12,7 @@
 UCLASS(editinlinenew, hidecategories=Object, meta=(DisplayName = "VF Init Rotation"))
 class UParticleModuleVectorFieldRotation : public UParticleModuleVectorFieldBase
 {
-	GENERATED_BODY()
-public:
-	UParticleModuleVectorFieldRotation(const FObjectInitializer& ObjectInitializer = FObjectInitializer::Get());
+	GENERATED_UCLASS_BODY()
 
 	/** Minimum initial rotation applied to the local vector field. */
 	UPROPERTY(EditAnywhere, Category=VectorField)
diff --git a/Engine/Source/Runtime/Engine/Classes/Particles/VectorField/ParticleModuleVectorFieldRotationRate.h b/Engine/Source/Runtime/Engine/Classes/Particles/VectorField/ParticleModuleVectorFieldRotationRate.h
--- a/Engine/Source/Runtime/Engine/Classes/Particles/VectorField/ParticleModuleVectorFieldRotationRate.h
+++ b/Engine/Source/Runtime/Engine/Classes/Particles/VectorField/ParticleModuleVectorFieldRotationRate.h
@@ -11,9 +11,7 @@
 UCLASS(editinlinenew, hidecategories=Object, meta=(DisplayName = "VF Rotation Rate"))
 class UParticleModuleVectorFieldRotationRate : public UParticleModuleVectorFieldBase
 {
-	GENERATED_BODY()
-public:
-	UParticleModuleVectorFieldRotationRate(const FObjectInitializer& ObjectInitializer = FObjectInitializer::Get());
+	GENERATED_UCLASS_BODY()
 
 	/** Constant rotation rate applied to the local vector field. */
 	UPROPERTY(EditAnywhere, Category=VectorField)
diff --git a/Engine/Source/Runtime/Engine/Classes/Particles/VectorField/ParticleModuleVectorFieldScale.h b/Engine/Source/Runtime/Engine/Classes/Particles/VectorField/ParticleModuleVectorFieldScale.h
--- a/Engine/Source/Runtime/Engine/Classes/Particles/VectorField/ParticleModuleVectorFieldScale.h
+++ b/Engine/Source/Runtime/Engine/Classes/Particles/VectorField/ParticleModuleVectorFieldScale.h
@@ -11,9 +11,7 @@
 UCLASS(editinlinenew, hidecategories=Object, MinimalAPI, meta=(DisplayName = "Vector Field Scale"))
 class UParticleModuleVectorFieldScale : public UParticleModuleVectorFieldBase
 {
-	GENERATED_BODY()
-public:
-	ENGINE_API UParticleModuleVectorFieldScale(const FObjectInitializer& ObjectInitializer = FObjectInitializer::Get());
+	GENERATED_UCLASS_BODY()
 
 	/** Per-particle vector field scale. Evaluated using emitter time. */
 	UPROPERTY(EditAnywhere, Category=VectorField)
diff --git a/Engine/Source/Runtime/Engine/Classes/Particles/VectorField/ParticleModuleVectorFieldScaleOverLife.h b/Engine/Source/Runtime/Engine/Classes/Particles/VectorField/ParticleModuleVectorFieldScaleOverLife.h
--- a/Engine/Source/Runtime/Engine/Classes/Particles/VectorField/ParticleModuleVectorFieldScaleOverLife.h
+++ b/Engine/Source/Runtime/Engine/Classes/Particles/VectorField/ParticleModuleVectorFieldScaleOverLife.h
@@ -11,9 +11,7 @@
 UCLASS(editinlinenew, hidecategories=Object, MinimalAPI, meta=(DisplayName = "VF Scale/Life"))
 class UParticleModuleVectorFieldScaleOverLife : public UParticleModuleVectorFieldBase
 {
-	GENERATED_BODY()
-public:
-	ENGINE_API UParticleModuleVectorFieldScaleOverLife(const FObjectInitializer& ObjectInitializer = FObjectInitializer::Get());
+	GENERATED_UCLASS_BODY()
 
 	/** Per-particle vector field scale. Evaluated using particle relative time. */
 	UPROPERTY(EditAnywhere, Category=VectorField)
diff --git a/Engine/Source/Runtime/Engine/Classes/Particles/Velocity/ParticleModuleVelocity.h b/Engine/Source/Runtime/Engine/Classes/Particles/Velocity/ParticleModuleVelocity.h
--- a/Engine/Source/Runtime/Engine/Classes/Particles/Velocity/ParticleModuleVelocity.h
+++ b/Engine/Source/Runtime/Engine/Classes/Particles/Velocity/ParticleModuleVelocity.h
@@ -8,9 +8,7 @@
 UCLASS(editinlinenew, hidecategories=Object, MinimalAPI, meta=(DisplayName = "Initial Velocity"))
 class UParticleModuleVelocity : public UParticleModuleVelocityBase
 {
-	GENERATED_BODY()
-public:
-	ENGINE_API UParticleModuleVelocity(const FObjectInitializer& ObjectInitializer = FObjectInitializer::Get());
+	GENERATED_UCLASS_BODY()
 
 	/** 
 	 *	The velocity to apply to a particle when it is spawned.
diff --git a/Engine/Source/Runtime/Engine/Classes/Particles/Velocity/ParticleModuleVelocityBase.h b/Engine/Source/Runtime/Engine/Classes/Particles/Velocity/ParticleModuleVelocityBase.h
--- a/Engine/Source/Runtime/Engine/Classes/Particles/Velocity/ParticleModuleVelocityBase.h
+++ b/Engine/Source/Runtime/Engine/Classes/Particles/Velocity/ParticleModuleVelocityBase.h
@@ -8,9 +8,7 @@
 UCLASS(editinlinenew, hidecategories=Object, abstract, meta=(DisplayName = "Velocity"))
 class UParticleModuleVelocityBase : public UParticleModule
 {
-	GENERATED_BODY()
-public:
-	UParticleModuleVelocityBase(const FObjectInitializer& ObjectInitializer = FObjectInitializer::Get());
+	GENERATED_UCLASS_BODY()
 
 	/**
 	 *	If true, then treat the velocity as world-space defined.
diff --git a/Engine/Source/Runtime/Engine/Classes/Particles/Velocity/ParticleModuleVelocityCone.h b/Engine/Source/Runtime/Engine/Classes/Particles/Velocity/ParticleModuleVelocityCone.h
--- a/Engine/Source/Runtime/Engine/Classes/Particles/Velocity/ParticleModuleVelocityCone.h
+++ b/Engine/Source/Runtime/Engine/Classes/Particles/Velocity/ParticleModuleVelocityCone.h
@@ -8,9 +8,7 @@
 UCLASS(editinlinenew, hidecategories=Object, MinimalAPI, meta=( DisplayName = "Velocity Cone" ))
 class UParticleModuleVelocityCone : public UParticleModuleVelocityBase
 {
-	GENERATED_BODY()
-public:
-	ENGINE_API UParticleModuleVelocityCone(const FObjectInitializer& ObjectInitializer = FObjectInitializer::Get());
+	GENERATED_UCLASS_BODY()
 
 	/** The Min value represents the inner cone angle value and the Max value represents the outer cone angle value. */
 	UPROPERTY(EditAnywhere, Category=Cone)
diff --git a/Engine/Source/Runtime/Engine/Classes/Particles/Velocity/ParticleModuleVelocityInheritParent.h b/Engine/Source/Runtime/Engine/Classes/Particles/Velocity/ParticleModuleVelocityInheritParent.h
--- a/Engine/Source/Runtime/Engine/Classes/Particles/Velocity/ParticleModuleVelocityInheritParent.h
+++ b/Engine/Source/Runtime/Engine/Classes/Particles/Velocity/ParticleModuleVelocityInheritParent.h
@@ -8,9 +8,7 @@
 UCLASS(editinlinenew, hidecategories=Object, MinimalAPI, meta=(DisplayName = "Inherit Parent Velocity"))
 class UParticleModuleVelocityInheritParent : public UParticleModuleVelocityBase
 {
-	GENERATED_BODY()
-public:
-	ENGINE_API UParticleModuleVelocityInheritParent(const FObjectInitializer& ObjectInitializer = FObjectInitializer::Get());
+	GENERATED_UCLASS_BODY()
 
 	/**
 	 *	The scale to apply tot he parent velocity prior to adding it to the particle velocity during spawn.
diff --git a/Engine/Source/Runtime/Engine/Classes/Particles/Velocity/ParticleModuleVelocityOverLifetime.h b/Engine/Source/Runtime/Engine/Classes/Particles/Velocity/ParticleModuleVelocityOverLifetime.h
--- a/Engine/Source/Runtime/Engine/Classes/Particles/Velocity/ParticleModuleVelocityOverLifetime.h
+++ b/Engine/Source/Runtime/Engine/Classes/Particles/Velocity/ParticleModuleVelocityOverLifetime.h
@@ -8,9 +8,7 @@
 UCLASS(editinlinenew, hidecategories=Object, MinimalAPI, meta=(DisplayName = "Velocity/Life"))
 class UParticleModuleVelocityOverLifetime : public UParticleModuleVelocityBase
 {
-	GENERATED_BODY()
-public:
-	ENGINE_API UParticleModuleVelocityOverLifetime(const FObjectInitializer& ObjectInitializer = FObjectInitializer::Get());
+	GENERATED_UCLASS_BODY()
 
 	/**
 	 *	The scaling  value applied to the velocity.
diff --git a/Engine/Source/Runtime/Engine/Classes/Particles/Velocity/ParticleModuleVelocity_Seeded.h b/Engine/Source/Runtime/Engine/Classes/Particles/Velocity/ParticleModuleVelocity_Seeded.h
--- a/Engine/Source/Runtime/Engine/Classes/Particles/Velocity/ParticleModuleVelocity_Seeded.h
+++ b/Engine/Source/Runtime/Engine/Classes/Particles/Velocity/ParticleModuleVelocity_Seeded.h
@@ -8,9 +8,7 @@
 UCLASS(editinlinenew, hidecategories=Object, meta=(DisplayName = "Initial Velocity (Seed)"))
 class UParticleModuleVelocity_Seeded : public UParticleModuleVelocity
 {
-	GENERATED_BODY()
-public:
-	UParticleModuleVelocity_Seeded(const FObjectInitializer& ObjectInitializer = FObjectInitializer::Get());
+	GENERATED_UCLASS_BODY()
 
 	/** The random seed(s) to use for looking up values in StartLocation */
 	UPROPERTY(EditAnywhere, Category=RandomSeed)
diff --git a/Engine/Source/Runtime/Engine/Classes/PhysicalMaterials/PhysicalMaterial.h b/Engine/Source/Runtime/Engine/Classes/PhysicalMaterials/PhysicalMaterial.h
--- a/Engine/Source/Runtime/Engine/Classes/PhysicalMaterials/PhysicalMaterial.h
+++ b/Engine/Source/Runtime/Engine/Classes/PhysicalMaterials/PhysicalMaterial.h
@@ -13,7 +13,7 @@ namespace physx
 USTRUCT()
 struct FTireFrictionScalePair
 {
-	GENERATED_BODY()
+	GENERATED_USTRUCT_BODY()
 
 	/** Tire type */
 	UPROPERTY(EditAnywhere, Category=TireFrictionScalePair)
@@ -36,9 +36,7 @@ struct FTireFrictionScalePair
 UCLASS(collapsecategories, hidecategories=Object)
 class ENGINE_API UPhysicalMaterial : public UObject
 {
-	GENERATED_BODY()
-public:
-	UPhysicalMaterial(const FObjectInitializer& ObjectInitializer = FObjectInitializer::Get());
+	GENERATED_UCLASS_BODY()
 
 	//
 	// Surface properties.
diff --git a/Engine/Source/Runtime/Engine/Classes/PhysicalMaterials/PhysicalMaterialPropertyBase.h b/Engine/Source/Runtime/Engine/Classes/PhysicalMaterials/PhysicalMaterialPropertyBase.h
--- a/Engine/Source/Runtime/Engine/Classes/PhysicalMaterials/PhysicalMaterialPropertyBase.h
+++ b/Engine/Source/Runtime/Engine/Classes/PhysicalMaterials/PhysicalMaterialPropertyBase.h
@@ -26,9 +26,7 @@
 UCLASS(collapsecategories, hidecategories=Object, editinlinenew, abstract,MinimalAPI, deprecated)
 class UDEPRECATED_PhysicalMaterialPropertyBase : public UObject
 {
-	GENERATED_BODY()
-public:
-	ENGINE_API UDEPRECATED_PhysicalMaterialPropertyBase(const FObjectInitializer& ObjectInitializer = FObjectInitializer::Get());
+	GENERATED_UCLASS_BODY()
 
 	// Begin UObject Interface
 	virtual bool IsAsset() const override { return false; }
diff --git a/Engine/Source/Runtime/Engine/Classes/PhysicsEngine/AggregateGeometry2D.h b/Engine/Source/Runtime/Engine/Classes/PhysicsEngine/AggregateGeometry2D.h
--- a/Engine/Source/Runtime/Engine/Classes/PhysicsEngine/AggregateGeometry2D.h
+++ b/Engine/Source/Runtime/Engine/Classes/PhysicsEngine/AggregateGeometry2D.h
@@ -9,7 +9,7 @@
 USTRUCT()
 struct ENGINE_API FConvexElement2D
 {
-	GENERATED_BODY()
+	GENERATED_USTRUCT_BODY()
 
 	// Vertices that lie on the convex hull
 	UPROPERTY()
@@ -30,7 +30,7 @@ struct ENGINE_API FConvexElement2D
 USTRUCT()
 struct ENGINE_API FCircleElement2D
 {
-	GENERATED_BODY()
+	GENERATED_USTRUCT_BODY()
 
 	UPROPERTY(Category=KSphereElem, VisibleAnywhere)
 	FVector2D Center;
@@ -68,7 +68,7 @@ public:
 USTRUCT()
 struct ENGINE_API FBoxElement2D
 {
-	GENERATED_BODY()
+	GENERATED_USTRUCT_BODY()
 
 	// Center of the box
 	UPROPERTY(Category=KBoxElem, VisibleAnywhere)
@@ -114,7 +114,7 @@ struct ENGINE_API FBoxElement2D
 USTRUCT()
 struct ENGINE_API FAggregateGeometry2D
 {
-	GENERATED_BODY()
+	GENERATED_USTRUCT_BODY()
 
 public:
 	UPROPERTY(EditAnywhere, editfixedsize, Category = KAggregateGeom)
diff --git a/Engine/Source/Runtime/Engine/Classes/PhysicsEngine/BodyInstance.h b/Engine/Source/Runtime/Engine/Classes/PhysicsEngine/BodyInstance.h
--- a/Engine/Source/Runtime/Engine/Classes/PhysicsEngine/BodyInstance.h
+++ b/Engine/Source/Runtime/Engine/Classes/PhysicsEngine/BodyInstance.h
@@ -64,7 +64,7 @@ struct FCollisionNotifyInfo;
 USTRUCT()
 struct ENGINE_API FCollisionResponse
 {
-	GENERATED_BODY()
+	GENERATED_USTRUCT_BODY()
 
 	FCollisionResponse();
 	FCollisionResponse(ECollisionResponse DefaultResponse);
@@ -126,7 +126,7 @@ enum class BodyInstanceSceneState
 USTRUCT()
 struct ENGINE_API FBodyInstance
 {
-	GENERATED_BODY()
+	GENERATED_USTRUCT_BODY()
 
 	/** 
 	 *	Index of this BodyInstance within the SkeletalMeshComponent/PhysicsAsset. 
@@ -862,4 +862,4 @@ struct FBodyInstanceEditorHelpers
 private:
 	FBodyInstanceEditorHelpers() {}
 };
-#endif
+#endif
\ No newline at end of file
diff --git a/Engine/Source/Runtime/Engine/Classes/PhysicsEngine/BodySetup.h b/Engine/Source/Runtime/Engine/Classes/PhysicsEngine/BodySetup.h
--- a/Engine/Source/Runtime/Engine/Classes/PhysicsEngine/BodySetup.h
+++ b/Engine/Source/Runtime/Engine/Classes/PhysicsEngine/BodySetup.h
@@ -24,9 +24,7 @@ namespace physx
 UCLASS(hidecategories=Object, MinimalAPI)
 class UBodySetup : public UObject
 {
-	GENERATED_BODY()
-public:
-	ENGINE_API UBodySetup(const FObjectInitializer& ObjectInitializer = FObjectInitializer::Get());
+	GENERATED_UCLASS_BODY()
 
 	/** Simplified collision representation of this  */
 	UPROPERTY()
diff --git a/Engine/Source/Runtime/Engine/Classes/PhysicsEngine/BodySetup2D.h b/Engine/Source/Runtime/Engine/Classes/PhysicsEngine/BodySetup2D.h
--- a/Engine/Source/Runtime/Engine/Classes/PhysicsEngine/BodySetup2D.h
+++ b/Engine/Source/Runtime/Engine/Classes/PhysicsEngine/BodySetup2D.h
@@ -8,9 +8,7 @@
 UCLASS()
 class ENGINE_API UBodySetup2D : public UBodySetup
 {
-	GENERATED_BODY()
-public:
-	UBodySetup2D(const FObjectInitializer& ObjectInitializer = FObjectInitializer::Get());
+	GENERATED_UCLASS_BODY()
 
 	UPROPERTY()
 	FAggregateGeometry2D AggGeom2D;
diff --git a/Engine/Source/Runtime/Engine/Classes/PhysicsEngine/DestructibleActor.h b/Engine/Source/Runtime/Engine/Classes/PhysicsEngine/DestructibleActor.h
--- a/Engine/Source/Runtime/Engine/Classes/PhysicsEngine/DestructibleActor.h
+++ b/Engine/Source/Runtime/Engine/Classes/PhysicsEngine/DestructibleActor.h
@@ -14,9 +14,7 @@ DECLARE_DYNAMIC_MULTICAST_DELEGATE_TwoParams(FActorFractureSignature, const FVec
 UCLASS(MinimalAPI, hideCategories=(Input), showCategories=("Input|MouseInput", "Input|TouchInput"), ComponentWrapperClass)
 class ADestructibleActor : public AActor
 {
-	GENERATED_BODY()
-public:
-	ENGINE_API ADestructibleActor(const FObjectInitializer& ObjectInitializer = FObjectInitializer::Get());
+	GENERATED_UCLASS_BODY()
 
 	/**
 	 * The component which holds the skinned mesh and physics data for this actor.
diff --git a/Engine/Source/Runtime/Engine/Classes/PhysicsEngine/PhysicsCollisionHandler.h b/Engine/Source/Runtime/Engine/Classes/PhysicsEngine/PhysicsCollisionHandler.h
--- a/Engine/Source/Runtime/Engine/Classes/PhysicsEngine/PhysicsCollisionHandler.h
+++ b/Engine/Source/Runtime/Engine/Classes/PhysicsEngine/PhysicsCollisionHandler.h
@@ -8,9 +8,7 @@ struct FCollisionNotifyInfo;
 UCLASS(hidecategories=(Object), Blueprintable)
 class ENGINE_API UPhysicsCollisionHandler : public UObject
 {
-	GENERATED_BODY()
-public:
-	UPhysicsCollisionHandler(const FObjectInitializer& ObjectInitializer = FObjectInitializer::Get());
+	GENERATED_UCLASS_BODY()
 
 	// This class is SUPER BASIC, but games can implement their own handlers to do more advanced physics collisions
 
diff --git a/Engine/Source/Runtime/Engine/Classes/PhysicsEngine/PhysicsConstraintActor.h b/Engine/Source/Runtime/Engine/Classes/PhysicsEngine/PhysicsConstraintActor.h
--- a/Engine/Source/Runtime/Engine/Classes/PhysicsEngine/PhysicsConstraintActor.h
+++ b/Engine/Source/Runtime/Engine/Classes/PhysicsEngine/PhysicsConstraintActor.h
@@ -10,9 +10,7 @@
 UCLASS(ConversionRoot, MinimalAPI, ComponentWrapperClass)
 class APhysicsConstraintActor : public ARigidBodyBase
 {
-	GENERATED_BODY()
-public:
-	ENGINE_API APhysicsConstraintActor(const FObjectInitializer& ObjectInitializer = FObjectInitializer::Get());
+	GENERATED_UCLASS_BODY()
 
 	// Cached reference to constraint component
 private_subobject:
diff --git a/Engine/Source/Runtime/Engine/Classes/PhysicsEngine/PhysicsConstraintComponent.h b/Engine/Source/Runtime/Engine/Classes/PhysicsEngine/PhysicsConstraintComponent.h
--- a/Engine/Source/Runtime/Engine/Classes/PhysicsEngine/PhysicsConstraintComponent.h
+++ b/Engine/Source/Runtime/Engine/Classes/PhysicsEngine/PhysicsConstraintComponent.h
@@ -12,9 +12,7 @@ DECLARE_DYNAMIC_MULTICAST_DELEGATE_OneParam(FConstraintBrokenSignature, int32, C
 UCLASS(ClassGroup=Physics, MinimalAPI, meta=(BlueprintSpawnableComponent), HideCategories=(Activation,"Components|Activation", Physics, Mobility), ShowCategories=("Physics|Components|PhysicsConstraint"))
 class UPhysicsConstraintComponent : public USceneComponent
 {
-	GENERATED_BODY()
-public:
-	ENGINE_API UPhysicsConstraintComponent(const FObjectInitializer& ObjectInitializer = FObjectInitializer::Get());
+	GENERATED_UCLASS_BODY()
 
 	/** Pointer to first Actor to constrain.  */
 	UPROPERTY(EditInstanceOnly, Category=Constraint)
diff --git a/Engine/Source/Runtime/Engine/Classes/PhysicsEngine/PhysicsConstraintTemplate.h b/Engine/Source/Runtime/Engine/Classes/PhysicsEngine/PhysicsConstraintTemplate.h
--- a/Engine/Source/Runtime/Engine/Classes/PhysicsEngine/PhysicsConstraintTemplate.h
+++ b/Engine/Source/Runtime/Engine/Classes/PhysicsEngine/PhysicsConstraintTemplate.h
@@ -17,9 +17,7 @@
 UCLASS(hidecategories=Object, MinimalAPI)
 class UPhysicsConstraintTemplate : public UObject
 {
-	GENERATED_BODY()
-public:
-	ENGINE_API UPhysicsConstraintTemplate(const FObjectInitializer& ObjectInitializer = FObjectInitializer::Get());
+	GENERATED_UCLASS_BODY()
 
 	UPROPERTY()
 	FName JointName_DEPRECATED;
diff --git a/Engine/Source/Runtime/Engine/Classes/PhysicsEngine/PhysicsHandleComponent.h b/Engine/Source/Runtime/Engine/Classes/PhysicsEngine/PhysicsHandleComponent.h
--- a/Engine/Source/Runtime/Engine/Classes/PhysicsEngine/PhysicsHandleComponent.h
+++ b/Engine/Source/Runtime/Engine/Classes/PhysicsEngine/PhysicsHandleComponent.h
@@ -17,9 +17,7 @@ namespace physx
 UCLASS(collapsecategories, ClassGroup=Physics, hidecategories=Object, MinimalAPI, meta=(BlueprintSpawnableComponent))
 class UPhysicsHandleComponent : public UActorComponent
 {
-	GENERATED_BODY()
-public:
-	ENGINE_API UPhysicsHandleComponent(const FObjectInitializer& ObjectInitializer = FObjectInitializer::Get());
+	GENERATED_UCLASS_BODY()
 
 	/** Component we are currently holding */
 	UPROPERTY()
diff --git a/Engine/Source/Runtime/Engine/Classes/PhysicsEngine/PhysicsSettings.h b/Engine/Source/Runtime/Engine/Classes/PhysicsEngine/PhysicsSettings.h
--- a/Engine/Source/Runtime/Engine/Classes/PhysicsEngine/PhysicsSettings.h
+++ b/Engine/Source/Runtime/Engine/Classes/PhysicsEngine/PhysicsSettings.h
@@ -15,7 +15,7 @@
 USTRUCT()
 struct FPhysicalSurfaceName
 {
-	GENERATED_BODY()
+	GENERATED_USTRUCT_BODY()
 
 	UPROPERTY()
 	TEnumAsByte<enum EPhysicalSurface> Type;
@@ -88,9 +88,7 @@ UCLASS(config=Engine, defaultconfig)
 class ENGINE_API UPhysicsSettings
 	: public UObject
 {
-	GENERATED_BODY()
-public:
-	UPhysicsSettings(const FObjectInitializer& ObjectInitializer = FObjectInitializer::Get());
+	GENERATED_UCLASS_BODY()
 
 	/** Default gravity. */
 	UPROPERTY(config, EditAnywhere, Category = Constants)
diff --git a/Engine/Source/Runtime/Engine/Classes/PhysicsEngine/PhysicsThruster.h b/Engine/Source/Runtime/Engine/Classes/PhysicsEngine/PhysicsThruster.h
--- a/Engine/Source/Runtime/Engine/Classes/PhysicsEngine/PhysicsThruster.h
+++ b/Engine/Source/Runtime/Engine/Classes/PhysicsEngine/PhysicsThruster.h
@@ -13,9 +13,7 @@
 UCLASS(hideCategories=(Input,Collision,Replication), showCategories=("Input|MouseInput", "Input|TouchInput"), ComponentWrapperClass)
 class APhysicsThruster : public ARigidBodyBase
 {
-	GENERATED_BODY()
-public:
-	APhysicsThruster(const FObjectInitializer& ObjectInitializer = FObjectInitializer::Get());
+	GENERATED_UCLASS_BODY()
 
 private_subobject:
 	/** Thruster component */
diff --git a/Engine/Source/Runtime/Engine/Classes/PhysicsEngine/PhysicsThrusterComponent.h b/Engine/Source/Runtime/Engine/Classes/PhysicsEngine/PhysicsThrusterComponent.h
--- a/Engine/Source/Runtime/Engine/Classes/PhysicsEngine/PhysicsThrusterComponent.h
+++ b/Engine/Source/Runtime/Engine/Classes/PhysicsEngine/PhysicsThrusterComponent.h
@@ -11,9 +11,7 @@
 UCLASS(hidecategories=(Object, Mobility, LOD), ClassGroup=Physics, showcategories=Trigger, MinimalAPI, meta=(BlueprintSpawnableComponent))
 class UPhysicsThrusterComponent : public USceneComponent
 {
-	GENERATED_BODY()
-public:
-	ENGINE_API UPhysicsThrusterComponent(const FObjectInitializer& ObjectInitializer = FObjectInitializer::Get());
+	GENERATED_UCLASS_BODY()
 
 	/** Strength of thrust force applied to the base object. */
 	UPROPERTY(BlueprintReadWrite, interp, Category=Physics)
diff --git a/Engine/Source/Runtime/Engine/Classes/PhysicsEngine/RadialForceActor.h b/Engine/Source/Runtime/Engine/Classes/PhysicsEngine/RadialForceActor.h
--- a/Engine/Source/Runtime/Engine/Classes/PhysicsEngine/RadialForceActor.h
+++ b/Engine/Source/Runtime/Engine/Classes/PhysicsEngine/RadialForceActor.h
@@ -8,9 +8,7 @@
 UCLASS(MinimalAPI, hideCategories=(Collision, Input), showCategories=("Input|MouseInput", "Input|TouchInput"), ComponentWrapperClass)
 class ARadialForceActor : public ARigidBodyBase
 {
-	GENERATED_BODY()
-public:
-	ENGINE_API ARadialForceActor(const FObjectInitializer& ObjectInitializer = FObjectInitializer::Get());
+	GENERATED_UCLASS_BODY()
 
 private_subobject:
 	/** Force component */
diff --git a/Engine/Source/Runtime/Engine/Classes/PhysicsEngine/RadialForceComponent.h b/Engine/Source/Runtime/Engine/Classes/PhysicsEngine/RadialForceComponent.h
--- a/Engine/Source/Runtime/Engine/Classes/PhysicsEngine/RadialForceComponent.h
+++ b/Engine/Source/Runtime/Engine/Classes/PhysicsEngine/RadialForceComponent.h
@@ -10,9 +10,7 @@
 UCLASS(hidecategories=(Object, Mobility), ClassGroup=Physics, showcategories=Trigger, meta=(BlueprintSpawnableComponent), MinimalAPI)
 class URadialForceComponent : public USceneComponent
 {
-	GENERATED_BODY()
-public:
-	ENGINE_API URadialForceComponent(const FObjectInitializer& ObjectInitializer = FObjectInitializer::Get());
+	GENERATED_UCLASS_BODY()
 
 	/** The radius to apply the force or impulse in */
 	UPROPERTY(interp, EditAnywhere, BlueprintReadWrite, Category=RadialForceComponent)
diff --git a/Engine/Source/Runtime/Engine/Classes/PhysicsEngine/RigidBodyBase.h b/Engine/Source/Runtime/Engine/Classes/PhysicsEngine/RigidBodyBase.h
--- a/Engine/Source/Runtime/Engine/Classes/PhysicsEngine/RigidBodyBase.h
+++ b/Engine/Source/Runtime/Engine/Classes/PhysicsEngine/RigidBodyBase.h
@@ -10,9 +10,7 @@
 UCLASS(ClassGroup=Physics, abstract,MinimalAPI)
 class ARigidBodyBase : public AActor
 {
-	GENERATED_BODY()
-public:
-	ENGINE_API ARigidBodyBase(const FObjectInitializer& ObjectInitializer = FObjectInitializer::Get());
+	GENERATED_UCLASS_BODY()
 
 };
 
diff --git a/Engine/Source/Runtime/Engine/Classes/Slate/ButtonStyleAsset.h b/Engine/Source/Runtime/Engine/Classes/Slate/ButtonStyleAsset.h
--- a/Engine/Source/Runtime/Engine/Classes/Slate/ButtonStyleAsset.h
+++ b/Engine/Source/Runtime/Engine/Classes/Slate/ButtonStyleAsset.h
@@ -12,9 +12,7 @@
 UCLASS(hidecategories=Object, MinimalAPI)
 class UButtonStyleAsset : public UObject
 {
-	GENERATED_BODY()
-public:
-	ENGINE_API UButtonStyleAsset(const FObjectInitializer& ObjectInitializer = FObjectInitializer::Get());
+	GENERATED_UCLASS_BODY()
 
 public:
 	/** The actual data describing the button's appearance. */
diff --git a/Engine/Source/Runtime/Engine/Classes/Slate/CheckboxStyleAsset.h b/Engine/Source/Runtime/Engine/Classes/Slate/CheckboxStyleAsset.h
--- a/Engine/Source/Runtime/Engine/Classes/Slate/CheckboxStyleAsset.h
+++ b/Engine/Source/Runtime/Engine/Classes/Slate/CheckboxStyleAsset.h
@@ -12,9 +12,7 @@
 UCLASS(hidecategories=Object, MinimalAPI)
 class UCheckBoxStyleAsset : public UObject
 {
-	GENERATED_BODY()
-public:
-	ENGINE_API UCheckBoxStyleAsset(const FObjectInitializer& ObjectInitializer = FObjectInitializer::Get());
+	GENERATED_UCLASS_BODY()
 
 public:
 	/** The actual data describing the Check Box's appearance. */
diff --git a/Engine/Source/Runtime/Engine/Classes/Slate/SlateBrushAsset.h b/Engine/Source/Runtime/Engine/Classes/Slate/SlateBrushAsset.h
--- a/Engine/Source/Runtime/Engine/Classes/Slate/SlateBrushAsset.h
+++ b/Engine/Source/Runtime/Engine/Classes/Slate/SlateBrushAsset.h
@@ -12,9 +12,7 @@
 UCLASS(hidecategories=Object, MinimalAPI, BlueprintType)
 class USlateBrushAsset : public UObject
 {
-	GENERATED_BODY()
-public:
-	ENGINE_API USlateBrushAsset(const FObjectInitializer& ObjectInitializer = FObjectInitializer::Get());
+	GENERATED_UCLASS_BODY()
 
 public:
 	/** The slate brush resource describing the texture's behavior. */
diff --git a/Engine/Source/Runtime/Engine/Classes/Sound/AmbientSound.h b/Engine/Source/Runtime/Engine/Classes/Sound/AmbientSound.h
--- a/Engine/Source/Runtime/Engine/Classes/Sound/AmbientSound.h
+++ b/Engine/Source/Runtime/Engine/Classes/Sound/AmbientSound.h
@@ -7,9 +7,7 @@
 UCLASS(AutoExpandCategories=Audio, ClassGroup=Sounds, MinimalAPI, hideCategories(Collision, Input, Game), showCategories=("Input|MouseInput", "Input|TouchInput", "Game|Damage"), ComponentWrapperClass)
 class AAmbientSound : public AActor
 {
-	GENERATED_BODY()
-public:
-	ENGINE_API AAmbientSound(const FObjectInitializer& ObjectInitializer = FObjectInitializer::Get());
+	GENERATED_UCLASS_BODY()
 
 private_subobject:
 	/** Audio component that handles sound playing */
diff --git a/Engine/Source/Runtime/Engine/Classes/Sound/AudioSettings.h b/Engine/Source/Runtime/Engine/Classes/Sound/AudioSettings.h
--- a/Engine/Source/Runtime/Engine/Classes/Sound/AudioSettings.h
+++ b/Engine/Source/Runtime/Engine/Classes/Sound/AudioSettings.h
@@ -12,9 +12,7 @@ UCLASS(config=Engine, defaultconfig)
 class ENGINE_API UAudioSettings
 	: public UObject
 {
-	GENERATED_BODY()
-public:
-	UAudioSettings(const FObjectInitializer& ObjectInitializer = FObjectInitializer::Get());
+	GENERATED_UCLASS_BODY()
 
 	/** The SoundClass assigned to newly created sounds */
 	UPROPERTY(config, EditAnywhere, Category="Audio", meta=(AllowedClasses="SoundClass", DisplayName="Default Sound Class"))
diff --git a/Engine/Source/Runtime/Engine/Classes/Sound/AudioVolume.h b/Engine/Source/Runtime/Engine/Classes/Sound/AudioVolume.h
--- a/Engine/Source/Runtime/Engine/Classes/Sound/AudioVolume.h
+++ b/Engine/Source/Runtime/Engine/Classes/Sound/AudioVolume.h
@@ -47,7 +47,7 @@ enum ReverbPreset
 USTRUCT(BlueprintType)
 struct FReverbSettings
 {
-	GENERATED_BODY()
+	GENERATED_USTRUCT_BODY()
 
 	/* Whether to apply the reverb settings below. */
 	UPROPERTY(EditAnywhere, BlueprintReadWrite, Category=ReverbSettings )
@@ -97,7 +97,7 @@ struct TStructOpsTypeTraits<FReverbSettings> : public TStructOpsTypeTraitsBase
 USTRUCT(BlueprintType)
 struct FInteriorSettings
 {
-	GENERATED_BODY()
+	GENERATED_USTRUCT_BODY()
 
 	// Whether these interior settings are the default values for the world
 	UPROPERTY()
@@ -155,9 +155,7 @@ struct FInteriorSettings
 UCLASS(hidecategories=(Advanced, Attachment, Collision, Volume), MinimalAPI)
 class AAudioVolume : public AVolume
 {
-	GENERATED_BODY()
-public:
-	ENGINE_API AAudioVolume(const FObjectInitializer& ObjectInitializer = FObjectInitializer::Get());
+	GENERATED_UCLASS_BODY()
 
 	/**
 	 * Priority of this volume. In the case of overlapping volumes the one with the highest priority
diff --git a/Engine/Source/Runtime/Engine/Classes/Sound/DialogueTypes.h b/Engine/Source/Runtime/Engine/Classes/Sound/DialogueTypes.h
--- a/Engine/Source/Runtime/Engine/Classes/Sound/DialogueTypes.h
+++ b/Engine/Source/Runtime/Engine/Classes/Sound/DialogueTypes.h
@@ -36,7 +36,7 @@ class UDialogueVoice;
 USTRUCT(BlueprintType)
 struct ENGINE_API FDialogueContext
 {
-	GENERATED_BODY()
+	GENERATED_USTRUCT_BODY()
 
 	FDialogueContext();
 
@@ -61,7 +61,7 @@ class UDialogueWave;
 USTRUCT()
 struct FDialogueWaveParameter
 {
-	GENERATED_BODY()
+	GENERATED_USTRUCT_BODY()
 
 	FDialogueWaveParameter();
 
@@ -77,7 +77,5 @@ struct FDialogueWaveParameter
 UCLASS(abstract)
 class UDialogueTypes : public UObject
 {
-	GENERATED_BODY()
-public:
-	UDialogueTypes(const FObjectInitializer& ObjectInitializer = FObjectInitializer::Get());
+	GENERATED_UCLASS_BODY()
 };
diff --git a/Engine/Source/Runtime/Engine/Classes/Sound/DialogueVoice.h b/Engine/Source/Runtime/Engine/Classes/Sound/DialogueVoice.h
--- a/Engine/Source/Runtime/Engine/Classes/Sound/DialogueVoice.h
+++ b/Engine/Source/Runtime/Engine/Classes/Sound/DialogueVoice.h
@@ -13,9 +13,7 @@
 UCLASS(hidecategories=Object, editinlinenew, MinimalAPI, BlueprintType)
 class UDialogueVoice : public UObject
 {
-	GENERATED_BODY()
-public:
-	ENGINE_API UDialogueVoice(const FObjectInitializer& ObjectInitializer = FObjectInitializer::Get());
+	GENERATED_UCLASS_BODY()
 
 	UPROPERTY(EditAnywhere, Category=DialogueVoice, AssetRegistrySearchable)
 	TEnumAsByte<EGrammaticalGender::Type> Gender;
@@ -35,4 +33,4 @@ public:
 
 	virtual void PostDuplicate(bool bDuplicateForPIE);
 	// End UObject interface. 
-};
+};
\ No newline at end of file
diff --git a/Engine/Source/Runtime/Engine/Classes/Sound/DialogueWave.h b/Engine/Source/Runtime/Engine/Classes/Sound/DialogueWave.h
--- a/Engine/Source/Runtime/Engine/Classes/Sound/DialogueWave.h
+++ b/Engine/Source/Runtime/Engine/Classes/Sound/DialogueWave.h
@@ -16,7 +16,7 @@ class UDialogueSoundWaveProxy;
 USTRUCT()
 struct FDialogueContextMapping
 {
-	GENERATED_BODY()
+	GENERATED_USTRUCT_BODY()
 
 	FDialogueContextMapping();
 
@@ -39,9 +39,7 @@ bool operator!=(const FDialogueContextMapping& LHS, const FDialogueContextMappin
 UCLASS(hidecategories=Object, editinlinenew, MinimalAPI, BlueprintType)
 class UDialogueWave : public UObject
 {
-	GENERATED_BODY()
-public:
-	ENGINE_API UDialogueWave(const FObjectInitializer& ObjectInitializer = FObjectInitializer::Get());
+	GENERATED_UCLASS_BODY()
 
 	/** true if this dialogue is considered to contain mature/adult content. */
 	UPROPERTY(EditAnywhere, Category=Filter, AssetRegistrySearchable)
@@ -88,4 +86,4 @@ public:
 
 protected:
 	void UpdateMappingProxy(FDialogueContextMapping& ContextMapping);
-};
+};
\ No newline at end of file
diff --git a/Engine/Source/Runtime/Engine/Classes/Sound/ReverbEffect.h b/Engine/Source/Runtime/Engine/Classes/Sound/ReverbEffect.h
--- a/Engine/Source/Runtime/Engine/Classes/Sound/ReverbEffect.h
+++ b/Engine/Source/Runtime/Engine/Classes/Sound/ReverbEffect.h
@@ -7,9 +7,7 @@
 UCLASS(BlueprintType, hidecategories=object, MinimalAPI)
 class UReverbEffect : public UObject
 {
-	GENERATED_BODY()
-public:
-	ENGINE_API UReverbEffect(const FObjectInitializer& ObjectInitializer = FObjectInitializer::Get());
+	GENERATED_UCLASS_BODY()
 
 	/** Density - 0.0 < 1.0 < 1.0 - Coloration of the late reverb - lower value is more */
 	UPROPERTY(Category=ReverbParameters, meta=(ClampMin = "0.0", ClampMax = "1.0"), EditAnywhere)
diff --git a/Engine/Source/Runtime/Engine/Classes/Sound/SoundAttenuation.h b/Engine/Source/Runtime/Engine/Classes/Sound/SoundAttenuation.h
--- a/Engine/Source/Runtime/Engine/Classes/Sound/SoundAttenuation.h
+++ b/Engine/Source/Runtime/Engine/Classes/Sound/SoundAttenuation.h
@@ -43,7 +43,7 @@ The settings for attenuating.
 USTRUCT(BlueprintType)
 struct ENGINE_API FAttenuationSettings
 {
-	GENERATED_BODY()
+	GENERATED_USTRUCT_BODY()
 
 	/* Enable attenuation via volume. */
 	UPROPERTY(EditAnywhere, BlueprintReadWrite, Category=Attenuation )
@@ -162,9 +162,7 @@ struct TStructOpsTypeTraits<FAttenuationSettings> : public TStructOpsTypeTraitsB
 UCLASS(BlueprintType, hidecategories=Object, editinlinenew, MinimalAPI)
 class USoundAttenuation : public UObject
 {
-	GENERATED_BODY()
-public:
-	ENGINE_API USoundAttenuation(const FObjectInitializer& ObjectInitializer = FObjectInitializer::Get());
+	GENERATED_UCLASS_BODY()
 
 	UPROPERTY(EditAnywhere, Category=Settings)
 	FAttenuationSettings Attenuation;
diff --git a/Engine/Source/Runtime/Engine/Classes/Sound/SoundBase.h b/Engine/Source/Runtime/Engine/Classes/Sound/SoundBase.h
--- a/Engine/Source/Runtime/Engine/Classes/Sound/SoundBase.h
+++ b/Engine/Source/Runtime/Engine/Classes/Sound/SoundBase.h
@@ -31,9 +31,7 @@ namespace EMaxConcurrentResolutionRule
 UCLASS(config=Engine, hidecategories=Object, abstract, editinlinenew, BlueprintType)
 class ENGINE_API USoundBase : public UObject
 {
-	GENERATED_BODY()
-public:
-	USoundBase(const FObjectInitializer& ObjectInitializer = FObjectInitializer::Get());
+	GENERATED_UCLASS_BODY()
 
 protected:
 	/** Sound class this sound belongs to */
diff --git a/Engine/Source/Runtime/Engine/Classes/Sound/SoundCue.h b/Engine/Source/Runtime/Engine/Classes/Sound/SoundCue.h
--- a/Engine/Source/Runtime/Engine/Classes/Sound/SoundCue.h
+++ b/Engine/Source/Runtime/Engine/Classes/Sound/SoundCue.h
@@ -10,7 +10,7 @@ struct FSoundParseParameters;
 USTRUCT()
 struct FSoundNodeEditorData
 {
-	GENERATED_BODY()
+	GENERATED_USTRUCT_BODY()
 
 	int32 NodePosX;
 
@@ -36,9 +36,7 @@ struct FSoundNodeEditorData
 UCLASS(hidecategories=object, MinimalAPI, BlueprintType)
 class USoundCue : public USoundBase
 {
-	GENERATED_BODY()
-public:
-	ENGINE_API USoundCue(const FObjectInitializer& ObjectInitializer = FObjectInitializer::Get());
+	GENERATED_UCLASS_BODY()
 
 	/* Indicates whether attenuation should use the Attenuation Overrides or the Attenuation Settings asset */
 	UPROPERTY(EditAnywhere, Category=Attenuation)
diff --git a/Engine/Source/Runtime/Engine/Classes/Sound/SoundGroups.h b/Engine/Source/Runtime/Engine/Classes/Sound/SoundGroups.h
--- a/Engine/Source/Runtime/Engine/Classes/Sound/SoundGroups.h
+++ b/Engine/Source/Runtime/Engine/Classes/Sound/SoundGroups.h
@@ -37,7 +37,7 @@ enum ESoundGroup
 USTRUCT()
 struct FSoundGroup
 {
-	GENERATED_BODY()
+	GENERATED_USTRUCT_BODY()
 
 	// The sound group enumeration we are setting values for
 	UPROPERTY(config)
@@ -71,9 +71,7 @@ struct FSoundGroup
 UCLASS(config=Engine, abstract)
 class USoundGroups : public UObject
 {
-	GENERATED_BODY()
-public:
-	USoundGroups(const FObjectInitializer& ObjectInitializer = FObjectInitializer::Get());
+	GENERATED_UCLASS_BODY()
 
 	void Initialize() const;
 
diff --git a/Engine/Source/Runtime/Engine/Classes/Sound/SoundMix.h b/Engine/Source/Runtime/Engine/Classes/Sound/SoundMix.h
--- a/Engine/Source/Runtime/Engine/Classes/Sound/SoundMix.h
+++ b/Engine/Source/Runtime/Engine/Classes/Sound/SoundMix.h
@@ -7,7 +7,7 @@
 USTRUCT()
 struct FAudioEQEffect
 {
-	GENERATED_BODY()
+	GENERATED_USTRUCT_BODY()
 
 	/* Start time of effect */
 	double RootTime;
@@ -73,7 +73,7 @@ struct FAudioEQEffect
 USTRUCT()
 struct FSoundClassAdjuster
 {
-	GENERATED_BODY()
+	GENERATED_USTRUCT_BODY()
 
 	/* The sound class this adjuster affects. */
 	UPROPERTY(EditAnywhere, Category=SoundClassAdjuster, DisplayName = "Sound Class" )
@@ -111,9 +111,7 @@ struct FSoundClassAdjuster
 UCLASS(BlueprintType, hidecategories=object, MinimalAPI)
 class USoundMix : public UObject
 {
-	GENERATED_BODY()
-public:
-	ENGINE_API USoundMix(const FObjectInitializer& ObjectInitializer = FObjectInitializer::Get());
+	GENERATED_UCLASS_BODY()
 
 	/** Whether to apply the EQ effect */
 	UPROPERTY(EditAnywhere, Category=EQ, AssetRegistrySearchable )
diff --git a/Engine/Source/Runtime/Engine/Classes/Sound/SoundNode.h b/Engine/Source/Runtime/Engine/Classes/Sound/SoundNode.h
--- a/Engine/Source/Runtime/Engine/Classes/Sound/SoundNode.h
+++ b/Engine/Source/Runtime/Engine/Classes/Sound/SoundNode.h
@@ -43,9 +43,7 @@ struct FActiveSound;
 UCLASS(abstract, hidecategories=Object, editinlinenew)
 class ENGINE_API USoundNode : public UObject
 {
-	GENERATED_BODY()
-public:
-	USoundNode(const FObjectInitializer& ObjectInitializer = FObjectInitializer::Get());
+	GENERATED_UCLASS_BODY()
 
 	static const int32 MAX_ALLOWED_CHILD_NODES = 32;
 
diff --git a/Engine/Source/Runtime/Engine/Classes/Sound/SoundNodeAttenuation.h b/Engine/Source/Runtime/Engine/Classes/Sound/SoundNodeAttenuation.h
--- a/Engine/Source/Runtime/Engine/Classes/Sound/SoundNodeAttenuation.h
+++ b/Engine/Source/Runtime/Engine/Classes/Sound/SoundNodeAttenuation.h
@@ -15,9 +15,7 @@ class USoundAttenuation;
 UCLASS(hidecategories=Object, editinlinenew, MinimalAPI, meta=( DisplayName="Attenuation" ))
 class USoundNodeAttenuation : public USoundNode
 {
-	GENERATED_BODY()
-public:
-	ENGINE_API USoundNodeAttenuation(const FObjectInitializer& ObjectInitializer = FObjectInitializer::Get());
+	GENERATED_UCLASS_BODY()
 
 	UPROPERTY(EditAnywhere, Category=Attenuation, meta=(EditCondition="!bOverrideAttenuation"))
 	USoundAttenuation* AttenuationSettings;
diff --git a/Engine/Source/Runtime/Engine/Classes/Sound/SoundNodeBranch.h b/Engine/Source/Runtime/Engine/Classes/Sound/SoundNodeBranch.h
--- a/Engine/Source/Runtime/Engine/Classes/Sound/SoundNodeBranch.h
+++ b/Engine/Source/Runtime/Engine/Classes/Sound/SoundNodeBranch.h
@@ -10,9 +10,7 @@
 UCLASS(hidecategories=Object, editinlinenew, MinimalAPI, meta=( DisplayName="Branch" ))
 class USoundNodeBranch : public USoundNode
 {
-	GENERATED_BODY()
-public:
-	ENGINE_API USoundNodeBranch(const FObjectInitializer& ObjectInitializer = FObjectInitializer::Get());
+	GENERATED_UCLASS_BODY()
 
 	/** The name of the boolean parameter to use to determine which branch we should take */
 	UPROPERTY(EditAnywhere, Category=Branch)
diff --git a/Engine/Source/Runtime/Engine/Classes/Sound/SoundNodeConcatenator.h b/Engine/Source/Runtime/Engine/Classes/Sound/SoundNodeConcatenator.h
--- a/Engine/Source/Runtime/Engine/Classes/Sound/SoundNodeConcatenator.h
+++ b/Engine/Source/Runtime/Engine/Classes/Sound/SoundNodeConcatenator.h
@@ -16,9 +16,7 @@
 UCLASS(hidecategories = Object, editinlinenew, MinimalAPI, meta = (DisplayName = "Concatenator"))
 class USoundNodeConcatenator : public USoundNode
 {
-	GENERATED_BODY()
-public:
-	ENGINE_API USoundNodeConcatenator(const FObjectInitializer& ObjectInitializer = FObjectInitializer::Get());
+	GENERATED_UCLASS_BODY()
 
 	/** Volume multiplier for each input. */
 	UPROPERTY(EditAnywhere, editfixedsize, Category=Concatenator)
diff --git a/Engine/Source/Runtime/Engine/Classes/Sound/SoundNodeDelay.h b/Engine/Source/Runtime/Engine/Classes/Sound/SoundNodeDelay.h
--- a/Engine/Source/Runtime/Engine/Classes/Sound/SoundNodeDelay.h
+++ b/Engine/Source/Runtime/Engine/Classes/Sound/SoundNodeDelay.h
@@ -10,9 +10,7 @@
 UCLASS(hidecategories=Object, editinlinenew, MinimalAPI, meta=( DisplayName="Delay" ))
 class USoundNodeDelay : public USoundNode
 {
-	GENERATED_BODY()
-public:
-	ENGINE_API USoundNodeDelay(const FObjectInitializer& ObjectInitializer = FObjectInitializer::Get());
+	GENERATED_UCLASS_BODY()
 
 	/* The lower bound of delay time in seconds. */
 	UPROPERTY(EditAnywhere, Category=Delay )
diff --git a/Engine/Source/Runtime/Engine/Classes/Sound/SoundNodeDialoguePlayer.h b/Engine/Source/Runtime/Engine/Classes/Sound/SoundNodeDialoguePlayer.h
--- a/Engine/Source/Runtime/Engine/Classes/Sound/SoundNodeDialoguePlayer.h
+++ b/Engine/Source/Runtime/Engine/Classes/Sound/SoundNodeDialoguePlayer.h
@@ -11,9 +11,7 @@
 UCLASS(hidecategories = Object, editinlinenew, MinimalAPI, meta = (DisplayName = "Dialogue Player"))
 class USoundNodeDialoguePlayer : public USoundNode
 {
-	GENERATED_BODY()
-public:
-	ENGINE_API USoundNodeDialoguePlayer(const FObjectInitializer& ObjectInitializer = FObjectInitializer::Get());
+	GENERATED_UCLASS_BODY()
 
 	UPROPERTY(EditAnywhere, Category=DialoguePlayer)
 	FDialogueWaveParameter DialogueWaveParameter;
diff --git a/Engine/Source/Runtime/Engine/Classes/Sound/SoundNodeDistanceCrossFade.h b/Engine/Source/Runtime/Engine/Classes/Sound/SoundNodeDistanceCrossFade.h
--- a/Engine/Source/Runtime/Engine/Classes/Sound/SoundNodeDistanceCrossFade.h
+++ b/Engine/Source/Runtime/Engine/Classes/Sound/SoundNodeDistanceCrossFade.h
@@ -8,7 +8,7 @@
 USTRUCT()
 struct FDistanceDatum
 {
-	GENERATED_BODY()
+	GENERATED_USTRUCT_BODY()
 
 	/* The FadeInDistance at which to start hearing this sound.
 	 * If you want to hear the sound up close then setting this to 0 might be a good option. 
@@ -73,9 +73,7 @@ struct FDistanceDatum
 UCLASS(hidecategories=Object, editinlinenew, MinimalAPI, meta=( DisplayName="Crossfade by Distance" ))
 class USoundNodeDistanceCrossFade : public USoundNode
 {
-	GENERATED_BODY()
-public:
-	ENGINE_API USoundNodeDistanceCrossFade(const FObjectInitializer& ObjectInitializer = FObjectInitializer::Get());
+	GENERATED_UCLASS_BODY()
 
 	/**
 	 * Each input needs to have the correct data filled in so the SoundNodeDistanceCrossFade is able
diff --git a/Engine/Source/Runtime/Engine/Classes/Sound/SoundNodeDoppler.h b/Engine/Source/Runtime/Engine/Classes/Sound/SoundNodeDoppler.h
--- a/Engine/Source/Runtime/Engine/Classes/Sound/SoundNodeDoppler.h
+++ b/Engine/Source/Runtime/Engine/Classes/Sound/SoundNodeDoppler.h
@@ -12,9 +12,7 @@ struct FListener;
 UCLASS(hidecategories=Object, editinlinenew, meta=( DisplayName="Doppler" ))
 class USoundNodeDoppler : public USoundNode
 {
-	GENERATED_BODY()
-public:
-	USoundNodeDoppler(const FObjectInitializer& ObjectInitializer = FObjectInitializer::Get());
+	GENERATED_UCLASS_BODY()
 
 	/* How much to scale the doppler shift (1.0 is normal). */
 	UPROPERTY(EditAnywhere, Category=Doppler )
diff --git a/Engine/Source/Runtime/Engine/Classes/Sound/SoundNodeGroupControl.h b/Engine/Source/Runtime/Engine/Classes/Sound/SoundNodeGroupControl.h
--- a/Engine/Source/Runtime/Engine/Classes/Sound/SoundNodeGroupControl.h
+++ b/Engine/Source/Runtime/Engine/Classes/Sound/SoundNodeGroupControl.h
@@ -11,9 +11,7 @@
 UCLASS(hidecategories=Object, editinlinenew, MinimalAPI, meta=( DisplayName="Group Control" ))
 class USoundNodeGroupControl : public USoundNode
 {
-	GENERATED_BODY()
-public:
-	ENGINE_API USoundNodeGroupControl(const FObjectInitializer& ObjectInitializer = FObjectInitializer::Get());
+	GENERATED_UCLASS_BODY()
 
 	/* How many active sounds are allowed for each group */
 	UPROPERTY(EditAnywhere, editfixedsize, Category=GroupControl)
diff --git a/Engine/Source/Runtime/Engine/Classes/Sound/SoundNodeLooping.h b/Engine/Source/Runtime/Engine/Classes/Sound/SoundNodeLooping.h
--- a/Engine/Source/Runtime/Engine/Classes/Sound/SoundNodeLooping.h
+++ b/Engine/Source/Runtime/Engine/Classes/Sound/SoundNodeLooping.h
@@ -11,9 +11,7 @@
 UCLASS(hidecategories=Object, editinlinenew, MinimalAPI, meta=( DisplayName="Looping" ))
 class USoundNodeLooping : public USoundNode
 {
-	GENERATED_BODY()
-public:
-	ENGINE_API USoundNodeLooping(const FObjectInitializer& ObjectInitializer = FObjectInitializer::Get());
+	GENERATED_UCLASS_BODY()
 
 	/* The amount of times to loop */
 	UPROPERTY(EditAnywhere, Category = Looping, meta = (ClampMin = 1))
diff --git a/Engine/Source/Runtime/Engine/Classes/Sound/SoundNodeMature.h b/Engine/Source/Runtime/Engine/Classes/Sound/SoundNodeMature.h
--- a/Engine/Source/Runtime/Engine/Classes/Sound/SoundNodeMature.h
+++ b/Engine/Source/Runtime/Engine/Classes/Sound/SoundNodeMature.h
@@ -12,9 +12,7 @@
 UCLASS(hidecategories=Object, editinlinenew, MinimalAPI, meta=( DisplayName="Mature" ))
 class USoundNodeMature : public USoundNode
 {
-	GENERATED_BODY()
-public:
-	ENGINE_API USoundNodeMature(const FObjectInitializer& ObjectInitializer = FObjectInitializer::Get());
+	GENERATED_UCLASS_BODY()
 
 
 public:
diff --git a/Engine/Source/Runtime/Engine/Classes/Sound/SoundNodeMixer.h b/Engine/Source/Runtime/Engine/Classes/Sound/SoundNodeMixer.h
--- a/Engine/Source/Runtime/Engine/Classes/Sound/SoundNodeMixer.h
+++ b/Engine/Source/Runtime/Engine/Classes/Sound/SoundNodeMixer.h
@@ -11,9 +11,7 @@
 UCLASS(hidecategories=Object, editinlinenew, MinimalAPI, meta=( DisplayName="Mixer" ))
 class USoundNodeMixer : public USoundNode
 {
-	GENERATED_BODY()
-public:
-	ENGINE_API USoundNodeMixer(const FObjectInitializer& ObjectInitializer = FObjectInitializer::Get());
+	GENERATED_UCLASS_BODY()
 
 	/** A volume for each input.  Automatically sized. */
 	UPROPERTY(EditAnywhere, export, editfixedsize, Category=Mixer)
diff --git a/Engine/Source/Runtime/Engine/Classes/Sound/SoundNodeModulator.h b/Engine/Source/Runtime/Engine/Classes/Sound/SoundNodeModulator.h
--- a/Engine/Source/Runtime/Engine/Classes/Sound/SoundNodeModulator.h
+++ b/Engine/Source/Runtime/Engine/Classes/Sound/SoundNodeModulator.h
@@ -11,9 +11,7 @@
 UCLASS(hidecategories=Object, editinlinenew, MinimalAPI, meta=( DisplayName="Modulator" ))
 class USoundNodeModulator : public USoundNode
 {
-	GENERATED_BODY()
-public:
-	ENGINE_API USoundNodeModulator(const FObjectInitializer& ObjectInitializer = FObjectInitializer::Get());
+	GENERATED_UCLASS_BODY()
 
 	/* The lower bound of pitch (1.0 is no change). */
 	UPROPERTY(EditAnywhere, Category=Modulation )
diff --git a/Engine/Source/Runtime/Engine/Classes/Sound/SoundNodeModulatorContinuous.h b/Engine/Source/Runtime/Engine/Classes/Sound/SoundNodeModulatorContinuous.h
--- a/Engine/Source/Runtime/Engine/Classes/Sound/SoundNodeModulatorContinuous.h
+++ b/Engine/Source/Runtime/Engine/Classes/Sound/SoundNodeModulatorContinuous.h
@@ -17,7 +17,7 @@ enum ModulationParamMode
 USTRUCT()
 struct FModulatorContinuousParams
 {
-	GENERATED_BODY()
+	GENERATED_USTRUCT_BODY()
 
 	FModulatorContinuousParams()
 		: Default(1.f)
@@ -66,9 +66,7 @@ struct FModulatorContinuousParams
 UCLASS(hidecategories=Object, editinlinenew, meta=( DisplayName="Continuous Modulator" ))
 class USoundNodeModulatorContinuous : public USoundNode
 {
-	GENERATED_BODY()
-public:
-	USoundNodeModulatorContinuous(const FObjectInitializer& ObjectInitializer = FObjectInitializer::Get());
+	GENERATED_UCLASS_BODY()
 
 	UPROPERTY(EditAnywhere, Category=ContinuousModulator)
 	FModulatorContinuousParams PitchModulationParams;
diff --git a/Engine/Source/Runtime/Engine/Classes/Sound/SoundNodeOscillator.h b/Engine/Source/Runtime/Engine/Classes/Sound/SoundNodeOscillator.h
--- a/Engine/Source/Runtime/Engine/Classes/Sound/SoundNodeOscillator.h
+++ b/Engine/Source/Runtime/Engine/Classes/Sound/SoundNodeOscillator.h
@@ -11,9 +11,7 @@
 UCLASS(hidecategories=Object, editinlinenew, meta=( DisplayName="Oscillator" ))
 class USoundNodeOscillator : public USoundNode
 {
-	GENERATED_BODY()
-public:
-	USoundNodeOscillator(const FObjectInitializer& ObjectInitializer = FObjectInitializer::Get());
+	GENERATED_UCLASS_BODY()
 
 	/* Whether to oscillate volume. */
 	UPROPERTY(EditAnywhere, Category=Oscillator )
diff --git a/Engine/Source/Runtime/Engine/Classes/Sound/SoundNodeRandom.h b/Engine/Source/Runtime/Engine/Classes/Sound/SoundNodeRandom.h
--- a/Engine/Source/Runtime/Engine/Classes/Sound/SoundNodeRandom.h
+++ b/Engine/Source/Runtime/Engine/Classes/Sound/SoundNodeRandom.h
@@ -11,9 +11,7 @@
 UCLASS(hidecategories=Object, editinlinenew, MinimalAPI, meta=( DisplayName="Random" ))
 class USoundNodeRandom : public USoundNode
 {
-	GENERATED_BODY()
-public:
-	ENGINE_API USoundNodeRandom(const FObjectInitializer& ObjectInitializer = FObjectInitializer::Get());
+	GENERATED_UCLASS_BODY()
 
 	UPROPERTY(EditAnywhere, editfixedsize, Category=Random)
 	TArray<float> Weights;
diff --git a/Engine/Source/Runtime/Engine/Classes/Sound/SoundNodeSwitch.h b/Engine/Source/Runtime/Engine/Classes/Sound/SoundNodeSwitch.h
--- a/Engine/Source/Runtime/Engine/Classes/Sound/SoundNodeSwitch.h
+++ b/Engine/Source/Runtime/Engine/Classes/Sound/SoundNodeSwitch.h
@@ -10,9 +10,7 @@
 UCLASS(hidecategories=Object, editinlinenew, MinimalAPI, meta=(DisplayName="Switch"))
 class USoundNodeSwitch : public USoundNode
 {
-	GENERATED_BODY()
-public:
-	ENGINE_API USoundNodeSwitch(const FObjectInitializer& ObjectInitializer = FObjectInitializer::Get());
+	GENERATED_UCLASS_BODY()
 
 	/** The name of the integer parameter to use to determine which branch we should take */
 	UPROPERTY(EditAnywhere, Category=Switch)
diff --git a/Engine/Source/Runtime/Engine/Classes/Sound/SoundNodeWaveParam.h b/Engine/Source/Runtime/Engine/Classes/Sound/SoundNodeWaveParam.h
--- a/Engine/Source/Runtime/Engine/Classes/Sound/SoundNodeWaveParam.h
+++ b/Engine/Source/Runtime/Engine/Classes/Sound/SoundNodeWaveParam.h
@@ -10,9 +10,7 @@
 UCLASS(hidecategories=Object, editinlinenew, MinimalAPI, meta=( DisplayName="Wave Param" ))
 class USoundNodeWaveParam : public USoundNode
 {
-	GENERATED_BODY()
-public:
-	ENGINE_API USoundNodeWaveParam(const FObjectInitializer& ObjectInitializer = FObjectInitializer::Get());
+	GENERATED_UCLASS_BODY()
 
 	/** The name of the wave parameter to use to look up the SoundWave we should play */
 	UPROPERTY(EditAnywhere, Category=WaveParam)
diff --git a/Engine/Source/Runtime/Engine/Classes/Sound/SoundNodeWavePlayer.h b/Engine/Source/Runtime/Engine/Classes/Sound/SoundNodeWavePlayer.h
--- a/Engine/Source/Runtime/Engine/Classes/Sound/SoundNodeWavePlayer.h
+++ b/Engine/Source/Runtime/Engine/Classes/Sound/SoundNodeWavePlayer.h
@@ -11,9 +11,7 @@
 UCLASS(hidecategories=Object, editinlinenew, MinimalAPI, meta=( DisplayName="Wave Player" ))
 class USoundNodeWavePlayer : public USoundNode
 {
-	GENERATED_BODY()
-public:
-	ENGINE_API USoundNodeWavePlayer(const FObjectInitializer& ObjectInitializer = FObjectInitializer::Get());
+	GENERATED_UCLASS_BODY()
 
 	UPROPERTY(EditAnywhere, Category=WavePlayer)
 	USoundWave* SoundWave;
diff --git a/Engine/Source/Runtime/Engine/Classes/Sound/SoundWave.h b/Engine/Source/Runtime/Engine/Classes/Sound/SoundWave.h
--- a/Engine/Source/Runtime/Engine/Classes/Sound/SoundWave.h
+++ b/Engine/Source/Runtime/Engine/Classes/Sound/SoundWave.h
@@ -64,7 +64,7 @@ struct FStreamedAudioChunk
 USTRUCT()
 struct FStreamedAudioPlatformData
 {
-	GENERATED_BODY()
+	GENERATED_USTRUCT_BODY()
 
 	/** Number of audio chunks. */
 	int32 NumChunks;
@@ -111,9 +111,7 @@ struct FStreamedAudioPlatformData
 UCLASS(hidecategories=Object, editinlinenew, BlueprintType)
 class ENGINE_API USoundWave : public USoundBase
 {
-	GENERATED_BODY()
-public:
-	USoundWave(const FObjectInitializer& ObjectInitializer = FObjectInitializer::Get());
+	GENERATED_UCLASS_BODY()
 
 	/** Platform agnostic compression quality. 1..100 with 1 being best compression and 100 being best quality. */
 	UPROPERTY(EditAnywhere, Category=Compression, meta=(ClampMin = "1", ClampMax = "100"), AssetRegistrySearchable)
diff --git a/Engine/Source/Runtime/Engine/Classes/Sound/SoundWaveStreaming.h b/Engine/Source/Runtime/Engine/Classes/Sound/SoundWaveStreaming.h
--- a/Engine/Source/Runtime/Engine/Classes/Sound/SoundWaveStreaming.h
+++ b/Engine/Source/Runtime/Engine/Classes/Sound/SoundWaveStreaming.h
@@ -14,9 +14,7 @@ DECLARE_DELEGATE_TwoParams( FOnSoundWaveStreamingUnderflow, class USoundWaveStre
 UCLASS()
 class ENGINE_API USoundWaveStreaming : public USoundWave
 {
-	GENERATED_BODY()
-public:
-	USoundWaveStreaming(const FObjectInitializer& ObjectInitializer = FObjectInitializer::Get());
+	GENERATED_UCLASS_BODY()
 
 private:
 	TArray<uint8> QueuedAudio;
@@ -50,4 +48,4 @@ public:
 	/** Called when GeneratePCMData is called but not enough data is available. Allows more data to be added, and will try again */
 	FOnSoundWaveStreamingUnderflow OnSoundWaveStreamingUnderflow;
 
-};
+};
\ No newline at end of file
diff --git a/Engine/Source/Runtime/Engine/Classes/Tests/AutomationTestSettings.h b/Engine/Source/Runtime/Engine/Classes/Tests/AutomationTestSettings.h
--- a/Engine/Source/Runtime/Engine/Classes/Tests/AutomationTestSettings.h
+++ b/Engine/Source/Runtime/Engine/Classes/Tests/AutomationTestSettings.h
@@ -15,7 +15,7 @@
 USTRUCT()
 struct FOpenTestAsset
 {
-	GENERATED_BODY()
+	GENERATED_USTRUCT_BODY()
 
 	/**
 	 * Asset reference
@@ -34,7 +34,7 @@ struct FOpenTestAsset
 USTRUCT()
 struct FExternalToolDefinition
 {
-	GENERATED_BODY()
+	GENERATED_USTRUCT_BODY()
 
 	/* The name of the tool / test. */
 	UPROPERTY(config, EditAnywhere, Category=ExternalTools)
@@ -67,7 +67,7 @@ struct FExternalToolDefinition
 USTRUCT()
 struct FImportFactorySettingValues
 {
-	GENERATED_BODY()
+	GENERATED_USTRUCT_BODY()
 
 	/* The name of the UProperty to change */
 	UPROPERTY(config, EditAnywhere, Category = Automation, meta = (ToolTip = "Name of the property to change.  Nested settings can be modified using \"Outer.Property\""))
@@ -85,7 +85,7 @@ struct FImportFactorySettingValues
 USTRUCT()
 struct FEditorImportExportTestDefinition
 {
-	GENERATED_BODY()
+	GENERATED_USTRUCT_BODY()
 
 	/* The file to import */
 	UPROPERTY(config, EditAnywhere, Category = Automation, meta = (FilePathFilter = "*"))
@@ -110,7 +110,7 @@ struct FEditorImportExportTestDefinition
 USTRUCT()
 struct FEditorImportWorkflowDefinition
 {
-	GENERATED_BODY()
+	GENERATED_USTRUCT_BODY()
 
 	/* The file to import */
 	UPROPERTY(config, EditAnywhere, Category = Automation, meta = (FilePathFilter = "*"))
@@ -127,7 +127,7 @@ struct FEditorImportWorkflowDefinition
 USTRUCT()
 struct FBuildPromotionImportWorkflowSettings
 {
-	GENERATED_BODY()
+	GENERATED_USTRUCT_BODY()
 
 	/* Import settings for the Diffuse texture */
 	UPROPERTY(config, EditAnywhere, Category = Automation)
@@ -180,7 +180,7 @@ struct FBuildPromotionImportWorkflowSettings
 USTRUCT()
 struct FBuildPromotionOpenAssetSettings
 {
-	GENERATED_BODY()
+	GENERATED_USTRUCT_BODY()
 
 	/* The blueprint asset to open */
 	UPROPERTY(config, EditAnywhere, Category = Automation, meta = (FilePathFilter = "uasset"))
@@ -213,7 +213,7 @@ struct FBuildPromotionOpenAssetSettings
 USTRUCT()
 struct FBuildPromotionBlueprintSettings
 {
-	GENERATED_BODY()
+	GENERATED_USTRUCT_BODY()
 
 	/** The starting mesh for the blueprint **/
 	UPROPERTY(EditAnywhere, Category = Automation, meta = (FilePathFilter = "uasset"))
@@ -230,7 +230,7 @@ struct FBuildPromotionBlueprintSettings
 USTRUCT()
 struct FBuildPromotionNewProjectSettings
 {
-	GENERATED_BODY()
+	GENERATED_USTRUCT_BODY()
 
 	/** The path for the new project */
 	UPROPERTY(EditAnywhere, Category = Automation)
@@ -247,7 +247,7 @@ struct FBuildPromotionNewProjectSettings
 USTRUCT()
 struct FBuildPromotionTestSettings
 {
-	GENERATED_BODY()
+	GENERATED_USTRUCT_BODY()
 
 	/** Default static mesh asset to apply materials to **/
 	UPROPERTY(EditAnywhere, Category = Automation)
@@ -280,7 +280,7 @@ struct FBuildPromotionTestSettings
 USTRUCT()
 struct FEditorMapPerformanceTestDefinition
 {
-	GENERATED_BODY()
+	GENERATED_USTRUCT_BODY()
 	
 	/** Map to be used for the Performance Capture **/
 	UPROPERTY(config, EditAnywhere, Category = Automation, meta = (FilePathFilter = "umap"))
@@ -297,7 +297,7 @@ struct FEditorMapPerformanceTestDefinition
 USTRUCT()
 struct FLaunchOnTestSettings
 {
-	GENERATED_BODY()
+	GENERATED_USTRUCT_BODY()
 
 	/** Map to be used for the Launch On test **/
 	UPROPERTY(config, EditAnywhere, Category = Automation, meta = (FilePathFilter = "umap"))
@@ -314,9 +314,7 @@ struct FLaunchOnTestSettings
 UCLASS(config=Engine)
 class ENGINE_API UAutomationTestSettings : public UObject
 {
-	GENERATED_BODY()
-public:
-	UAutomationTestSettings(const FObjectInitializer& ObjectInitializer = FObjectInitializer::Get());
+	GENERATED_UCLASS_BODY()
 
 public:
 
diff --git a/Engine/Source/Runtime/Engine/Classes/Tests/TextPropertyTestObject.h b/Engine/Source/Runtime/Engine/Classes/Tests/TextPropertyTestObject.h
--- a/Engine/Source/Runtime/Engine/Classes/Tests/TextPropertyTestObject.h
+++ b/Engine/Source/Runtime/Engine/Classes/Tests/TextPropertyTestObject.h
@@ -10,9 +10,7 @@
 UCLASS()
 class UTextPropertyTestObject : public UObject
 {
-	GENERATED_BODY()
-public:
-	UTextPropertyTestObject(const FObjectInitializer& ObjectInitializer = FObjectInitializer::Get());
+	GENERATED_UCLASS_BODY()
 
 public:
 	UPROPERTY()
@@ -21,4 +19,4 @@ public:
 	FText UndefaultedText;
 	UPROPERTY()
 	FText TransientText;
-};
+};
\ No newline at end of file
diff --git a/Engine/Source/Runtime/Engine/Classes/VectorField/VectorField.h b/Engine/Source/Runtime/Engine/Classes/VectorField/VectorField.h
--- a/Engine/Source/Runtime/Engine/Classes/VectorField/VectorField.h
+++ b/Engine/Source/Runtime/Engine/Classes/VectorField/VectorField.h
@@ -11,9 +11,7 @@
 UCLASS(MinimalAPI)
 class UVectorField : public UObject
 {
-	GENERATED_BODY()
-public:
-	ENGINE_API UVectorField(const FObjectInitializer& ObjectInitializer = FObjectInitializer::Get());
+	GENERATED_UCLASS_BODY()
 
 	/** Bounds of the volume in local space. */
 	UPROPERTY(EditAnywhere, Category=VectorFieldBounds)
diff --git a/Engine/Source/Runtime/Engine/Classes/VectorField/VectorFieldAnimated.h b/Engine/Source/Runtime/Engine/Classes/VectorField/VectorFieldAnimated.h
--- a/Engine/Source/Runtime/Engine/Classes/VectorField/VectorFieldAnimated.h
+++ b/Engine/Source/Runtime/Engine/Classes/VectorField/VectorFieldAnimated.h
@@ -22,9 +22,7 @@ enum EVectorFieldConstructionOp
 UCLASS(MinimalAPI)
 class UVectorFieldAnimated : public UVectorField
 {
-	GENERATED_BODY()
-public:
-	ENGINE_API UVectorFieldAnimated(const FObjectInitializer& ObjectInitializer = FObjectInitializer::Get());
+	GENERATED_UCLASS_BODY()
 
 	/** The texture from which to create the vector field. */
 	UPROPERTY(EditAnywhere, Category=Reconstruction)
diff --git a/Engine/Source/Runtime/Engine/Classes/VectorField/VectorFieldStatic.h b/Engine/Source/Runtime/Engine/Classes/VectorField/VectorFieldStatic.h
--- a/Engine/Source/Runtime/Engine/Classes/VectorField/VectorFieldStatic.h
+++ b/Engine/Source/Runtime/Engine/Classes/VectorField/VectorFieldStatic.h
@@ -11,9 +11,7 @@
 UCLASS(hidecategories=VectorFieldBounds, MinimalAPI)
 class UVectorFieldStatic : public UVectorField
 {
-	GENERATED_BODY()
-public:
-	ENGINE_API UVectorFieldStatic(const FObjectInitializer& ObjectInitializer = FObjectInitializer::Get());
+	GENERATED_UCLASS_BODY()
 
 	/** Size of the vector field volume. */
 	UPROPERTY(Category=VectorFieldStatic, VisibleAnywhere)
diff --git a/Engine/Source/Runtime/Engine/Classes/Vehicles/TireType.h b/Engine/Source/Runtime/Engine/Classes/Vehicles/TireType.h
--- a/Engine/Source/Runtime/Engine/Classes/Vehicles/TireType.h
+++ b/Engine/Source/Runtime/Engine/Classes/Vehicles/TireType.h
@@ -11,9 +11,7 @@
 UCLASS()
 class ENGINE_API UTireType : public UDataAsset
 {
-	GENERATED_BODY()
-public:
-	UTireType(const FObjectInitializer& ObjectInitializer = FObjectInitializer::Get());
+	GENERATED_UCLASS_BODY()
 
 private:
 
diff --git a/Engine/Source/Runtime/Engine/Classes/Vehicles/VehicleWheel.h b/Engine/Source/Runtime/Engine/Classes/Vehicles/VehicleWheel.h
--- a/Engine/Source/Runtime/Engine/Classes/Vehicles/VehicleWheel.h
+++ b/Engine/Source/Runtime/Engine/Classes/Vehicles/VehicleWheel.h
@@ -19,9 +19,7 @@ namespace physx
 UCLASS(BlueprintType, Blueprintable)
 class ENGINE_API UVehicleWheel : public UObject
 {
-	GENERATED_BODY()
-public:
-	UVehicleWheel(const FObjectInitializer& ObjectInitializer = FObjectInitializer::Get());
+	GENERATED_UCLASS_BODY()
 
 	/** 
 	 * Static mesh with collision setup for wheel, will be used to create wheel shape
diff --git a/Engine/Source/Runtime/Engine/Classes/Vehicles/WheeledVehicleMovementComponent.h b/Engine/Source/Runtime/Engine/Classes/Vehicles/WheeledVehicleMovementComponent.h
--- a/Engine/Source/Runtime/Engine/Classes/Vehicles/WheeledVehicleMovementComponent.h
+++ b/Engine/Source/Runtime/Engine/Classes/Vehicles/WheeledVehicleMovementComponent.h
@@ -86,7 +86,7 @@ struct FTireShaderOutput
 USTRUCT()
 struct ENGINE_API FWheelSetup
 {
-	GENERATED_BODY()
+	GENERATED_USTRUCT_BODY()
 
 	// The wheel class to use
 	UPROPERTY(EditAnywhere, Category=WheelSetup)
@@ -106,7 +106,7 @@ struct ENGINE_API FWheelSetup
 USTRUCT()
 struct ENGINE_API FReplicatedVehicleState
 {
-	GENERATED_BODY()
+	GENERATED_USTRUCT_BODY()
 
 	// input replication: steering
 	UPROPERTY()
@@ -132,7 +132,7 @@ struct ENGINE_API FReplicatedVehicleState
 USTRUCT()
 struct ENGINE_API FVehicleInputRate
 {
-	GENERATED_BODY()
+	GENERATED_USTRUCT_BODY()
 
 	// Rate at which the input value rises
 	UPROPERTY(EditAnywhere, Category=VehicleInputRate)
@@ -161,9 +161,7 @@ struct ENGINE_API FVehicleInputRate
 UCLASS(Abstract, hidecategories=(PlanarMovement, "Components|Movement|Planar", Activation, "Components|Activation"))
 class ENGINE_API UWheeledVehicleMovementComponent : public UPawnMovementComponent, public IRVOAvoidanceInterface
 {
-	GENERATED_BODY()
-public:
-	UWheeledVehicleMovementComponent(const FObjectInitializer& ObjectInitializer = FObjectInitializer::Get());
+	GENERATED_UCLASS_BODY()
 
 	/** Wheels to create */
 	UPROPERTY(EditAnywhere, Category=VehicleSetup)
@@ -500,10 +498,8 @@ protected:
 	void UpdateState(float DeltaTime);
 
 	/** Pass current state to server */
-	UFUNCTION(reliable, server="ServerUpdateState_Implementation", WithValidation="ServerUpdateState_Validate")
+	UFUNCTION(reliable, server, WithValidation)
 	void ServerUpdateState(float InSteeringInput, float InThrottleInput, float InBrakeInput, float InHandbrakeInput, int32 CurrentGear);
-	virtual void ServerUpdateState_Implementation(float InSteeringInput, float InThrottleInput, float InBrakeInput, float InHandbrakeInput, int32 CurrentGear);
-	virtual bool ServerUpdateState_Validate(float , float , float , float , int32 );
 
 	/** Update RVO Avoidance for simulation */
 	void UpdateAvoidance(float DeltaTime);
diff --git a/Engine/Source/Runtime/Engine/Classes/Vehicles/WheeledVehicleMovementComponent4W.h b/Engine/Source/Runtime/Engine/Classes/Vehicles/WheeledVehicleMovementComponent4W.h
--- a/Engine/Source/Runtime/Engine/Classes/Vehicles/WheeledVehicleMovementComponent4W.h
+++ b/Engine/Source/Runtime/Engine/Classes/Vehicles/WheeledVehicleMovementComponent4W.h
@@ -25,7 +25,7 @@ namespace EVehicleDifferential4W
 USTRUCT()
 struct FVehicleDifferential4WData
 {
-	GENERATED_BODY()
+	GENERATED_USTRUCT_BODY()
 
 	/** Type of differential */
 	UPROPERTY(EditAnywhere, Category=Setup)
@@ -59,7 +59,7 @@ struct FVehicleDifferential4WData
 USTRUCT()
 struct FVehicleEngineData
 {
-	GENERATED_BODY()
+	GENERATED_USTRUCT_BODY()
 
 	/** Torque (Nm) at a given RPM*/
 	UPROPERTY(EditAnywhere, Category = Setup)
@@ -93,7 +93,7 @@ struct FVehicleEngineData
 USTRUCT()
 struct FVehicleGearData
 {
-	GENERATED_BODY()
+	GENERATED_USTRUCT_BODY()
 
 	/** Determines the amount of torque multiplication*/
 	UPROPERTY(EditAnywhere, Category = Setup)
@@ -111,7 +111,7 @@ struct FVehicleGearData
 USTRUCT()
 struct FVehicleTransmissionData
 {
-	GENERATED_BODY()
+	GENERATED_USTRUCT_BODY()
 	/** Whether to use automatic transmission */
 	UPROPERTY(EditAnywhere, Category = VehicleSetup, meta=(FriendlyName = "Automatic Transmission"))
 	bool bUseGearAutoBox;
@@ -148,9 +148,7 @@ struct FVehicleTransmissionData
 UCLASS()
 class ENGINE_API UWheeledVehicleMovementComponent4W : public UWheeledVehicleMovementComponent
 {
-	GENERATED_BODY()
-public:
-	UWheeledVehicleMovementComponent4W(const FObjectInitializer& ObjectInitializer = FObjectInitializer::Get());
+	GENERATED_UCLASS_BODY()
 
 	/** Engine */
 	UPROPERTY(EditAnywhere, Category = MechanicalSetup)
diff --git a/Engine/Source/Runtime/Engine/Classes/VisualLogger/VisualLoggerAutomationTests.h b/Engine/Source/Runtime/Engine/Classes/VisualLogger/VisualLoggerAutomationTests.h
--- a/Engine/Source/Runtime/Engine/Classes/VisualLogger/VisualLoggerAutomationTests.h
+++ b/Engine/Source/Runtime/Engine/Classes/VisualLogger/VisualLoggerAutomationTests.h
@@ -8,8 +8,6 @@
 UCLASS(NotBlueprintable, Transient, hidecategories = UObject, notplaceable)
 class UVisualLoggerAutomationTests : public UObject
 {
-	GENERATED_BODY()
-public:
-	UVisualLoggerAutomationTests(const FObjectInitializer& ObjectInitializer = FObjectInitializer::Get());
+	GENERATED_UCLASS_BODY()
 };
 
diff --git a/Engine/Source/Runtime/Engine/Classes/VisualLogger/VisualLoggerKismetLibrary.h b/Engine/Source/Runtime/Engine/Classes/VisualLogger/VisualLoggerKismetLibrary.h
--- a/Engine/Source/Runtime/Engine/Classes/VisualLogger/VisualLoggerKismetLibrary.h
+++ b/Engine/Source/Runtime/Engine/Classes/VisualLogger/VisualLoggerKismetLibrary.h
@@ -7,9 +7,7 @@
 UCLASS(MinimalAPI)
 class UVisualLoggerKismetLibrary : public UBlueprintFunctionLibrary
 {
-	GENERATED_BODY()
-public:
-	ENGINE_API UVisualLoggerKismetLibrary(const FObjectInitializer& ObjectInitializer = FObjectInitializer::Get());
+	GENERATED_UCLASS_BODY()
 
 	/** Logs simple text string with Visual Logger - recording for Visual Logs has to be enabled to record this data */
 	UFUNCTION(BlueprintCallable, Category = "Debug|VisualLogger", meta = (HidePin = "WorldContextObject", DefaultToSelf = "WorldContextObject", FriendlyName = "LogText"))
diff --git a/Engine/Source/Runtime/Engine/Public/AI/RVOAvoidanceInterface.h b/Engine/Source/Runtime/Engine/Public/AI/RVOAvoidanceInterface.h
--- a/Engine/Source/Runtime/Engine/Public/AI/RVOAvoidanceInterface.h
+++ b/Engine/Source/Runtime/Engine/Public/AI/RVOAvoidanceInterface.h
@@ -12,15 +12,12 @@
 UINTERFACE(MinimalAPI, meta=(CannotImplementInterfaceInBlueprint))
 class URVOAvoidanceInterface : public UInterface
 {
-	GENERATED_BODY()
-public:
-	ENGINE_API URVOAvoidanceInterface(const FObjectInitializer& ObjectInitializer = FObjectInitializer::Get());
+	GENERATED_UINTERFACE_BODY()
 };
 
 class ENGINE_API IRVOAvoidanceInterface
 {
-	GENERATED_BODY()
-public:
+	GENERATED_IINTERFACE_BODY()
 
 	/** Store the AvoidanceUID generated by the Avoidance Manager **/
 	virtual void SetRVOAvoidanceUID(int32 UID) = 0;
