https://jira.ol.epicgames.net/browse/UEPLAT-514

Adding Texture LOD settings to Device Profiles
- this will improve how we currently edit and override lod settings on a per device type bases.
- LOD Settings can now be set in the Device Profile Editor
- Any device profiles with no values set for LOD Settings, will use the default objects.
- These can be found in BaseDeviceProfiles.ini in the /Script/Engine.TextureLODSettings section
- Appropriate defaults set for those groups not listed.

-Other fixes
- - Incorrect with_editor check in HTML5 target platform, changed to with_engine

diff --git a/Engine/Config/BaseDeviceProfiles.ini b/Engine/Config/BaseDeviceProfiles.ini
--- a/Engine/Config/BaseDeviceProfiles.ini
+++ b/Engine/Config/BaseDeviceProfiles.ini
@@ -28,6 +28,33 @@
 +DeviceProfileNameAndTypes=LinuxNoEditor,LinuxNoEditor
 +DeviceProfileNameAndTypes=LinuxServer,LinuxServer
 
+[/Script/Engine.TextureLODSettings]
+TextureLODGroups=(Group=TEXTUREGROUP_World,MinLODSize=1,MaxLODSize=4096,LODBias=0,MinMagFilter=aniso,MipFilter=point,MipGenSettings=TMGS_SimpleAverage)
++TextureLODGroups=(Group=TEXTUREGROUP_WorldNormalMap,MinLODSize=1,MaxLODSize=4096,LODBias=0,MinMagFilter=aniso,MipFilter=point,MipGenSettings=TMGS_SimpleAverage)
++TextureLODGroups=(Group=TEXTUREGROUP_WorldSpecular,MinLODSize=1,MaxLODSize=4096,LODBias=0,MinMagFilter=aniso,MipFilter=point,MipGenSettings=TMGS_SimpleAverage)
++TextureLODGroups=(Group=TEXTUREGROUP_Character,MinLODSize=1,MaxLODSize=4096,LODBias=0,MinMagFilter=aniso,MipFilter=point,MipGenSettings=TMGS_SimpleAverage)
++TextureLODGroups=(Group=TEXTUREGROUP_CharacterNormalMap,MinLODSize=1,MaxLODSize=4096,LODBias=0,MinMagFilter=aniso,MipFilter=point,MipGenSettings=TMGS_SimpleAverage)
++TextureLODGroups=(Group=TEXTUREGROUP_CharacterSpecular,MinLODSize=1,MaxLODSize=4096,LODBias=0,MinMagFilter=aniso,MipFilter=point,MipGenSettings=TMGS_SimpleAverage)
++TextureLODGroups=(Group=TEXTUREGROUP_Weapon,MinLODSize=1,MaxLODSize=4096,LODBias=0,MinMagFilter=aniso,MipFilter=point,MipGenSettings=TMGS_SimpleAverage)
++TextureLODGroups=(Group=TEXTUREGROUP_WeaponNormalMap,MinLODSize=1,MaxLODSize=4096,LODBias=0,MinMagFilter=aniso,MipFilter=point,MipGenSettings=TMGS_SimpleAverage)
++TextureLODGroups=(Group=TEXTUREGROUP_WeaponSpecular,MinLODSize=1,MaxLODSize=4096,LODBias=0,MinMagFilter=aniso,MipFilter=point,MipGenSettings=TMGS_SimpleAverage)
++TextureLODGroups=(Group=TEXTUREGROUP_Vehicle,MinLODSize=1,MaxLODSize=4096,LODBias=0,MinMagFilter=aniso,MipFilter=point,MipGenSettings=TMGS_SimpleAverage)
++TextureLODGroups=(Group=TEXTUREGROUP_VehicleNormalMap,MinLODSize=1,MaxLODSize=4096,LODBias=0,MinMagFilter=aniso,MipFilter=point,MipGenSettings=TMGS_SimpleAverage)
++TextureLODGroups=(Group=TEXTUREGROUP_VehicleSpecular,MinLODSize=1,MaxLODSize=4096,LODBias=0,MinMagFilter=aniso,MipFilter=point,MipGenSettings=TMGS_SimpleAverage)
++TextureLODGroups=(Group=TEXTUREGROUP_Cinematic,MinLODSize=1,MaxLODSize=4096,LODBias=0,MinMagFilter=aniso,MipFilter=point,MipGenSettings=TMGS_SimpleAverage)
++TextureLODGroups=(Group=TEXTUREGROUP_Effects,MinLODSize=1,MaxLODSize=4096,LODBias=0,MinMagFilter=linear,MipFilter=point,MipGenSettings=TMGS_SimpleAverage)
++TextureLODGroups=(Group=TEXTUREGROUP_EffectsNotFiltered,MinLODSize=1,MaxLODSize=4096,LODBias=0,MinMagFilter=aniso,MipFilter=point,MipGenSettings=TMGS_SimpleAverage)
++TextureLODGroups=(Group=TEXTUREGROUP_Skybox,MinLODSize=1,MaxLODSize=4096,LODBias=0,MinMagFilter=aniso,MipFilter=point,MipGenSettings=TMGS_SimpleAverage)
++TextureLODGroups=(Group=TEXTUREGROUP_UI,MinLODSize=1,MaxLODSize=4096,LODBias=0,MinMagFilter=aniso,MipFilter=point,MipGenSettings=TMGS_SimpleAverage)
++TextureLODGroups=(Group=TEXTUREGROUP_Lightmap,MinLODSize=1,MaxLODSize=4096,LODBias=0,MinMagFilter=aniso,MipFilter=point,MipGenSettings=TMGS_SimpleAverage)
++TextureLODGroups=(Group=TEXTUREGROUP_Shadowmap,MinLODSize=1,MaxLODSize=4096,LODBias=0,MinMagFilter=aniso,MipFilter=point,NumStreamedMips=3,MipGenSettings=TMGS_SimpleAverage)
++TextureLODGroups=(Group=TEXTUREGROUP_RenderTarget,MinLODSize=1,MaxLODSize=4096,LODBias=0,MinMagFilter=aniso,MipFilter=point,MipGenSettings=TMGS_SimpleAverage)
++TextureLODGroups=(Group=TEXTUREGROUP_MobileFlattened,MinLODSize=1,MaxLODSize=4096,LODBias=0,MinMagFilter=aniso,MipFilter=point,MipGenSettings=TMGS_SimpleAverage)
++TextureLODGroups=(Group=TEXTUREGROUP_Terrain_Heightmap,MinLODSize=1,MaxLODSize=4096,LODBias=0,MinMagFilter=aniso,MipFilter=point,MipGenSettings=TMGS_SimpleAverage)
++TextureLODGroups=(Group=TEXTUREGROUP_Terrain_Weightmap,MinLODSize=1,MaxLODSize=4096,LODBias=0,MinMagFilter=aniso,MipFilter=point,MipGenSettings=TMGS_SimpleAverage)
++TextureLODGroups=(Group=TEXTUREGROUP_Bokeh,MinLODSize=1,MaxLODSize=256,LODBias=0,MinMagFilter=linear,MipFilter=linear,MipGenSettings=TMGS_SimpleAverage)
++TextureLODGroups=(Group=TEXTUREGROUP_Pixels2D,MinLODSize=1,MaxLODSize=4096,LODBias=0,MinMagFilter=point,MipFilter=point,MipGenSettings=TMGS_SimpleAverage)
+
 [Windows DeviceProfile]
 DeviceType=Windows
 BaseProfileName=
diff --git a/Engine/Config/BaseEngine.ini b/Engine/Config/BaseEngine.ini
--- a/Engine/Config/BaseEngine.ini
+++ b/Engine/Config/BaseEngine.ini
@@ -1638,32 +1638,6 @@ FLightPrimitiveInteractionInitialBlockSize=512
 [SystemSettings]
 ; NOTE THAT ANY ITEMS IN THIS SECTION WILL AFFECT ALL PLATFORMS!!!
 
-TEXTUREGROUP_World=(MinLODSize=1,MaxLODSize=4096,LODBias=0,MinMagFilter=aniso,MipFilter=point,MipGenSettings=TMGS_SimpleAverage)
-TEXTUREGROUP_WorldNormalMap=(MinLODSize=1,MaxLODSize=4096,LODBias=0,MinMagFilter=aniso,MipFilter=point,MipGenSettings=TMGS_SimpleAverage)
-TEXTUREGROUP_WorldSpecular=(MinLODSize=1,MaxLODSize=4096,LODBias=0,MinMagFilter=aniso,MipFilter=point,MipGenSettings=TMGS_SimpleAverage)
-TEXTUREGROUP_Character=(MinLODSize=1,MaxLODSize=4096,LODBias=0,MinMagFilter=aniso,MipFilter=point,MipGenSettings=TMGS_SimpleAverage)
-TEXTUREGROUP_CharacterNormalMap=(MinLODSize=1,MaxLODSize=4096,LODBias=0,MinMagFilter=aniso,MipFilter=point,MipGenSettings=TMGS_SimpleAverage)
-TEXTUREGROUP_CharacterSpecular=(MinLODSize=1,MaxLODSize=4096,LODBias=0,MinMagFilter=aniso,MipFilter=point,MipGenSettings=TMGS_SimpleAverage)
-TEXTUREGROUP_Weapon=(MinLODSize=1,MaxLODSize=4096,LODBias=0,MinMagFilter=aniso,MipFilter=point,MipGenSettings=TMGS_SimpleAverage)
-TEXTUREGROUP_WeaponNormalMap=(MinLODSize=1,MaxLODSize=4096,LODBias=0,MinMagFilter=aniso,MipFilter=point,MipGenSettings=TMGS_SimpleAverage)
-TEXTUREGROUP_WeaponSpecular=(MinLODSize=1,MaxLODSize=4096,LODBias=0,MinMagFilter=aniso,MipFilter=point,MipGenSettings=TMGS_SimpleAverage)
-TEXTUREGROUP_Vehicle=(MinLODSize=1,MaxLODSize=4096,LODBias=0,MinMagFilter=aniso,MipFilter=point,MipGenSettings=TMGS_SimpleAverage)
-TEXTUREGROUP_VehicleNormalMap=(MinLODSize=1,MaxLODSize=4096,LODBias=0,MinMagFilter=aniso,MipFilter=point,MipGenSettings=TMGS_SimpleAverage)
-TEXTUREGROUP_VehicleSpecular=(MinLODSize=1,MaxLODSize=4096,LODBias=0,MinMagFilter=aniso,MipFilter=point,MipGenSettings=TMGS_SimpleAverage)
-TEXTUREGROUP_Cinematic=(MinLODSize=1,MaxLODSize=4096,LODBias=0,MinMagFilter=aniso,MipFilter=point,MipGenSettings=TMGS_SimpleAverage)
-TEXTUREGROUP_Effects=(MinLODSize=1,MaxLODSize=4096,LODBias=0,MinMagFilter=linear,MipFilter=point,MipGenSettings=TMGS_SimpleAverage)
-TEXTUREGROUP_EffectsNotFiltered=(MinLODSize=1,MaxLODSize=4096,LODBias=0,MinMagFilter=aniso,MipFilter=point,MipGenSettings=TMGS_SimpleAverage)
-TEXTUREGROUP_Skybox=(MinLODSize=1,MaxLODSize=4096,LODBias=0,MinMagFilter=aniso,MipFilter=point,MipGenSettings=TMGS_SimpleAverage)
-TEXTUREGROUP_UI=(MinLODSize=1,MaxLODSize=4096,LODBias=0,MinMagFilter=aniso,MipFilter=point,MipGenSettings=TMGS_SimpleAverage)
-TEXTUREGROUP_Lightmap=(MinLODSize=1,MaxLODSize=4096,LODBias=0,MinMagFilter=aniso,MipFilter=point,MipGenSettings=TMGS_SimpleAverage)
-TEXTUREGROUP_Shadowmap=(MinLODSize=1,MaxLODSize=4096,LODBias=0,MinMagFilter=aniso,MipFilter=point,NumStreamedMips=3,MipGenSettings=TMGS_SimpleAverage)
-TEXTUREGROUP_RenderTarget=(MinLODSize=1,MaxLODSize=4096,LODBias=0,MinMagFilter=aniso,MipFilter=point,MipGenSettings=TMGS_SimpleAverage)
-TEXTUREGROUP_MobileFlattened=(MinLODSize=1,MaxLODSize=4096,LODBias=0,MinMagFilter=aniso,MipFilter=point,MipGenSettings=TMGS_SimpleAverage)
-TEXTUREGROUP_Terrain_Heightmap=(MinLODSize=1,MaxLODSize=4096,LODBias=0,MinMagFilter=aniso,MipFilter=point,MipGenSettings=TMGS_SimpleAverage)
-TEXTUREGROUP_Terrain_Weightmap=(MinLODSize=1,MaxLODSize=4096,LODBias=0,MinMagFilter=aniso,MipFilter=point,MipGenSettings=TMGS_SimpleAverage)
-TEXTUREGROUP_Bokeh=(MinLODSize=1,MaxLODSize=256,LODBias=0,MinMagFilter=linear,MipFilter=linear,MipGenSettings=TMGS_SimpleAverage)
-TEXTUREGROUP_Pixels2D=(MinLODSize=1,MaxLODSize=4096,LODBias=0,MinMagFilter=point,MipFilter=point,MipGenSettings=TMGS_SimpleAverage)
-
 [SystemSettingsEditor]
 ; System settings overrides for the editor.  Ideally the editor should use the same settings as the game.
 
diff --git a/Engine/Source/Developer/Android/AndroidTargetPlatform/Private/AndroidTargetPlatform.h b/Engine/Source/Developer/Android/AndroidTargetPlatform/Private/AndroidTargetPlatform.h
--- a/Engine/Source/Developer/Android/AndroidTargetPlatform/Private/AndroidTargetPlatform.h
+++ b/Engine/Source/Developer/Android/AndroidTargetPlatform/Private/AndroidTargetPlatform.h
@@ -131,7 +131,12 @@ public:
 
 	virtual void GetTextureFormats( const UTexture* InTexture, TArray<FName>& OutFormats ) const override;
 
-	virtual const struct FTextureLODSettings& GetTextureLODSettings( ) const override;
+	virtual const UTextureLODSettings& GetTextureLODSettings() const override;
+
+	virtual void RegisterTextureLODSettings(const UTextureLODSettings* InTextureLODSettings) override
+	{
+		TextureLODSettings = InTextureLODSettings;
+	}
 
 	virtual FName GetWaveFormat( const class USoundWave* Wave ) const override;
 #endif //WITH_ENGINE
@@ -199,7 +204,7 @@ private:
 
 #if WITH_ENGINE
 	// Holds a cache of the target LOD settings.
-	FTextureLODSettings TextureLODSettings;
+	const UTextureLODSettings* TextureLODSettings;
 
 	// Holds the static mesh LOD settings.
 	FStaticMeshLODSettings StaticMeshLODSettings;
diff --git a/Engine/Source/Developer/Android/AndroidTargetPlatform/Private/AndroidTargetPlatform.inl b/Engine/Source/Developer/Android/AndroidTargetPlatform/Private/AndroidTargetPlatform.inl
--- a/Engine/Source/Developer/Android/AndroidTargetPlatform/Private/AndroidTargetPlatform.inl
+++ b/Engine/Source/Developer/Android/AndroidTargetPlatform/Private/AndroidTargetPlatform.inl
@@ -32,7 +32,7 @@ inline FAndroidTargetPlatform<TPlatformProperties>::FAndroidTargetPlatform( ) :
 {
 	#if WITH_ENGINE
 		FConfigCacheIni::LoadLocalIniFile(EngineSettings, TEXT("Engine"), true, *TTargetPlatformBase<TPlatformProperties>::PlatformName());
-		TextureLODSettings.Initialize(EngineSettings, TEXT("SystemSettings"));
+			TextureLODSettings = nullptr; // These are registered by the device profile system.
 		StaticMeshLODSettings.Initialize(EngineSettings);
 	#endif
 
@@ -282,11 +282,10 @@ void FAndroidTargetPlatform<TPlatformProperties>::GetReflectionCaptureFormats( T
 }
 
 
-
 template<class TPlatformProperties>
-const FTextureLODSettings& FAndroidTargetPlatform<TPlatformProperties>::GetTextureLODSettings( ) const
+const UTextureLODSettings& FAndroidTargetPlatform<TPlatformProperties>::GetTextureLODSettings() const
 {
-	return TextureLODSettings;
+	return *TextureLODSettings;
 }
 
 
diff --git a/Engine/Source/Developer/DesktopTargetPlatform/Private/DesktopTargetPlatform.cpp b/Engine/Source/Developer/DesktopTargetPlatform/Private/DesktopTargetPlatform.cpp
--- a/Engine/Source/Developer/DesktopTargetPlatform/Private/DesktopTargetPlatform.cpp
+++ b/Engine/Source/Developer/DesktopTargetPlatform/Private/DesktopTargetPlatform.cpp
@@ -15,7 +15,6 @@ FDesktopTargetPlatform::FDesktopTargetPlatform()
 #if WITH_ENGINE
 	// use non-platform specific settings
 	FConfigCacheIni::LoadLocalIniFile(EngineSettings, TEXT("Engine"), true, NULL);
-	TextureLODSettings.Initialize(EngineSettings, TEXT("SystemSettings"));
 	StaticMeshLODSettings.Initialize(EngineSettings);
 #endif // #if WITH_ENGINE
 
@@ -62,12 +61,6 @@ void FDesktopTargetPlatform::GetTextureFormats( const UTexture* Texture, TArray<
 }
 
 
-const FTextureLODSettings& FDesktopTargetPlatform::GetTextureLODSettings() const
-{
-	return TextureLODSettings;
-}
-
-
 FName FDesktopTargetPlatform::GetWaveFormat( const class USoundWave* Wave ) const
 {
 	static FName NAME_OGG(TEXT("OGG"));
diff --git a/Engine/Source/Developer/DesktopTargetPlatform/Private/DesktopTargetPlatform.h b/Engine/Source/Developer/DesktopTargetPlatform/Private/DesktopTargetPlatform.h
--- a/Engine/Source/Developer/DesktopTargetPlatform/Private/DesktopTargetPlatform.h
+++ b/Engine/Source/Developer/DesktopTargetPlatform/Private/DesktopTargetPlatform.h
@@ -125,9 +125,17 @@ public:
 		return StaticMeshLODSettings;
 	}
 
-	virtual void GetTextureFormats( const UTexture* Texture, TArray<FName>& OutFormats ) const override;
+	virtual void GetTextureFormats(const UTexture* Texture, TArray<FName>& OutFormats) const override;
+	
+	virtual const UTextureLODSettings& GetTextureLODSettings() const override
+	{
+		return *TextureLODSettings;
+	}
 
-	virtual const struct FTextureLODSettings& GetTextureLODSettings( ) const override;
+	virtual void RegisterTextureLODSettings(const UTextureLODSettings* InTextureLODSettings) override
+	{
+		TextureLODSettings = InTextureLODSettings;
+	}
 
 	virtual FName GetWaveFormat( const class USoundWave* Wave ) const override;
 #endif // WITH_ENGINE
@@ -152,7 +160,7 @@ private:
 	FConfigFile EngineSettings;
 
 	// Holds the cache of the target LOD settings.
-	FTextureLODSettings TextureLODSettings;
+	const UTextureLODSettings* TextureLODSettings;
 
 	// Holds the static mesh LOD settings.
 	FStaticMeshLODSettings StaticMeshLODSettings;
diff --git a/Engine/Source/Developer/HTML5/HTML5TargetPlatform/Private/HTML5TargetPlatform.cpp b/Engine/Source/Developer/HTML5/HTML5TargetPlatform/Private/HTML5TargetPlatform.cpp
--- a/Engine/Source/Developer/HTML5/HTML5TargetPlatform/Private/HTML5TargetPlatform.cpp
+++ b/Engine/Source/Developer/HTML5/HTML5TargetPlatform/Private/HTML5TargetPlatform.cpp
@@ -6,8 +6,8 @@
 
 #include "HTML5TargetPlatformPrivatePCH.h"
 
-#if WITH_EDITOR
-#include "TextureLODSettings.h"
+#if WITH_ENGINE
+#include "DeviceProfiles/DeviceProfile.h"
 #endif 
 
  /* FHTML5TargetPlatform structors
@@ -19,7 +19,7 @@ FHTML5TargetPlatform::FHTML5TargetPlatform( )
 
 #if WITH_ENGINE
 	// load up texture settings from the config file
-	HTML5LODSettings.Initialize(HTML5EngineSettings, TEXT("SystemSettings"));
+	HTML5LODSettings = nullptr;
 	StaticMeshLODSettings.Initialize(HTML5EngineSettings);
 #endif
 }
@@ -309,9 +309,9 @@ void FHTML5TargetPlatform::GetTextureFormats( const UTexture* Texture, TArray<FN
 }
 
 
-const struct FTextureLODSettings& FHTML5TargetPlatform::GetTextureLODSettings( ) const
+const UTextureLODSettings& FHTML5TargetPlatform::GetTextureLODSettings() const
 {
-	return HTML5LODSettings;
+	return *HTML5LODSettings;
 }
 
 
diff --git a/Engine/Source/Developer/HTML5/HTML5TargetPlatform/Private/HTML5TargetPlatform.h b/Engine/Source/Developer/HTML5/HTML5TargetPlatform/Private/HTML5TargetPlatform.h
--- a/Engine/Source/Developer/HTML5/HTML5TargetPlatform/Private/HTML5TargetPlatform.h
+++ b/Engine/Source/Developer/HTML5/HTML5TargetPlatform/Private/HTML5TargetPlatform.h
@@ -70,7 +70,12 @@ public:
 		OutFormats.Add(FName(TEXT("EncodedHDR")));
 	}
 
-	virtual const struct FTextureLODSettings& GetTextureLODSettings( ) const override;
+	virtual const UTextureLODSettings& GetTextureLODSettings() const override;
+
+	virtual void RegisterTextureLODSettings(const UTextureLODSettings* InTextureLODSettings) override
+	{
+		HTML5LODSettings = InTextureLODSettings;
+	}
 
 	virtual FName GetWaveFormat( const class USoundWave* Wave ) const override;
 #endif // WITH_ENGINE
@@ -101,7 +106,7 @@ private:
 
 #if WITH_ENGINE
 	// Holds the cached target LOD settings.
-	FTextureLODSettings HTML5LODSettings;
+	const UTextureLODSettings* HTML5LODSettings;
 
 	// Holds the static mesh LOD settings.
 	FStaticMeshLODSettings StaticMeshLODSettings;
diff --git a/Engine/Source/Developer/Linux/LinuxTargetPlatform/Private/LinuxTargetPlatform.h b/Engine/Source/Developer/Linux/LinuxTargetPlatform/Private/LinuxTargetPlatform.h
--- a/Engine/Source/Developer/Linux/LinuxTargetPlatform/Private/LinuxTargetPlatform.h
+++ b/Engine/Source/Developer/Linux/LinuxTargetPlatform/Private/LinuxTargetPlatform.h
@@ -37,7 +37,7 @@ public:
 	
 #if WITH_ENGINE
 		FConfigCacheIni::LoadLocalIniFile(EngineSettings, TEXT("Engine"), true, *this->PlatformName());
-		TextureLODSettings.Initialize(EngineSettings, TEXT("SystemSettings"));
+		TextureLODSettings = nullptr;
 		StaticMeshLODSettings.Initialize(EngineSettings);
 
 		// Get the Target RHIs for this platform, we do not always want all those that are supported.
@@ -169,9 +169,14 @@ public:
 		}
 	}
 
-	virtual const struct FTextureLODSettings& GetTextureLODSettings( ) const override
+	virtual const UTextureLODSettings& GetTextureLODSettings() const override
 	{
-		return TextureLODSettings;
+		return *TextureLODSettings;
+	}
+
+	virtual void RegisterTextureLODSettings(const UTextureLODSettings* InTextureLODSettings) override
+	{
+		TextureLODSettings = InTextureLODSettings;
 	}
 
 	virtual FName GetWaveFormat( const class USoundWave* Wave ) const override
@@ -247,7 +252,7 @@ private:
 	FConfigFile EngineSettings;
 
 	// Holds the texture LOD settings.
-	FTextureLODSettings TextureLODSettings;
+	const UTextureLODSettings* TextureLODSettings;
 
 	// Holds static mesh LOD settings.
 	FStaticMeshLODSettings StaticMeshLODSettings;
diff --git a/Engine/Source/Developer/Mac/MacTargetPlatform/Private/GenericMacTargetPlatform.h b/Engine/Source/Developer/Mac/MacTargetPlatform/Private/GenericMacTargetPlatform.h
--- a/Engine/Source/Developer/Mac/MacTargetPlatform/Private/GenericMacTargetPlatform.h
+++ b/Engine/Source/Developer/Mac/MacTargetPlatform/Private/GenericMacTargetPlatform.h
@@ -33,7 +33,7 @@ public:
 
 		#if WITH_ENGINE
 			FConfigCacheIni::LoadLocalIniFile(EngineSettings, TEXT("Engine"), true, *this->PlatformName());
-			TextureLODSettings.Initialize(EngineSettings, TEXT("SystemSettings"));
+			TextureLODSettings = nullptr;
 			StaticMeshLODSettings.Initialize(EngineSettings);
 		#endif
 	}
@@ -125,9 +125,14 @@ return TSuper::SupportsFeature(Feature);
 	}
 
 
-	virtual const struct FTextureLODSettings& GetTextureLODSettings( ) const override
+	virtual const UTextureLODSettings& GetTextureLODSettings() const override
 	{
-		return TextureLODSettings;
+		return *TextureLODSettings;
+	}
+
+	virtual void RegisterTextureLODSettings(const UTextureLODSettings* InTextureLODSettings) override
+	{
+		TextureLODSettings = InTextureLODSettings;
 	}
 
 
@@ -203,7 +208,7 @@ private:
 	FConfigFile EngineSettings;
 
 	// Holds the texture LOD settings.
-	FTextureLODSettings TextureLODSettings;
+	const UTextureLODSettings* TextureLODSettings;
 
 	// Holds the static mesh LOD settings.
 	FStaticMeshLODSettings StaticMeshLODSettings;
diff --git a/Engine/Source/Developer/TargetPlatform/Public/Interfaces/ITargetPlatform.h b/Engine/Source/Developer/TargetPlatform/Public/Interfaces/ITargetPlatform.h
--- a/Engine/Source/Developer/TargetPlatform/Public/Interfaces/ITargetPlatform.h
+++ b/Engine/Source/Developer/TargetPlatform/Public/Interfaces/ITargetPlatform.h
@@ -314,7 +314,12 @@ public:
 	 *
 	 * @return A texture LOD settings structure.
 	 */
-	virtual const struct FTextureLODSettings& GetTextureLODSettings() const = 0;
+	virtual const class UTextureLODSettings& GetTextureLODSettings() const = 0;
+
+	/**
+	* Register Basic LOD Settings for this platform
+	*/
+	virtual void RegisterTextureLODSettings(const class UTextureLODSettings* InTextureLODSettings) = 0;
 
 	/**
 	 * Gets the static mesh LOD settings used by this platform.
diff --git a/Engine/Source/Developer/TextureFormatASTC/Private/TextureFormatASTC.cpp b/Engine/Source/Developer/TextureFormatASTC/Private/TextureFormatASTC.cpp
--- a/Engine/Source/Developer/TextureFormatASTC/Private/TextureFormatASTC.cpp
+++ b/Engine/Source/Developer/TextureFormatASTC/Private/TextureFormatASTC.cpp
@@ -306,7 +306,7 @@ class FTextureFormatASTC : public ITextureFormat
 //	// Since we want to have per texture [group] compression settings, we need to have the key based on the texture
 //	virtual FString GetDerivedDataKeyString(const class UTexture& Texture) const override
 //	{
-//		const int32 LODBias = GSystemSettings.TextureLODSettings.CalculateLODBias(Texture.Source.GetSizeX(), Texture.Source.GetSizeY(), Texture.LODGroup, Texture.LODBias, Texture.NumCinematicMipLevels, Texture.MipGenSettings);
+//		const int32 LODBias = UDeviceProfileManager::Get().GetActiveProfile()->GetTextureLODSettings()->CalculateLODBias(Texture.Source.GetSizeX(), Texture.Source.GetSizeY(), Texture.LODGroup, Texture.LODBias, Texture.NumCinematicMipLevels, Texture.MipGenSettings);
 //		check(LODBias >= 0);
 //		return FString::Printf(TEXT("%02u%d_"), (uint32)LODBias, CVarVirtualTextureReducedMemoryEnabled->GetValueOnGameThread());
 //	}
diff --git a/Engine/Source/Developer/Windows/WindowsTargetPlatform/Private/GenericWindowsTargetPlatform.h b/Engine/Source/Developer/Windows/WindowsTargetPlatform/Private/GenericWindowsTargetPlatform.h
--- a/Engine/Source/Developer/Windows/WindowsTargetPlatform/Private/GenericWindowsTargetPlatform.h
+++ b/Engine/Source/Developer/Windows/WindowsTargetPlatform/Private/GenericWindowsTargetPlatform.h
@@ -29,7 +29,7 @@ public:
 		
 		#if WITH_ENGINE
 			FConfigCacheIni::LoadLocalIniFile(EngineSettings, TEXT("Engine"), true, *PlatformName());
-			TextureLODSettings.Initialize(EngineSettings, TEXT("SystemSettings"));
+			TextureLODSettings = nullptr; // These are registered by the device profile system.
 			StaticMeshLODSettings.Initialize(EngineSettings);
 
 			// Get the Target RHIs for this platform, we do not always want all those that are supported.
@@ -149,9 +149,14 @@ public:
 		}
 	}
 
-	virtual const struct FTextureLODSettings& GetTextureLODSettings( ) const override
+	virtual const UTextureLODSettings& GetTextureLODSettings() const override
 	{
-		return TextureLODSettings;
+		return *TextureLODSettings;
+	}
+
+	virtual void RegisterTextureLODSettings(const UTextureLODSettings* InTextureLODSettings) override
+	{
+		TextureLODSettings = InTextureLODSettings;
 	}
 
 	virtual FName GetWaveFormat( const class USoundWave* Wave ) const override
@@ -244,7 +249,7 @@ private:
 	FConfigFile EngineSettings;
 
 	// Holds the texture LOD settings.
-	FTextureLODSettings TextureLODSettings;
+	const UTextureLODSettings* TextureLODSettings;
 
 	// Holds static mesh LOD settings.
 	FStaticMeshLODSettings StaticMeshLODSettings;
diff --git a/Engine/Source/Developer/iOS/IOSTargetPlatform/Private/IOSTargetPlatform.cpp b/Engine/Source/Developer/iOS/IOSTargetPlatform/Private/IOSTargetPlatform.cpp
--- a/Engine/Source/Developer/iOS/IOSTargetPlatform/Private/IOSTargetPlatform.cpp
+++ b/Engine/Source/Developer/iOS/IOSTargetPlatform/Private/IOSTargetPlatform.cpp
@@ -14,7 +14,7 @@ FIOSTargetPlatform::FIOSTargetPlatform()
 {
 #if WITH_ENGINE
 	FConfigCacheIni::LoadLocalIniFile(EngineSettings, TEXT("Engine"), true, *PlatformName());
-	TextureLODSettings.Initialize(EngineSettings, TEXT("SystemSettings"));
+	TextureLODSettings = nullptr; // TextureLODSettings are registered by the device profile.
 	StaticMeshLODSettings.Initialize(EngineSettings);
 #endif // #if WITH_ENGINE
 
@@ -445,9 +445,9 @@ void FIOSTargetPlatform::GetTextureFormats( const UTexture* Texture, TArray<FNam
 }
 
 
-const FTextureLODSettings& FIOSTargetPlatform::GetTextureLODSettings() const
+const UTextureLODSettings& FIOSTargetPlatform::GetTextureLODSettings() const
 {
-	return TextureLODSettings;
+	return *TextureLODSettings;
 }
 
 
diff --git a/Engine/Source/Developer/iOS/IOSTargetPlatform/Private/IOSTargetPlatform.h b/Engine/Source/Developer/iOS/IOSTargetPlatform/Private/IOSTargetPlatform.h
--- a/Engine/Source/Developer/iOS/IOSTargetPlatform/Private/IOSTargetPlatform.h
+++ b/Engine/Source/Developer/iOS/IOSTargetPlatform/Private/IOSTargetPlatform.h
@@ -95,7 +95,12 @@ public:
 
 	virtual void GetTextureFormats( const UTexture* Texture, TArray<FName>& OutFormats ) const override;
 
-	virtual const struct FTextureLODSettings& GetTextureLODSettings( ) const override;
+	virtual const UTextureLODSettings& GetTextureLODSettings() const override;
+
+	virtual void RegisterTextureLODSettings(const UTextureLODSettings* InTextureLODSettings) override
+	{
+		TextureLODSettings = InTextureLODSettings;
+	}
 
 	virtual FName GetWaveFormat( const class USoundWave* Wave ) const override;
 #endif // WITH_ENGINE
@@ -152,7 +157,7 @@ private:
 	FConfigFile EngineSettings;
 
 	// Holds the cache of the target LOD settings.
-	FTextureLODSettings TextureLODSettings;
+	const UTextureLODSettings* TextureLODSettings;
 
 	// Holds the static mesh LOD settings.
 	FStaticMeshLODSettings StaticMeshLODSettings;
diff --git a/Engine/Source/Editor/DetailCustomizations/Private/DeviceProfileDetails.cpp b/Engine/Source/Editor/DetailCustomizations/Private/DeviceProfileDetails.cpp
--- a/Engine/Source/Editor/DetailCustomizations/Private/DeviceProfileDetails.cpp
+++ b/Engine/Source/Editor/DetailCustomizations/Private/DeviceProfileDetails.cpp
@@ -25,9 +25,6 @@ void FDeviceProfileDetails::CustomizeDetails(IDetailLayoutBuilder& DetailBuilder
 	TSharedPtr<IPropertyHandle> MeshLODSettingsHandle = DetailBuilder.GetProperty("MeshLODSettings");
 	DetailBuilder.HideProperty(MeshLODSettingsHandle);
 
-	TSharedPtr<IPropertyHandle> TextureLODSettingsHandle = DetailBuilder.GetProperty("TextureLODSettings");
-	DetailBuilder.HideProperty(TextureLODSettingsHandle);
-
 	// Setup the parent profile panel
 	ParentProfileDetails = MakeShareable(new FDeviceProfileParentPropertyDetails(&DetailBuilder));
 	ParentProfileDetails->CreateParentPropertyView();
@@ -35,6 +32,9 @@ void FDeviceProfileDetails::CustomizeDetails(IDetailLayoutBuilder& DetailBuilder
 	// Setup the console variable editor
 	ConsoleVariablesDetails = MakeShareable(new FDeviceProfileConsoleVariablesPropertyDetails(&DetailBuilder));
 	ConsoleVariablesDetails->CreateConsoleVariablesPropertyView();
+
+	TextureLODSettingsDetails = MakeShareable(new FDeviceProfileTextureLODSettingsDetails(&DetailBuilder));
+	TextureLODSettingsDetails->CreateTextureLODSettingsPropertyView();
 }
 
 
@@ -445,7 +445,7 @@ void FDeviceProfileParentPropertyDetails::CreateParentPropertyView()
 	if(ActiveDeviceProfile != nullptr)
 	{
 		TArray<UDeviceProfile*> AllPossibleParents;
-		GEngine->GetDeviceProfileManager()->GetAllPossibleParentProfiles(ActiveDeviceProfile, AllPossibleParents);
+		UDeviceProfileManager::Get().GetAllPossibleParentProfiles(ActiveDeviceProfile, AllPossibleParents);
 
 		for(auto& NextProfile : AllPossibleParents)
 		{
@@ -572,7 +572,6 @@ FDeviceProfileConsoleVariablesPropertyDetails::FDeviceProfileConsoleVariablesPro
 
 void FDeviceProfileConsoleVariablesPropertyDetails::CreateConsoleVariablesPropertyView()
 {
-
 	FSimpleDelegate OnCVarPropertyChangedDelegate = FSimpleDelegate::CreateSP(this, &FDeviceProfileConsoleVariablesPropertyDetails::OnCVarPropertyChanged);
 	CVarsHandle->SetOnPropertyValueChanged(OnCVarPropertyChangedDelegate);
 
@@ -841,5 +840,4 @@ void FDeviceProfileConsoleVariablesPropertyDetails::OnCVarPropertyChanged()
 	DetailBuilder->ForceRefreshDetails();
 }
 
-
 #undef LOCTEXT_NAMESPACE
\ No newline at end of file
diff --git a/Engine/Source/Editor/DetailCustomizations/Private/DeviceProfileDetails.h b/Engine/Source/Editor/DetailCustomizations/Private/DeviceProfileDetails.h
--- a/Engine/Source/Editor/DetailCustomizations/Private/DeviceProfileDetails.h
+++ b/Engine/Source/Editor/DetailCustomizations/Private/DeviceProfileDetails.h
@@ -2,7 +2,7 @@
 
 #pragma once
 
-class UDeviceProfile;
+#include "TextureLODSettingsDetails.h"
 
 
 /* FDeviceProfileParentPropertyDetails
@@ -65,7 +65,7 @@ private:
 	TArray<TSharedPtr<FString>> AvailableParentProfiles;
 	
 	/** A reference to the object we are showing these properties for */
-	UDeviceProfile* ActiveDeviceProfile;
+	class UDeviceProfile* ActiveDeviceProfile;
 };
 
 
@@ -159,7 +159,7 @@ private:
 };
 
 
-/* FDeviceProfileParentPropertyDetails
+/* FDeviceProfileDetails
  *****************************************************************************/
 
 /**
@@ -193,4 +193,7 @@ private:
 
 	/** Reference to the console variables property view */
 	TSharedPtr<FDeviceProfileConsoleVariablesPropertyDetails> ConsoleVariablesDetails;
+
+	/** Reference to the console variables property view */
+	TSharedPtr<FDeviceProfileTextureLODSettingsDetails> TextureLODSettingsDetails;
 };
diff --git a/Engine/Source/Editor/DeviceProfileEditor/Private/DeviceProfileConsoleVariableColumn.cpp b/Engine/Source/Editor/DeviceProfileEditor/Private/DeviceProfileConsoleVariableColumn.cpp
--- a/Engine/Source/Editor/DeviceProfileEditor/Private/DeviceProfileConsoleVariableColumn.cpp
+++ b/Engine/Source/Editor/DeviceProfileEditor/Private/DeviceProfileConsoleVariableColumn.cpp
@@ -134,7 +134,7 @@ bool FDeviceProfileConsoleVariableColumn::Supports(const TSharedRef< IPropertyTa
 		{
 			const FPropertyInfo& PropertyInfo = PropertyPath->GetRootProperty();
 			UProperty* Property = PropertyInfo.Property.Get();
-			if( Property->IsA( UArrayProperty::StaticClass() ) )
+			if (Property->GetName() == TEXT("CVars") && Property->IsA(UArrayProperty::StaticClass()))
 			{
 				return true;
 			}
diff --git a/Engine/Source/Editor/DeviceProfileEditor/Private/DeviceProfileEditorPCH.h b/Engine/Source/Editor/DeviceProfileEditor/Private/DeviceProfileEditorPCH.h
--- a/Engine/Source/Editor/DeviceProfileEditor/Private/DeviceProfileEditorPCH.h
+++ b/Engine/Source/Editor/DeviceProfileEditor/Private/DeviceProfileEditorPCH.h
@@ -17,6 +17,7 @@
 
 #include "DetailsPanel/SDeviceProfileDetailsPanel.h"
 #include "DeviceProfileConsoleVariableColumn.h"
+#include "DeviceProfileTextureLODSettingsColumn.h"
 #include "SDeviceProfileEditorSingleProfileView.h"
 #include "SDeviceProfileCreateProfilePanel.h"
 #include "SDeviceProfileSelectionPanel.h"
diff --git a/Engine/Source/Editor/DeviceProfileEditor/Private/DeviceProfileTextureLODSettingsColumn.cpp b/Engine/Source/Editor/DeviceProfileEditor/Private/DeviceProfileTextureLODSettingsColumn.cpp
new file mode 100644
--- /dev/null
+++ b/Engine/Source/Editor/DeviceProfileEditor/Private/DeviceProfileTextureLODSettingsColumn.cpp
@@ -0,0 +1,171 @@
+// Copyright 1998-2015 Epic Games, Inc. All Rights Reserved.
+
+// Module includes
+#include "DeviceProfileEditorPCH.h"
+
+// Property table includes
+#include "PropertyEditorModule.h"
+#include "PropertyPath.h"
+#include "IPropertyTableCell.h"
+#include "PropertyHandle.h"
+
+// Misc includes
+#include "DeviceProfiles/DeviceProfile.h"
+#include "ScopedTransaction.h"
+
+
+#define LOCTEXT_NAMESPACE "DeviceProfileEditor"
+
+
+/**
+* Formatter of the Texture LOD Settings property for a device profile.
+*/
+class FTextureLODSettingsCellPresenter : public TSharedFromThis< FTextureLODSettingsCellPresenter > , public IPropertyTableCellPresenter
+{
+public:
+	/** 
+	 * Constructor 
+	 */
+	FTextureLODSettingsCellPresenter(TWeakObjectPtr<UDeviceProfile> InOwnerProfile, const FOnEditDeviceProfileTextureLODSettingsRequestDelegate& OnTextureLODSettingsEditRequest )
+		: OwnerProfile(InOwnerProfile)
+		, OnEditTextureLODSettingsRequest(OnTextureLODSettingsEditRequest)
+	{
+	}
+
+	virtual ~FTextureLODSettingsCellPresenter() {}
+
+	/**
+	 * Event handler triggered when the user presses the edit TextureLODSettings button
+	 *
+	 * @return Whether the event was handled.
+	 */
+	FReply HandleEditTextureLODSettingsButtonPressed()
+	{
+		OnEditTextureLODSettingsRequest.ExecuteIfBound(OwnerProfile);
+		return FReply::Handled();
+	}
+
+public:
+
+	/** Begin IPropertyTableCellPresenter interface */
+	virtual TSharedRef<class SWidget> ConstructDisplayWidget() override;
+
+	virtual bool RequiresDropDown() override
+	{
+		return false;
+	}
+
+	virtual TSharedRef< class SWidget > ConstructEditModeCellWidget() override
+	{
+		return ConstructDisplayWidget();
+	}
+
+	virtual TSharedRef< class SWidget > ConstructEditModeDropDownWidget() override
+	{
+		return SNullWidget::NullWidget;
+	}
+
+	virtual TSharedRef< class SWidget > WidgetToFocusOnEdit() override
+	{
+		return SNullWidget::NullWidget;
+	}
+
+	virtual bool HasReadOnlyEditMode() override
+	{
+		return true;
+	}
+
+	virtual FString GetValueAsString() override
+	{
+		return TEXT("");
+	}
+
+	virtual FText GetValueAsText() override
+	{
+		return FText::FromString(TEXT(""));
+	}
+	/** End IPropertyTableCellPresenter interface */
+
+
+private:
+
+	/** The object we will link to */
+	TWeakObjectPtr<UDeviceProfile> OwnerProfile;
+
+	/** Delegate triggered when the user opts to edit the TextureLODSettings from the button in this cell */
+	FOnEditDeviceProfileTextureLODSettingsRequestDelegate OnEditTextureLODSettingsRequest;
+};
+
+
+TSharedRef<class SWidget> FTextureLODSettingsCellPresenter::ConstructDisplayWidget()
+{
+	return SNew(SBorder)
+	.Padding(0.0f)
+	.VAlign(VAlign_Center)
+	.HAlign(HAlign_Center)
+	.BorderImage(FEditorStyle::GetBrush("NoBorder"))
+	.Content()
+	[
+		SNew(SButton)
+		.OnClicked(this, &FTextureLODSettingsCellPresenter::HandleEditTextureLODSettingsButtonPressed)
+		.ContentPadding(2.0f)
+		.ForegroundColor(FSlateColor::UseForeground())
+		.IsFocusable(false)
+		[
+			SNew(SImage)
+			.Image(FEditorStyle::GetBrush("PropertyWindow.Button_Edit"))
+			.ColorAndOpacity(FSlateColor::UseForeground())
+		]
+	];
+}
+
+
+FDeviceProfileTextureLODSettingsColumn::FDeviceProfileTextureLODSettingsColumn()
+{
+}
+
+
+bool FDeviceProfileTextureLODSettingsColumn::Supports(const TSharedRef< IPropertyTableColumn >& Column, const TSharedRef< IPropertyTableUtilities >& Utilities) const
+{
+	if( Column->GetDataSource()->IsValid() )
+	{
+		TSharedPtr< FPropertyPath > PropertyPath = Column->GetDataSource()->AsPropertyPath();
+		if( PropertyPath.IsValid() && PropertyPath->GetNumProperties() > 0 )
+		{
+			const FPropertyInfo& PropertyInfo = PropertyPath->GetRootProperty();
+			UProperty* Property = PropertyInfo.Property.Get();
+			if (Property->GetName() == TEXT("TextureLODGroups") && Property->IsA(UArrayProperty::StaticClass()))
+			{
+				return true;
+			}
+		}
+	}
+
+	return false;
+}
+
+
+TSharedPtr< SWidget > FDeviceProfileTextureLODSettingsColumn::CreateColumnLabel(const TSharedRef< IPropertyTableColumn >& Column, const TSharedRef< IPropertyTableUtilities >& Utilities, const FName& Style) const
+{
+	return NULL;
+}
+
+
+TSharedPtr< IPropertyTableCellPresenter > FDeviceProfileTextureLODSettingsColumn::CreateCellPresenter(const TSharedRef< IPropertyTableCell >& Cell, const TSharedRef< IPropertyTableUtilities >& Utilities, const FName& Style) const
+{
+	TSharedPtr< IPropertyHandle > PropertyHandle = Cell->GetPropertyHandle();
+	if (PropertyHandle.IsValid())
+	{
+		TArray<UObject*> OuterObjects;
+		PropertyHandle->GetOuterObjects(OuterObjects);
+		if (OuterObjects.Num() == 1)
+		{
+			return MakeShareable(new FTextureLODSettingsCellPresenter(CastChecked<UDeviceProfile>(OuterObjects[0]),OnEditTextureLODSettingsRequestDelegate));
+		}
+	}
+
+	return NULL;
+}
+
+
+#undef LOCTEXT_NAMESPACE
\ No newline at end of file
diff --git a/Engine/Source/Editor/DeviceProfileEditor/Private/DeviceProfileTextureLODSettingsColumn.h b/Engine/Source/Editor/DeviceProfileEditor/Private/DeviceProfileTextureLODSettingsColumn.h
new file mode 100644
--- /dev/null
+++ b/Engine/Source/Editor/DeviceProfileEditor/Private/DeviceProfileTextureLODSettingsColumn.h
@@ -0,0 +1,46 @@
+// Copyright 1998-2015 Epic Games, Inc. All Rights Reserved.
+
+#pragma once
+
+#include "IPropertyTableCustomColumn.h"
+
+
+/** Delegate triggered when user opts to edit TextureLODSettings **/
+DECLARE_DELEGATE_OneParam(FOnEditDeviceProfileTextureLODSettingsRequestDelegate, const TWeakObjectPtr<UDeviceProfile>&);
+
+
+/**
+ * A property table custom column used to bring the user to an editor which 
+ * will manage the Texture LOD Settings associated with the device profile
+ */
+class FDeviceProfileTextureLODSettingsColumn : public IPropertyTableCustomColumn
+{
+public:
+	FDeviceProfileTextureLODSettingsColumn();
+
+public:
+
+	// IPropertyTableCustomColumn interface
+
+	virtual bool Supports( const TSharedRef< IPropertyTableColumn >& Column, const TSharedRef< IPropertyTableUtilities >& Utilities ) const override;
+	virtual TSharedPtr< SWidget > CreateColumnLabel( const TSharedRef< IPropertyTableColumn >& Column, const TSharedRef< IPropertyTableUtilities >& Utilities, const FName& Style ) const override;
+	virtual TSharedPtr< IPropertyTableCellPresenter > CreateCellPresenter( const TSharedRef< IPropertyTableCell >& Cell, const TSharedRef< IPropertyTableUtilities >& Utilities, const FName& Style ) const override;
+
+public:
+
+	/**
+	 * Delegate used to notify listeners that an edit request was triggered from the property table.
+	 *
+	 * @return The delegate.
+	 */
+	FOnEditDeviceProfileTextureLODSettingsRequestDelegate& OnEditTextureLODSettingsRequest()
+	{
+		return OnEditTextureLODSettingsRequestDelegate;
+	}
+
+private:
+
+	/** Delegate triggered when user opts to edit TextureLODSettings. **/
+	FOnEditDeviceProfileTextureLODSettingsRequestDelegate OnEditTextureLODSettingsRequestDelegate;
+};
+
diff --git a/Engine/Source/Editor/DeviceProfileEditor/Private/SDeviceProfileCreateProfilePanel.cpp b/Engine/Source/Editor/DeviceProfileEditor/Private/SDeviceProfileCreateProfilePanel.cpp
--- a/Engine/Source/Editor/DeviceProfileEditor/Private/SDeviceProfileCreateProfilePanel.cpp
+++ b/Engine/Source/Editor/DeviceProfileEditor/Private/SDeviceProfileCreateProfilePanel.cpp
@@ -285,7 +285,7 @@ bool SDeviceProfileCreateProfilePanel::IsCreateProfileButtonEnabled() const
 
 FReply SDeviceProfileCreateProfilePanel::HandleCreateDeviceProfileButtonClicked()
 {
-	UDeviceProfile* NewProfile = DeviceProfileManager->CreateProfile( *DeviceProfileNameTextBox->GetText().ToString(), *SelectedDeviceProfileType, SelectedDeviceProfileParent.IsValid() ? SelectedDeviceProfileParent->GetName() : TEXT("") );
+	DeviceProfileManager->CreateProfile( *DeviceProfileNameTextBox->GetText().ToString(), *SelectedDeviceProfileType, SelectedDeviceProfileParent.IsValid() ? SelectedDeviceProfileParent->GetName() : TEXT("") );
 
 	// Set the components for this panel back to default
 	ResetComponentsState();
diff --git a/Engine/Source/Editor/DeviceProfileEditor/Private/SDeviceProfileEditor.cpp b/Engine/Source/Editor/DeviceProfileEditor/Private/SDeviceProfileEditor.cpp
--- a/Engine/Source/Editor/DeviceProfileEditor/Private/SDeviceProfileEditor.cpp
+++ b/Engine/Source/Editor/DeviceProfileEditor/Private/SDeviceProfileEditor.cpp
@@ -92,7 +92,7 @@ SDeviceProfileSourceControl::~SDeviceProfileSourceControl()
 
 FReply SDeviceProfileSourceControl::HandleSaveDefaultsButtonPressed()
 {
-	GEngine->GetDeviceProfileManager()->SaveProfiles(true);
+	UDeviceProfileManager::Get().SaveProfiles(true);
 
 	return FReply::Handled();
 }
@@ -226,7 +226,7 @@ void SDeviceProfileSourceControl::Tick( const FGeometry& AllottedGeometry, const
 
 void SDeviceProfileEditor::Construct( const FArguments& InArgs )
 {
-	DeviceProfileManager = GEngine->GetDeviceProfileManager();
+	DeviceProfileManager = &UDeviceProfileManager::Get();
 
 	// Setup the tab layout for the editor.
 	TSharedRef<FWorkspaceItem> RootMenuGroup = FWorkspaceItem::NewGroup( LOCTEXT("RootMenuGroupName", "Root") );
@@ -514,9 +514,14 @@ TSharedRef< SWidget > SDeviceProfileEditor::SetupPropertyEditor()
 	// Bind our action to open a single editor when requested from the property table
 	CVarsColumn->OnEditCVarsRequest().BindRaw(this, &SDeviceProfileEditor::HandleDeviceProfileViewAlone);
 
+	// Adapt the TextureLODSettings column as a button to open a single editor which will allow better control of the Texture Groups
+	TSharedRef<FDeviceProfileTextureLODSettingsColumn> TextureLODSettingsColumn = MakeShareable(new FDeviceProfileTextureLODSettingsColumn());
+	TextureLODSettingsColumn->OnEditTextureLODSettingsRequest().BindRaw(this, &SDeviceProfileEditor::HandleDeviceProfileViewAlone);
+
 	// Add our Custom Rows to the table
 	TArray<TSharedRef<IPropertyTableCustomColumn>> CustomColumns;
 	CustomColumns.Add(CVarsColumn);
+	CustomColumns.Add(TextureLODSettingsColumn);
 
 	return PropertyEditorModule.CreatePropertyTableWidget(PropertyTable.ToSharedRef(), CustomColumns);
 }
diff --git a/Engine/Source/Editor/DeviceProfileServices/Private/DeviceProfileServicesUIManager.cpp b/Engine/Source/Editor/DeviceProfileServices/Private/DeviceProfileServicesUIManager.cpp
--- a/Engine/Source/Editor/DeviceProfileServices/Private/DeviceProfileServicesUIManager.cpp
+++ b/Engine/Source/Editor/DeviceProfileServices/Private/DeviceProfileServicesUIManager.cpp
@@ -10,7 +10,7 @@ DEFINE_LOG_CATEGORY_STATIC(LogDeviceProfileServices, Log, All);
 
 FDeviceProfileServicesUIManager::FDeviceProfileServicesUIManager()
 {
-	GEngine->GetDeviceProfileManager()->OnManagerUpdated().AddRaw(this, &FDeviceProfileServicesUIManager::HandleRefreshUIData);
+	UDeviceProfileManager::Get().OnManagerUpdated().AddRaw(this, &FDeviceProfileServicesUIManager::HandleRefreshUIData);
 	HandleRefreshUIData();
 	CreatePlatformMap();
 }
@@ -38,9 +38,9 @@ const TArray<TSharedPtr<FString> > FDeviceProfileServicesUIManager::GetPlatformL
 
 void FDeviceProfileServicesUIManager::GetProfilesByType( TArray<UDeviceProfile*>& OutDeviceProfiles, const FString& InType )
 {
-	for( int32 Idx = 0; Idx < GEngine->GetDeviceProfileManager()->Profiles.Num(); Idx++ )
+	for( int32 Idx = 0; Idx < UDeviceProfileManager::Get().Profiles.Num(); Idx++ )
 	{
-		UDeviceProfile* CurrentDevice = CastChecked<UDeviceProfile>( GEngine->GetDeviceProfileManager()->Profiles[Idx] );
+		UDeviceProfile* CurrentDevice = CastChecked<UDeviceProfile>( UDeviceProfileManager::Get().Profiles[Idx] );
 		if ( CurrentDevice->DeviceType == InType )
 		{
 			OutDeviceProfiles.Add( CurrentDevice );
@@ -64,9 +64,9 @@ void FDeviceProfileServicesUIManager::HandleRefreshUIData()
 {
 	// Rebuild profile to platform map
 	DeviceToPlatformMap.Empty();
-	for( int32 Idx = 0; Idx < GEngine->GetDeviceProfileManager()->Profiles.Num(); Idx++ )
+	for( int32 Idx = 0; Idx < UDeviceProfileManager::Get().Profiles.Num(); Idx++ )
 	{
-		UDeviceProfile* CurrentDevice = CastChecked<UDeviceProfile>( GEngine->GetDeviceProfileManager()->Profiles[Idx] );
+		UDeviceProfile* CurrentDevice = CastChecked<UDeviceProfile>( UDeviceProfileManager::Get().Profiles[Idx] );
 		DeviceToPlatformMap.Add( CurrentDevice->GetName(), CurrentDevice->DeviceType );
 	}
 }
diff --git a/Engine/Source/Editor/TextureEditor/Private/TextureEditorToolkit.cpp b/Engine/Source/Editor/TextureEditor/Private/TextureEditorToolkit.cpp
--- a/Engine/Source/Editor/TextureEditor/Private/TextureEditorToolkit.cpp
+++ b/Engine/Source/Editor/TextureEditor/Private/TextureEditorToolkit.cpp
@@ -10,7 +10,7 @@
 #include "Engine/TextureRenderTargetCube.h"
 #include "Engine/TextureRenderTarget2D.h"
 #include "Engine/Texture2D.h"
-#include "SystemSettings.h"
+#include "DeviceProfiles/DeviceProfile.h"
 #include "Engine/LightMapTexture2D.h"
 #include "Engine/ShadowMapTexture2D.h"
 #include "Engine/Texture2DDynamic.h"
@@ -352,7 +352,7 @@ void FTextureEditorToolkit::PopulateQuickInfo( )
 
 	// In game max bias and dimensions
 	uint32 MaxInGameWidth, MaxInGameHeight;
-	int32 MipLevel = GSystemSettings.TextureLODSettings.CalculateLODBias(Texture);
+	int32 MipLevel = UDeviceProfileManager::Get().GetActiveProfile()->GetTextureLODSettings()->CalculateLODBias(Texture);
 	CalculateEffectiveTextureDimensions(MipLevel, MaxInGameWidth, MaxInGameHeight);
 
 	// Editor bias and dimensions (takes user specified mip setting into account)
@@ -600,7 +600,7 @@ TSharedRef<SWidget> FTextureEditorToolkit::BuildTexturePropertiesWidget( )
 void FTextureEditorToolkit::CalculateEffectiveTextureDimensions( int32 LODBias, uint32& EffectiveTextureWidth, uint32& EffectiveTextureHeight )
 {
 	//Calculate in-game max resolution and store in EffectiveTextureWidth, EffectiveTextureHeight
-	GSystemSettings.TextureLODSettings.ComputeInGameMaxResolution(LODBias, *Texture, (uint32 &)EffectiveTextureWidth, (uint32 &)EffectiveTextureHeight);
+	UDeviceProfileManager::Get().GetActiveProfile()->GetTextureLODSettings()->ComputeInGameMaxResolution(LODBias, *Texture, (uint32 &)EffectiveTextureWidth, (uint32 &)EffectiveTextureHeight);
 }
 
 
diff --git a/Engine/Source/Editor/UnrealEd/Private/CookOnTheFlyServer.cpp b/Engine/Source/Editor/UnrealEd/Private/CookOnTheFlyServer.cpp
--- a/Engine/Source/Editor/UnrealEd/Private/CookOnTheFlyServer.cpp
+++ b/Engine/Source/Editor/UnrealEd/Private/CookOnTheFlyServer.cpp
@@ -34,7 +34,7 @@
 // shader compiler processAsyncResults
 #include "ShaderCompiler.h"
 #include "Engine/LevelStreaming.h"
-#include "TextureLODSettings.h"
+#include "Engine/TextureLODSettings.h"
 
 DEFINE_LOG_CATEGORY_STATIC(LogCook, Log, All);
 
@@ -2080,7 +2080,7 @@ bool UCookOnTheFlyServer::GetCurrentIniVersionStrings( const ITargetPlatform* Ta
 	PlatformIniFiles.Empty();
 
 
-	const FTextureLODSettings& LodSettings = TargetPlatform->GetTextureLODSettings();
+	const UTextureLODSettings& LodSettings = TargetPlatform->GetTextureLODSettings();
 
 	UEnum* TextureGroupEnum = FindObject<UEnum>( NULL, TEXT("Engine.TextureGroup") );
 	UEnum* TextureMipGenSettingsEnum = FindObject<UEnum>( NULL, TEXT("Engine.TextureMipGenSettings") );
diff --git a/Engine/Source/Editor/UnrealEd/Private/Dialogs/SMeshMergingDialog.cpp b/Engine/Source/Editor/UnrealEd/Private/Dialogs/SMeshMergingDialog.cpp
--- a/Engine/Source/Editor/UnrealEd/Private/Dialogs/SMeshMergingDialog.cpp
+++ b/Engine/Source/Editor/UnrealEd/Private/Dialogs/SMeshMergingDialog.cpp
@@ -39,8 +39,8 @@ void SMeshMergingDialog::Construct(const FArguments& InArgs)
 	
 
 	// Setup available resolutions for lightmap
-	const int32 MinTexResolution = 1 << FTextureLODSettings::FTextureLODGroup().MinLODMipCount;
-	const int32 MaxTexResolution = 1 << FTextureLODSettings::FTextureLODGroup().MaxLODMipCount;
+	const int32 MinTexResolution = 1 << FTextureLODGroup().MinLODMipCount;
+	const int32 MaxTexResolution = 1 << FTextureLODGroup().MaxLODMipCount;
 	for (int32 LightmapRes = MinTexResolution; LightmapRes <= MaxTexResolution; LightmapRes*=2)
 	{
 		LightMapResolutionOptions.Add(MakeShareable(new FString(FString::FormatAsNumber(LightmapRes))));
diff --git a/Engine/Source/Editor/UnrealEd/Private/EditorExporters.cpp b/Engine/Source/Editor/UnrealEd/Private/EditorExporters.cpp
--- a/Engine/Source/Editor/UnrealEd/Private/EditorExporters.cpp
+++ b/Engine/Source/Editor/UnrealEd/Private/EditorExporters.cpp
@@ -36,7 +36,7 @@
 #include "AI/Navigation/NavigationSystem.h"
 #include "Engine/TextureCube.h"
 #include "EngineUtils.h"
-#include "TextureLODSettings.h"
+#include "DeviceProfiles/DeviceProfileManager.h"
 
 DEFINE_LOG_CATEGORY_STATIC(LogEditorExporters, Log, All);
 
@@ -1050,14 +1050,7 @@ public:
 	FIntPoint FindMaxTextureSize(UMaterialInterface* InMaterialInterface, FIntPoint MinimumSize = FIntPoint(1, 1)) const
 	{
 		// static lod settings so that we only initialize them once
-		static FTextureLODSettings GameTextureLODSettings;
-		static bool bAreLODSettingsInitialized = false;
-		if (!bAreLODSettingsInitialized)
-		{
-			// initialize LOD settings with game texture resolutions, since we don't want to use 
-			// potentially bloated editor settings
-			GameTextureLODSettings.Initialize(GEngineIni, TEXT("SystemSettings"));
-		}
+		UTextureLODSettings* GameTextureLODSettings = UDeviceProfileManager::Get().GetActiveProfile()->GetTextureLODSettings();
 
 		TArray<UTexture*> MaterialTextures;
 
@@ -1087,7 +1080,7 @@ public:
 				LocalSize = FIntPoint(TexCube->GetSizeX(), TexCube->GetSizeY());
 			}
 
-			int32 LocalBias = GameTextureLODSettings.CalculateLODBias(Texture);
+			int32 LocalBias = GameTextureLODSettings->CalculateLODBias(Texture);
 
 			// bias the texture size based on LOD group
 			FIntPoint BiasedLocalSize(LocalSize.X >> LocalBias, LocalSize.Y >> LocalBias);
diff --git a/Engine/Source/Runtime/Engine/Classes/DeviceProfiles/DeviceProfile.h b/Engine/Source/Runtime/Engine/Classes/DeviceProfiles/DeviceProfile.h
--- a/Engine/Source/Runtime/Engine/Classes/DeviceProfiles/DeviceProfile.h
+++ b/Engine/Source/Runtime/Engine/Classes/DeviceProfiles/DeviceProfile.h
@@ -12,8 +12,8 @@
 DECLARE_DELEGATE(FOnCVarsUpdated);
 
 
-UCLASS(config=DeviceProfiles, perobjectconfig, Blueprintable)
-class ENGINE_API UDeviceProfile : public UObject
+UCLASS(config=DeviceProfiles, perObjectConfig)
+class ENGINE_API UDeviceProfile : public UTextureLODSettings
 {
 	GENERATED_UCLASS_BODY()
 
@@ -54,6 +54,11 @@ public:
 		return CVarsUpdatedDelegate;
 	}
 
+public:
+	/** 
+	 * Access to the device profiles Texture LOD Settings
+	 */
+	UTextureLODSettings* GetTextureLODSettings() const;
 
 private:
 
@@ -62,6 +67,10 @@ private:
 
 public:
 
+	// Begin UObject interface
+	virtual void PostInitProperties() override;
+	// End UObject interface
+
 #if WITH_EDITOR
 	// Begin UObject interface
 	virtual void PostEditChangeProperty( FPropertyChangedEvent& PropertyChangedEvent ) override;
diff --git a/Engine/Source/Runtime/Engine/Classes/DeviceProfiles/DeviceProfileManager.h b/Engine/Source/Runtime/Engine/Classes/DeviceProfiles/DeviceProfileManager.h
--- a/Engine/Source/Runtime/Engine/Classes/DeviceProfiles/DeviceProfileManager.h
+++ b/Engine/Source/Runtime/Engine/Classes/DeviceProfiles/DeviceProfileManager.h
@@ -6,12 +6,11 @@
 
 #pragma once
 
+#include "DeviceProfile.h"
 #include "DeviceProfileManager.generated.h"
 
-class UDeviceProfile;
-
 // Delegate used to refresh the UI when the profiles change
-DECLARE_MULTICAST_DELEGATE( FOnManagerUpdated );
+DECLARE_MULTICAST_DELEGATE( FOnDeviceProfileManagerUpdated );
 
 /**
  * Implements a helper class that manages all profiles in the Device
@@ -21,7 +20,7 @@ class ENGINE_API UDeviceProfileManager : public UObject
 {
 public:
 
-	GENERATED_UCLASS_BODY()
+	GENERATED_BODY()
 
 	/**
 	 * Startup and select the active device profile
@@ -36,7 +35,7 @@ public:
 	 *
 	 * @return the created profile.
 	 */
-	UDeviceProfile* CreateProfile( const FString& ProfileName );
+	UDeviceProfile& CreateProfile( const FString& ProfileName );
 
 	/**
 	 * Create a copy of a device profile from a copy.
@@ -46,7 +45,7 @@ public:
 	 *
 	 * @return the created profile.
 	 */
-	UDeviceProfile* CreateProfile( const FString& ProfileName, const FString& ProfileType, const FString& ParentName = TEXT("") );
+	UDeviceProfile& CreateProfile( const FString& ProfileName, const FString& ProfileType, const FString& ParentName = TEXT("") );
 
 	/**
 	 * Delete a profile.
@@ -61,14 +60,14 @@ public:
 	 * @param ProfileName - The profile name to find.
 	 * @return The found profile.
 	 */
-	UDeviceProfile* FindProfile( const FString& ProfileName );
+	UDeviceProfile& FindProfile( const FString& ProfileName );
 
 	/**
 	 * Get the device profile .ini name.
 	 *
 	 * @return the device profile .ini name.
 	 */
-	const FString GetDeviceProfileIniName();
+	const FString GetDeviceProfileIniName() const;
 
 	/**
 	 * Load the device profiles from the config file.
@@ -80,7 +79,7 @@ public:
 	 *
 	 * @return The delegate.
 	 */
-	FOnManagerUpdated& OnManagerUpdated();
+	FOnDeviceProfileManagerUpdated& OnManagerUpdated();
 
 	/**
 	 * Save the device profiles.
@@ -107,8 +106,10 @@ public:
 	 * @param ChildProfile				- The profile we are looking for potential parents
 	 * @param PossibleParentProfiles	- The list of profiles which would be suitable as a parent for the given profile
 	 */
-	void GetAllPossibleParentProfiles( const UDeviceProfile* ChildProfile, OUT TArray<UDeviceProfile*>& PossibleParentProfiles );
+	void GetAllPossibleParentProfiles( const UDeviceProfile* ChildProfile, OUT TArray<UDeviceProfile*>& PossibleParentProfiles ) const;
 
+	static class UDeviceProfileManager* DeviceProfileManagerSingleton;
+	static UDeviceProfileManager& Get();
 
 public:
 
@@ -119,13 +120,11 @@ public:
 private:
 
 	// Holds a delegate to be invoked profiles are updated.
-	FOnManagerUpdated ManagerUpdatedDelegate;
+	FOnDeviceProfileManagerUpdated ManagerUpdatedDelegate;
 
 	// Holds the selected device profile
-	class UDeviceProfile* ActiveDeviceProfile;
+	UDeviceProfile* ActiveDeviceProfile;
 
 	// Holds the device profile .ini location
 	static FString DeviceProfileFileName;
-
-	int32 RenameIndex;
 };
\ No newline at end of file
diff --git a/Engine/Source/Runtime/Engine/Classes/Engine/Engine.h b/Engine/Source/Runtime/Engine/Classes/Engine/Engine.h
--- a/Engine/Source/Runtime/Engine/Classes/Engine/Engine.h
+++ b/Engine/Source/Runtime/Engine/Classes/Engine/Engine.h
@@ -1393,21 +1393,6 @@ public:
 	UPROPERTY(globalconfig)
 	uint32 bEnableVisualLogRecordingOnStart;
 
-private:
-
-	/** Used to hold the device profile manager and device profiles */
-	UPROPERTY(transient)
-	UDeviceProfileManager* DeviceProfileManager;
-
-public:
-
-	/**
-	 * Get the reference to the engines device profile manager
-	 *
-	 * return A reference to the dpm singleton, initializes if not already.
-	 */
-	UDeviceProfileManager* GetDeviceProfileManager();
-
 private:
 
 	/** Whether the engine should be playing sounds.  If false at initialization time the AudioDevice will not be created */
diff --git a/Engine/Source/Runtime/Engine/Classes/Engine/TextureDefines.h b/Engine/Source/Runtime/Engine/Classes/Engine/TextureDefines.h
--- a/Engine/Source/Runtime/Engine/Classes/Engine/TextureDefines.h
+++ b/Engine/Source/Runtime/Engine/Classes/Engine/TextureDefines.h
@@ -109,3 +109,14 @@ namespace ETexturePowerOfTwoSetting
 		// in the texture compressor DDC key
 	};
 }
+
+// Must match enum ESamplerFilter in RHIDefinitions.h
+UENUM()
+enum class ETextureSamplerFilter : uint8
+{
+	Point,
+	Bilinear,
+	Trilinear,
+	AnisotropicPoint,
+	AnisotropicLinear,
+};
diff --git a/Engine/Source/Runtime/Engine/Classes/Engine/TextureLODSettings.h b/Engine/Source/Runtime/Engine/Classes/Engine/TextureLODSettings.h
new file mode 100644
--- /dev/null
+++ b/Engine/Source/Runtime/Engine/Classes/Engine/TextureLODSettings.h
@@ -0,0 +1,214 @@
+// Copyright 1998-2015 Epic Games, Inc. All Rights Reserved.
+
+
+#pragma once
+
+#include "TextureDefines.h"
+#include "TextureLODSettings.generated.h"
+
+/** LOD settings for a single texture group. */
+USTRUCT()
+struct ENGINE_API FTextureLODGroup
+{
+	GENERATED_USTRUCT_BODY()
+
+	FTextureLODGroup()
+		: MinLODMipCount(0)
+		, MaxLODMipCount(12)
+		, LODBias(0)
+		, NumStreamedMips(-1)
+		, MipGenSettings(TextureMipGenSettings::TMGS_SimpleAverage)
+		, MinLODSize(1)
+		, MaxLODSize(4096)
+		, MinMagFilter(NAME_Aniso)
+		, MipFilter(NAME_Point)
+	{
+		MinLODMipCount = FMath::CeilLogTwo(MinLODSize);
+		MaxLODMipCount = FMath::CeilLogTwo(MaxLODSize);
+
+		// Linear filtering
+		if (MinMagFilter == NAME_Linear)
+		{
+			if (MipFilter == NAME_Point)
+			{
+				Filter = ETextureSamplerFilter::Bilinear;
+			}
+			else
+			{
+				Filter = ETextureSamplerFilter::Trilinear;
+			}
+		}
+		// Point. Don't even care about mip filter.
+		else if (MinMagFilter == NAME_Point)
+		{
+			Filter = ETextureSamplerFilter::Point;
+		}
+		// Aniso or unknown.
+		else
+		{
+			if (MipFilter == NAME_Point)
+			{
+				Filter = ETextureSamplerFilter::AnisotropicPoint;
+			}
+			else
+			{
+				Filter = ETextureSamplerFilter::AnisotropicLinear;
+			}
+		}
+	}
+
+	/** Minimum LOD mip count below which the code won't bias.						*/
+	UPROPERTY()
+	TEnumAsByte<TextureGroup> Group;
+
+	/** Minimum LOD mip count below which the code won't bias.						*/
+	int32 MinLODMipCount;
+	
+	/** Maximum LOD mip count. Bias will be adjusted so texture won't go above.		*/
+	int32 MaxLODMipCount;
+	
+	/** Group LOD bias.																*/
+	UPROPERTY()
+	int32 LODBias;
+	
+	/** Sampler filter state.														*/
+	ETextureSamplerFilter Filter;
+	
+	/** Number of mip-levels that can be streamed. -1 means all mips can stream.	*/
+	UPROPERTY()
+	int32 NumStreamedMips;
+
+	/** Defines how the the mip-map generation works, e.g. sharpening				*/
+	UPROPERTY()
+	TEnumAsByte<TextureMipGenSettings> MipGenSettings;
+
+	UPROPERTY()
+	int32 MinLODSize;
+
+	UPROPERTY()
+	int32 MaxLODSize;
+
+	UPROPERTY()
+	FName MinMagFilter;
+
+	UPROPERTY()
+	FName MipFilter;
+};
+
+/**
+ * Structure containing all information related to an LOD group and providing helper functions to calculate
+ * the LOD bias of a given group.
+ */
+UCLASS(config=DeviceProfiles, perObjectConfig)
+class ENGINE_API UTextureLODSettings : public UObject
+{
+	GENERATED_UCLASS_BODY()
+
+public:
+
+
+	/**
+	 * Calculates and returns the LOD bias based on texture LOD group, LOD bias and maximum size.
+	 *
+	 * @param	Texture			Texture object to calculate LOD bias for.
+	 * @param	bIncTextureMips	If true, takes the textures LOD & Cinematic mips into consideration
+	 * @return	LOD bias
+	 */
+	int32 CalculateLODBias(const UTexture* Texture, bool bIncTextureMips = true) const;
+
+	/**
+	 * Calculates and returns the LOD bias based on the information provided.
+	 *
+	 * @param	Width						Width of the texture
+	 * @param	Height						Height of the texture
+	 * @param	LODGroup					Which LOD group the texture belongs to
+	 * @param	LODBias						LOD bias to include in the calculation
+	 * @param	NumCinematicMipLevels		The texture cinematic mip levels to include in the calculation
+	 * @param	MipGenSetting				Mip generation setting
+	 * @return	LOD bias
+	 */
+	int32 CalculateLODBias( int32 Width, int32 Height, int32 LODGroup, int32 LODBias, int32 NumCinematicMipLevels, TextureMipGenSettings MipGenSetting ) const;
+
+	/** 
+	* Useful for stats in the editor.
+	*
+	* @param LODBias			Default LOD at which the texture renders. Platform dependent, call UTextureLODSettings::CalculateLODBias(Texture)
+	*/
+	void ComputeInGameMaxResolution(int32 LODBias, UTexture& Texture, uint32& OutSizeX, uint32& OutSizeY) const;
+
+#if WITH_EDITORONLY_DATA
+	void GetMipGenSettings( const UTexture& Texture, TextureMipGenSettings& OutMipGenSettings, float& OutSharpen, uint32& OutKernelSize, bool& bOutDownsampleWithAverage, bool& bOutSharpenWithoutColorShift, bool &bOutBorderColorBlack ) const;
+#endif // #if WITH_EDITORONLY_DATA
+
+	/**
+	 * Will return the LODBias for a passed in LODGroup
+	 *
+	 * @param	InLODGroup		The LOD Group ID 
+	 * @return	LODBias
+	 */
+	int32 GetTextureLODGroupLODBias( int32 InLODGroup ) const;
+
+	/**
+	 * Returns the LODGroup setting for number of streaming mip-levels.
+	 * -1 means that all mip-levels are allowed to stream.
+	 *
+	 * @param	InLODGroup		The LOD Group ID 
+	 * @return	Number of streaming mip-levels for textures in the specified LODGroup
+	 */
+	int32 GetNumStreamedMips( int32 InLODGroup ) const;
+
+	int32 GetMinLODMipCount( int32 InLODGroup ) const;
+	int32 GetMaxLODMipCount( int32 InLODGroup ) const;
+
+	/**
+	 * Returns the filter state that should be used for the passed in texture, taking
+	 * into account other system settings.
+	 *
+	 * @param	Texture		Texture to retrieve filter state for, must not be 0
+	 * @return	Filter sampler state for passed in texture
+	 */
+	ETextureSamplerFilter GetSamplerFilter(const UTexture* Texture) const;
+
+	ETextureSamplerFilter GetSamplerFilter(int32 InLODGroup) const;
+
+	/**
+	 * Returns the LODGroup mip gen settings
+	 *
+	 * @param	InLODGroup		The LOD Group ID 
+	 * @return	TextureMipGenSettings for lod group
+	 */
+	const TextureMipGenSettings GetTextureMipGenSettings( int32 InLODGroup ) const; 
+
+
+	/**
+	 * Returns the texture group names, sorted like enum.
+	 *
+	 * @return array of texture group names
+	 */
+	static TArray<FString> GetTextureGroupNames();
+
+	/**
+	 * TextureLODGroups access with bounds check
+	 *
+	 * @param   GroupIndex      usually from Texture.LODGroup
+	 * @return                  A handle to the indexed LOD group. 
+	 */
+	FTextureLODGroup& GetTextureLODGroup(TextureGroup GroupIndex);
+
+	/**
+	* TextureLODGroups access with bounds check
+	*
+	* @param   GroupIndex      usually from Texture.LODGroup
+	* @return                  A handle to the indexed LOD group.
+	*/
+	const FTextureLODGroup& GetTextureLODGroup(TextureGroup GroupIndex) const;
+
+protected:
+	void SetupLODGroup(int32 GroupId);
+
+public:
+
+	/** Array of LOD settings with entries per group. */
+	UPROPERTY(EditAnywhere, config, Category="Texture LOD Settings")
+	TArray<FTextureLODGroup> TextureLODGroups;
+};
diff --git a/Engine/Source/Runtime/Engine/Private/ContentStreaming.cpp b/Engine/Source/Runtime/Engine/Private/ContentStreaming.cpp
--- a/Engine/Source/Runtime/Engine/Private/ContentStreaming.cpp
+++ b/Engine/Source/Runtime/Engine/Private/ContentStreaming.cpp
@@ -2187,7 +2187,7 @@ FStreamingManagerTexture::FStreamingManagerTexture()
 
 	for ( int32 LODGroup=0; LODGroup < TEXTUREGROUP_MAX; ++LODGroup )
 	{
-		FTextureLODSettings::FTextureLODGroup& TexGroup = GSystemSettings.TextureLODSettings.GetTextureLODGroup( LODGroup );
+		const FTextureLODGroup& TexGroup = UDeviceProfileManager::Get().GetActiveProfile()->GetTextureLODSettings()->GetTextureLODGroup(TextureGroup(LODGroup));
 		ThreadSettings.NumStreamedMips[LODGroup] = TexGroup.NumStreamedMips;
 	}
 
@@ -4410,7 +4410,7 @@ bool FStreamingManagerTexture::HandleNumStreamedMipsCommand( const TCHAR* Cmd, F
 	int32 NumMips = ( NumMipsString.Len() > 0 ) ? FCString::Atoi(*NumMipsString) : MAX_int32;
 	if ( LODGroup >= 0 && LODGroup < TEXTUREGROUP_MAX )
 	{
-		FTextureLODSettings::FTextureLODGroup& TexGroup = GSystemSettings.TextureLODSettings.GetTextureLODGroup(LODGroup);
+		FTextureLODGroup& TexGroup = UDeviceProfileManager::Get().GetActiveProfile()->GetTextureLODSettings()->GetTextureLODGroup(TextureGroup(LODGroup));
 		if ( NumMips >= -1 && NumMips <= MAX_TEXTURE_MIP_COUNT )
 		{
 			TexGroup.NumStreamedMips = NumMips;
diff --git a/Engine/Source/Runtime/Engine/Private/DeviceProfiles/DeviceProfile.cpp b/Engine/Source/Runtime/Engine/Private/DeviceProfiles/DeviceProfile.cpp
--- a/Engine/Source/Runtime/Engine/Private/DeviceProfiles/DeviceProfile.cpp
+++ b/Engine/Source/Runtime/Engine/Private/DeviceProfiles/DeviceProfile.cpp
@@ -12,7 +12,7 @@ UDeviceProfile::UDeviceProfile(const FObjectInitializer& ObjectInitializer)
 	bVisible = true;
 
 	FString DeviceProfileFileName = FPaths::EngineConfigDir() + TEXT("Deviceprofiles.ini");
-	LoadConfig(GetClass(), *DeviceProfileFileName);
+//	LoadConfig(GetClass(), *DeviceProfileFileName, UE4::LCPF_ReadParentSections);
 }
 
 
@@ -40,6 +40,51 @@ void UDeviceProfile::GatherParentCVarInformationRecursively(OUT TMap<FString, FS
 	}
 }
 
+UTextureLODSettings* UDeviceProfile::GetTextureLODSettings() const
+{
+	return (UTextureLODSettings*)this;
+}
+
+
+void UDeviceProfile::PostInitProperties()
+{
+	Super::PostInitProperties();
+
+	// Ensure the Texture LOD Groups are in order of TextureGroup Enum
+	TextureLODGroups.Sort([]
+		(const FTextureLODGroup& Lhs, const FTextureLODGroup& Rhs)
+		{
+			return (int32)Lhs.Group < (int32)Rhs.Group;
+		}
+	);
+
+	// Make sure every Texture Group has an entry, any that aren't specified for this profile should use it's parents values, or the defaults.
+	UDeviceProfile* ParentProfile = nullptr;
+	if (HasAnyFlags(RF_ClassDefaultObject) == false)
+	{
+		if (BaseProfileName.IsEmpty() == false)
+		{
+			ParentProfile = FindObject<UDeviceProfile>(GetTransientPackage(), *BaseProfileName);
+		}
+		if (ParentProfile == nullptr)
+		{
+			ParentProfile = CastChecked<UDeviceProfile>(UDeviceProfile::StaticClass()->GetDefaultObject());
+		}
+	}
+
+	for (int32 GroupId = 0; GroupId < (int32)TEXTUREGROUP_MAX; ++GroupId)
+	{
+		if (TextureLODGroups.Num() < (GroupId + 1) || TextureLODGroups[GroupId].Group > GroupId)
+		{
+			TextureLODGroups.Insert((ParentProfile ? ParentProfile->TextureLODGroups[GroupId] : FTextureLODGroup()), GroupId);
+			TextureLODGroups[GroupId].Group = (TextureGroup)GroupId;
+		}
+	}
+
+#define SETUPLODGROUP(GroupId) SetupLODGroup(GroupId);
+	FOREACH_ENUM_TEXTUREGROUP(SETUPLODGROUP)
+#undef SETUPLODGROUP
+}
 
 #if WITH_EDITOR
 
diff --git a/Engine/Source/Runtime/Engine/Private/DeviceProfiles/DeviceProfileManager.cpp b/Engine/Source/Runtime/Engine/Private/DeviceProfiles/DeviceProfileManager.cpp
--- a/Engine/Source/Runtime/Engine/Private/DeviceProfiles/DeviceProfileManager.cpp
+++ b/Engine/Source/Runtime/Engine/Private/DeviceProfiles/DeviceProfileManager.cpp
@@ -9,6 +9,30 @@
 
 FString UDeviceProfileManager::DeviceProfileFileName;
 
+UDeviceProfileManager* UDeviceProfileManager::DeviceProfileManagerSingleton = nullptr;
+
+UDeviceProfileManager& UDeviceProfileManager::Get()
+{
+	if (DeviceProfileManagerSingleton == nullptr)
+	{
+		DeviceProfileManagerSingleton = NewObject<UDeviceProfileManager>();
+
+		DeviceProfileManagerSingleton->AddToRoot();
+		DeviceProfileManagerSingleton->LoadProfiles();
+
+		FString DeviceProfileSelectionModule;
+		GConfig->GetString(TEXT("DeviceProfileManager"), TEXT("DeviceProfileSelectionModule"), DeviceProfileSelectionModule, GEngineIni);
+		IDeviceProfileSelectorModule& DPSelectorModule = FModuleManager::LoadModuleChecked<IDeviceProfileSelectorModule>(*DeviceProfileSelectionModule);
+
+		UDeviceProfile* ActiveProfile = &DeviceProfileManagerSingleton->FindProfile(DPSelectorModule.GetRuntimeDeviceProfileName());
+		DeviceProfileManagerSingleton->SetActiveDeviceProfile(ActiveProfile);
+
+		InitializeSharedSamplerStates();
+	}
+	return *DeviceProfileManagerSingleton;
+}
+
+
 void UDeviceProfileManager::InitializeCVarsForActiveDeviceProfile()
 {
 	// Find the device profile selector module used in this instance
@@ -123,35 +147,24 @@ void UDeviceProfileManager::InitializeCVarsForActiveDeviceProfile()
 }
 
 
-UDeviceProfileManager::UDeviceProfileManager( const FObjectInitializer& ObjectInitializer )
-	: Super( ObjectInitializer )
+UDeviceProfile& UDeviceProfileManager::CreateProfile( const FString& ProfileName )
 {
-	RenameIndex = 0;
-#if WITH_EDITOR
-	LoadProfiles();
-#endif
-}
-
-
-UDeviceProfile* UDeviceProfileManager::CreateProfile( const FString& ProfileName )
-{
-	UDeviceProfile* NewProfile = NewObject<UDeviceProfile>(GetTransientPackage(), *ProfileName, RF_Transient | RF_Public);
+	UDeviceProfile* NewProfile = NewObject<UDeviceProfile>();
 	Profiles.Add( NewProfile );
 
 	// Inform the UI that the device list has changed
 	ManagerUpdatedDelegate.Broadcast();
 
-	return NewProfile;
+	return *NewProfile;
 }
 
 
-UDeviceProfile* UDeviceProfileManager::CreateProfile( const FString& ProfileName, const FString& ProfileType, const FString& ParentName )
+UDeviceProfile& UDeviceProfileManager::CreateProfile( const FString& ProfileName, const FString& ProfileType, const FString& ParentName )
 {
 	UDeviceProfile* DeviceProfile = FindObject<UDeviceProfile>( GetTransientPackage(), *ProfileName );
 	if( DeviceProfile == NULL )
 	{
-		DeviceProfile = NewObject<UDeviceProfile>(GetTransientPackage(), *ProfileName, RF_Transient | RF_Public);
-		DeviceProfile->LoadConfig( UDeviceProfile::StaticClass(), *DeviceProfileFileName );
+		DeviceProfile = NewObject<UDeviceProfile>(GetTransientPackage(), *ProfileName);
 		DeviceProfile->BaseProfileName = ParentName != TEXT("") ? ParentName : DeviceProfile->BaseProfileName;
 		DeviceProfile->DeviceType = ProfileType;
 
@@ -163,7 +176,7 @@ UDeviceProfile* UDeviceProfileManager::CreateProfile( const FString& ProfileName
 			DeviceProfile->Parent = FindObject<UDeviceProfile>( GetTransientPackage(), *DeviceProfile->BaseProfileName );
 			if( DeviceProfile->Parent == NULL )
 			{
-				DeviceProfile->Parent = CreateProfile( DeviceProfile->BaseProfileName, ProfileType );
+				DeviceProfile->Parent = &CreateProfile( DeviceProfile->BaseProfileName, ProfileType );
 			}
 			ObjectTemplate = CastChecked<UDeviceProfile>(DeviceProfile->Parent);
 		}
@@ -180,7 +193,7 @@ UDeviceProfile* UDeviceProfileManager::CreateProfile( const FString& ProfileName
 		ManagerUpdatedDelegate.Broadcast(); 
 	}
 
-	return DeviceProfile;
+	return *DeviceProfile;
 }
 
 
@@ -190,9 +203,9 @@ void UDeviceProfileManager::DeleteProfile( UDeviceProfile* Profile )
 }
 
 
-UDeviceProfile* UDeviceProfileManager::FindProfile( const FString& ProfileName )
+UDeviceProfile& UDeviceProfileManager::FindProfile( const FString& ProfileName )
 {
-	UDeviceProfile* FoundProfile = NULL;
+	UDeviceProfile* FoundProfile = nullptr;
 
 	for( int32 Idx = 0; Idx < Profiles.Num(); Idx++ )
 	{
@@ -204,17 +217,17 @@ UDeviceProfile* UDeviceProfileManager::FindProfile( const FString& ProfileName )
 		}
 	}
 
-	return FoundProfile;
+	return FoundProfile != nullptr ? *FoundProfile : CreateProfile(ProfileName);
 }
 
 
-const FString UDeviceProfileManager::GetDeviceProfileIniName()
+const FString UDeviceProfileManager::GetDeviceProfileIniName() const
 {
 	return DeviceProfileFileName;
 }
 
 
-FOnManagerUpdated& UDeviceProfileManager::OnManagerUpdated()
+FOnDeviceProfileManagerUpdated& UDeviceProfileManager::OnManagerUpdated()
 {
 	return ManagerUpdatedDelegate;
 }
@@ -238,6 +251,25 @@ void UDeviceProfileManager::LoadProfiles()
 				CreateProfile(NewDeviceProfileSelectorPlatformName, NewDeviceProfileSelectorPlatformType);
 			}
 		}
+
+		// Register Texture LOD settings with each Target Platform
+		ITargetPlatformManagerModule& TargetPlatformManager = GetTargetPlatformManagerRef();
+		const TArray<ITargetPlatform*>& TargetPlatforms = TargetPlatformManager.GetActiveTargetPlatforms();
+		for (int32 PlatformIndex = 0; PlatformIndex < TargetPlatforms.Num(); ++PlatformIndex)
+		{
+			ITargetPlatform* Platform = TargetPlatforms[PlatformIndex];
+			if (const UTextureLODSettings* TextureLODSettingsObj = (UTextureLODSettings*)&FindProfile(Platform->PlatformName()))
+			{
+				// Set TextureLODSettings
+				Platform->RegisterTextureLODSettings(TextureLODSettingsObj);
+			}
+			else
+			{
+				// No Device Profile found that matches target platform
+			}
+		}
+
+
 		ManagerUpdatedDelegate.Broadcast();
 	}
 }
@@ -289,7 +321,7 @@ UDeviceProfile* UDeviceProfileManager::GetActiveProfile() const
 }
 
 
-void UDeviceProfileManager::GetAllPossibleParentProfiles(const UDeviceProfile* ChildProfile, OUT TArray<UDeviceProfile*>& PossibleParentProfiles)
+void UDeviceProfileManager::GetAllPossibleParentProfiles(const UDeviceProfile* ChildProfile, OUT TArray<UDeviceProfile*>& PossibleParentProfiles) const
 {
 	for(auto& NextProfile : Profiles)
 	{
diff --git a/Engine/Source/Runtime/Engine/Private/EnginePrivate.h b/Engine/Source/Runtime/Engine/Private/EnginePrivate.h
--- a/Engine/Source/Runtime/Engine/Private/EnginePrivate.h
+++ b/Engine/Source/Runtime/Engine/Private/EnginePrivate.h
@@ -96,3 +96,6 @@ EnginePrivate.h: Unreal engine private header file.
 #include "TimerManager.h"					// Game play timers
 #include "SlateCore.h"
 #include "SlateBasics.h"
+#include "Engine/TextureLODSettings.h"
+#include "DeviceProfiles/DeviceProfile.h"
+#include "DeviceProfiles/DeviceProfileManager.h"
\ No newline at end of file
diff --git a/Engine/Source/Runtime/Engine/Private/SceneManagement.cpp b/Engine/Source/Runtime/Engine/Private/SceneManagement.cpp
--- a/Engine/Source/Runtime/Engine/Private/SceneManagement.cpp
+++ b/Engine/Source/Runtime/Engine/Private/SceneManagement.cpp
@@ -355,7 +355,7 @@ void FSharedSamplerState::InitRHI()
 
 	FSamplerStateInitializerRHI SamplerStateInitializer
 	(
-		GSystemSettings.TextureLODSettings.GetSamplerFilter(TEXTUREGROUP_World),
+	(ESamplerFilter)UDeviceProfileManager::Get().GetActiveProfile()->GetTextureLODSettings()->GetSamplerFilter(TEXTUREGROUP_World),
 		bWrap ? AM_Wrap : AM_Clamp,
 		bWrap ? AM_Wrap : AM_Clamp,
 		bWrap ? AM_Wrap : AM_Clamp,
diff --git a/Engine/Source/Runtime/Engine/Private/SystemSettings.cpp b/Engine/Source/Runtime/Engine/Private/SystemSettings.cpp
--- a/Engine/Source/Runtime/Engine/Private/SystemSettings.cpp
+++ b/Engine/Source/Runtime/Engine/Private/SystemSettings.cpp
@@ -54,13 +54,6 @@ static inline const TCHAR* GetSectionName(bool bIsEditor)
 	return bIsEditor ? GIniSectionEditor : GIniSectionGame;
 }
 
-/**
- * ctor
- */
-FSystemSettingsDataTextureDetail::FSystemSettingsDataTextureDetail()
-{
-}
-
 /**
  * ctor
  */
@@ -120,111 +113,6 @@ void FSystemSettingsData::LoadFromIni( const TCHAR* IniSection, const FString& I
 
 		IConsoleManager::Get().CallAllConsoleVariableSinks();
 	}
-
-
-	// Read the texture group LOD settings.
-	TextureLODSettings.Initialize( IniFilename, IniSection );
-}
-
-void FSystemSettingsData::SaveToIni( const TCHAR* IniSection )
-{
-	// Save the texture group LOD settings.
-	WriteTextureLODGroupsToIni( IniSection );
-
-	GConfig->Flush( false, GEngineIni );
-}
-
-/**
- * Writes all texture group LOD settings to the specified ini.
- *
- * @param	IniSection			The .ini section to save to.
- */
-void FSystemSettingsData::WriteTextureLODGroupsToIni(const TCHAR* IniSection)
-{
-#define WRITETEXTURELODGROUPTOINI(Group) WriteTextureLODGroupToIni(Group, TEXT(#Group), IniSection);
-	FOREACH_ENUM_TEXTUREGROUP(WRITETEXTURELODGROUPTOINI)
-#undef WRITETEXTURELODGROUPTOINI
-}
-
-/**
-* Returns a string for the specified texture group LOD settings to the specified ini.
-*
-* @param	TextureGroupID		Index/enum of the group
-* @param	GroupName			String representation of the texture group
-*/
-FString FSystemSettingsData::GetLODGroupString( TextureGroup TextureGroupID, const TCHAR* GroupName )
-{
-	const FExposedTextureLODSettings::FTextureLODGroup& Group = TextureLODSettings.GetTextureLODGroup(TextureGroupID);
-
-	const int32 MinLODSize = 1 << Group.MinLODMipCount;
-	const int32 MaxLODSize = 1 << Group.MaxLODMipCount;
-
-	FName MinMagFilter = NAME_Aniso;
-	FName MipFilter = NAME_Linear;
-	switch(Group.Filter)
-	{
-		case SF_Point:
-			MinMagFilter = NAME_Point;
-			MipFilter = NAME_Point;
-			break;
-		case SF_Bilinear:
-			MinMagFilter = NAME_Linear;
-			MipFilter = NAME_Point;
-			break;
-		case SF_Trilinear:
-			MinMagFilter = NAME_Linear;
-			MipFilter = NAME_Linear;
-			break;
-		case SF_AnisotropicPoint:
-			MinMagFilter = NAME_Aniso;
-			MipFilter = NAME_Point;
-			break;
-		case SF_AnisotropicLinear:
-			MinMagFilter = NAME_Aniso;
-			MipFilter = NAME_Linear;
-			break;
-	}
-
-	FString NumStreamedMipsText;
-	if ( Group.NumStreamedMips >= 0 )
-	{
-		NumStreamedMipsText = FString::Printf( TEXT(",NumStreamedMips=%i"), Group.NumStreamedMips );
-	}
-
-	return FString::Printf( TEXT("(MinLODSize=%i,MaxLODSize=%i,LODBias=%i,MinMagFilter=%s,MipFilter=%s%s,MipGenSettings=%s)"),
-		MinLODSize, MaxLODSize, Group.LODBias, *MinMagFilter.GetPlainNameString(), *MipFilter.GetPlainNameString(), *NumStreamedMipsText, UTexture::GetMipGenSettingsString(Group.MipGenSettings) );
-}
-
-/**
- * Writes the specified texture group LOD settings to the specified ini.
- *
- * @param	TextureGroupID		Index/enum of the group to parse
- * @param	GroupName			String representation of the texture group, to be used as the ini key.
- * @param	IniSection			The .ini section to save to.
- */
-void FSystemSettingsData::WriteTextureLODGroupToIni(TextureGroup TextureGroupID, const TCHAR* GroupName, const TCHAR* IniSection)
-{
-	const FString Entry = GetLODGroupString( TextureGroupID, GroupName );
-	GConfig->SetString( IniSection, GroupName, *Entry, GEngineIni );
-}
-
-/**
- * Dump helpers
- */
-void FSystemSettingsData::DumpTextureLODGroups()
-{
-#define DUMPTEXTURELODGROUP(Group) DumpTextureLODGroup(Group, TEXT(#Group));
-	FOREACH_ENUM_TEXTUREGROUP(DUMPTEXTURELODGROUP)
-#undef DUMPTEXTURELODGROUP
-}
-
-/**
- * Dump helpers
- */
-void FSystemSettingsData::DumpTextureLODGroup(TextureGroup TextureGroupID, const TCHAR* GroupName)
-{
-	const FString Entry = GetLODGroupString( TextureGroupID, GroupName );
-	UE_LOG(LogSystemSettings, Log, TEXT("\t%s: %s"), GroupName, *Entry );
 }
 
 /**
@@ -360,14 +248,6 @@ bool FSystemSettings::Exec( UWorld* InWorld, const TCHAR* Cmd, FOutputDevice& Ar
 	return false;
 }
 
-void FSystemSettings::SetTextureLODGroup(TextureGroup TextureGroupID, int32 MinLODSize, int32 MaxLODSize, int32 LODBias, TextureMipGenSettings MipGenSettings)
-{
-	TextureLODSettings.GetTextureLODGroup(TextureGroupID).MinLODMipCount	= FMath::CeilLogTwo( MinLODSize );
-	TextureLODSettings.GetTextureLODGroup(TextureGroupID).MaxLODMipCount	= FMath::CeilLogTwo( MaxLODSize );
-	TextureLODSettings.GetTextureLODGroup(TextureGroupID).LODBias			= LODBias;
-	TextureLODSettings.GetTextureLODGroup(TextureGroupID).MipGenSettings	= MipGenSettings;
-}
-
 /*-----------------------------------------------------------------------------
 	Resolution.
 -----------------------------------------------------------------------------*/
@@ -380,102 +260,6 @@ void FSystemSettings::LoadFromIni()
 	FSystemSettingsData::LoadFromIni(GetSectionName(bIsEditor));
 }
 
-void FSystemSettings::SaveToIni()
-{
-	// don't write changes in the editor
-	if (bIsEditor)
-	{
-		UE_LOG(LogSystemSettings, Log, TEXT("Can't save system settings to ini in an editor mode"));
-		return;
-	}
-	FSystemSettingsData::SaveToIni(GetSectionName(bIsEditor));
-}
-
-/**
- * Helper for ApplyNewSettings when the engine is running. Applies the changes needed for the runtime system.
- *
- * We can assume the game is running if this code is called.
- */
-void FSystemSettings::ApplySettingsAtRuntime(const FSystemSettingsData& NewSettings, bool bWriteToIni)
-{
-	// Some of these settings are shared between threads, so we
-	// must flush the rendering thread before changing anything.
-	FlushRenderingCommands();
-
-	// Track settings we might have to put back
-	FExposedTextureLODSettings InMemoryTextureLODSettings = TextureLODSettings;
-
-	// Read settings from .ini.  This is necessary because settings which need to wait for a restart
-	// will be on disk but may not be in memory.  Therefore, we read from disk before capturing old
-	// values to revert to.
-	LoadFromIni();
-
-	// see what settings are actually changing.
-	// Ugly casts because system settings is multi-inherited from all the consituent classes for backwards compatibility
-
-	// Texture Detail
-	bool bTextureDetailChanged = (const FSystemSettingsDataTextureDetail&)(*this) != (const FSystemSettingsDataTextureDetail&)NewSettings;
-
-	// Make a copy of the existing settings so we can compare for changes
-	FSystemSettingsData OldData = *this;
-
-	// Set new settings. Would look prettier if we didn't derive from the Data class...
-	(FSystemSettingsData&)(*this) = NewSettings;
-
-	// apply any runtime changes that need to be made
-	bool bUpdateTextureStreamingSucceeded = false;
-	if (bTextureDetailChanged)
-	{
-		bUpdateTextureStreamingSucceeded = UTexture::ForceUpdateTextureStreaming();
-	}
-
-	// If requested, save the settings to ini.
-	if ( bWriteToIni )
-	{
-		SaveToIni();
-	}
-
-	// If texture detail settings couldn't be applied because we're loading seekfree,
-	// revert the new settings to their previous in-memory values.
-	if ( bTextureDetailChanged && !bUpdateTextureStreamingSucceeded )
-	{
-		TextureLODSettings = InMemoryTextureLODSettings;
-	}
-}
-
-/**
- * Sets new system settings (optionally writes out to the ini). 
- */
-void FSystemSettings::ApplyNewSettings( const FSystemSettingsData& NewSettings, bool bWriteToIni )
-{
-	// we can set any setting before the engine is initialized so don't bother restoring values.
-	bool bEngineIsInitialized = GEngine != NULL;
-
-	// if the engine is running, there are certain values we can't set immediately
-	if (bEngineIsInitialized)
-	{
-		// apply settings to the runtime system.
-		ApplySettingsAtRuntime(NewSettings, bWriteToIni);
-
-		ApplyOverrides();
-	}
-	else
-	{
-		// if the engine is not initialized we don't need to worry about all the deferred settings etc. 
-		// as we do above.
-		// Set new settings. Would look prettier if we didn't derive from the Data class...
-		(FSystemSettingsData&)(*this) = NewSettings;
-
-		// If requested, save the settings to ini.
-		if ( bWriteToIni )
-		{
-			SaveToIni();
-		}
-
-		ApplyOverrides();
-	}
-}
-
 void FSystemSettings::ApplyOverrides()
 {
 	EConsoleVariableFlags SetBy = ECVF_SetByMask;
@@ -530,17 +314,6 @@ void FSystemSettings::ApplyOverrides()
 			static auto CVar = IConsoleManager::Get().FindConsoleVariable(TEXT("r.Shadow.PreShadowResolutionFactor"));
 			CVar->Set(1.0f, SetBy);
 		}
-
-		for (int32 GroupIndex = 0; GroupIndex < TEXTUREGROUP_MAX; GroupIndex++)
-		{
-			FTextureLODSettings::FTextureLODGroup& CurrentGroup = TextureLODSettings.GetTextureLODGroup(GroupIndex);
-			// Use the best quality texture filtering
-			CurrentGroup.Filter = SF_AnisotropicLinear;
-			// Raise texture max sizes to 4096
-			CurrentGroup.MinLODMipCount = 12;
-			CurrentGroup.MaxLODMipCount = 12;
-			CurrentGroup.LODBias = -1000;
-		}
 	}
 }
 
diff --git a/Engine/Source/Runtime/Engine/Private/Texture2D.cpp b/Engine/Source/Runtime/Engine/Private/Texture2D.cpp
--- a/Engine/Source/Runtime/Engine/Private/Texture2D.cpp
+++ b/Engine/Source/Runtime/Engine/Private/Texture2D.cpp
@@ -439,7 +439,7 @@ FString UTexture2D::GetDesc()
 	UpdateCachedLODBias();
 #endif //#if WITH_EDITOR
 
-	GSystemSettings.TextureLODSettings.ComputeInGameMaxResolution(GetCachedLODBias(), *this, EffectiveSizeX, EffectiveSizeY);
+	UDeviceProfileManager::Get().GetActiveProfile()->GetTextureLODSettings()->ComputeInGameMaxResolution(GetCachedLODBias(), *this, EffectiveSizeX, EffectiveSizeY);
 
 	return FString::Printf( TEXT("%s %dx%d -> %dx%d[%s]"), 
 		NeverStream ? TEXT("NeverStreamed") : TEXT("Streamed"), 
@@ -1352,7 +1352,7 @@ void FTexture2DResource::CreateSamplerStates(float MipMapBias)
 	// Create the sampler state RHI resource.
 	FSamplerStateInitializerRHI SamplerStateInitializer
 	(
-	  GSystemSettings.TextureLODSettings.GetSamplerFilter(Owner),
+	  (ESamplerFilter)UDeviceProfileManager::Get().GetActiveProfile()->GetTextureLODSettings()->GetSamplerFilter(Owner),
 	  Owner->AddressX == TA_Wrap ? AM_Wrap : (Owner->AddressX == TA_Clamp ? AM_Clamp : AM_Mirror),
 	  Owner->AddressY == TA_Wrap ? AM_Wrap : (Owner->AddressY == TA_Clamp ? AM_Clamp : AM_Mirror),
 	  AM_Wrap,
@@ -1363,7 +1363,7 @@ void FTexture2DResource::CreateSamplerStates(float MipMapBias)
 	// Create a custom sampler state for using this texture in a deferred pass, where ddx / ddy are discontinuous
 	FSamplerStateInitializerRHI DeferredPassSamplerStateInitializer
 	(
-	  GSystemSettings.TextureLODSettings.GetSamplerFilter(Owner),
+	  (ESamplerFilter)UDeviceProfileManager::Get().GetActiveProfile()->GetTextureLODSettings()->GetSamplerFilter(Owner),
 	  Owner->AddressX == TA_Wrap ? AM_Wrap : (Owner->AddressX == TA_Clamp ? AM_Clamp : AM_Mirror),
 	  Owner->AddressY == TA_Wrap ? AM_Wrap : (Owner->AddressY == TA_Clamp ? AM_Clamp : AM_Mirror),
 	  AM_Wrap,
@@ -2330,7 +2330,7 @@ FIncomingTextureArrayDataEntry::FIncomingTextureArrayDataEntry(UTexture2D* InTex
 	NumMips = InTexture->GetNumMips();
 	LODGroup = (TextureGroup)InTexture->LODGroup;
 	Format = InTexture->GetPixelFormat();
-	Filter = GSystemSettings.TextureLODSettings.GetSamplerFilter(InTexture);
+	Filter = (ESamplerFilter)UDeviceProfileManager::Get().GetActiveProfile()->GetTextureLODSettings()->GetSamplerFilter(InTexture);
 	bSRGB = InTexture->SRGB;
 
 	MipData.Empty(NumMips);
diff --git a/Engine/Source/Runtime/Engine/Private/Texture2DDynamic.cpp b/Engine/Source/Runtime/Engine/Private/Texture2DDynamic.cpp
--- a/Engine/Source/Runtime/Engine/Private/Texture2DDynamic.cpp
+++ b/Engine/Source/Runtime/Engine/Private/Texture2DDynamic.cpp
@@ -35,7 +35,7 @@ void FTexture2DDynamicResource::InitRHI()
 	// Create the sampler state RHI resource.
 	FSamplerStateInitializerRHI SamplerStateInitializer
 	(
-		GSystemSettings.TextureLODSettings.GetSamplerFilter( Owner ),
+		(ESamplerFilter)UDeviceProfileManager::Get().GetActiveProfile()->GetTextureLODSettings()->GetSamplerFilter( Owner ),
 		AM_Wrap,
 		AM_Wrap,
 		AM_Wrap
diff --git a/Engine/Source/Runtime/Engine/Private/TextureCube.cpp b/Engine/Source/Runtime/Engine/Private/TextureCube.cpp
--- a/Engine/Source/Runtime/Engine/Private/TextureCube.cpp
+++ b/Engine/Source/Runtime/Engine/Private/TextureCube.cpp
@@ -228,7 +228,7 @@ public:
 		// Create the sampler state RHI resource.
 		FSamplerStateInitializerRHI SamplerStateInitializer
 		(
-			GSystemSettings.TextureLODSettings.GetSamplerFilter( Owner ),
+			(ESamplerFilter)UDeviceProfileManager::Get().GetActiveProfile()->GetTextureLODSettings()->GetSamplerFilter( Owner ),
 			AM_Clamp,
 			AM_Clamp,
 			AM_Clamp
diff --git a/Engine/Source/Runtime/Engine/Private/TextureDerivedData.cpp b/Engine/Source/Runtime/Engine/Private/TextureDerivedData.cpp
--- a/Engine/Source/Runtime/Engine/Private/TextureDerivedData.cpp
+++ b/Engine/Source/Runtime/Engine/Private/TextureDerivedData.cpp
@@ -276,7 +276,7 @@ static void GetTextureDerivedDataKey(
  */
 static void GetTextureBuildSettings(
 	const UTexture& Texture,
-	const FTextureLODSettings& TextureLODSettings,
+	const UTextureLODSettings& TextureLODSettings,
 	FTextureBuildSettings& OutBuildSettings
 	)
 {
@@ -347,7 +347,7 @@ static void GetTextureBuildSettings(
 	OutBuildSettings.bFlipGreenChannel = Texture.bFlipGreenChannel;
 	OutBuildSettings.CompositeTextureMode = Texture.CompositeTextureMode;
 	OutBuildSettings.CompositePower = Texture.CompositePower;
-	OutBuildSettings.LODBias = GSystemSettings.TextureLODSettings.CalculateLODBias(Texture.Source.GetSizeX(), Texture.Source.GetSizeY(), Texture.LODGroup, Texture.LODBias, Texture.NumCinematicMipLevels, Texture.MipGenSettings);
+	OutBuildSettings.LODBias = UDeviceProfileManager::Get().GetActiveProfile()->GetTextureLODSettings()->CalculateLODBias(Texture.Source.GetSizeX(), Texture.Source.GetSizeY(), Texture.LODGroup, Texture.LODBias, Texture.NumCinematicMipLevels, Texture.MipGenSettings);
 	OutBuildSettings.bStreamable = !Texture.NeverStream && (Texture.LODGroup != TEXTUREGROUP_UI) && (Cast<const UTexture2D>(&Texture) != NULL);
 	OutBuildSettings.PowerOfTwoMode = Texture.PowerOfTwoMode;
 	OutBuildSettings.PaddingColor = Texture.PaddingColor;
@@ -393,7 +393,8 @@ static void GetBuildSettingsForRunningPlatform(
 
 		// Assume there is at least one format and the first one is what we want at runtime.
 		check(PlatformFormats.Num());
-		GetTextureBuildSettings(Texture, GSystemSettings.TextureLODSettings, OutBuildSettings);
+		const UTextureLODSettings* LODSettings = (UTextureLODSettings*)&UDeviceProfileManager::Get().FindProfile(CurrentPlatform->PlatformName());
+		GetTextureBuildSettings(Texture, *LODSettings, OutBuildSettings);
 		OutBuildSettings.TextureFormatName = PlatformFormats[0];
 	}
 }
@@ -1384,7 +1385,7 @@ void UTexture2D::GetMipData(int32 FirstMipToLoad, void** OutMipData)
 
 void UTexture::UpdateCachedLODBias( bool bIncTextureMips )
 {
-	CachedCombinedLODBias = GSystemSettings.TextureLODSettings.CalculateLODBias( this, bIncTextureMips );
+	CachedCombinedLODBias = UDeviceProfileManager::Get().GetActiveProfile()->GetTextureLODSettings()->CalculateLODBias(this, bIncTextureMips);
 }
 
 #if WITH_EDITOR
diff --git a/Engine/Source/Runtime/Engine/Private/TextureLODSettings.cpp b/Engine/Source/Runtime/Engine/Private/TextureLODSettings.cpp
--- a/Engine/Source/Runtime/Engine/Private/TextureLODSettings.cpp
+++ b/Engine/Source/Runtime/Engine/Private/TextureLODSettings.cpp
@@ -1,25 +1,10 @@
 // Copyright 1998-2015 Epic Games, Inc. All Rights Reserved.
 
 #include "EnginePrivate.h"
-#include "TextureLODSettings.h"
 
-void FTextureLODSettings::Initialize( const FString& IniFilename, const TCHAR* IniSection )
+UTextureLODSettings::UTextureLODSettings(const FObjectInitializer& ObjectInitializer)
+	: Super(ObjectInitializer)
 {
-	// look up the file object
-	FConfigFile* ConfigFile = GConfig->FindConfigFile(IniFilename);
-	if (ConfigFile)
-	{
-		// pass to the other initialize function
-		Initialize(*ConfigFile, IniSection);
-	}
-}
-
-void FTextureLODSettings::Initialize(const FConfigFile& IniFile, const TCHAR* IniSection)
-{
-	// Read individual entries from a config file.
-#define GROUPREADENTRY(g) ReadEntry( g, TEXT(#g), IniFile, IniSection );
-	FOREACH_ENUM_TEXTUREGROUP(GROUPREADENTRY)
-#undef GROUPREADENTRY
 }
 
 /**
@@ -27,7 +12,7 @@ void FTextureLODSettings::Initialize(const FConfigFile& IniFile, const TCHAR* In
  *
  * @return array of texture group names
  */
-TArray<FString> FTextureLODSettings::GetTextureGroupNames()
+TArray<FString> UTextureLODSettings::GetTextureGroupNames()
 {
 	TArray<FString> TextureGroupNames;
 
@@ -38,94 +23,46 @@ TArray<FString> FTextureLODSettings::GetTextureGroupNames()
 	return TextureGroupNames;
 }
 
-void FTextureLODSettings::ReadEntry( int32 GroupId, const TCHAR* GroupName, const FConfigFile& IniFile, const TCHAR* IniSection )
+void UTextureLODSettings::SetupLODGroup(int32 GroupId)
 {
-	// Look for string in filename/ section.
-	FString Entry;
-	if (IniFile.GetString(IniSection, GroupName, Entry))
-	{
-		// Trim whitespace at the beginning.
-		Entry = Entry.Trim();
-		// Remove brackets.
-		Entry = Entry.Replace( TEXT("("), TEXT("") );
-		Entry = Entry.Replace( TEXT(")"), TEXT("") );
-		
-		// Parse minimum LOD mip count.
-		int32	MinLODSize = 0;
-		if( FParse::Value( *Entry, TEXT("MinLODSize="), MinLODSize ) )
-		{
-			TextureLODGroups[GroupId].MinLODMipCount = FMath::CeilLogTwo( MinLODSize );
-		}
-
-		// Parse maximum LOD mip count.
-		int32 MaxLODSize = 0;
-		if( FParse::Value( *Entry, TEXT("MaxLODSize="), MaxLODSize ) )
-		{
-			TextureLODGroups[GroupId].MaxLODMipCount = FMath::CeilLogTwo( MaxLODSize );
-		}
-
-		// Parse LOD bias.
-		int32 LODBias = 0;
-		if( FParse::Value( *Entry, TEXT("LODBias="), LODBias ) )
-		{
-			TextureLODGroups[GroupId].LODBias = LODBias;
-		}
+	TextureLODGroups[GroupId].MinLODMipCount = FMath::CeilLogTwo(TextureLODGroups[GroupId].MinLODSize);
+	TextureLODGroups[GroupId].MaxLODMipCount = FMath::CeilLogTwo(TextureLODGroups[GroupId].MaxLODSize);
 
-		// Parse min/map/mip filter names.
-		FName MinMagFilter = NAME_Aniso;
-		FParse::Value( *Entry, TEXT("MinMagFilter="), MinMagFilter );
-		FName MipFilter = NAME_Point;
-		FParse::Value( *Entry, TEXT("MipFilter="), MipFilter );
+	// Convert into single filter enum. The code is layed out such that invalid input will 
+	// map to the default state of highest quality filtering.
 
+	// Linear filtering
+	if (TextureLODGroups[GroupId].MinMagFilter == NAME_Linear)
+	{
+		if (TextureLODGroups[GroupId].MipFilter == NAME_Point)
 		{
-			FString MipGenSettings;
-			FParse::Value( *Entry, TEXT("MipGenSettings="), MipGenSettings );
-			TextureLODGroups[GroupId].MipGenSettings = UTexture::GetMipGenSettingsFromString(*MipGenSettings, true);
+			TextureLODGroups[GroupId].Filter = ETextureSamplerFilter::Bilinear;
 		}
-
-		// Convert into single filter enum. The code is layed out such that invalid input will 
-		// map to the default state of highest quality filtering.
-
-		// Linear filtering
-		if( MinMagFilter == NAME_Linear )
+		else
 		{
-			if( MipFilter == NAME_Point )
-			{
-				TextureLODGroups[GroupId].Filter = SF_Bilinear;
-			}
-			else
-			{
-				TextureLODGroups[GroupId].Filter = SF_Trilinear;
-			}
+			TextureLODGroups[GroupId].Filter = ETextureSamplerFilter::Trilinear;
 		}
-		// Point. Don't even care about mip filter.
-		else if( MinMagFilter == NAME_Point )
+	}
+	// Point. Don't even care about mip filter.
+	else if (TextureLODGroups[GroupId].MinMagFilter == NAME_Point)
+	{
+		TextureLODGroups[GroupId].Filter = ETextureSamplerFilter::Point;
+	}
+	// Aniso or unknown.
+	else
+	{
+		if (TextureLODGroups[GroupId].MipFilter == NAME_Point)
 		{
-			TextureLODGroups[GroupId].Filter = SF_Point;
+			TextureLODGroups[GroupId].Filter = ETextureSamplerFilter::AnisotropicPoint;
 		}
-		// Aniso or unknown.
 		else
 		{
-			if( MipFilter == NAME_Point )
-			{
-				TextureLODGroups[GroupId].Filter = SF_AnisotropicPoint;
-			}
-			else
-			{
-				TextureLODGroups[GroupId].Filter = SF_AnisotropicLinear;
-			}
-		}
-
-		// Parse NumStreamedMips
-		int32 NumStreamedMips = -1;
-		if( FParse::Value( *Entry, TEXT("NumStreamedMips="), NumStreamedMips ) )
-		{
-			TextureLODGroups[GroupId].NumStreamedMips = NumStreamedMips;
+			TextureLODGroups[GroupId].Filter = ETextureSamplerFilter::AnisotropicLinear;
 		}
 	}
 }
 
-int32 FTextureLODSettings::CalculateLODBias( const UTexture* Texture, bool bIncTextureMips ) const
+int32 UTextureLODSettings::CalculateLODBias(const UTexture* Texture, bool bIncTextureMips) const
 {	
 	check( Texture );
 	TextureMipGenSettings MipGenSetting = TMGS_MAX;
@@ -135,13 +72,13 @@ int32 FTextureLODSettings::CalculateLODBias( const UTexture* Texture, bool bIncT
 	return CalculateLODBias(Texture->GetSurfaceWidth(), Texture->GetSurfaceHeight(), Texture->LODGroup, (bIncTextureMips ? Texture->LODBias : 0), (bIncTextureMips ? Texture->NumCinematicMipLevels : 0), MipGenSetting);
 }
 
-int32 FTextureLODSettings::CalculateLODBias( int32 Width, int32 Height, int32 LODGroup, int32 LODBias, int32 NumCinematicMipLevels, TextureMipGenSettings InMipGenSetting ) const
+int32 UTextureLODSettings::CalculateLODBias(int32 Width, int32 Height, int32 LODGroup, int32 LODBias, int32 NumCinematicMipLevels, TextureMipGenSettings InMipGenSetting) const
 {	
 	// Find LOD group.
 	const FTextureLODGroup& LODGroupInfo = TextureLODGroups[LODGroup];
 
 	// Test to see if we have no mip generation as in which case the LOD bias will be ignored
-	const TextureMipGenSettings FinalMipGenSetting = (InMipGenSetting == TMGS_FromTextureGroup) ? LODGroupInfo.MipGenSettings : InMipGenSetting;
+	const TextureMipGenSettings FinalMipGenSetting = (InMipGenSetting == TMGS_FromTextureGroup) ? (TextureMipGenSettings)LODGroupInfo.MipGenSettings : InMipGenSetting;
 	if ( FinalMipGenSetting == TMGS_NoMipmaps )
 	{
 		return 0;
@@ -164,7 +101,7 @@ int32 FTextureLODSettings::CalculateLODBias( int32 Width, int32 Height, int32 LO
 /** 
 * Useful for stats in the editor.
 */
-void FTextureLODSettings::ComputeInGameMaxResolution(int32 LODBias, UTexture &Texture, uint32 &OutSizeX, uint32 &OutSizeY) const
+void UTextureLODSettings::ComputeInGameMaxResolution(int32 LODBias, UTexture &Texture, uint32 &OutSizeX, uint32 &OutSizeY) const
 {
 	uint32 ImportedSizeX = FMath::TruncToInt(Texture.GetSurfaceWidth());
 	uint32 ImportedSizeY = FMath::TruncToInt(Texture.GetSurfaceHeight());
@@ -188,19 +125,26 @@ void FTextureLODSettings::ComputeInGameMaxResolution(int32 LODBias, UTexture &Te
 * @param   GroupIndex      usually from Texture.LODGroup
 * @return                  A handle to the indexed LOD group. 
 */
-const FTextureLODSettings::FTextureLODGroup& FTextureLODSettings::GetTextureLODGroup(TextureGroup GroupIndex) const
+FTextureLODGroup& UTextureLODSettings::GetTextureLODGroup(TextureGroup GroupIndex)
 {
-	if((uint32)GroupIndex >= TEXTUREGROUP_MAX)
-	{
-		// to prevent crash
-		GroupIndex = (TextureGroup)0;
-	}
-//	check((uint32)GroupIndex < TEXTUREGROUP_MAX);
+	check(GroupIndex >= 0 && GroupIndex < TEXTUREGROUP_MAX);
+	return TextureLODGroups[GroupIndex];
+}
+
+/**
+* TextureLODGroups access with bounds check
+*
+* @param   GroupIndex      usually from Texture.LODGroup
+* @return                  A handle to the indexed LOD group.
+*/
+const FTextureLODGroup& UTextureLODSettings::GetTextureLODGroup(TextureGroup GroupIndex) const
+{
+	check(GroupIndex >= 0 && GroupIndex < TEXTUREGROUP_MAX);
 	return TextureLODGroups[GroupIndex];
 }
 
 #if WITH_EDITORONLY_DATA
-void FTextureLODSettings::GetMipGenSettings( const UTexture& Texture, TextureMipGenSettings& OutMipGenSettings, float& OutSharpen, uint32& OutKernelSize, bool& bOutDownsampleWithAverage, bool& bOutSharpenWithoutColorShift, bool &bOutBorderColorBlack ) const
+void UTextureLODSettings::GetMipGenSettings(const UTexture& Texture, TextureMipGenSettings& OutMipGenSettings, float& OutSharpen, uint32& OutKernelSize, bool& bOutDownsampleWithAverage, bool& bOutSharpenWithoutColorShift, bool &bOutBorderColorBlack) const
 {
 	TextureMipGenSettings Setting = (TextureMipGenSettings)Texture.MipGenSettings;
 
@@ -251,7 +195,7 @@ void FTextureLODSettings::GetMipGenSettings( const UTexture& Texture, TextureMip
  * @param	InLODGroup		The LOD Group ID 
  * @return	LODBias
  */
-int32 FTextureLODSettings::GetTextureLODGroupLODBias( int32 InLODGroup ) const
+int32 UTextureLODSettings::GetTextureLODGroupLODBias(int32 InLODGroup) const
 {
 	int32 Retval = 0;
 
@@ -271,7 +215,7 @@ int32 FTextureLODSettings::GetTextureLODGroupLODBias( int32 InLODGroup ) const
  * @param	InLODGroup		The LOD Group ID 
  * @return	Number of streaming mip-levels for textures in the specified LODGroup
  */
-int32 FTextureLODSettings::GetMinLODMipCount( int32 InLODGroup ) const
+int32 UTextureLODSettings::GetMinLODMipCount(int32 InLODGroup) const
 {
 	int32 Retval = 0;
 
@@ -292,7 +236,7 @@ int32 FTextureLODSettings::GetMinLODMipCount( int32 InLODGroup ) const
  * @param	InLODGroup		The LOD Group ID 
  * @return	Number of streaming mip-levels for textures in the specified LODGroup
  */
-int32 FTextureLODSettings::GetMaxLODMipCount( int32 InLODGroup ) const
+int32 UTextureLODSettings::GetMaxLODMipCount(int32 InLODGroup) const
 {
 	int32 Retval = 0;
 
@@ -311,7 +255,7 @@ int32 FTextureLODSettings::GetMaxLODMipCount( int32 InLODGroup ) const
  * @param	InLODGroup		The LOD Group ID 
  * @return	Number of streaming mip-levels for textures in the specified LODGroup
  */
-int32 FTextureLODSettings::GetNumStreamedMips( int32 InLODGroup ) const
+int32 UTextureLODSettings::GetNumStreamedMips(int32 InLODGroup) const
 {
 	int32 Retval = 0;
 
@@ -328,7 +272,7 @@ int32 FTextureLODSettings::GetNumStreamedMips( int32 InLODGroup ) const
  * @param	InLODGroup		The LOD Group ID 
  * @return	TextureMipGenSettings for lod group
  */
- const TextureMipGenSettings& FTextureLODSettings::GetTextureMipGenSettings( int32 InLODGroup ) const
+const TextureMipGenSettings UTextureLODSettings::GetTextureMipGenSettings(int32 InLODGroup) const
 {
 	return TextureLODGroups[InLODGroup].MipGenSettings; 
 }
@@ -342,16 +286,16 @@ int32 FTextureLODSettings::GetNumStreamedMips( int32 InLODGroup ) const
  * @param	Texture		Texture to retrieve filter state for, must not be 0
  * @return	Filter sampler state for passed in texture
  */
-ESamplerFilter FTextureLODSettings::GetSamplerFilter( const UTexture* Texture ) const
+ETextureSamplerFilter UTextureLODSettings::GetSamplerFilter(const UTexture* Texture) const
 {
 	// Default to point filtering.
-	ESamplerFilter Filter = SF_Point;
+	ETextureSamplerFilter Filter = ETextureSamplerFilter::Point;
 
 	switch(Texture->Filter)
 	{
-		case TF_Nearest: Filter = SF_Point; break;
-		case TF_Bilinear: Filter = SF_Bilinear; break;
-		case TF_Trilinear: Filter = SF_Trilinear; break;
+		case TF_Nearest: Filter = ETextureSamplerFilter::Point; break;
+		case TF_Bilinear: Filter = ETextureSamplerFilter::Bilinear; break;
+		case TF_Trilinear: Filter = ETextureSamplerFilter::Trilinear; break;
 
 		// TF_Default
 		default:
@@ -362,7 +306,7 @@ ESamplerFilter FTextureLODSettings::GetSamplerFilter( const UTexture* Texture )
 	return Filter;
 }
 
-ESamplerFilter FTextureLODSettings::GetSamplerFilter( int32 InLODGroup) const
+ETextureSamplerFilter UTextureLODSettings::GetSamplerFilter(int32 InLODGroup) const
 {
 	return TextureLODGroups[InLODGroup].Filter;
 }
\ No newline at end of file
diff --git a/Engine/Source/Runtime/Engine/Private/TextureRenderTarget2D.cpp b/Engine/Source/Runtime/Engine/Private/TextureRenderTarget2D.cpp
--- a/Engine/Source/Runtime/Engine/Private/TextureRenderTarget2D.cpp
+++ b/Engine/Source/Runtime/Engine/Private/TextureRenderTarget2D.cpp
@@ -337,7 +337,7 @@ void FTextureRenderTarget2DResource::InitDynamicRHI()
 	// Create the sampler state RHI resource.
 	FSamplerStateInitializerRHI SamplerStateInitializer
 	(
-		GSystemSettings.TextureLODSettings.GetSamplerFilter( Owner ),
+		(ESamplerFilter)UDeviceProfileManager::Get().GetActiveProfile()->GetTextureLODSettings()->GetSamplerFilter( Owner ),
 		Owner->AddressX == TA_Wrap ? AM_Wrap : (Owner->AddressX == TA_Clamp ? AM_Clamp : AM_Mirror),
 		Owner->AddressY == TA_Wrap ? AM_Wrap : (Owner->AddressY == TA_Clamp ? AM_Clamp : AM_Mirror),
 		AM_Wrap
diff --git a/Engine/Source/Runtime/Engine/Private/TextureRenderTargetCube.cpp b/Engine/Source/Runtime/Engine/Private/TextureRenderTargetCube.cpp
--- a/Engine/Source/Runtime/Engine/Private/TextureRenderTargetCube.cpp
+++ b/Engine/Source/Runtime/Engine/Private/TextureRenderTargetCube.cpp
@@ -251,7 +251,7 @@ void FTextureRenderTargetCubeResource::InitDynamicRHI()
 	// Create the sampler state RHI resource.
 	FSamplerStateInitializerRHI SamplerStateInitializer
 	(
-		GSystemSettings.TextureLODSettings.GetSamplerFilter(Owner),
+		(ESamplerFilter)UDeviceProfileManager::Get().GetActiveProfile()->GetTextureLODSettings()->GetSamplerFilter(Owner),
 		AM_Clamp,
 		AM_Clamp,
 		AM_Clamp
diff --git a/Engine/Source/Runtime/Engine/Private/UnrealEngine.cpp b/Engine/Source/Runtime/Engine/Private/UnrealEngine.cpp
--- a/Engine/Source/Runtime/Engine/Private/UnrealEngine.cpp
+++ b/Engine/Source/Runtime/Engine/Private/UnrealEngine.cpp
@@ -3344,7 +3344,7 @@ bool UEngine::HandleListTexturesCommand( const TCHAR* Cmd, FOutputDevice& Ar )
 	SortedTextures.Sort( FCompareFSortedTexture( bAlphaSort ) );
 
 	// Retrieve mapping from LOD group enum value to text representation.
-	TArray<FString> TextureGroupNames = FTextureLODSettings::GetTextureGroupNames();
+	TArray<FString> TextureGroupNames = UTextureLODSettings::GetTextureGroupNames();
 
 	// Display.
 	int32 TotalMaxSize		= 0;
@@ -10652,17 +10652,6 @@ bool UEngine::ShouldAbsorbAuthorityOnlyEvent()
 }
 
 
-UDeviceProfileManager* UEngine::GetDeviceProfileManager()
-{
-	if(DeviceProfileManager == NULL)
-	{
-		DeviceProfileManager = NewObject<UDeviceProfileManager>(GetTransientPackage(), TEXT("GlobalDeviceProfileManager"), RF_Public | RF_Transient);
-	}
-
-	return DeviceProfileManager;
-}
-
-
 bool UEngine::ShouldAbsorbCosmeticOnlyEvent()
 {
 	for (auto It = WorldList.CreateIterator(); It; ++It)
diff --git a/Engine/Source/Runtime/Engine/Public/Engine.h b/Engine/Source/Runtime/Engine/Public/Engine.h
--- a/Engine/Source/Runtime/Engine/Public/Engine.h
+++ b/Engine/Source/Runtime/Engine/Public/Engine.h
@@ -292,6 +292,8 @@
 #include "Engine/StreamableManager.h"
 #include "Engine/TextureDefines.h"
 #include "Tests/TextPropertyTestObject.h"
+#include "Engine/TextureLODSettings.h"
+#include "DeviceProfiles/DeviceProfileManager.h"
 #include "Engine/Texture.h"
 #include "Engine/LightMapTexture2D.h"
 #include "Engine/ShadowMapTexture2D.h"
@@ -318,4 +320,4 @@
 #include "EngineUtils.h"
 #include "TimerManager.h"					// Game play timers
 #include "SlateCore.h"
-#include "SlateBasics.h"
+#include "SlateBasics.h"
\ No newline at end of file
diff --git a/Engine/Source/Runtime/Engine/Public/SystemSettings.h b/Engine/Source/Runtime/Engine/Public/SystemSettings.h
--- a/Engine/Source/Runtime/Engine/Public/SystemSettings.h
+++ b/Engine/Source/Runtime/Engine/Public/SystemSettings.h
@@ -6,89 +6,25 @@
 
 #pragma once
 
-#include "TextureLODSettings.h"
 #include "ShowFlags.h"
 
-/*-----------------------------------------------------------------------------
-	System settings and scalability options.
------------------------------------------------------------------------------*/
-
-/** Augments TextureLODSettings with access to TextureLODGroups. */
-struct FExposedTextureLODSettings : public FTextureLODSettings
-{
-public:
-	/** @return		A handle to the indexed LOD group. */
-	FTextureLODGroup& GetTextureLODGroup(int32 GroupIndex)
-	{
-		check( GroupIndex >= 0 && GroupIndex < TEXTUREGROUP_MAX );
-		return TextureLODGroups[GroupIndex];
-	}
-	/** @return		A handle to the indexed LOD group. */
-	const FTextureLODGroup& GetTextureLODGroup(int32 GroupIndex) const
-	{
-		check( GroupIndex >= 0 && GroupIndex < TEXTUREGROUP_MAX );
-		return TextureLODGroups[GroupIndex];
-	}
-};
-
-/** Encapsulates all the system settings that make up Texture Detail friendly settings */
-struct FSystemSettingsDataTextureDetail
-{
-	FSystemSettingsDataTextureDetail();
-
-	friend bool operator==(const FSystemSettingsDataTextureDetail& LHS, const FSystemSettingsDataTextureDetail& RHS)
-	{
-		// WRH - 2007/09/22 - This relies on the fact that the various settings are binary comparable
-		return FMemory::Memcmp(&LHS, &RHS, sizeof(LHS)) == 0;
-	}
-	friend bool operator!=(const FSystemSettingsDataTextureDetail& LHS, const FSystemSettingsDataTextureDetail& RHS)
-	{
-		return !(LHS==RHS);
-	}
-
-	/** Global texture LOD settings.									*/
-	FExposedTextureLODSettings TextureLODSettings;
-};
-
 /** 
  * Struct that holds the actual data for the system settings. 
  *
  * Uses the insane derivation for backwards compatibility purposes only. Would be cleaner to use composition.
  */
 struct FSystemSettingsData 
-	: public FSystemSettingsDataTextureDetail
 {
 	/** ctor */
 	FSystemSettingsData();
 
 	/** loads settings from the given section in the given ini */
 	void LoadFromIni(const TCHAR* IniSection, const FString& IniFilename = GEngineIni, bool bAllowMissingValues = true, bool* FoundValues=NULL);
-
-	/** saves settings to the given section in the engine ini */
-	void SaveToIni(const TCHAR* IniSection);
-	
-private:
-	/**
-	 * Returns a string for the specified texture group LOD settings to the specified ini.
-	 *
-	 * @param	TextureGroupID		Index/enum of the group
-	 * @param	GroupName			String representation of the texture group
-	 */
-	FString GetLODGroupString( TextureGroup TextureGroupID, const TCHAR* GroupName );
-
-	void WriteTextureLODGroupsToIni(const TCHAR* IniSection);
-	void WriteTextureLODGroupToIni(TextureGroup TextureGroupID, const TCHAR* GroupName, const TCHAR* IniSection);
-
-	void DumpTextureLODGroups();
-	void DumpTextureLODGroup(TextureGroup TextureGroupID, const TCHAR* GroupName);
 };
 
 
 /**
  * Class that actually holds the current system settings
- *
- * Derive from FSystemSettingsData instead of holding one purely for backwards
- * compatibility reasons (each data element used to be a member of the class).
  */
 class FSystemSettings : public FSystemSettingsData, public FExec, private FNoncopyable
 {
@@ -103,11 +39,6 @@ public:
 	 */
 	ENGINE_API void Initialize( bool bSetupForEditor );
 
-	/**
-	 * Sets new system settings (optionally writes out to the ini).
-	 */
-	ENGINE_API void ApplyNewSettings(const FSystemSettingsData& NewSettings, bool bWriteToIni);
-
 	/** Applies setting overrides based on command line options. */
 	void ApplyOverrides();
 	
@@ -134,46 +65,6 @@ private:
 	 */
 	void LoadFromIni();
 
-	/**
-	 * Saves current settings to the ini. (purposely override the inherited name so people can't accidentally call it.)
-	 */
-	void SaveToIni();
-
-	/**
-	 * Reads a single entry and parses it into the group array.
-	 *
-	 * @param	TextureGroupID		Index/enum of group to parse
-	 * @param	MinLODSize			Minimum size, in pixels, below which the code won't bias.
-	 * @param	MaxLODSize			Maximum size, in pixels, above which the code won't bias.
-	 * @param	LODBias				Group LOD bias.
-	 * @param	MipGenSettings		Defines how the the mip-map generation works, e.g. sharpening
-	 */
-	void SetTextureLODGroup(TextureGroup TextureGroupID, int MinLODSize, int32 MaxLODSize, int32 LODBias, TextureMipGenSettings MipGenSettings);
-
-	/**
-	 * Writes all texture group LOD settings to the specified ini.
-	 *
-	 * @param	IniFilename			The .ini file to save to.
-	 * @param	IniSection			The .ini section to save to.
-	 */
-	void WriteTextureLODGroupsToIni(const TCHAR* IniSection);
-
-	/**
-	 * Writes the specified texture group LOD settings to the specified ini.
-	 *
-	 * @param	TextureGroupID		Index/enum of the group to parse
-	 * @param	GroupName			String representation of the texture group, to be used as the ini key.
-	 * @param	IniSection			The .ini section to save to.
-	 */
-	void WriteTextureLODGroupToIni(TextureGroup TextureGroupID, const TCHAR* GroupName, const TCHAR* IniSection);
-
-	/**
-	 * Helper for ApplyNewSettings when the engine is running. Applies the changes needed for the runtime system.
-	 *
-	 * We can assume the game is running if this code is called.
-	 */
-	void ApplySettingsAtRuntime(const FSystemSettingsData& NewSettings, bool bWriteToInii);
-
 	/** 
 	 * Makes System Settings take effect on the rendering thread 
 	 */
diff --git a/Engine/Source/Runtime/Engine/Public/TextureLODSettings.h b/Engine/Source/Runtime/Engine/Public/TextureLODSettings.h
deleted file mode 100644
diff --git a/Engine/Source/Runtime/Launch/Private/LaunchEngineLoop.cpp b/Engine/Source/Runtime/Launch/Private/LaunchEngineLoop.cpp
--- a/Engine/Source/Runtime/Launch/Private/LaunchEngineLoop.cpp
+++ b/Engine/Source/Runtime/Launch/Private/LaunchEngineLoop.cpp
@@ -1212,20 +1212,6 @@ int32 FEngineLoop::PreInit( const TCHAR* CmdLine )
 		InitializeStdOutDevice();
 	}
 
-	// Initialize the RHI.
-	RHIInit(bHasEditorToken);
-
-	InitializeSharedSamplerStates();
-
-	if (!FPlatformProperties::RequiresCookedData())
-	{
-		check(!GShaderCompilingManager);
-		GShaderCompilingManager = new FShaderCompilingManager();
-
-		check(!GDistanceFieldAsyncQueue);
-		GDistanceFieldAsyncQueue = new FDistanceFieldAsyncQueue();
-	}
-
 	FIOSystem::Get(); // force it to be created if it isn't already
 
 	// allow the platform to start up any features it may need
@@ -1298,6 +1284,18 @@ int32 FEngineLoop::PreInit( const TCHAR* CmdLine )
 
 	SlowTask.EnterProgressFrame(10);
 
+	// Initialize the RHI.
+	RHIInit(bHasEditorToken);
+
+	if (!FPlatformProperties::RequiresCookedData())
+	{
+		check(!GShaderCompilingManager);
+		GShaderCompilingManager = new FShaderCompilingManager();
+
+		check(!GDistanceFieldAsyncQueue);
+		GDistanceFieldAsyncQueue = new FDistanceFieldAsyncQueue();
+	}
+
 	{
 		DECLARE_SCOPE_CYCLE_COUNTER(TEXT("Initial UObject load"), STAT_InitialUObjectLoad, STATGROUP_LoadTime);
 
diff --git a/Engine/Source/Runtime/MediaAssets/Private/Assets/MediaTextureResource.cpp b/Engine/Source/Runtime/MediaAssets/Private/Assets/MediaTextureResource.cpp
--- a/Engine/Source/Runtime/MediaAssets/Private/Assets/MediaTextureResource.cpp
+++ b/Engine/Source/Runtime/MediaAssets/Private/Assets/MediaTextureResource.cpp
@@ -1,6 +1,7 @@
 // Copyright 1998-2015 Epic Games, Inc. All Rights Reserved.
 
 #include "MediaAssetsPrivatePCH.h"
+#include "DeviceProfiles/DeviceProfile.h"
 
 
 /* FMediaTextureResource structors
@@ -49,7 +50,7 @@ void FMediaTextureResource::InitDynamicRHI()
 
 	// Create the sampler state RHI resource.
 	FSamplerStateInitializerRHI SamplerStateInitializer(
-		GSystemSettings.TextureLODSettings.GetSamplerFilter(Owner),
+		(ESamplerFilter)UDeviceProfileManager::Get().GetActiveProfile()->GetTextureLODSettings()->GetSamplerFilter(Owner),
 		Owner->AddressX == TA_Wrap ? AM_Wrap : (Owner->AddressX == TA_Clamp ? AM_Clamp : AM_Mirror),
 		Owner->AddressY == TA_Wrap ? AM_Wrap : (Owner->AddressY == TA_Clamp ? AM_Clamp : AM_Mirror),
 		AM_Wrap
