Fixed navmesh drawing crashing if bDrawPathCollidingGeometry is true and there's no navoctree (yet) #UE4

diff --git a/Engine/Source/Runtime/Engine/Private/Components/NavMeshRenderingComponent.cpp b/Engine/Source/Runtime/Engine/Private/Components/NavMeshRenderingComponent.cpp
--- a/Engine/Source/Runtime/Engine/Private/Components/NavMeshRenderingComponent.cpp
+++ b/Engine/Source/Runtime/Engine/Private/Components/NavMeshRenderingComponent.cpp
@@ -463,34 +463,36 @@ void UNavMeshRenderingComponent::GatherData(FNavMeshSceneProxyData& CurrentData)
 			// draw all geometry gathered in navoctree
 			const FNavigationOctree* NavOctree = NavMesh->GetWorld()->GetNavigationSystem()->GetNavOctree();
 
-			TArray<FVector> PathCollidingGeomVerts;
-			TArray <int32> PathCollidingGeomIndices;
-			for (FNavigationOctree::TConstIterator<> It(*NavOctree); It.HasPendingNodes(); It.Advance())
+			if (NavOctree)
 			{
-				const FNavigationOctree::FNode& Node = It.GetCurrentNode();
-				for (FNavigationOctree::ElementConstIt ElementIt(Node.GetElementIt()); ElementIt; ElementIt++)
+				TArray<FVector> PathCollidingGeomVerts;
+				TArray <int32> PathCollidingGeomIndices;
+				for (FNavigationOctree::TConstIterator<> It(*NavOctree); It.HasPendingNodes(); It.Advance())
 				{
-					const FNavigationOctreeElement& Element = *ElementIt;
-					if (Element.ShouldUseGeometry(NavMesh->GetConfig()) && Element.Data->CollisionData.Num())
+					const FNavigationOctree::FNode& Node = It.GetCurrentNode();
+					for (FNavigationOctree::ElementConstIt ElementIt(Node.GetElementIt()); ElementIt; ElementIt++)
 					{
-						const FRecastGeometryCache CachedGeometry(Element.Data->CollisionData.GetData());
-						AppendGeometry(PathCollidingGeomVerts, PathCollidingGeomIndices, CachedGeometry.Verts, CachedGeometry.Header.NumVerts, CachedGeometry.Indices, CachedGeometry.Header.NumFaces);
+						const FNavigationOctreeElement& Element = *ElementIt;
+						if (Element.ShouldUseGeometry(NavMesh->GetConfig()) && Element.Data->CollisionData.Num())
+						{
+							const FRecastGeometryCache CachedGeometry(Element.Data->CollisionData.GetData());
+							AppendGeometry(PathCollidingGeomVerts, PathCollidingGeomIndices, CachedGeometry.Verts, CachedGeometry.Header.NumVerts, CachedGeometry.Indices, CachedGeometry.Header.NumFaces);
+						}
 					}
-				}
-				FOREACH_OCTREE_CHILD_NODE(ChildRef)
-				{
-					if (Node.HasChild(ChildRef))
+					FOREACH_OCTREE_CHILD_NODE(ChildRef)
 					{
-						It.PushChild(ChildRef);
+						if (Node.HasChild(ChildRef))
+						{
+							It.PushChild(ChildRef);
+						}
 					}
 				}
+				CurrentData.PathCollidingGeomIndices = PathCollidingGeomIndices;
+				for (const auto& Vertex : PathCollidingGeomVerts)
+				{
+					CurrentData.PathCollidingGeomVerts.Add(FDynamicMeshVertex(Vertex));
+				}
 			}
-			CurrentData.PathCollidingGeomIndices = PathCollidingGeomIndices;
-			for (const auto& Vertex : PathCollidingGeomVerts)
-			{
-				CurrentData.PathCollidingGeomVerts.Add(FDynamicMeshVertex(Vertex));
-			}
-
 		}
 
 		if (CurrentData.NavMeshGeometry.BuiltMeshIndices.Num() > 0)
