Expanded UNavigationSystem::UpdateNavOctree to handle components that are potentially navigation relevant, but are not relevant at the moment of calling #UE4

This is needed to properly handle component's collision disabling.
UE-11725

diff --git a/Engine/Source/Runtime/Engine/Private/AI/Navigation/NavigationSystem.cpp b/Engine/Source/Runtime/Engine/Private/AI/Navigation/NavigationSystem.cpp
--- a/Engine/Source/Runtime/Engine/Private/AI/Navigation/NavigationSystem.cpp
+++ b/Engine/Source/Runtime/Engine/Private/AI/Navigation/NavigationSystem.cpp
@@ -2334,25 +2334,40 @@ void UNavigationSystem::UpdateNavOctree(UActorComponent* Comp)
 
 	// special case for early out: use cached nav relevancy
 	UPrimitiveComponent* PrimComp = Cast<UPrimitiveComponent>(Comp);
-	if (PrimComp && !PrimComp->bNavigationRelevant)
+	if (PrimComp == nullptr || PrimComp->bNavigationRelevant == true)
 	{
-		return;
-	}
-
-	INavRelevantInterface* NavElement = Cast<INavRelevantInterface>(Comp);
-	if (NavElement)
-	{
-		AActor* OwnerActor = Comp ? Comp->GetOwner() : NULL;
-		if (OwnerActor)
+		INavRelevantInterface* NavElement = Cast<INavRelevantInterface>(Comp);
+		if (NavElement)
 		{
-			UNavigationSystem* NavSys = UNavigationSystem::GetCurrent(OwnerActor->GetWorld());
-			if (NavSys)
+			AActor* OwnerActor = Comp ? Comp->GetOwner() : NULL;
+			if (OwnerActor)
 			{
-				if (OwnerActor->IsComponentRelevantForNavigation(Comp))
+				UNavigationSystem* NavSys = UNavigationSystem::GetCurrent(OwnerActor->GetWorld());
+				if (NavSys)
 				{
-					NavSys->UpdateNavOctreeElement(Comp, NavElement, OctreeUpdate_Default);
+					if (OwnerActor->IsComponentRelevantForNavigation(Comp))
+					{
+						NavSys->UpdateNavOctreeElement(Comp, NavElement, OctreeUpdate_Default);
+					}
+					else
+					{
+						NavSys->UnregisterNavOctreeElement(Comp, NavElement, OctreeUpdate_Default);
+					}
 				}
-				else
+			}
+		}
+	}
+	else if (PrimComp != nullptr && PrimComp->CanEverAffectNavigation()) // implies PrimComp->bNavigationRelevant == false
+	{
+		// could have been relevant before and not it isn't. Need to check if there's an octree element ID for it
+		INavRelevantInterface* NavElement = Cast<INavRelevantInterface>(Comp);
+		if (NavElement)
+		{
+			AActor* OwnerActor = Comp ? Comp->GetOwner() : NULL;
+			if (OwnerActor)
+			{
+				UNavigationSystem* NavSys = UNavigationSystem::GetCurrent(OwnerActor->GetWorld());
+				if (NavSys)
 				{
 					NavSys->UnregisterNavOctreeElement(Comp, NavElement, OctreeUpdate_Default);
 				}
