RecalcRequiredBones can clear cached bones array so need it to be called before we setup our pose update data.

diff --git a/Engine/Source/Runtime/Engine/Private/Components/SkeletalMeshComponent.cpp b/Engine/Source/Runtime/Engine/Private/Components/SkeletalMeshComponent.cpp
--- a/Engine/Source/Runtime/Engine/Private/Components/SkeletalMeshComponent.cpp
+++ b/Engine/Source/Runtime/Engine/Private/Components/SkeletalMeshComponent.cpp
@@ -1003,6 +1003,12 @@ void USkeletalMeshComponent::RefreshBoneTransforms(FActorComponentTickFunction*
 		return;
 	}
 
+	// Recalculate the RequiredBones array, if necessary
+	if (!bRequiredBonesUpToDate)
+	{
+		RecalcRequiredBones(PredictedLODLevel);
+	}
+
 	const bool bDoEvaluationRateOptimization = bEnableUpdateRateOptimizations && AnimUpdateRateParams->DoEvaluationRateOptimizations();
 
 	//Handle update rate optimization setup
@@ -1034,18 +1040,13 @@ void USkeletalMeshComponent::RefreshBoneTransforms(FActorComponentTickFunction*
 	AActor* Owner = GetOwner();
 	UE_LOG(LogAnimation, Verbose, TEXT("RefreshBoneTransforms(%s)"), *GetNameSafe(Owner));
 
-	// Recalculate the RequiredBones array, if necessary
-	if (!bRequiredBonesUpToDate)
-	{
-		RecalcRequiredBones(PredictedLODLevel);
-	}
-
 	AnimEvaluationContext.SkeletalMesh = SkeletalMesh;
 	AnimEvaluationContext.AnimInstance = AnimScriptInstance;
 
 	AnimEvaluationContext.bDoEvaluation = bShouldDoEvaluation;
 	
 	AnimEvaluationContext.bDoInterpolation = bDoEvaluationRateOptimization && !bInvalidCachedBones && AnimUpdateRateParams->ShouldInterpolateSkippedFrames();
+	//FPlatformMisc::LowLevelOutputDebugStringf(TEXT("DoInt: %s DoURO: %s InvCacBones: %s ShouldInterp: %s\nCach Num: %i Local Num:%i\n"), B(AnimEvaluationContext.bDoInterpolation), B(bDoEvaluationRateOptimization), B(bInvalidCachedBones), B(AnimUpdateRateParams->ShouldInterpolateSkippedFrames()), CachedLocalAtoms.Num(), LocalAtoms.Num());
 	AnimEvaluationContext.bDuplicateToCacheBones = bInvalidCachedBones || (bDoEvaluationRateOptimization && AnimEvaluationContext.bDoEvaluation && !AnimEvaluationContext.bDoInterpolation);
 
 	if (!bDoEvaluationRateOptimization)
