Fixed typo in subtitle tooltip

diff --git a/Engine/Source/Runtime/Engine/Classes/Engine/EngineTypes.h b/Engine/Source/Runtime/Engine/Classes/Engine/EngineTypes.h
--- a/Engine/Source/Runtime/Engine/Classes/Engine/EngineTypes.h
+++ b/Engine/Source/Runtime/Engine/Classes/Engine/EngineTypes.h
@@ -1057,7 +1057,7 @@ struct FSubtitleCue
 {
 	GENERATED_USTRUCT_BODY()
 
-	/** The text too appear in the subtitle. */
+	/** The text to appear in the subtitle. */
 	UPROPERTY(EditAnywhere, BlueprintReadWrite, Category=SubtitleCue)
 	FText Text;
 
