Fix notifies at the end of animation sequences in montages not triggering because of blend out.

diff --git a/Engine/Source/Runtime/Engine/Classes/Animation/AnimInstance.h b/Engine/Source/Runtime/Engine/Classes/Animation/AnimInstance.h
--- a/Engine/Source/Runtime/Engine/Classes/Animation/AnimInstance.h
+++ b/Engine/Source/Runtime/Engine/Classes/Animation/AnimInstance.h
@@ -337,6 +337,20 @@ struct FNativeStateBinding
 	}
 };
 
+/** Tracks state of active slot nodes in the graph */
+struct FMontageActiveSlotTracker
+{
+	//Weighting of root motion from this montage
+	float RootMotionWeight;
+
+	//Is the montage slot part of the active graph this tick
+	bool  bIsRelevantThisTick;
+
+	//Was the motnage slot part of the active graph last tick
+	bool  bWasRelevantOnPreviousTick;
+
+	FMontageActiveSlotTracker() : RootMotionWeight(0.f), bIsRelevantThisTick(false), bWasRelevantOnPreviousTick(false) {}
+};
 
 UCLASS(transient, Blueprintable, hideCategories=AnimInstance, BlueprintType)
 class ENGINE_API UAnimInstance : public UObject
@@ -404,7 +418,7 @@ public:
 	void UpdateSlotNodeWeight(FName SlotNodeName, float Weight);
 	// if it doesn't tick, it will keep old weight, so we'll have to clear it in the beginning of tick
 	void ClearSlotNodeWeights();
-	bool IsActiveSlotNode(FName SlotNodeName) const;
+	bool IsSlotNodeRelevantForNotifies(FName SlotNodeName) const;
 
 	// Allow slot nodes to store off their root motion weight during ticking
 	void UpdateSlotRootMotionWeight(FName SlotNodeName, float Weight);
@@ -769,10 +783,7 @@ private:
 	UPROPERTY(Transient)
 	int16 GraphTraversalCounter;
 
-	TMap<FName, float> ActiveSlotWeights;
-
-	// Mapping from slot name to weighting for that root motion
-	TMap<FName, float> ActiveSlotRootMotionWeights;
+	TMap<FName, FMontageActiveSlotTracker> SlotWeightTracker;
 
 public:
 	/** 
diff --git a/Engine/Source/Runtime/Engine/Classes/Animation/AnimMontage.h b/Engine/Source/Runtime/Engine/Classes/Animation/AnimMontage.h
--- a/Engine/Source/Runtime/Engine/Classes/Animation/AnimMontage.h
+++ b/Engine/Source/Runtime/Engine/Classes/Animation/AnimMontage.h
@@ -192,6 +192,10 @@ public:
 	// transient PreviousWeight - Weight of previous tick
 	float PreviousWeight;
 
+	// transient NotifyWeight   - Weight for spawned notifies, modified slightly to make sure
+	//                          - we spawn all notifies
+	float NotifyWeight;
+
 	/** Currently Active AnimNotifyState, stored as a copy of the event as we need to
 		call NotifyEnd on the event after a deletion in the editor. After this the event
 		is removed correctly. */
diff --git a/Engine/Source/Runtime/Engine/Private/Animation/AnimInstance.cpp b/Engine/Source/Runtime/Engine/Private/Animation/AnimInstance.cpp
--- a/Engine/Source/Runtime/Engine/Private/Animation/AnimInstance.cpp
+++ b/Engine/Source/Runtime/Engine/Private/Animation/AnimInstance.cpp
@@ -1588,8 +1588,7 @@ void UAnimInstance::SlotEvaluatePose(FName SlotNodeName, const FA2Pose & SourceP
 
 void UAnimInstance::ReinitializeSlotNodes()
 {
-	ActiveSlotWeights.Empty();
-	ActiveSlotRootMotionWeights.Empty();
+	SlotWeightTracker.Empty();
 	
 	// Increment counter
 	SlotNodeInitializationCounter++;
@@ -1604,63 +1603,55 @@ void UAnimInstance::RegisterSlotNodeWithAnimInstance(FName SlotNodeName)
 {
 	// verify if same slot node name exists
 	// then warn users, this is invalid
-	for (auto Iter = ActiveSlotWeights.CreateConstIterator(); Iter; ++Iter)
+	if (SlotWeightTracker.Contains(SlotNodeName))
 	{
-		// if same name found, we should warn user, and make sure they know about it
-		if (SlotNodeName == Iter.Key())
-		{
-			FMessageLog("AnimBlueprint").Warning(FText::Format(LOCTEXT("AnimInstance_SlotNode", "SLOTNODE: '{0}' already exists. Each slot node has to have unique name."), FText::FromString(SlotNodeName.ToString())));
-			return;
-		}
+		FMessageLog("AnimBlueprint").Warning(FText::Format(LOCTEXT("AnimInstance_SlotNode", "SLOTNODE: '{0}' already exists. Each slot node has to have unique name."), FText::FromString(SlotNodeName.ToString())));
+		return;
 	}
 
-	ActiveSlotWeights.Add(SlotNodeName, 0.f);
-	ActiveSlotRootMotionWeights.Add(SlotNodeName, 0.f);
+	SlotWeightTracker.Add(SlotNodeName, FMontageActiveSlotTracker());
 }
 
 void UAnimInstance::UpdateSlotNodeWeight(FName SlotNodeName, float Weight)
 {
-	float* CurrentWeight = ActiveSlotWeights.Find(SlotNodeName);
-	if (CurrentWeight)
+	FMontageActiveSlotTracker* Tracker = SlotWeightTracker.Find(SlotNodeName);
+	if (Tracker)
 	{
-		*CurrentWeight = Weight;
+		// Count as relevant if we are weighted in
+		Tracker->bIsRelevantThisTick = Tracker->bIsRelevantThisTick || Weight > ZERO_ANIMWEIGHT_THRESH;
 	}
 }
 
 void UAnimInstance::ClearSlotNodeWeights()
 {
-	for (auto Iter = ActiveSlotWeights.CreateIterator(); Iter; ++Iter)
-	{
-		float& Weight = Iter.Value();
-		Weight = 0.f;
-	}
-
-	for (auto Iter = ActiveSlotRootMotionWeights.CreateIterator(); Iter; ++Iter)
+	for (auto Iter = SlotWeightTracker.CreateIterator(); Iter; ++Iter)
 	{
-		float& Weight = Iter.Value();
-		Weight = 0.f;
+		FMontageActiveSlotTracker& Tracker = Iter.Value();
+		Tracker.bWasRelevantOnPreviousTick = Tracker.bIsRelevantThisTick;
+		Tracker.bIsRelevantThisTick = false;
+		Tracker.RootMotionWeight = 0.f;
 	}
 }
 
-bool UAnimInstance::IsActiveSlotNode(FName SlotNodeName) const
+bool UAnimInstance::IsSlotNodeRelevantForNotifies(FName SlotNodeName) const
 {
-	const float* SlotNodeWeight = ActiveSlotWeights.Find(SlotNodeName);
-	return ( SlotNodeWeight && (*SlotNodeWeight > ZERO_ANIMWEIGHT_THRESH) );
+	const FMontageActiveSlotTracker* Tracker = SlotWeightTracker.Find(SlotNodeName);
+	return ( Tracker && (Tracker->bIsRelevantThisTick || Tracker->bWasRelevantOnPreviousTick) );
 }
 
 void UAnimInstance::UpdateSlotRootMotionWeight(FName SlotNodeName, float Weight)
 {
-	float* CurrentWeight = ActiveSlotRootMotionWeights.Find(SlotNodeName);
-	if (CurrentWeight)
+	FMontageActiveSlotTracker* Tracker = SlotWeightTracker.Find(SlotNodeName);
+	if (Tracker)
 	{
-		*CurrentWeight += Weight;
+		Tracker->RootMotionWeight += Weight;
 	}
 }
 
 float UAnimInstance::GetSlotRootMotionWeight(FName SlotNodeName) const
 {
-	const float* SlotNodeWeight = ActiveSlotRootMotionWeights.Find(SlotNodeName);
-	return SlotNodeWeight ? *SlotNodeWeight : 0.f;
+	const FMontageActiveSlotTracker* Tracker = SlotWeightTracker.Find(SlotNodeName);
+	return Tracker ? Tracker->RootMotionWeight : 0.f;
 }
 
 float UAnimInstance::GetCurveValue(FName CurveName)
diff --git a/Engine/Source/Runtime/Engine/Private/Animation/AnimMontage.cpp b/Engine/Source/Runtime/Engine/Private/Animation/AnimMontage.cpp
--- a/Engine/Source/Runtime/Engine/Private/Animation/AnimMontage.cpp
+++ b/Engine/Source/Runtime/Engine/Private/Animation/AnimMontage.cpp
@@ -1098,6 +1098,10 @@ void FAnimMontageInstance::UpdateWeight(float DeltaTime)
 
 		// update weight
 		FAnimationRuntime::TickBlendWeight(DeltaTime, DesiredWeight, Weight, BlendTime);
+
+		// Notify weight is max of previous and current as notify could have come
+		// from any point between now and last tick
+		NotifyWeight = FMath::Max(PreviousWeight, Weight);
 	}
 }
 
@@ -1381,7 +1385,7 @@ void FAnimMontageInstance::HandleEvents(float PreviousTrackPos, float CurrentTra
 		// we'll do this for all slots for now
 		for (auto SlotTrack = Montage->SlotAnimTracks.CreateIterator(); SlotTrack; ++SlotTrack)
 		{
-			if (AnimInstance->IsActiveSlotNode(SlotTrack->SlotName))
+			if (AnimInstance->IsSlotNodeRelevantForNotifies(SlotTrack->SlotName))
 			{
 				for (auto AnimSegment = SlotTrack->AnimTrack.AnimSegments.CreateIterator(); AnimSegment; ++AnimSegment)
 				{
@@ -1391,7 +1395,7 @@ void FAnimMontageInstance::HandleEvents(float PreviousTrackPos, float CurrentTra
 		}
 
 		// Queue all these notifies.
-		AnimInstance->AddAnimNotifies(Notifies, Weight);
+		AnimInstance->AddAnimNotifies(Notifies, NotifyWeight);
 	}
 
 	// Update active state branching points, before we handle the immediate tick marker.
