#ue4 - Improve comment on "PhysMat" pin on FHitResult to indicate that "bReturnPhysicalMaterial" must be set.

diff --git a/Engine/Source/Runtime/Engine/Classes/Engine/EngineTypes.h b/Engine/Source/Runtime/Engine/Classes/Engine/EngineTypes.h
--- a/Engine/Source/Runtime/Engine/Classes/Engine/EngineTypes.h
+++ b/Engine/Source/Runtime/Engine/Classes/Engine/EngineTypes.h
@@ -1504,7 +1504,10 @@ struct ENGINE_API FHitResult
 	UPROPERTY()
 	int32 Item;
 
-	/** Physical material that was hit. Must set bReturnPhysicalMaterial to true in the query params for this to be returned. */
+	/**
+	 * Physical material that was hit.
+	 * @note Must set bReturnPhysicalMaterial on the swept PrimitiveComponent or in the query params for this to be returned.
+	 */
 	UPROPERTY()
 	TWeakObjectPtr<class UPhysicalMaterial> PhysMaterial;
 
