Added bRelativeToInitialTransform option to CameraAnims. Gives users the option to author camera animations with an initial offset, if so desired (e.g. for a Begin/Loop/End setup)

diff --git a/Engine/Source/Runtime/Engine/Classes/Camera/CameraAnim.h b/Engine/Source/Runtime/Engine/Classes/Camera/CameraAnim.h
--- a/Engine/Source/Runtime/Engine/Classes/Camera/CameraAnim.h
+++ b/Engine/Source/Runtime/Engine/Classes/Camera/CameraAnim.h
@@ -12,6 +12,7 @@ class UCameraAnim : public UObject
 {
 	GENERATED_UCLASS_BODY()
 
+public:
 	/** The UInterpGroup that holds our actual interpolation data. */
 	UPROPERTY()
 	class UInterpGroup* CameraInterpGroup;
@@ -20,8 +21,8 @@ class UCameraAnim : public UObject
 	/** This is to preview and they only exists in editor */
 	UPROPERTY(transient)
 	class UInterpGroup* PreviewInterpGroup;
-
 #endif // WITH_EDITORONLY_DATA
+
 	/** Length, in seconds. */
 	UPROPERTY()
 	float AnimLength;
@@ -30,7 +31,14 @@ class UCameraAnim : public UObject
 	UPROPERTY()
 	FBox BoundingBox;
 
-	/** The */
+	/** 
+	 * If true, assume all transform keys are intended be offsets from the start of the animation. This allows the animation to be authored at any world location and be applied as a delta to the camera. 
+	 * If false, assume all transform keys are authored relative to the world origin. Positions will be directly applied as deltas to the camera.
+	*/
+	UPROPERTY(EditDefaultsOnly, Category=CameraAnim)
+	uint8 bRelativeToInitialTransform : 1;
+
+	/** The FOV  */
 	UPROPERTY()
 	float BaseFOV;
 
@@ -42,11 +50,6 @@ class UCameraAnim : public UObject
 	UPROPERTY()
 	float BasePostProcessBlendWeight;
 
-protected:
-	/** Internal. Computes and stores the local AABB of the camera's motion. */
-	void CalcLocalAABB();
-
-public:
 	// Begin UObject Interface
 	virtual void PreSave() override;
 	virtual void PostLoad() override;
@@ -65,6 +68,10 @@ public:
 	 * @return Returns the local-space axis-aligned bounding box of the entire motion of this animation. 
 	 */
 	ENGINE_API FBox GetAABB(FVector const& BaseLoc, FRotator const& BaseRot, float Scale) const;
+
+protected:
+	/** Internal. Computes and stores the local AABB of the camera's motion. */
+	void CalcLocalAABB();
 };
 
 
diff --git a/Engine/Source/Runtime/Engine/Private/Camera/CameraAnim.cpp b/Engine/Source/Runtime/Engine/Private/Camera/CameraAnim.cpp
--- a/Engine/Source/Runtime/Engine/Private/Camera/CameraAnim.cpp
+++ b/Engine/Source/Runtime/Engine/Private/Camera/CameraAnim.cpp
@@ -20,6 +20,7 @@ UCameraAnim::UCameraAnim(const FObjectInitializer& ObjectInitializer)
 	: Super(ObjectInitializer)
 {
 	AnimLength = 3.0f;
+	bRelativeToInitialTransform = true;
 	BaseFOV = 90.0f;
 }
 
diff --git a/Engine/Source/Runtime/Engine/Private/PlayerCameraManager.cpp b/Engine/Source/Runtime/Engine/Private/PlayerCameraManager.cpp
--- a/Engine/Source/Runtime/Engine/Private/PlayerCameraManager.cpp
+++ b/Engine/Source/Runtime/Engine/Private/PlayerCameraManager.cpp
@@ -278,16 +278,18 @@ void APlayerCameraManager::GetCachedPostProcessBlends(TArray<FPostProcessSetting
 
 void APlayerCameraManager::ApplyAnimToCamera(ACameraActor const* AnimatedCamActor, UCameraAnimInst const* AnimInst, FMinimalViewInfo& InOutPOV)
 {
-	float const Scale = AnimInst->CurrentBlendWeight;
+	if (AnimInst->CamAnim->bRelativeToInitialTransform)
+	{
+		// move animated cam actor to initial-relative position
+		FTransform const AnimatedCamToWorld = AnimatedCamActor->GetTransform();
+		FTransform const AnimatedCamToInitialCam = AnimatedCamToWorld * AnimInst->InitialCamToWorld.Inverse();
+		ACameraActor* const MutableCamActor = const_cast<ACameraActor*>(AnimatedCamActor);
+		MutableCamActor->SetActorTransform(AnimatedCamToInitialCam);
+	}
 
+	float const Scale = AnimInst->CurrentBlendWeight;
 	FRotationMatrix const CameraToWorld(InOutPOV.Rotation);
 
-	// move animated cam actor to initial-relative position
-	FTransform const AnimatedCamToWorld = AnimatedCamActor->GetTransform();
-	FTransform const AnimatedCamToInitialCam = AnimatedCamToWorld * AnimInst->InitialCamToWorld.Inverse();
-	ACameraActor* const MutableCamActor = const_cast<ACameraActor*>(AnimatedCamActor);
-	MutableCamActor->SetActorTransform(AnimatedCamToInitialCam);		// set it back because that's what the code below expects
-
 	if (AnimInst->PlaySpace == ECameraAnimPlaySpace::CameraLocal)
 	{
 		// the code in the else block will handle this just fine, but this path provides efficiency and simplicity for the most common case
