Fixed DrawDebugBox using DepthPriority incorrectly.

diff --git a/Engine/Source/Runtime/Engine/Private/DrawDebugHelpers.cpp b/Engine/Source/Runtime/Engine/Private/DrawDebugHelpers.cpp
--- a/Engine/Source/Runtime/Engine/Private/DrawDebugHelpers.cpp
+++ b/Engine/Source/Runtime/Engine/Private/DrawDebugHelpers.cpp
@@ -121,57 +121,57 @@ void DrawDebugBox(const UWorld* InWorld, FVector const& Center, FVector const& B
 		ULineBatchComponent* const LineBatcher = GetDebugLineBatcher( InWorld, bPersistentLines, LifeTime, (DepthPriority == SDPG_Foreground) );
 		if(LineBatcher != NULL)
 		{
-			float LineLifeTime = (LifeTime > 0.f)? LifeTime : LineBatcher->DefaultLifeTime;
+			float const LineLifeTime = (LifeTime > 0.f) ? LifeTime : LineBatcher->DefaultLifeTime;
 			TArray<struct FBatchedLine> Lines;
 
-			FTransform Transform(Rotation);
+			FTransform const Transform(Rotation);
 			FVector Start = Transform.TransformPosition(FVector( Box.X,  Box.Y,  Box.Z));
 			FVector End = Transform.TransformPosition(FVector( Box.X, -Box.Y, Box.Z));
-			new(Lines) FBatchedLine(Center + Start, Center + End, Color, LineLifeTime, DepthPriority, 0.f);
+			new(Lines) FBatchedLine(Center + Start, Center + End, Color, LineLifeTime, 0.f, DepthPriority);
 
 			Start = Transform.TransformPosition(FVector( Box.X, -Box.Y,  Box.Z));
 			End = Transform.TransformPosition(FVector(-Box.X, -Box.Y, Box.Z));
-			new(Lines) FBatchedLine(Center + Start, Center + End, Color, LineLifeTime,  DepthPriority, 0.f);
+			new(Lines) FBatchedLine(Center + Start, Center + End, Color, LineLifeTime, 0.f, DepthPriority);
 
 			Start = Transform.TransformPosition(FVector(-Box.X, -Box.Y,  Box.Z));
 			End = Transform.TransformPosition(FVector(-Box.X,  Box.Y, Box.Z));
-			new(Lines) FBatchedLine(Center + Start, Center + End, Color, LineLifeTime,  DepthPriority, 0.f);
+			new(Lines) FBatchedLine(Center + Start, Center + End, Color, LineLifeTime, 0.f, DepthPriority);
 
 			Start = Transform.TransformPosition(FVector(-Box.X,  Box.Y,  Box.Z));
 			End = Transform.TransformPosition(FVector( Box.X,  Box.Y, Box.Z));
-			new(Lines) FBatchedLine(Center + Start, Center + End, Color, LineLifeTime,  DepthPriority, 0.f);
+			new(Lines) FBatchedLine(Center + Start, Center + End, Color, LineLifeTime, 0.f, DepthPriority);
 
 			Start = Transform.TransformPosition(FVector( Box.X,  Box.Y, -Box.Z));
 			End = Transform.TransformPosition(FVector( Box.X, -Box.Y, -Box.Z));
-			new(Lines) FBatchedLine(Center + Start, Center + End, Color, LineLifeTime,  DepthPriority, 0.f);
+			new(Lines) FBatchedLine(Center + Start, Center + End, Color, LineLifeTime, 0.f, DepthPriority);
 
 			Start = Transform.TransformPosition(FVector( Box.X, -Box.Y, -Box.Z));
 			End = Transform.TransformPosition(FVector(-Box.X, -Box.Y, -Box.Z));
-			new(Lines) FBatchedLine(Center + Start, Center + End, Color, LineLifeTime,  DepthPriority, 0.f);
+			new(Lines) FBatchedLine(Center + Start, Center + End, Color, LineLifeTime, 0.f, DepthPriority);
 
 			Start = Transform.TransformPosition(FVector(-Box.X, -Box.Y, -Box.Z));
 			End = Transform.TransformPosition(FVector(-Box.X,  Box.Y, -Box.Z));
-			new(Lines) FBatchedLine(Center + Start, Center + End, Color, LineLifeTime,  DepthPriority, 0.f);
+			new(Lines) FBatchedLine(Center + Start, Center + End, Color, LineLifeTime, 0.f, DepthPriority);
 
 			Start = Transform.TransformPosition(FVector(-Box.X,  Box.Y, -Box.Z));
 			End = Transform.TransformPosition(FVector( Box.X,  Box.Y, -Box.Z));
-			new(Lines) FBatchedLine(Center + Start, Center + End, Color, LineLifeTime,  DepthPriority, 0.f);
+			new(Lines )FBatchedLine(Center + Start, Center + End, Color, LineLifeTime, 0.f, DepthPriority);
 
 			Start = Transform.TransformPosition(FVector( Box.X,  Box.Y,  Box.Z));
 			End = Transform.TransformPosition(FVector( Box.X,  Box.Y, -Box.Z));
-			new(Lines) FBatchedLine(Center + Start, Center + End, Color, LineLifeTime,  DepthPriority, 0.f);
+			new(Lines) FBatchedLine(Center + Start, Center + End, Color, LineLifeTime, 0.f, DepthPriority);
 
 			Start = Transform.TransformPosition(FVector( Box.X, -Box.Y,  Box.Z));
 			End = Transform.TransformPosition(FVector( Box.X, -Box.Y, -Box.Z));
-			new(Lines) FBatchedLine(Center + Start, Center + End, Color, LineLifeTime,  DepthPriority, 0.f);
+			new(Lines) FBatchedLine(Center + Start, Center + End, Color, LineLifeTime, 0.f, DepthPriority);
 
 			Start = Transform.TransformPosition(FVector(-Box.X, -Box.Y,  Box.Z));
 			End = Transform.TransformPosition(FVector(-Box.X, -Box.Y, -Box.Z));
-			new(Lines) FBatchedLine(Center + Start, Center + End, Color, LineLifeTime,  DepthPriority, 0.f);
+			new(Lines) FBatchedLine(Center + Start, Center + End, Color, LineLifeTime, 0.f, DepthPriority);
 
 			Start = Transform.TransformPosition(FVector(-Box.X,  Box.Y,  Box.Z));
 			End = Transform.TransformPosition(FVector(-Box.X,  Box.Y, -Box.Z));
-			new(Lines) FBatchedLine(Center + Start, Center + End, Color, LineLifeTime,  DepthPriority, 0.f);
+			new(Lines) FBatchedLine(Center + Start, Center + End, Color, LineLifeTime, 0.f, DepthPriority);
 
 			LineBatcher->DrawLines(Lines);
 		}
