Fix double scaling for conex element rendering. UE-11761

diff --git a/Engine/Source/Runtime/Engine/Private/PhysicsEngine/PhysDrawing.cpp b/Engine/Source/Runtime/Engine/Private/PhysicsEngine/PhysDrawing.cpp
--- a/Engine/Source/Runtime/Engine/Private/PhysicsEngine/PhysDrawing.cpp
+++ b/Engine/Source/Runtime/Engine/Private/PhysicsEngine/PhysDrawing.cpp
@@ -646,7 +646,7 @@ void FKAggregateGeom::GetAggGeom(const FTransform& Transform, const FColor Color
 				FColor ConvexColor = bPerHullColor ? DebugUtilColor[i%NUM_DEBUG_UTIL_COLORS] : Color;
 				FTransform ElemTM = ConvexElems[i].GetTransform();
 				ElemTM *= Transform;
-				ConvexElems[i].DrawElemWire(Collector.GetPDI(ViewIndex), ElemTM, Scale3D.X, ConvexColor);
+				ConvexElems[i].DrawElemWire(Collector.GetPDI(ViewIndex), ElemTM, 1.f, ConvexColor);	//we pass in 1 for scale because the ElemTM already has the scale baked into it
 			}
 		}
 	}
