#ue4 - Perf: use PhysX single-query functionality (with eANY_HIT flag) where possible.

- 25% speedup for RayCastTest(), GeomSweepTest(), GeomOverlapBlockingTest().
- Renamed FPxQueryFilterCallback::bSingleQuery to bIgnoreTouches for clarity.

UE-11917

diff --git a/Engine/Source/Runtime/Engine/Private/Collision/PhysXCollision.cpp b/Engine/Source/Runtime/Engine/Private/Collision/PhysXCollision.cpp
--- a/Engine/Source/Runtime/Engine/Private/Collision/PhysXCollision.cpp
+++ b/Engine/Source/Runtime/Engine/Private/Collision/PhysXCollision.cpp
@@ -480,8 +480,8 @@ float DebugLineLifetime = 2.f;
 
 /**
  * Helper to lock/unlock multiple scenes that also makes sure to unlock everything when it goes out of scope.
- * Multiple locks are NOT SAFE. You can't call LockRead() if already locked.
- * Multiple unlocks are safe (repeated unlocks do nothing after the first successful unlock.).
+ * Multiple locks on the same scene are NOT SAFE. You can't call LockRead() if already locked.
+ * Multiple unlocks on the same scene are safe (repeated unlocks do nothing after the first successful unlock).
  */
 struct FScopedMultiSceneReadLock
 {
@@ -655,9 +655,9 @@ PxSceneQueryHitType::Enum FPxQueryFilterCallback::preFilter(const PxFilterData&
 
 	PxSceneQueryHitType::Enum Result = FPxQueryFilterCallback::CalcQueryHitType(filterData, ShapeFilter, true);
 
-	if(bSingleQuery && Result == PxSceneQueryHitType::eTOUCH)
+	if (Result == PxSceneQueryHitType::eTOUCH && bIgnoreTouches)
 	{
-		Result = PxSceneQueryHitType::eNONE; // return eNONE for raycastSingle/sweepSingle as it's meaningless to return eTOUCH for those
+		Result = PxSceneQueryHitType::eNONE;
 	}
 
 #if ENABLE_PREFILTER_LOGGING
@@ -836,14 +836,14 @@ bool RaycastTest(const UWorld* World, const FVector Start, const FVector End, EC
 			PxSceneQueryFilterData PQueryFilterData(PFilter, PxSceneQueryFilterFlag::eSTATIC | PxSceneQueryFilterFlag::eDYNAMIC | PxSceneQueryFilterFlag::ePREFILTER);
 			PxSceneQueryFlags POutputFlags = PxHitFlags();
 			FPxQueryFilterCallback PQueryCallback(Params.IgnoreComponents);
-			PQueryCallback.bSingleQuery = true;
+			PQueryCallback.bIgnoreTouches = true; // pre-filter to ignore touches and only get blocking hits.
 
 			FPhysScene* PhysScene = World->GetPhysicsScene();
 			{
 				// Enable scene locks, in case they are required
 				PxScene* SyncScene = PhysScene->GetPhysXScene(PST_Sync);
 				SCOPED_SCENE_READ_LOCK(SyncScene);
-				bHaveBlockingHit = SyncScene->raycastSingle(U2PVector(Start), PDir, DeltaMag, POutputFlags, PHit, PQueryFilterData, &PQueryCallback);
+				bHaveBlockingHit = SyncScene->raycastAny(U2PVector(Start), PDir, DeltaMag, PHit, PQueryFilterData, &PQueryCallback);
 			}
 
 			// Test async scene if we have no blocking hit, and async tests are requested
@@ -851,7 +851,7 @@ bool RaycastTest(const UWorld* World, const FVector Start, const FVector End, EC
 			{
 				PxScene* AsyncScene = PhysScene->GetPhysXScene(PST_Async);
 				SCOPED_SCENE_READ_LOCK(AsyncScene);
-				bHaveBlockingHit = AsyncScene->raycastSingle(U2PVector(Start), PDir, DeltaMag, POutputFlags, PHit, PQueryFilterData, &PQueryCallback);
+				bHaveBlockingHit = AsyncScene->raycastAny(U2PVector(Start), PDir, DeltaMag, PHit, PQueryFilterData, &PQueryCallback);
 			}
 		}
 #endif // WITH_PHYSX
@@ -908,7 +908,7 @@ bool RaycastSingle(const UWorld* World, struct FHitResult& OutHit, const FVector
 			PxSceneQueryFilterData PQueryFilterData(PFilter, PxSceneQueryFilterFlag::eSTATIC | PxSceneQueryFilterFlag::eDYNAMIC | PxSceneQueryFilterFlag::ePREFILTER);
 			PxSceneQueryFlags POutputFlags = PxSceneQueryFlag::ePOSITION | PxSceneQueryFlag::eNORMAL | PxSceneQueryFlag::eDISTANCE | PxSceneQueryFlag::eMTD;
 			FPxQueryFilterCallback PQueryCallback(Params.IgnoreComponents);
-			PQueryCallback.bSingleQuery = true; // this needs to be set for raycastSingle so we return eNONE instead of eTOUCH
+			PQueryCallback.bIgnoreTouches = true; // pre-filter to ignore touches and only get blocking hits.
 
 			PxVec3 PStart = U2PVector(Start);
 			PxVec3 PDir = U2PVector(Delta / DeltaMag);
@@ -1194,7 +1194,7 @@ bool GeomSweepTest(const UWorld* World, const struct FCollisionShape& CollisionS
 
 		FPxQueryFilterCallbackSweep PQueryCallbackSweep(Params.IgnoreComponents);
 		PQueryCallbackSweep.DiscardInitialOverlaps = !Params.bFindInitialOverlaps;
-		PQueryCallbackSweep.bSingleQuery = true;
+		PQueryCallbackSweep.bIgnoreTouches = true; // pre-filter to ignore touches and only get blocking hits.
 
 		PxTransform PStartTM(U2PVector(Start), PGeomRot);
 		PxVec3 PDir = U2PVector(Delta/DeltaMag);
@@ -1205,7 +1205,7 @@ bool GeomSweepTest(const UWorld* World, const struct FCollisionShape& CollisionS
 			PxScene* SyncScene = PhysScene->GetPhysXScene(PST_Sync);
 			SCOPED_SCENE_READ_LOCK(SyncScene);
 			PxSweepHit PQueryHit;
-			bHaveBlockingHit = SyncScene->sweepSingle(PGeom, PStartTM, PDir, DeltaMag, POutputFlags, PQueryHit, PQueryFilterData, &PQueryCallbackSweep);
+			bHaveBlockingHit = SyncScene->sweepAny(PGeom, PStartTM, PDir, DeltaMag, POutputFlags, PQueryHit, PQueryFilterData, &PQueryCallbackSweep);
 		}
 
 		// Test async scene if async tests are requested and there was no blocking hit was found in the sync scene (since no hit info other than a boolean yes/no is recorded)
@@ -1214,7 +1214,7 @@ bool GeomSweepTest(const UWorld* World, const struct FCollisionShape& CollisionS
 			PxScene* AsyncScene = PhysScene->GetPhysXScene(PST_Async);
 			SCOPED_SCENE_READ_LOCK(AsyncScene);
 			PxSweepHit PTestHit;
-			bHaveBlockingHit = AsyncScene->sweepSingle(PGeom, PStartTM, PDir, DeltaMag, POutputFlags, PTestHit, PQueryFilterData, &PQueryCallbackSweep);
+			bHaveBlockingHit = AsyncScene->sweepAny(PGeom, PStartTM, PDir, DeltaMag, POutputFlags, PTestHit, PQueryFilterData, &PQueryCallbackSweep);
 		}
 	}
 
@@ -1265,7 +1265,7 @@ bool GeomSweepSingle(const UWorld* World, const struct FCollisionShape& Collisio
 		PxSceneQueryFilterData PQueryFilterData(PFilter, PxSceneQueryFilterFlag::eSTATIC | PxSceneQueryFilterFlag::eDYNAMIC | PxSceneQueryFilterFlag::ePREFILTER);
 		PxSceneQueryFlags POutputFlags = PxSceneQueryFlag::ePOSITION | PxSceneQueryFlag::eNORMAL | PxSceneQueryFlag::eDISTANCE | PxSceneQueryFlag::eMTD;
 		FPxQueryFilterCallbackSweep PQueryCallbackSweep(Params.IgnoreComponents);
-		PQueryCallbackSweep.bSingleQuery = true; //LOC_MOD need to do this to return eNONE instead of eTOUCH for r
+		PQueryCallbackSweep.bIgnoreTouches = true; // pre-filter to ignore touches and only get blocking hits.
 		PQueryCallbackSweep.DiscardInitialOverlaps = !Params.bFindInitialOverlaps;	
 
 		PxTransform PStartTM(U2PVector(Start), PGeomRot);
@@ -1575,6 +1575,7 @@ bool GeomOverlapMultiImp_PhysX(const UWorld* World, const PxGeometry& PGeom, con
 		PxFilterData PFilter = CreateQueryFilterData(TraceChannel, Params.bTraceComplex, ResponseParams.CollisionResponse, ObjectParams, InfoType != EQueryInfo::IsAnything);
 		PxSceneQueryFilterData PQueryFilterData(PFilter, PxSceneQueryFilterFlag::eSTATIC | PxSceneQueryFilterFlag::eDYNAMIC | PxSceneQueryFilterFlag::ePREFILTER);
 		FPxQueryFilterCallback PQueryCallback(Params.IgnoreComponents);
+		PQueryCallback.bIgnoreTouches = (InfoType == EQueryInfo::IsBlocking); // pre-filter to ignore touches and only get blocking hits, if that's what we're after.
 
 		// Enable scene locks, in case they are required
 		FScopedMultiSceneReadLock SceneLocks;
@@ -1592,7 +1593,7 @@ bool GeomOverlapMultiImp_PhysX(const UWorld* World, const PxGeometry& PGeom, con
 
 		PxI32 NumHits = 0;
 		
-		if (InfoType == EQueryInfo::IsAnything)
+		if ((InfoType == EQueryInfo::IsAnything) || (InfoType == EQueryInfo::IsBlocking))
 		{
 			if (SyncScene->overlapAny(PGeom, PGeomPose, POverlapResult, PQueryFilterData, &PQueryCallback))
 			{
@@ -1601,6 +1602,7 @@ bool GeomOverlapMultiImp_PhysX(const UWorld* World, const PxGeometry& PGeom, con
 		}
 		else
 		{
+			checkSlow(InfoType == EQueryInfo::GatherAll);
 			NumHits = SyncScene->overlapMultiple(PGeom, PGeomPose, POverlapArray, OVERLAP_BUFFER_SIZE_MAX_SYNC_QUERIES, PQueryFilterData, &PQueryCallback);
 			if (NumHits == -1)
 			{
@@ -1616,17 +1618,6 @@ bool GeomOverlapMultiImp_PhysX(const UWorld* World, const PxGeometry& PGeom, con
 				// Not using anything from this scene, so unlock it.
 				SceneLocks.UnlockRead(SyncScene, PST_Sync);
 			}
-
-			if (InfoType == EQueryInfo::IsBlocking)	//in the case where we just want to know if we're blocking, quickly check if anything blocks and return
-			{
-				for (PxI32 HitIdx = 0; HitIdx < NumHits; ++HitIdx)
-				{
-					if (IsBlocking(POverlapArray[HitIdx].shape, PFilter))
-					{
-						return true;
-					}
-				}
-			}
 		}
 
 		// Test async scene if async tests are requested and there was no overflow
@@ -1636,14 +1627,8 @@ bool GeomOverlapMultiImp_PhysX(const UWorld* World, const PxGeometry& PGeom, con
 			
 			// we can't use scoped because we later do a conversion which depends on these results and it should all be atomic
 			SceneLocks.LockRead(AsyncScene, PST_Async);
-
-			// Write into the same PHits buffer
-			PxOverlapHit* PAsyncOverlapArray = POverlapArray + NumHits;
-			const PxI32 AsyncBufferSize = (ARRAY_COUNT(RawPOverlapArray) - NumHits);
-			check(AsyncBufferSize > 0);
-			PxI32 NumAsyncHits = 0;
-			
-			if (InfoType == EQueryInfo::IsAnything)
+		
+			if ((InfoType == EQueryInfo::IsAnything) || (InfoType == EQueryInfo::IsBlocking))
 			{
 				if (AsyncScene->overlapAny(PGeom, PGeomPose, POverlapResult, PQueryFilterData, &PQueryCallback))
 				{
@@ -1652,7 +1637,13 @@ bool GeomOverlapMultiImp_PhysX(const UWorld* World, const PxGeometry& PGeom, con
 			}
 			else
 			{
-				NumAsyncHits = AsyncScene->overlapMultiple(PGeom, PGeomPose, PAsyncOverlapArray, AsyncBufferSize, PQueryFilterData, &PQueryCallback);
+				checkSlow(InfoType == EQueryInfo::GatherAll);
+				// Write into the same PHits buffer
+				PxOverlapHit* PAsyncOverlapArray = POverlapArray + NumHits;
+				const PxI32 AsyncBufferSize = (ARRAY_COUNT(RawPOverlapArray) - NumHits);
+				check(AsyncBufferSize > 0);
+
+				PxI32 NumAsyncHits = AsyncScene->overlapMultiple(PGeom, PGeomPose, PAsyncOverlapArray, AsyncBufferSize, PQueryFilterData, &PQueryCallback);
 
 				if (NumAsyncHits == -1)
 				{
@@ -1669,17 +1660,6 @@ bool GeomOverlapMultiImp_PhysX(const UWorld* World, const PxGeometry& PGeom, con
 					SceneLocks.UnlockRead(AsyncScene, PST_Async);
 				}
 
-				if (InfoType == EQueryInfo::IsBlocking)	//in the case where we just want to know if we're blocking, quickly check if anything blocks and return
-				{
-					for (PxI32 HitIdx = 0; HitIdx < NumAsyncHits; ++HitIdx)
-					{
-						if (IsBlocking(PAsyncOverlapArray[HitIdx].shape, PFilter))
-						{
-							return true;
-						}
-					}
-				}
-
 				NumHits += NumAsyncHits;
 			}
 		}
diff --git a/Engine/Source/Runtime/Engine/Private/Collision/PhysXCollision.h b/Engine/Source/Runtime/Engine/Private/Collision/PhysXCollision.h
--- a/Engine/Source/Runtime/Engine/Private/Collision/PhysXCollision.h
+++ b/Engine/Source/Runtime/Engine/Private/Collision/PhysXCollision.h
@@ -89,14 +89,14 @@ public:
 	/** Result of PreFilter callback. */
 	PxSceneQueryHitType::Enum PrefilterReturnValue;
 
-	/** Whether this is a raycastSingle or sweepSingle, because in 3.3 you can't return eTOUCH for single queries */
-	bool bSingleQuery;
+	/** Whether to ignore touches (convert an eTOUCH result to eNONE). */
+	bool bIgnoreTouches;
 
 
 	FPxQueryFilterCallback()
 	{
 		PrefilterReturnValue = PxSceneQueryHitType::eNONE;
-		bSingleQuery = false;
+		bIgnoreTouches = false;
 	}
 
 	FPxQueryFilterCallback(const FilterIgnoreComponentsArrayType& InIgnoreComponents)
@@ -104,7 +104,7 @@ public:
 		PrefilterReturnValue = PxSceneQueryHitType::eNONE;
 		
 		IgnoreComponents = InIgnoreComponents;
-		bSingleQuery = false;
+		bIgnoreTouches = false;
 	}
 
 	/** 
