Add extension points suggested by Scaleform.

diff --git a/Engine/Source/Editor/UnrealEd/Private/PlayLevel.cpp b/Engine/Source/Editor/UnrealEd/Private/PlayLevel.cpp
--- a/Engine/Source/Editor/UnrealEd/Private/PlayLevel.cpp
+++ b/Engine/Source/Editor/UnrealEd/Private/PlayLevel.cpp
@@ -58,6 +58,7 @@
 #include "EngineUtils.h"
 #include "GameMapsSettings.h"
 #include "GameFramework/Pawn.h"
+#include "GameDelegates.h"
 
 DEFINE_LOG_CATEGORY_STATIC(LogPlayLevel, Log, All);
 
@@ -231,6 +232,8 @@ void UEditorEngine::EndPlayMap()
 		EditorWorld->GetNavigationSystem()->OnPIEEnd();
 	}
 
+	FGameDelegates::Get().GetEndPlayMapDelegate().Broadcast();
+
 	EditorWorld->bAllowAudioPlayback = true;
 	EditorWorld = NULL;
 
diff --git a/Engine/Source/Runtime/Engine/Classes/Engine/GameViewportClient.h b/Engine/Source/Runtime/Engine/Classes/Engine/GameViewportClient.h
--- a/Engine/Source/Runtime/Engine/Classes/Engine/GameViewportClient.h
+++ b/Engine/Source/Runtime/Engine/Classes/Engine/GameViewportClient.h
@@ -477,6 +477,48 @@ public:
 		return CloseRequestedDelegate;
 	}
 
+	/** Accessor for the delegate called when the game viewport is created. */
+	static FSimpleMulticastDelegate& OnViewportCreated()
+	{
+		return CreatedDelegate;
+	}
+
+	// Accessor for the delegate called when a player is added to the game viewport
+	FOnGameViewportClientPlayerAction& OnPlayerAdded()
+	{
+		return PlayerAddedDelegate;
+	}
+
+	// Accessor for the delegate called when a player is removed from the game viewport
+	FOnGameViewportClientPlayerAction& OnPlayerRemoved()
+	{
+		return PlayerRemovedDelegate;
+	}
+
+	// Accessor for the delegate called when the engine starts drawing a game viewport
+	FSimpleMulticastDelegate& OnBeginDraw()
+	{
+		return BeginDrawDelegate;
+	}
+
+	// Accessor for the delegate called when the game viewport is drawn, before drawing the console
+	FSimpleMulticastDelegate& OnDrawn()
+	{
+		return DrawnDelegate;
+	}
+
+	// Accessor for the delegate called when the engine finishes drawing a game viewport
+	FSimpleMulticastDelegate& OnEndDraw()
+	{
+		return EndDrawDelegate;
+	}
+
+	// Accessor for the delegate called when ticking the game viewport
+	FOnGameViewportTick& OnTick()
+	{
+		return TickDelegate;
+	}
+
 	/** Return the engine show flags for this viewport */
 	virtual FEngineShowFlags* GetEngineShowFlags() 
 	{ 
@@ -700,6 +742,27 @@ private:
 	/** Delegate called when a request to close the viewport is received */
 	FOnCloseRequested CloseRequestedDelegate;
 
+	/** Delegate called when the game viewport is created. */
+	static FSimpleMulticastDelegate CreatedDelegate;
+
+	/** Delegate called when a player is added to the game viewport */
+	FOnGameViewportClientPlayerAction PlayerAddedDelegate;
+
+	/** Delegate called when a player is removed from the game viewport */
+	FOnGameViewportClientPlayerAction PlayerRemovedDelegate;
+
+	/** Delegate called when the engine starts drawing a game viewport */
+	FSimpleMulticastDelegate BeginDrawDelegate;
+
+	/** Delegate called when the game viewport is drawn, before drawing the console */
+	FSimpleMulticastDelegate DrawnDelegate;
+
+	/** Delegate called when the engine finishes drawing a game viewport */
+	FSimpleMulticastDelegate EndDrawDelegate;
+
+	/** Delegate called when ticking the game viewport */
+	FOnGameViewportTick TickDelegate;
+
 	/** Data needed to display perframe stat tracking when STAT UNIT is enabled */
 	FStatUnitData* StatUnitData;
 
diff --git a/Engine/Source/Runtime/Engine/Classes/Engine/GameViewportDelegates.h b/Engine/Source/Runtime/Engine/Classes/Engine/GameViewportDelegates.h
--- a/Engine/Source/Runtime/Engine/Classes/Engine/GameViewportDelegates.h
+++ b/Engine/Source/Runtime/Engine/Classes/Engine/GameViewportDelegates.h
@@ -31,3 +31,19 @@ DECLARE_DELEGATE_FourParams(FOnPNGScreenshotCaptured, int32, int32, const TArray
  * @see UGameViewportClient
  */
 DECLARE_DELEGATE_OneParam(FOnCloseRequested, FViewport*);
+
+/**
+ * Delegate type used by UGameViewportClient for when a player is added or removed
+ *
+ * The first parameter is the player index
+ * @see UGameViewportClient
+ */
+DECLARE_MULTICAST_DELEGATE_OneParam(FOnGameViewportClientPlayerAction, int32);
+
+/**
+ * Delegate type used by UGameViewportClient for tick callbacks
+ *
+ * The first parameter is the time delta
+ * @see UGameViewportClient
+ */
+DECLARE_MULTICAST_DELEGATE_OneParam(FOnGameViewportTick, float);
diff --git a/Engine/Source/Runtime/Engine/Private/GameEngine.cpp b/Engine/Source/Runtime/Engine/Private/GameEngine.cpp
--- a/Engine/Source/Runtime/Engine/Private/GameEngine.cpp
+++ b/Engine/Source/Runtime/Engine/Private/GameEngine.cpp
@@ -37,6 +37,7 @@
 #include "Engine/GameEngine.h"
 #include "GameFramework/GameUserSettings.h"
 #include "GameFramework/GameMode.h"
+#include "GameDelegates.h"
 
 ENGINE_API bool GDisallowNetworkTravel = false;
 
@@ -458,6 +459,9 @@ void UGameEngine::Init(IEngineLoop* InEngineLoop)
 		{
 			SwitchGameWindowToUseGameViewport();
 		}
+
+		UGameViewportClient::OnViewportCreated().Broadcast();
+
 		FString Error;
 		if(ViewportClient->SetupInitialLocalPlayer(Error) == NULL)
 		{
@@ -674,6 +678,7 @@ bool UGameEngine::Exec( UWorld* InWorld, const TCHAR* Cmd, FOutputDevice& Ar )
 bool UGameEngine::HandleExitCommand( const TCHAR* Cmd, FOutputDevice& Ar )
 {
 	Ar.Log( TEXT("Closing by request") );
+	FGameDelegates::Get().GetExitCommandDelegate().Broadcast();
 	FPlatformMisc::RequestExit( 0 );
 	return true;
 }
diff --git a/Engine/Source/Runtime/Engine/Private/GameViewportClient.cpp b/Engine/Source/Runtime/Engine/Private/GameViewportClient.cpp
--- a/Engine/Source/Runtime/Engine/Private/GameViewportClient.cpp
+++ b/Engine/Source/Runtime/Engine/Private/GameViewportClient.cpp
@@ -48,6 +48,9 @@ extern ENGINE_API class FLightMap2D* GDebugSelectedLightmap;
 /** Delegate called at the end of the frame when a screenshot is captured */
 FOnScreenshotCaptured UGameViewportClient::ScreenshotCapturedDelegate;
 
+/** Delegate called when the game viewport is created. */
+FSimpleMulticastDelegate UGameViewportClient::CreatedDelegate;
+
 /** A list of all the stat names which are enabled for this viewport (static so they persist between runs) */
 TArray<FString> UGameViewportClient::EnabledStats;
 
@@ -221,6 +224,7 @@ TSharedPtr<class SViewport> UGameViewportClient::GetGameViewportWidget()
 
 void UGameViewportClient::Tick( float DeltaTime )
 {
+	TickDelegate.Broadcast(DeltaTime);
 }
 
 FString UGameViewportClient::ConsoleCommand( const FString& Command)
@@ -768,6 +772,8 @@ void UGameViewportClient::Draw(FViewport* InViewport, FCanvas* SceneCanvas)
 	//Valid SceneCanvas is required.  Make this explicit.
 	check(SceneCanvas);
 
+	BeginDrawDelegate.Broadcast();
+
 	FCanvas* DebugCanvas = InViewport->GetDebugCanvas();
 
 	// Create a temporary canvas if there isn't already one.
@@ -1160,6 +1166,9 @@ void UGameViewportClient::Draw(FViewport* InViewport, FCanvas* SceneCanvas)
 		{
 			DebugCanvas->Flush_GameThread();
 		}
+
+		DrawnDelegate.Broadcast();
+
 		// Allow the viewport to render additional stuff
 		PostRender(DebugCanvasObject);
 		
@@ -1241,6 +1250,8 @@ void UGameViewportClient::Draw(FViewport* InViewport, FCanvas* SceneCanvas)
 	{
 		DrawStatsHUD( GetWorld(), InViewport, DebugCanvas, DebugCanvasObject, DebugProperties, PlayerCameraLocation, PlayerCameraRotation );
 	}
+
+	EndDrawDelegate.Broadcast();
 }
 
 void UGameViewportClient::ProcessScreenShots(FViewport* InViewport)
@@ -1330,7 +1341,7 @@ void UGameViewportClient::Precache()
 			for(TObjectIterator<USoundWave> It;It;++It)
 			{
 				USoundWave* SoundWave = *It;
-				AudioDevice->Precache(SoundWave);
+				AudioDevice->Precache( SoundWave );
 			}
 			UE_LOG(LogPlayerManagement, Log, TEXT("Precaching sounds completed..."));
 		}
@@ -1903,6 +1914,8 @@ void UGameViewportClient::NotifyPlayerAdded( int32 PlayerIndex, ULocalPlayer* Ad
 	{
 		GameLayerManager->NotifyPlayerAdded(PlayerIndex, AddedPlayer);
 	}
+
+	PlayerAddedDelegate.Broadcast( PlayerIndex );
 }
 
 void UGameViewportClient::NotifyPlayerRemoved( int32 PlayerIndex, ULocalPlayer* RemovedPlayer )
@@ -1914,6 +1927,8 @@ void UGameViewportClient::NotifyPlayerRemoved( int32 PlayerIndex, ULocalPlayer*
 	{
 		GameLayerManager->NotifyPlayerRemoved(PlayerIndex, RemovedPlayer);
 	}
+
+	PlayerRemovedDelegate.Broadcast( PlayerIndex );
 }
 
 void UGameViewportClient::AddViewportWidgetContent( TSharedRef<SWidget> ViewportContent, const int32 ZOrder )
@@ -1968,7 +1983,7 @@ void UGameViewportClient::RemoveAllViewportWidgets()
 
 	TSharedPtr< IGameLayerManager > GameLayerManager(GameLayerManagerPtr.Pin());
 	if ( GameLayerManager.IsValid() )
-	{
+		{
 		GameLayerManager->ClearWidgets();
 	}
 }
diff --git a/Engine/Source/Runtime/Engine/Private/LocalPlayer.cpp b/Engine/Source/Runtime/Engine/Private/LocalPlayer.cpp
--- a/Engine/Source/Runtime/Engine/Private/LocalPlayer.cpp
+++ b/Engine/Source/Runtime/Engine/Private/LocalPlayer.cpp
@@ -23,6 +23,8 @@
 #include "GameFramework/GameMode.h"
 #include "GameFramework/PlayerState.h"
 
+#include "GameDelegates.h"
+
 DEFINE_LOG_CATEGORY(LogPlayerManagement);
 
 #if !(UE_BUILD_SHIPPING || UE_BUILD_TEST)
@@ -1027,6 +1029,9 @@ bool ULocalPlayer::HandleExitCommand( const TCHAR* Cmd, FOutputDevice& Ar )
 	{
 		ViewportClient->CloseRequested(ViewportClient->Viewport);
 	}
+
+	FGameDelegates::Get().GetExitCommandDelegate().Broadcast();
+
 	return true;
 }
 
diff --git a/Engine/Source/Runtime/Engine/Public/GameDelegates.h b/Engine/Source/Runtime/Engine/Public/GameDelegates.h
--- a/Engine/Source/Runtime/Engine/Public/GameDelegates.h
+++ b/Engine/Source/Runtime/Engine/Public/GameDelegates.h
@@ -55,10 +55,27 @@ public:
 	/** Return a single FGameDelegates object */
 	static FGameDelegates& Get();
 
+	// Called when an exit command is received
+	FSimpleMulticastDelegate& GetExitCommandDelegate()
+	{
+		return ExitCommandDelegate;
+	}
+
+	// Called when ending playing a map
+	FSimpleMulticastDelegate& GetEndPlayMapDelegate()
+	{
+		return EndPlayMapDelegate;
+	}
+
 	// Implement all delegates declared above
 	DEFINE_GAME_DELEGATE(CookModificationDelegate);
 	DEFINE_GAME_DELEGATE(AssignStreamingChunkDelegate);
 	DEFINE_GAME_DELEGATE(AssignLayerChunkDelegate);
 	DEFINE_GAME_DELEGATE(ExtendedSaveGameInfoDelegate);
 	DEFINE_GAME_DELEGATE(WebServerActionDelegate);	
+
+private:
+
+	FSimpleMulticastDelegate ExitCommandDelegate;
+	FSimpleMulticastDelegate EndPlayMapDelegate;
 };
