Fix deferred actor addition and removal not iterating over the physx scenes correctly.

diff --git a/Engine/Source/Runtime/Engine/Private/PhysicsEngine/PhysScene.cpp b/Engine/Source/Runtime/Engine/Private/PhysicsEngine/PhysScene.cpp
--- a/Engine/Source/Runtime/Engine/Private/PhysicsEngine/PhysScene.cpp
+++ b/Engine/Source/Runtime/Engine/Private/PhysicsEngine/PhysScene.cpp
@@ -353,16 +353,16 @@ void FPhysScene::TermBody(FBodyInstance* BodyInstance)
 #endif
 	}
 
-		// Remove body from any pending deferred addition / removal
-		for(int32 SceneIdx = 0 ; SceneIdx < PST_MAX ; ++SceneIdx)
+	// Remove body from any pending deferred addition / removal
+	for(FDeferredSceneData& Deferred : DeferredSceneData)
+	{
+		int32 FoundIdx = INDEX_NONE;
+		if(Deferred.AddInstances.Find(BodyInstance, FoundIdx))
 		{
-			int32 FoundIdx = INDEX_NONE;
-		if(DeferredAddInstances[SceneIdx].Find(BodyInstance, FoundIdx))
-			{
-				DeferredAddActors->RemoveAt(FoundIdx);
-				DeferredAddInstances->RemoveAt(FoundIdx);
-			}
+			Deferred.AddActors.RemoveAtSwap(FoundIdx);
+			Deferred.AddInstances.RemoveAtSwap(FoundIdx);
 		}
+	}
 
 #if WITH_PHYSX
 	RemoveActiveBody(BodyInstance, PST_Sync);
@@ -1313,6 +1313,7 @@ void FPhysScene::InitPhysScene(uint32 SceneType)
 
 	// Store index of PhysX Scene in this FPhysScene
 	this->PhysXSceneIndex[SceneType] = PhysXSceneCount;
+	DeferredSceneData[SceneType].SceneIndex = PhysXSceneCount;
 
 	// Increment scene count
 	PhysXSceneCount++;
@@ -1396,108 +1397,135 @@ void FPhysScene::TermPhysScene(uint32 SceneType)
 }
 
 #if WITH_PHYSX
-void FPhysScene::FlushDeferredActors()
+
+void FPhysScene::FDeferredSceneData::FlushDeferredActors()
 {
-	check(IsInGameThread());
+	check(AddInstances.Num() == AddActors.Num());
 
-	for(int32 SceneIdx = 0 ; SceneIdx < PST_MAX ; ++SceneIdx)
+	if (AddInstances.Num() > 0)
 	{
-		check(DeferredAddInstances[SceneIdx].Num() == DeferredAddActors[SceneIdx].Num());
+		PxScene* Scene = GetPhysXSceneFromIndex(SceneIndex);
+		SCOPED_SCENE_WRITE_LOCK(Scene);
 
-		if(DeferredAddInstances[SceneIdx].Num() > 0)
+		Scene->addActors(AddActors.GetData(), AddActors.Num());
+		int32 Idx = -1;
+		for (FBodyInstance* Instance : AddInstances)
 		{
-			PxScene* Scene = GetPhysXScene(SceneIdx);
-			SCOPED_SCENE_WRITE_LOCK(Scene);
+			++Idx;
+			Instance->CurrentSceneState = BodyInstanceSceneState::Added;
 
-			Scene->addActors(DeferredAddActors[SceneIdx].GetData(), DeferredAddActors[SceneIdx].Num());
-			int32 Idx = -1;
-			for(FBodyInstance* Instance : DeferredAddInstances[SceneIdx])
+			if (Instance->IsDynamic())
 			{
-				++Idx;
-				Instance->CurrentSceneState = BodyInstanceSceneState::Added;
-
-				if(Instance->IsDynamic())
-				{
-					// Extra setup necessary for dynamic objects.
-					Instance->InitDynamicProperties();
-				}
+				// Extra setup necessary for dynamic objects.
+				Instance->InitDynamicProperties();
 			}
-
-			DeferredAddInstances->Empty();
-			DeferredAddActors->Empty();
 		}
 
-		check(DeferredRemoveInstances[SceneIdx].Num() == DeferredRemoveActors[SceneIdx].Num());
+		AddInstances.Empty();
+		AddActors.Empty();
+	}
 
-		if(DeferredRemoveInstances[SceneIdx].Num() > 0)
-		{
-			PxScene* Scene = GetPhysXScene(SceneIdx);
-			SCOPED_SCENE_WRITE_LOCK(Scene);
+	check(RemoveInstances.Num() == RemoveActors.Num());
 
-			Scene->removeActors(DeferredRemoveActors[SceneIdx].GetData(), DeferredRemoveActors[SceneIdx].Num());
+	if (RemoveInstances.Num() > 0)
+	{
+		PxScene* Scene = GetPhysXSceneFromIndex(SceneIndex);
+		SCOPED_SCENE_WRITE_LOCK(Scene);
 
-			for(FBodyInstance* Instance : DeferredAddInstances[SceneIdx])
-			{
-				Instance->CurrentSceneState = BodyInstanceSceneState::Removed;
-			}
+		Scene->removeActors(RemoveActors.GetData(), RemoveActors.Num());
 
-			DeferredRemoveInstances->Empty();
-			DeferredRemoveActors->Empty();
+		for (FBodyInstance* Instance : AddInstances)
+		{
+			Instance->CurrentSceneState = BodyInstanceSceneState::Removed;
 		}
+
+		RemoveInstances.Empty();
+		RemoveActors.Empty();
 	}
 }
 
-void FPhysScene::DeferAddActor(FBodyInstance* OwningInstance, PxActor* Actor, EPhysicsSceneType SceneType)
+void FPhysScene::FlushDeferredActors()
+{
+	check(IsInGameThread());
+
+	for(FDeferredSceneData& Deferred : DeferredSceneData)
+	{
+		Deferred.FlushDeferredActors();
+	}
+}
+
+
+
+void FPhysScene::FDeferredSceneData::DeferAddActor(FBodyInstance* OwningInstance, PxActor* Actor)
 {
-	check(OwningInstance && Actor && SceneType < PST_MAX);
 	// Allowed to be unadded or awaiting add here (objects can be in more than one scene
-	if(OwningInstance->CurrentSceneState == BodyInstanceSceneState::NotAdded ||
-	   OwningInstance->CurrentSceneState == BodyInstanceSceneState::AwaitingAdd)
+	if (OwningInstance->CurrentSceneState == BodyInstanceSceneState::NotAdded ||
+		OwningInstance->CurrentSceneState == BodyInstanceSceneState::AwaitingAdd)
 	{
 		OwningInstance->CurrentSceneState = BodyInstanceSceneState::AwaitingAdd;
-		DeferredAddInstances[SceneType].Add(OwningInstance);
-		DeferredAddActors[SceneType].Add(Actor);
+		AddInstances.Add(OwningInstance);
+		AddActors.Add(Actor);
 	}
-	else if(OwningInstance->CurrentSceneState == BodyInstanceSceneState::AwaitingRemove)
+	else if (OwningInstance->CurrentSceneState == BodyInstanceSceneState::AwaitingRemove)
 	{
 		// We were waiting to be removed, but we're canceling that
 
 		OwningInstance->CurrentSceneState = BodyInstanceSceneState::Added;
-		DeferredRemoveInstances[SceneType].RemoveSingle(OwningInstance);
-		DeferredRemoveActors[SceneType].RemoveSingle(Actor);
+		RemoveInstances.RemoveSingle(OwningInstance);
+		RemoveActors.RemoveSingle(Actor);
 	}
 }
 
-void FPhysScene::DeferAddActors(TArray<FBodyInstance*>& OwningInstances, TArray<PxActor*>& Actors, EPhysicsSceneType SceneType)
+void FPhysScene::DeferAddActor(FBodyInstance* OwningInstance, PxActor* Actor, EPhysicsSceneType SceneType)
 {
-	int32 Num = OwningInstances.Num();
+	check(IsInGameThread());
+	check(OwningInstance && Actor && SceneType < PST_MAX);
+	DeferredSceneData[SceneType].DeferAddActor(OwningInstance, Actor);
+}
 
-	DeferredAddInstances[SceneType].Reserve(DeferredAddInstances->Num() + Num);
-	DeferredAddActors[SceneType].Reserve(DeferredAddActors->Num() + Num);
 
-	for(int32 Idx = 0 ; Idx < Num ; ++Idx)
+void FPhysScene::FDeferredSceneData::DeferAddActors(TArray<FBodyInstance*>& OwningInstances, TArray<PxActor*>& Actors)
+{
+	int32 Num = OwningInstances.Num();
+	AddInstances.Reserve(AddInstances.Num() + Num);
+	AddActors.Reserve(AddActors.Num() + Num);
+
+	for (int32 Idx = 0; Idx < Num; ++Idx)
 	{
-		DeferAddActor(OwningInstances[Idx], Actors[Idx], SceneType);
+		DeferAddActor(OwningInstances[Idx], Actors[Idx]);
 	}
 }
 
-void FPhysScene::DeferRemoveActor(FBodyInstance* OwningInstance, PxActor* Actor, EPhysicsSceneType SceneType)
+void FPhysScene::DeferAddActors(TArray<FBodyInstance*>& OwningInstances, TArray<PxActor*>& Actors, EPhysicsSceneType SceneType)
 {
-	check(OwningInstance && Actor && SceneType < PST_MAX);
-	if(OwningInstance->CurrentSceneState == BodyInstanceSceneState::Added ||
-	   OwningInstance->CurrentSceneState == BodyInstanceSceneState::AwaitingRemove)
+	check(IsInGameThread());
+	check(SceneType < PST_MAX);
+	DeferredSceneData[SceneType].DeferAddActors(OwningInstances, Actors);
+}
+
+void FPhysScene::FDeferredSceneData::DeferRemoveActor(FBodyInstance* OwningInstance, PxActor* Actor)
+{
+	if (OwningInstance->CurrentSceneState == BodyInstanceSceneState::Added ||
+		OwningInstance->CurrentSceneState == BodyInstanceSceneState::AwaitingRemove)
 	{
 		OwningInstance->CurrentSceneState = BodyInstanceSceneState::AwaitingRemove;
-		DeferredRemoveInstances[SceneType].Add(OwningInstance);
-		DeferredRemoveActors[SceneType].Add(Actor);
+		RemoveInstances.Add(OwningInstance);
+		RemoveActors.Add(Actor);
 	}
-	else if(OwningInstance->CurrentSceneState == BodyInstanceSceneState::AwaitingAdd)
+	else if (OwningInstance->CurrentSceneState == BodyInstanceSceneState::AwaitingAdd)
 	{
 		// We were waiting to add but now we're canceling it
 		OwningInstance->CurrentSceneState = BodyInstanceSceneState::Removed;
-		DeferredAddInstances[SceneType].RemoveSingle(OwningInstance);
-		DeferredAddActors[SceneType].RemoveSingle(Actor);
+		AddInstances.RemoveSingle(OwningInstance);
+		AddActors.RemoveSingle(Actor);
 	}
 }
 
+void FPhysScene::DeferRemoveActor(FBodyInstance* OwningInstance, PxActor* Actor, EPhysicsSceneType SceneType)
+{
+	check(IsInGameThread());
+	check(OwningInstance && Actor && SceneType < PST_MAX);
+	DeferredSceneData[SceneType].DeferRemoveActor(OwningInstance, Actor);
+}
+
 #endif
\ No newline at end of file
diff --git a/Engine/Source/Runtime/Engine/Public/PhysicsPublic.h b/Engine/Source/Runtime/Engine/Public/PhysicsPublic.h
--- a/Engine/Source/Runtime/Engine/Public/PhysicsPublic.h
+++ b/Engine/Source/Runtime/Engine/Public/PhysicsPublic.h
@@ -297,16 +297,6 @@ public:
 	float							FrameTimeSmoothingFactor[PST_MAX];
 
 #if WITH_PHYSX
-	/** Body instances awaiting scene add */
-	TArray<FBodyInstance*> DeferredAddInstances[PST_MAX];
-	/** PhysX Actors awaiting scene add */
-	TArray<PxActor*> DeferredAddActors[PST_MAX];
-
-	/** Body instances awaiting scene remove */
-	TArray<FBodyInstance*> DeferredRemoveInstances[PST_MAX];
-	/** PhysX Actors awaiting scene remove */
-	TArray<PxActor*> DeferredRemoveActors[PST_MAX];
-
 	/** Flush the deferred actor and instance arrays, either adding or removing from the scene */
 	void FlushDeferredActors();
 
@@ -354,6 +344,31 @@ private:
 	FGraphEventRef PhysicsSceneCompletion;
 
 #if WITH_PHYSX
+
+	struct FDeferredSceneData
+	{
+		/** The PhysX scene index used*/
+		int32 SceneIndex;
+
+		/** Body instances awaiting scene add */
+		TArray<FBodyInstance*> AddInstances;
+		/** PhysX Actors awaiting scene add */
+		TArray<PxActor*> AddActors;
+
+		/** Body instances awaiting scene remove */
+		TArray<FBodyInstance*> RemoveInstances;
+		/** PhysX Actors awaiting scene remove */
+		TArray<PxActor*> RemoveActors;
+
+		void FlushDeferredActors();
+		void DeferAddActor(FBodyInstance* OwningInstance, PxActor* Actor);
+		void DeferAddActors(TArray<FBodyInstance*>& OwningInstances, TArray<PxActor*>& Actors);
+		void DeferRemoveActor(FBodyInstance* OwningInstance, PxActor* Actor);
+	};
+
+	FDeferredSceneData DeferredSceneData[PST_MAX];
+
+
 	/** Dispatcher for CPU tasks */
 	class PxCpuDispatcher*			CPUDispatcher;
 	/** Simulation event callback object */
