#ue4 - Add virtual hook in CharacterMovementComponent to determine if a client's state is too different than the server's (ServerCheckClientError).

- This is used by ServerMoveHandleClientError() to determine whether to send a client error correction.
- Added FNetworkPredictionData_Server_Character::bForceClientUpdate to allow users to force a client correction.

UE-9654

diff --git a/Engine/Source/Runtime/Engine/Classes/GameFramework/CharacterMovementComponent.h b/Engine/Source/Runtime/Engine/Classes/GameFramework/CharacterMovementComponent.h
--- a/Engine/Source/Runtime/Engine/Classes/GameFramework/CharacterMovementComponent.h
+++ b/Engine/Source/Runtime/Engine/Classes/GameFramework/CharacterMovementComponent.h
@@ -1639,8 +1639,19 @@ protected:
 	/** Call the appropriate replicated servermove() function to send a client player move to the server. */
 	virtual void CallServerMove(const class FSavedMove_Character* NewMove, const class FSavedMove_Character* OldMove);
 	
-	/** Have the server check if the client is outside an error tolerance, and set a client adjustment if so. ClientLoc will be a relative location if MovementBaseUtility::UseRelativePosition(ClientMovementBase) is true. */
-	virtual void ServerMoveHandleClientError(float TimeStamp, float DeltaTime, const FVector& Accel, const FVector& ClientLoc, UPrimitiveComponent* ClientMovementBase, FName ClientBaseBoneName, uint8 ClientMovementMode);
+	/**
+	 * Have the server check if the client is outside an error tolerance, and queue a client adjustment if so.
+	 * If either GetPredictionData_Server_Character()->bForceClientUpdate or ServerCheckClientError() are true, the client adjustment will be sent.
+	 * RelativeClientLocation will be a relative location if MovementBaseUtility::UseRelativePosition(ClientMovementBase) is true, or a world location if false.
+	 * @see ServerCheckClientError()
+	 */
+	virtual void ServerMoveHandleClientError(float ClientTimeStamp, float DeltaTime, const FVector& Accel, const FVector& RelativeClientLocation, UPrimitiveComponent* ClientMovementBase, FName ClientBaseBoneName, uint8 ClientMovementMode);
+
+	/**
+	 * Check for Server-Client disagreement in position or other movement state important enough to trigger a client correction.
+	 * @see ServerMoveHandleClientError()
+	 */
+	virtual bool ServerCheckClientError(float ClientTimeStamp, float DeltaTime, const FVector& Accel, const FVector& ClientWorldLocation, const FVector& RelativeClientLocation, UPrimitiveComponent* ClientMovementBase, FName ClientBaseBoneName, uint8 ClientMovementMode);
 
 	/* Process a move at the given time stamp, given the compressed flags representing various events that occurred (ie jump). */
 	virtual void MoveAutonomous( float ClientTimeStamp, float DeltaTime, uint8 CompressedFlags, const FVector& NewAccel);
@@ -2017,6 +2028,8 @@ public:
 	// @TODO: don't duplicate between server and client data (though it's used by both)
 	float MaxResponseTime;
 
+	uint32 bForceClientUpdate:1;	// Force client update on the next ServerMoveHandleClientError() call.
+
 	/** @return time delta to use for the current ServerMove() */
 	float GetServerMoveDeltaTime(float TimeStamp) const;
 };
diff --git a/Engine/Source/Runtime/Engine/Private/Components/CharacterMovementComponent.cpp b/Engine/Source/Runtime/Engine/Private/Components/CharacterMovementComponent.cpp
--- a/Engine/Source/Runtime/Engine/Private/Components/CharacterMovementComponent.cpp
+++ b/Engine/Source/Runtime/Engine/Private/Components/CharacterMovementComponent.cpp
@@ -6162,7 +6162,7 @@ void UCharacterMovementComponent::ServerMove_Implementation(
 }
 
 
-void UCharacterMovementComponent::ServerMoveHandleClientError(float TimeStamp, float DeltaTime, const FVector& Accel, const FVector& RelativeClientLoc, UPrimitiveComponent* ClientMovementBase, FName ClientBaseBoneName, uint8 ClientMovementMode)
+void UCharacterMovementComponent::ServerMoveHandleClientError(float ClientTimeStamp, float DeltaTime, const FVector& Accel, const FVector& RelativeClientLoc, UPrimitiveComponent* ClientMovementBase, FName ClientBaseBoneName, uint8 ClientMovementMode)
 {
 	if (RelativeClientLoc == FVector(1.f,2.f,3.f)) // first part of double servermove
 	{
@@ -6191,11 +6191,8 @@ void UCharacterMovementComponent::ServerMoveHandleClientError(float TimeStamp, f
 	}
 
 	// Compute the client error from the server's position
-	const FVector LocDiff = CharacterOwner->GetActorLocation() - ClientLoc;
-	const uint8 CurrentPackedMovementMode = PackNetworkMovementMode();
-
 	// If client has accumulated a noticeable positional error, correct him.
-	if (GetDefault<AGameNetworkManager>()->ExceedsAllowablePositionError(LocDiff) || (CurrentPackedMovementMode != ClientMovementMode))
+	if (ServerData->bForceClientUpdate || ServerCheckClientError(ClientTimeStamp, DeltaTime, Accel, ClientLoc, RelativeClientLoc, ClientMovementBase, ClientBaseBoneName, ClientMovementMode))
 	{
 		UPrimitiveComponent* MovementBase = CharacterOwner->GetMovementBase();
 		ServerData->PendingAdjustment.NewVel = Velocity;
@@ -6218,8 +6215,9 @@ void UCharacterMovementComponent::ServerMoveHandleClientError(float TimeStamp, f
 #if !(UE_BUILD_SHIPPING || UE_BUILD_TEST)
 		if (CVarNetShowCorrections.GetValueOnGameThread() != 0)
 		{
+			const FVector LocDiff = CharacterOwner->GetActorLocation() - ClientLoc;
 			UE_LOG(LogNetPlayerMovement, Warning, TEXT("******** Client Error at %f is %f Accel %s LocDiff %s ClientLoc %s, ServerLoc: %s, Base: %s, Bone: %s"),
-				TimeStamp, LocDiff.Size(), *Accel.ToString(), *LocDiff.ToString(), *ClientLoc.ToString(), *CharacterOwner->GetActorLocation().ToString(), *GetNameSafe(MovementBase), *ServerData->PendingAdjustment.NewBaseBoneName.ToString());
+				ClientTimeStamp, LocDiff.Size(), *Accel.ToString(), *LocDiff.ToString(), *ClientLoc.ToString(), *CharacterOwner->GetActorLocation().ToString(), *GetNameSafe(MovementBase), *ServerData->PendingAdjustment.NewBaseBoneName.ToString());
 			const float DebugLifetime = CVarNetCorrectionLifetime.GetValueOnGameThread();
 			DrawDebugCapsule(GetWorld(), CharacterOwner->GetActorLocation(), CharacterOwner->GetSimpleCollisionHalfHeight(), CharacterOwner->GetSimpleCollisionRadius(), FQuat::Identity, FColor(100, 255, 100), true, DebugLifetime);
 			DrawDebugCapsule(GetWorld(), ClientLoc                    , CharacterOwner->GetSimpleCollisionHalfHeight(), CharacterOwner->GetSimpleCollisionRadius(), FQuat::Identity, FColor(255, 100, 100), true, DebugLifetime);
@@ -6228,15 +6226,16 @@ void UCharacterMovementComponent::ServerMoveHandleClientError(float TimeStamp, f
 
 		ServerData->LastUpdateTime = GetWorld()->TimeSeconds;
 		ServerData->PendingAdjustment.DeltaTime = DeltaTime;
-		ServerData->PendingAdjustment.TimeStamp = TimeStamp;
+		ServerData->PendingAdjustment.TimeStamp = ClientTimeStamp;
 		ServerData->PendingAdjustment.bAckGoodMove = false;
-		ServerData->PendingAdjustment.MovementMode = CurrentPackedMovementMode;
+		ServerData->PendingAdjustment.MovementMode = PackNetworkMovementMode();
 	}
 	else
 	{
 		if (GetDefault<AGameNetworkManager>()->ClientAuthorativePosition)
 		{
-			if (!LocDiff.IsZero() || ClientMovementMode != CurrentPackedMovementMode)
+			const FVector LocDiff = CharacterOwner->GetActorLocation() - ClientLoc;
+			if (!LocDiff.IsZero() || ClientMovementMode != PackNetworkMovementMode())
 			{
 				// Just set the position. On subsequent moves we will resolve initially overlapping conditions.
 				UpdatedComponent->SetWorldLocation(ClientLoc, false);
@@ -6254,9 +6253,30 @@ void UCharacterMovementComponent::ServerMoveHandleClientError(float TimeStamp, f
 		}
 
 		// acknowledge receipt of this successful servermove()
-		ServerData->PendingAdjustment.TimeStamp = TimeStamp;
+		ServerData->PendingAdjustment.TimeStamp = ClientTimeStamp;
 		ServerData->PendingAdjustment.bAckGoodMove = true;
 	}
+
+	ServerData->bForceClientUpdate = false;
+}
+
+bool UCharacterMovementComponent::ServerCheckClientError(float ClientTimeStamp, float DeltaTime, const FVector& Accel, const FVector& ClientWorldLocation, const FVector& RelativeClientLocation, UPrimitiveComponent* ClientMovementBase, FName ClientBaseBoneName, uint8 ClientMovementMode)
+{
+	// Check location difference against global setting
+	const FVector LocDiff = CharacterOwner->GetActorLocation() - ClientWorldLocation;
+	if (GetDefault<AGameNetworkManager>()->ExceedsAllowablePositionError(LocDiff))
+	{
+		return true;
+	}
+
+	// Check for disagreement in movement mode
+	const uint8 CurrentPackedMovementMode = PackNetworkMovementMode();
+	if (CurrentPackedMovementMode != ClientMovementMode)
+	{
+		return true;
+	}
+
+	return false;
 }
 
 
@@ -6400,6 +6420,7 @@ void UCharacterMovementComponent::SendClientAdjustment()
 
 	ServerData->PendingAdjustment.TimeStamp = 0;
 	ServerData->PendingAdjustment.bAckGoodMove = false;
+	ServerData->bForceClientUpdate = false;
 }
 
 
@@ -7027,6 +7048,7 @@ FNetworkPredictionData_Server_Character::FNetworkPredictionData_Server_Character
 	, CurrentClientTimeStamp(0.f)
 	, LastUpdateTime(0.f)
 	, MaxResponseTime(0.125f)
+	, bForceClientUpdate(false)
 {
 }
 
