#ue4 - Don't use an invalid floor as the character base component when initially entering Walking mode.

- Add UCharacterMovementComponent::SetBaseFromFloor() helper to validate that the floor is indeed walkable before calling SetBase() with the hit component.

UDN: https://udn.unrealengine.com/questions/239213/how-do-i-prevent-projectile-collision-from-impacti.html

diff --git a/Engine/Source/Runtime/Engine/Classes/GameFramework/CharacterMovementComponent.h b/Engine/Source/Runtime/Engine/Classes/GameFramework/CharacterMovementComponent.h
--- a/Engine/Source/Runtime/Engine/Classes/GameFramework/CharacterMovementComponent.h
+++ b/Engine/Source/Runtime/Engine/Classes/GameFramework/CharacterMovementComponent.h
@@ -1028,9 +1028,14 @@ public:
 	 */
 	virtual bool StepUp(const FVector& GravDir, const FVector& Delta, const FHitResult &Hit, struct UCharacterMovementComponent::FStepDownResult* OutStepDownResult = NULL);
 
-	/** Update the base of the character, which is the obejct we are standing on. */
+	/** Update the base of the character, which is the PrimitiveComponent we are standing on. */
 	virtual void SetBase(UPrimitiveComponent* NewBase, const FName BoneName = NAME_None, bool bNotifyActor=true);
 
+	/**
+	 * Update the base of the character, using the given floor result if it is walkable, or null if not. Calls SetBase().
+	 */
+	void SetBaseFromFloor(const FFindFloorResult& FloorResult);
+
 	/** Applies repulsion force to all touched components. */
 	virtual void ApplyRepulsionForce(float DeltaSeconds);
 	
@@ -1407,21 +1412,22 @@ protected:
 	 * Sweeps a vertical trace to find the floor for the capsule at the given location. Will attempt to perch if ShouldComputePerchResult() returns true for the downward sweep result.
 	 *
 	 * @param CapsuleLocation:		Location where the capsule sweep should originate
-	 * @param OutFloorResult:		[Out] Contains the result of the floor check.
+	 * @param OutFloorResult:		[Out] Contains the result of the floor check. The HitResult will contain the valid sweep or line test upon success, or the result of the sweep upon failure.
 	 * @param bZeroDelta:			If true, the capsule was not actively moving in this update (can be used to avoid unnecessary floor tests).
 	 * @param DownwardSweepResult:	If non-null and it contains valid blocking hit info, this will be used as the result of a downward sweep test instead of doing it as part of the update.
 	 */
 	virtual void FindFloor(const FVector& CapsuleLocation, struct FFindFloorResult& OutFloorResult, bool bZeroDelta, const FHitResult* DownwardSweepResult = NULL) const;
 
 	/**
-	 * Compute distance to the floor from bottom sphere of capsule. This is the swept distance of the capsule to the first point impacted by the lower sphere.
+	 * Compute distance to the floor from bottom sphere of capsule and store the result in OutFloorResult.
+	 * This distance is the swept distance of the capsule to the first point impacted by the lower hemisphere, or distance from the bottom of the capsule in the case of a line trace.
 	 * SweepDistance MUST be greater than or equal to the line distance.
 	 * @see FindFloor
 	 *
 	 * @param CapsuleLocation:	Location of the capsule used for the query
-	 * @param LineDistance:		If non-zero, max distance to test for a simple line check from the capsule base. Used before the sweep test, and only returns a valid result if the impact normal is a walkable normal.
+	 * @param LineDistance:		If non-zero, max distance to test for a simple line check from the capsule base. Used only if the sweep test fails to find a walkable floor, and only returns a valid result if the impact normal is a walkable normal.
 	 * @param SweepDistance:	If non-zero, max distance to use when sweeping a capsule downwards for the test.
-	 * @param OutFloorResult:	Result of the floor check.
+	 * @param OutFloorResult:	Result of the floor check. The HitResult will contain the valid sweep or line test upon success, or the result of the sweep upon failure.
 	 * @param SweepRadius:		The radius to use for sweep tests. Should be <= capsule radius.
 	 * @param DownwardSweepResult:	If non-null and it contains valid blocking hit info, this will be used as the result of a downward sweep test instead of doing it as part of the update.
 	 */
diff --git a/Engine/Source/Runtime/Engine/Private/Character.cpp b/Engine/Source/Runtime/Engine/Private/Character.cpp
--- a/Engine/Source/Runtime/Engine/Private/Character.cpp
+++ b/Engine/Source/Runtime/Engine/Private/Character.cpp
@@ -545,8 +545,12 @@ namespace MovementBaseUtility
 
 
 /**	Change the Pawn's base. */
-void ACharacter::SetBase( UPrimitiveComponent* NewBaseComponent, const FName BoneName, bool bNotifyPawn )
+void ACharacter::SetBase( UPrimitiveComponent* NewBaseComponent, const FName InBoneName, bool bNotifyPawn )
 {
+	// If NewBaseComponent is null, ignore bone name.
+	const FName BoneName = (NewBaseComponent ? InBoneName : NAME_None);
+
+	// See what changed.
 	const bool bBaseChanged = (NewBaseComponent != BasedMovement.MovementBase);
 	const bool bBoneChanged = (BoneName != BasedMovement.BoneName);
 
diff --git a/Engine/Source/Runtime/Engine/Private/Components/CharacterMovementComponent.cpp b/Engine/Source/Runtime/Engine/Private/Components/CharacterMovementComponent.cpp
--- a/Engine/Source/Runtime/Engine/Private/Components/CharacterMovementComponent.cpp
+++ b/Engine/Source/Runtime/Engine/Private/Components/CharacterMovementComponent.cpp
@@ -708,7 +708,7 @@ void UCharacterMovementComponent::OnMovementModeChanged(EMovementMode PreviousMo
 		// make sure we update our new floor/base on initial entry of the walking physics
 		FindFloor(UpdatedComponent->GetComponentLocation(), CurrentFloor, false);
 		AdjustFloorHeight();
-		SetBase(CurrentFloor.HitResult.Component.Get(), CurrentFloor.HitResult.BoneName);
+		SetBaseFromFloor(CurrentFloor);
 	}
 	else
 	{
@@ -1263,7 +1263,19 @@ void UCharacterMovementComponent::SetBase( UPrimitiveComponent* NewBase, FName B
 {
 	if (CharacterOwner)
 	{
-		CharacterOwner->SetBase(NewBase, BoneName, bNotifyActor);
+		CharacterOwner->SetBase(NewBase, NewBase ? BoneName : NAME_None, bNotifyActor);
+	}
+}
+
+void UCharacterMovementComponent::SetBaseFromFloor(const FFindFloorResult& FloorResult)
+{
+	if (FloorResult.IsWalkableFloor())
+	{
+		SetBase(FloorResult.HitResult.GetComponent(), FloorResult.HitResult.BoneName);
+	}
+	else
+	{
+		SetBase(nullptr);
 	}
 }
 
@@ -6532,8 +6544,16 @@ void UCharacterMovementComponent::ClientAdjustPosition_Implementation
 		if (PreviousBase)
 		{
 			FindFloor(UpdatedComponent->GetComponentLocation(), CurrentFloor, false);
-			FinalBase = CurrentFloor.HitResult.Component.Get();
-			FinalBaseBoneName = CurrentFloor.HitResult.BoneName;
+			if (CurrentFloor.IsWalkableFloor())
+			{
+				FinalBase = CurrentFloor.HitResult.Component.Get();
+				FinalBaseBoneName = CurrentFloor.HitResult.BoneName;
+			}
+			else
+			{
+				FinalBase = nullptr;
+				FinalBaseBoneName = NAME_None;
+			}
 		}
 	}
 	SetBase(FinalBase, FinalBaseBoneName);
