Fix for null texture lod settings on certain platforms. Wrong target platform list was processed.

diff --git a/Engine/Source/Runtime/Engine/Private/DeviceProfiles/DeviceProfileManager.cpp b/Engine/Source/Runtime/Engine/Private/DeviceProfiles/DeviceProfileManager.cpp
--- a/Engine/Source/Runtime/Engine/Private/DeviceProfiles/DeviceProfileManager.cpp
+++ b/Engine/Source/Runtime/Engine/Private/DeviceProfiles/DeviceProfileManager.cpp
@@ -254,7 +254,7 @@ void UDeviceProfileManager::LoadProfiles()
 
 		// Register Texture LOD settings with each Target Platform
 		ITargetPlatformManagerModule& TargetPlatformManager = GetTargetPlatformManagerRef();
-		const TArray<ITargetPlatform*>& TargetPlatforms = TargetPlatformManager.GetActiveTargetPlatforms();
+		const TArray<ITargetPlatform*>& TargetPlatforms = TargetPlatformManager.GetTargetPlatforms();
 		for (int32 PlatformIndex = 0; PlatformIndex < TargetPlatforms.Num(); ++PlatformIndex)
 		{
 			ITargetPlatform* Platform = TargetPlatforms[PlatformIndex];
