Fix for Target platform being used on game builds. Also a small optimization as we dont need to load unnecessary device profiles on other devices.

diff --git a/Engine/Source/Runtime/Engine/Private/DeviceProfiles/DeviceProfileManager.cpp b/Engine/Source/Runtime/Engine/Private/DeviceProfiles/DeviceProfileManager.cpp
--- a/Engine/Source/Runtime/Engine/Private/DeviceProfiles/DeviceProfileManager.cpp
+++ b/Engine/Source/Runtime/Engine/Private/DeviceProfiles/DeviceProfileManager.cpp
@@ -18,7 +18,10 @@ UDeviceProfileManager& UDeviceProfileManager::Get()
 		DeviceProfileManagerSingleton = NewObject<UDeviceProfileManager>();
 
 		DeviceProfileManagerSingleton->AddToRoot();
-		DeviceProfileManagerSingleton->LoadProfiles();
+		if (!FPlatformProperties::RequiresCookedData())
+		{
+			DeviceProfileManagerSingleton->LoadProfiles();
+		}
 
 		FString DeviceProfileSelectionModule;
 		GConfig->GetString(TEXT("DeviceProfileManager"), TEXT("DeviceProfileSelectionModule"), DeviceProfileSelectionModule, GEngineIni);
