Fix play state on changing montages in Persona

diff --git a/Engine/Source/Runtime/Engine/Private/Animation/AnimSingleNodeInstance.cpp b/Engine/Source/Runtime/Engine/Private/Animation/AnimSingleNodeInstance.cpp
--- a/Engine/Source/Runtime/Engine/Private/Animation/AnimSingleNodeInstance.cpp
+++ b/Engine/Source/Runtime/Engine/Private/Animation/AnimSingleNodeInstance.cpp
@@ -71,6 +71,7 @@ void UAnimSingleNodeInstance::SetAnimationAsset(class UAnimationAsset* NewAsset,
 			RegisterSlotNodeWithAnimInstance(Montage->SlotAnimTracks[0].SlotName);
 		}
 		RestartMontage( Montage );
+		SetPlaying(bPlaying);
 	}
 	else
 	{
