Reverting erroneous reversion of 2481510 in 2483033

diff --git a/Engine/Source/Runtime/Engine/Private/UnrealEngine.cpp b/Engine/Source/Runtime/Engine/Private/UnrealEngine.cpp
--- a/Engine/Source/Runtime/Engine/Private/UnrealEngine.cpp
+++ b/Engine/Source/Runtime/Engine/Private/UnrealEngine.cpp
@@ -3337,7 +3337,7 @@ bool UEngine::HandleListTexturesCommand( const TCHAR* Cmd, FOutputDevice& Ar )
 	SortedTextures.Sort( FCompareFSortedTexture( bAlphaSort ) );
 
 	// Retrieve mapping from LOD group enum value to text representation.
-	TArray<FString> TextureGroupNames = FTextureLODSettings::GetTextureGroupNames();
+	TArray<FString> TextureGroupNames = UTextureLODSettings::GetTextureGroupNames();
 
 	// Display.
 	int32 TotalMaxSize		= 0;
@@ -10645,17 +10645,6 @@ bool UEngine::ShouldAbsorbAuthorityOnlyEvent()
 }
 
 
-UDeviceProfileManager* UEngine::GetDeviceProfileManager()
-{
-	if(DeviceProfileManager == NULL)
-	{
-		DeviceProfileManager = NewObject<UDeviceProfileManager>(GetTransientPackage(), TEXT("GlobalDeviceProfileManager"), RF_Public | RF_Transient);
-	}
-
-	return DeviceProfileManager;
-}
-
-
 bool UEngine::ShouldAbsorbCosmeticOnlyEvent()
 {
 	for (auto It = WorldList.CreateIterator(); It; ++It)
