DemoNetDriver closes channels for bTearOff actors, and actors that aren't bRelevantForNetworkReplays anymore.

diff --git a/Engine/Source/Runtime/Engine/Private/DemoNetDriver.cpp b/Engine/Source/Runtime/Engine/Private/DemoNetDriver.cpp
--- a/Engine/Source/Runtime/Engine/Private/DemoNetDriver.cpp
+++ b/Engine/Source/Runtime/Engine/Private/DemoNetDriver.cpp
@@ -630,31 +630,35 @@ static void DemoReplicateActor(AActor* Actor, UNetConnection* Connection, bool I
 	*/
 	if ( Actor != NULL && Actor->GetRemoteRole() != ROLE_None && ( Actor == Connection->PlayerController || Cast< APlayerController >( Actor ) == NULL ) )
 	{
-		if (Actor->bRelevantForNetworkReplays)
+		const bool bShouldHaveChannel =
+			Actor->bRelevantForNetworkReplays &&
+			!Actor->bTearOff &&
+			(!Actor->IsNetStartupActor() || Connection->ClientHasInitializedLevelFor(Actor));
+
+		UActorChannel* Channel = Connection->ActorChannels.FindRef(Actor);
+
+		if (bShouldHaveChannel && Channel == NULL)
 		{
 			// Create a new channel for this actor.
-			const bool StartupActor = Actor->IsNetStartupActor();
-			UActorChannel* Channel = Connection->ActorChannels.FindRef(Actor);
-
-			if (!Channel && (!StartupActor || Connection->ClientHasInitializedLevelFor(Actor)))
+			Channel = (UActorChannel*)Connection->CreateChannel(CHTYPE_Actor, 1);
+			if (Channel != NULL)
 			{
-				// create a channel if possible
-				Channel = (UActorChannel*)Connection->CreateChannel(CHTYPE_Actor, 1);
-				if (Channel != NULL)
-				{
-					Channel->SetChannelActor(Actor);
-				}
+				Channel->SetChannelActor(Actor);
 			}
+		}
 
-			if (Channel)
+		if (Channel != NULL && !Channel->Closing)
+		{
+			// Send it out!
+			if (Channel->IsNetReady(0))
 			{
-				// Send it out!
-				check(!Channel->Closing);
-
-				if (Channel->IsNetReady(0))
-				{
-					Channel->ReplicateActor();
-				}
+				Channel->ReplicateActor();
+			}
+			
+			// Close the channel if this actor shouldn't have one
+			if (!bShouldHaveChannel)
+			{
+				Channel->Close();
 			}
 		}
 	}
