Fix to root node index mismatch on child anim blueprints after structurally changing and compiling parent when child is empty and considered data-only

diff --git a/Engine/Source/Runtime/Engine/Private/Animation/AnimBlueprintGeneratedClass.cpp b/Engine/Source/Runtime/Engine/Private/Animation/AnimBlueprintGeneratedClass.cpp
--- a/Engine/Source/Runtime/Engine/Private/Animation/AnimBlueprintGeneratedClass.cpp
+++ b/Engine/Source/Runtime/Engine/Private/Animation/AnimBlueprintGeneratedClass.cpp
@@ -203,6 +203,21 @@ void UAnimBlueprintGeneratedClass::Link(FArchive& Ar, bool bRelinkExistingProper
 		}
 	}
 
+	// Pull info down from root anim class
+	UAnimBlueprintGeneratedClass* RootClass = this;
+	while(UAnimBlueprintGeneratedClass* NextClass = Cast<UAnimBlueprintGeneratedClass>(RootClass->GetSuperClass()))
+	{
+		RootClass = NextClass;
+	}
+
+	if(RootClass != this)
+	{
+		// Copy root, state notifies and baked machines from the root class
+		RootAnimNodeIndex = RootClass->RootAnimNodeIndex;
+		AnimNotifies = RootClass->AnimNotifies;
+		BakedStateMachines = RootClass->BakedStateMachines;
+	}
+
 	if (AnimNodeProperties.Num() > 0)
 	{
 		const bool bValidRootIndex = (RootAnimNodeIndex >= 0) && (RootAnimNodeIndex < AnimNodeProperties.Num());
