Fixing assert for no audio device

jira UE-12086

diff --git a/Engine/Source/Runtime/Engine/Private/UnrealEngine.cpp b/Engine/Source/Runtime/Engine/Private/UnrealEngine.cpp
--- a/Engine/Source/Runtime/Engine/Private/UnrealEngine.cpp
+++ b/Engine/Source/Runtime/Engine/Private/UnrealEngine.cpp
@@ -1648,6 +1648,7 @@ bool UEngine::InitializeAudioDeviceManager()
 		// Initialize the audio device.
 		if (bUseSound == true)
 		{
+
 			// get the module name from the ini file
 			FString AudioDeviceModuleName;
 			GConfig->GetString(TEXT("Audio"), TEXT("AudioDeviceModuleName"), AudioDeviceModuleName, GEngineIni);
@@ -1664,6 +1665,8 @@ bool UEngine::InitializeAudioDeviceManager()
 					AudioDeviceManager = new FAudioDeviceManager();
 					AudioDeviceManager->RegisterAudioDeviceModule(AudioDeviceModule);
 
+					bool bSucceeded = false;
+
 					// Create a new audio device.
 					FAudioDevice* NewAudioDevice = AudioDeviceManager->CreateAudioDevice(MainAudioDeviceHandle);
 					if (NewAudioDevice)
@@ -1671,6 +1674,7 @@ bool UEngine::InitializeAudioDeviceManager()
 						// Initialize the audio device
 						if (NewAudioDevice->Init())
 						{
+							bSucceeded = true;
 							AudioDeviceManager->SetActiveDevice(MainAudioDeviceHandle);
 						}
 						else
@@ -1679,6 +1683,12 @@ bool UEngine::InitializeAudioDeviceManager()
 							AudioDeviceManager->ShutdownAudioDevice(MainAudioDeviceHandle);
 						}
 					}
+
+					// if we failed to create or init a main audio device, shut down the audio device manager
+					if (!bSucceeded)
+					{
+						ShutdownAudioDeviceManager();
+					}
 				}
 			}
 		}
