Fully initialize FNetGuidCacheObject, which is needed server side when saving checkpoints

diff --git a/Engine/Source/Runtime/Engine/Private/PackageMapClient.cpp b/Engine/Source/Runtime/Engine/Private/PackageMapClient.cpp
--- a/Engine/Source/Runtime/Engine/Private/PackageMapClient.cpp
+++ b/Engine/Source/Runtime/Engine/Private/PackageMapClient.cpp
@@ -427,12 +427,27 @@ void UPackageMapClient::InternalWriteObject( FArchive & Ar, FNetworkGUID NetGUID
 		// Serialize Name of object
 		Ar << ObjectPathName;
 
+		FNetGuidCacheObject* CacheObject = GuidCache->ObjectLookup.Find( NetGUID );
+
 		if ( bIsPackage )
 		{
 			FGuid PackageGuid = CastChecked< const UPackage >( Object )->GetGuid();
 			Ar << PackageGuid;
 
 			UE_LOG( LogNetPackageMap, VeryVerbose, TEXT( "InternalWriteObject: Package: %s, GUID: %s" ), *ObjectPathName, *PackageGuid.ToString() );
+
+			if ( CacheObject != NULL )
+			{
+				CacheObject->PackageGuid = PackageGuid;
+			}
+		}
+
+		if ( CacheObject != NULL )
+		{
+			CacheObject->PathName			= FName( *ObjectPathName );
+			CacheObject->OuterGUID			= OuterNetGUID;
+			CacheObject->bNoLoad			= ExportFlags.bNoLoad;
+			CacheObject->bIgnoreWhenMissing = ExportFlags.bNoLoad;
 		}
 
 		if ( GuidCache->IsExportingNetGUIDBunch )
