Initial network replay checkpoint implementation

diff --git a/Engine/Source/Runtime/Engine/Classes/Engine/DemoNetDriver.h b/Engine/Source/Runtime/Engine/Classes/Engine/DemoNetDriver.h
--- a/Engine/Source/Runtime/Engine/Classes/Engine/DemoNetDriver.h
+++ b/Engine/Source/Runtime/Engine/Classes/Engine/DemoNetDriver.h
@@ -57,6 +57,11 @@ class UDemoNetDriver
 	double		MaxRecordTime;
 	int32		RecordCountSinceFlush;
 	float		TimeToSkip;
+
+	bool		bSavingCheckpoint;
+	double		LastCheckpointTime;
+
+	void		SaveCheckpoint();
 public:
 
 	// UNetDriver interface.
@@ -92,7 +97,8 @@ public:
 
 	void TickDemoRecord( float DeltaSeconds );
 	void PauseChannels( const bool bPause );
-	bool ReadDemoFrame();
+	bool ConditionallyReadDemoFrame();
+	bool ReadDemoFrame( FArchive* Archive );
 	void TickDemoPlayback( float DeltaSeconds );
 	void SpawnDemoRecSpectator( UNetConnection* Connection );
 	void ResetDemoState();
@@ -100,4 +106,5 @@ public:
 	void StopDemo();
 
 	void ReplayStreamingReady( bool bSuccess, bool bRecord );
+	void CheckpointReady( bool bSuccess );
 };
diff --git a/Engine/Source/Runtime/Engine/Private/DemoNetDriver.cpp b/Engine/Source/Runtime/Engine/Private/DemoNetDriver.cpp
--- a/Engine/Source/Runtime/Engine/Private/DemoNetDriver.cpp
+++ b/Engine/Source/Runtime/Engine/Private/DemoNetDriver.cpp
@@ -24,8 +24,10 @@ DEFINE_LOG_CATEGORY_STATIC( LogDemo, Log, All );
 static TAutoConsoleVariable<float> CVarDemoRecordHz( TEXT( "demo.RecordHz" ), 10, TEXT( "Number of demo frames recorded per second" ) );
 static TAutoConsoleVariable<float> CVarDemoTimeDilation( TEXT( "demo.TimeDilation" ), -1.0f, TEXT( "Override time dilation during demo playback (-1 = don't override)" ) );
 static TAutoConsoleVariable<float> CVarDemoSkipTime( TEXT( "demo.SkipTime" ), 0, TEXT( "Skip fixed amount of network replay time (in seconds)" ) );
+static TAutoConsoleVariable<int32> CVarEnableCheckpoints( TEXT( "demo.EnableCheckpoints" ), 0, TEXT( "Whether or not checkpoints save on the server" ) );
+static TAutoConsoleVariable<int32> CVarTestCheckpoint( TEXT( "demo.TestCheckpoint" ), 0, TEXT( "For testing only, jump to a particular checkpoint" ) );
 
-static const int32 MAX_DEMO_READ_WRITE_BUFFER = 1024 * 32;
+static const int32 MAX_DEMO_READ_WRITE_BUFFER = 1024 * 2;
 
 #define DEMO_CHECKSUMS 0		// When setting this to 1, this will invalidate all demos, you will need to re-record and playback
 
@@ -364,10 +366,6 @@ bool UDemoNetDriver::InitListen( FNetworkNotify* InNotify, FURL& ListenURL, bool
 		return false;
 	}
 
-	// use the same byte format regardless of platform so that the demos are cross platform
-	//@note: swap on non console platforms as the console archives have byte swapping compiled out by default
-	//FileAr->SetByteSwapping(true);
-
 	FNetworkDemoHeader DemoHeader;
 
 	DemoHeader.LevelName = World->GetCurrentLevel()->GetOutermost()->GetName();
@@ -495,6 +493,12 @@ void UDemoNetDriver::TickFlush( float DeltaSeconds )
 			}
 		}
 
+		if ( CVarTestCheckpoint.GetValueOnGameThread() > 0 )
+		{
+			ReplayStreamer->GotoCheckpoint( CVarTestCheckpoint.GetValueOnGameThread(), FOnCheckpointReadyDelegate::CreateUObject( this, &UDemoNetDriver::CheckpointReady ) );
+			CVarTestCheckpoint.AsVariable()->Set( TEXT( "0" ), ECVF_SetByConsole );
+		}
+
 		if ( bDemoPlaybackDone )
 		{
 			return;
@@ -664,6 +668,146 @@ static void DemoReplicateActor(AActor* Actor, UNetConnection* Connection, bool I
 	}
 }
 
+static void SerializeGuidCache( TSharedPtr< class FNetGUIDCache > GuidCache, FArchive* CheckpointArchive )
+{
+	int32 NumValues = 0;
+
+	for ( auto It = GuidCache->ObjectLookup.CreateIterator(); It; ++It )
+	{
+		if ( It.Value().Object == NULL )
+		{
+			continue;
+		}
+
+		if ( !It.Value().Object->IsNameStableForNetworking() )
+		{
+			continue;
+		}
+
+		NumValues++;
+	}
+
+	*CheckpointArchive << NumValues;
+
+	UE_LOG( LogDemo, Warning, TEXT( "Checkpoint. SerializeGuidCache: %i" ), NumValues );
+
+	for ( auto It = GuidCache->ObjectLookup.CreateIterator(); It; ++It )
+	{
+		if ( It.Value().Object == NULL )
+		{
+			continue;
+		}
+
+		if ( !It.Value().Object->IsNameStableForNetworking() )
+		{
+			continue;
+		}
+
+		FString PathName = It.Value().Object->GetName();
+
+		*CheckpointArchive << It.Key();
+		*CheckpointArchive << It.Value().OuterGUID;
+		*CheckpointArchive << PathName;
+		*CheckpointArchive << It.Value().PackageGuid;
+
+		uint8 Flags = 0;
+		
+		Flags |= It.Value().bNoLoad ? ( 1 << 0 ) : 0;
+		Flags |= It.Value().bIgnoreWhenMissing ? ( 1 << 1 ) : 0;
+
+		*CheckpointArchive << Flags;
+	}
+}
+
+void UDemoNetDriver::SaveCheckpoint()
+{
+	if ( DemoCurrentTime > 30.0f )
+	{
+		return;
+	}
+
+	FArchive* CheckpointArchive = ReplayStreamer->GetCheckpointArchive();
+
+	if ( CheckpointArchive == nullptr )
+	{
+		return;
+	}
+
+	const double StartCheckpointTime = FPlatformTime::Seconds();
+
+	// First, save the current guid cache
+	SerializeGuidCache( GuidCache, CheckpointArchive );
+
+	const uint32 GuidCacheSize = CheckpointArchive->TotalSize();
+
+	// Save total absolute demo time in MS
+	uint32 SavedAbsTimeMS = GetDemoCurrentTimeInMS();
+	*CheckpointArchive << SavedAbsTimeMS;
+
+	FURL CheckpointURL;
+	CheckpointURL.Map = TEXT( "Checkpoint" );
+
+	UDemoNetConnection* CheckpointConnection = NewObject<UDemoNetConnection>();
+	ClientConnections.Add( CheckpointConnection );
+	CheckpointConnection->InitConnection( this, USOCK_Open, CheckpointURL, 1000000 );
+	CheckpointConnection->InitSendBuffer();
+
+	// Some hackery to make the player thinks this checkpoint connection owns it
+	CheckpointConnection->PlayerController					= ClientConnections[0]->PlayerController;
+	CheckpointConnection->PlayerController->Player			= CheckpointConnection;
+	CheckpointConnection->PlayerController->NetConnection	= CheckpointConnection;
+	//CheckpointConnection->OwningActor						= CheckpointConnection->PlayerController;
+
+	// Make sure we have the exact same actor channel indexes
+	for ( auto It = ClientConnections[0]->ActorChannels.CreateIterator(); It; ++It )
+	{
+		UActorChannel* Channel = (UActorChannel*)CheckpointConnection->CreateChannel( CHTYPE_Actor, true, It.Value()->ChIndex );
+		if ( Channel != NULL )
+		{
+			Channel->SetChannelActor( It.Value()->Actor );
+		}
+	}
+
+	bSavingCheckpoint = true;
+
+	DemoReplicateActor( World->GetWorldSettings(), CheckpointConnection, false );
+
+	for ( int32 i = 0; i < World->NetworkActors.Num(); i++ )
+	{
+		AActor* Actor = World->NetworkActors[i];
+
+		Actor->PreReplication( *FindOrCreateRepChangedPropertyTracker( Actor ).Get() );
+		DemoReplicateActor( Actor, CheckpointConnection, false );
+	}
+
+	CheckpointConnection->FlushNet();
+
+	bSavingCheckpoint = false;
+
+	// Write a count of 0 to signal the end of the frame
+	int32 EndCount = 0;
+	*CheckpointArchive << EndCount;
+
+	// Undo hackery
+	ClientConnections[0]->PlayerController->Player			= ClientConnections[0];
+	ClientConnections[0]->PlayerController->NetConnection	= ClientConnections[0];
+
+	CheckpointConnection->Close();
+	CheckpointConnection->CleanUp();
+
+	const uint32 CheckpointSize = CheckpointArchive->TotalSize() - GuidCacheSize;
+
+	const int32 TotalSize = CheckpointArchive->TotalSize();
+
+	ReplayStreamer->FlushCheckpoint( 1 );
+
+	const double EndCheckpointTime = FPlatformTime::Seconds();
+
+	const float CheckpointTimeInMS = ( EndCheckpointTime - StartCheckpointTime ) * 1000.0f;
+
+	UE_LOG( LogDemo, Warning, TEXT( "Checkpoint. Total: %i, Rep size: %i, PackageMap: %u, Time: %2.2f" ), TotalSize, CheckpointSize, GuidCacheSize, CheckpointTimeInMS );
+}
+
 void UDemoNetDriver::TickDemoRecord( float DeltaSeconds )
 {
 	if ( ClientConnections.Num() == 0 )
@@ -694,6 +838,17 @@ void UDemoNetDriver::TickDemoRecord( float DeltaSeconds )
 
 	LastRecordTime = CurrentSeconds;
 
+	if ( CVarEnableCheckpoints.GetValueOnGameThread() == 1 )
+	{
+		const double CHECKPOINT_DELAY = 10.0;
+
+		if ( CurrentSeconds - LastCheckpointTime > CHECKPOINT_DELAY )
+		{
+			SaveCheckpoint();
+			LastCheckpointTime = CurrentSeconds;
+		}
+	}
+
 	// Save out a frame
 	DemoFrameNum++;
 	ReplicationFrame++;
@@ -775,7 +930,7 @@ void UDemoNetDriver::PauseChannels( const bool bPause )
 	bChannelsArePaused = bPause;
 }
 
-bool UDemoNetDriver::ReadDemoFrame()
+bool UDemoNetDriver::ConditionallyReadDemoFrame()
 {
 	FArchive* FileAr = ReplayStreamer->GetStreamingArchive();
 
@@ -794,7 +949,7 @@ bool UDemoNetDriver::ReadDemoFrame()
 
 	if ( ReplayStreamer->GetLastError() != ENetworkReplayError::None )
 	{
-		UE_LOG( LogDemo, Error, TEXT( "UDemoNetDriver::ReadDemoFrame: ReplayStreamer ERROR: %s" ), ENetworkReplayError::ToString( ReplayStreamer->GetLastError() ) );
+		UE_LOG( LogDemo, Error, TEXT( "UDemoNetDriver::ConditionallyReadDemoFrame: ReplayStreamer ERROR: %s" ), ENetworkReplayError::ToString( ReplayStreamer->GetLastError() ) );
 		StopDemo();
 		return false;
 	}
@@ -823,7 +978,7 @@ bool UDemoNetDriver::ReadDemoFrame()
 
 		if ( DeltaTimeChecksum != ServerDeltaTimeCheksum )
 		{
-			UE_LOG( LogDemo, Error, TEXT( "UDemoNetDriver::ReadDemoFrame: DeltaTimeChecksum != ServerDeltaTimeCheksum" ) );
+			UE_LOG( LogDemo, Error, TEXT( "UDemoNetDriver::ConditionallyReadDemoFrame: DeltaTimeChecksum != ServerDeltaTimeCheksum" ) );
 			StopDemo();
 			return false;
 		}
@@ -832,7 +987,7 @@ bool UDemoNetDriver::ReadDemoFrame()
 
 	if ( FileAr->IsError() )
 	{
-		UE_LOG( LogDemo, Error, TEXT( "UDemoNetDriver::ReadDemoFrame: Failed to read demo ServerDeltaTime" ) );
+		UE_LOG( LogDemo, Error, TEXT( "UDemoNetDriver::ConditionallyReadDemoFrame: Failed to read demo ServerDeltaTime" ) );
 		StopDemo();
 		return false;
 	}
@@ -844,15 +999,20 @@ bool UDemoNetDriver::ReadDemoFrame()
 		return false;
 	}
 
+	return ReadDemoFrame( FileAr );
+}
+
+bool UDemoNetDriver::ReadDemoFrame( FArchive* Archive )
+{
 	while ( true )
 	{
 		uint8 ReadBuffer[ MAX_DEMO_READ_WRITE_BUFFER ];
 
 		int32 PacketBytes;
 
-		*FileAr << PacketBytes;
+		*Archive << PacketBytes;
 
-		if ( FileAr->IsError() )
+		if ( Archive->IsError() )
 		{
 			UE_LOG( LogDemo, Error, TEXT( "UDemoNetDriver::ReadDemoFrame: Failed to read demo PacketBytes" ) );
 			StopDemo();
@@ -879,9 +1039,9 @@ bool UDemoNetDriver::ReadDemoFrame()
 		}
 
 		// Read data from file.
-		FileAr->Serialize( ReadBuffer, PacketBytes );
+		Archive->Serialize( ReadBuffer, PacketBytes );
 
-		if ( FileAr->IsError() )
+		if ( Archive->IsError() )
 		{
 			UE_LOG( LogDemo, Error, TEXT( "UDemoNetDriver::ReadDemoFrame: Failed to read demo file packet" ) );
 			StopDemo();
@@ -891,7 +1051,7 @@ bool UDemoNetDriver::ReadDemoFrame()
 #if DEMO_CHECKSUMS == 1
 		{
 			uint32 ServerChecksum = 0;
-			*FileAr << ServerChecksum;
+			*Archive << ServerChecksum;
 
 			const uint32 Checksum = FCrc::MemCrc32( ReadBuffer, PacketBytes, 0 );
 
@@ -970,7 +1130,7 @@ void UDemoNetDriver::TickDemoPlayback( float DeltaSeconds )
 	}
 
 	// Read demo frames until we are caught up
-	while ( ReadDemoFrame() )
+	while ( ConditionallyReadDemoFrame() )
 	{
 		DemoFrameNum++;
 	}
@@ -1028,6 +1188,83 @@ void UDemoNetDriver::ReplayStreamingReady( bool bSuccess, bool bRecord )
 	}
 }
 
+void UDemoNetDriver::CheckpointReady( bool bSuccess )
+{
+	if ( !bSuccess )
+	{
+		GetWorld()->GetGameInstance()->HandleDemoPlaybackFailure( EDemoPlayFailure::DemoNotFound, FString( EDemoPlayFailure::ToString( EDemoPlayFailure::DemoNotFound ) ) );
+		return;
+	}
+
+	FURL ConnectURL;
+	ConnectURL.Map = DemoFilename;
+
+	for ( int32 i = ServerConnection->OpenChannels.Num() - 1; i >= 0; i-- )
+	{
+		UChannel* OpenChannel = ServerConnection->OpenChannels[i];
+		if ( OpenChannel != NULL )
+		{
+			UActorChannel* ActorChannel = Cast< UActorChannel >( OpenChannel );
+			if ( ActorChannel != NULL && ActorChannel->GetActor() != NULL )
+			{
+				GetWorld()->DestroyActor( ActorChannel->GetActor(), true );
+			}
+		}
+	}
+
+	ServerConnection->Close();
+	ServerConnection->CleanUp();
+
+	ServerConnection = NewObject<UNetConnection>(GetTransientPackage(), UDemoNetConnection::StaticClass());
+	ServerConnection->InitConnection( this, USOCK_Pending, ConnectURL, 1000000 );
+
+	// Create fake control channel
+	ServerConnection->CreateChannel( CHTYPE_Control, 1 );
+
+	FArchive* CheckpointArchive = ReplayStreamer->GetCheckpointArchive();
+
+	GuidCache->ObjectLookup.Empty();
+	GuidCache->NetGUIDLookup.Empty();
+
+	int32 NumValues = 0;
+	*CheckpointArchive << NumValues;
+
+	for ( int32 i = 0; i < NumValues; i++ )
+	{
+		FNetworkGUID Guid;
+		
+		*CheckpointArchive << Guid;
+		
+		FNetGuidCacheObject CacheObject;
+
+		FString PathName;
+
+		*CheckpointArchive << CacheObject.OuterGUID;
+		*CheckpointArchive << PathName;
+		*CheckpointArchive << CacheObject.PackageGuid;
+
+		CacheObject.PathName = FName( *PathName );
+
+		uint8 Flags = 0;
+		*CheckpointArchive << Flags;
+
+		CacheObject.bNoLoad = ( Flags & ( 1 << 0 ) ) ? true : false;
+		CacheObject.bIgnoreWhenMissing = ( Flags & ( 1 << 1 ) ) ? true : false;		
+
+		GuidCache->ObjectLookup.Add( Guid, CacheObject );
+		//GuidCache->GetObjectFromNetGUID( Guid, false );
+	}
+
+	uint32 SavedAbsTimeMS = 0;
+	*CheckpointArchive << SavedAbsTimeMS;
+
+	DemoCurrentTime = (float)SavedAbsTimeMS / 1000.0f;
+
+	ReadDemoFrame( CheckpointArchive );
+
+	bDemoPlaybackDone = false;
+}
+
 /*-----------------------------------------------------------------------------
 	UDemoNetConnection.
 -----------------------------------------------------------------------------*/
@@ -1059,6 +1296,20 @@ FString UDemoNetConnection::LowLevelGetRemoteAddress( bool bAppendPort )
 
 void UDemoNetConnection::LowLevelSend( void* Data, int32 Count )
 {
+	if ( GetDriver()->bSavingCheckpoint )
+	{
+		FArchive* CheckpointArchive = GetDriver()->ReplayStreamer->GetCheckpointArchive();
+
+		if ( CheckpointArchive == nullptr )
+		{
+			return;
+		}
+	
+		*CheckpointArchive << Count;
+		CheckpointArchive->Serialize( Data, Count );
+		return;
+	}
+
 	if ( Count == 0 )
 	{
 		UE_LOG( LogDemo, Warning, TEXT( "UDemoNetConnection::LowLevelSend: Ignoring empty packet." ) );
diff --git a/Engine/Source/Runtime/NetworkReplayStreaming/HttpNetworkReplayStreaming/Private/HttpNetworkReplayStreaming.cpp b/Engine/Source/Runtime/NetworkReplayStreaming/HttpNetworkReplayStreaming/Private/HttpNetworkReplayStreaming.cpp
--- a/Engine/Source/Runtime/NetworkReplayStreaming/HttpNetworkReplayStreaming/Private/HttpNetworkReplayStreaming.cpp
+++ b/Engine/Source/Runtime/NetworkReplayStreaming/HttpNetworkReplayStreaming/Private/HttpNetworkReplayStreaming.cpp
@@ -91,6 +91,8 @@ int64 HttpStreamFArchive::TotalSize()
 
 void HttpStreamFArchive::Seek( int64 InPos ) 
 {
+	check( InPos < Buffer.Num() );
+
 	Pos = InPos;
 }
 
@@ -108,7 +110,8 @@ FHttpNetworkReplayStreamer::FHttpNetworkReplayStreamer() :
 	bStopStreamingCalled( false ), 
 	bStreamIsLive( false ),
 	NumDownloadChunks( 0 ),
-	DemoTimeInMS( 0 ),
+	TotalDemoTimeInMS( 0 ),
+	FlushCheckpointTime( 0 ),
 	HighPriorityEndTime( 0 ),
 	StreamerLastError( ENetworkReplayError::None )
 {
@@ -142,12 +145,15 @@ void FHttpNetworkReplayStreamer::StartStreaming( const FString& StreamName, bool
 	StartStreamingDelegate = Delegate;
 
 	// Setup the archives
-	StreamArchive.ArIsLoading = !bRecord;
-	StreamArchive.ArIsSaving = !StreamArchive.ArIsLoading;
-	StreamArchive.bAtEndOfReplay = false;
+	StreamArchive.ArIsLoading		= !bRecord;
+	StreamArchive.ArIsSaving		= !StreamArchive.ArIsLoading;
+	StreamArchive.bAtEndOfReplay	= false;
+
+	HeaderArchive.ArIsLoading		= StreamArchive.ArIsLoading;
+	HeaderArchive.ArIsSaving		= StreamArchive.ArIsSaving;
 
-	HeaderArchive.ArIsLoading = !bRecord;
-	HeaderArchive.ArIsSaving = !StreamArchive.ArIsLoading;
+	CheckpointArchive.ArIsLoading	= StreamArchive.ArIsLoading;
+	CheckpointArchive.ArIsSaving	= StreamArchive.ArIsSaving;
 
 	LastChunkTime = FPlatformTime::Seconds();
 
@@ -237,7 +243,7 @@ void FHttpNetworkReplayStreamer::UploadHeader()
 
 	HttpState = EHttptate::UploadingHeader;
 
-	HttpRequest->SetURL( FString::Printf( TEXT( "%supload?Version=%s&Session=%s&NumChunks=%i&Time=%i&Filename=replay.header" ), *ServerURL, *SessionVersion, *SessionName, StreamFileCount, DemoTimeInMS ) );
+	HttpRequest->SetURL( FString::Printf( TEXT( "%supload?Version=%s&Session=%s&NumChunks=%i&Time=%i&Filename=replay.header" ), *ServerURL, *SessionVersion, *SessionName, StreamFileCount, TotalDemoTimeInMS ) );
 	HttpRequest->SetVerb( TEXT( "POST" ) );
 	HttpRequest->SetHeader( TEXT( "Content-Type" ), TEXT( "application/octet-stream" ) );
 	HttpRequest->SetContent( HeaderArchive.Buffer );
@@ -259,6 +265,12 @@ void FHttpNetworkReplayStreamer::FlushStream()
 		return;
 	}
 
+	if ( StreamArchive.Buffer.Num() == 0 )
+	{
+		// Nothing to flush
+		return;
+	}
+
 	if ( IsHttpBusy() )
 	{
 		// If we're currently busy we can't flush right now
@@ -275,7 +287,7 @@ void FHttpNetworkReplayStreamer::FlushStream()
 
 	HttpState = EHttptate::UploadingStream;
 
-	HttpRequest->SetURL( FString::Printf( TEXT( "%supload?Version=%s&Session=%s&NumChunks=%i&Time=%i&Filename=stream.%i" ), *ServerURL, *SessionVersion, *SessionName, StreamFileCount + 1, DemoTimeInMS, StreamFileCount ) );
+	HttpRequest->SetURL( FString::Printf( TEXT( "%supload?Version=%s&Session=%s&NumChunks=%i&Time=%i&Filename=stream.%i" ), *ServerURL, *SessionVersion, *SessionName, StreamFileCount + 1, TotalDemoTimeInMS, StreamFileCount ) );
 	HttpRequest->SetVerb( TEXT( "POST" ) );
 	HttpRequest->SetHeader( TEXT( "Content-Type" ), TEXT( "application/octet-stream" ) );
 	HttpRequest->SetContent( StreamArchive.Buffer );
@@ -290,6 +302,77 @@ void FHttpNetworkReplayStreamer::FlushStream()
 	LastChunkTime = FPlatformTime::Seconds();
 }
 
+void FHttpNetworkReplayStreamer::FlushCheckpoint( const uint32 TimeInMS )
+{
+	check( CheckpointArchive.Buffer.Num() > 0 );
+	FlushCheckpointTime = TimeInMS;
+	
+	// Flush any existing stream, we need checkpoints to line up with the next chunk
+	check( !IsHttpBusy() || StreamArchive.Buffer.Num() == 0 );
+
+	FlushStream();
+}
+
+void FHttpNetworkReplayStreamer::GotoCheckpoint( const uint32 TimeInMS, const FOnCheckpointReadyDelegate& Delegate )
+{
+	if ( IsHttpBusy() )
+	{
+		return;
+	}
+
+	TSharedRef<class IHttpRequest> HttpRequest = FHttpModule::Get().CreateRequest();
+
+	// Download the next stream chunk
+	HttpRequest->SetURL( FString::Printf( TEXT( "%sdownload?Version=%s&Session=%s&Filename=checkpoint.%i" ), *ServerURL, *SessionVersion, *SessionName, 0 ) );
+	HttpRequest->SetVerb( TEXT( "GET" ) );
+
+	HttpRequest->OnProcessRequestComplete().BindRaw( this, &FHttpNetworkReplayStreamer::HttpDownloadCheckpointFinished );
+
+	HttpState = EHttptate::DownloadingCheckpoint;
+
+	HttpRequest->ProcessRequest();
+
+	GotoCheckpointDelegate = Delegate;
+}
+
+void FHttpNetworkReplayStreamer::FlushCheckpointInternal( uint32 TimeInMS )
+{
+	if ( SessionName.IsEmpty() || !CheckpointArchive.ArIsSaving )
+	{
+		// IF there is no active session, or we are not recording, we don't need to flush
+		return;
+	}
+
+	if ( IsHttpBusy() )
+	{
+		// If we're currently busy we can't flush right now
+		return;
+	}
+
+	// Upload any new streamed data to the http server
+	UE_LOG( LogHttpReplay, Log, TEXT( "FHttpNetworkReplayStreamer::FlushCheckpoint. Size: %i, StreamFileCount: %i" ), CheckpointArchive.Buffer.Num(), StreamFileCount );
+
+	// Create the Http request and add to pending request list
+	TSharedRef<class IHttpRequest> HttpRequest = FHttpModule::Get().CreateRequest();
+
+	HttpRequest->OnProcessRequestComplete().BindRaw( this, &FHttpNetworkReplayStreamer::HttpUploadFinished );
+
+	HttpState = EHttptate::UploadingStream;
+
+	// Save the chunk to the checkpoint
+	CheckpointArchive << StreamFileCount;
+
+	HttpRequest->SetURL( FString::Printf( TEXT( "%supload?Version=%s&Session=%s&NumChunks=%i&Time=%i&Filename=checkpoint.%i" ), *ServerURL, *SessionVersion, *SessionName, StreamFileCount, TotalDemoTimeInMS, 0 ) );
+	HttpRequest->SetVerb( TEXT( "POST" ) );
+	HttpRequest->SetHeader( TEXT( "Content-Type" ), TEXT( "application/octet-stream" ) );
+	HttpRequest->SetContent( CheckpointArchive.Buffer );
+
+	CheckpointArchive.Buffer.Empty();
+	CheckpointArchive.Pos = 0;
+
+	HttpRequest->ProcessRequest();
+}
+
 void FHttpNetworkReplayStreamer::DownloadHeader()
 {
 	// Download header first
@@ -407,7 +490,12 @@ FArchive* FHttpNetworkReplayStreamer::GetHeaderArchive()
 
 FArchive* FHttpNetworkReplayStreamer::GetStreamingArchive()
 {
-	return &StreamArchive;//( StreamArchive.IsSaving() || IsDataAvailable() ) ? &StreamArchive : NULL;
+	return &StreamArchive;
+}
+
+FArchive* FHttpNetworkReplayStreamer::GetCheckpointArchive()
+{	
+	return &CheckpointArchive;
 }
 
 FArchive* FHttpNetworkReplayStreamer::GetMetadataArchive()
@@ -417,7 +505,7 @@ FArchive* FHttpNetworkReplayStreamer::GetMetadataArchive()
 
 void FHttpNetworkReplayStreamer::UpdateTotalDemoTime( uint32 TimeInMS )
 {
-	DemoTimeInMS = TimeInMS;
+	TotalDemoTimeInMS = TimeInMS;
 }
 
 bool FHttpNetworkReplayStreamer::IsDataAvailable() const
@@ -572,7 +660,7 @@ void FHttpNetworkReplayStreamer::HttpUploadFinished( FHttpRequestPtr , FHttpResp
 
 			HttpState = EHttptate::StopUploading;
 
-			HttpRequest->SetURL( FString::Printf( TEXT( "%sstopuploading?Version=%s&Session=%s&NumChunks=%i&Time=%i" ), *ServerURL, *SessionVersion, *SessionName, StreamFileCount, DemoTimeInMS ) );
+			HttpRequest->SetURL( FString::Printf( TEXT( "%sstopuploading?Version=%s&Session=%s&NumChunks=%i&Time=%i" ), *ServerURL, *SessionVersion, *SessionName, StreamFileCount, TotalDemoTimeInMS ) );
 			HttpRequest->SetVerb( TEXT( "POST" ) );
 			HttpRequest->SetHeader( TEXT( "Content-Type" ), TEXT( "application/octet-stream" ) );
 
@@ -604,9 +692,9 @@ void FHttpNetworkReplayStreamer::HttpStartDownloadingFinished( FHttpRequestPtr H
 		bStreamIsLive = State == TEXT( "Live" );
 
 		NumDownloadChunks = FCString::Atoi( *NumChunksString );
-		DemoTimeInMS = FCString::Atoi( *DemoTimeString );
+		TotalDemoTimeInMS = FCString::Atoi( *DemoTimeString );
 
-		UE_LOG( LogHttpReplay, Log, TEXT( "FHttpNetworkReplayStreamer::HttpStartDownloadingFinished. Viewer: %s, State: %s, NumChunks: %i, DemoTime: %2.2f" ), *ViewerName, *State, NumDownloadChunks, (float)DemoTimeInMS / 1000 );
+		UE_LOG( LogHttpReplay, Log, TEXT( "FHttpNetworkReplayStreamer::HttpStartDownloadingFinished. Viewer: %s, State: %s, NumChunks: %i, DemoTime: %2.2f" ), *ViewerName, *State, NumDownloadChunks, (float)TotalDemoTimeInMS / 1000 );
 
 		// First, download the header
 		if ( NumDownloadChunks > 0 )
@@ -686,7 +774,7 @@ void FHttpNetworkReplayStreamer::HttpDownloadFinished( FHttpRequestPtr HttpReque
 		bStreamIsLive = State == TEXT( "Live" );
 
 		NumDownloadChunks = FCString::Atoi( *NumChunksString );
-		DemoTimeInMS = FCString::Atoi( *DemoTimeString );
+		TotalDemoTimeInMS = FCString::Atoi( *DemoTimeString );
 
 		if ( HttpResponse->GetContent().Num() > 0 || bStreamIsLive )
 		{
@@ -696,7 +784,7 @@ void FHttpNetworkReplayStreamer::HttpDownloadFinished( FHttpRequestPtr HttpReque
 				StreamFileCount++;
 			}
 
-			UE_LOG( LogHttpReplay, Verbose, TEXT( "FHttpNetworkReplayStreamer::HttpDownloadFinished. State: %s, Progress: %i / %i, DemoTime: %2.2f" ), *State, StreamFileCount, NumDownloadChunks, (float)DemoTimeInMS / 1000 );
+			UE_LOG( LogHttpReplay, Verbose, TEXT( "FHttpNetworkReplayStreamer::HttpDownloadFinished. State: %s, Progress: %i / %i, DemoTime: %2.2f" ), *State, StreamFileCount, NumDownloadChunks, (float)TotalDemoTimeInMS / 1000 );
 		}
 		else
 		{
@@ -713,6 +801,51 @@ void FHttpNetworkReplayStreamer::HttpDownloadFinished( FHttpRequestPtr HttpReque
 	}
 }
 
+void FHttpNetworkReplayStreamer::HttpDownloadCheckpointFinished( FHttpRequestPtr HttpRequest, FHttpResponsePtr HttpResponse, bool bSucceeded )
+{
+	check( HttpState == EHttptate::DownloadingCheckpoint );
+	check( StreamArchive.IsLoading() );
+
+	if ( bSucceeded && HttpResponse->GetResponseCode() == EHttpResponseCodes::Ok )
+	{
+		StreamArchive.Buffer.Empty();
+		StreamArchive.Pos = 0;
+		StreamArchive.bAtEndOfReplay = false;
+
+		CheckpointArchive.Buffer = HttpResponse->GetContent();
+		CheckpointArchive.Pos = 0;
+
+		int32 NumValues = 0;
+		CheckpointArchive << NumValues;
+
+		// Seek to the end, where we added our internal meta data
+		CheckpointArchive.Seek( CheckpointArchive.Buffer.Num() - 4 );
+		
+		// Set the next chunk to be the first chunk that represents the start of the stream after the checkpoint
+		CheckpointArchive << StreamFileCount;
+
+		// Make sure we read to the very end
+		check( CheckpointArchive.Pos == CheckpointArchive.Buffer.Num() );
+
+		// Remove the data we serialized from the end of the checkpoint
+		CheckpointArchive.Buffer.SetNum( CheckpointArchive.Buffer.Num() - 4 );
+
+		CheckpointArchive.Pos = 0;
+
+		GotoCheckpointDelegate.ExecuteIfBound( true );
+	}
+	else
+	{
+		UE_LOG( LogHttpReplay, Error, TEXT( "FHttpNetworkReplayStreamer::HttpDownloadCheckpointFinished. FAILED." ) );
+		SetLastError( ENetworkReplayError::ServiceUnavailable );
+		GotoCheckpointDelegate.ExecuteIfBound( false );
+	}
+
+	HttpState = EHttptate::Idle;
+
+	GotoCheckpointDelegate = FOnCheckpointReadyDelegate();
+}
+
 void FHttpNetworkReplayStreamer::HttpRefreshViewerFinished( FHttpRequestPtr HttpRequest, FHttpResponsePtr HttpResponse, bool bSucceeded )
 {
 	check( HttpState == EHttptate::RefreshingViewer );
@@ -815,6 +948,12 @@ void FHttpNetworkReplayStreamer::Tick( const float DeltaTime )
 	}
 	else if ( StreamerState == EStreamerState::StreamingUp )
 	{
+		if ( FlushCheckpointTime > 0 )
+		{
+			FlushCheckpointInternal( FlushCheckpointTime );
+			FlushCheckpointTime = 0;
+		}
+
 		const double FLUSH_TIME_IN_SECONDS = 10;
 
 		if ( FPlatformTime::Seconds() - LastChunkTime > FLUSH_TIME_IN_SECONDS )
diff --git a/Engine/Source/Runtime/NetworkReplayStreaming/HttpNetworkReplayStreaming/Private/HttpNetworkReplayStreaming.h b/Engine/Source/Runtime/NetworkReplayStreaming/HttpNetworkReplayStreaming/Private/HttpNetworkReplayStreaming.h
--- a/Engine/Source/Runtime/NetworkReplayStreaming/HttpNetworkReplayStreaming/Private/HttpNetworkReplayStreaming.h
+++ b/Engine/Source/Runtime/NetworkReplayStreaming/HttpNetworkReplayStreaming/Private/HttpNetworkReplayStreaming.h
@@ -37,9 +37,12 @@ public:
 	virtual void		StopStreaming() override;
 	virtual FArchive*	GetHeaderArchive() override;
 	virtual FArchive*	GetStreamingArchive() override;
+	virtual FArchive*	GetCheckpointArchive() override;
+	virtual void		FlushCheckpoint( const uint32 TimeInMS ) override;
+	virtual void		GotoCheckpoint( const uint32 TimeInMS, const FOnCheckpointReadyDelegate& Delegate ) override;
 	virtual FArchive*	GetMetadataArchive() override;
 	virtual void		UpdateTotalDemoTime( uint32 TimeInMS ) override;
-	virtual uint32		GetTotalDemoTime() const override { return DemoTimeInMS; }
+	virtual uint32		GetTotalDemoTime() const override { return TotalDemoTimeInMS; }
 	virtual bool		IsDataAvailable() const override;
 	virtual void		SetHighPriorityTimeRange( const uint32 StartTimeInMS, const uint32 EndTimeInMS ) override;
 	virtual bool		IsDataAvailableForTimeRange( const uint32 StartTimeInMS, const uint32 EndTimeInMS ) override;
@@ -55,11 +58,13 @@ public:
 	void DownloadNextChunk();
 	void RefreshViewer();
 	void SetLastError( const ENetworkReplayError::Type InLastError );
+	void FlushCheckpointInternal( uint32 TimeInMS );
 
 	/** Delegates */
 	void HttpStartDownloadingFinished( FHttpRequestPtr HttpRequest, FHttpResponsePtr HttpResponse, bool bSucceeded );
 	void HttpDownloadHeaderFinished( FHttpRequestPtr HttpRequest, FHttpResponsePtr HttpResponse, bool bSucceeded );
 	void HttpDownloadFinished( FHttpRequestPtr HttpRequest, FHttpResponsePtr HttpResponse, bool bSucceeded );
+	void HttpDownloadCheckpointFinished( FHttpRequestPtr HttpRequest, FHttpResponsePtr HttpResponse, bool bSucceeded );
 	void HttpRefreshViewerFinished( FHttpRequestPtr HttpRequest, FHttpResponsePtr HttpResponse, bool bSucceeded );
 	void HttpStartUploadingFinished( FHttpRequestPtr HttpRequest, FHttpResponsePtr HttpResponse, bool bSucceeded );
 	void HttpStopUploadingFinished( FHttpRequestPtr HttpRequest, FHttpResponsePtr HttpResponse, bool bSucceeded );
@@ -84,6 +89,8 @@ public:
 		DownloadingStream,			// We are in the process of downloading the replay stream
 		RefreshingViewer,			// We are refreshing the server to let it know we're still viewing
 		EnumeratingSessions,		// We are in the process of downloading the available sessions
+		UploadingCheckpoint,		// We are uploading a checkpoint
+		DownloadingCheckpoint,		// We are downloading a checkpoint
 	};
 
 	/** EStreamerState - Overall state of the streamer */
@@ -100,6 +107,7 @@ public:
 
 	HttpStreamFArchive		HeaderArchive;			// Archive used to buffer the header stream
 	HttpStreamFArchive		StreamArchive;			// Archive used to buffer the data stream
+	HttpStreamFArchive		CheckpointArchive;		// Archive used to buffer checkpoint data
 	FString					SessionName;			// Name of the session on the http replay server
 	FString					SessionVersion;			// Version of the session
 	FString					ServerURL;				// The address of the server
@@ -111,7 +119,8 @@ public:
 	bool					bStopStreamingCalled;
 	bool					bStreamIsLive;			// If true, we are viewing a live stream
 	int32					NumDownloadChunks;
-	uint32					DemoTimeInMS;
+	uint32					TotalDemoTimeInMS;
+	uint32					FlushCheckpointTime;
 	FString					ViewerName;
 	uint32					HighPriorityEndTime;
 
@@ -119,6 +128,7 @@ public:
 
 	FOnStreamReadyDelegate			StartStreamingDelegate;		// Delegate passed in to StartStreaming
 	FOnEnumerateStreamsComplete		EnumerateStreamsDelegate;
+	FOnCheckpointReadyDelegate		GotoCheckpointDelegate;
 };
 
 class FHttpNetworkReplayStreamingFactory : public INetworkReplayStreamingFactory, public FTickableGameObject
diff --git a/Engine/Source/Runtime/NetworkReplayStreaming/NetworkReplayStreaming/Public/NetworkReplayStreaming.h b/Engine/Source/Runtime/NetworkReplayStreaming/NetworkReplayStreaming/Public/NetworkReplayStreaming.h
--- a/Engine/Source/Runtime/NetworkReplayStreaming/NetworkReplayStreaming/Public/NetworkReplayStreaming.h
+++ b/Engine/Source/Runtime/NetworkReplayStreaming/NetworkReplayStreaming/Public/NetworkReplayStreaming.h
@@ -57,8 +57,20 @@ namespace ENetworkReplayError
 	}
 }
 
-DECLARE_MULTICAST_DELEGATE_TwoParams( FOnStreamReady, bool, bool );
-typedef FOnStreamReady::FDelegate FOnStreamReadyDelegate;
+/**
+ * Delegate called when StartStreaming() completes.
+ *
+ * @param bWasSuccessful Whether streaming was started.
+ * @param bRecord Whether streaming is recording or not (vs playing)
+ */
+DECLARE_DELEGATE_TwoParams( FOnStreamReadyDelegate, bool, bool );
+
+/**
+ * Delegate called when GotoCheckpoint() completes.
+ *
+ * @param bWasSuccessful Whether streaming was started.
+ */
+DECLARE_DELEGATE_OneParam( FOnCheckpointReadyDelegate, bool );
 
 /**
  * Delegate called when DeleteFinishedStream() completes.
@@ -84,6 +96,9 @@ public:
 	virtual void StopStreaming() = 0;
 	virtual FArchive* GetHeaderArchive() = 0;
 	virtual FArchive* GetStreamingArchive() = 0;
+	virtual FArchive* GetCheckpointArchive() = 0;
+	virtual void FlushCheckpoint( const uint32 TimeInMS ) = 0;
+	virtual void GotoCheckpoint( const uint32 TimeInMS, const FOnCheckpointReadyDelegate& Delegate ) = 0;
 	virtual FArchive* GetMetadataArchive() = 0;
 	virtual void UpdateTotalDemoTime( uint32 TimeInMS ) = 0;
 	virtual uint32 GetTotalDemoTime() const = 0;
diff --git a/Engine/Source/Runtime/NetworkReplayStreaming/NullNetworkReplayStreaming/Public/NullNetworkReplayStreaming.h b/Engine/Source/Runtime/NetworkReplayStreaming/NullNetworkReplayStreaming/Public/NullNetworkReplayStreaming.h
--- a/Engine/Source/Runtime/NetworkReplayStreaming/NullNetworkReplayStreaming/Public/NullNetworkReplayStreaming.h
+++ b/Engine/Source/Runtime/NetworkReplayStreaming/NullNetworkReplayStreaming/Public/NullNetworkReplayStreaming.h
@@ -20,6 +20,9 @@ public:
 	virtual void StopStreaming() override;
 	virtual FArchive* GetHeaderArchive() override;
 	virtual FArchive* GetStreamingArchive() override;
+	virtual FArchive* GetCheckpointArchive() override { return NULL; }
+	virtual void FlushCheckpoint( const uint32 TimeInMS ) override { }
+	virtual void GotoCheckpoint( const uint32 TimeInMS, const FOnCheckpointReadyDelegate& Delegate ) override { };
 	virtual FArchive* GetMetadataArchive() override;
 	virtual void UpdateTotalDemoTime( uint32 TimeInMS ) override { }
 	virtual uint32 GetTotalDemoTime() const override { return 0; }
