Consolidate bTearOff handling on clients to channel close, this removes some redundant code.

diff --git a/Engine/Source/Runtime/Engine/Private/DataChannel.cpp b/Engine/Source/Runtime/Engine/Private/DataChannel.cpp
--- a/Engine/Source/Runtime/Engine/Private/DataChannel.cpp
+++ b/Engine/Source/Runtime/Engine/Private/DataChannel.cpp
@@ -1445,6 +1445,9 @@ bool UActorChannel::CleanUp( const bool bForDestroy )
 			{
 				Actor->Role = ROLE_Authority;
 				Actor->SetReplicates(false);
+				bTornOff = true;
+				Actor->TornOff();
+				Actor = NULL;
 			}
 			else if (Dormant)
 			{
@@ -1933,17 +1936,6 @@ void UActorChannel::ProcessBunch( FInBunch & Bunch )
 	{
 		Actor->PostNetInit();
 	}
-
-	// Tear off an actor on the client-side
-	if ( Actor && Actor->bTearOff && Connection->Driver->GetNetMode() == NM_Client )
-	{
-		Actor->Role = ROLE_Authority;
-		Actor->SetReplicates(false);
-		bTornOff = true;
-		Connection->ActorChannels.Remove( Actor );
-		Actor->TornOff();
-		Actor = NULL;
-	}
 }
 
 bool UActorChannel::ReplicateActor()
