Don't find the "currently executing" timer when there is no executing timer and you pass an invalid handle to FindTimer

diff --git a/Engine/Source/Runtime/Engine/Private/TimerManager.cpp b/Engine/Source/Runtime/Engine/Private/TimerManager.cpp
--- a/Engine/Source/Runtime/Engine/Private/TimerManager.cpp
+++ b/Engine/Source/Runtime/Engine/Private/TimerManager.cpp
@@ -80,47 +80,50 @@ FTimerData const* FTimerManager::DEPRECATED_FindTimer(FTimerUnifiedDelegate cons
 
 FTimerData const* FTimerManager::FindTimer(FTimerHandle const& InHandle, int32* OutTimerIndex) const
 {
-	if (CurrentlyExecutingTimer.TimerHandle == InHandle)
+	if (InHandle.IsValid())
 	{
-		// found it currently executing
-		if (OutTimerIndex)
+		if (CurrentlyExecutingTimer.TimerHandle == InHandle)
 		{
-			*OutTimerIndex = -1;
+			// found it currently executing
+			if (OutTimerIndex)
+			{
+				*OutTimerIndex = -1;
+			}
+			return &CurrentlyExecutingTimer;
 		}
-		return &CurrentlyExecutingTimer;
-	}
 
-	int32 ActiveTimerIdx = FindTimerInList(ActiveTimerHeap, InHandle);
-	if (ActiveTimerIdx != INDEX_NONE)
-	{
-		// found it in the active heap
-		if (OutTimerIndex)
+		int32 ActiveTimerIdx = FindTimerInList(ActiveTimerHeap, InHandle);
+		if (ActiveTimerIdx != INDEX_NONE)
 		{
-			*OutTimerIndex = ActiveTimerIdx;
+			// found it in the active heap
+			if (OutTimerIndex)
+			{
+				*OutTimerIndex = ActiveTimerIdx;
+			}
+			return &ActiveTimerHeap[ActiveTimerIdx];
 		}
-		return &ActiveTimerHeap[ActiveTimerIdx];
-	}
 
-	int32 PausedTimerIdx = FindTimerInList(PausedTimerList, InHandle);
-	if (PausedTimerIdx != INDEX_NONE)
-	{
-		// found it in the paused list
-		if (OutTimerIndex)
+		int32 PausedTimerIdx = FindTimerInList(PausedTimerList, InHandle);
+		if (PausedTimerIdx != INDEX_NONE)
 		{
-			*OutTimerIndex = PausedTimerIdx;
+			// found it in the paused list
+			if (OutTimerIndex)
+			{
+				*OutTimerIndex = PausedTimerIdx;
+			}
+			return &PausedTimerList[PausedTimerIdx];
 		}
-		return &PausedTimerList[PausedTimerIdx];
-	}
 
-	int32 PendingTimerIdx = FindTimerInList(PendingTimerList, InHandle);
-	if (PendingTimerIdx != INDEX_NONE)
-	{
-		// found it in the pending list
-		if (OutTimerIndex)
+		int32 PendingTimerIdx = FindTimerInList(PendingTimerList, InHandle);
+		if (PendingTimerIdx != INDEX_NONE)
 		{
-			*OutTimerIndex = PendingTimerIdx;
+			// found it in the pending list
+			if (OutTimerIndex)
+			{
+				*OutTimerIndex = PendingTimerIdx;
+			}
+			return &PendingTimerList[PendingTimerIdx];
 		}
-		return &PendingTimerList[PendingTimerIdx];
 	}
 
 	return nullptr;
