Fix issue where startup actors were getting destroyed when loading replay checkpoints

diff --git a/Engine/Source/Runtime/Engine/Private/DemoNetDriver.cpp b/Engine/Source/Runtime/Engine/Private/DemoNetDriver.cpp
--- a/Engine/Source/Runtime/Engine/Private/DemoNetDriver.cpp
+++ b/Engine/Source/Runtime/Engine/Private/DemoNetDriver.cpp
@@ -1205,7 +1205,7 @@ void UDemoNetDriver::CheckpointReady( bool bSuccess )
 		if ( OpenChannel != NULL )
 		{
 			UActorChannel* ActorChannel = Cast< UActorChannel >( OpenChannel );
-			if ( ActorChannel != NULL && ActorChannel->GetActor() != NULL )
+			if ( ActorChannel != NULL && ActorChannel->GetActor() != NULL && !ActorChannel->GetActor()->IsNetStartupActor() )
 			{
 				GetWorld()->DestroyActor( ActorChannel->GetActor(), true );
 			}
