Expose fixed time step as a general project setting. UE-10104

diff --git a/Engine/Source/Runtime/Engine/Classes/Engine/Engine.h b/Engine/Source/Runtime/Engine/Classes/Engine/Engine.h
--- a/Engine/Source/Runtime/Engine/Classes/Engine/Engine.h
+++ b/Engine/Source/Runtime/Engine/Classes/Engine/Engine.h
@@ -1176,12 +1176,20 @@ public:
 	UPROPERTY(config)
 	uint32 bEnableEditorPSysRealtimeLOD:1;
 
+	/** Whether to use a fixed framerate. */
+	UPROPERTY(config, EditAnywhere, Category = Framerate)
+	uint32 bUseFixedFrameRate : 1;
+	
+	/** The fixed framerate to use. */
+	UPROPERTY(config, EditAnywhere, Category = Framerate, meta=(EditCondition="bUseFixedFrameRate"))
+	float FixedFrameRate;
+
 	/** Whether to enable framerate smoothing.																		*/
-	UPROPERTY(config, EditAnywhere, Category=Framerate)
+	UPROPERTY(config, EditAnywhere, Category=Framerate, meta=(EditCondition="!bUseFixedFrameRate"))
 	uint32 bSmoothFrameRate:1;
 
 	/** Range of framerates in which smoothing will kick in */
-	UPROPERTY(config, EditAnywhere, Category=Framerate, meta=(UIMin=0, UIMax=200))
+	UPROPERTY(config, EditAnywhere, Category=Framerate, meta=(UIMin=0, UIMax=200, EditCondition="!bUseFixedFrameRate"))
 	FFloatRange SmoothedFrameRateRange;
 
 	/** 
@@ -1338,7 +1346,7 @@ public:
 	float DisplayGamma;
 
 	/** Minimum desired framerate setting */
-	UPROPERTY(config, EditAnywhere, Category=Framerate, meta=(UIMin=0, ClampMin=0))
+	UPROPERTY(config, EditAnywhere, Category=Framerate, meta=(UIMin=0, ClampMin=0, EditCondition="!bUseFixedFrameRate"))
 	float MinDesiredFrameRate;
 
 private:
diff --git a/Engine/Source/Runtime/Engine/Private/GameEngine.cpp b/Engine/Source/Runtime/Engine/Private/GameEngine.cpp
--- a/Engine/Source/Runtime/Engine/Private/GameEngine.cpp
+++ b/Engine/Source/Runtime/Engine/Private/GameEngine.cpp
@@ -394,6 +394,9 @@ UEngine::UEngine(const FObjectInitializer& ObjectInitializer)
 	EndStreamingPauseDelegate = NULL;
 
 	bCanBlueprintsTickByDefault = true;
+
+	bUseFixedFrameRate = false;
+	FixedFrameRate = 30.f;
 }
 
 void UGameEngine::Init(IEngineLoop* InEngineLoop)
diff --git a/Engine/Source/Runtime/Engine/Private/UnrealEngine.cpp b/Engine/Source/Runtime/Engine/Private/UnrealEngine.cpp
--- a/Engine/Source/Runtime/Engine/Private/UnrealEngine.cpp
+++ b/Engine/Source/Runtime/Engine/Private/UnrealEngine.cpp
@@ -1001,11 +1001,11 @@ void UEngine::UpdateTimeAndHandleMaxTickRate()
 	FApp::UpdateLastTime();
 
 	// Calculate delta time and update time.
-	if( bUseFixedTimeStep )
+	if( bUseFixedTimeStep || bUseFixedFrameRate )
 	{
 		bTimeWasManipulated = true;
-
-		FApp::SetDeltaTime(FApp::GetFixedDeltaTime());
+		const float FrameRate = bUseFixedTimeStep ? FApp::GetFixedDeltaTime() : (1.f / FixedFrameRate);
+		FApp::SetDeltaTime(FrameRate);
 		LastTime = FApp::GetCurrentTime();
 		FApp::SetCurrentTime(FApp::GetCurrentTime() + FApp::GetDeltaTime());
 	}
