Persist current location/rotation when loading replay checkpoints + destroy all non startup actors when loading checkpoints

diff --git a/Engine/Source/Runtime/Engine/Private/DemoNetDriver.cpp b/Engine/Source/Runtime/Engine/Private/DemoNetDriver.cpp
--- a/Engine/Source/Runtime/Engine/Private/DemoNetDriver.cpp
+++ b/Engine/Source/Runtime/Engine/Private/DemoNetDriver.cpp
@@ -1196,9 +1196,33 @@ void UDemoNetDriver::CheckpointReady( bool bSuccess )
 		return;
 	}
 
+	if ( SpectatorController == NULL )
+	{
+		UE_LOG( LogDemo, Warning, TEXT( "UDemoNetConnection::CheckpointReady: No spectator player controller." ) );
+		return;
+	}
+
+	if ( SpectatorController->GetSpectatorPawn() == NULL )
+	{
+		UE_LOG( LogDemo, Warning, TEXT( "UDemoNetConnection::CheckpointReady: No spectator pawn." ) );
+		return;
+	}
+
+	const FVector OldLocation = SpectatorController->GetSpectatorPawn()->GetActorLocation();
+	const FRotator OldRotation = SpectatorController->GetControlRotation();//GetSpectatorPawn()->GetActorRotation();
+
 	FURL ConnectURL;
 	ConnectURL.Map = DemoFilename;
 
+#if 1
+	for ( FActorIterator It( GetWorld() ); It; ++It )
+	{
+		if ( !It->IsNetStartupActor() )
+		{
+			GetWorld()->DestroyActor( *It, true );
+		}
+	}
+#else
 	for ( int32 i = ServerConnection->OpenChannels.Num() - 1; i >= 0; i-- )
 	{
 		UChannel* OpenChannel = ServerConnection->OpenChannels[i];
@@ -1211,6 +1235,9 @@ void UDemoNetDriver::CheckpointReady( bool bSuccess )
 			}
 		}
 	}
+#endif
+
+	SpectatorController = NULL;
 
 	ServerConnection->Close();
 	ServerConnection->CleanUp();
@@ -1263,6 +1290,16 @@ void UDemoNetDriver::CheckpointReady( bool bSuccess )
 	ReadDemoFrame( CheckpointArchive );
 
 	bDemoPlaybackDone = false;
+
+	// Persist location and rotation from previous spectator
+	if ( SpectatorController != NULL )
+	{
+		SpectatorController->SetInitialLocationAndRotation( OldLocation, OldRotation );
+	}
+	else
+	{
+		UE_LOG( LogDemo, Warning, TEXT( "UDemoNetConnection::CheckpointReady: No spectator controller after checkpoint." ) );
+	}
 }
 
 /*-----------------------------------------------------------------------------
@@ -1382,11 +1419,6 @@ void UDemoNetConnection::HandleClientPlayer( APlayerController* PC, UNetConnecti
 		if ( It->IsA( APlayerStart::StaticClass() ) )
 		{
 			PC->SetInitialLocationAndRotation( It->GetActorLocation(), It->GetActorRotation() );
-
-			if ( PC->GetPawn() )
-			{
-				PC->GetPawn()->TeleportTo( It->GetActorLocation(), It->GetActorRotation(), false, true );
-			}
 			break;
 		}
 	}
