Remove virtual from BlueprintImplementableEvents that have clear native virtual already

diff --git a/Engine/Plugins/Experimental/LiveEditor/Source/LiveEditor/Classes/LiveEditorBlueprint.h b/Engine/Plugins/Experimental/LiveEditor/Source/LiveEditor/Classes/LiveEditorBlueprint.h
--- a/Engine/Plugins/Experimental/LiveEditor/Source/LiveEditor/Classes/LiveEditorBlueprint.h
+++ b/Engine/Plugins/Experimental/LiveEditor/Source/LiveEditor/Classes/LiveEditorBlueprint.h
@@ -12,13 +12,13 @@ class ULiveEditorBlueprint : public UObject
 
 public:
 	UFUNCTION(BlueprintImplementableEvent)
-	virtual void OnInit();
+	void OnInit();
 
 	UFUNCTION(BlueprintImplementableEvent)
-	virtual void OnShutdown();
+	void OnShutdown();
 
 	UFUNCTION(BlueprintImplementableEvent)
-	virtual void Tick( float DeltaTime );
+	void Tick( float DeltaTime );
 
 	UFUNCTION(BlueprintCallable, Category="Development|LiveEditor")
 	void SaveCheckpoint( int32 CheckpointID, const TArray<FLiveEditorCheckpointData> &CheckpointData );
diff --git a/Engine/Source/Developer/FunctionalTesting/Classes/FunctionalTest.h b/Engine/Source/Developer/FunctionalTesting/Classes/FunctionalTest.h
--- a/Engine/Source/Developer/FunctionalTesting/Classes/FunctionalTest.h
+++ b/Engine/Source/Developer/FunctionalTesting/Classes/FunctionalTest.h
@@ -95,7 +95,7 @@ public:
 
 	/** Used by debug drawing to gather actors this test is using and point at them on the level to better understand test's setup */
 	UFUNCTION(BlueprintImplementableEvent, Category = "FunctionalTesting")
-	virtual TArray<AActor*> DebugGatherRelevantActors() const;
+	TArray<AActor*> DebugGatherRelevantActors() const;
 
 	virtual void GatherRelevantActors(TArray<AActor*>& OutActors) const;
 
diff --git a/Engine/Source/Runtime/AIModule/Classes/AIController.h b/Engine/Source/Runtime/AIModule/Classes/AIController.h
--- a/Engine/Source/Runtime/AIModule/Classes/AIController.h
+++ b/Engine/Source/Runtime/AIModule/Classes/AIController.h
@@ -124,7 +124,7 @@ public:
 
 	/** Event called when PossessedPawn is possessed by this controller. */
 	UFUNCTION(BlueprintImplementableEvent)
-	virtual void OnPossess(APawn* PossessedPawn);
+	void OnPossess(APawn* PossessedPawn);
 
 	/** Makes AI go toward specified Goal actor (destination will be continuously updated)
 	 *  @param AcceptanceRadius - finish move if pawn gets close enough
@@ -237,7 +237,7 @@ public:
 
 protected:
 	UFUNCTION(BlueprintImplementableEvent)
-	virtual void OnUsingBlackBoard(UBlackboardComponent* BlackboardComp, UBlackboardData* BlackboardAsset);
+	void OnUsingBlackBoard(UBlackboardComponent* BlackboardComp, UBlackboardData* BlackboardAsset);
 
 	virtual bool InitializeBlackboard(UBlackboardComponent& BlackboardComp, UBlackboardData& BlackboardAsset);
 
diff --git a/Engine/Source/Runtime/CoreUObject/Classes/Object.h b/Engine/Source/Runtime/CoreUObject/Classes/Object.h
--- a/Engine/Source/Runtime/CoreUObject/Classes/Object.h
+++ b/Engine/Source/Runtime/CoreUObject/Classes/Object.h
@@ -896,7 +896,7 @@ public:
 	 * @param	EntryPoint	The entry point to start code execution at.
 	 */
 	UFUNCTION(BlueprintImplementableEvent, meta=(BlueprintInternalUseOnly = "true"))
-	virtual void ExecuteUbergraph(int32 EntryPoint);
+	void ExecuteUbergraph(int32 EntryPoint);
 };
 
 #endif
diff --git a/Engine/Source/Runtime/CoreUObject/Public/UObject/UObject.h b/Engine/Source/Runtime/CoreUObject/Public/UObject/UObject.h
--- a/Engine/Source/Runtime/CoreUObject/Public/UObject/UObject.h
+++ b/Engine/Source/Runtime/CoreUObject/Public/UObject/UObject.h
@@ -1140,7 +1140,7 @@ public:
 	{
 		int32 EntryPoint;
 	};
-	virtual void ExecuteUbergraph(int32 EntryPoint)
+	void ExecuteUbergraph(int32 EntryPoint)
 	{
 		Object_eventExecuteUbergraph_Parms Parms;
 		Parms.EntryPoint=EntryPoint;
diff --git a/Engine/Source/Runtime/Engine/Classes/Camera/PlayerCameraManager.h b/Engine/Source/Runtime/Engine/Classes/Camera/PlayerCameraManager.h
--- a/Engine/Source/Runtime/Engine/Classes/Camera/PlayerCameraManager.h
+++ b/Engine/Source/Runtime/Engine/Classes/Camera/PlayerCameraManager.h
@@ -432,7 +432,7 @@ public:
 	 * final camera POV. 
 	 */
 	UFUNCTION(BlueprintImplementableEvent, BlueprintCosmetic)
-	virtual bool BlueprintUpdateCamera(AActor* CameraTarget, FVector& NewCameraLocation, FRotator& NewCameraRotation, float& NewCameraFOV);
+	bool BlueprintUpdateCamera(AActor* CameraTarget, FVector& NewCameraLocation, FRotator& NewCameraRotation, float& NewCameraFOV);
 
 	/** Returns the PlayerController that owns this camera. */
 	UFUNCTION(BlueprintCallable, Category="Game|Player")
diff --git a/Engine/Source/Runtime/Engine/Classes/Engine/GameInstance.h b/Engine/Source/Runtime/Engine/Classes/Engine/GameInstance.h
--- a/Engine/Source/Runtime/Engine/Classes/Engine/GameInstance.h
+++ b/Engine/Source/Runtime/Engine/Classes/Engine/GameInstance.h
@@ -61,14 +61,14 @@ public:
 
 	/** Opportunity for blueprints to handle the game instance being initialized. */
 	UFUNCTION(BlueprintImplementableEvent, meta=(FriendlyName = "Init"))
-	virtual void ReceiveInit();
+	void ReceiveInit();
 
 	/** virtual function to allow custom GameInstances an opportunity to do cleanup when shutting down */
 	virtual void Shutdown();
 
 	/** Opportunity for blueprints to handle the game instance being shutdown. */
 	UFUNCTION(BlueprintImplementableEvent, meta=(FriendlyName = "Shutdown"))
-	virtual void ReceiveShutdown();
+	void ReceiveShutdown();
 
 	/** Opportunity for blueprints to handle network errors. */
 	UFUNCTION(BlueprintImplementableEvent, meta=(FriendlyName = "NetworkError"))
diff --git a/Engine/Source/Runtime/Engine/Classes/GameFramework/Actor.h b/Engine/Source/Runtime/Engine/Classes/GameFramework/Actor.h
--- a/Engine/Source/Runtime/Engine/Classes/GameFramework/Actor.h
+++ b/Engine/Source/Runtime/Engine/Classes/GameFramework/Actor.h
@@ -993,7 +993,7 @@ public:
 	
 	/** Event when play begins for this actor. */
 	UFUNCTION(BlueprintImplementableEvent, meta=(FriendlyName = "BeginPlay"))
-	virtual void ReceiveBeginPlay();
+	void ReceiveBeginPlay();
 
 	/** Event when play begins for this actor. */
 	virtual void BeginPlay();
@@ -1006,22 +1006,22 @@ public:
 
 	/** Event when this actor takes ANY damage */
 	UFUNCTION(BlueprintImplementableEvent, BlueprintAuthorityOnly, meta=(FriendlyName = "AnyDamage"), Category="Game|Damage")
-	virtual void ReceiveAnyDamage(float Damage, const class UDamageType* DamageType, class AController* InstigatedBy, AActor* DamageCauser);
+	void ReceiveAnyDamage(float Damage, const class UDamageType* DamageType, class AController* InstigatedBy, AActor* DamageCauser);
 	
 	/** 
 	 * Event when this actor takes RADIAL damage 
 	 * @todo Pass it the full array of hits instead of just one?
 	 */
 	UFUNCTION(BlueprintImplementableEvent, BlueprintAuthorityOnly, meta=(FriendlyName = "RadialDamage"), Category="Game|Damage")
-	virtual void ReceiveRadialDamage(float DamageReceived, const class UDamageType* DamageType, FVector Origin, const struct FHitResult& HitInfo, class AController* InstigatedBy, AActor* DamageCauser);
+	void ReceiveRadialDamage(float DamageReceived, const class UDamageType* DamageType, FVector Origin, const struct FHitResult& HitInfo, class AController* InstigatedBy, AActor* DamageCauser);
 
 	/** Event when this actor takes POINT damage */
 	UFUNCTION(BlueprintImplementableEvent, BlueprintAuthorityOnly, meta=(FriendlyName = "PointDamage"), Category="Game|Damage")
-	virtual void ReceivePointDamage(float Damage, const class UDamageType* DamageType, FVector HitLocation, FVector HitNormal, class UPrimitiveComponent* HitComponent, FName BoneName, FVector ShotFromDirection, class AController* InstigatedBy, AActor* DamageCauser);
+	void ReceivePointDamage(float Damage, const class UDamageType* DamageType, FVector HitLocation, FVector HitNormal, class UPrimitiveComponent* HitComponent, FName BoneName, FVector ShotFromDirection, class AController* InstigatedBy, AActor* DamageCauser);
 
 	/** Event called every frame */
 	UFUNCTION(BlueprintImplementableEvent, meta=(FriendlyName = "Tick"))
-	virtual void ReceiveTick(float DeltaSeconds);
+	void ReceiveTick(float DeltaSeconds);
 
 	/** 
 	 *	Event when this actor overlaps another actor, for example a player walking into a trigger.
@@ -1070,10 +1070,6 @@ public:
 	UFUNCTION(BlueprintImplementableEvent, meta=(FriendlyName = "TouchLeave"), Category="Touch Input")
 	virtual void ReceiveActorOnInputTouchLeave(const ETouchIndex::Type FingerIndex);
 
-	/** Event when keys/touches/tilt/etc happen */
-	UFUNCTION(BlueprintImplementableEvent, meta=(DeprecatedFunction))
-	virtual void ReceiveInput(const FString& InputName, float Value, FVector VectorValue, bool bStarted, bool bEnded);
-
 	/** 
 	 * Returns list of actors this actor is overlapping (any component overlapping any component). Does not return itself.
 	 * @param OverlappingActors		[out] Returned list of overlapping actors
@@ -1113,7 +1109,7 @@ public:
 	 * @param	Rotation	The rotation.
 	 */
 	UFUNCTION(BlueprintImplementableEvent, meta=(BlueprintInternalUseOnly = "true", FriendlyName = "Construction Script"))
-	virtual void UserConstructionScript();
+	void UserConstructionScript();
 
 	/**
 	 * Destroy this actor. Returns true if destroyed, false if indestructible.
@@ -1125,7 +1121,7 @@ public:
 	bool Destroy(bool bNetForce = false, bool bShouldModifyLevel = true );
 
 	UFUNCTION(BlueprintImplementableEvent, meta = (Keywords = "delete", FriendlyName = "Destroyed"))
-	virtual void ReceiveDestroyed();
+	void ReceiveDestroyed();
 
 	/** Event triggered when the actor is destroyed. */
 	UPROPERTY(BlueprintAssignable, Category="Game")
@@ -1134,7 +1130,7 @@ public:
 
 	/** Event to notify blueprints this actor is about to be deleted. */
 	UFUNCTION(BlueprintImplementableEvent, meta=(Keywords = "delete", FriendlyName = "End Play"))
-	virtual void ReceiveEndPlay(EEndPlayReason::Type EndPlayReason);
+	void ReceiveEndPlay(EEndPlayReason::Type EndPlayReason);
 
 	/** Event triggered when the actor is being removed from a level. */
 	UPROPERTY(BlueprintAssignable, Category="Game")
@@ -2077,11 +2073,11 @@ public:
 
 	/** Event called when this Actor becomes the view target for the given PlayerController. */
 	UFUNCTION(BlueprintImplementableEvent, meta=(FriendlyName="OnBecomeViewTarget", Keywords="Activate Camera"), Category=Actor)
-	virtual void K2_OnBecomeViewTarget( class APlayerController* PC );
+	void K2_OnBecomeViewTarget( class APlayerController* PC );
 
 	/** Event called when this Actor is no longer the view target for the given PlayerController. */
 	UFUNCTION(BlueprintImplementableEvent, meta=(FriendlyName="OnEndViewTarget", Keywords="Deactivate Camera"), Category=Actor)
-	virtual void K2_OnEndViewTarget( class APlayerController* PC );
+	void K2_OnEndViewTarget( class APlayerController* PC );
 
 	/**
 	 *	Calculate camera view point, when viewing this actor.
diff --git a/Engine/Source/Runtime/Engine/Classes/GameFramework/Character.h b/Engine/Source/Runtime/Engine/Classes/GameFramework/Character.h
--- a/Engine/Source/Runtime/Engine/Classes/GameFramework/Character.h
+++ b/Engine/Source/Runtime/Engine/Classes/GameFramework/Character.h
@@ -476,7 +476,7 @@ public:
 
 	/** Let blueprint know that we were launched */
 	UFUNCTION(BlueprintImplementableEvent)
-	virtual void OnLaunched(FVector LaunchVelocity, bool bXYOverride, bool bZOverride);
+	void OnLaunched(FVector LaunchVelocity, bool bXYOverride, bool bZOverride);
 
 	/** Event fired when the character has just started jumping */
 	UFUNCTION(BlueprintNativeEvent, Category="Pawn|Character")
@@ -512,7 +512,7 @@ public:
 	* @see OnMovementModeChanged()
 	*/
 	UFUNCTION(BlueprintImplementableEvent)
-	virtual void OnLanded(const FHitResult& Hit);
+	void OnLanded(const FHitResult& Hit);
 
 	/**
 	 * Event fired when the Character is walking off a surface and is about to fall because CharacterMovement->CurrentFloor became unwalkable.
@@ -560,7 +560,7 @@ public:
 	 * @param	ScaledHalfHeightAdjust	difference after component scale is taken in to account.
 	 */
 	UFUNCTION(BlueprintImplementableEvent, meta=(FriendlyName = "OnEndCrouch"))
-	virtual void K2_OnEndCrouch(float HalfHeightAdjust, float ScaledHalfHeightAdjust);
+	void K2_OnEndCrouch(float HalfHeightAdjust, float ScaledHalfHeightAdjust);
 
 	/**
 	 * Called when Character crouches. Called on non-owned Characters through bIsCrouched replication.
@@ -575,7 +575,7 @@ public:
 	 * @param	ScaledHalfHeightAdjust	difference after component scale is taken in to account.
 	 */
 	UFUNCTION(BlueprintImplementableEvent, meta=(FriendlyName = "OnStartCrouch"))
-	virtual void K2_OnStartCrouch(float HalfHeightAdjust, float ScaledHalfHeightAdjust);
+	void K2_OnStartCrouch(float HalfHeightAdjust, float ScaledHalfHeightAdjust);
 
 	/**
 	 * Called from CharacterMovementComponent to notify the character that the movement mode has changed.
@@ -594,7 +594,7 @@ public:
 	 * @param	NewCustomMode		New custom mode (applicable if NewMovementMode is Custom)
 	 */
 	UFUNCTION(BlueprintImplementableEvent, meta=(FriendlyName = "OnMovementModeChanged"))
-	virtual void K2_OnMovementModeChanged(EMovementMode PrevMovementMode, EMovementMode NewMovementMode, uint8 PrevCustomMode, uint8 NewCustomMode);
+	void K2_OnMovementModeChanged(EMovementMode PrevMovementMode, EMovementMode NewMovementMode, uint8 PrevCustomMode, uint8 NewCustomMode);
 
 	/** Event for implementing custom character movement mode. Called by CharacterMovement if MovementMode is set to Custom. */
 	UFUNCTION(BlueprintImplementableEvent, meta=(FriendlyName= "UpdateCustomMovement"))
diff --git a/Engine/Source/Runtime/Engine/Classes/GameFramework/Controller.h b/Engine/Source/Runtime/Engine/Classes/GameFramework/Controller.h
--- a/Engine/Source/Runtime/Engine/Classes/GameFramework/Controller.h
+++ b/Engine/Source/Runtime/Engine/Classes/GameFramework/Controller.h
@@ -291,7 +291,7 @@ protected:
 
 	/** Event when this controller instigates ANY damage */
 	UFUNCTION(BlueprintImplementableEvent, BlueprintAuthorityOnly)
-	virtual void ReceiveInstigatedAnyDamage(float Damage, const class UDamageType* DamageType, class AActor* DamagedActor, class AActor* DamageCauser);
+	void ReceiveInstigatedAnyDamage(float Damage, const class UDamageType* DamageType, class AActor* DamagedActor, class AActor* DamageCauser);
 
 public:
 	/** Called when the level this controller is in is unloaded via streaming. */
diff --git a/Engine/Source/Runtime/Engine/Classes/GameFramework/GameMode.h b/Engine/Source/Runtime/Engine/Classes/GameFramework/GameMode.h
--- a/Engine/Source/Runtime/Engine/Classes/GameFramework/GameMode.h
+++ b/Engine/Source/Runtime/Engine/Classes/GameFramework/GameMode.h
@@ -437,7 +437,7 @@ public:
 
 	/** Notification that a player has successfully logged in, and has been given a player controller */
 	UFUNCTION(BlueprintImplementableEvent, Category="Game", meta=(FriendlyName="PostLogin"))
-	virtual void K2_PostLogin( APlayerController* NewPlayer );
+	void K2_PostLogin( APlayerController* NewPlayer );
 
 	/** 
 	 * Spawns a PlayerController at the specified location; split out from Login()/HandleSeamlessTravelPlayer() for easier overriding 
diff --git a/Engine/Source/Runtime/Engine/Classes/GameFramework/HUD.h b/Engine/Source/Runtime/Engine/Classes/GameFramework/HUD.h
--- a/Engine/Source/Runtime/Engine/Classes/GameFramework/HUD.h
+++ b/Engine/Source/Runtime/Engine/Classes/GameFramework/HUD.h
@@ -171,7 +171,7 @@ public:
 
 	/** Hook to allow blueprints to do custom HUD drawing. @see bSuppressNativeHUD to control HUD drawing in base class. */
 	UFUNCTION(BlueprintImplementableEvent, BlueprintCosmetic)
-	virtual void ReceiveDrawHUD(int32 SizeX, int32 SizeY);
+	void ReceiveDrawHUD(int32 SizeX, int32 SizeY);
 
 	/** Called when a hit box is clicked on. Provides the name associated with that box. */
 	UFUNCTION(BlueprintImplementableEvent, BlueprintCosmetic)
diff --git a/Engine/Source/Runtime/GameplayAbilities/Public/Abilities/GameplayAbility.h b/Engine/Source/Runtime/GameplayAbilities/Public/Abilities/GameplayAbility.h
--- a/Engine/Source/Runtime/GameplayAbilities/Public/Abilities/GameplayAbility.h
+++ b/Engine/Source/Runtime/GameplayAbilities/Public/Abilities/GameplayAbility.h
@@ -276,7 +276,7 @@ protected:
 
 	/** Returns true if this ability can be activated right now. Has no side effects */
 	UFUNCTION(BlueprintImplementableEvent, Category = Ability, FriendlyName = "ShouldAbilityRespondToEvent")
-	virtual bool K2_ShouldAbilityRespondToEvent(FGameplayAbilityActorInfo ActorInfo, FGameplayEventData Payload) const;
+	bool K2_ShouldAbilityRespondToEvent(FGameplayAbilityActorInfo ActorInfo, FGameplayEventData Payload) const;
 
 	bool bHasBlueprintShouldAbilityRespondToEvent;
 		
@@ -286,7 +286,7 @@ protected:
 	
 	/** Returns true if this ability can be activated right now. Has no side effects */
 	UFUNCTION(BlueprintImplementableEvent, Category = Ability, FriendlyName="CanActivateAbility")
-	virtual bool K2_CanActivateAbility(FGameplayAbilityActorInfo ActorInfo, FGameplayTagContainer& RelevantTags) const;
+	bool K2_CanActivateAbility(FGameplayAbilityActorInfo ActorInfo, FGameplayTagContainer& RelevantTags) const;
 
 	bool bHasBlueprintCanUse;
 
@@ -306,10 +306,10 @@ protected:
 	 *  
 	 */
 	UFUNCTION(BlueprintImplementableEvent, Category = Ability, FriendlyName = "ActivateAbility")
-	virtual void K2_ActivateAbility();
+	void K2_ActivateAbility();
 
 	UFUNCTION(BlueprintImplementableEvent, Category = Ability, FriendlyName = "ActivateAbilityFromEvent")
-	virtual void K2_ActivateAbilityFromEvent(const FGameplayEventData& EventData);
+	void K2_ActivateAbilityFromEvent(const FGameplayEventData& EventData);
 
 	virtual void ActivateAbility(const FGameplayAbilitySpecHandle Handle, const FGameplayAbilityActorInfo* ActorInfo, const FGameplayAbilityActivationInfo ActivationInfo, const FGameplayEventData* TriggerEventData);
 
@@ -368,7 +368,7 @@ public:
 	// --------------------------------------
 
 	UFUNCTION(BlueprintImplementableEvent, Category = Ability, FriendlyName="CommitExecute")
-	virtual void K2_CommitExecute();
+	void K2_CommitExecute();
 
 	/** Does the commit atomically (consume resources, do cooldowns, etc) */
 	virtual void CommitExecute(const FGameplayAbilitySpecHandle Handle, const FGameplayAbilityActorInfo* ActorInfo, const FGameplayAbilityActivationInfo ActivationInfo);
@@ -408,7 +408,7 @@ protected:
 
 	/** Kismet event, will be called if an ability ends normally or abnormally */
 	UFUNCTION(BlueprintImplementableEvent, Category = Ability, FriendlyName = "OnEndAbility")
-	virtual void K2_OnEndAbility();
+	void K2_OnEndAbility();
 
 	/** Native function, called if an ability ends normally or abnormally. If bReplicate is set to true, try to replicate the ending to the client/server */
 	virtual void EndAbility(const FGameplayAbilitySpecHandle Handle, const FGameplayAbilityActorInfo* ActorInfo, const FGameplayAbilityActivationInfo ActivationInfo, bool bReplicateEndAbility);
