#UE4 Fixed an issue where old playerlayers with references to slots that have been removed would get reused after ClearWidgets in SGameLayerManager.

diff --git a/Engine/Source/Runtime/Engine/Private/Slate/SGameLayerManager.cpp b/Engine/Source/Runtime/Engine/Private/Slate/SGameLayerManager.cpp
--- a/Engine/Source/Runtime/Engine/Private/Slate/SGameLayerManager.cpp
+++ b/Engine/Source/Runtime/Engine/Private/Slate/SGameLayerManager.cpp
@@ -111,6 +111,16 @@ void SGameLayerManager::ClearWidgetsForPlayer(ULocalPlayer* Player)
 void SGameLayerManager::ClearWidgets()
 {
 	PlayerCanvas->ClearChildren();
+
+	for(const auto& LayerIt : PlayerLayers)
+	{
+		const TSharedPtr<FPlayerLayer>& Layer = LayerIt.Value;
+		if (Layer.IsValid())
+		{
+			Layer->Slot = nullptr;
+		}
+	}
+	PlayerLayers.Reset();
 }
 
 void SGameLayerManager::Tick(const FGeometry& AllottedGeometry, const double InCurrentTime, const float InDeltaTime)
