Better fix for bTearOff handling. Same change as CL 2483434 (clients handle bTearOff on channel close only), but also the UNetDriver closes channels of bTearOff actors immediately.

diff --git a/Engine/Source/Runtime/Engine/Private/DataChannel.cpp b/Engine/Source/Runtime/Engine/Private/DataChannel.cpp
--- a/Engine/Source/Runtime/Engine/Private/DataChannel.cpp
+++ b/Engine/Source/Runtime/Engine/Private/DataChannel.cpp
@@ -1445,6 +1445,9 @@ bool UActorChannel::CleanUp( const bool bForDestroy )
 			{
 				Actor->Role = ROLE_Authority;
 				Actor->SetReplicates(false);
+				bTornOff = true;
+				Actor->TornOff();
+				Actor = NULL;
 			}
 			else if (Dormant)
 			{
@@ -1933,17 +1936,6 @@ void UActorChannel::ProcessBunch( FInBunch & Bunch )
 	{
 		Actor->PostNetInit();
 	}
-
-	// Tear off an actor on the client-side
-	if ( Actor && Actor->bTearOff && Connection->Driver->GetNetMode() == NM_Client )
-	{
-		Actor->Role = ROLE_Authority;
-		Actor->SetReplicates(false);
-		bTornOff = true;
-		Connection->ActorChannels.Remove( Actor );
-		Actor->TornOff();
-		Actor = NULL;
-	}
 }
 
 bool UActorChannel::ReplicateActor()
diff --git a/Engine/Source/Runtime/Engine/Private/NetworkDriver.cpp b/Engine/Source/Runtime/Engine/Private/NetworkDriver.cpp
--- a/Engine/Source/Runtime/Engine/Private/NetworkDriver.cpp
+++ b/Engine/Source/Runtime/Engine/Private/NetworkDriver.cpp
@@ -2438,7 +2438,9 @@ int32 UNetDriver::ServerReplicateActors(float DeltaSeconds)
 						}
 						
 						// if the actor is now relevant or was recently relevant
-						if( bIsRelevant || (Channel && Time - Channel->RelevantTime < RelevantTimeout) )
+						const bool bIsRecentlyRelevant = bIsRelevant || (Channel && Time - Channel->RelevantTime < RelevantTimeout);
+
+						if( bIsRecentlyRelevant )
 						{	
 							FinalRelevantCount++;
 
@@ -2508,8 +2510,10 @@ int32 UNetDriver::ServerReplicateActors(float DeltaSeconds)
 								}
 							}
 						}
-						// otherwise close the actor channel if it exists for this connection
-						else if ( Channel != NULL )
+						
+						// If the actor wasn't recently relevant, or if it was torn off,
+						// close the actor channel if it exists for this connection
+						if ((!bIsRecentlyRelevant || Actor->bTearOff) && Channel != NULL)
 						{
 							// Non startup (map) actors have their channels closed immediately, which destroys them.
 							// Startup actors get to keep their channels open.
