#MATERIALS: Fix Get Materials to return proper number when overriden

https://jira.ol.epicgames.net/browse/UE-7635

diff --git a/Engine/Source/Runtime/Engine/Classes/Components/SkinnedMeshComponent.h b/Engine/Source/Runtime/Engine/Classes/Components/SkinnedMeshComponent.h
--- a/Engine/Source/Runtime/Engine/Classes/Components/SkinnedMeshComponent.h
+++ b/Engine/Source/Runtime/Engine/Classes/Components/SkinnedMeshComponent.h
@@ -464,10 +464,6 @@ public:
 	virtual int32 GetNumMaterials() const override;
 	// End UPrimitiveComponent interface
 
-	// Begin MeshComponent interface
-	virtual TArray<class UMaterialInterface*> GetMaterials() const override;
-	// End MeshComponent interface
-
 	/**
 	 *	Sets the value of the bForceWireframe flag and reattaches the component as necessary.
 	 *
diff --git a/Engine/Source/Runtime/Engine/Classes/Components/StaticMeshComponent.h b/Engine/Source/Runtime/Engine/Classes/Components/StaticMeshComponent.h
--- a/Engine/Source/Runtime/Engine/Classes/Components/StaticMeshComponent.h
+++ b/Engine/Source/Runtime/Engine/Classes/Components/StaticMeshComponent.h
@@ -311,10 +311,6 @@ public:
 	virtual bool IsNavigationRelevant() const override;
 	virtual void GetNavigationData(FNavigationRelevantData& Data) const override;
 	// End INavRelevantInterface interface.
-	
-	// Begin UMeshComponent interface
-	virtual TArray<class UMaterialInterface*> GetMaterials() const override;
-	// End UMeshComponent interface
 	/**
 	 *	Returns true if the component uses texture lightmaps
 	 *
diff --git a/Engine/Source/Runtime/Engine/Private/Components/MeshComponent.cpp b/Engine/Source/Runtime/Engine/Private/Components/MeshComponent.cpp
--- a/Engine/Source/Runtime/Engine/Private/Components/MeshComponent.cpp
+++ b/Engine/Source/Runtime/Engine/Private/Components/MeshComponent.cpp
@@ -138,5 +138,18 @@ void UMeshComponent::SetTextureForceResidentFlag( bool bForceMiplevelsToBeReside
 
 TArray<class UMaterialInterface*> UMeshComponent::GetMaterials() const
 {
-	return OverrideMaterials;
+	TArray<class UMaterialInterface*> OutMaterials;
+	int32 TotalNumMaterials = GetNumMaterials();
+	if(TotalNumMaterials > 0)
+	{
+		// make sure to extend it
+		OutMaterials.AddZeroed(TotalNumMaterials);
+
+		for(int32 MaterialIndex=0; MaterialIndex < TotalNumMaterials; ++MaterialIndex)
+		{
+			OutMaterials[MaterialIndex] = GetMaterial(MaterialIndex);
+		}
+	}
+
+	return OutMaterials;
 }
diff --git a/Engine/Source/Runtime/Engine/Private/Components/SkinnedMeshComponent.cpp b/Engine/Source/Runtime/Engine/Private/Components/SkinnedMeshComponent.cpp
--- a/Engine/Source/Runtime/Engine/Private/Components/SkinnedMeshComponent.cpp
+++ b/Engine/Source/Runtime/Engine/Private/Components/SkinnedMeshComponent.cpp
@@ -273,11 +273,15 @@ void USkinnedMeshComponent::CreateRenderState_Concurrent()
 
 		if(MeshObject)
 		{
-			int32 UseLOD = PredictedLODLevel;
+			int32 UseLOD;
 			if(MasterPoseComponent.IsValid())
 			{
 				UseLOD = FMath::Clamp(MasterPoseComponent->PredictedLODLevel, 0, MeshObject->GetSkeletalMeshResource().LODModels.Num()-1);
 			}
+			else
+			{
+				UseLOD = FMath::Clamp(PredictedLODLevel, 0, MeshObject->GetSkeletalMeshResource().LODModels.Num()-1);
+			}
 
 			// We just recreated RenderState, that means it could be new mesh or different mesh
 			// Clear ActiveVertexAnims, so that it doesn't use from previous state
@@ -469,22 +473,7 @@ void USkinnedMeshComponent::UpdateSlaveComponent()
 	MarkRenderDynamicDataDirty();
 }
 
-TArray<class UMaterialInterface*> USkinnedMeshComponent::GetMaterials() const
-{
-	TArray<class UMaterialInterface*> OutMaterials = Super::GetMaterials();
-
-	// if no material is overriden, look for mesh material;
-	if(OutMaterials.Num() == 0 && SkeletalMesh)
-	{
-		for (auto Material : SkeletalMesh->Materials)
-		{
-			OutMaterials.Add(Material.MaterialInterface);
-		}
-	}
-
-	return OutMaterials;
-}
-
+// this has to be skeletalmesh material. You can't have more than what SkeletalMesh materials have
 int32 USkinnedMeshComponent::GetNumMaterials() const
 {
 	if (SkeletalMesh)
diff --git a/Engine/Source/Runtime/Engine/Private/Components/StaticMeshComponent.cpp b/Engine/Source/Runtime/Engine/Private/Components/StaticMeshComponent.cpp
--- a/Engine/Source/Runtime/Engine/Private/Components/StaticMeshComponent.cpp
+++ b/Engine/Source/Runtime/Engine/Private/Components/StaticMeshComponent.cpp
@@ -1467,19 +1467,6 @@ int32 UStaticMeshComponent::GetNumMaterials() const
 	}
 }
 
-TArray<class UMaterialInterface*> UStaticMeshComponent::GetMaterials() const
-{
-	TArray<class UMaterialInterface*> OutMaterials = Super::GetMaterials();
-
-	// if no material is overriden, look for mesh material;
-	if(OutMaterials.Num() == 0 && StaticMesh)
-	{
-		OutMaterials = StaticMesh->Materials;
-	}
-
-	return OutMaterials;
-}
-
 UMaterialInterface* UStaticMeshComponent::GetMaterial(int32 MaterialIndex) const
 {
 	// If we have a base materials array, use that
