#ue4 - UWorld::DestroyActor() and AActor::Destroy() prevent recursion by checking IsPendingKillPending() instead of IsPendingKill().

- Prevents possible infinite recursion during destruction with cycles that try to destroy something marked for destruction earlier in the stack.
- Prevents calling destruction events multiple times for the same objects when there are cycles.
- Return value of both methods has been changed to the value of IsPendingKillPending().

diff --git a/Engine/Source/Runtime/Engine/Classes/Engine/World.h b/Engine/Source/Runtime/Engine/Classes/Engine/World.h
--- a/Engine/Source/Runtime/Engine/Classes/Engine/World.h
+++ b/Engine/Source/Runtime/Engine/Classes/Engine/World.h
@@ -2725,7 +2725,7 @@ public:
 	 * @param	ThisActor				Actor to remove.
 	 * @param	bNetForce				[opt] Ignored unless called during play.  Default is false.
 	 * @param	bShouldModifyLevel		[opt] If true, Modify() the level before removing the actor.  Default is true.
-	 * @return							true if destroy, false if actor couldn't be destroyed.
+	 * @return							true if destroyed or already marked for destruction, false if actor couldn't be destroyed.
 	 */
 	bool DestroyActor( AActor* Actor, bool bNetForce=false, bool bShouldModifyLevel=true );
 
diff --git a/Engine/Source/Runtime/Engine/Classes/GameFramework/Actor.h b/Engine/Source/Runtime/Engine/Classes/GameFramework/Actor.h
--- a/Engine/Source/Runtime/Engine/Classes/GameFramework/Actor.h
+++ b/Engine/Source/Runtime/Engine/Classes/GameFramework/Actor.h
@@ -1112,11 +1112,11 @@ public:
 	void UserConstructionScript();
 
 	/**
-	 * Destroy this actor. Returns true if destroyed, false if indestructible.
+	 * Destroy this actor. Returns true the actor is destroyed or already marked for destruction, false if indestructible.
 	 * Destruction is latent. It occurs at the end of the tick.
 	 * @param	bNetForce				[opt] Ignored unless called during play.  Default is false.
 	 * @param	bShouldModifyLevel		[opt] If true, Modify() the level before removing the actor.  Default is true.	
-	 * returns	the state of the RF_PendingKill flag
+	 * returns	true if destroyed or already marked for destruction, false if indestructible.
 	 */
 	bool Destroy(bool bNetForce = false, bool bShouldModifyLevel = true );
 
diff --git a/Engine/Source/Runtime/Engine/Private/Actor.cpp b/Engine/Source/Runtime/Engine/Private/Actor.cpp
--- a/Engine/Source/Runtime/Engine/Private/Actor.cpp
+++ b/Engine/Source/Runtime/Engine/Private/Actor.cpp
@@ -2828,8 +2828,8 @@ bool AActor::GetActorEnableCollision()
 
 bool AActor::Destroy( bool bNetForce, bool bShouldModifyLevel )
 {
-	// It's already pending kill, no need to beat the corpse
-	if (!IsPendingKill())
+	// It's already pending kill or in DestroyActor(), no need to beat the corpse
+	if (!IsPendingKillPending())
 	{
 		UWorld* World = GetWorld();
 		if (World)
@@ -2841,7 +2841,8 @@ bool AActor::Destroy( bool bNetForce, bool bShouldModifyLevel )
 			UE_LOG(LogSpawn, Warning, TEXT("Destroying %s, which doesn't have a valid world pointer"), *GetPathName());
 		}
 	}
-	return IsPendingKill();
+
+	return IsPendingKillPending();
 }
 
 void AActor::K2_DestroyActor()
diff --git a/Engine/Source/Runtime/Engine/Private/LevelActor.cpp b/Engine/Source/Runtime/Engine/Private/LevelActor.cpp
--- a/Engine/Source/Runtime/Engine/Private/LevelActor.cpp
+++ b/Engine/Source/Runtime/Engine/Private/LevelActor.cpp
@@ -399,6 +399,13 @@ bool UWorld::DestroyActor( AActor* ThisActor, bool bNetForce, bool bShouldModify
 		UE_LOG(LogSpawn, Warning, TEXT("Destroying %s, which doesn't have a valid world pointer"), *ThisActor->GetPathName());
 	}
 
+	// If already on list to be deleted, pretend the call was successful.
+	// We don't want recursive calls to trigger destruction notifications multiple times.
+	if (ThisActor->IsPendingKillPending())
+	{
+		return true;
+	}
+
 	// In-game deletion rules.
 	if( AreActorsInitialized() )
 	{
@@ -409,12 +416,6 @@ bool UWorld::DestroyActor( AActor* ThisActor, bool bNetForce, bool bShouldModify
 			return false;
 		}
 
-		// If already on list to be deleted, pretend the call was successful.
-		if( ThisActor->IsPendingKill() )
-		{
-			return true;
-		}
-
 		// Can't kill if wrong role.
 		if( ThisActor->Role!=ROLE_Authority && !bNetForce && !ThisActor->bNetTemporary )
 		{
@@ -441,6 +442,8 @@ bool UWorld::DestroyActor( AActor* ThisActor, bool bNetForce, bool bShouldModify
 	{
 		ThisActor->Modify();
 	}
+
+	// Prevent recursion
 	ThisActor->bPendingKillPending = true;
 
 	// Notify the texture streaming manager about the destruction of this actor.
@@ -453,11 +456,10 @@ bool UWorld::DestroyActor( AActor* ThisActor, bool bNetForce, bool bShouldModify
 	TArray<AActor*> AttachedActors;
 	ThisActor->GetAttachedActors(AttachedActors);
 
-	TInlineComponentArray<USceneComponent*> SceneComponents;
-	ThisActor->GetComponents(SceneComponents);
-
 	if (AttachedActors.Num() > 0)
 	{
+		TInlineComponentArray<USceneComponent*> SceneComponents;
+		ThisActor->GetComponents(SceneComponents);
 
 		for (TArray< AActor* >::TConstIterator AttachedActorIt(AttachedActors); AttachedActorIt; ++AttachedActorIt)
 		{
@@ -498,26 +500,12 @@ bool UWorld::DestroyActor( AActor* ThisActor, bool bNetForce, bool bShouldModify
 #endif
 	}
 
-	if( ThisActor->IsPendingKill() )
-	{
-		return true;
-	}
-
 	ThisActor->ClearComponentOverlaps();
 
-	if (ThisActor->IsPendingKill())
-	{
-		return true;
-	}
-
 	// If this actor has an owner, notify it that it has lost a child.
 	if( ThisActor->GetOwner() )
 	{
 		ThisActor->SetOwner(NULL);
-		if( ThisActor->IsPendingKill() )
-		{
-			return true;
-		}
 	}
 	// Notify net players that this guy has been destroyed.
 	if( NetDriver )
