Final fix from PVS Studio warnings

diff --git a/Engine/Source/Runtime/Engine/Private/Materials/HLSLMaterialTranslator.h b/Engine/Source/Runtime/Engine/Private/Materials/HLSLMaterialTranslator.h
--- a/Engine/Source/Runtime/Engine/Private/Materials/HLSLMaterialTranslator.h
+++ b/Engine/Source/Runtime/Engine/Private/Materials/HLSLMaterialTranslator.h
@@ -3106,7 +3106,7 @@ protected:
 
 		int32 Domain = Material->GetMaterialDomain();
 		if( (Domain != MD_Surface && Domain != MD_DeferredDecal) && 
-			(SourceCoordinateSpace == TRANSFORMSOURCE_Tangent || SourceCoordinateSpace == TRANSFORMSOURCE_Local || DestCoordType == TRANSFORM_Tangent || DestCoordType == TRANSFORM_Tangent) )
+			(SourceCoordinateSpace == TRANSFORMSOURCE_Tangent || SourceCoordinateSpace == TRANSFORMSOURCE_Local || DestCoordType == TRANSFORM_Tangent || DestCoordType == TRANSFORM_Local))
 		{
 			return Errorf(TEXT("Local and tangent transforms are only supported in the Surface and Deferred Decal material domains!"));
 		}
