Preload device profiles before they are needed
Load only the device profiles for the running platform on game builds.

diff --git a/Engine/Source/Runtime/Engine/Classes/DeviceProfiles/DeviceProfileManager.h b/Engine/Source/Runtime/Engine/Classes/DeviceProfiles/DeviceProfileManager.h
--- a/Engine/Source/Runtime/Engine/Classes/DeviceProfiles/DeviceProfileManager.h
+++ b/Engine/Source/Runtime/Engine/Classes/DeviceProfiles/DeviceProfileManager.h
@@ -28,15 +28,6 @@ public:
 	 */
 	static void InitializeCVarsForActiveDeviceProfile();
 
-	/**
-	 * Create a device profile.
-	 *
-	 * @param ProfileName - The profile name.
-	 *
-	 * @return the created profile.
-	 */
-	UDeviceProfile& CreateProfile( const FString& ProfileName );
-
 	/**
 	 * Create a copy of a device profile from a copy.
 	 *
diff --git a/Engine/Source/Runtime/Engine/Private/DeviceProfiles/DeviceProfileManager.cpp b/Engine/Source/Runtime/Engine/Private/DeviceProfiles/DeviceProfileManager.cpp
--- a/Engine/Source/Runtime/Engine/Private/DeviceProfiles/DeviceProfileManager.cpp
+++ b/Engine/Source/Runtime/Engine/Private/DeviceProfiles/DeviceProfileManager.cpp
@@ -150,49 +150,40 @@ void UDeviceProfileManager::InitializeCVarsForActiveDeviceProfile()
 }
 
 
-UDeviceProfile& UDeviceProfileManager::CreateProfile( const FString& ProfileName )
-{
-	UDeviceProfile* NewProfile = NewObject<UDeviceProfile>();
-	Profiles.Add( NewProfile );
-
-	// Inform the UI that the device list has changed
-	ManagerUpdatedDelegate.Broadcast();
-
-	return *NewProfile;
-}
-
-
-UDeviceProfile& UDeviceProfileManager::CreateProfile( const FString& ProfileName, const FString& ProfileType, const FString& ParentName )
+UDeviceProfile& UDeviceProfileManager::CreateProfile( const FString& ProfileName, const FString& ProfileType, const FString& InSpecifyParentName )
 {
 	UDeviceProfile* DeviceProfile = FindObject<UDeviceProfile>( GetTransientPackage(), *ProfileName );
 	if( DeviceProfile == NULL )
 	{
-		DeviceProfile = NewObject<UDeviceProfile>(GetTransientPackage(), *ProfileName);
-		DeviceProfile->BaseProfileName = ParentName != TEXT("") ? ParentName : DeviceProfile->BaseProfileName;
-		DeviceProfile->DeviceType = ProfileType;
-
-		UDeviceProfile* ObjectTemplate = NULL;
+		// Build Parent objects first. Important for setup
+		FString ParentName = InSpecifyParentName;
+		if (ParentName.Len() == 0)
+		{
+			const FString SectionName = FString::Printf(TEXT("%s %s"), *ProfileName, *UDeviceProfile::StaticClass()->GetName());
+			GConfig->GetString(*SectionName, TEXT("BaseProfileName"), ParentName, GetDeviceProfileIniName());
+		}
 
+		UObject* ParentObject = nullptr;
 		// Recursively build the parent tree
-		if( DeviceProfile->BaseProfileName != TEXT("") )
+		if (ParentName.Len() > 0)
 		{
-			DeviceProfile->Parent = FindObject<UDeviceProfile>( GetTransientPackage(), *DeviceProfile->BaseProfileName );
-			if( DeviceProfile->Parent == NULL )
+			ParentObject = FindObject<UDeviceProfile>(GetTransientPackage(), *ParentName);
+			if (ParentObject == nullptr)
 			{
-				DeviceProfile->Parent = &CreateProfile( DeviceProfile->BaseProfileName, ProfileType );
+				ParentObject = &CreateProfile(ParentName, ProfileType);
 			}
-			ObjectTemplate = CastChecked<UDeviceProfile>(DeviceProfile->Parent);
 		}
 
-		if( ObjectTemplate )
-		{
-			DeviceProfile->BaseProfileName = ObjectTemplate->GetName();
-		}
-		DeviceProfile->DeviceType = ProfileType;
+		// Create the profile after it's parents have been created.
+		DeviceProfile = NewObject<UDeviceProfile>(GetTransientPackage(), *ProfileName);
+		DeviceProfile->DeviceType = DeviceProfile->DeviceType.Len() > 0 ? DeviceProfile->DeviceType : ProfileType;
+		DeviceProfile->BaseProfileName = DeviceProfile->BaseProfileName.Len() > 0 ? DeviceProfile->BaseProfileName : ParentName;
+		DeviceProfile->Parent = ParentObject;
 
+		// Add the new profile to the accessible device profile list
 		Profiles.Add( DeviceProfile );
 
-		// Inform the UI that the device list has changed
+		// Inform any listeners that the device list has changed
 		ManagerUpdatedDelegate.Broadcast(); 
 	}
 
@@ -220,7 +211,7 @@ UDeviceProfile& UDeviceProfileManager::FindProfile( const FString& ProfileName )
 		}
 	}
 
-	return FoundProfile != nullptr ? *FoundProfile : CreateProfile(ProfileName);
+	return FoundProfile != nullptr ? *FoundProfile : CreateProfile(ProfileName, FPlatformProperties::PlatformName());
 }
 
 
@@ -255,24 +246,22 @@ void UDeviceProfileManager::LoadProfiles()
 			}
 		}
 
-		// Register Texture LOD settings with each Target Platform
-		ITargetPlatformManagerModule& TargetPlatformManager = GetTargetPlatformManagerRef();
-		const TArray<ITargetPlatform*>& TargetPlatforms = TargetPlatformManager.GetTargetPlatforms();
-		for (int32 PlatformIndex = 0; PlatformIndex < TargetPlatforms.Num(); ++PlatformIndex)
+		if (!FPlatformProperties::RequiresCookedData())
 		{
-			ITargetPlatform* Platform = TargetPlatforms[PlatformIndex];
-			if (const UTextureLODSettings* TextureLODSettingsObj = (UTextureLODSettings*)&FindProfile(Platform->PlatformName()))
-			{
-				// Set TextureLODSettings
-				Platform->RegisterTextureLODSettings(TextureLODSettingsObj);
-			}
-			else
+			// Register Texture LOD settings with each Target Platform
+			ITargetPlatformManagerModule& TargetPlatformManager = GetTargetPlatformManagerRef();
+			const TArray<ITargetPlatform*>& TargetPlatforms = TargetPlatformManager.GetTargetPlatforms();
+			for (int32 PlatformIndex = 0; PlatformIndex < TargetPlatforms.Num(); ++PlatformIndex)
 			{
-				// No Device Profile found that matches target platform
+				ITargetPlatform* Platform = TargetPlatforms[PlatformIndex];
+				if (const UTextureLODSettings* TextureLODSettingsObj = (UTextureLODSettings*)&FindProfile(Platform->PlatformName()))
+				{
+					// Set TextureLODSettings
+					Platform->RegisterTextureLODSettings(TextureLODSettingsObj);
+				}
 			}
 		}
 
-
 		ManagerUpdatedDelegate.Broadcast();
 	}
 }
