PR #767: Fixed GetSocketTransform using the wrong transform space when a nonexistent socket is specified. (Contributed by [REDACTED_PERSON])
https://github.com/EpicGames/UnrealEngine/pull/767

diff --git a/Engine/Source/Runtime/Engine/Private/Components/StaticMeshComponent.cpp b/Engine/Source/Runtime/Engine/Private/Components/StaticMeshComponent.cpp
--- a/Engine/Source/Runtime/Engine/Private/Components/StaticMeshComponent.cpp
+++ b/Engine/Source/Runtime/Engine/Private/Components/StaticMeshComponent.cpp
@@ -683,7 +683,7 @@ FTransform UStaticMeshComponent::GetSocketTransform(FName InSocketName, ERelativ
 		}
 	}
 
-	return ComponentToWorld;
+	return Super::GetSocketTransform(InSocketName, TransformSpace);
 }
 
 
