PR #763/764: Get Actor Tickable When Paused (Contributed by [REDACTED_PERSON])
https://github.com/EpicGames/UnrealEngine/pull/763
https://github.com/EpicGames/UnrealEngine/pull/764

diff --git a/Engine/Source/Runtime/Engine/Classes/GameFramework/Actor.h b/Engine/Source/Runtime/Engine/Classes/GameFramework/Actor.h
--- a/Engine/Source/Runtime/Engine/Classes/GameFramework/Actor.h
+++ b/Engine/Source/Runtime/Engine/Classes/GameFramework/Actor.h
@@ -955,6 +955,10 @@ public:
 	/** Remove tick dependency on PrerequisiteComponent. */
 	UFUNCTION(BlueprintCallable, Category="Utilities", meta=(Keywords = "dependency"))
 	virtual void RemoveTickPrerequisiteComponent(UActorComponent* PrerequisiteComponent);
+	
+	/** Gets whether this actor can tick when paused. */
+	UFUNCTION(BlueprintCallable, Category="Utilities")
+	bool GetTickableWhenPaused();
 
 	/** Sets whether this actor can tick when paused. */
 	UFUNCTION(BlueprintCallable, Category="Utilities")
diff --git a/Engine/Source/Runtime/Engine/Private/Actor.cpp b/Engine/Source/Runtime/Engine/Private/Actor.cpp
--- a/Engine/Source/Runtime/Engine/Private/Actor.cpp
+++ b/Engine/Source/Runtime/Engine/Private/Actor.cpp
@@ -380,6 +380,11 @@ void AActor::RemoveTickPrerequisiteComponent(UActorComponent* PrerequisiteCompon
 	}
 }
 
+bool AActor::GetTickableWhenPaused()
+{
+	return PrimaryActorTick.bTickEvenWhenPaused;
+}
+
 void AActor::SetTickableWhenPaused(bool bTickableWhenPaused)
 {
 	PrimaryActorTick.bTickEvenWhenPaused = bTickableWhenPaused;
