UE-11329: Expose 'pause' for montages

diff --git a/Engine/Source/Runtime/Engine/Classes/Animation/AnimInstance.h b/Engine/Source/Runtime/Engine/Classes/Animation/AnimInstance.h
--- a/Engine/Source/Runtime/Engine/Classes/Animation/AnimInstance.h
+++ b/Engine/Source/Runtime/Engine/Classes/Animation/AnimInstance.h
@@ -489,6 +489,10 @@ public:
 	UFUNCTION(BlueprintCallable, Category = "Animation")
 	void Montage_Stop(float InBlendOutTime, UAnimMontage * Montage = NULL);
 
+	/** Pauses the animation montage. If reference is NULL, it will stop ALL active montages. */
+	UFUNCTION(BlueprintCallable, Category = "Animation")
+	void Montage_Pause(UAnimMontage * Montage = NULL);
+
 	/** Makes a montage jump to a named section. If Montage reference is NULL, it will do that to all active montages. */
 	UFUNCTION(BlueprintCallable, Category="Animation")
 	void Montage_JumpToSection(FName SectionName, UAnimMontage * Montage = NULL);
diff --git a/Engine/Source/Runtime/Engine/Private/Animation/AnimInstance.cpp b/Engine/Source/Runtime/Engine/Private/Animation/AnimInstance.cpp
--- a/Engine/Source/Runtime/Engine/Private/Animation/AnimInstance.cpp
+++ b/Engine/Source/Runtime/Engine/Private/Animation/AnimInstance.cpp
@@ -2120,6 +2120,30 @@ void UAnimInstance::Montage_Stop(float InBlendOutTime, UAnimMontage * Montage)
 	}
 }
 
+void UAnimInstance::Montage_Pause(UAnimMontage * Montage)
+{
+	if (Montage)
+	{
+		FAnimMontageInstance * MontageInstance = GetActiveInstanceForMontage(*Montage);
+		if (MontageInstance)
+		{
+			MontageInstance->Pause();
+		}
+	}
+	else
+	{
+		// If no Montage reference, do it on all active ones.
+		for (int32 InstanceIndex = 0; InstanceIndex < MontageInstances.Num(); InstanceIndex++)
+		{
+			FAnimMontageInstance * MontageInstance = MontageInstances[InstanceIndex];
+			if (MontageInstance && MontageInstance->IsActive())
+			{
+				MontageInstance->Pause();
+			}
+		}
+	}
+}
+
 void UAnimInstance::Montage_JumpToSection(FName SectionName, UAnimMontage * Montage)
 {
 	if (Montage)
