UE4 - fix for mystery crash from crash reporter. Remove NULL textures from arrays returned by preimitive components. JIRA UE-11810 BUGG 40104

diff --git a/Engine/Source/Runtime/Engine/Classes/Components/PrimitiveComponent.h b/Engine/Source/Runtime/Engine/Classes/Components/PrimitiveComponent.h
--- a/Engine/Source/Runtime/Engine/Classes/Components/PrimitiveComponent.h
+++ b/Engine/Source/Runtime/Engine/Classes/Components/PrimitiveComponent.h
@@ -24,6 +24,13 @@ struct FStreamingTexturePrimitiveInfo
 	UTexture* Texture;
 	FSphere Bounds;
 	float TexelFactor;
+
+	FStreamingTexturePrimitiveInfo()
+		: Texture(nullptr)
+		, Bounds(0)
+		, TexelFactor(1.0f)
+	{
+	}
 };
 
 /** Determines whether a Character can attempt to step up onto a component when they walk in to it. */
@@ -1059,7 +1066,14 @@ public:
 	 * @param OutStreamingTextures - Upon return, contains a list of the streaming textures used by the primitive.
 	 */
 	virtual void GetStreamingTextureInfo(TArray<struct FStreamingTexturePrimitiveInfo>& OutStreamingTextures) const
-	{}
+	{
+	}
+
+	/**
+	 * Call GetStreamingTextureInfo and remove the elements with a NULL texture
+	 * @param OutStreamingTextures - Upon return, contains a list of the non-null streaming textures used by the primitive.
+	 */
+	void GetStreamingTextureInfoWithNULLRemoval(TArray<struct FStreamingTexturePrimitiveInfo>& OutStreamingTextures) const;
 
 	/**
 	 * Determines the DPG the primitive's primary elements are drawn in.
diff --git a/Engine/Source/Runtime/Engine/Private/Components/PrimitiveComponent.cpp b/Engine/Source/Runtime/Engine/Private/Components/PrimitiveComponent.cpp
--- a/Engine/Source/Runtime/Engine/Private/Components/PrimitiveComponent.cpp
+++ b/Engine/Source/Runtime/Engine/Private/Components/PrimitiveComponent.cpp
@@ -164,6 +164,17 @@ bool UPrimitiveComponent::HasStaticLighting() const
 	return ((Mobility == EComponentMobility::Static) || bLightAsIfStatic) && SupportsStaticLighting();
 }
 
+void UPrimitiveComponent::GetStreamingTextureInfoWithNULLRemoval(TArray<struct FStreamingTexturePrimitiveInfo>& OutStreamingTextures) const
+{
+	GetStreamingTextureInfo(OutStreamingTextures);
+	for (int32 Index = 0; Index < OutStreamingTextures.Num(); Index++)
+	{
+		if (!OutStreamingTextures[Index].Texture || !Cast<UTexture2D>(OutStreamingTextures[Index].Texture))
+		{
+			OutStreamingTextures.RemoveAt(Index--);
+		}
+	}
+}
 
 void UPrimitiveComponent::GetUsedTextures(TArray<UTexture*>& OutTextures, EMaterialQualityLevel::Type QualityLevel)
 {
diff --git a/Engine/Source/Runtime/Engine/Private/ContentStreaming.cpp b/Engine/Source/Runtime/Engine/Private/ContentStreaming.cpp
--- a/Engine/Source/Runtime/Engine/Private/ContentStreaming.cpp
+++ b/Engine/Source/Runtime/Engine/Private/ContentStreaming.cpp
@@ -2657,7 +2657,7 @@ bool FStreamingManagerTexture::AddDynamicPrimitive( const UPrimitiveComponent* P
 	if (Primitive && Primitive->IsRegistered())
 	{
 		TArray<FStreamingTexturePrimitiveInfo> TextureInstanceInfos;
-		Primitive->GetStreamingTextureInfo( TextureInstanceInfos );
+		Primitive->GetStreamingTextureInfoWithNULLRemoval( TextureInstanceInfos );
 
 		FSpawnedPrimitiveData* PrimitiveData = NULL;
 		TArray<FSpawnedTextureInstance>* TextureInstances = NULL;
@@ -3032,7 +3032,7 @@ void FStreamingManagerTexture::NotifyTimedPrimitiveAttached( const UPrimitiveCom
 	if ( Primitive  && Primitive->IsRegistered())
 	{
 		TArray<FStreamingTexturePrimitiveInfo> TextureInstanceInfos;
-		Primitive->GetStreamingTextureInfo( TextureInstanceInfos );
+		Primitive->GetStreamingTextureInfoWithNULLRemoval( TextureInstanceInfos );
 
 		FSpawnedPrimitiveData* PrimitiveData = NULL;
 		TArray<FSpawnedTextureInstance>* TextureInstances = NULL;
@@ -3063,7 +3063,7 @@ void FStreamingManagerTexture::NotifyTimedPrimitiveDetached( const UPrimitiveCom
 	if ( Primitive )
 	{
 		TArray<FStreamingTexturePrimitiveInfo> TextureInstanceInfos;
-		Primitive->GetStreamingTextureInfo( TextureInstanceInfos );
+		Primitive->GetStreamingTextureInfoWithNULLRemoval( TextureInstanceInfos );
 
 		FSpawnedPrimitiveData* PrimitiveData = NULL;
 		TArray<FSpawnedTextureInstance>* TextureInstances = NULL;
diff --git a/Engine/Source/Runtime/Engine/Private/Level.cpp b/Engine/Source/Runtime/Engine/Private/Level.cpp
--- a/Engine/Source/Runtime/Engine/Private/Level.cpp
+++ b/Engine/Source/Runtime/Engine/Private/Level.cpp
@@ -1273,7 +1273,7 @@ void ULevel::BuildStreamingData(UTexture2D* UpdateSpecificTextureOnly/*=NULL*/)
 				TArray<FStreamingTexturePrimitiveInfo> PrimitiveStreamingTextures;
 
 				// Ask the primitive to enumerate the streaming textures it uses.
-				Primitive->GetStreamingTextureInfo(PrimitiveStreamingTextures);
+				Primitive->GetStreamingTextureInfoWithNULLRemoval(PrimitiveStreamingTextures);
 
 				for(int32 TextureIndex = 0;TextureIndex < PrimitiveStreamingTextures.Num();TextureIndex++)
 				{
diff --git a/Engine/Source/Runtime/Engine/Private/UnrealEngine.cpp b/Engine/Source/Runtime/Engine/Private/UnrealEngine.cpp
--- a/Engine/Source/Runtime/Engine/Private/UnrealEngine.cpp
+++ b/Engine/Source/Runtime/Engine/Private/UnrealEngine.cpp
@@ -3284,7 +3284,7 @@ bool UEngine::HandleListTexturesCommand( const TCHAR* Cmd, FOutputDevice& Ar )
 		// that GetStreamingTextureInfo doesn't check whether a texture is actually streamable or not
 		// and is also implemented for skeletal meshes and such.
 		TArray<FStreamingTexturePrimitiveInfo> StreamingTextures;
-		PrimitiveComponent->GetStreamingTextureInfo( StreamingTextures );
+		PrimitiveComponent->GetStreamingTextureInfoWithNULLRemoval( StreamingTextures );
 
 		// Increase usage count for all referenced textures
 		for( int32 TextureIndex=0; TextureIndex<StreamingTextures.Num(); TextureIndex++ )
