Converted UTextRenderComponent to use FText so it can respond to live-culture updates

UE-11282 - Convert Text Render Components to use FText

UDN - Text Render Localization

The FString variant of SetText is now deprecated for both C++ and Blueprints.

diff --git a/Engine/Source/Runtime/Engine/Classes/Components/TextRenderComponent.h b/Engine/Source/Runtime/Engine/Classes/Components/TextRenderComponent.h
--- a/Engine/Source/Runtime/Engine/Classes/Components/TextRenderComponent.h
+++ b/Engine/Source/Runtime/Engine/Classes/Components/TextRenderComponent.h
@@ -39,7 +39,7 @@ class ENGINE_API UTextRenderComponent : public UPrimitiveComponent
 
 	/** Text content, can be multi line using <br> as line separator */
 	UPROPERTY(EditAnywhere, BlueprintReadOnly, Category=Text, meta=(MultiLine=true))
-	FString Text;
+	FText Text;
 
 	/** Text material */
 	UPROPERTY(EditAnywhere, BlueprintReadOnly, Category=Text)
@@ -87,10 +87,21 @@ class ENGINE_API UTextRenderComponent : public UPrimitiveComponent
 
 	// -----------------------------
 	
-	/** Change the text value and signal the primitives to be rebuilt */
-	UFUNCTION(BlueprintCallable, Category="Rendering|Components|TextRender")
+	/**
+	 * Change the text value and signal the primitives to be rebuilt 
+	 * The FString variant is deprecated in favor of the FText variant
+	 */
+	DEPRECATED(4.8, "Passing text as FString is deprecated, please use FText instead (likely via a LOCTEXT).")
+	UFUNCTION(BlueprintCallable, Category="Rendering|Components|TextRender", meta=(FriendlyName="Set Text (String)", DeprecatedFunction, DeprecationMessage="Use the SetText function taking an FText instead."))
 	void SetText(const FString& Value);
 
+	/** Change the text value and signal the primitives to be rebuilt */
+	void SetText(const FText& Value);
+
+	/** Change the text value and signal the primitives to be rebuilt */
+	UFUNCTION(BlueprintCallable, Category="Rendering|Components|TextRender", meta=(FriendlyName="Set Text"))
+	void K2_SetText(const FText& Value);
+
 	/** Change the text material and signal the primitives to be rebuilt */
 	UFUNCTION(BlueprintCallable, Category="Rendering|Components|TextRender")
 	void SetTextMaterial(UMaterialInterface* Material);
@@ -145,11 +156,19 @@ class ENGINE_API UTextRenderComponent : public UPrimitiveComponent
 
 	// Begin USceneComponent interface.
 	virtual FBoxSphereBounds CalcBounds(const FTransform& LocalToWorld) const override;
-	// Begin USceneComponent interface.
+	// End USceneComponent interface.
+
+	// Begin UActorComponent interface.
+	virtual void TickComponent(float DeltaTime, enum ELevelTick TickType, FActorComponentTickFunction* ThisTickFunction) override;
+	// End UActorComponent interface.
 
 	// Begin UObject interface.
 	virtual void PostLoad() override;
 	// End UObject interface.
+
+private:
+	/** The state of the text the last time it was updated (used to allow updates when the text is changed due to a culture update) */
+	FTextSnapshot TextLastUpdate;
 };
 
 
diff --git a/Engine/Source/Runtime/Engine/Private/Components/TextRenderComponent.cpp b/Engine/Source/Runtime/Engine/Private/Components/TextRenderComponent.cpp
--- a/Engine/Source/Runtime/Engine/Private/Components/TextRenderComponent.cpp
+++ b/Engine/Source/Runtime/Engine/Private/Components/TextRenderComponent.cpp
@@ -7,6 +7,8 @@
 #include "Engine/TextRenderActor.h"
 #include "LocalVertexFactory.h"
 
+#define LOCTEXT_NAMESPACE "TextRenderComponent"
+
 ATextRenderActor::ATextRenderActor(const FObjectInitializer& ObjectInitializer)
 	: Super(ObjectInitializer)
 {
@@ -385,7 +387,7 @@ private:
 	const FColor TextRenderColor;
 	UMaterialInterface* TextMaterial;
 	UFont* Font;
-	FString Text;
+	FText Text;
 	float XScale;
 	float YScale;
 	float HorizSpacingAdjust;
@@ -587,7 +589,7 @@ bool  FTextRenderSceneProxy::BuildStringMesh( TArray<FDynamicMeshVertex>& OutVer
 
 	const int32 PageIndex = 0;
 
-	FTextIterator It(*Text);
+	FTextIterator It(*Text.ToString());
 	while (It.NextLine())
 	{
 		FVector2D LineSize = ComputeTextSize(It, Font, XScale, YScale, HorizSpacingAdjust);
@@ -703,7 +705,12 @@ UTextRenderComponent::UTextRenderComponent(const FObjectInitializer& ObjectIniti
 	};
 	static FConstructorStatics ConstructorStatics;
 
-	Text = "Text";
+	PrimaryComponentTick.bCanEverTick = true;
+	bTickInEditor = true;
+
+	Text = LOCTEXT("DefaultText", "Text");
+	TextLastUpdate = FTextSnapshot(Text);
+
 	Font = ConstructorStatics.Font.Get();
 	TextMaterial = ConstructorStatics.TextMaterial.Get();
 
@@ -771,7 +778,7 @@ FBoxSphereBounds UTextRenderComponent::CalcBounds(const FTransform& LocalToWorld
 		FVector2D LeftTop(FLT_MAX, FLT_MAX);
 		float FirstLineHeight = -1;
 
-		FTextIterator It(*Text);
+		FTextIterator It(*Text.ToString());
 
 		float AdjustedXScale = WorldSize * XScale * InvDefaultSize;
 		float AdjustedYScale = WorldSize * YScale * InvDefaultSize;
@@ -820,7 +827,7 @@ FMatrix UTextRenderComponent::GetRenderMatrix() const
 		float AdjustedXScale = WorldSize * XScale * InvDefaultSize;
 		float AdjustedYScale = WorldSize * YScale * InvDefaultSize;
 
-		FTextIterator It(*Text);
+		FTextIterator It(*Text.ToString());
 		while (It.NextLine())
 		{
 			FVector2D LineSize = ComputeTextSize(It, Font, AdjustedXScale, AdjustedYScale, HorizSpacingAdjust);
@@ -847,9 +854,20 @@ FMatrix UTextRenderComponent::GetRenderMatrix() const
 }
 
 void UTextRenderComponent::SetText(const FString& Value)
+{
+	K2_SetText(FText::FromString(Value));
+}
+
+void UTextRenderComponent::SetText(const FText& Value)
+{
+	K2_SetText(Value);
+}
+
+void UTextRenderComponent::K2_SetText(const FText& Value)
 {
 	Text = Value;
-	MarkRenderStateDirty();	
+	TextLastUpdate = FTextSnapshot(Text); // update this immediately as we're calling MarkRenderStateDirty ourselves
+	MarkRenderStateDirty();
 }
 
 void UTextRenderComponent::SetTextMaterial(class UMaterialInterface* Value)
@@ -921,13 +939,29 @@ FVector UTextRenderComponent::GetTextWorldSize() const
 	return Bounds.GetBox().GetSize();
 }
 
+void UTextRenderComponent::TickComponent(float DeltaTime, enum ELevelTick TickType, FActorComponentTickFunction* ThisTickFunction)
+{
+	if (!TextLastUpdate.IdenticalTo(Text))
+	{
+		// The pointer used by the bound text has changed, however the text may still be the same - check that now
+		if (!TextLastUpdate.IsDisplayStringEqualTo(Text))
+		{
+			// The source text has changed, so we need to update our render data
+			MarkRenderStateDirty();	
+		}
+
+		// Update this even if the text is lexically identical, as it will update the pointer compared by IdenticalTo for the next Tick
+		TextLastUpdate = FTextSnapshot(Text);
+	}
+}
+
 void UTextRenderComponent::PostLoad()
 {
 	// Try and fix up assets created before the vertical alignment fix was implemented. Because we didn't flag that
 	// fix with its own version, use the version number closest to that CL
 	if (GetLinkerUE4Version() < VER_UE4_PACKAGE_REQUIRES_LOCALIZATION_GATHER_FLAGGING)
 	{
-		float Offset = CalculateVerticalAlignmentOffset(*Text, Font, XScale, YScale, HorizSpacingAdjust, VerticalAlignment);
+		float Offset = CalculateVerticalAlignmentOffset(*Text.ToString(), Font, XScale, YScale, HorizSpacingAdjust, VerticalAlignment);
 		FTransform RelativeTransform = GetRelativeTransform();
 		FTransform CorrectionLeft = FTransform::Identity;
 		FTransform CorrectionRight = FTransform::Identity;
@@ -965,3 +999,5 @@ UTextRenderComponent* ATextRenderActor::GetTextRender() const { return TextRende
 /** Returns SpriteComponent subobject **/
 UBillboardComponent* ATextRenderActor::GetSpriteComponent() const { return SpriteComponent; }
 #endif
+
+#undef LOCTEXT_NAMESPACE
