Added safety checks for collision component's existence in APawn::UpdateNavigationRelevance overrides #UE4

diff --git a/Engine/Source/Runtime/Engine/Private/Character.cpp b/Engine/Source/Runtime/Engine/Private/Character.cpp
--- a/Engine/Source/Runtime/Engine/Private/Character.cpp
+++ b/Engine/Source/Runtime/Engine/Private/Character.cpp
@@ -167,7 +167,10 @@ void ACharacter::GetSimpleCollisionCylinder(float& CollisionRadius, float& Colli
 
 void ACharacter::UpdateNavigationRelevance()
 {
-	CapsuleComponent->SetCanEverAffectNavigation(bCanAffectNavigationGeneration);
+	if (CapsuleComponent)
+	{
+		CapsuleComponent->SetCanEverAffectNavigation(bCanAffectNavigationGeneration);
+	}
 }
 
 void ACharacter::ApplyWorldOffset(const FVector& InOffset, bool bWorldShift)
diff --git a/Engine/Source/Runtime/Engine/Private/DefaultPawn.cpp b/Engine/Source/Runtime/Engine/Private/DefaultPawn.cpp
--- a/Engine/Source/Runtime/Engine/Private/DefaultPawn.cpp
+++ b/Engine/Source/Runtime/Engine/Private/DefaultPawn.cpp
@@ -138,7 +138,10 @@ void ADefaultPawn::SetupPlayerInputComponent(UInputComponent* InputComponent)
 
 void ADefaultPawn::UpdateNavigationRelevance()
 {
-	CollisionComponent->SetCanEverAffectNavigation(bCanAffectNavigationGeneration);
+	if (CollisionComponent)
+	{
+		CollisionComponent->SetCanEverAffectNavigation(bCanAffectNavigationGeneration);
+	}
 }
 
 void ADefaultPawn::MoveRight(float Val)
