Don't set a bTearOff actor's pointer to null in UActorChannel::CleanUp so that the channel is removed from the connection's ActorChannels map in UActorChannel::SetClosingFlag

diff --git a/Engine/Source/Runtime/Engine/Private/DataChannel.cpp b/Engine/Source/Runtime/Engine/Private/DataChannel.cpp
--- a/Engine/Source/Runtime/Engine/Private/DataChannel.cpp
+++ b/Engine/Source/Runtime/Engine/Private/DataChannel.cpp
@@ -1447,7 +1447,6 @@ bool UActorChannel::CleanUp( const bool bForDestroy )
 				Actor->SetReplicates(false);
 				bTornOff = true;
 				Actor->TornOff();
-				Actor = NULL;
 			}
 			else if (Dormant)
 			{
