UE4 - added cvar to collect garbage every frame, for debugging

diff --git a/Engine/Source/Runtime/Engine/Private/LevelTick.cpp b/Engine/Source/Runtime/Engine/Private/LevelTick.cpp
--- a/Engine/Source/Runtime/Engine/Private/LevelTick.cpp
+++ b/Engine/Source/Runtime/Engine/Private/LevelTick.cpp
@@ -701,6 +701,11 @@ static TAutoConsoleVariable<int32> CVarAllowAsyncRenderThreadUpdates(
 	0,
 	TEXT("Used to control async renderthread updates."));
 
+static TAutoConsoleVariable<int32> CVarCollectGarbageEveryFrame(
+	TEXT("CollectGarbageEveryFrame"),
+	0,
+	TEXT("Used to debug garbage collection...Collects garbage every frame if the value is > 0."));
+
 void UWorld::MarkActorComponentForNeededEndOfFrameUpdate(class UActorComponent* Component, bool bForceGameThread)
 {
 	check(!bPostTickComponentUpdate); // can't call this while we are doing the updates
@@ -1314,6 +1319,11 @@ void UWorld::Tick( ELevelTick TickType, float DeltaSeconds )
 		}
 	}
 
+	if (CVarCollectGarbageEveryFrame.GetValueOnGameThread() > 0)
+	{
+		ForceGarbageCollection(true);
+	}
+
 	// players only request from last frame
 	if (bPlayersOnlyPending)
 	{
