Fix crash shutting down GDC demo after recording a movie

diff --git a/Engine/Source/Runtime/Engine/Private/AVIWriter.cpp b/Engine/Source/Runtime/Engine/Private/AVIWriter.cpp
--- a/Engine/Source/Runtime/Engine/Private/AVIWriter.cpp
+++ b/Engine/Source/Runtime/Engine/Private/AVIWriter.cpp
@@ -328,15 +328,16 @@ public:
 	{
 		if (bCapturing)
 		{
-			if (CaptureViewport)
-			{
-				ViewportColorBuffer.Empty();
-			}
+			Control->Stop();
+
+			// Trigger the event to force the capture thread to unblock and quit.
+			GCaptureSyncEvent->Trigger();
+			GCaptureSyncEvent->Wait();
+
 			MovieCaptureIndex = 0;
 			bCapturing = false;
 			bReadyForCapture = false;
 			bMatineeFinished = false;
-			Control->Stop();
 			CaptureViewport = NULL;
 			CaptureSlateRenderer = NULL;
 			FrameNumber = 0;
@@ -349,6 +350,11 @@ public:
 			FWindowsPlatformMisc::CoUninitialize();
 			FPlatformProcess::ReturnSynchEventToPool(GCaptureSyncEvent);
 			GCaptureSyncEvent = nullptr;
+
+			if(CaptureViewport)
+			{
+				ViewportColorBuffer.Empty();
+			}
 		}
 	}
 
