Core - Added proper shutdown for the FScreenSaverInhibitor

diff --git a/Engine/Source/Runtime/Engine/Private/UnrealEngine.cpp b/Engine/Source/Runtime/Engine/Private/UnrealEngine.cpp
--- a/Engine/Source/Runtime/Engine/Private/UnrealEngine.cpp
+++ b/Engine/Source/Runtime/Engine/Private/UnrealEngine.cpp
@@ -543,25 +543,43 @@ public:
  */
 class FScreenSaverInhibitor : public FRunnable
 {
-	// FRunnable interface. Not required to be implemented.
-	bool Init() { return true; }
-	void Stop() {}
-	void Exit() {}
+public:
+	/** Default constructor. */
+	FScreenSaverInhibitor()
+		: bEnabled(true)
+	{}
+
+	protected:
+
+	bool Init() override
+	{ 
+		return true; 
+	}
+
+	void Stop() override
+	{
+		bEnabled = false;
+		FPlatformMisc::MemoryBarrier();
+	}
 
 	/**
 	 * Prevents screensaver from kicking in by calling FPlatformMisc::PreventScreenSaver every 50 seconds.
-	 * 
-	 * @return	never returns
 	 */
-	uint32 Run()
+	uint32 Run() override
 	{
-		while( true )
+		while( bEnabled )
 		{
-			FPlatformProcess::Sleep( 50 );
+			const int32 NUM_SECONDS_TO_SLEEP = 50;
+			for( int32 Sec = 0; Sec < NUM_SECONDS_TO_SLEEP && bEnabled; ++Sec )
+			{
+				FPlatformProcess::Sleep( 1 );
+			}
 			FPlatformMisc::PreventScreenSaver();
 		}
 		return 0;
 	}
+
+	bool bEnabled;
 };
 
 /*-----------------------------------------------------------------------------
@@ -957,7 +975,6 @@ void UEngine::PreExit()
 
 	if (ScreenSaverInhibitor)
 	{
-		ScreenSaverInhibitor->Kill();
 		delete ScreenSaverInhibitor;
 	}
 
