Add ability to destroy bTearOff actors immediately during replay fast-forward. This functionality can be toggled via the CVar demo.FastForwardDestroyTearOffActors.

diff --git a/Engine/Source/Runtime/Engine/Classes/Engine/DemoNetDriver.h b/Engine/Source/Runtime/Engine/Classes/Engine/DemoNetDriver.h
--- a/Engine/Source/Runtime/Engine/Classes/Engine/DemoNetDriver.h
+++ b/Engine/Source/Runtime/Engine/Classes/Engine/DemoNetDriver.h
@@ -62,6 +62,10 @@ class UDemoNetDriver
 	double		LastCheckpointTime;
 
 	void		SaveCheckpoint();
+
+private:
+	bool		bIsFastForwarding;
+
 public:
 
 	// UNetDriver interface.
@@ -77,6 +81,7 @@ public:
 	virtual bool IsAvailable() const override { return true; }
 	ENGINE_API void SkipTime(float InTimeToSkip);
 	bool InitConnectInternal( FString& Error );
+	virtual bool ShouldClientDestroyTearOffActors() const;
 
 public:
 
diff --git a/Engine/Source/Runtime/Engine/Classes/Engine/NetDriver.h b/Engine/Source/Runtime/Engine/Classes/Engine/NetDriver.h
--- a/Engine/Source/Runtime/Engine/Classes/Engine/NetDriver.h
+++ b/Engine/Source/Runtime/Engine/Classes/Engine/NetDriver.h
@@ -531,6 +531,9 @@ public:
 	*/
 	TSharedPtr<FRepChangedPropertyTracker> FindOrCreateRepChangedPropertyTracker(UObject *Obj);
 
+	/** Returns true if the client should destroy immediately any actor that becomes torn-off */
+	virtual bool ShouldClientDestroyTearOffActors() const { return false; }
+
 protected:
 
 	/** Adds (fully initialized, ready to go) client connection to the ClientConnections list + any other game related setup */
diff --git a/Engine/Source/Runtime/Engine/Private/DataChannel.cpp b/Engine/Source/Runtime/Engine/Private/DataChannel.cpp
--- a/Engine/Source/Runtime/Engine/Private/DataChannel.cpp
+++ b/Engine/Source/Runtime/Engine/Private/DataChannel.cpp
@@ -1441,14 +1441,14 @@ bool UActorChannel::CleanUp( const bool bForDestroy )
 		checkSlow(Connection->Driver->IsValidLowLevel());
 		if (Actor != NULL)
 		{
-			if (Actor->bTearOff)
+			if (Actor->bTearOff && !Connection->Driver->ShouldClientDestroyTearOffActors())
 			{
 				Actor->Role = ROLE_Authority;
 				Actor->SetReplicates(false);
 				bTornOff = true;
 				Actor->TornOff();
 			}
-			else if (Dormant)
+			else if (Dormant && !Actor->bTearOff)
 			{
 				Connection->DormantActors.Add(Actor);
 			}
diff --git a/Engine/Source/Runtime/Engine/Private/DemoNetDriver.cpp b/Engine/Source/Runtime/Engine/Private/DemoNetDriver.cpp
--- a/Engine/Source/Runtime/Engine/Private/DemoNetDriver.cpp
+++ b/Engine/Source/Runtime/Engine/Private/DemoNetDriver.cpp
@@ -27,6 +27,7 @@ static TAutoConsoleVariable<float> CVarDemoTimeDilation( TEXT( "demo.TimeDilatio
 static TAutoConsoleVariable<float> CVarDemoSkipTime( TEXT( "demo.SkipTime" ), 0, TEXT( "Skip fixed amount of network replay time (in seconds)" ) );
 static TAutoConsoleVariable<int32> CVarEnableCheckpoints( TEXT( "demo.EnableCheckpoints" ), 0, TEXT( "Whether or not checkpoints save on the server" ) );
 static TAutoConsoleVariable<int32> CVarTestCheckpoint( TEXT( "demo.TestCheckpoint" ), 0, TEXT( "For testing only, jump to a particular checkpoint" ) );
+static TAutoConsoleVariable<int32> CVarDemoFastForwardDestroyTearOffActors( TEXT( "demo.FastForwardDestroyTearOffActors" ), 1, TEXT( "If true, the driver will destroy any torn-off actors immediately while fast-forwarding a replay." ) );
 
 static const int32 MAX_DEMO_READ_WRITE_BUFFER = 1024 * 2;
 
@@ -50,7 +51,8 @@ bool UDemoNetDriver::InitBase( bool bInitAsClient, FNetworkNotify* InNotify, con
 		bIsRecordingDemoFrame	= false;
 		bDemoPlaybackDone		= false;
 		bChannelsArePaused		= false;
-		TimeToSkip = 0.0f;
+		TimeToSkip				= 0.0f;
+		bIsFastForwarding		= false;
 
 		ResetDemoState();
 
@@ -523,6 +525,16 @@ void UDemoNetDriver::ProcessRemoteFunction( class AActor* Actor, class UFunction
 	}
 }
 
+bool UDemoNetDriver::ShouldClientDestroyTearOffActors() const
+{
+	if ( CVarDemoFastForwardDestroyTearOffActors.GetValueOnGameThread() != 0 )
+	{
+		return bIsFastForwarding;
+	}
+
+	return false;
+}
+
 bool UDemoNetDriver::Exec( UWorld* InWorld, const TCHAR* Cmd, FOutputDevice& Ar )
 {
 	return Super::Exec( InWorld, Cmd, Ar);
@@ -1121,6 +1133,8 @@ void UDemoNetDriver::TickDemoPlayback( float DeltaSeconds )
 		DemoCurrentTime += TimeToSkip;
 
 		TimeToSkip = 0.0f;
+
+		bIsFastForwarding = true;
 	}
 
 	DemoCurrentTime += DeltaSeconds;
@@ -1135,6 +1149,8 @@ void UDemoNetDriver::TickDemoPlayback( float DeltaSeconds )
 	{
 		DemoFrameNum++;
 	}
+
+	bIsFastForwarding = false;
 }
 
 void UDemoNetDriver::SpawnDemoRecSpectator( UNetConnection* Connection )
