Replay checkpoint work.

* Fix for first checkpoint not saving
* Allow to play 0th checkpoint with cvar
* Save checkpoint after stream chunk is saved.
* Auto increment checkpoint index
* Only save what was in previous stream to checkpoint
* Property call delegate when streaming starts for recording

diff --git a/Engine/Source/Runtime/Engine/Classes/Engine/DemoNetDriver.h b/Engine/Source/Runtime/Engine/Classes/Engine/DemoNetDriver.h
--- a/Engine/Source/Runtime/Engine/Classes/Engine/DemoNetDriver.h
+++ b/Engine/Source/Runtime/Engine/Classes/Engine/DemoNetDriver.h
@@ -47,7 +47,7 @@ class UDemoNetDriver
 	FString DemoSpectatorClass;
 
 	/** Our network replay streamer */
-	TSharedPtr< class INetworkReplayStreamer > ReplayStreamer;
+	TSharedPtr< class INetworkReplayStreamer >	ReplayStreamer;
 
 	uint32 GetDemoCurrentTimeInMS() { return (uint32)( (double)DemoCurrentTime * 1000 ); }
 
@@ -60,6 +60,7 @@ class UDemoNetDriver
 
 	bool		bSavingCheckpoint;
 	double		LastCheckpointTime;
+	int32		CurrentCheckpointIndex;
 
 	void		SaveCheckpoint();
 
diff --git a/Engine/Source/Runtime/Engine/Private/DemoNetDriver.cpp b/Engine/Source/Runtime/Engine/Private/DemoNetDriver.cpp
--- a/Engine/Source/Runtime/Engine/Private/DemoNetDriver.cpp
+++ b/Engine/Source/Runtime/Engine/Private/DemoNetDriver.cpp
@@ -26,7 +26,7 @@ static TAutoConsoleVariable<float> CVarDemoRecordHz( TEXT( "demo.RecordHz" ), 10
 static TAutoConsoleVariable<float> CVarDemoTimeDilation( TEXT( "demo.TimeDilation" ), -1.0f, TEXT( "Override time dilation during demo playback (-1 = don't override)" ) );
 static TAutoConsoleVariable<float> CVarDemoSkipTime( TEXT( "demo.SkipTime" ), 0, TEXT( "Skip fixed amount of network replay time (in seconds)" ) );
 static TAutoConsoleVariable<int32> CVarEnableCheckpoints( TEXT( "demo.EnableCheckpoints" ), 0, TEXT( "Whether or not checkpoints save on the server" ) );
-static TAutoConsoleVariable<int32> CVarTestCheckpoint( TEXT( "demo.TestCheckpoint" ), 0, TEXT( "For testing only, jump to a particular checkpoint" ) );
+static TAutoConsoleVariable<int32> CVarTestCheckpoint( TEXT( "demo.TestCheckpoint" ), -1, TEXT( "For testing only, jump to a particular checkpoint" ) );
 static TAutoConsoleVariable<int32> CVarDemoFastForwardDestroyTearOffActors( TEXT( "demo.FastForwardDestroyTearOffActors" ), 1, TEXT( "If true, the driver will destroy any torn-off actors immediately while fast-forwarding a replay." ) );
 
 static const int32 MAX_DEMO_READ_WRITE_BUFFER = 1024 * 2;
@@ -53,6 +53,7 @@ bool UDemoNetDriver::InitBase( bool bInitAsClient, FNetworkNotify* InNotify, con
 		bChannelsArePaused		= false;
 		TimeToSkip				= 0.0f;
 		bIsFastForwarding		= false;
+		CurrentCheckpointIndex	= 0;
 
 		ResetDemoState();
 
@@ -496,10 +497,10 @@ void UDemoNetDriver::TickFlush( float DeltaSeconds )
 			}
 		}
 
-		if ( CVarTestCheckpoint.GetValueOnGameThread() > 0 )
+		if ( CVarTestCheckpoint.GetValueOnGameThread() >= 0 )
 		{
 			ReplayStreamer->GotoCheckpoint( CVarTestCheckpoint.GetValueOnGameThread(), FOnCheckpointReadyDelegate::CreateUObject( this, &UDemoNetDriver::CheckpointReady ) );
-			CVarTestCheckpoint.AsVariable()->Set( TEXT( "0" ), ECVF_SetByConsole );
+			CVarTestCheckpoint.AsVariable()->Set( TEXT( "-1" ), ECVF_SetByConsole );
 		}
 
 		if ( bDemoPlaybackDone )
@@ -734,18 +735,16 @@ static void SerializeGuidCache( TSharedPtr< class FNetGUIDCache > GuidCache, FAr
 
 void UDemoNetDriver::SaveCheckpoint()
 {
-	if ( DemoCurrentTime > 30.0f )
-	{
-		return;
-	}
-
 	FArchive* CheckpointArchive = ReplayStreamer->GetCheckpointArchive();
 
 	if ( CheckpointArchive == nullptr )
 	{
+		// This doesn't mean error, it means the streamer isn't ready to save checkpoints
 		return;
 	}
 
+	check( CheckpointArchive->TotalSize() == 0 );
+
 	const double StartCheckpointTime = FPlatformTime::Seconds();
 
 	// First, save the current guid cache
@@ -783,14 +782,22 @@ void UDemoNetDriver::SaveCheckpoint()
 
 	bSavingCheckpoint = true;
 
-	DemoReplicateActor( World->GetWorldSettings(), CheckpointConnection, false );
+	// Replicate *only* the actors that were in the previous frame, we want to be able to re-create up to that point with this single checkpoint
+	// It's important that we don't catch any new actors that the next frame will also catch, that will cause conflict with bOpen (the open will occur twice on the same channel)
+	if ( CheckpointConnection->ActorChannels.Contains( World->GetWorldSettings() ) )
+	{
+		DemoReplicateActor( World->GetWorldSettings(), CheckpointConnection, false );
+	}
 
 	for ( int32 i = 0; i < World->NetworkActors.Num(); i++ )
 	{
 		AActor* Actor = World->NetworkActors[i];
 
-		Actor->PreReplication( *FindOrCreateRepChangedPropertyTracker( Actor ).Get() );
-		DemoReplicateActor( Actor, CheckpointConnection, false );
+		if ( CheckpointConnection->ActorChannels.Contains( Actor ) )
+		{
+			Actor->PreReplication( *FindOrCreateRepChangedPropertyTracker( Actor ).Get() );
+			DemoReplicateActor( Actor, CheckpointConnection, false );
+		}
 	}
 
 	CheckpointConnection->FlushNet();
@@ -812,7 +819,7 @@ void UDemoNetDriver::SaveCheckpoint()
 
 	const int32 TotalSize = CheckpointArchive->TotalSize();
 
-	ReplayStreamer->FlushCheckpoint( 1 );
+	ReplayStreamer->FlushCheckpoint( CurrentCheckpointIndex++ );
 
 	const double EndCheckpointTime = FPlatformTime::Seconds();
 
@@ -851,17 +858,6 @@ void UDemoNetDriver::TickDemoRecord( float DeltaSeconds )
 
 	LastRecordTime = CurrentSeconds;
 
-	if ( CVarEnableCheckpoints.GetValueOnGameThread() == 1 )
-	{
-		const double CHECKPOINT_DELAY = 10.0;
-
-		if ( CurrentSeconds - LastCheckpointTime > CHECKPOINT_DELAY )
-		{
-			SaveCheckpoint();
-			LastCheckpointTime = CurrentSeconds;
-		}
-	}
-
 	// Save out a frame
 	DemoFrameNum++;
 	ReplicationFrame++;
@@ -913,6 +909,18 @@ void UDemoNetDriver::TickDemoRecord( float DeltaSeconds )
 	int32 EndCount = 0;
 
 	*FileAr << EndCount;
+
+	// Save a checkpoint if it's time
+	if ( CVarEnableCheckpoints.GetValueOnGameThread() == 1 )
+	{
+		const double CHECKPOINT_DELAY = 10.0;
+
+		if ( CurrentSeconds - LastCheckpointTime > CHECKPOINT_DELAY )
+		{
+			SaveCheckpoint();
+			LastCheckpointTime = CurrentSeconds;
+		}
+	}
 }
 
 void UDemoNetDriver::PauseChannels( const bool bPause )
@@ -1232,6 +1240,7 @@ void UDemoNetDriver::CheckpointReady( bool bSuccess )
 	ConnectURL.Map = DemoFilename;
 
 #if 1
+	// Destroy all non startup actors. They will get restored with the checkpoint
 	for ( FActorIterator It( GetWorld() ); It; ++It )
 	{
 		if ( !It->IsNetStartupActor() )
diff --git a/Engine/Source/Runtime/NetworkReplayStreaming/HttpNetworkReplayStreaming/Private/HttpNetworkReplayStreaming.cpp b/Engine/Source/Runtime/NetworkReplayStreaming/HttpNetworkReplayStreaming/Private/HttpNetworkReplayStreaming.cpp
--- a/Engine/Source/Runtime/NetworkReplayStreaming/HttpNetworkReplayStreaming/Private/HttpNetworkReplayStreaming.cpp
+++ b/Engine/Source/Runtime/NetworkReplayStreaming/HttpNetworkReplayStreaming/Private/HttpNetworkReplayStreaming.cpp
@@ -270,9 +270,9 @@ void FHttpNetworkReplayStreamer::UploadHeader()
 
 void FHttpNetworkReplayStreamer::FlushStream()
 {
-	if ( SessionName.IsEmpty() || !StreamArchive.ArIsSaving )
+	if ( bNeedToUploadHeader || !StreamArchive.ArIsSaving )
 	{
-		// IF there is no active session, or we are not recording, we don't need to flush
+		// If we haven't uploaded the header, or we are not recording, we don't need to flush
 		return;
 	}
 
@@ -497,6 +497,12 @@ FArchive* FHttpNetworkReplayStreamer::GetStreamingArchive()
 
 FArchive* FHttpNetworkReplayStreamer::GetCheckpointArchive()
 {	
+	if ( bNeedToUploadHeader )
+	{
+		// If we need to upload the header, we're not ready to save checkpoints
+		return NULL;
+	}
+
 	return &CheckpointArchive;
 }
 
@@ -629,18 +635,25 @@ void FHttpNetworkReplayStreamer::HttpHeaderUploadFinished( FHttpRequestPtr HttpR
 	check( InFlightHttpRequest->Request == HttpRequest );
 	check( InFlightHttpRequest->Type == EQueuedHttpRequestType::UploadingHeader );
 	check( StreamerState == EStreamerState::StreamingUp );
+	check( StartStreamingDelegate.IsBound() );
 
 	InFlightHttpRequest = NULL;
 
 	if ( bSucceeded && HttpResponse->GetResponseCode() == EHttpResponseCodes::Ok )
 	{
 		UE_LOG( LogHttpReplay, Log, TEXT( "FHttpNetworkReplayStreamer::HttpHeaderUploadFinished." ) );
+
+		StartStreamingDelegate.ExecuteIfBound( true, true );
 	}
 	else
 	{
 		UE_LOG( LogHttpReplay, Error, TEXT( "FHttpNetworkReplayStreamer::HttpHeaderUploadFinished. FAILED" ) );
+		StartStreamingDelegate.ExecuteIfBound( false, true );
 		SetLastError( ENetworkReplayError::ServiceUnavailable );
 	}
+
+	// Reset delegate
+	StartStreamingDelegate = FOnStreamReadyDelegate();
 }
 
 void FHttpNetworkReplayStreamer::HttpUploadFinished( FHttpRequestPtr HttpRequest, FHttpResponsePtr HttpResponse, bool bSucceeded )
@@ -696,7 +709,7 @@ void FHttpNetworkReplayStreamer::HttpStartDownloadingFinished( FHttpRequestPtr H
 		{
 			UE_LOG( LogHttpReplay, Warning, TEXT( "FHttpNetworkReplayStreamer::HttpStartDownloadingFinished. NO CHUNKS" ) );
 
-			StartStreamingDelegate.ExecuteIfBound( false, StreamArchive.IsSaving() );
+			StartStreamingDelegate.ExecuteIfBound( false, true );
 
 			// Reset delegate
 			StartStreamingDelegate = FOnStreamReadyDelegate();
@@ -708,7 +721,7 @@ void FHttpNetworkReplayStreamer::HttpStartDownloadingFinished( FHttpRequestPtr H
 	{
 		UE_LOG( LogHttpReplay, Error, TEXT( "FHttpNetworkReplayStreamer::HttpStartDownloadingFinished. FAILED" ) );
 
-		StartStreamingDelegate.ExecuteIfBound( false, StreamArchive.IsSaving() );
+		StartStreamingDelegate.ExecuteIfBound( false, true );
 
 		// Reset delegate
 		StartStreamingDelegate = FOnStreamReadyDelegate();
@@ -724,6 +737,7 @@ void FHttpNetworkReplayStreamer::HttpDownloadHeaderFinished( FHttpRequestPtr Htt
 	check( InFlightHttpRequest->Type == EQueuedHttpRequestType::DownloadingHeader );
 	check( StreamerState == EStreamerState::StreamingDown );
 	check( StreamArchive.IsLoading() );
+	check( StartStreamingDelegate.IsBound() );
 
 	InFlightHttpRequest = NULL;
 
@@ -733,7 +747,7 @@ void FHttpNetworkReplayStreamer::HttpDownloadHeaderFinished( FHttpRequestPtr Htt
 
 		UE_LOG( LogHttpReplay, Log, TEXT( "FHttpNetworkReplayStreamer::HttpDownloadHeaderFinished. Size: %i" ), HeaderArchive.Buffer.Num() );
 
-		StartStreamingDelegate.ExecuteIfBound( true, StreamArchive.IsSaving() );
+		StartStreamingDelegate.ExecuteIfBound( true, false );
 	}
 	else
 	{
@@ -741,7 +755,7 @@ void FHttpNetworkReplayStreamer::HttpDownloadHeaderFinished( FHttpRequestPtr Htt
 		UE_LOG( LogHttpReplay, Error, TEXT( "FHttpNetworkReplayStreamer::HttpDownloadHeaderFinished. FAILED." ) );
 
 		StreamArchive.Buffer.Empty();
-		StartStreamingDelegate.ExecuteIfBound( false, StreamArchive.IsSaving() );
+		StartStreamingDelegate.ExecuteIfBound( false, false );
 
 		SetLastError( ENetworkReplayError::ServiceUnavailable );
 	}
