Fixed rescale issue when snap to target

diff --git a/Engine/Source/Runtime/Engine/Private/Components/SceneComponent.cpp b/Engine/Source/Runtime/Engine/Private/Components/SceneComponent.cpp
--- a/Engine/Source/Runtime/Engine/Private/Components/SceneComponent.cpp
+++ b/Engine/Source/Runtime/Engine/Private/Components/SceneComponent.cpp
@@ -1036,9 +1036,8 @@ void USceneComponent::AttachTo(class USceneComponent* Parent, FName InSocketName
 				{
 					// when snap, we'd like to give socket or bone scale only
 					// to do so, get parent and get socket relative to parent
-					FTransform ParentToWorld = AttachParent->GetComponentToWorld();
 					FTransform SocketTransform = AttachParent->GetSocketTransform(AttachSocketName);
-					FTransform RelativeTM = SocketTransform.GetRelativeTransform(ParentToWorld);
+					FTransform RelativeTM = ComponentToWorld.GetRelativeTransform(SocketTransform);
 					RelativeScale3D = RelativeTM.GetScale3D();
 				}
 			}
