#ue4 - Add missing Super::TickComponent() to RotatingMovementComponent. Calls to Super should occur only if ShouldSkipUpdate() returns false.

diff --git a/Engine/Source/Runtime/Engine/Private/Components/ProjectileMovementComponent.cpp b/Engine/Source/Runtime/Engine/Private/Components/ProjectileMovementComponent.cpp
--- a/Engine/Source/Runtime/Engine/Private/Components/ProjectileMovementComponent.cpp
+++ b/Engine/Source/Runtime/Engine/Private/Components/ProjectileMovementComponent.cpp
@@ -91,13 +91,15 @@ void UProjectileMovementComponent::InitializeComponent()
 void UProjectileMovementComponent::TickComponent(float DeltaTime, enum ELevelTick TickType, FActorComponentTickFunction *ThisTickFunction)
 {
 	QUICK_SCOPE_CYCLE_COUNTER( STAT_ProjectileMovementComponent_TickComponent );
-	Super::TickComponent(DeltaTime, TickType, ThisTickFunction);
 
 	// skip if don't want component updated when not rendered or updated component can't move
-	if ( ShouldSkipUpdate(DeltaTime) )
+	if (ShouldSkipUpdate(DeltaTime))
 	{
 		return;
 	}
+
+	Super::TickComponent(DeltaTime, TickType, ThisTickFunction);
+
 	AActor* ActorOwner = UpdatedComponent->GetOwner();
 	if ( !ActorOwner || !CheckStillInWorld() )
 	{
diff --git a/Engine/Source/Runtime/Engine/Private/Components/RotatingMovementComponent.cpp b/Engine/Source/Runtime/Engine/Private/Components/RotatingMovementComponent.cpp
--- a/Engine/Source/Runtime/Engine/Private/Components/RotatingMovementComponent.cpp
+++ b/Engine/Source/Runtime/Engine/Private/Components/RotatingMovementComponent.cpp
@@ -14,11 +14,13 @@ URotatingMovementComponent::URotatingMovementComponent(const FObjectInitializer&
 void URotatingMovementComponent::TickComponent(float DeltaTime, enum ELevelTick TickType, FActorComponentTickFunction *ThisTickFunction)
 {
 	// skip if we don't want component updated when not rendered or if updated component can't move
-	if ( ShouldSkipUpdate(DeltaTime) )
+	if (ShouldSkipUpdate(DeltaTime))
 	{
 		return;
 	}
 
+	Super::TickComponent(DeltaTime, TickType, ThisTickFunction);
+
 	// Compute new rotation
 	const FQuat OldRotation = UpdatedComponent->GetComponentQuat();
 	const FQuat DeltaRotation = (RotationRate * DeltaTime).Quaternion();
diff --git a/Engine/Source/Runtime/Engine/Private/FloatingPawnMovement.cpp b/Engine/Source/Runtime/Engine/Private/FloatingPawnMovement.cpp
--- a/Engine/Source/Runtime/Engine/Private/FloatingPawnMovement.cpp
+++ b/Engine/Source/Runtime/Engine/Private/FloatingPawnMovement.cpp
@@ -18,9 +18,14 @@ UFloatingPawnMovement::UFloatingPawnMovement(const FObjectInitializer& ObjectIni
 
 void UFloatingPawnMovement::TickComponent(float DeltaTime, enum ELevelTick TickType, FActorComponentTickFunction *ThisTickFunction)
 {
+	if (ShouldSkipUpdate(DeltaTime))
+	{
+		return;
+	}
+
 	Super::TickComponent(DeltaTime, TickType, ThisTickFunction);
 
-	if (!PawnOwner || !UpdatedComponent || ShouldSkipUpdate(DeltaTime))
+	if (!PawnOwner || !UpdatedComponent)
 	{
 		return;
 	}
