Replay checkpoint sanity checks

diff --git a/Engine/Source/Runtime/Engine/Private/DataChannel.cpp b/Engine/Source/Runtime/Engine/Private/DataChannel.cpp
--- a/Engine/Source/Runtime/Engine/Private/DataChannel.cpp
+++ b/Engine/Source/Runtime/Engine/Private/DataChannel.cpp
@@ -581,6 +581,7 @@ bool UChannel::ReceivedNextBunch( FInBunch & Bunch, bool & bOutSkipAck )
 			{
 				// If we receive a reliable at this point, this means reliables are out of order, which shouldn't be possible
 				check( !HandleBunch->bReliable );
+				check( !Connection->InternalAck );	// Shouldn't be possible for 100% reliable connections
 
 				// Don't ack this packet (since we won't process all of it)
 				bOutSkipAck = true;
