- Fix supporting event triggered before anim instance gets destroyed
- Potential fix for community crash

https://jira.ol.epicgames.net/browse/UE-12112
https://jira.ol.epicgames.net/browse/UE-9957

diff --git a/Engine/Source/Runtime/Engine/Classes/Animation/AnimInstance.h b/Engine/Source/Runtime/Engine/Classes/Animation/AnimInstance.h
--- a/Engine/Source/Runtime/Engine/Classes/Animation/AnimInstance.h
+++ b/Engine/Source/Runtime/Engine/Classes/Animation/AnimInstance.h
@@ -709,6 +709,7 @@ public:
 
 	void InitializeAnimation();
 	void UpdateAnimation(float DeltaSeconds);
+	void UninitializeAnimation();
 
 	// Native initialization override point
 	virtual void NativeInitializeAnimation();
diff --git a/Engine/Source/Runtime/Engine/Private/Animation/AnimInstance.cpp b/Engine/Source/Runtime/Engine/Private/Animation/AnimInstance.cpp
--- a/Engine/Source/Runtime/Engine/Private/Animation/AnimInstance.cpp
+++ b/Engine/Source/Runtime/Engine/Private/Animation/AnimInstance.cpp
@@ -234,6 +234,8 @@ void UAnimInstance::InitializeAnimation()
 {
 	SCOPE_CYCLE_COUNTER(STAT_AnimInitTime);
 
+	UninitializeAnimation();
+
 	// make sure your skeleton is initialized
 	// you can overwrite different skeleton
 	USkeletalMeshComponent* OwnerComponent = GetSkelMeshComponent();
@@ -308,6 +310,35 @@ void UAnimInstance::InitializeAnimation()
 	BindNativeDelegates();
 }
 
+void UAnimInstance::UninitializeAnimation()
+{
+	StopAllMontages(0.f);
+
+	USkeletalMeshComponent * SkelMeshComp = GetSkelMeshComponent();
+	if (SkelMeshComp)
+	{
+		// Tick currently active AnimNotifyState
+		for(int32 Index=0; Index<ActiveAnimNotifyState.Num(); Index++)
+		{
+			const FAnimNotifyEvent& AnimNotifyEvent = ActiveAnimNotifyState[Index];
+			AnimNotifyEvent.NotifyStateClass->NotifyEnd(SkelMeshComp, Cast<UAnimSequenceBase>(AnimNotifyEvent.NotifyStateClass->GetOuter()));
+		}
+
+		TArray<FName> ParamsToClearCopy = MaterialParamatersToClear;
+		for(int i = 0; i < ParamsToClearCopy.Num(); ++i)
+		{
+			AddCurveValue(ParamsToClearCopy[i], 0.0f, ACF_DrivesMaterial);
+		}
+	}
+
+	ActiveAnimNotifyState.Empty();
+	EventCurves.Empty();
+	MorphTargetCurves.Empty();
+	MaterialParameterCurves.Empty();
+	MaterialParamatersToClear.Empty();
+	AnimNotifies.Empty();
+}
+
 #if WITH_EDITORONLY_DATA
 bool UAnimInstance::UpdateSnapshotAndSkipRemainingUpdate()
 {
diff --git a/Engine/Source/Runtime/Engine/Private/Components/SkeletalMeshComponent.cpp b/Engine/Source/Runtime/Engine/Private/Components/SkeletalMeshComponent.cpp
--- a/Engine/Source/Runtime/Engine/Private/Components/SkeletalMeshComponent.cpp
+++ b/Engine/Source/Runtime/Engine/Private/Components/SkeletalMeshComponent.cpp
@@ -209,6 +209,11 @@ void USkeletalMeshComponent::OnUnregister()
 	ReleaseAllClothingResources();
 #endif// #if WITH_APEX_CLOTHING
 
+	if (AnimScriptInstance)
+	{
+		AnimScriptInstance->UninitializeAnimation();
+	}
+
 	Super::OnUnregister();
 }
 
@@ -2003,16 +2008,16 @@ void USkeletalMeshComponent::UnregisterOnSkeletalMeshPropertyChanged( FDelegateH
 
 void USkeletalMeshComponent::ValidateAnimation()
 {
-	if (SkeletalMesh && SkeletalMesh->Skeleton == NULL)
+	if (SkeletalMesh && SkeletalMesh->Skeleton == nullptr)
 	{
 		UE_LOG(LogAnimation, Warning, TEXT("SkeletalMesh %s has no skeleton. This needs to fixed before an animation can be set"), *SkeletalMesh->GetName());
 		if (AnimationMode == EAnimationMode::AnimationSingleNode)
 		{
-			AnimationData.AnimToPlay = NULL;
+			AnimationData.AnimToPlay = nullptr;
 		}
 		else
 		{
-			AnimBlueprintGeneratedClass = NULL;
+			AnimBlueprintGeneratedClass = nullptr;
 		}
 		return;
 	}
@@ -2021,24 +2026,32 @@ void USkeletalMeshComponent::ValidateAnimation()
 	{
 		if(AnimationData.AnimToPlay && SkeletalMesh && AnimationData.AnimToPlay->GetSkeleton() != SkeletalMesh->Skeleton)
 		{
-			UE_LOG(LogAnimation, Warning, TEXT("Animation %s is incompatible with skeleton %s, removing animation from actor."), *AnimationData.AnimToPlay->GetName(), *SkeletalMesh->Skeleton->GetName());
-			AnimationData.AnimToPlay = NULL;
+			if (SkeletalMesh->Skeleton)
+			{
+				UE_LOG(LogAnimation, Warning, TEXT("Animation %s is incompatible with skeleton %s, removing animation from actor."), *AnimationData.AnimToPlay->GetName(), *SkeletalMesh->Skeleton->GetName());
+			}
+			else
+			{
+				UE_LOG(LogAnimation, Warning, TEXT("Animation %s is incompatible because mesh %s has no skeleton, removing animation from actor."), *AnimationData.AnimToPlay->GetName());
+			}
+
+			AnimationData.AnimToPlay = nullptr;
 		}
 	}
 	else
 	{
 		if(AnimBlueprintGeneratedClass && SkeletalMesh && AnimBlueprintGeneratedClass->TargetSkeleton != SkeletalMesh->Skeleton)
 		{
-			if(SkeletalMesh->Skeleton == nullptr)
+			if(SkeletalMesh->Skeleton)
 			{
-				UE_LOG(LogAnimation, Warning, TEXT("AnimBP %s is incompatible because mesh %s has no skeleton, removing AnimBP from actor."), *AnimBlueprintGeneratedClass->GetName(), *SkeletalMesh->GetName());
+				UE_LOG(LogAnimation, Warning, TEXT("AnimBP %s is incompatible with skeleton %s, removing AnimBP from actor."), *AnimBlueprintGeneratedClass->GetName(), *SkeletalMesh->Skeleton->GetName());
 			}
 			else
 			{
-				UE_LOG(LogAnimation, Warning, TEXT("AnimBP %s is incompatible with skeleton %s, removing AnimBP from actor."), *AnimBlueprintGeneratedClass->GetName(), *SkeletalMesh->Skeleton->GetName());
+				UE_LOG(LogAnimation, Warning, TEXT("AnimBP %s is incompatible because mesh %s has no skeleton, removing AnimBP from actor."), *AnimBlueprintGeneratedClass->GetName(), *SkeletalMesh->GetName());
 			}
 
-			AnimBlueprintGeneratedClass = NULL;
+			AnimBlueprintGeneratedClass = nullptr;
 		}
 	}
 }
