Add ability to skip calling RepNotify functions during replay fast-forward, this can be a significant optimization. Can be toggled via the CVar demo.FastForwardSkipRepNotifies.
Time spent fast-forwarding will print to LogDemo.

diff --git a/Engine/Source/Runtime/Engine/Classes/Engine/DemoNetDriver.h b/Engine/Source/Runtime/Engine/Classes/Engine/DemoNetDriver.h
--- a/Engine/Source/Runtime/Engine/Classes/Engine/DemoNetDriver.h
+++ b/Engine/Source/Runtime/Engine/Classes/Engine/DemoNetDriver.h
@@ -66,6 +66,7 @@ class UDemoNetDriver
 
 private:
 	bool		bIsFastForwarding;
+	double		FastForwardStartSeconds;
 
 public:
 
@@ -83,6 +84,7 @@ public:
 	ENGINE_API void SkipTime(float InTimeToSkip);
 	bool InitConnectInternal( FString& Error );
 	virtual bool ShouldClientDestroyTearOffActors() const override;
+	virtual bool ShouldSkipRepNotifies() const override;
 
 public:
 
diff --git a/Engine/Source/Runtime/Engine/Classes/Engine/NetDriver.h b/Engine/Source/Runtime/Engine/Classes/Engine/NetDriver.h
--- a/Engine/Source/Runtime/Engine/Classes/Engine/NetDriver.h
+++ b/Engine/Source/Runtime/Engine/Classes/Engine/NetDriver.h
@@ -534,6 +534,9 @@ public:
 	/** Returns true if the client should destroy immediately any actor that becomes torn-off */
 	virtual bool ShouldClientDestroyTearOffActors() const { return false; }
 
+	/** Returns whether or not properties that are replicating using this driver should not call RepNotify functions. */
+	virtual bool ShouldSkipRepNotifies() const { return false; }
+
 protected:
 
 	/** Adds (fully initialized, ready to go) client connection to the ClientConnections list + any other game related setup */
diff --git a/Engine/Source/Runtime/Engine/Private/DataReplication.cpp b/Engine/Source/Runtime/Engine/Private/DataReplication.cpp
--- a/Engine/Source/Runtime/Engine/Private/DataReplication.cpp
+++ b/Engine/Source/Runtime/Engine/Private/DataReplication.cpp
@@ -1127,6 +1127,11 @@ void FObjectReplicator::CallRepNotifies()
 		return;
 	}
 
+	if ( Connection != NULL && Connection->Driver != NULL && Connection->Driver->ShouldSkipRepNotifies() )
+	{
+		return;
+	}
+
 	RepLayout->CallRepNotifies( RepState, Object );
 
 	if ( RepNotifies.Num() > 0 )
diff --git a/Engine/Source/Runtime/Engine/Private/DemoNetDriver.cpp b/Engine/Source/Runtime/Engine/Private/DemoNetDriver.cpp
--- a/Engine/Source/Runtime/Engine/Private/DemoNetDriver.cpp
+++ b/Engine/Source/Runtime/Engine/Private/DemoNetDriver.cpp
@@ -19,6 +19,7 @@
 #include "NetworkReplayStreaming.h"
 #include "Net/UnrealNetwork.h"
 #include "Net/NetworkProfiler.h"
+#include "Net/DataReplication.h"
 
 DEFINE_LOG_CATEGORY_STATIC( LogDemo, Log, All );
 
@@ -28,6 +29,7 @@ static TAutoConsoleVariable<float> CVarDemoSkipTime( TEXT( "demo.SkipTime" ), 0,
 static TAutoConsoleVariable<int32> CVarEnableCheckpoints( TEXT( "demo.EnableCheckpoints" ), 0, TEXT( "Whether or not checkpoints save on the server" ) );
 static TAutoConsoleVariable<int32> CVarTestCheckpoint( TEXT( "demo.TestCheckpoint" ), -1, TEXT( "For testing only, jump to a particular checkpoint" ) );
 static TAutoConsoleVariable<int32> CVarDemoFastForwardDestroyTearOffActors( TEXT( "demo.FastForwardDestroyTearOffActors" ), 1, TEXT( "If true, the driver will destroy any torn-off actors immediately while fast-forwarding a replay." ) );
+static TAutoConsoleVariable<int32> CVarDemoFastForwardSkipRepNotifies( TEXT( "demo.FastForwardSkipRepNotifies" ), 1, TEXT( "If true, the driver will optimize fast-forwarding by deferring calls to RepNotify functions until the fast-forward is complete. " ) );
 
 static const int32 MAX_DEMO_READ_WRITE_BUFFER = 1024 * 2;
 
@@ -53,6 +55,7 @@ bool UDemoNetDriver::InitBase( bool bInitAsClient, FNetworkNotify* InNotify, con
 		bChannelsArePaused		= false;
 		TimeToSkip				= 0.0f;
 		bIsFastForwarding		= false;
+		FastForwardStartSeconds	= 0.0;
 		CurrentCheckpointIndex	= 0;
 
 		ResetDemoState();
@@ -536,6 +539,16 @@ bool UDemoNetDriver::ShouldClientDestroyTearOffActors() const
 	return false;
 }
 
+bool UDemoNetDriver::ShouldSkipRepNotifies() const
+{
+	if ( CVarDemoFastForwardSkipRepNotifies.GetValueOnGameThread() != 0 )
+	{
+		return bIsFastForwarding;
+	}
+
+	return false;
+}
+
 bool UDemoNetDriver::Exec( UWorld* InWorld, const TCHAR* Cmd, FOutputDevice& Ar )
 {
 	return Super::Exec( InWorld, Cmd, Ar);
@@ -1143,6 +1156,8 @@ void UDemoNetDriver::TickDemoPlayback( float DeltaSeconds )
 		TimeToSkip = 0.0f;
 
 		bIsFastForwarding = true;
+
+		FastForwardStartSeconds = FPlatformTime::Seconds();
 	}
 
 	DemoCurrentTime += DeltaSeconds;
@@ -1158,7 +1173,29 @@ void UDemoNetDriver::TickDemoPlayback( float DeltaSeconds )
 		DemoFrameNum++;
 	}
 
-	bIsFastForwarding = false;
+	if ( bIsFastForwarding )
+	{
+		bIsFastForwarding = false;
+
+		// Flush all pending RepNotifies that were built up during the fast-forward.
+		if ( ServerConnection != nullptr)
+		{
+			for ( auto& ChannelPair : ServerConnection->ActorChannels )
+			{
+				if ( ChannelPair.Value != NULL )
+				{
+					for (auto& ReplicatorPair : ChannelPair.Value->ReplicationMap)
+					{
+						ReplicatorPair.Value->CallRepNotifies();
+					}
+				}
+			}
+		}
+
+		const auto FastForwardTotalSeconds = FPlatformTime::Seconds() - FastForwardStartSeconds;
+
+		UE_LOG( LogDemo, Log, TEXT( "Fast forward took %.2f seconds." ), FastForwardTotalSeconds );
+	}
 }
 
 void UDemoNetDriver::SpawnDemoRecSpectator( UNetConnection* Connection )
