#marked set morphtarget to be unsafe to be called by construction script.

- It needs tick, and we don't tick from construction script.

https://jira.ol.epicgames.net/browse/UE-4489

diff --git a/Engine/Source/Runtime/Engine/Classes/Components/SkeletalMeshComponent.h b/Engine/Source/Runtime/Engine/Classes/Components/SkeletalMeshComponent.h
--- a/Engine/Source/Runtime/Engine/Classes/Components/SkeletalMeshComponent.h
+++ b/Engine/Source/Runtime/Engine/Classes/Components/SkeletalMeshComponent.h
@@ -569,7 +569,7 @@ public:
 	 *
 	 * @param bRemoveZeroWeight : Used by editor code when it should stay in the active list with zero weight
 	 */
-	UFUNCTION(BlueprintCallable, Category="Components|SkeletalMesh")
+	UFUNCTION(BlueprintCallable, Category="Components|SkeletalMesh", meta=(UnsafeDuringActorConstruction="true"))
 	void SetMorphTarget(FName MorphTargetName, float Value, bool bRemoveZeroWeight=true);
 
 	/**
