Handle 0 size checkpoints

diff --git a/Engine/Source/Runtime/Engine/Private/DemoNetDriver.cpp b/Engine/Source/Runtime/Engine/Private/DemoNetDriver.cpp
--- a/Engine/Source/Runtime/Engine/Private/DemoNetDriver.cpp
+++ b/Engine/Source/Runtime/Engine/Private/DemoNetDriver.cpp
@@ -832,7 +832,10 @@ void UDemoNetDriver::SaveCheckpoint()
 
 	const int32 TotalSize = CheckpointArchive->TotalSize();
 
-	ReplayStreamer->FlushCheckpoint( CurrentCheckpointIndex++ );
+	if ( CheckpointArchive->TotalSize() > 0 )
+	{
+		ReplayStreamer->FlushCheckpoint( CurrentCheckpointIndex++ );
+	}
 
 	const double EndCheckpointTime = FPlatformTime::Seconds();
 
@@ -1258,6 +1261,13 @@ void UDemoNetDriver::CheckpointReady( bool bSuccess )
 		return;
 	}
 
+	FArchive* CheckpointArchive = ReplayStreamer->GetCheckpointArchive();
+
+	if ( CheckpointArchive->TotalSize() == 0 )
+	{
+		return;
+	}
+
 	if ( SpectatorController == NULL )
 	{
 		UE_LOG( LogDemo, Warning, TEXT( "UDemoNetConnection::CheckpointReady: No spectator player controller." ) );
@@ -1311,8 +1321,6 @@ void UDemoNetDriver::CheckpointReady( bool bSuccess )
 	// Create fake control channel
 	ServerConnection->CreateChannel( CHTYPE_Control, 1 );
 
-	FArchive* CheckpointArchive = ReplayStreamer->GetCheckpointArchive();
-
 	GuidCache->ObjectLookup.Empty();
 	GuidCache->NetGUIDLookup.Empty();
 
diff --git a/Engine/Source/Runtime/NetworkReplayStreaming/HttpNetworkReplayStreaming/Private/HttpNetworkReplayStreaming.cpp b/Engine/Source/Runtime/NetworkReplayStreaming/HttpNetworkReplayStreaming/Private/HttpNetworkReplayStreaming.cpp
--- a/Engine/Source/Runtime/NetworkReplayStreaming/HttpNetworkReplayStreaming/Private/HttpNetworkReplayStreaming.cpp
+++ b/Engine/Source/Runtime/NetworkReplayStreaming/HttpNetworkReplayStreaming/Private/HttpNetworkReplayStreaming.cpp
@@ -821,6 +821,13 @@ void FHttpNetworkReplayStreamer::HttpDownloadCheckpointFinished( FHttpRequestPtr
 
 	if ( bSucceeded && HttpResponse->GetResponseCode() == EHttpResponseCodes::Ok )
 	{
+		if ( HttpResponse->GetContent().Num() == 0 )
+		{
+			UE_LOG( LogHttpReplay, Warning, TEXT( "FHttpNetworkReplayStreamer::HttpDownloadCheckpointFinished. Checkpoint empty." ) );
+			GotoCheckpointDelegate.ExecuteIfBound( true );
+			return;
+		}
+
 		StreamArchive.Buffer.Empty();
 		StreamArchive.Pos = 0;
 		StreamArchive.bAtEndOfReplay = false;
