Sanity check 100% reliable connections aren't skipping acks

diff --git a/Engine/Source/Runtime/Engine/Private/NetConnection.cpp b/Engine/Source/Runtime/Engine/Private/NetConnection.cpp
--- a/Engine/Source/Runtime/Engine/Private/NetConnection.cpp
+++ b/Engine/Source/Runtime/Engine/Private/NetConnection.cpp
@@ -1016,6 +1016,8 @@ void UNetConnection::ReceivedPacket( FBitReader& Reader )
 
 	ValidateSendBuffer();
 
+	check( !bSkipAck || !InternalAck );		// 100% reliable connections shouldn't be skipping acks
+
 	// Acknowledge the packet.
 	if ( !bSkipAck )
 	{
