#ue4 - Remove uses of GWorld in CollisionConversions debug drawing code.

UE-12114

diff --git a/Engine/Source/Runtime/Engine/Private/Collision/CollisionConversions.cpp b/Engine/Source/Runtime/Engine/Private/Collision/CollisionConversions.cpp
--- a/Engine/Source/Runtime/Engine/Private/Collision/CollisionConversions.cpp
+++ b/Engine/Source/Runtime/Engine/Private/Collision/CollisionConversions.cpp
@@ -19,7 +19,7 @@ extern TAutoConsoleVariable<int32> CVarShowInitialOverlaps;
 static const PxQueryHit InvalidQueryHit;
 
 // Forward declare, I don't want to move the entire function right now or we lose change history.
-static bool ConvertOverlappedShapeToImpactHit(const PxLocationHit& PHit, const FVector& StartLoc, const FVector& EndLoc, FHitResult& OutResult, const PxGeometry& Geom, const PxTransform& QueryTM, const PxFilterData& QueryFilter, bool bReturnPhysMat);
+static bool ConvertOverlappedShapeToImpactHit(const UWorld* World, const PxLocationHit& PHit, const FVector& StartLoc, const FVector& EndLoc, FHitResult& OutResult, const PxGeometry& Geom, const PxTransform& QueryTM, const PxFilterData& QueryFilter, bool bReturnPhysMat);
 
 DECLARE_CYCLE_STAT(TEXT("ConvertQueryHit"), STAT_ConvertQueryImpactHit, STATGROUP_Collision);
 DECLARE_CYCLE_STAT(TEXT("ConvertOverlapToHit"), STAT_CollisionConvertOverlapToHit, STATGROUP_Collision);
@@ -340,14 +340,14 @@ static void SetHitResultFromShapeAndFaceIndex(const PxShape* PShape,  const PxRi
 	}
 }
 
-void ConvertQueryImpactHit(const PxLocationHit& PHit, FHitResult& OutResult, float CheckLength, const PxFilterData& QueryFilter, const FVector& StartLoc, const FVector& EndLoc, const PxGeometry* const Geom, const PxTransform& QueryTM, bool bReturnFaceIndex, bool bReturnPhysMat)
+void ConvertQueryImpactHit(const UWorld* World, const PxLocationHit& PHit, FHitResult& OutResult, float CheckLength, const PxFilterData& QueryFilter, const FVector& StartLoc, const FVector& EndLoc, const PxGeometry* const Geom, const PxTransform& QueryTM, bool bReturnFaceIndex, bool bReturnPhysMat)
 {
 	SCOPE_CYCLE_COUNTER(STAT_ConvertQueryImpactHit);
 
 	checkSlow(PHit.flags & PxHitFlag::eDISTANCE);
 	if (Geom != NULL && PHit.hadInitialOverlap())
 	{
-		ConvertOverlappedShapeToImpactHit(PHit, StartLoc, EndLoc, OutResult, *Geom, QueryTM, QueryFilter, bReturnPhysMat);
+		ConvertOverlappedShapeToImpactHit(World, PHit, StartLoc, EndLoc, OutResult, *Geom, QueryTM, QueryFilter, bReturnPhysMat);
 		return;
 	}
 
@@ -422,7 +422,7 @@ void ConvertQueryImpactHit(const PxLocationHit& PHit, FHitResult& OutResult, flo
 	}
 }
 
-void ConvertRaycastResults(int32 NumHits, PxRaycastHit* Hits, float CheckLength, const PxFilterData& QueryFilter, TArray<FHitResult>& OutHits, const FVector& StartLoc, const FVector& EndLoc, bool bReturnFaceIndex, bool bReturnPhysMat)
+void ConvertRaycastResults(const UWorld* World, int32 NumHits, PxRaycastHit* Hits, float CheckLength, const PxFilterData& QueryFilter, TArray<FHitResult>& OutHits, const FVector& StartLoc, const FVector& EndLoc, bool bReturnFaceIndex, bool bReturnPhysMat)
 {
 	OutHits.Reserve(NumHits);
 
@@ -432,7 +432,7 @@ void ConvertRaycastResults(int32 NumHits, PxRaycastHit* Hits, float CheckLength,
 		FHitResult NewResult(ForceInit);
 		const PxRaycastHit& PHit = Hits[i];
 
-		ConvertQueryImpactHit(PHit, NewResult, CheckLength, QueryFilter, StartLoc, EndLoc, NULL, PStartTM, bReturnFaceIndex, bReturnPhysMat);
+		ConvertQueryImpactHit(World, PHit, NewResult, CheckLength, QueryFilter, StartLoc, EndLoc, NULL, PStartTM, bReturnFaceIndex, bReturnPhysMat);
 
 		OutHits.Add(NewResult);
 	}
@@ -441,7 +441,7 @@ void ConvertRaycastResults(int32 NumHits, PxRaycastHit* Hits, float CheckLength,
 	OutHits.Sort( FCompareFHitResultTime() );
 }
 
-bool AddSweepResults(int32 NumHits, const PxSweepHit* Hits, float CheckLength, const PxFilterData& QueryFilter, TArray<FHitResult>& OutHits, const FVector& StartLoc, const FVector& EndLoc, const PxGeometry& Geom, const PxTransform& QueryTM, float MaxDistance, bool bReturnPhysMat)
+bool AddSweepResults(const UWorld* World, int32 NumHits, const PxSweepHit* Hits, float CheckLength, const PxFilterData& QueryFilter, TArray<FHitResult>& OutHits, const FVector& StartLoc, const FVector& EndLoc, const PxGeometry& Geom, const PxTransform& QueryTM, float MaxDistance, bool bReturnPhysMat)
 {
 	OutHits.Reserve(NumHits);
 	bool bHadBlockingHit = false;
@@ -453,7 +453,7 @@ bool AddSweepResults(int32 NumHits, const PxSweepHit* Hits, float CheckLength, c
 		if(PHit.distance <= MaxDistance)
 		{
 			FHitResult NewResult(ForceInit);
-			ConvertQueryImpactHit(PHit, NewResult, CheckLength, QueryFilter, StartLoc, EndLoc, &Geom, QueryTM, false, bReturnPhysMat);
+			ConvertQueryImpactHit(World, PHit, NewResult, CheckLength, QueryFilter, StartLoc, EndLoc, &Geom, QueryTM, false, bReturnPhysMat);
 			bHadBlockingHit |= NewResult.bBlockingHit;
 			OutHits.Add(NewResult);
 		}
@@ -466,15 +466,15 @@ bool AddSweepResults(int32 NumHits, const PxSweepHit* Hits, float CheckLength, c
 
 /* Function to find the best normal from the list of triangles that are overlapping our geom. */
 template<typename GeomType>
-FVector FindBestOverlappingNormal(const PxGeometry& Geom, const PxTransform& QueryTM, const GeomType& ShapeGeom, const PxTransform& PShapeWorldPose, PxU32* HitTris, int32 NumTrisHit, bool bCanDrawOverlaps = false)
+FVector FindBestOverlappingNormal(const UWorld* World, const PxGeometry& Geom, const PxTransform& QueryTM, const GeomType& ShapeGeom, const PxTransform& PShapeWorldPose, PxU32* HitTris, int32 NumTrisHit, bool bCanDrawOverlaps = false)
 {
 #if DRAW_OVERLAPPING_TRIS
 	const float Lifetime = 5.f;
-	bCanDrawOverlaps &= GWorld->IsGameWorld() && GWorld->PersistentLineBatcher && (GWorld->PersistentLineBatcher->BatchedLines.Num() < 2048);
+	bCanDrawOverlaps &= World && World->IsGameWorld() && World->PersistentLineBatcher && (World->PersistentLineBatcher->BatchedLines.Num() < 2048);
 	if (bCanDrawOverlaps)
 	{
 		TArray<FOverlapResult> Overlaps;
-		DrawGeomOverlaps(GWorld, Geom, QueryTM, Overlaps, Lifetime);
+		DrawGeomOverlaps(World, Geom, QueryTM, Overlaps, Lifetime);
 	}
 	const FLinearColor LineColor = FLinearColor::Green;
 	const FLinearColor NormalColor = FLinearColor::Red;
@@ -512,20 +512,20 @@ FVector FindBestOverlappingNormal(const PxGeometry& Geom, const PxTransform& Que
 		}
 
 #if DRAW_OVERLAPPING_TRIS
-		if (bCanDrawOverlaps && (GWorld->PersistentLineBatcher->BatchedLines.Num() < 2048))
+		if (bCanDrawOverlaps && (World->PersistentLineBatcher->BatchedLines.Num() < 2048))
 		{
 			static const float LineThickness = 0.9f;
 			static const float NormalThickness = 0.75f;
 			static const float PointThickness = 5.0f;
-			GWorld->PersistentLineBatcher->DrawLine(A, B, LineColor, SDPG_Foreground, LineThickness, Lifetime);
-			GWorld->PersistentLineBatcher->DrawLine(B, C, LineColor, SDPG_Foreground, LineThickness, Lifetime);
-			GWorld->PersistentLineBatcher->DrawLine(C, A, LineColor, SDPG_Foreground, LineThickness, Lifetime);
+			World->PersistentLineBatcher->DrawLine(A, B, LineColor, SDPG_Foreground, LineThickness, Lifetime);
+			World->PersistentLineBatcher->DrawLine(B, C, LineColor, SDPG_Foreground, LineThickness, Lifetime);
+			World->PersistentLineBatcher->DrawLine(C, A, LineColor, SDPG_Foreground, LineThickness, Lifetime);
 			const FVector Centroid((A + B + C) / 3.f);
-			GWorld->PersistentLineBatcher->DrawLine(Centroid, Centroid + (35.0f*TriNormal), NormalColor, SDPG_Foreground, NormalThickness, Lifetime);
-			GWorld->PersistentLineBatcher->DrawPoint(Centroid + (35.0f*TriNormal), NormalColor, PointThickness, SDPG_Foreground, Lifetime);
-			GWorld->PersistentLineBatcher->DrawPoint(A, PointColor, PointThickness, SDPG_Foreground, Lifetime);
-			GWorld->PersistentLineBatcher->DrawPoint(B, PointColor, PointThickness, SDPG_Foreground, Lifetime);
-			GWorld->PersistentLineBatcher->DrawPoint(C, PointColor, PointThickness, SDPG_Foreground, Lifetime);
+			World->PersistentLineBatcher->DrawLine(Centroid, Centroid + (35.0f*TriNormal), NormalColor, SDPG_Foreground, NormalThickness, Lifetime);
+			World->PersistentLineBatcher->DrawPoint(Centroid + (35.0f*TriNormal), NormalColor, PointThickness, SDPG_Foreground, Lifetime);
+			World->PersistentLineBatcher->DrawPoint(A, PointColor, PointThickness, SDPG_Foreground, Lifetime);
+			World->PersistentLineBatcher->DrawPoint(B, PointColor, PointThickness, SDPG_Foreground, Lifetime);
+			World->PersistentLineBatcher->DrawPoint(C, PointColor, PointThickness, SDPG_Foreground, Lifetime);
 		}
 #endif // DRAW_OVERLAPPING_TRIS
 	}
@@ -645,7 +645,7 @@ static bool CanFindOverlappedTriangle(const PxShape* PShape)
 }
 
 
-static bool FindOverlappedTriangleNormal_Internal(const PxGeometry& Geom, const PxTransform& QueryTM, const PxShape* PShape, const PxTransform& PShapeWorldPose, FVector& OutNormal, bool bCanDrawOverlaps = false)
+static bool FindOverlappedTriangleNormal_Internal(const UWorld* World, const PxGeometry& Geom, const PxTransform& QueryTM, const PxShape* PShape, const PxTransform& PShapeWorldPose, FVector& OutNormal, bool bCanDrawOverlaps = false)
 {
 	if (CanFindOverlappedTriangle(PShape))
 	{
@@ -667,11 +667,11 @@ static bool FindOverlappedTriangleNormal_Internal(const PxGeometry& Geom, const
 			{
 				if (bIsTriMesh)
 				{
-					OutNormal = FindBestOverlappingNormal(Geom, QueryTM, PTriMeshGeom, PShapeWorldPose, HitTris, NumTrisHit, bCanDrawOverlaps);
+					OutNormal = FindBestOverlappingNormal(World, Geom, QueryTM, PTriMeshGeom, PShapeWorldPose, HitTris, NumTrisHit, bCanDrawOverlaps);
 				}
 				else
 				{
-					OutNormal = FindBestOverlappingNormal(Geom, QueryTM, PHeightfieldGeom, PShapeWorldPose, HitTris, NumTrisHit, bCanDrawOverlaps);
+					OutNormal = FindBestOverlappingNormal(World, Geom, QueryTM, PHeightfieldGeom, PShapeWorldPose, HitTris, NumTrisHit, bCanDrawOverlaps);
 				}
 
 				return true;
@@ -683,7 +683,7 @@ static bool FindOverlappedTriangleNormal_Internal(const PxGeometry& Geom, const
 }
 
 
-static bool FindOverlappedTriangleNormal(const PxGeometry& Geom, const PxTransform& QueryTM, const PxShape* PShape, const PxTransform& PShapeWorldPose, FVector& OutNormal, float Inflation, bool bCanDrawOverlaps = false)
+static bool FindOverlappedTriangleNormal(const UWorld* World, const PxGeometry& Geom, const PxTransform& QueryTM, const PxShape* PShape, const PxTransform& PShapeWorldPose, FVector& OutNormal, float Inflation, bool bCanDrawOverlaps = false)
 {
 	bool bSuccess = false;
 
@@ -691,7 +691,7 @@ static bool FindOverlappedTriangleNormal(const PxGeometry& Geom, const PxTransfo
 	{
 		if (Inflation <= 0.f)
 		{
-			bSuccess = FindOverlappedTriangleNormal_Internal(Geom, QueryTM, PShape, PShapeWorldPose, OutNormal, bCanDrawOverlaps);
+			bSuccess = FindOverlappedTriangleNormal_Internal(World, Geom, QueryTM, PShape, PShapeWorldPose, OutNormal, bCanDrawOverlaps);
 		}
 		else
 		{
@@ -702,7 +702,7 @@ static bool FindOverlappedTriangleNormal(const PxGeometry& Geom, const PxTransfo
 				{
 					const PxCapsuleGeometry* InCapsule = static_cast<const PxCapsuleGeometry*>(&Geom);
 					PxCapsuleGeometry InflatedCapsule(InCapsule->radius + Inflation, InCapsule->halfHeight); // don't inflate halfHeight, radius is added all around.
-					bSuccess = FindOverlappedTriangleNormal_Internal(InflatedCapsule, QueryTM, PShape, PShapeWorldPose, OutNormal, bCanDrawOverlaps);
+					bSuccess = FindOverlappedTriangleNormal_Internal(World, InflatedCapsule, QueryTM, PShape, PShapeWorldPose, OutNormal, bCanDrawOverlaps);
 					break;
 				}
 
@@ -710,7 +710,7 @@ static bool FindOverlappedTriangleNormal(const PxGeometry& Geom, const PxTransfo
 				{
 					const PxBoxGeometry* InBox = static_cast<const PxBoxGeometry*>(&Geom);
 					PxBoxGeometry InflatedBox(InBox->halfExtents + PxVec3(Inflation));
-					bSuccess = FindOverlappedTriangleNormal_Internal(InflatedBox, QueryTM, PShape, PShapeWorldPose, OutNormal, bCanDrawOverlaps);
+					bSuccess = FindOverlappedTriangleNormal_Internal(World, InflatedBox, QueryTM, PShape, PShapeWorldPose, OutNormal, bCanDrawOverlaps);
 					break;
 				}
 
@@ -718,7 +718,7 @@ static bool FindOverlappedTriangleNormal(const PxGeometry& Geom, const PxTransfo
 				{
 					const PxSphereGeometry* InSphere = static_cast<const PxSphereGeometry*>(&Geom);
 					PxSphereGeometry InflatedSphere(InSphere->radius + Inflation);
-					bSuccess = FindOverlappedTriangleNormal_Internal(InflatedSphere, QueryTM, PShape, PShapeWorldPose, OutNormal, bCanDrawOverlaps);
+					bSuccess = FindOverlappedTriangleNormal_Internal(World, InflatedSphere, QueryTM, PShape, PShapeWorldPose, OutNormal, bCanDrawOverlaps);
 					break;
 				}
 
@@ -736,7 +736,7 @@ static bool FindOverlappedTriangleNormal(const PxGeometry& Geom, const PxTransfo
 
 
 /** Util to convert an overlapped shape into a sweep hit result, returns whether it was a blocking hit. */
-static bool ConvertOverlappedShapeToImpactHit(const PxLocationHit& PHit, const FVector& StartLoc, const FVector& EndLoc, FHitResult& OutResult, const PxGeometry& Geom, const PxTransform& QueryTM, const PxFilterData& QueryFilter, bool bReturnPhysMat)
+static bool ConvertOverlappedShapeToImpactHit(const UWorld* World, const PxLocationHit& PHit, const FVector& StartLoc, const FVector& EndLoc, FHitResult& OutResult, const PxGeometry& Geom, const PxTransform& QueryTM, const PxFilterData& QueryFilter, bool bReturnPhysMat)
 {
 	SCOPE_CYCLE_COUNTER(STAT_CollisionConvertOverlapToHit);
 
@@ -783,11 +783,11 @@ static bool ConvertOverlappedShapeToImpactHit(const PxLocationHit& PHit, const F
 	}
 
 #if DRAW_OVERLAPPING_TRIS
-	if (CVarShowInitialOverlaps.GetValueOnAnyThread() != 0 && GWorld->IsGameWorld())
+	if (CVarShowInitialOverlaps.GetValueOnAnyThread() != 0 && World && World->IsGameWorld())
 	{
 		FVector DummyNormal(0.f);
 		const PxTransform PShapeWorldPose = PxShapeExt::getGlobalPose(*PShape, *PActor);
-		FindOverlappedTriangleNormal(Geom, QueryTM, PShape, PShapeWorldPose, DummyNormal, 0.f, true);
+		FindOverlappedTriangleNormal(World, Geom, QueryTM, PShape, PShapeWorldPose, DummyNormal, 0.f, true);
 	}
 #endif
 
@@ -810,8 +810,8 @@ static bool ConvertOverlappedShapeToImpactHit(const PxLocationHit& PHit, const F
 		else
 		{
 			static const float SmallOverlapInflation = 0.250f;
-			if (FindOverlappedTriangleNormal(Geom, QueryTM, PShape, PShapeWorldPose, OutResult.ImpactNormal, 0.f, false) ||
-				FindOverlappedTriangleNormal(Geom, QueryTM, PShape, PShapeWorldPose, OutResult.ImpactNormal, SmallOverlapInflation, false))
+			if (FindOverlappedTriangleNormal(World, Geom, QueryTM, PShape, PShapeWorldPose, OutResult.ImpactNormal, 0.f, false) ||
+				FindOverlappedTriangleNormal(World, Geom, QueryTM, PShape, PShapeWorldPose, OutResult.ImpactNormal, SmallOverlapInflation, false))
 			{
 				// Success
 			}
diff --git a/Engine/Source/Runtime/Engine/Private/Collision/CollisionConversions.h b/Engine/Source/Runtime/Engine/Private/Collision/CollisionConversions.h
--- a/Engine/Source/Runtime/Engine/Private/Collision/CollisionConversions.h
+++ b/Engine/Source/Runtime/Engine/Private/Collision/CollisionConversions.h
@@ -20,7 +20,7 @@
  * @param	bReturnFaceIndex	True if we want to lookup the face index
  * @param	bReturnPhysMat		True if we want to lookup the physical material
  */
-void ConvertQueryImpactHit(const PxLocationHit& PHit, FHitResult& OutResult, float CheckLength, const PxFilterData& QueryFilter, const FVector& StartLoc, const FVector& EndLoc, const PxGeometry* const Geom, const PxTransform& QueryTM, bool bReturnFaceIndex, bool bReturnPhysMat);
+void ConvertQueryImpactHit(const UWorld* World, const PxLocationHit& PHit, FHitResult& OutResult, float CheckLength, const PxFilterData& QueryFilter, const FVector& StartLoc, const FVector& EndLoc, const PxGeometry* const Geom, const PxTransform& QueryTM, bool bReturnFaceIndex, bool bReturnPhysMat);
 
 /** 
  * Util to convert physX raycast results to our hit results
@@ -34,7 +34,7 @@ void ConvertQueryImpactHit(const PxLocationHit& PHit, FHitResult& OutResult, flo
  * @param	bReturnFaceIndex	True if we want to lookup the face index
  * @param	bReturnPhysMat		True if we want to lookup the physical material
  */
-void ConvertRaycastResults(int32 NumHits, PxRaycastHit* Hits, float CheckLength, const PxFilterData& QueryFilter, TArray<FHitResult>& OutHits, const FVector& StartLoc, const FVector& EndLoc, bool bReturnFaceIndex, bool bReturnPhysMat);
+void ConvertRaycastResults(const UWorld* World, int32 NumHits, PxRaycastHit* Hits, float CheckLength, const PxFilterData& QueryFilter, TArray<FHitResult>& OutHits, const FVector& StartLoc, const FVector& EndLoc, bool bReturnFaceIndex, bool bReturnPhysMat);
 
 /** 
  * Util to convert physX sweep results to unreal hit results and add to array
@@ -48,7 +48,7 @@ void ConvertRaycastResults(int32 NumHits, PxRaycastHit* Hits, float CheckLength,
  * @param	Geom
  * @return	true if any blocking hit was found within MaxDistance.
  */
-bool AddSweepResults(int32 NumHits, const PxSweepHit* Hits, float CheckLength, const PxFilterData& QueryFilter, TArray<FHitResult>& OutHits, const FVector& StartLoc, const FVector& EndLoc, const PxGeometry& Geom, const PxTransform& QueryTM, float MaxDistance, bool bReturnPhysMat);
+bool AddSweepResults(const UWorld* World, int32 NumHits, const PxSweepHit* Hits, float CheckLength, const PxFilterData& QueryFilter, TArray<FHitResult>& OutHits, const FVector& StartLoc, const FVector& EndLoc, const PxGeometry& Geom, const PxTransform& QueryTM, float MaxDistance, bool bReturnPhysMat);
 
 /** 
  * Util to convert physX overlap query to our overlap result
@@ -58,7 +58,7 @@ bool AddSweepResults(int32 NumHits, const PxSweepHit* Hits, float CheckLength, c
  * @param	OutOverlap	(out) Result converted
  * @param	QueryFilter	Query Filter 
  */
-void ConvertQueryOverlap(const PxShape* PShape, const PxRigidActor* PActor, FOverlapResult& OutOverlap, const PxFilterData& QueryFilter);
+void ConvertQueryOverlap(const UWorld* World, const PxShape* PShape, const PxRigidActor* PActor, FOverlapResult& OutOverlap, const PxFilterData& QueryFilter);
 
 /**
  * Util to determine if a shape is deemed blocking based on the query filter
@@ -78,7 +78,7 @@ bool IsBlocking(const PxShape* PShape, const PxFilterData& QueryFilter);
  * @param	QueryFilter	Query filter for converting
  * @return	OutOverlaps	Converted data
  */
-bool ConvertOverlapResults(int32 NumOverlaps,  PxOverlapHit* POverlapResults, const PxFilterData& QueryFilter, TArray<FOverlapResult>& OutOverlaps);
+bool ConvertOverlapResults(const UWorld* World, int32 NumOverlaps,  PxOverlapHit* POverlapResults, const PxFilterData& QueryFilter, TArray<FOverlapResult>& OutOverlaps);
 
 
 #endif // WITH_PHYX
diff --git a/Engine/Source/Runtime/Engine/Private/Collision/PhysXCollision.cpp b/Engine/Source/Runtime/Engine/Private/Collision/PhysXCollision.cpp
--- a/Engine/Source/Runtime/Engine/Private/Collision/PhysXCollision.cpp
+++ b/Engine/Source/Runtime/Engine/Private/Collision/PhysXCollision.cpp
@@ -952,7 +952,7 @@ bool RaycastSingle(const UWorld* World, struct FHitResult& OutHit, const FVector
 			if (bHaveBlockingHit) // If we got a hit
 			{
 				PxTransform PStartTM(U2PVector(Start));
-				ConvertQueryImpactHit(PHit, OutHit, DeltaMag, PFilter, Start, End, NULL, PStartTM, Params.bReturnFaceIndex, Params.bReturnPhysicalMaterial);
+				ConvertQueryImpactHit(World, PHit, OutHit, DeltaMag, PFilter, Start, End, NULL, PStartTM, Params.bReturnFaceIndex, Params.bReturnPhysicalMaterial);
 			}
 		}
 #endif //WITH_PHYSX
@@ -1130,7 +1130,7 @@ bool RaycastMulti(const UWorld* World, TArray<struct FHitResult>& OutHits, const
 
 		if (NumHits > 0)
 		{
-			ConvertRaycastResults(NumHits, PHits, DeltaMag, PFilter, OutHits, Start, End, Params.bReturnFaceIndex, Params.bReturnPhysicalMaterial);
+			ConvertRaycastResults(World, NumHits, PHits, DeltaMag, PFilter, OutHits, Start, End, Params.bReturnFaceIndex, Params.bReturnPhysicalMaterial);
 		}
 
 		bHaveBlockingHit = bBlockingHit;
@@ -1333,7 +1333,7 @@ bool GeomSweepSingle(const UWorld* World, const struct FCollisionShape& Collisio
 
 		if(bHaveBlockingHit) // If we got a hit, convert it to unreal type
 		{
-			ConvertQueryImpactHit(PHit, OutHit, DeltaMag, PFilter, Start, End, &PGeom, PStartTM, false, Params.bReturnPhysicalMaterial);
+			ConvertQueryImpactHit(World, PHit, OutHit, DeltaMag, PFilter, Start, End, &PGeom, PStartTM, false, Params.bReturnPhysicalMaterial);
 		}
 	}
 
@@ -1431,7 +1431,7 @@ bool GeomSweepMulti_PhysX(const UWorld* World, const PxGeometry& PGeom, const Px
 				{
 					TestLength = PHits[NumHits - 1].distance;
 					FHitResult TestHit;
-					ConvertQueryImpactHit(PHits[NumHits - 1], TestHit, DeltaMag, PFilter, Start, End, &PGeom, PStartTM, false, Params.bReturnPhysicalMaterial);
+					ConvertQueryImpactHit(World, PHits[NumHits - 1], TestHit, DeltaMag, PFilter, Start, End, &PGeom, PStartTM, false, Params.bReturnPhysicalMaterial);
 					HitLocation = TestHit.Location;
 				}
 
@@ -1503,7 +1503,7 @@ bool GeomSweepMulti_PhysX(const UWorld* World, const PxGeometry& PGeom, const Px
 		// Convert all hits to unreal structs. This will remove any hits further than MinBlockDistance, and sort results.
 		if (NumHits > 0)
 		{
-			bBlockingHit |= AddSweepResults(NumHits, PHits, DeltaMag, PFilter, OutHits, Start, End, PGeom, PStartTM, MinBlockDistance, Params.bReturnPhysicalMaterial);
+			bBlockingHit |= AddSweepResults(World, NumHits, PHits, DeltaMag, PFilter, OutHits, Start, End, PGeom, PStartTM, MinBlockDistance, Params.bReturnPhysicalMaterial);
 		}
 	}
 
diff --git a/Engine/Source/Runtime/Engine/Private/PhysicsEngine/BodyInstance.cpp b/Engine/Source/Runtime/Engine/Private/PhysicsEngine/BodyInstance.cpp
--- a/Engine/Source/Runtime/Engine/Private/PhysicsEngine/BodyInstance.cpp
+++ b/Engine/Source/Runtime/Engine/Private/PhysicsEngine/BodyInstance.cpp
@@ -3521,7 +3521,8 @@ bool FBodyInstance::LineTrace(struct FHitResult& OutHit, const FVector& Start, c
 				QueryFilter.word2 = 0xFFFFF;
 
 				PxTransform PStartTM(U2PVector(Start));
-				ConvertQueryImpactHit(BestHit, OutHit, DeltaMag, QueryFilter, Start, End, NULL, PStartTM, false, bReturnPhysicalMaterial);
+				const UPrimitiveComponent* Owner = OwnerComponent.Get();
+				ConvertQueryImpactHit(Owner ? Owner->GetWorld() : nullptr, BestHit, OutHit, DeltaMag, QueryFilter, Start, End, NULL, PStartTM, false, bReturnPhysicalMaterial);
 				bHitSomething = true;
 			}
 		}
@@ -3623,7 +3624,8 @@ bool FBodyInstance::InternalSweepPhysX(struct FHitResult& OutHit, const FVector&
 						PHit.shape = PShape;
 						PHit.actor = PShape->getActor();
 						PxTransform PStartTransform(U2PVector(Start));
-						ConvertQueryImpactHit(PHit, OutHit, DeltaMag, QueryFilter, Start, End, NULL, PStartTransform, false, false);
+						const UPrimitiveComponent* Owner = OwnerComponent.Get();
+						ConvertQueryImpactHit(Owner ? Owner->GetWorld() : nullptr, PHit, OutHit, DeltaMag, QueryFilter, Start, End, NULL, PStartTransform, false, false);
 						return true;
 					}
 				}
