#ue4 - Fix build

diff --git a/Engine/Source/Runtime/Engine/Private/Collision/PhysXCollision.cpp b/Engine/Source/Runtime/Engine/Private/Collision/PhysXCollision.cpp
--- a/Engine/Source/Runtime/Engine/Private/Collision/PhysXCollision.cpp
+++ b/Engine/Source/Runtime/Engine/Private/Collision/PhysXCollision.cpp
@@ -1668,7 +1668,7 @@ bool GeomOverlapMultiImp_PhysX(const UWorld* World, const PxGeometry& PGeom, con
 		{
 			if (NumHits > 0)
 			{
-				bHaveBlockingHit = ConvertOverlapResults(NumHits, POverlapArray, PFilter, OutOverlaps);
+				bHaveBlockingHit = ConvertOverlapResults(World, NumHits, POverlapArray, PFilter, OutOverlaps);
 			}
 
 #if !(UE_BUILD_SHIPPING || UE_BUILD_TEST)
