#ue4 - Actually fix build...

diff --git a/Engine/Source/Runtime/Engine/Private/Collision/CollisionConversions.h b/Engine/Source/Runtime/Engine/Private/Collision/CollisionConversions.h
--- a/Engine/Source/Runtime/Engine/Private/Collision/CollisionConversions.h
+++ b/Engine/Source/Runtime/Engine/Private/Collision/CollisionConversions.h
@@ -78,7 +78,7 @@ bool IsBlocking(const PxShape* PShape, const PxFilterData& QueryFilter);
  * @param	QueryFilter	Query filter for converting
  * @return	OutOverlaps	Converted data
  */
-bool ConvertOverlapResults(const UWorld* World, int32 NumOverlaps,  PxOverlapHit* POverlapResults, const PxFilterData& QueryFilter, TArray<FOverlapResult>& OutOverlaps);
+bool ConvertOverlapResults(int32 NumOverlaps,  PxOverlapHit* POverlapResults, const PxFilterData& QueryFilter, TArray<FOverlapResult>& OutOverlaps);
 
 
 #endif // WITH_PHYX
diff --git a/Engine/Source/Runtime/Engine/Private/Collision/PhysXCollision.cpp b/Engine/Source/Runtime/Engine/Private/Collision/PhysXCollision.cpp
--- a/Engine/Source/Runtime/Engine/Private/Collision/PhysXCollision.cpp
+++ b/Engine/Source/Runtime/Engine/Private/Collision/PhysXCollision.cpp
@@ -1668,7 +1668,7 @@ bool GeomOverlapMultiImp_PhysX(const UWorld* World, const PxGeometry& PGeom, con
 		{
 			if (NumHits > 0)
 			{
-				bHaveBlockingHit = ConvertOverlapResults(World, NumHits, POverlapArray, PFilter, OutOverlaps);
+				bHaveBlockingHit = ConvertOverlapResults(NumHits, POverlapArray, PFilter, OutOverlaps);
 			}
 
 #if !(UE_BUILD_SHIPPING || UE_BUILD_TEST)
