UE4 - add some code to try to track down a problem with lightmaps and garbage collection

diff --git a/Engine/Source/Runtime/Engine/Private/Components/StaticMeshComponent.cpp b/Engine/Source/Runtime/Engine/Private/Components/StaticMeshComponent.cpp
--- a/Engine/Source/Runtime/Engine/Private/Components/StaticMeshComponent.cpp
+++ b/Engine/Source/Runtime/Engine/Private/Components/StaticMeshComponent.cpp
@@ -1581,6 +1581,30 @@ void UStaticMeshComponent::ApplyComponentInstanceData(FStaticMeshComponentInstan
 		{
 			LODData[i].LightMap = StaticMeshInstanceData->CachedStaticLighting.LODDataLightMap[i];
 			LODData[i].ShadowMap = StaticMeshInstanceData->CachedStaticLighting.LODDataShadowMap[i];
+
+// this code is to try to track down a mystery GC crash from crash reporter...if this code crashes, and then we know it is component reinstancing
+			if (LODData[i].LightMap.GetReference())
+			{
+				FLightMap2D* LightMap = LODData[i].LightMap->GetLightMap2D();
+				if (LightMap)
+				{
+					if (LightMap->IsValid(0))
+					{
+						UTexture2D* Tex = LightMap->GetTexture(0);
+						Tex->GetResourceSize(EResourceSizeMode::Exclusive);
+					}
+					if (LightMap->IsValid(1))
+					{
+						UTexture2D* Tex = LightMap->GetTexture(1);
+						Tex->GetResourceSize(EResourceSizeMode::Exclusive);
+					}
+					if (LightMap->GetSkyOcclusionTexture())
+					{
+						UTexture2D* Tex = LightMap->GetSkyOcclusionTexture();
+						Tex->GetResourceSize(EResourceSizeMode::Exclusive);
+					}
+				}
+			}
 		}
 
 		IrrelevantLights = StaticMeshInstanceData->CachedStaticLighting.IrrelevantLights;
diff --git a/Engine/Source/Runtime/Engine/Private/LightMap.cpp b/Engine/Source/Runtime/Engine/Private/LightMap.cpp
--- a/Engine/Source/Runtime/Engine/Private/LightMap.cpp
+++ b/Engine/Source/Runtime/Engine/Private/LightMap.cpp
@@ -1591,6 +1591,11 @@ UTexture2D* FLightMap2D::GetTexture(uint32 BasisIndex)
 	return Textures[BasisIndex];
 }
 
+UTexture2D* FLightMap2D::GetSkyOcclusionTexture()
+{
+	return SkyOcclusionTexture;
+}
+
 /**
  * Returns whether the specified basis has a valid lightmap texture or not.
  * @param	BasisIndex - The basis index.
diff --git a/Engine/Source/Runtime/Engine/Public/LightMap.h b/Engine/Source/Runtime/Engine/Public/LightMap.h
--- a/Engine/Source/Runtime/Engine/Public/LightMap.h
+++ b/Engine/Source/Runtime/Engine/Public/LightMap.h
@@ -211,6 +211,12 @@ public:
 	ENGINE_API const UTexture2D* GetTexture(uint32 BasisIndex) const;
 	ENGINE_API UTexture2D* GetTexture(uint32 BasisIndex);
 
+	/**
+	 * Returns SkyOcclusionTexture.
+	 * @return	The SkyOcclusionTexture.
+	 */
+	ENGINE_API UTexture2D* GetSkyOcclusionTexture();
+
 	/**
 	 * Returns whether the specified basis has a valid lightmap texture or not.
 	 * @param	BasisIndex - The basis index.
