When a checkpoint is missing, it's no longer catastrophic

diff --git a/Engine/Source/Runtime/Engine/Private/DemoNetDriver.cpp b/Engine/Source/Runtime/Engine/Private/DemoNetDriver.cpp
--- a/Engine/Source/Runtime/Engine/Private/DemoNetDriver.cpp
+++ b/Engine/Source/Runtime/Engine/Private/DemoNetDriver.cpp
@@ -716,7 +716,7 @@ static void SerializeGuidCache( TSharedPtr< class FNetGUIDCache > GuidCache, FAr
 
 	*CheckpointArchive << NumValues;
 
-	UE_LOG( LogDemo, Warning, TEXT( "Checkpoint. SerializeGuidCache: %i" ), NumValues );
+	UE_LOG( LogDemo, Verbose, TEXT( "Checkpoint. SerializeGuidCache: %i" ), NumValues );
 
 	for ( auto It = GuidCache->ObjectLookup.CreateIterator(); It; ++It )
 	{
@@ -841,7 +841,7 @@ void UDemoNetDriver::SaveCheckpoint()
 
 	const float CheckpointTimeInMS = ( EndCheckpointTime - StartCheckpointTime ) * 1000.0f;
 
-	UE_LOG( LogDemo, Warning, TEXT( "Checkpoint. Total: %i, Rep size: %i, PackageMap: %u, Time: %2.2f" ), TotalSize, CheckpointSize, GuidCacheSize, CheckpointTimeInMS );
+	UE_LOG( LogDemo, Log, TEXT( "Checkpoint. Total: %i, Rep size: %i, PackageMap: %u, Time: %2.2f" ), TotalSize, CheckpointSize, GuidCacheSize, CheckpointTimeInMS );
 }
 
 void UDemoNetDriver::TickDemoRecord( float DeltaSeconds )
@@ -1257,7 +1257,7 @@ void UDemoNetDriver::CheckpointReady( bool bSuccess )
 {
 	if ( !bSuccess )
 	{
-		GetWorld()->GetGameInstance()->HandleDemoPlaybackFailure( EDemoPlayFailure::DemoNotFound, FString( EDemoPlayFailure::ToString( EDemoPlayFailure::DemoNotFound ) ) );
+		UE_LOG( LogDemo, Warning, TEXT( "UDemoNetConnection::CheckpointReady: Failed to go to checkpoint." ) );
 		return;
 	}
 
@@ -1265,6 +1265,7 @@ void UDemoNetDriver::CheckpointReady( bool bSuccess )
 
 	if ( CheckpointArchive->TotalSize() == 0 )
 	{
+		UE_LOG( LogDemo, Warning, TEXT( "UDemoNetConnection::CheckpointReady: Checkpoint is empty." ) );
 		return;
 	}
 
