Fix bad perf when changing rendering settings

https://jira.ol.epicgames.net/browse/UE-12207

diff --git a/Engine/Source/Runtime/Engine/Private/RenderingSettings.cpp b/Engine/Source/Runtime/Engine/Private/RenderingSettings.cpp
--- a/Engine/Source/Runtime/Engine/Private/RenderingSettings.cpp
+++ b/Engine/Source/Runtime/Engine/Private/RenderingSettings.cpp
@@ -102,7 +102,7 @@ void URendererSettings::PostEditChangeProperty(FPropertyChangedEvent& PropertyCh
 
 		// If the renderer settings are updated and saved, we need to update the user interface settings too, to prevent data loss
 		// from old curve data for UI settings.
-		if ( IsTemplate() )
+		if ( IsTemplate() && PropertyChangedEvent.ChangeType != EPropertyChangeType::Interactive )
 		{
 			UUserInterfaceSettings* UISettings = GetMutableDefault<UUserInterfaceSettings>(UUserInterfaceSettings::StaticClass());
 			UISettings->UpdateDefaultConfigFile();
