Fix confusion about override materials vs materials list in details panel

https://jira.ol.epicgames.net/browse/UE-11858

diff --git a/Engine/Source/Runtime/Engine/Classes/Components/MeshComponent.h b/Engine/Source/Runtime/Engine/Classes/Components/MeshComponent.h
--- a/Engine/Source/Runtime/Engine/Classes/Components/MeshComponent.h
+++ b/Engine/Source/Runtime/Engine/Classes/Components/MeshComponent.h
@@ -20,7 +20,7 @@ class ENGINE_API UMeshComponent : public UPrimitiveComponent
 	GENERATED_UCLASS_BODY()
 
 	/** Per-Component material overrides.  These must NOT be set directly or a race condition can occur between GC and the rendering thread. */
-	UPROPERTY(EditAnywhere, Category=Rendering, Meta=(ToolTip="Material overrides."))
+	UPROPERTY(EditAnywhere, AdvancedDisplay, Category=Rendering, Meta=(ToolTip="Material overrides."))
 	TArray<class UMaterialInterface*> OverrideMaterials;
 	
 	UFUNCTION(BlueprintCallable, Category="Components|Mesh")
