Fix issue where montages with a play rate of -1 and a current position of 0 would not progress.

diff --git a/Engine/Source/Runtime/Engine/Private/Animation/AnimMontage.cpp b/Engine/Source/Runtime/Engine/Private/Animation/AnimMontage.cpp
--- a/Engine/Source/Runtime/Engine/Private/Animation/AnimMontage.cpp
+++ b/Engine/Source/Runtime/Engine/Private/Animation/AnimMontage.cpp
@@ -1261,10 +1261,12 @@ void FAnimMontageInstance::Advance(float DeltaTime, struct FRootMotionMovementPa
 					float PrevPosition = Position;
 					Position = FMath::Clamp<float>(Position + ActualDeltaPos, CurrentSection.GetTime(), CurrentSectionEndPos);
 
-					if( FMath::Abs(ActualDeltaMove) > 0.f )
+					const float PositionBeforeFiringEvents = Position;
+
+					if (FMath::Abs(ActualDeltaMove) > 0.f)
 					{
 						// Extract Root Motion for this time slice, and accumulate it.
-						if( bExtractRootMotion && AnimInstance.IsValid() )
+						if (bExtractRootMotion && AnimInstance.IsValid())
 						{
 							FTransform RootMotion = Montage->ExtractRootMotionFromTrackRange(PrevPosition, Position);
 							if (bBlendRootMotion)
@@ -1280,13 +1282,13 @@ void FAnimMontageInstance::Advance(float DeltaTime, struct FRootMotionMovementPa
 						}
 
 						// If about to reach the end of the montage...
-						if( NextSectionIndex == INDEX_NONE )
+						if (NextSectionIndex == INDEX_NONE)
 						{
 							// ... trigger blend out if within blend out time window.
 							const float DeltaPosToEnd = bPlayingForward ? (CurrentSectionLength - PosInSection) : PosInSection;
 							const float DeltaTimeToEnd = DeltaPosToEnd / CombinedPlayRate;
 							const float BlendOutTime = FMath::Max<float>(Montage->BlendOutTime * DefaultBlendTimeMultiplier, KINDA_SMALL_NUMBER);
-							if( DeltaTimeToEnd <= BlendOutTime )
+							if (DeltaTimeToEnd <= BlendOutTime)
 							{
 								Stop(DeltaTimeToEnd, false);
 							}
@@ -1297,35 +1299,30 @@ void FAnimMontageInstance::Advance(float DeltaTime, struct FRootMotionMovementPa
 						// first need to have event handler to handle it
 
 						// Save position before firing events.
-						const float PositionBeforeFiringEvents = Position;
-						if( !bInterrupted )
+						if (!bInterrupted)
 						{
 							HandleEvents(PrevPosition, Position, BranchingPointMarker);
 						}
 
-						// if we reached end of section, and we were not processing a branching point, and no events has messed with out current position..
-						// .. Move to next section.
-						// (this also handles looping, the same as jumping to a different section).
-						if ((AdvanceType != ETAA_Default) && !BranchingPointMarker && (PositionBeforeFiringEvents == Position))
-						{
-							// Get recent NextSectionIndex in case it's been changed by previous events.
-							const int32 RecentNextSectionIndex = bPlayingForward ? NextSections[CurrentSectionIndex] : PrevSections[CurrentSectionIndex];
+					}
+					// if we reached end of section, and we were not processing a branching point, and no events has messed with out current position..
+					// .. Move to next section.
+					// (this also handles looping, the same as jumping to a different section).
+					if ((AdvanceType != ETAA_Default) && !BranchingPointMarker && (PositionBeforeFiringEvents == Position))
+					{
+						// Get recent NextSectionIndex in case it's been changed by previous events.
+						const int32 RecentNextSectionIndex = bPlayingForward ? NextSections[CurrentSectionIndex] : PrevSections[CurrentSectionIndex];
 
-							if( RecentNextSectionIndex != INDEX_NONE )
-							{
-								float LatestNextSectionStartTime, LatestNextSectionEndTime;
-								Montage->GetSectionStartAndEndTime(RecentNextSectionIndex, LatestNextSectionStartTime, LatestNextSectionEndTime);
+						if( RecentNextSectionIndex != INDEX_NONE )
+						{
+							float LatestNextSectionStartTime, LatestNextSectionEndTime;
+							Montage->GetSectionStartAndEndTime(RecentNextSectionIndex, LatestNextSectionStartTime, LatestNextSectionEndTime);
 
-								// Jump to next section's appropriate starting point (start or end).
-								float EndOffset = KINDA_SMALL_NUMBER/2.f; //KINDA_SMALL_NUMBER/2 because we use KINDA_SMALL_NUMBER to offset notifies for triggering and SMALL_NUMBER is too small
-								Position = bPlayingForward ? LatestNextSectionStartTime : (LatestNextSectionEndTime - EndOffset);
-							}
+							// Jump to next section's appropriate starting point (start or end).
+							float EndOffset = KINDA_SMALL_NUMBER/2.f; //KINDA_SMALL_NUMBER/2 because we use KINDA_SMALL_NUMBER to offset notifies for triggering and SMALL_NUMBER is too small
+							Position = bPlayingForward ? LatestNextSectionStartTime : (LatestNextSectionEndTime - EndOffset);
 						}
 					}
-					else
-					{
-						break;
-					}
 				}
 				else
 				{
