Merge from John Barrett: Exposed UPackageMapClient as ENGINE_API.

diff --git a/Engine/Source/Runtime/Engine/Classes/Engine/PackageMapClient.h b/Engine/Source/Runtime/Engine/Classes/Engine/PackageMapClient.h
--- a/Engine/Source/Runtime/Engine/Classes/Engine/PackageMapClient.h
+++ b/Engine/Source/Runtime/Engine/Classes/Engine/PackageMapClient.h
@@ -88,7 +88,7 @@ public:
 };
 
 UCLASS(transient)
-class UPackageMapClient : public UPackageMap
+class ENGINE_API UPackageMapClient : public UPackageMap
 {
 public:
 	GENERATED_BODY()
