PR #946: Fixing PVS-Studio warning (Contributed by [REDACTED_PERSON])

https://github.com/EpicGames/UnrealEngine/pull/946

I have gone through these changes and verified that whatever behavior change they cause could only be an improvement (generally they just prevent possible access to nullptrs, but there are also some style corrections). A separate changelist is being prepared that includes changes that could unintentionally caues behavior change for the worse.

diff --git a/Engine/Source/Developer/GammaUI/Private/GammaUIPanel.cpp b/Engine/Source/Developer/GammaUI/Private/GammaUIPanel.cpp
--- a/Engine/Source/Developer/GammaUI/Private/GammaUIPanel.cpp
+++ b/Engine/Source/Developer/GammaUI/Private/GammaUIPanel.cpp
@@ -44,8 +44,7 @@ void SGammaUIPanel::Construct(const SGammaUIPanel::FArguments& InArgs)
 
 float SGammaUIPanel::OnGetGamma() const
 {
-	float DisplayGamma = GEngine->DisplayGamma;
-	return GEngine ? DisplayGamma : 2.2f;
+	return GEngine ? GEngine->DisplayGamma : 2.2f;
 }
 
 void SGammaUIPanel::OnGammaChanged(float NewValue)
diff --git a/Engine/Source/Editor/BehaviorTreeEditor/Private/BehaviorTreeDebugger.cpp b/Engine/Source/Editor/BehaviorTreeEditor/Private/BehaviorTreeDebugger.cpp
--- a/Engine/Source/Editor/BehaviorTreeEditor/Private/BehaviorTreeDebugger.cpp
+++ b/Engine/Source/Editor/BehaviorTreeEditor/Private/BehaviorTreeDebugger.cpp
@@ -1112,7 +1112,7 @@ void FBehaviorTreeDebugger::OnInstanceSelectedInDropdown(UBehaviorTreeComponent*
 
 		AController* OldController = TreeInstance.IsValid() ? Cast<AController>(TreeInstance->GetOwner()) : NULL;
 		APawn* OldPawn = OldController != NULL ? OldController->GetPawn() : NULL;
-		USelection* SelectedActors = GEditor->GetSelectedActors();
+		USelection* SelectedActors = GEditor ? GEditor->GetSelectedActors() : NULL;
 		if (SelectedActors && OldPawn)
 		{
 			SelectedActors->Deselect(OldPawn);
@@ -1120,7 +1120,7 @@ void FBehaviorTreeDebugger::OnInstanceSelectedInDropdown(UBehaviorTreeComponent*
 
 		TreeInstance = SelectedInstance;
 
-		if (SelectedActors && GEditor && SelectedInstance && SelectedInstance->GetOwner())
+		if (SelectedActors && SelectedInstance && SelectedInstance->GetOwner())
 		{
 			AController* TestController = Cast<AController>(SelectedInstance->GetOwner());
 			APawn* Pawn = TestController != NULL ? TestController->GetPawn() : NULL;
diff --git a/Engine/Source/Editor/DetailCustomizations/Private/PhysicsSettingsDetails.cpp b/Engine/Source/Editor/DetailCustomizations/Private/PhysicsSettingsDetails.cpp
--- a/Engine/Source/Editor/DetailCustomizations/Private/PhysicsSettingsDetails.cpp
+++ b/Engine/Source/Editor/DetailCustomizations/Private/PhysicsSettingsDetails.cpp
@@ -82,7 +82,7 @@ void SPhysicalSurfaceListItem::OnTextChanged(const FText& NewText)
 void SPhysicalSurfaceListItem::NewNameEntered(const FText& NewText, ETextCommit::Type CommitInfo)
 {
 	// Don't digest the number if we just clicked away from the pop-up
-	if((CommitInfo == ETextCommit::OnEnter) || (CommitInfo == ETextCommit::OnUserMovedFocus) || (CommitInfo == ETextCommit::OnUserMovedFocus))
+	if((CommitInfo == ETextCommit::OnEnter) || (CommitInfo == ETextCommit::OnUserMovedFocus))
 	{
 		FString NewName = NewText.ToString();
 		if(NewName.Find(TEXT(" "))==INDEX_NONE)
diff --git a/Engine/Source/Editor/FontEditor/Private/SFontEditorViewport.cpp b/Engine/Source/Editor/FontEditor/Private/SFontEditorViewport.cpp
--- a/Engine/Source/Editor/FontEditor/Private/SFontEditorViewport.cpp
+++ b/Engine/Source/Editor/FontEditor/Private/SFontEditorViewport.cpp
@@ -133,7 +133,7 @@ void FFontEditorViewportClient::Draw(FViewport* Viewport, FCanvas* Canvas)
 					RenderInfo->Renderer->GetThumbnailSize(Texture, 1.0f, Width, Height);
 
 					// Don't draw if we are outside of our range
-					if (YPos + Height >= 0 && YPos <= (int32)Viewport->GetSizeXY().Y)
+					if (YPos + (int32)Height >= 0 && YPos <= (int32)Viewport->GetSizeXY().Y)
 					{
 						// If hit testing, draw a tile instead
 						if (Canvas->IsHitTesting())
diff --git a/Engine/Source/Editor/GraphEditor/Private/KismetPins/SGraphPinIndex.cpp b/Engine/Source/Editor/GraphEditor/Private/KismetPins/SGraphPinIndex.cpp
--- a/Engine/Source/Editor/GraphEditor/Private/KismetPins/SGraphPinIndex.cpp
+++ b/Engine/Source/Editor/GraphEditor/Private/KismetPins/SGraphPinIndex.cpp
@@ -31,10 +31,13 @@ FEdGraphPinType SGraphPinIndex::OnGetPinType() const
 
 void SGraphPinIndex::OnTypeChanged(const FEdGraphPinType& PinType)
 {
-	GraphPinObj->PinType = PinType;
-	// Let the node know that one of its' pins had their pin type changed
-	if (GraphPinObj && GraphPinObj->GetOwningNode())
+	if (GraphPinObj)
 	{
-		GraphPinObj->GetOwningNode()->PinTypeChanged(GraphPinObj);
+		GraphPinObj->PinType = PinType;
+		// Let the node know that one of its' pins had their pin type changed
+		if (GraphPinObj->GetOwningNode())
+		{
+			GraphPinObj->GetOwningNode()->PinTypeChanged(GraphPinObj);
+		}
 	}
 }
diff --git a/Engine/Source/Editor/LandscapeEditor/Private/Tests/LandscapeAutomationTests.cpp b/Engine/Source/Editor/LandscapeEditor/Private/Tests/LandscapeAutomationTests.cpp
--- a/Engine/Source/Editor/LandscapeEditor/Private/Tests/LandscapeAutomationTests.cpp
+++ b/Engine/Source/Editor/LandscapeEditor/Private/Tests/LandscapeAutomationTests.cpp
@@ -81,7 +81,7 @@ bool FBeginModifyLandscapeCommand::Update()
 	//Find a location on the edge of the landscape along the x axis so the default camera can see it in the distance.
 	FVector LandscapeSizePerComponent = LandscapeEdMode->UISettings->NewLandscape_QuadsPerSection * LandscapeEdMode->UISettings->NewLandscape_SectionsPerComponent * LandscapeEdMode->UISettings->NewLandscape_Scale;
 	FVector TargetLoctaion(0);
-	TargetLoctaion.X = -LandscapeSizePerComponent.X * (LandscapeEdMode->UISettings->NewLandscape_ComponentCount.X / 2);
+	TargetLoctaion.X = -LandscapeSizePerComponent.X * (LandscapeEdMode->UISettings->NewLandscape_ComponentCount.X / 2.f);
 
 	ALandscapeProxy* Proxy = LandscapeEdMode->CurrentToolTarget.LandscapeInfo.Get()->GetCurrentLevelLandscapeProxy(true);
 	if (Proxy)
diff --git a/Engine/Source/Editor/LevelEditor/Private/LevelEditorContextMenu.cpp b/Engine/Source/Editor/LevelEditor/Private/LevelEditorContextMenu.cpp
--- a/Engine/Source/Editor/LevelEditor/Private/LevelEditorContextMenu.cpp
+++ b/Engine/Source/Editor/LevelEditor/Private/LevelEditorContextMenu.cpp
@@ -534,7 +534,7 @@ void FLevelEditorContextMenu::BuildGroupMenu( FMenuBuilder& MenuBuilder, const F
 void FLevelEditorContextMenuImpl::FillSelectActorMenu( FMenuBuilder& MenuBuilder )
 {
 	FText SelectAllActorStr = FText::Format( LOCTEXT("SelectActorsOfSameClass", "Select All {0}(s)"), FText::FromString( SelectionInfo.SelectionStr ) );
-	int32 NumSelectedSurfaces = NumSelectedSurfaces = AssetSelectionUtils::GetNumSelectedSurfaces( SelectionInfo.SharedWorld );
+	int32 NumSelectedSurfaces = AssetSelectionUtils::GetNumSelectedSurfaces( SelectionInfo.SharedWorld );
 
 	MenuBuilder.BeginSection("SelectActorGeneral", LOCTEXT("SelectAnyHeading", "General") );
 	{
diff --git a/Engine/Source/Editor/Matinee/Private/MatineeDraw.cpp b/Engine/Source/Editor/Matinee/Private/MatineeDraw.cpp
--- a/Engine/Source/Editor/Matinee/Private/MatineeDraw.cpp
+++ b/Engine/Source/Editor/Matinee/Private/MatineeDraw.cpp
@@ -1438,8 +1438,7 @@ void FMatineeViewportClient::CreatePushPropertiesOntoGraphButton( FCanvas* Canva
 		Track->IsA(UInterpTrackFloatBase::StaticClass()) ||
 		Track->IsA(UInterpTrackVectorBase::StaticClass()) || 
 		Track->IsA(UInterpTrackMove::StaticClass()) || 
-		Track->IsA(UInterpTrackLinearColorBase::StaticClass()) ||
-		Track->IsA(UInterpTrackVectorBase::StaticClass())
+		Track->IsA(UInterpTrackLinearColorBase::StaticClass())
 	)
 	{
 		UTexture2D* GraphTex = NULL;
diff --git a/Engine/Source/Editor/Sequencer/Private/SequencerActorBindingManager.cpp b/Engine/Source/Editor/Sequencer/Private/SequencerActorBindingManager.cpp
--- a/Engine/Source/Editor/Sequencer/Private/SequencerActorBindingManager.cpp
+++ b/Engine/Source/Editor/Sequencer/Private/SequencerActorBindingManager.cpp
@@ -553,7 +553,7 @@ void FSequencerActorBindingManager::PropagatePuppetActorChanges( const TSharedRe
 			}
 		}
 		
-		if( ensure( FoundSpawnable != NULL ) )
+		if (ensure( FoundSpawnable != NULL && PuppetActor != NULL ))
 		{
 			UClass* ActorClass = FoundSpawnable->GetClass();
 
diff --git a/Engine/Source/Editor/UnrealEd/Classes/Settings/ContentBrowserSettings.h b/Engine/Source/Editor/UnrealEd/Classes/Settings/ContentBrowserSettings.h
--- a/Engine/Source/Editor/UnrealEd/Classes/Settings/ContentBrowserSettings.h
+++ b/Engine/Source/Editor/UnrealEd/Classes/Settings/ContentBrowserSettings.h
@@ -56,7 +56,7 @@ public:
 	/** Gets whether we are allowed to display the developers folder or not, optional flag ignoring the override */
 	bool GetDisplayDevelopersFolder( bool bExcludeOverride = false ) const
 	{ 
-		return ( ( bExcludeOverride ? false : OverrideDisplayDevelopersFolder ) | DisplayDevelopersFolder );
+		return ( ( bExcludeOverride ? false : OverrideDisplayDevelopersFolder ) || DisplayDevelopersFolder );
 	}
 
 	/** Sets whether we are allowed to display the plugin folders or not, optional flag for setting override instead */
@@ -68,7 +68,7 @@ public:
 	/** Gets whether we are allowed to display the plugin folders or not, optional flag ignoring the override */
 	bool GetDisplayPluginFolders( bool bExcludeOverride = false ) const
 	{ 
-		return ( ( bExcludeOverride ? false : OverrideDisplayPluginFolders ) | DisplayPluginFolders );
+		return ( ( bExcludeOverride ? false : OverrideDisplayPluginFolders ) || DisplayPluginFolders );
 	}
 
 	/** Sets whether we are allowed to display collection folders or not */
diff --git a/Engine/Source/Editor/UnrealEd/Private/EdGraphUtilities.cpp b/Engine/Source/Editor/UnrealEd/Private/EdGraphUtilities.cpp
--- a/Engine/Source/Editor/UnrealEd/Private/EdGraphUtilities.cpp
+++ b/Engine/Source/Editor/UnrealEd/Private/EdGraphUtilities.cpp
@@ -236,12 +236,10 @@ void FEdGraphUtilities::MergeChildrenGraphsIn(UEdGraph* MergeTarget, UEdGraph* P
 // Tries to rename the graph to have a name similar to BaseName
 void FEdGraphUtilities::RenameGraphCloseToName(UEdGraph* Graph, const FString& BaseName, int32 StartIndex)
 {
-	bool bFoundName = false;
-
 	FString NewName = BaseName;
 
 	int32 NameIndex = StartIndex;
-	while (!bFoundName)
+	for (;;)
 	{
 		if (Graph->Rename(*NewName, Graph->GetOuter(), REN_Test))
 		{
diff --git a/Engine/Source/Editor/UnrealEd/Private/Editor.cpp b/Engine/Source/Editor/UnrealEd/Private/Editor.cpp
--- a/Engine/Source/Editor/UnrealEd/Private/Editor.cpp
+++ b/Engine/Source/Editor/UnrealEd/Private/Editor.cpp
@@ -5084,7 +5084,7 @@ void UEditorEngine::ReplaceActors(UActorFactory* Factory, const FAssetData& Asse
 			UBlueprint* Blueprint = CastChecked<UBlueprint>(Asset);
 			// Only try to copy properties if this blueprint is based on the actor
 			UClass* OldActorClass = OldActor->GetClass();
-			if (Blueprint->GeneratedClass->IsChildOf(OldActorClass))
+			if (Blueprint->GeneratedClass->IsChildOf(OldActorClass) && NewActor != NULL)
 			{
 				NewActor->UnregisterAllComponents();
 				UEditorEngine::CopyPropertiesForUnrelatedObjects(OldActor, NewActor);
diff --git a/Engine/Source/Editor/UnrealEd/Private/EditorExporters.cpp b/Engine/Source/Editor/UnrealEd/Private/EditorExporters.cpp
--- a/Engine/Source/Editor/UnrealEd/Private/EditorExporters.cpp
+++ b/Engine/Source/Editor/UnrealEd/Private/EditorExporters.cpp
@@ -2364,7 +2364,7 @@ bool UAnimSequenceExporterFBX::ExportBinary( UObject* Object, const TCHAR* Type,
 {
 	UAnimSequence* AnimSequence = CastChecked<UAnimSequence>( Object );
 	USkeleton* AnimSkeleton = AnimSequence->GetSkeleton();
-	USkeletalMesh* PreviewMesh = AnimSkeleton->GetAssetPreviewMesh(AnimSequence);
+	USkeletalMesh* PreviewMesh = AnimSkeleton ? AnimSkeleton->GetAssetPreviewMesh(AnimSequence) : NULL;
 
 	if (AnimSkeleton && PreviewMesh)
 	{
diff --git a/Engine/Source/Editor/UnrealEd/Private/Fbx/FbxAnimationExport.cpp b/Engine/Source/Editor/UnrealEd/Private/Fbx/FbxAnimationExport.cpp
--- a/Engine/Source/Editor/UnrealEd/Private/Fbx/FbxAnimationExport.cpp
+++ b/Engine/Source/Editor/UnrealEd/Private/Fbx/FbxAnimationExport.cpp
@@ -341,9 +341,7 @@ void FFbxExporter::ExportAnimSequencesAsSingle( USkeletalMesh* SkelMesh, const A
  */
 void FFbxExporter::ExportMatineeGroup(class AMatineeActor* MatineeActor, USkeletalMeshComponent* SkeletalMeshComponent)
 {
-	float MatineeLength = MatineeActor->MatineeData->InterpLength;
-	
-	if (Scene == NULL || MatineeActor == NULL || SkeletalMeshComponent == NULL || MatineeLength == 0)
+	if (Scene == NULL || MatineeActor == NULL || SkeletalMeshComponent == NULL || MatineeActor->MatineeData->InterpLength == 0)
 	{
 		return;
 	}
diff --git a/Engine/Source/Editor/UnrealEd/Private/Kismet2/KismetReinstanceUtilities.cpp b/Engine/Source/Editor/UnrealEd/Private/Kismet2/KismetReinstanceUtilities.cpp
--- a/Engine/Source/Editor/UnrealEd/Private/Kismet2/KismetReinstanceUtilities.cpp
+++ b/Engine/Source/Editor/UnrealEd/Private/Kismet2/KismetReinstanceUtilities.cpp
@@ -991,7 +991,7 @@ void FBlueprintCompileReinstancer::ReplaceInstancesOfClass(UClass* OldClass, UCl
 	}
 
 	SelectedActors->EndBatchSelectOperation();
-	if (bSelectionChanged)
+	if (bSelectionChanged && GEditor)
 	{
 		GEditor->NoteSelectionChange();
 	}
diff --git a/Engine/Source/Editor/UnrealEd/Private/LevelEditorViewport.cpp b/Engine/Source/Editor/UnrealEd/Private/LevelEditorViewport.cpp
--- a/Engine/Source/Editor/UnrealEd/Private/LevelEditorViewport.cpp
+++ b/Engine/Source/Editor/UnrealEd/Private/LevelEditorViewport.cpp
@@ -1577,13 +1577,13 @@ FSceneView* FLevelEditorViewportClient::CalcSceneView(FSceneViewFamily* ViewFami
 			FLevelEditorViewportClient* ViewportClient = GEditor->LevelViewportClients[ViewportIndex];
 
 			//only change camera for a viewport that is looking at the same scene
-			if (GetScene() != ViewportClient->GetScene())
+			if (ViewportClient == NULL || GetScene() != ViewportClient->GetScene())
 			{
 				continue;
 			}
 
 			// go over all other level viewports
-			if (ViewportClient && ViewportClient->Viewport && ViewportClient != this)
+			if (ViewportClient->Viewport && ViewportClient != this)
 			{
 				// force camera of same-typed viewports
 				if (ViewportClient->GetViewportType() == GetViewportType())
diff --git a/Engine/Source/Editor/UnrealEd/Private/StaticLightingSystem/StaticLightingSystem.cpp b/Engine/Source/Editor/UnrealEd/Private/StaticLightingSystem/StaticLightingSystem.cpp
--- a/Engine/Source/Editor/UnrealEd/Private/StaticLightingSystem/StaticLightingSystem.cpp
+++ b/Engine/Source/Editor/UnrealEd/Private/StaticLightingSystem/StaticLightingSystem.cpp
@@ -1624,20 +1624,20 @@ void FStaticLightingSystem::AddPrimitiveStaticLightingInfo(FStaticLightingPrimit
 	for(int32 MeshIndex = 0;MeshIndex < PrimitiveInfo.Meshes.Num();MeshIndex++)
 	{
 		FStaticLightingMesh* Mesh = PrimitiveInfo.Meshes[MeshIndex];
-		Mesh->VisibilityIds.Add(PrimitiveInfo.VisibilityId);
-		if (!GLightmassDebugOptions.bSortMappings && bBuildActorLighting)
+		if (Mesh)
 		{
-			if (Mesh)
+			Mesh->VisibilityIds.Add(PrimitiveInfo.VisibilityId);
+			if (!GLightmassDebugOptions.bSortMappings && bBuildActorLighting)
 			{
-				Mesh->Guid = FGuid(0,0,0,DeterministicIndex++);
+				Mesh->Guid = FGuid(0, 0, 0, DeterministicIndex++);
 			}
-		}
-		Meshes.Add(Mesh);
-		LightingMeshBounds += Mesh->BoundingBox;
+			Meshes.Add(Mesh);
+			LightingMeshBounds += Mesh->BoundingBox;
 
-		if (Mesh->bCastShadow)
-		{
-			UpdateAutomaticImportanceVolumeBounds( Mesh->BoundingBox );
+			if (Mesh->bCastShadow)
+			{
+				UpdateAutomaticImportanceVolumeBounds(Mesh->BoundingBox);
+			}
 		}
 	}
 
diff --git a/Engine/Source/Editor/UnrealEd/Private/Tests/EditorBuildPromotionTests.cpp b/Engine/Source/Editor/UnrealEd/Private/Tests/EditorBuildPromotionTests.cpp
--- a/Engine/Source/Editor/UnrealEd/Private/Tests/EditorBuildPromotionTests.cpp
+++ b/Engine/Source/Editor/UnrealEd/Private/Tests/EditorBuildPromotionTests.cpp
@@ -1546,7 +1546,7 @@ namespace BuildPromotionTestHelper
 			}
 
 			//Apply if this is a material.  (Editor won't close unless we do)
-			if (CurrentAsset->IsA(UMaterialInterface::StaticClass()))
+			if (CurrentAsset && CurrentAsset->IsA(UMaterialInterface::StaticClass()))
 			{
 				EditorBuildPromotionTestUtils::SendUpdateMaterialCommand();
 			}
diff --git a/Engine/Source/Runtime/AIModule/Private/AITypes.cpp b/Engine/Source/Runtime/AIModule/Private/AITypes.cpp
--- a/Engine/Source/Runtime/AIModule/Private/AITypes.cpp
+++ b/Engine/Source/Runtime/AIModule/Private/AITypes.cpp
@@ -48,7 +48,7 @@ namespace FAIResources
 
 	void RegisterResource(const FAIResourceID& Resource)
 	{
-		if (FAIResourceID::GetSize() - FAIResources::ResourceIDs.Num() > 0)
+		if (FAIResourceID::GetSize() > static_cast<uint32>(FAIResources::ResourceIDs.Num()))
 		{
 			ResourceIDs.AddZeroed(FAIResourceID::GetSize() - FAIResources::ResourceIDs.Num());
 		}
diff --git a/Engine/Source/Runtime/AIModule/Private/Navigation/PathFollowingComponent.cpp b/Engine/Source/Runtime/AIModule/Private/Navigation/PathFollowingComponent.cpp
--- a/Engine/Source/Runtime/AIModule/Private/Navigation/PathFollowingComponent.cpp
+++ b/Engine/Source/Runtime/AIModule/Private/Navigation/PathFollowingComponent.cpp
@@ -712,7 +712,6 @@ int32 UPathFollowingComponent::OptimizeSegmentVisibility(int32 StartIndex)
 			Path->ShortcutNodeRefs.Reserve(RaycastResult.CorridorPolysCount);
 			Path->ShortcutNodeRefs.SetNumUninitialized(RaycastResult.CorridorPolysCount);
 		}
-		FPlatformMemory::Memcmp(Path->ShortcutNodeRefs.GetData(), RaycastResult.CorridorPolys, RaycastResult.CorridorPolysCount * sizeof(NavNodeRef));
 #else
 		const bool RaycastHitResult = NavData->Raycast(PawnLocation, AdjustedDestination, HitLocation, QueryFilter);
 #endif
diff --git a/Engine/Source/Runtime/Core/Private/Windows/WindowsApplication.cpp b/Engine/Source/Runtime/Core/Private/Windows/WindowsApplication.cpp
--- a/Engine/Source/Runtime/Core/Private/Windows/WindowsApplication.cpp
+++ b/Engine/Source/Runtime/Core/Private/Windows/WindowsApplication.cpp
@@ -50,11 +50,11 @@ FWindowsApplication::FWindowsApplication( const HINSTANCE HInstance, const HICON
 	, bIsMouseAttached( false )
 	, XInput( XInputInterface::Create( MessageHandler ) )
 	, bHasLoadedInputPlugins( false )
+	, bAllowedToDeferMessageProcessing(true)
 	, CVarDeferMessageProcessing( 
 		TEXT( "Slate.DeferWindowsMessageProcessing" ),
 		bAllowedToDeferMessageProcessing,
 		TEXT( "Whether windows message processing is deferred until tick or if they are processed immediately" ) )
-	, bAllowedToDeferMessageProcessing( true )
 	, bInModalSizeLoop( false )
 
 {
@@ -1935,11 +1935,11 @@ void FWindowsApplication::QueryConnectedMice()
 		if (Device.dwType != RIM_TYPEMOUSE)
 			continue;
 		//Force the use of ANSI versions of these calls
-		if (GetRawInputDeviceInfoA(Device.hDevice, RIDI_DEVICENAME, nullptr, &NameLen) < 0)
+		if (GetRawInputDeviceInfoA(Device.hDevice, RIDI_DEVICENAME, nullptr, &NameLen) == static_cast<UINT>(-1))
 			continue;
 
 		Name.Reset(new char[NameLen+1]);
-		if (GetRawInputDeviceInfoA(Device.hDevice, RIDI_DEVICENAME, Name.GetOwnedPointer(), &NameLen) < 0)
+		if (GetRawInputDeviceInfoA(Device.hDevice, RIDI_DEVICENAME, Name.GetOwnedPointer(), &NameLen) == static_cast<UINT>(-1))
 			continue;
 
 		Name[NameLen] = 0;
diff --git a/Engine/Source/Runtime/Core/Private/Windows/WindowsPlatformSplash.cpp b/Engine/Source/Runtime/Core/Private/Windows/WindowsPlatformSplash.cpp
--- a/Engine/Source/Runtime/Core/Private/Windows/WindowsPlatformSplash.cpp
+++ b/Engine/Source/Runtime/Core/Private/Windows/WindowsPlatformSplash.cpp
@@ -312,7 +312,7 @@ HBITMAP LoadSplashBitmap()
 		}
 
 		// Extract the image into the HBITMAP    
-		if (SUCCEEDED(hr))
+		if (SUCCEEDED(hr) && ToRenderBitmapSource)
 		{
 			hr = ToRenderBitmapSource->CopyPixels(
 				NULL,
diff --git a/Engine/Source/Runtime/Core/Public/Windows/WindowsPlatformCodeAnalysis.h b/Engine/Source/Runtime/Core/Public/Windows/WindowsPlatformCodeAnalysis.h
--- a/Engine/Source/Runtime/Core/Public/Windows/WindowsPlatformCodeAnalysis.h
+++ b/Engine/Source/Runtime/Core/Public/Windows/WindowsPlatformCodeAnalysis.h
@@ -92,4 +92,16 @@
 	// so we can get the benefits of analysis on cross-platform code sooner.
 	#pragma warning(disable : 6320) // warning C6320: Exception-filter expression is the constant EXCEPTION_EXECUTE_HANDLER. This might mask exceptions that were not intended to be handled.
 
+	//PVS-Studio settings:
+	//-V::505
+	//-V:TRYCOMPRESSION:519,547
+	//-V:check:547
+	//-V:ENABLE_TEXT_ERROR_CHECKING_RESULTS:560
+	//-V:ENABLE_LOC_TESTING:560
+	//-V:WITH_EDITOR:560
+	//-V:UE_LOG_ACTIVE:560
+	//-V:verify:501
+	//-V:check:501
+	//-V::677,720,704
+
 #endif
\ No newline at end of file
diff --git a/Engine/Source/Runtime/CoreUObject/Private/UObject/Class.cpp b/Engine/Source/Runtime/CoreUObject/Private/UObject/Class.cpp
--- a/Engine/Source/Runtime/CoreUObject/Private/UObject/Class.cpp
+++ b/Engine/Source/Runtime/CoreUObject/Private/UObject/Class.cpp
@@ -805,7 +805,7 @@ void UStruct::SerializeTaggedProperties(FArchive& Ar, uint8* Data, UStruct* Defa
 				if (InAr.UE4Ver() < VER_UE4_STRUCT_GUID_IN_PROPERTY_TAG)
 				{
 					// Old Implementation
-					return !StructProperty->UseBinaryOrNativeSerialization(InAr);
+					return StructProperty && !StructProperty->UseBinaryOrNativeSerialization(InAr);
 				}
 				return PropertyTag.StructGuid.IsValid() && StructProperty && StructProperty->Struct && (PropertyTag.StructGuid == StructProperty->Struct->GetCustomGuid());
 			};
diff --git a/Engine/Source/Runtime/CoreUObject/Private/UObject/SavePackage.cpp b/Engine/Source/Runtime/CoreUObject/Private/UObject/SavePackage.cpp
--- a/Engine/Source/Runtime/CoreUObject/Private/UObject/SavePackage.cpp
+++ b/Engine/Source/Runtime/CoreUObject/Private/UObject/SavePackage.cpp
@@ -715,10 +715,9 @@ public:
 		FString TmpFilename = FPaths::GetPath( DstFilename ) / ( FPaths::GetBaseFilename( DstFilename ) + TEXT( "_SaveCompressed.tmp" ));
 		// Create file reader and writer...
 		FMemoryReader Reader(*(FBufferArchive*)(SrcLinker->Saver),true);
-		FArchive* FileReader = &Reader;
 		FArchive* FileWriter = IFileManager::Get().CreateFileWriter( *TmpFilename );
 		// ... and abort if either operation wasn't successful.
-		if( !FileReader || !FileWriter )
+		if( !FileWriter )
 		{
 			// Delete potentially created reader or writer.
 			delete FileWriter;
@@ -729,7 +728,7 @@ public:
 		}
 
 
-		CompressArchive(FileReader,FileWriter,SrcLinker);
+		CompressArchive(&Reader,FileWriter,SrcLinker);
 
 		// Tear down file writer. This will flush the writer first.
 		delete FileWriter;
diff --git a/Engine/Source/Runtime/CoreUObject/Private/UObject/UObjectBase.cpp b/Engine/Source/Runtime/CoreUObject/Private/UObject/UObjectBase.cpp
--- a/Engine/Source/Runtime/CoreUObject/Private/UObject/UObjectBase.cpp
+++ b/Engine/Source/Runtime/CoreUObject/Private/UObject/UObjectBase.cpp
@@ -253,18 +253,18 @@ bool UObjectBase::IsValidLowLevelFast(bool bRecursive /*= true*/) const
 		UE_LOG(LogUObjectBase, Error, TEXT("Object flags are invalid or either Class or Outer is misaligned"));
 		return false;
 	}
-	// Avoid infinite recursion so call IsValidLowLevelFast on the class object with bRecirsive = false.
-	if (bRecursive && !Class->IsValidLowLevelFast(false))
-	{
-		UE_LOG(LogUObjectBase, Error, TEXT("Class object failed IsValidLowLevelFast test."));
-		return false;
-	}
 	// These should all be non-NULL (except CDO-alignment check which should be 0)
 	if (Class == NULL || Class->ClassDefaultObject == NULL || ((UPTRINT)Class->ClassDefaultObject & AlignmentCheck) != 0)
 	{
 		UE_LOG(LogUObjectBase, Error, TEXT("Class pointer is invalid or CDO is invalid."));
 		return false;
 	}
+	// Avoid infinite recursion so call IsValidLowLevelFast on the class object with bRecirsive = false.
+	if (bRecursive && !Class->IsValidLowLevelFast(false))
+	{
+		UE_LOG(LogUObjectBase, Error, TEXT("Class object failed IsValidLowLevelFast test."));
+		return false;
+	}
 	// Lightweight versions of index checks.
 	if (!GUObjectArray.IsValidIndex(this) || !Name.IsValidIndexFast())
 	{
@@ -809,11 +809,18 @@ void UObjectBaseShutdown()
  */
 const TCHAR* DebugFName(UObject* Object)
 {
-	// Hardcoded static array. This function is only used inside the debugger so it should be fine to return it.
-	static TCHAR TempName[256];
-	FName Name = Object->GetFName();
-	FCString::Strcpy(TempName,Object ? *FName::SafeString(Name.GetDisplayIndex(), Name.GetNumber()) : TEXT("NULL"));
-	return TempName;
+	if ( Object )
+	{
+		// Hardcoded static array. This function is only used inside the debugger so it should be fine to return it.
+		static TCHAR TempName[256];
+		FName Name = Object->GetFName();
+		FCString::Strcpy(TempName, *FName::SafeString(Name.GetDisplayIndex(), Name.GetNumber()));
+		return TempName;
+	}
+	else
+	{
+		return TEXT("NULL");
+	}
 }
 
 /**
diff --git a/Engine/Source/Runtime/CoreUObject/Private/UObject/UObjectGlobals.cpp b/Engine/Source/Runtime/CoreUObject/Private/UObject/UObjectGlobals.cpp
--- a/Engine/Source/Runtime/CoreUObject/Private/UObject/UObjectGlobals.cpp
+++ b/Engine/Source/Runtime/CoreUObject/Private/UObject/UObjectGlobals.cpp
@@ -2557,35 +2557,37 @@ private:
 		if( InObject )
 		{
 			UObject* Object = const_cast<UObject*>(InObject);
-			// Remove references to pending kill objects if we're allowed to do so.
-			UObject* ObjectToAdd = Object->HasAnyFlags( RF_Unreachable ) ? Object : NULL;
-			if (Object && Object->HasAnyFlags(ReferenceSearchFlags))
+			if ( Object )
 			{
-				// Stop recursing, and add to the list of references
-				if (FoundReferencesList)
+				// Remove references to pending kill objects if we're allowed to do so.
+				UObject* ObjectToAdd = Object->HasAnyFlags(RF_Unreachable) ? Object : NULL;
+				if (Object->HasAnyFlags(ReferenceSearchFlags))
 				{
-					if (!CurrentReferenceInfo)
+					// Stop recursing, and add to the list of references
+					if (FoundReferencesList)
 					{
-						CurrentReferenceInfo = new(FoundReferencesList->ExternalReferences) FReferencerInformation(CurrentObject);
+						if (!CurrentReferenceInfo)
+						{
+							CurrentReferenceInfo = new(FoundReferencesList->ExternalReferences) FReferencerInformation(CurrentObject);
+						}
+						if (InReferencingObject != NULL)
+						{
+							CurrentReferenceInfo->ReferencingProperties.AddUnique(dynamic_cast<const UProperty*>(InReferencingObject));
+						}
+						CurrentReferenceInfo->TotalReferences++;
 					}
-					if ( InReferencingObject != NULL )
+					if (ObjectToAdd)
 					{
-						CurrentReferenceInfo->ReferencingProperties.AddUnique(dynamic_cast<const UProperty*>(InReferencingObject));
+						// Mark it as reachable.
+						Object->ClearFlags(RF_Unreachable);
 					}
-					CurrentReferenceInfo->TotalReferences++;
 				}
-				if (ObjectToAdd)
+				else if (ObjectToAdd)
 				{
-					// Mark it as reachable.
-					Object->ClearFlags( RF_Unreachable );
+					// Add encountered object reference to list of to be serialized objects if it hasn't already been added.
+					AddToObjectList(InReferencingObject, InReferencingProperty, ObjectToAdd);
 				}
 			}
-			else if (ObjectToAdd)
-			{
-				// Add encountered object reference to list of to be serialized objects if it hasn't already been added.
-				AddToObjectList(InReferencingObject, InReferencingProperty, ObjectToAdd);
-			}
-
 		}
 	}
 
diff --git a/Engine/Source/Runtime/Engine/Private/Animation/AnimSequence.cpp b/Engine/Source/Runtime/Engine/Private/Animation/AnimSequence.cpp
--- a/Engine/Source/Runtime/Engine/Private/Animation/AnimSequence.cpp
+++ b/Engine/Source/Runtime/Engine/Private/Animation/AnimSequence.cpp
@@ -1839,77 +1839,80 @@ void FillUpTransformBasedOnRig(USkeleton* Skeleton, TArray<FTransform>& NodeSpac
 
 	const URig* Rig = Skeleton->GetRig();
 
-	// this one has to collect all Nodes in Rig data
-	// since we're comparing two of them together. 
-	int32 NodeNum = Rig->GetNodeNum();
-
-	if (Rig && NodeNum > 0)
+	if (Rig)
 	{
-		NodeSpaceBases.Empty(NodeNum);
-		NodeSpaceBases.AddUninitialized(NodeNum);
+		// this one has to collect all Nodes in Rig data
+		// since we're comparing two of them together. 
+		int32 NodeNum = Rig->GetNodeNum();
 
-		Rotations.Empty(NodeNum);
-		Rotations.AddUninitialized(NodeNum);
+		if (NodeNum > 0)
+		{
+			NodeSpaceBases.Empty(NodeNum);
+			NodeSpaceBases.AddUninitialized(NodeNum);
 
-		Translations.Empty(NodeNum);
-		Translations.AddUninitialized(NodeNum);
+			Rotations.Empty(NodeNum);
+			Rotations.AddUninitialized(NodeNum);
 
-		TranslationParentFlags.Empty(Translations.Num());
-		TranslationParentFlags.AddZeroed(Translations.Num());
-		
-		const USkeletalMesh* PreviewMesh = Skeleton->GetPreviewMesh();
+			Translations.Empty(NodeNum);
+			Translations.AddUninitialized(NodeNum);
 
-		for ( int32 Index = 0; Index < NodeNum; ++Index )
-		{
-			const FName NodeName = Rig->GetNodeName(Index);
-			const FName& BoneName = Skeleton->GetRigBoneMapping(NodeName);
-			const int32& BoneIndex = FindMeshBoneIndexFromBoneName(Skeleton, BoneName);
-			
-			if (BoneIndex == INDEX_NONE)
-			{
-				// add identity
-				NodeSpaceBases[Index].SetIdentity();
-				Rotations[Index].SetIdentity();
-				Translations[Index].SetIdentity();
-			}
-			else
-			{
-				// initialize with SpaceBases - assuming World Based
-				NodeSpaceBases[Index] = SpaceBases[BoneIndex];
-				Rotations[Index] = SpaceBases[BoneIndex];
-				Translations[Index] = SpaceBases[BoneIndex];
+			TranslationParentFlags.Empty(Translations.Num());
+			TranslationParentFlags.AddZeroed(Translations.Num());
 
-				const FTransformBase* TransformBase = Rig->GetTransformBaseByNodeName(NodeName);
+			const USkeletalMesh* PreviewMesh = Skeleton->GetPreviewMesh();
 
-				if (TransformBase != NULL)
+			for (int32 Index = 0; Index < NodeNum; ++Index)
+			{
+				const FName NodeName = Rig->GetNodeName(Index);
+				const FName& BoneName = Skeleton->GetRigBoneMapping(NodeName);
+				const int32& BoneIndex = FindMeshBoneIndexFromBoneName(Skeleton, BoneName);
+
+				if (BoneIndex == INDEX_NONE)
+				{
+					// add identity
+					NodeSpaceBases[Index].SetIdentity();
+					Rotations[Index].SetIdentity();
+					Translations[Index].SetIdentity();
+				}
+				else
 				{
-					// orientation constraint			
-					const auto& RotConstraint = TransformBase->Constraints[EControlConstraint::Type::Orientation];
+					// initialize with SpaceBases - assuming World Based
+					NodeSpaceBases[Index] = SpaceBases[BoneIndex];
+					Rotations[Index] = SpaceBases[BoneIndex];
+					Translations[Index] = SpaceBases[BoneIndex];
+
+					const FTransformBase* TransformBase = Rig->GetTransformBaseByNodeName(NodeName);
 
-					if (RotConstraint.TransformConstraints.Num() > 0)
+					if (TransformBase != NULL)
 					{
-						const FName& ParentBoneName = Skeleton->GetRigBoneMapping(RotConstraint.TransformConstraints[0].ParentSpace);
-						const int32& ParentBoneIndex = FindMeshBoneIndexFromBoneName(Skeleton, ParentBoneName);
+						// orientation constraint			
+						const auto& RotConstraint = TransformBase->Constraints[EControlConstraint::Type::Orientation];
 
-						if (ParentBoneIndex != INDEX_NONE)
+						if (RotConstraint.TransformConstraints.Num() > 0)
 						{
-							Rotations[Index] = SpaceBases[BoneIndex].GetRelativeTransform(SpaceBases[ParentBoneIndex]);
-						}
-					}
+							const FName& ParentBoneName = Skeleton->GetRigBoneMapping(RotConstraint.TransformConstraints[0].ParentSpace);
+							const int32& ParentBoneIndex = FindMeshBoneIndexFromBoneName(Skeleton, ParentBoneName);
 
-					// translation constraint
-					const auto& TransConstraint = TransformBase->Constraints[EControlConstraint::Type::Translation];
+							if (ParentBoneIndex != INDEX_NONE)
+							{
+								Rotations[Index] = SpaceBases[BoneIndex].GetRelativeTransform(SpaceBases[ParentBoneIndex]);
+							}
+						}
 
-					if (TransConstraint.TransformConstraints.Num() > 0)
-					{
-						const FName& ParentBoneName = Skeleton->GetRigBoneMapping(TransConstraint.TransformConstraints[0].ParentSpace);
-						const int32& ParentBoneIndex = FindMeshBoneIndexFromBoneName(Skeleton, ParentBoneName);
+						// translation constraint
+						const auto& TransConstraint = TransformBase->Constraints[EControlConstraint::Type::Translation];
 
-						if (ParentBoneIndex != INDEX_NONE)
+						if (TransConstraint.TransformConstraints.Num() > 0)
 						{
-							// I think translation has to include rotation, otherwise it won't work
-							Translations[Index] = SpaceBases[BoneIndex].GetRelativeTransform(SpaceBases[ParentBoneIndex]);
-							TranslationParentFlags[Index] = true;
+							const FName& ParentBoneName = Skeleton->GetRigBoneMapping(TransConstraint.TransformConstraints[0].ParentSpace);
+							const int32& ParentBoneIndex = FindMeshBoneIndexFromBoneName(Skeleton, ParentBoneName);
+
+							if (ParentBoneIndex != INDEX_NONE)
+							{
+								// I think translation has to include rotation, otherwise it won't work
+								Translations[Index] = SpaceBases[BoneIndex].GetRelativeTransform(SpaceBases[ParentBoneIndex]);
+								TranslationParentFlags[Index] = true;
+							}
 						}
 					}
 				}
diff --git a/Engine/Source/Runtime/Engine/Private/LevelActor.cpp b/Engine/Source/Runtime/Engine/Private/LevelActor.cpp
--- a/Engine/Source/Runtime/Engine/Private/LevelActor.cpp
+++ b/Engine/Source/Runtime/Engine/Private/LevelActor.cpp
@@ -866,11 +866,10 @@ void UWorld::RefreshStreamingLevels( const TArray<class ULevelStreaming*>& InLev
 		for( int32 LevelIndex=0; LevelIndex<InLevelsToRefresh.Num(); LevelIndex++ )
 		{
 			ULevelStreaming* StreamingLevel = InLevelsToRefresh[LevelIndex];
-			ULevel* LoadedLevel = StreamingLevel->GetLoadedLevel();
+			ULevel* LoadedLevel = StreamingLevel ? StreamingLevel->GetLoadedLevel() : nullptr;
 
-			if( StreamingLevel 
-				&&	LoadedLevel 
-				&&	LoadedLevel->bIsVisible )
+			if( LoadedLevel &&
+				LoadedLevel->bIsVisible )
 			{
 				RemoveFromWorld( LoadedLevel );
 			}
diff --git a/Engine/Source/Runtime/Engine/Private/ModelLight.cpp b/Engine/Source/Runtime/Engine/Private/ModelLight.cpp
--- a/Engine/Source/Runtime/Engine/Private/ModelLight.cpp
+++ b/Engine/Source/Runtime/Engine/Private/ModelLight.cpp
@@ -1055,7 +1055,7 @@ void UModel::ApplyStaticLighting()
 					uint32 BaseY = SurfaceInfo.BaseY - SurfaceStaticLighting->MappedRect.Min.Y;
 					if (GLightmassDebugOptions.bPadMappings && GAllowLightmapPadding)
 					{
-						if ((PaddedSizeX - 2 > 0) && ((PaddedSizeY - 2) > 0))
+						if ((PaddedSizeX > 2) && (PaddedSizeY > 2))
 						{
 							PaddedSizeX -= 2;
 							PaddedSizeY -= 2;
diff --git a/Engine/Source/Runtime/Engine/Private/World.cpp b/Engine/Source/Runtime/Engine/Private/World.cpp
--- a/Engine/Source/Runtime/Engine/Private/World.cpp
+++ b/Engine/Source/Runtime/Engine/Private/World.cpp
@@ -3080,7 +3080,7 @@ void UWorld::CleanupWorld(bool bSessionEnded, bool bCleanupResources, UWorld* Ne
 #if WITH_EDITOR
 	// If we're server traveling, we need to break the reference dependency here (caused by levelscript)
 	// to avoid a GC crash for not cleaning up the gameinfo referenced by levelscript
-	if( IsGameWorld() && !GIsEditor && !IsRunningCommandlet() && bSessionEnded && bCleanupResources )
+	if (IsGameWorld() && !GIsEditor && !IsRunningCommandlet() && bSessionEnded && bCleanupResources && PersistentLevel)
 	{
 		if( ULevelScriptBlueprint* LSBP = PersistentLevel->GetLevelScriptBlueprint(true) )
 		{
diff --git a/Engine/Source/Runtime/GameLiveStreaming/GameLiveStreaming.cpp b/Engine/Source/Runtime/GameLiveStreaming/GameLiveStreaming.cpp
--- a/Engine/Source/Runtime/GameLiveStreaming/GameLiveStreaming.cpp
+++ b/Engine/Source/Runtime/GameLiveStreaming/GameLiveStreaming.cpp
@@ -436,8 +436,6 @@ UTexture2D* FGameLiveStreaming::GetWebCamTexture( bool& bIsImageFlippedHorizonta
 	UTexture2D* WebCamTexture = nullptr;
 	if( IsBroadcastingGame() && bDrawSimpleWebCamVideo && LiveStreamer->IsWebCamEnabled() )
 	{
-		bool bIsImageFlippedHorizontally = false;
-		bool bIsImageFlippedVertically = false;
 		WebCamTexture = LiveStreamer->GetWebCamTexture( bIsImageFlippedHorizontally, bIsImageFlippedVertically );
 	}
 	return WebCamTexture;
diff --git a/Engine/Source/Runtime/OpenGLDrv/Private/OpenGLTexture.cpp b/Engine/Source/Runtime/OpenGLDrv/Private/OpenGLTexture.cpp
--- a/Engine/Source/Runtime/OpenGLDrv/Private/OpenGLTexture.cpp
+++ b/Engine/Source/Runtime/OpenGLDrv/Private/OpenGLTexture.cpp
@@ -878,7 +878,7 @@ void TOpenGLTexture<RHIResourceType>::Unlock(uint32 MipIndex,uint32 ArrayIndex)
 
 	// Should we use client-storage to improve update time on platforms that require it
 	bool const bRenderable = (this->GetFlags() & (TexCreate_RenderTargetable|TexCreate_ResolveTargetable|TexCreate_DepthStencilTargetable)) != 0;
-	bool const bUseClientStorage = FOpenGL::SupportsClientStorage() && !FOpenGL::SupportsTextureView() && !FOpenGL::SupportsTextureView() && !bRenderable && !this->GetSizeZ() && !GLFormat.bCompressed;
+	bool const bUseClientStorage = FOpenGL::SupportsClientStorage() && !FOpenGL::SupportsTextureView() && !bRenderable && !this->GetSizeZ() && !GLFormat.bCompressed;
 	check(bUseClientStorage || IsValidRef(PixelBuffers[BufferIndex]));
 	
 #if PLATFORM_ANDROID
diff --git a/Engine/Source/Runtime/Renderer/Private/DrawingPolicy.cpp b/Engine/Source/Runtime/Renderer/Private/DrawingPolicy.cpp
--- a/Engine/Source/Runtime/Renderer/Private/DrawingPolicy.cpp
+++ b/Engine/Source/Runtime/Renderer/Private/DrawingPolicy.cpp
@@ -59,7 +59,7 @@ void FMeshDrawingPolicy::SetMeshRenderState(
 
 		// Use bitwise logic ops to avoid branches
 	RHICmdList.SetRasterizerState( GetStaticRasterizerState<true>(
-			( Mesh.bWireframe | IsWireframe() ) ? FM_Wireframe : FM_Solid, ( ( IsTwoSided() & !NeedsBackfacePass() ) | Mesh.bDisableBackfaceCulling ) ? CM_None :
+			( Mesh.bWireframe || IsWireframe() ) ? FM_Wireframe : FM_Solid, ( ( IsTwoSided() && !NeedsBackfacePass() ) || Mesh.bDisableBackfaceCulling ) ? CM_None :
 			( ( (View.bReverseCulling ^ bBackFace) ^ Mesh.ReverseCulling ) ? CM_CCW : CM_CW )
 			));
 }
diff --git a/Engine/Source/Runtime/Renderer/Private/ShadowRendering.cpp b/Engine/Source/Runtime/Renderer/Private/ShadowRendering.cpp
--- a/Engine/Source/Runtime/Renderer/Private/ShadowRendering.cpp
+++ b/Engine/Source/Runtime/Renderer/Private/ShadowRendering.cpp
@@ -3109,9 +3109,9 @@ bool FDeferredShadingSceneRenderer::RenderReflectiveShadowMaps(FRHICommandListIm
 			{
 				FSceneViewState* ViewState = (FSceneViewState*)ProjectedShadowInfo->DependentView->State;
 
-				FLightPropagationVolume* LightPropagationVolume = ViewState->GetLightPropagationVolume();
+				FLightPropagationVolume* LightPropagationVolume = ViewState ? ViewState->GetLightPropagationVolume() : NULL;
 
-				if ( ViewState && LightPropagationVolume )
+				if ( LightPropagationVolume )
 				{
 					if ( ProjectedShadowInfo->bWholeSceneShadow )
 					{
diff --git a/Engine/Source/Runtime/SlateCore/Private/Application/ThrottleManager.cpp b/Engine/Source/Runtime/SlateCore/Private/Application/ThrottleManager.cpp
--- a/Engine/Source/Runtime/SlateCore/Private/Application/ThrottleManager.cpp
+++ b/Engine/Source/Runtime/SlateCore/Private/Application/ThrottleManager.cpp
@@ -7,8 +7,8 @@
  *****************************************************************************/
 
 FSlateThrottleManager::FSlateThrottleManager( )
-	: CVarAllowThrottle(TEXT("Slate.bAllowThrottling"), bShouldThrottle, TEXT("Allow Slate to throttle parts of the engine to ensure the UI is responsive") )
-	, bShouldThrottle(1)
+	: bShouldThrottle(1)
+	, CVarAllowThrottle(TEXT("Slate.bAllowThrottling"), bShouldThrottle, TEXT("Allow Slate to throttle parts of the engine to ensure the UI is responsive") )
 	, ThrottleCount(0)
 { }
 
@@ -42,8 +42,8 @@ void FSlateThrottleManager::LeaveResponsiveMode( FThrottleRequest& InHandle )
 	if( InHandle.IsValid() )
 	{
 		// Decrement throttle count. If it becomes zero we are no longer throttling
+		check(ThrottleCount > 0);
 		--ThrottleCount;
-		check(ThrottleCount >= 0);
 
 		InHandle.Index = INDEX_NONE;
 	}
