Fixes from PVS-Studio static analysis

These changes could cause small behavior change, although they appear to be obvioiusly what was intended by the authors

diff --git a/Engine/Plugins/Media/WmfMedia/Source/WmfMedia/Private/Tracks/WmfMediaVideoTrack.h b/Engine/Plugins/Media/WmfMedia/Source/WmfMedia/Private/Tracks/WmfMediaVideoTrack.h
--- a/Engine/Plugins/Media/WmfMedia/Source/WmfMedia/Private/Tracks/WmfMediaVideoTrack.h
+++ b/Engine/Plugins/Media/WmfMedia/Source/WmfMedia/Private/Tracks/WmfMediaVideoTrack.h
@@ -39,7 +39,7 @@ public:
 		}
 		else
 		{
-			FrameRate = Numerator / Denominator;
+			FrameRate = static_cast<float>(Numerator) / Denominator;
 		}
 	}
 
diff --git a/Engine/Source/Editor/SoundCueEditor/Private/SoundCueEditorUtilities.cpp b/Engine/Source/Editor/SoundCueEditor/Private/SoundCueEditorUtilities.cpp
--- a/Engine/Source/Editor/SoundCueEditor/Private/SoundCueEditorUtilities.cpp
+++ b/Engine/Source/Editor/SoundCueEditor/Private/SoundCueEditorUtilities.cpp
@@ -32,7 +32,7 @@ void FSoundCueEditorUtilities::CreateWaveContainers(TArray<USoundWave*>& Selecte
 {
 	const int32 NodeSpacing = 70;
 	
-	Location.Y -= (SelectedWaves.Num() - 1) * NodeSpacing / 2;
+	Location.Y -= static_cast<float>((SelectedWaves.Num() - 1) * NodeSpacing) / 2.f;
 
 	for(int32 WaveIndex = 0; WaveIndex < SelectedWaves.Num(); WaveIndex++)
 	{
diff --git a/Engine/Source/Editor/UnrealEd/Private/AssetDeleteModel.cpp b/Engine/Source/Editor/UnrealEd/Private/AssetDeleteModel.cpp
--- a/Engine/Source/Editor/UnrealEd/Private/AssetDeleteModel.cpp
+++ b/Engine/Source/Editor/UnrealEd/Private/AssetDeleteModel.cpp
@@ -290,7 +290,7 @@ bool FAssetDeleteModel::ComputeCanReplaceReferences()
 		PendingDeletedObjects.Add(PendingDelete->GetObject());
 	}
 
-	return ObjectTools::AreObjectsValidForReplace(PendingDeletedObjects) && ObjectTools::AreObjectsValidForReplace(PendingDeletedObjects);
+	return ObjectTools::AreObjectsValidForReplace(PendingDeletedObjects);
 }
 
 bool FAssetDeleteModel::CanReplaceReferences() const
diff --git a/Engine/Source/Runtime/Engine/Private/GameplayStatics.cpp b/Engine/Source/Runtime/Engine/Private/GameplayStatics.cpp
--- a/Engine/Source/Runtime/Engine/Private/GameplayStatics.cpp
+++ b/Engine/Source/Runtime/Engine/Private/GameplayStatics.cpp
@@ -985,7 +985,7 @@ void UGameplayStatics::ActivateReverbEffect(UObject* WorldContextObject, class U
 
 void UGameplayStatics::DeactivateReverbEffect(UObject* WorldContextObject, FName TagName)
 {
-	if (GEngine || !GEngine->UseSound())
+	if (GEngine == nullptr || !GEngine->UseSound())
 	{
 		return;
 	}
diff --git a/Engine/Source/Runtime/Engine/Private/ShadowMap.cpp b/Engine/Source/Runtime/Engine/Private/ShadowMap.cpp
--- a/Engine/Source/Runtime/Engine/Private/ShadowMap.cpp
+++ b/Engine/Source/Runtime/Engine/Private/ShadowMap.cpp
@@ -971,9 +971,9 @@ void FShadowMap2D::EncodeSingleTexture(FShadowMapPendingTexture& PendingTexture,
 
 						for (int32 NeighborIndex = 0; NeighborIndex < ARRAY_COUNT(Neighbors); NeighborIndex++)
 						{
-							if (DestY + Neighbors[NeighborIndex].Y >= 0 
+							if (static_cast<int32>(DestY) + Neighbors[NeighborIndex].Y >= 0 
 								&& DestY + Neighbors[NeighborIndex].Y < MipSizeY
-								&& DestX + Neighbors[NeighborIndex].X >= 0 
+								&& static_cast<int32>(DestX) + Neighbors[NeighborIndex].X >= 0
 								&& DestX + Neighbors[NeighborIndex].X < MipSizeX)
 							{
 								const FFourDistanceFieldSamples& FourNeighborSamples = MipData[MipIndex][(DestY + Neighbors[NeighborIndex].Y) * MipSizeX + DestX + Neighbors[NeighborIndex].X];
@@ -981,9 +981,9 @@ void FShadowMap2D::EncodeSingleTexture(FShadowMapPendingTexture& PendingTexture,
 
 								if (NeighborSample.Coverage > 0)
 								{
-									if (DestY + Neighbors[NeighborIndex].Y * 2 >= 0 
+									if (static_cast<int32>(DestY) + Neighbors[NeighborIndex].Y * 2 >= 0
 										&& DestY + Neighbors[NeighborIndex].Y * 2 < MipSizeY
-										&& DestX + Neighbors[NeighborIndex].X * 2 >= 0 
+										&& static_cast<int32>(DestX) + Neighbors[NeighborIndex].X * 2 >= 0
 										&& DestX + Neighbors[NeighborIndex].X * 2 < MipSizeX)
 									{
 										// Lookup the second neighbor in the first neighbor's direction
diff --git a/Engine/Source/Runtime/GameMenuBuilder/Private/SGameMenuPageWidget.cpp b/Engine/Source/Runtime/GameMenuBuilder/Private/SGameMenuPageWidget.cpp
--- a/Engine/Source/Runtime/GameMenuBuilder/Private/SGameMenuPageWidget.cpp
+++ b/Engine/Source/Runtime/GameMenuBuilder/Private/SGameMenuPageWidget.cpp
@@ -625,7 +625,7 @@ void SGameMenuPageWidget::OnSubPanelStateChange(bool bWasOpened)
 void SGameMenuPageWidget::Tick( const FGeometry& AllottedGeometry, const double InCurrentTime, const float InDeltaTime )
 {
 	//ugly code seeing if the console is open
-	UConsole* ViewportConsole = ( GEngine->GameViewport != nullptr ) ? GEngine->GameViewport->ViewportConsole : nullptr;
+	UConsole* ViewportConsole = (GEngine && GEngine->GameViewport != nullptr) ? GEngine->GameViewport->ViewportConsole : nullptr;
 	if ( ViewportConsole != nullptr && ( ViewportConsole->ConsoleState == "Typing" || ViewportConsole->ConsoleState == "Open" ) )
 	{
 		if (!bConsoleVisible)
@@ -931,7 +931,7 @@ FReply SGameMenuPageWidget::OnKeyDown(const FGeometry& MyGeometry, const FKeyEve
 	
 	if ((CurrentMenu.IsValid() == true) && (!bControlsLocked))
 	{
-		bool bNavigationLocked = bControlsLocked | PendingMainMenu.IsValid() | PendingSubMenu.IsValid();
+		bool bNavigationLocked = bControlsLocked || PendingMainMenu.IsValid() || PendingSubMenu.IsValid();
 		const FKey Key = InKeyEvent.GetKey();
 
 		if (bNavigationLocked == false)
