UMG - Exposing a lot more functions for the FKey structure to blueprints.  You can now check what kind of key it is, you can also now get the DisplayText for it.

diff --git a/Engine/Source/Runtime/Engine/Classes/Kismet/KismetInputLibrary.h b/Engine/Source/Runtime/Engine/Classes/Kismet/KismetInputLibrary.h
--- a/Engine/Source/Runtime/Engine/Classes/Kismet/KismetInputLibrary.h
+++ b/Engine/Source/Runtime/Engine/Classes/Kismet/KismetInputLibrary.h
@@ -23,6 +23,48 @@ class UKismetInputLibrary : public UBlueprintFunctionLibrary
 	 */
 	UFUNCTION(BlueprintPure, meta=(FriendlyName = "Equal (Key)", CompactNodeTitle = "=="), Category="Utilities|Key")
 	static bool EqualEqual_KeyKey(FKey A, FKey B);
+	
+	/**
+	 * @returns True if the key is a modifier key: Ctrl, Command, Alt, Shift
+	 */
+	UFUNCTION(BlueprintPure, Category="Utilities|Key")
+	static bool Key_IsModifierKey(const FKey& Key);
+	
+	/**
+	 * @returns True if the key is a gamepad button
+	 */
+	UFUNCTION(BlueprintPure, Category="Utilities|Key")
+	static bool Key_IsGamepadKey(const FKey& Key);
+	
+	/**
+	 * @returns True if the key is a mouse button
+	 */
+	UFUNCTION(BlueprintPure, Category="Utilities|Key")
+	static bool Key_IsMouseButton(const FKey& Key);
+	
+	/**
+	 * @returns True if the key is a keyboard button
+	 */
+	UFUNCTION(BlueprintPure, Category="Utilities|Key")
+	static bool Key_IsKeyboardKey(const FKey& Key);
+	
+	/**
+	 * @returns True if the key is a float axis
+	 */
+	UFUNCTION(BlueprintPure, Category="Utilities|Key")
+	static bool Key_IsFloatAxis(const FKey& Key);
+	
+	/**
+	 * @returns True if the key is a vector axis
+	 */
+	UFUNCTION(BlueprintPure, Category="Utilities|Key")
+	static bool Key_IsVectorAxis(const FKey& Key);
+	
+	/**
+	 * @returns The display name of the key.
+	 */
+	UFUNCTION(BlueprintPure, Category="Utilities|Key")
+	static FText Key_GetDisplayName(const FKey& Key);
 
 	/**
 	 * Returns whether or not this character is an auto-repeated keystroke
@@ -157,7 +199,7 @@ class UKismetInputLibrary : public UBlueprintFunctionLibrary
 
 	/** Mouse buttons that are currently pressed */
 	UFUNCTION(BlueprintPure, meta=( FriendlyName = "IsMouseButtonDown" ), Category="Utilities|PointerEvent")
-		static bool PointerEvent_IsMouseButtonDown(const FPointerEvent& Input, FKey MouseButton);
+	static bool PointerEvent_IsMouseButtonDown(const FPointerEvent& Input, FKey MouseButton);
 
 	/** Mouse button that caused this event to be raised (possibly EB_None) */
 	UFUNCTION(BlueprintPure, meta=( FriendlyName = "GetEffectingButton" ), Category="Utilities|PointerEvent")
diff --git a/Engine/Source/Runtime/Engine/Private/KismetInputLibrary.cpp b/Engine/Source/Runtime/Engine/Private/KismetInputLibrary.cpp
--- a/Engine/Source/Runtime/Engine/Private/KismetInputLibrary.cpp
+++ b/Engine/Source/Runtime/Engine/Private/KismetInputLibrary.cpp
@@ -28,6 +28,41 @@ bool UKismetInputLibrary::EqualEqual_KeyKey(FKey A, FKey B)
 	return A == B;
 }
 
+bool UKismetInputLibrary::Key_IsModifierKey(const FKey& Key)
+{
+	return Key.IsModifierKey();
+}
+
+bool UKismetInputLibrary::Key_IsGamepadKey(const FKey& Key)
+{
+	return Key.IsGamepadKey();
+}
+
+bool UKismetInputLibrary::Key_IsMouseButton(const FKey& Key)
+{
+	return Key.IsMouseButton();
+}
+
+bool UKismetInputLibrary::Key_IsKeyboardKey(const FKey& Key)
+{
+	return Key.IsBindableInBlueprints() && (Key.IsGamepadKey() == false && Key.IsMouseButton() == false);
+}
+
+bool UKismetInputLibrary::Key_IsFloatAxis(const FKey& Key)
+{
+	return Key.IsFloatAxis();
+}
+
+bool UKismetInputLibrary::Key_IsVectorAxis(const FKey& Key)
+{
+	return Key.IsVectorAxis();
+}
+
+FText UKismetInputLibrary::Key_GetDisplayName(const FKey& Key)
+{
+	return Key.GetDisplayName();
+}
+
 bool UKismetInputLibrary::InputEvent_IsRepeat(const FInputEvent& Input)
 {
 	return Input.IsRepeat();
