Core - Fixed a deadlock in the FScreenSaverInhibitor

diff --git a/Engine/Source/Runtime/Engine/Private/UnrealEngine.cpp b/Engine/Source/Runtime/Engine/Private/UnrealEngine.cpp
--- a/Engine/Source/Runtime/Engine/Private/UnrealEngine.cpp
+++ b/Engine/Source/Runtime/Engine/Private/UnrealEngine.cpp
@@ -975,6 +975,8 @@ void UEngine::PreExit()
 
 	if (ScreenSaverInhibitor)
 	{
+		// Resume the thread to avoid a deadlock while waiting for finish.
+		ScreenSaverInhibitor->Suspend( false );
 		delete ScreenSaverInhibitor;
 	}
 
