Allow FComponentReference to reference SceneComponents

diff --git a/Engine/Source/Runtime/Engine/Private/EngineTypes.cpp b/Engine/Source/Runtime/Engine/Private/EngineTypes.cpp
--- a/Engine/Source/Runtime/Engine/Private/EngineTypes.cpp
+++ b/Engine/Source/Runtime/Engine/Private/EngineTypes.cpp
@@ -140,12 +140,12 @@ USceneComponent* FComponentReference::GetComponent(AActor* OwningActor) const
 				if(ObjProp != NULL)
 				{
 					// .. and return the component that is there
-					Result = Cast<UPrimitiveComponent>(ObjProp->GetObjectPropertyValue_InContainer(SearchActor));
+					Result = Cast<USceneComponent>(ObjProp->GetObjectPropertyValue_InContainer(SearchActor));
 				}
 			}
 			else
 			{
-				Result = Cast<UPrimitiveComponent>(SearchActor->GetRootComponent());
+				Result = Cast<USceneComponent>(SearchActor->GetRootComponent());
 			}
 		}
 	}
