Optimize ComponentsThatNeedEndOfFrameUpdate

diff --git a/Engine/Source/Runtime/Engine/Classes/Components/ActorComponent.h b/Engine/Source/Runtime/Engine/Classes/Components/ActorComponent.h
--- a/Engine/Source/Runtime/Engine/Classes/Components/ActorComponent.h
+++ b/Engine/Source/Runtime/Engine/Classes/Components/ActorComponent.h
@@ -136,6 +136,10 @@ private:
 	/** Indicates that BeginPlay has been called, but EndPlay has not yet */
 	uint32 bHasBegunPlay:1;
 
+	/** Tracks whether the component has been added to one of the world's end of frame update lists */
+	uint32 MarkedForEndOfFrameUpdateState:2;
+	friend struct FMarkComponentEndOfFrameUpdateState;
+
 	friend class FActorComponentInstanceData;
 
 	UPROPERTY()
@@ -143,6 +147,8 @@ private:
 
 public:
 
+	uint32 GetMarkedForEndOfFrameUpdateState() const { return MarkedForEndOfFrameUpdateState; }
+
 	void DetermineUCSModifiedProperties();
 	void GetUCSModifiedProperties(TSet<const UProperty*>& ModifiedProperties) const;
 
@@ -626,7 +632,4 @@ private:
 	virtual void Tick( float DeltaTime ) final { check(0); }
 
 #endif
-};
-
-
-
+};
\ No newline at end of file
diff --git a/Engine/Source/Runtime/Engine/Classes/Engine/World.h b/Engine/Source/Runtime/Engine/Classes/Engine/World.h
--- a/Engine/Source/Runtime/Engine/Classes/Engine/World.h
+++ b/Engine/Source/Runtime/Engine/Classes/Engine/World.h
@@ -743,10 +743,10 @@ private:
 	FPhysScene*									PhysicsScene;
 
 	/** Set of components that need updates at the end of the frame */
-	TSet<TWeakObjectPtr<class UActorComponent> > ComponentsThatNeedEndOfFrameUpdate;
+	TArray<TWeakObjectPtr<class UActorComponent> > ComponentsThatNeedEndOfFrameUpdate;
 
 	/** Set of components that need recreates at the end of the frame */
-	TSet<TWeakObjectPtr<class UActorComponent> > ComponentsThatNeedEndOfFrameUpdate_OnGameThread;
+	TArray<TWeakObjectPtr<class UActorComponent> > ComponentsThatNeedEndOfFrameUpdate_OnGameThread;
 
 	/** The state of async tracing - abstracted into its own object for easier reference */
 	FWorldAsyncTraceState AsyncTraceState;
diff --git a/Engine/Source/Runtime/Engine/Private/Components/ActorComponent.cpp b/Engine/Source/Runtime/Engine/Private/Components/ActorComponent.cpp
--- a/Engine/Source/Runtime/Engine/Private/Components/ActorComponent.cpp
+++ b/Engine/Source/Runtime/Engine/Private/Components/ActorComponent.cpp
@@ -1219,9 +1219,10 @@ void UActorComponent::MarkForNeededEndOfFrameUpdate()
 		return;
 	}
 
-	if(GetWorld())
+	UWorld* World = GetWorld();
+	if (World)
 	{
-		GetWorld()->MarkActorComponentForNeededEndOfFrameUpdate(this, RequiresGameThreadEndOfFrameUpdates());
+		World->MarkActorComponentForNeededEndOfFrameUpdate(this, RequiresGameThreadEndOfFrameUpdates());
 	}
 	else if (!HasAnyFlags(RF_Unreachable))
 	{
@@ -1237,10 +1238,11 @@ void UActorComponent::MarkForNeededEndOfFrameRecreate()
 		return;
 	}
 
-	if (GetWorld())
+	UWorld* World = GetWorld();
+	if (World)
 	{
 		// by convention, recreates are always done on the gamethread
-		GetWorld()->MarkActorComponentForNeededEndOfFrameUpdate(this, true);
+		World->MarkActorComponentForNeededEndOfFrameUpdate(this, true);
 	}
 	else if (!HasAnyFlags(RF_Unreachable))
 	{
diff --git a/Engine/Source/Runtime/Engine/Private/LevelTick.cpp b/Engine/Source/Runtime/Engine/Private/LevelTick.cpp
--- a/Engine/Source/Runtime/Engine/Private/LevelTick.cpp
+++ b/Engine/Source/Runtime/Engine/Private/LevelTick.cpp
@@ -706,22 +706,55 @@ static TAutoConsoleVariable<int32> CVarCollectGarbageEveryFrame(
 	0,
 	TEXT("Used to debug garbage collection...Collects garbage every frame if the value is > 0."));
 
-void UWorld::MarkActorComponentForNeededEndOfFrameUpdate(class UActorComponent* Component, bool bForceGameThread)
+namespace EComponentMarkedForEndOfFrameUpdateState
 {
-	check(!bPostTickComponentUpdate); // can't call this while we are doing the updates
-	if (!bForceGameThread)
+	enum Type
 	{
-		bool bAllowConcurrentUpdates = !!CVarAllowAsyncRenderThreadUpdates.GetValueOnGameThread();
-		bForceGameThread = !bAllowConcurrentUpdates;
-	}
+		Unmarked,
+		Marked,
+		MarkedForGameThread,
+	};
+}
+
+// Utility struct to allow world direct access to UActorComponent::MarkedForEndOfFrameUpdateState without friending all of UActorComponent
+struct FMarkComponentEndOfFrameUpdateState
+{
+	friend class UWorld;
 
-	if (bForceGameThread)
+private:
+	FORCEINLINE static void Set(UActorComponent* Component, const EComponentMarkedForEndOfFrameUpdateState::Type UpdateState)
 	{
-		ComponentsThatNeedEndOfFrameUpdate_OnGameThread.Add(Component);
+		checkSlow(UpdateState < 4); // Only 2 bits are allocated to store this value
+		Component->MarkedForEndOfFrameUpdateState = UpdateState;
 	}
-	else
+};
+
+void UWorld::MarkActorComponentForNeededEndOfFrameUpdate(UActorComponent* Component, bool bForceGameThread)
+{
+	check(!bPostTickComponentUpdate); // can't call this while we are doing the updates
+	if (Component->GetMarkedForEndOfFrameUpdateState() != EComponentMarkedForEndOfFrameUpdateState::MarkedForGameThread)
 	{
-		ComponentsThatNeedEndOfFrameUpdate.Add(Component);
+		if (!bForceGameThread)
+		{
+			bool bAllowConcurrentUpdates = !!CVarAllowAsyncRenderThreadUpdates.GetValueOnGameThread();
+			bForceGameThread = !bAllowConcurrentUpdates;
+		}
+
+		if (bForceGameThread)
+		{
+			// In this case we need to remove it from the other list
+			if (Component->GetMarkedForEndOfFrameUpdateState() == EComponentMarkedForEndOfFrameUpdateState::Marked)
+			{
+				ComponentsThatNeedEndOfFrameUpdate.RemoveSwap(Component);
+			}
+			ComponentsThatNeedEndOfFrameUpdate_OnGameThread.Add(Component);
+			FMarkComponentEndOfFrameUpdateState::Set(Component, EComponentMarkedForEndOfFrameUpdateState::MarkedForGameThread);
+		}
+		else if (Component->GetMarkedForEndOfFrameUpdateState() != EComponentMarkedForEndOfFrameUpdateState::Marked)
+		{
+			ComponentsThatNeedEndOfFrameUpdate.Add(Component);
+			FMarkComponentEndOfFrameUpdateState::Set(Component, EComponentMarkedForEndOfFrameUpdateState::Marked);
+		}
 	}
 }
 
@@ -737,31 +770,25 @@ void UWorld::SendAllEndOfFrameUpdates(FGraphEventArray* OutCompletion)
 	if (!OutCompletion)
 	{
 		// this is a viewer or something, just do everything on the gamethread
-		for (TSet<TWeakObjectPtr<UActorComponent> >::TIterator It(ComponentsThatNeedEndOfFrameUpdate); It; ++It)
-		{
-			ComponentsThatNeedEndOfFrameUpdate_OnGameThread.Add(*It);
-		}
+		ComponentsThatNeedEndOfFrameUpdate_OnGameThread.Append(ComponentsThatNeedEndOfFrameUpdate);
 		ComponentsThatNeedEndOfFrameUpdate.Empty(ComponentsThatNeedEndOfFrameUpdate.Num());
 	}
-	else
-	{
-		// remove any gamethread updates from the async update list
-		for (TSet<TWeakObjectPtr<UActorComponent> >::TIterator It(ComponentsThatNeedEndOfFrameUpdate_OnGameThread); It; ++It)
-		{
-			ComponentsThatNeedEndOfFrameUpdate.Remove(*It);
-		}
-	}
 
 	// Game thread updates need to happen before we go wide on the other threads.
 	// These updates are things that have said that they are NOT SAFE to run concurrently.
-	for (TSet<TWeakObjectPtr<UActorComponent> >::TIterator It(ComponentsThatNeedEndOfFrameUpdate_OnGameThread); It; ++It)
+	for (TArray<TWeakObjectPtr<UActorComponent> >::TIterator It(ComponentsThatNeedEndOfFrameUpdate_OnGameThread); It; ++It)
 	{
 		UActorComponent* Component = It->Get();
-		if (Component && !Component->IsPendingKill() && Component->IsRegistered() && !Component->IsTemplate())
+		if (Component)
 		{
-			FScopeCycleCounterUObject ComponentScope(Component);
-			FScopeCycleCounterUObject AdditionalScope(STATS ? Component->AdditionalStatObject() : NULL);
-			Component->DoDeferredRenderUpdates_Concurrent();
+			checkSlow(Component->GetMarkedForEndOfFrameUpdateState() != EComponentMarkedForEndOfFrameUpdateState::Unmarked);
+			if ( !Component->IsPendingKill() && Component->IsRegistered() && !Component->IsTemplate())
+			{
+				FScopeCycleCounterUObject ComponentScope(Component);
+				FScopeCycleCounterUObject AdditionalScope(STATS ? Component->AdditionalStatObject() : NULL);
+				Component->DoDeferredRenderUpdates_Concurrent();
+			}
+			FMarkComponentEndOfFrameUpdateState::Set(Component, EComponentMarkedForEndOfFrameUpdateState::Unmarked);
 		}
 	}
 
@@ -814,21 +841,26 @@ void UWorld::SendAllEndOfFrameUpdates(FGraphEventArray* OutCompletion)
 			//@todo optimization, this loop could be done on another thread
 			// First get the async transform and render data updates underway
 			FTaskArray* Array = NULL;
-			for (TSet<TWeakObjectPtr<UActorComponent> >::TIterator It(ComponentsThatNeedEndOfFrameUpdate); It; ++It)
+			for (TArray<TWeakObjectPtr<UActorComponent> >::TIterator It(ComponentsThatNeedEndOfFrameUpdate); It; ++It)
 			{
 				UActorComponent* NextComponent = It->Get();
-				if (NextComponent && !NextComponent->IsPendingKill() && NextComponent->IsRegistered() && !NextComponent->IsTemplate())
+				if (NextComponent)
 				{
-					if (!Array)
+					checkSlow(NextComponent->GetMarkedForEndOfFrameUpdateState() == EComponentMarkedForEndOfFrameUpdateState::Marked);
+					if (!NextComponent->IsPendingKill() && NextComponent->IsRegistered() && !NextComponent->IsTemplate())
 					{
-						Array = new FTaskArray;
-					}
-					Array->Add(NextComponent);
-					if (Array->Num() == NUM_COMPONENTS_PER_TASK)
-					{
-						new (*OutCompletion) FGraphEventRef(TGraphTask<FDoRenderthreadUpdatesTask>::CreateTask(NULL, ENamedThreads::GameThread).ConstructAndDispatchWhenReady(Array)); 
-						Array = NULL; // Array belongs to the task
+						if (!Array)
+						{
+							Array = new FTaskArray;
+						}
+						Array->Add(NextComponent);
+						if (Array->Num() == NUM_COMPONENTS_PER_TASK)
+						{
+							new (*OutCompletion) FGraphEventRef(TGraphTask<FDoRenderthreadUpdatesTask>::CreateTask(NULL, ENamedThreads::GameThread).ConstructAndDispatchWhenReady(Array)); 
+							Array = NULL; // Array belongs to the task
+						}
 					}
+					FMarkComponentEndOfFrameUpdateState::Set(NextComponent, EComponentMarkedForEndOfFrameUpdateState::Unmarked);
 				}
 			}
 			if (Array) // partial array if we had one
@@ -837,11 +869,12 @@ void UWorld::SendAllEndOfFrameUpdates(FGraphEventArray* OutCompletion)
 				Array = NULL; // Array belongs to the task
 			}
 		}
+		ComponentsThatNeedEndOfFrameUpdate.Empty(ComponentsThatNeedEndOfFrameUpdate.Num());
 	}
 
-	bPostTickComponentUpdate = false;
-	ComponentsThatNeedEndOfFrameUpdate.Empty(ComponentsThatNeedEndOfFrameUpdate.Num());
 	ComponentsThatNeedEndOfFrameUpdate_OnGameThread.Empty(ComponentsThatNeedEndOfFrameUpdate_OnGameThread.Num());
+
+	bPostTickComponentUpdate = false;
 }
 
 
