Fix for an unnecessary quilifier error in Device profile Manager. Only arose on clang build.

diff --git a/Engine/Source/Runtime/Engine/Classes/DeviceProfiles/DeviceProfileManager.h b/Engine/Source/Runtime/Engine/Classes/DeviceProfiles/DeviceProfileManager.h
--- a/Engine/Source/Runtime/Engine/Classes/DeviceProfiles/DeviceProfileManager.h
+++ b/Engine/Source/Runtime/Engine/Classes/DeviceProfiles/DeviceProfileManager.h
@@ -100,7 +100,7 @@ private:
 	*
 	* @return The selected profile.
 	*/
-	static const FString UDeviceProfileManager::GetActiveProfileName();
+	static const FString GetActiveProfileName();
 
 	/**
 	 * Set the active device profile - set via the device profile blueprint.
