Change single node anim instance finished playing checks to account for asset play rate

diff --git a/Engine/Source/Runtime/Engine/Private/Animation/AnimSingleNodeInstance.cpp b/Engine/Source/Runtime/Engine/Private/Animation/AnimSingleNodeInstance.cpp
--- a/Engine/Source/Runtime/Engine/Private/Animation/AnimSingleNodeInstance.cpp
+++ b/Engine/Source/Runtime/Engine/Private/Animation/AnimSingleNodeInstance.cpp
@@ -245,7 +245,8 @@ void UAnimSingleNodeInstance::NativeUpdateAnimation(float DeltaTimeX)
 			// if it's not looping, just set play to be false when reached to end
 			if (!bLooping)
 			{
-				if ((NewPlayRate < 0.f && CurrentTime <= 0.f) || (NewPlayRate > 0.f && CurrentTime >= Sequence->SequenceLength))
+				const float CombinedPlayRate = NewPlayRate*Sequence->RateScale;
+				if ((CombinedPlayRate < 0.f && CurrentTime <= 0.f) || (CombinedPlayRate > 0.f && CurrentTime >= Sequence->SequenceLength))
 				{
 					SetPlaying(false);
 				}
@@ -258,7 +259,8 @@ void UAnimSingleNodeInstance::NativeUpdateAnimation(float DeltaTimeX)
 			// if it's not looping, just set play to be false when reached to end
 			if (!bLooping)
 			{
-				if ((NewPlayRate < 0.f && CurrentTime <= 0.f) || (NewPlayRate > 0.f && CurrentTime >= Composite->SequenceLength))
+				const float CombinedPlayRate = NewPlayRate*Composite->RateScale;
+				if ((CombinedPlayRate < 0.f && CurrentTime <= 0.f) || (CombinedPlayRate > 0.f && CurrentTime >= Composite->SequenceLength))
 				{
 					SetPlaying(false);
 				}
