#ue4 Functions in native BP libraries are callable from native code game project. (UEBP-40 - Native code generation from BP)

diff --git a/Engine/Source/Runtime/Engine/Classes/Kismet/DataTableFunctionLibrary.h b/Engine/Source/Runtime/Engine/Classes/Kismet/DataTableFunctionLibrary.h
--- a/Engine/Source/Runtime/Engine/Classes/Kismet/DataTableFunctionLibrary.h
+++ b/Engine/Source/Runtime/Engine/Classes/Kismet/DataTableFunctionLibrary.h
@@ -20,8 +20,8 @@ namespace EEvaluateCurveTableResult
     };
 }
 
-UCLASS(MinimalAPI)
-class UDataTableFunctionLibrary : public UBlueprintFunctionLibrary
+UCLASS()
+class ENGINE_API UDataTableFunctionLibrary : public UBlueprintFunctionLibrary
 {
 	GENERATED_UCLASS_BODY()
 
@@ -35,7 +35,7 @@ class UDataTableFunctionLibrary : public UBlueprintFunctionLibrary
     UFUNCTION(BlueprintCallable, CustomThunk, Category = "DataTable", meta=(CustomStructureParam = "OutRow", BlueprintInternalUseOnly="true"))
     static bool GetDataTableRowFromName(UDataTable* Table, FName RowName, FTableRowBase& OutRow);
     
-	ENGINE_API static bool Generic_GetDataTableRowFromName(UDataTable* Table, FName RowName, void* OutRowPtr);
+	static bool Generic_GetDataTableRowFromName(UDataTable* Table, FName RowName, void* OutRowPtr);
 
     /** Based on UDataTableFunctionLibrary::GetDataTableRow */
     DECLARE_FUNCTION(execGetDataTableRowFromName)
diff --git a/Engine/Source/Runtime/Engine/Classes/Kismet/HeadMountedDisplayFunctionLibrary.h b/Engine/Source/Runtime/Engine/Classes/Kismet/HeadMountedDisplayFunctionLibrary.h
--- a/Engine/Source/Runtime/Engine/Classes/Kismet/HeadMountedDisplayFunctionLibrary.h
+++ b/Engine/Source/Runtime/Engine/Classes/Kismet/HeadMountedDisplayFunctionLibrary.h
@@ -15,7 +15,7 @@ namespace EOrientPositionSelector
 }
 
 UCLASS()
-class UHeadMountedDisplayFunctionLibrary : public UBlueprintFunctionLibrary
+class ENGINE_API UHeadMountedDisplayFunctionLibrary : public UBlueprintFunctionLibrary
 {
 	GENERATED_UCLASS_BODY()
 
diff --git a/Engine/Source/Runtime/Engine/Classes/Kismet/KismetArrayLibrary.h b/Engine/Source/Runtime/Engine/Classes/Kismet/KismetArrayLibrary.h
--- a/Engine/Source/Runtime/Engine/Classes/Kismet/KismetArrayLibrary.h
+++ b/Engine/Source/Runtime/Engine/Classes/Kismet/KismetArrayLibrary.h
@@ -3,8 +3,8 @@
 #pragma once
 #include "KismetArrayLibrary.generated.h"
 
-UCLASS(MinimalAPI)
-class UKismetArrayLibrary : public UBlueprintFunctionLibrary
+UCLASS()
+class ENGINE_API UKismetArrayLibrary : public UBlueprintFunctionLibrary
 {
 	GENERATED_UCLASS_BODY()
 	/** 
@@ -166,21 +166,21 @@ class UKismetArrayLibrary : public UBlueprintFunctionLibrary
 	static void SetArrayPropertyByName(UObject* Object, FName PropertyName, const TArray<int32>& Value);
 
 	// Native functions that will be called by the below custom thunk layers, which read off the property address, and call the appropriate native handler
-	ENGINE_API static int32 GenericArray_Add(void* TargetArray, const UArrayProperty* ArrayProp, const void* NewItem);
-	ENGINE_API static int32 GenericArray_AddUnique(void* TargetArray, const UArrayProperty* ArrayProp, const void* NewItem);
-	ENGINE_API static void GenericArray_Shuffle(void* TargetArray, const UArrayProperty* ArrayProp);
-	ENGINE_API static void GenericArray_Append(void* TargetArray, const UArrayProperty* TargetArrayProp, void* SourceArray, const UArrayProperty* SourceArrayProperty);
-	ENGINE_API static void GenericArray_Insert(void* TargetArray, const UArrayProperty* ArrayProp, const void* NewItem, int32 Index);
-	ENGINE_API static void GenericArray_Remove(void* TargetArray, const UArrayProperty* ArrayProp, int32 IndexToRemove);
-	ENGINE_API static bool GenericArray_RemoveItem(void* TargetArray, const UArrayProperty* ArrayProp, const void* Item);
-	ENGINE_API static void GenericArray_Clear(void* TargetArray, const UArrayProperty* ArrayProp);
-	ENGINE_API static void GenericArray_Resize(void* TargetArray, const UArrayProperty* ArrayProp, int32 Size);
-	ENGINE_API static int32 GenericArray_Length(void* TargetArray, const UArrayProperty* ArrayProp);
-	ENGINE_API static int32 GenericArray_LastIndex(void* TargetArray, const UArrayProperty* ArrayProp);
-	ENGINE_API static void GenericArray_Get(void* TargetArray, const UArrayProperty* ArrayProp, int32 Index, void* Item);
-	ENGINE_API static void GenericArray_Set(void* TargetArray, const UArrayProperty* ArrayProp, int32 Index, const void* NewItem, bool bSizeToFit);
-	ENGINE_API static int32 GenericArray_Find(void* TargetArray, const UArrayProperty* ArrayProperty, const void* ItemToFind);
-	ENGINE_API static void GenericArray_SetArrayPropertyByName(UObject* OwnerObject, FName ArrayPropertyName, const void* SrcArrayAddr);
+	static int32 GenericArray_Add(void* TargetArray, const UArrayProperty* ArrayProp, const void* NewItem);
+	static int32 GenericArray_AddUnique(void* TargetArray, const UArrayProperty* ArrayProp, const void* NewItem);
+	static void GenericArray_Shuffle(void* TargetArray, const UArrayProperty* ArrayProp);
+	static void GenericArray_Append(void* TargetArray, const UArrayProperty* TargetArrayProp, void* SourceArray, const UArrayProperty* SourceArrayProperty);
+	static void GenericArray_Insert(void* TargetArray, const UArrayProperty* ArrayProp, const void* NewItem, int32 Index);
+	static void GenericArray_Remove(void* TargetArray, const UArrayProperty* ArrayProp, int32 IndexToRemove);
+	static bool GenericArray_RemoveItem(void* TargetArray, const UArrayProperty* ArrayProp, const void* Item);
+	static void GenericArray_Clear(void* TargetArray, const UArrayProperty* ArrayProp);
+	static void GenericArray_Resize(void* TargetArray, const UArrayProperty* ArrayProp, int32 Size);
+	static int32 GenericArray_Length(void* TargetArray, const UArrayProperty* ArrayProp);
+	static int32 GenericArray_LastIndex(void* TargetArray, const UArrayProperty* ArrayProp);
+	static void GenericArray_Get(void* TargetArray, const UArrayProperty* ArrayProp, int32 Index, void* Item);
+	static void GenericArray_Set(void* TargetArray, const UArrayProperty* ArrayProp, int32 Index, const void* NewItem, bool bSizeToFit);
+	static int32 GenericArray_Find(void* TargetArray, const UArrayProperty* ArrayProperty, const void* ItemToFind);
+	static void GenericArray_SetArrayPropertyByName(UObject* OwnerObject, FName ArrayPropertyName, const void* SrcArrayAddr);
 
 private:
 	static void GenericArray_HandleBool(const UProperty* Property, void* ItemPtr);
diff --git a/Engine/Source/Runtime/Engine/Classes/Kismet/KismetGuidLibrary.h b/Engine/Source/Runtime/Engine/Classes/Kismet/KismetGuidLibrary.h
--- a/Engine/Source/Runtime/Engine/Classes/Kismet/KismetGuidLibrary.h
+++ b/Engine/Source/Runtime/Engine/Classes/Kismet/KismetGuidLibrary.h
@@ -5,8 +5,8 @@
 #include "KismetGuidLibrary.generated.h"
 
 
-UCLASS(MinimalAPI)
-class UKismetGuidLibrary
+UCLASS()
+class ENGINE_API UKismetGuidLibrary
 	: public UBlueprintFunctionLibrary
 {
 	GENERATED_UCLASS_BODY()
diff --git a/Engine/Source/Runtime/Engine/Classes/Kismet/KismetInputLibrary.h b/Engine/Source/Runtime/Engine/Classes/Kismet/KismetInputLibrary.h
--- a/Engine/Source/Runtime/Engine/Classes/Kismet/KismetInputLibrary.h
+++ b/Engine/Source/Runtime/Engine/Classes/Kismet/KismetInputLibrary.h
@@ -6,8 +6,8 @@
 typedef FKeyEvent FKeyEvent;
 #include "KismetInputLibrary.generated.h"
 
-UCLASS(MinimalAPI)
-class UKismetInputLibrary : public UBlueprintFunctionLibrary
+UCLASS()
+class ENGINE_API UKismetInputLibrary : public UBlueprintFunctionLibrary
 {
 	GENERATED_UCLASS_BODY()
 
diff --git a/Engine/Source/Runtime/Engine/Classes/Kismet/KismetMathLibrary.h b/Engine/Source/Runtime/Engine/Classes/Kismet/KismetMathLibrary.h
--- a/Engine/Source/Runtime/Engine/Classes/Kismet/KismetMathLibrary.h
+++ b/Engine/Source/Runtime/Engine/Classes/Kismet/KismetMathLibrary.h
@@ -55,8 +55,8 @@ namespace EEasingFunc
 	};
 }
 
-UCLASS(MinimalAPI)
-class UKismetMathLibrary : public UBlueprintFunctionLibrary
+UCLASS()
+class ENGINE_API UKismetMathLibrary : public UBlueprintFunctionLibrary
 {
 	GENERATED_UCLASS_BODY()
 
diff --git a/Engine/Source/Runtime/Engine/Classes/Kismet/KismetNodeHelperLibrary.h b/Engine/Source/Runtime/Engine/Classes/Kismet/KismetNodeHelperLibrary.h
--- a/Engine/Source/Runtime/Engine/Classes/Kismet/KismetNodeHelperLibrary.h
+++ b/Engine/Source/Runtime/Engine/Classes/Kismet/KismetNodeHelperLibrary.h
@@ -9,8 +9,8 @@
 
 #include "KismetNodeHelperLibrary.generated.h"
 
-UCLASS(MinimalAPI)
-class UKismetNodeHelperLibrary : public UBlueprintFunctionLibrary
+UCLASS()
+class ENGINE_API UKismetNodeHelperLibrary : public UBlueprintFunctionLibrary
 {
 	GENERATED_UCLASS_BODY()
 
diff --git a/Engine/Source/Runtime/Engine/Classes/Kismet/KismetStringLibrary.h b/Engine/Source/Runtime/Engine/Classes/Kismet/KismetStringLibrary.h
--- a/Engine/Source/Runtime/Engine/Classes/Kismet/KismetStringLibrary.h
+++ b/Engine/Source/Runtime/Engine/Classes/Kismet/KismetStringLibrary.h
@@ -4,8 +4,8 @@
 #include "BlueprintFunctionLibrary.h"
 #include "KismetStringLibrary.generated.h"
 
-UCLASS(MinimalAPI)
-class UKismetStringLibrary : public UBlueprintFunctionLibrary
+UCLASS()
+class ENGINE_API UKismetStringLibrary : public UBlueprintFunctionLibrary
 {
 	GENERATED_UCLASS_BODY()
 
diff --git a/Engine/Source/Runtime/Engine/Classes/Kismet/KismetSystemLibrary.h b/Engine/Source/Runtime/Engine/Classes/Kismet/KismetSystemLibrary.h
--- a/Engine/Source/Runtime/Engine/Classes/Kismet/KismetSystemLibrary.h
+++ b/Engine/Source/Runtime/Engine/Classes/Kismet/KismetSystemLibrary.h
@@ -55,8 +55,8 @@ struct FGenericStruct
 	int32 Data;
 };
 
-UCLASS(MinimalAPI)
-class UKismetSystemLibrary : public UBlueprintFunctionLibrary
+UCLASS()
+class ENGINE_API UKismetSystemLibrary : public UBlueprintFunctionLibrary
 {
 	GENERATED_UCLASS_BODY()
 
@@ -73,19 +73,19 @@ class UKismetSystemLibrary : public UBlueprintFunctionLibrary
 
 	// Return true if the object is usable : non-null and not pending kill
 	UFUNCTION(BlueprintPure, Category = "Utilities")
-	static ENGINE_API bool IsValid(const UObject* Object);
+	static bool IsValid(const UObject* Object);
 
 	// Return true if the class is usable : non-null and not pending kill
 	UFUNCTION(BlueprintPure, Category = "Utilities")
-	static ENGINE_API bool IsValidClass(UClass* Class);
+	static bool IsValidClass(UClass* Class);
 
 	// Returns the display name (or actor label), for displaying to end users.  Note:  In editor builds, this is the actor label.  In non-editor builds, this is the actual object name.  This function should not be used to uniquely identify actors!
 	UFUNCTION(BlueprintPure, Category="Utilities")
-	static ENGINE_API FString GetDisplayName(const UObject* Object);
+	static FString GetDisplayName(const UObject* Object);
 
 	// Returns the display name of a class
 	UFUNCTION(BlueprintPure, Category = "Utilities", meta = (FriendlyName = "Get Display Name"))
-	static ENGINE_API FString GetClassDisplayName(UClass* Class);
+	static FString GetClassDisplayName(UClass* Class);
 
 	// Engine build number, for displaying to end users.
 	UFUNCTION(BlueprintPure, Category="Development")
@@ -487,7 +487,7 @@ class UKismetSystemLibrary : public UBlueprintFunctionLibrary
 	UFUNCTION(BlueprintCallable, CustomThunk, meta = (BlueprintInternalUseOnly = "true", CustomStructureParam = "Value"))
 	static void SetStructurePropertyByName(UObject* Object, FName PropertyName, const FGenericStruct& Value);
 
-	ENGINE_API static void Generic_SetStructurePropertyByName(UObject* OwnerObject, FName StructPropertyName, const void* SrcStructAddr);
+	static void Generic_SetStructurePropertyByName(UObject* OwnerObject, FName StructPropertyName, const void* SrcStructAddr);
 
 	/** Based on UKismetArrayLibrary::execSetArrayPropertyByName */
 	DECLARE_FUNCTION(execSetStructurePropertyByName)
@@ -519,10 +519,10 @@ class UKismetSystemLibrary : public UBlueprintFunctionLibrary
 	 * @return				true if there was an overlap that passed the filters, false otherwise.
 	 */
 	UFUNCTION(BlueprintCallable, Category="Collision", meta=(WorldContext="WorldContextObject", AutoCreateRefTerm="ActorsToIgnore", FriendlyName = "SphereOverlapActors"))
-	static ENGINE_API bool SphereOverlapActors_NEW(UObject* WorldContextObject, const FVector SpherePos, float SphereRadius, const TArray<TEnumAsByte<EObjectTypeQuery> > & ObjectTypes, UClass* ActorClassFilter, const TArray<AActor*>& ActorsToIgnore, TArray<class AActor*>& OutActors);
+	static bool SphereOverlapActors_NEW(UObject* WorldContextObject, const FVector SpherePos, float SphereRadius, const TArray<TEnumAsByte<EObjectTypeQuery> > & ObjectTypes, UClass* ActorClassFilter, const TArray<AActor*>& ActorsToIgnore, TArray<class AActor*>& OutActors);
 
 	UFUNCTION(BlueprintCallable, Category="Collision", meta=(DeprecatedFunction, DeprecationMessage = "Use new SphereOverlapActors", WorldContext="WorldContextObject", AutoCreateRefTerm="ActorsToIgnore"))
-	static ENGINE_API bool SphereOverlapActors_DEPRECATED(UObject* WorldContextObject, const FVector SpherePos, float SphereRadius, EOverlapFilterOption Filter, UClass* ActorClassFilter, const TArray<AActor*>& ActorsToIgnore, TArray<class AActor*>& OutActors);
+	static bool SphereOverlapActors_DEPRECATED(UObject* WorldContextObject, const FVector SpherePos, float SphereRadius, EOverlapFilterOption Filter, UClass* ActorClassFilter, const TArray<AActor*>& ActorsToIgnore, TArray<class AActor*>& OutActors);
 
 	/**
 	 * Returns an array of components that overlap the given sphere.
@@ -536,10 +536,10 @@ class UKismetSystemLibrary : public UBlueprintFunctionLibrary
 	 * @return				true if there was an overlap that passed the filters, false otherwise.
 	 */
 	UFUNCTION(BlueprintCallable, Category="Collision", meta=(WorldContext="WorldContextObject", AutoCreateRefTerm="ActorsToIgnore", FriendlyName="SphereOverlapComponents"))
-	static ENGINE_API bool SphereOverlapComponents_NEW(UObject* WorldContextObject, const FVector SpherePos, float SphereRadius, const TArray<TEnumAsByte<EObjectTypeQuery> > & ObjectTypes, UClass* ComponentClassFilter, const TArray<AActor*>& ActorsToIgnore, TArray<class UPrimitiveComponent*>& OutComponents);
+	static bool SphereOverlapComponents_NEW(UObject* WorldContextObject, const FVector SpherePos, float SphereRadius, const TArray<TEnumAsByte<EObjectTypeQuery> > & ObjectTypes, UClass* ComponentClassFilter, const TArray<AActor*>& ActorsToIgnore, TArray<class UPrimitiveComponent*>& OutComponents);
 
 	UFUNCTION(BlueprintCallable, Category="Collision", meta=(DeprecatedFunction, DeprecationMessage = "Use new SphereOverlapComponents", WorldContext="WorldContextObject", AutoCreateRefTerm="ActorsToIgnore"))
-	static ENGINE_API bool SphereOverlapComponents_DEPRECATED(UObject* WorldContextObject, const FVector SpherePos, float SphereRadius, EOverlapFilterOption Filter, UClass* ComponentClassFilter, const TArray<AActor*>& ActorsToIgnore, TArray<class UPrimitiveComponent*>& OutComponents);
+	static bool SphereOverlapComponents_DEPRECATED(UObject* WorldContextObject, const FVector SpherePos, float SphereRadius, EOverlapFilterOption Filter, UClass* ComponentClassFilter, const TArray<AActor*>& ActorsToIgnore, TArray<class UPrimitiveComponent*>& OutComponents);
 	
 
 	/**
@@ -554,10 +554,10 @@ class UKismetSystemLibrary : public UBlueprintFunctionLibrary
 	 * @return				true if there was an overlap that passed the filters, false otherwise.
 	 */
 	UFUNCTION(BlueprintCallable, Category="Collision", meta=(WorldContext="WorldContextObject", AutoCreateRefTerm="ActorsToIgnore", FriendlyName="BoxOverlapActors"))
-	static ENGINE_API bool BoxOverlapActors_NEW(UObject* WorldContextObject, const FVector BoxPos, FVector BoxExtent, const TArray<TEnumAsByte<EObjectTypeQuery> > & ObjectTypes, UClass* ActorClassFilter, const TArray<AActor*>& ActorsToIgnore, TArray<class AActor*>& OutActors);
+	static bool BoxOverlapActors_NEW(UObject* WorldContextObject, const FVector BoxPos, FVector BoxExtent, const TArray<TEnumAsByte<EObjectTypeQuery> > & ObjectTypes, UClass* ActorClassFilter, const TArray<AActor*>& ActorsToIgnore, TArray<class AActor*>& OutActors);
 
 	UFUNCTION(BlueprintCallable, Category="Collision", meta=(DeprecatedFunction, DeprecationMessage = "Use new BoxOverlapActors", WorldContext="WorldContextObject", AutoCreateRefTerm="ActorsToIgnore"))
-	static ENGINE_API bool BoxOverlapActors_DEPRECATED(UObject* WorldContextObject, const FVector BoxPos, FVector BoxExtent, EOverlapFilterOption Filter, UClass* ActorClassFilter, const TArray<AActor*>& ActorsToIgnore, TArray<class AActor*>& OutActors);
+	static bool BoxOverlapActors_DEPRECATED(UObject* WorldContextObject, const FVector BoxPos, FVector BoxExtent, EOverlapFilterOption Filter, UClass* ActorClassFilter, const TArray<AActor*>& ActorsToIgnore, TArray<class AActor*>& OutActors);
 
 	/**
 	 * Returns an array of components that overlap the given axis-aligned box.
@@ -571,10 +571,10 @@ class UKismetSystemLibrary : public UBlueprintFunctionLibrary
 	 * @return				true if there was an overlap that passed the filters, false otherwise.
 	 */
 	UFUNCTION(BlueprintCallable, Category="Collision", meta=(WorldContext="WorldContextObject", AutoCreateRefTerm="ActorsToIgnore", FriendlyName="BoxOverlapComponents"))
-	static ENGINE_API bool BoxOverlapComponents_NEW(UObject* WorldContextObject, const FVector BoxPos, FVector Extent, const TArray<TEnumAsByte<EObjectTypeQuery> > & ObjectTypes, UClass* ComponentClassFilter, const TArray<AActor*>& ActorsToIgnore, TArray<class UPrimitiveComponent*>& OutComponents);
+	static bool BoxOverlapComponents_NEW(UObject* WorldContextObject, const FVector BoxPos, FVector Extent, const TArray<TEnumAsByte<EObjectTypeQuery> > & ObjectTypes, UClass* ComponentClassFilter, const TArray<AActor*>& ActorsToIgnore, TArray<class UPrimitiveComponent*>& OutComponents);
 
 	UFUNCTION(BlueprintCallable, Category="Collision", meta=(DeprecatedFunction, DeprecationMessage = "Use new BoxOverlapComponents", WorldContext="WorldContextObject", AutoCreateRefTerm="ActorsToIgnore"))
-	static ENGINE_API bool BoxOverlapComponents_DEPRECATED(UObject* WorldContextObject, const FVector BoxPos, FVector Extent, EOverlapFilterOption Filter, UClass* ComponentClassFilter, const TArray<AActor*>& ActorsToIgnore, TArray<class UPrimitiveComponent*>& OutComponents);
+	static bool BoxOverlapComponents_DEPRECATED(UObject* WorldContextObject, const FVector BoxPos, FVector Extent, EOverlapFilterOption Filter, UClass* ComponentClassFilter, const TArray<AActor*>& ActorsToIgnore, TArray<class UPrimitiveComponent*>& OutComponents);
 
 
 	/**
@@ -590,10 +590,10 @@ class UKismetSystemLibrary : public UBlueprintFunctionLibrary
 	 * @return				true if there was an overlap that passed the filters, false otherwise.
 	 */
 	UFUNCTION(BlueprintCallable, Category="Collision", meta=(WorldContext="WorldContextObject", AutoCreateRefTerm="ActorsToIgnore", FriendlyName="CapsuleOverlapActors"))
-	static ENGINE_API bool CapsuleOverlapActors_NEW(UObject* WorldContextObject, const FVector CapsulePos, float Radius, float HalfHeight, const TArray<TEnumAsByte<EObjectTypeQuery> > & ObjectTypes, UClass* ActorClassFilter, const TArray<AActor*>& ActorsToIgnore, TArray<class AActor*>& OutActors);
+	static bool CapsuleOverlapActors_NEW(UObject* WorldContextObject, const FVector CapsulePos, float Radius, float HalfHeight, const TArray<TEnumAsByte<EObjectTypeQuery> > & ObjectTypes, UClass* ActorClassFilter, const TArray<AActor*>& ActorsToIgnore, TArray<class AActor*>& OutActors);
 
 	UFUNCTION(BlueprintCallable, Category="Collision", meta=(DeprecatedFunction, DeprecationMessage = "Use new CapsuleOverlapActors", WorldContext="WorldContextObject", AutoCreateRefTerm="ActorsToIgnore"))
-	static ENGINE_API bool CapsuleOverlapActors_DEPRECATED(UObject* WorldContextObject, const FVector CapsulePos, float Radius, float HalfHeight, EOverlapFilterOption Filter, UClass* ActorClassFilter, const TArray<AActor*>& ActorsToIgnore, TArray<class AActor*>& OutActors);
+	static bool CapsuleOverlapActors_DEPRECATED(UObject* WorldContextObject, const FVector CapsulePos, float Radius, float HalfHeight, EOverlapFilterOption Filter, UClass* ActorClassFilter, const TArray<AActor*>& ActorsToIgnore, TArray<class AActor*>& OutActors);
 	/**
 	 * Returns an array of components that overlap the given capsule.
 	 * @param WorldContext	World context
@@ -607,10 +607,10 @@ class UKismetSystemLibrary : public UBlueprintFunctionLibrary
 	 * @return				true if there was an overlap that passed the filters, false otherwise.
 	 */
 	UFUNCTION(BlueprintCallable, Category="Collision", meta=(WorldContext="WorldContextObject", AutoCreateRefTerm="ActorsToIgnore", FriendlyName="CapsuleOverlapComponents") )
-	static ENGINE_API bool CapsuleOverlapComponents_NEW(UObject* WorldContextObject, const FVector CapsulePos, float Radius, float HalfHeight, const TArray<TEnumAsByte<EObjectTypeQuery> > & ObjectTypes, UClass* ComponentClassFilter, const TArray<AActor*>& ActorsToIgnore, TArray<class UPrimitiveComponent*>& OutComponents);
+	static bool CapsuleOverlapComponents_NEW(UObject* WorldContextObject, const FVector CapsulePos, float Radius, float HalfHeight, const TArray<TEnumAsByte<EObjectTypeQuery> > & ObjectTypes, UClass* ComponentClassFilter, const TArray<AActor*>& ActorsToIgnore, TArray<class UPrimitiveComponent*>& OutComponents);
 
 	UFUNCTION(BlueprintCallable, Category="Collision", meta=(DeprecatedFunction, DeprecationMessage = "Use new CapsuleOverlapComponents", WorldContext="WorldContextObject", AutoCreateRefTerm="ActorsToIgnore") )
-	static ENGINE_API bool CapsuleOverlapComponents_DEPRECATED(UObject* WorldContextObject, const FVector CapsulePos, float Radius, float HalfHeight, EOverlapFilterOption Filter, UClass* ComponentClassFilter, const TArray<AActor*>& ActorsToIgnore, TArray<class UPrimitiveComponent*>& OutComponents);
+	static bool CapsuleOverlapComponents_DEPRECATED(UObject* WorldContextObject, const FVector CapsulePos, float Radius, float HalfHeight, EOverlapFilterOption Filter, UClass* ComponentClassFilter, const TArray<AActor*>& ActorsToIgnore, TArray<class UPrimitiveComponent*>& OutComponents);
 	
 
 	/**
@@ -624,10 +624,10 @@ class UKismetSystemLibrary : public UBlueprintFunctionLibrary
 	 * @return						true if there was an overlap that passed the filters, false otherwise.
 	 */
 	UFUNCTION(BlueprintCallable, Category="Collision", meta=(AutoCreateRefTerm="ActorsToIgnore", FriendlyName="ComponentOverlapActors"))
-	static ENGINE_API bool ComponentOverlapActors_NEW(UPrimitiveComponent* Component, const FTransform& ComponentTransform, const TArray<TEnumAsByte<EObjectTypeQuery> > & ObjectTypes, UClass* ActorClassFilter, const TArray<AActor*>& ActorsToIgnore, TArray<class AActor*>& OutActors);
+	static bool ComponentOverlapActors_NEW(UPrimitiveComponent* Component, const FTransform& ComponentTransform, const TArray<TEnumAsByte<EObjectTypeQuery> > & ObjectTypes, UClass* ActorClassFilter, const TArray<AActor*>& ActorsToIgnore, TArray<class AActor*>& OutActors);
 
 	UFUNCTION(BlueprintCallable, Category="Collision", meta=(DeprecatedFunction, DeprecationMessage = "Use new ComponentOverlapActors", AutoCreateRefTerm="ActorsToIgnore"))
-	static ENGINE_API bool ComponentOverlapActors_DEPRECATED(UPrimitiveComponent* Component, const FTransform& ComponentTransform, EOverlapFilterOption Filter, UClass* ActorClassFilter, const TArray<AActor*>& ActorsToIgnore, TArray<class AActor*>& OutActors);
+	static bool ComponentOverlapActors_DEPRECATED(UPrimitiveComponent* Component, const FTransform& ComponentTransform, EOverlapFilterOption Filter, UClass* ActorClassFilter, const TArray<AActor*>& ActorsToIgnore, TArray<class AActor*>& OutActors);
 	
 	/**
 	 * Returns an array of components that overlap the given component.
@@ -640,10 +640,10 @@ class UKismetSystemLibrary : public UBlueprintFunctionLibrary
 	 * @return						true if there was an overlap that passed the filters, false otherwise.
 	 */
 	UFUNCTION(BlueprintCallable, Category="Collision", meta=(AutoCreateRefTerm="ActorsToIgnore", FriendlyName="ComponentOverlapComponents"))
-	static ENGINE_API bool ComponentOverlapComponents_NEW(UPrimitiveComponent* Component, const FTransform& ComponentTransform, const TArray<TEnumAsByte<EObjectTypeQuery> > & ObjectTypes, UClass* ComponentClassFilter, const TArray<AActor*>& ActorsToIgnore, TArray<class UPrimitiveComponent*>& OutComponents);
+	static bool ComponentOverlapComponents_NEW(UPrimitiveComponent* Component, const FTransform& ComponentTransform, const TArray<TEnumAsByte<EObjectTypeQuery> > & ObjectTypes, UClass* ComponentClassFilter, const TArray<AActor*>& ActorsToIgnore, TArray<class UPrimitiveComponent*>& OutComponents);
 
 	UFUNCTION(BlueprintCallable, Category="Collision", meta=(DeprecatedFunction, DeprecationMessage = "Use new ComponentOverlapComponents", AutoCreateRefTerm="ActorsToIgnore"))
-	static ENGINE_API bool ComponentOverlapComponents_DEPRECATED(UPrimitiveComponent* Component, const FTransform& ComponentTransform, EOverlapFilterOption Filter, UClass* ComponentClassFilter, const TArray<AActor*>& ActorsToIgnore, TArray<class UPrimitiveComponent*>& OutComponents);
+	static bool ComponentOverlapComponents_DEPRECATED(UPrimitiveComponent* Component, const FTransform& ComponentTransform, EOverlapFilterOption Filter, UClass* ComponentClassFilter, const TArray<AActor*>& ActorsToIgnore, TArray<class UPrimitiveComponent*>& OutComponents);
 
 	/**
 	 * Does a collision trace along the given line and returns the first blocking hit encountered.
@@ -658,7 +658,7 @@ class UKismetSystemLibrary : public UBlueprintFunctionLibrary
 	 * @return				True if there was a hit, false otherwise.
 	 */
 	UFUNCTION(BlueprintCallable, Category="Collision", meta=(bIgnoreSelf="true", WorldContext="WorldContextObject", AutoCreateRefTerm="ActorsToIgnore", FriendlyName="LineTraceByChannel", Keywords="raycast"))
-	static ENGINE_API bool LineTraceSingle_NEW(UObject* WorldContextObject, const FVector Start, const FVector End, ETraceTypeQuery TraceChannel, bool bTraceComplex, const TArray<AActor*>& ActorsToIgnore, EDrawDebugTrace::Type DrawDebugType, FHitResult& OutHit, bool bIgnoreSelf);
+	static bool LineTraceSingle_NEW(UObject* WorldContextObject, const FVector Start, const FVector End, ETraceTypeQuery TraceChannel, bool bTraceComplex, const TArray<AActor*>& ActorsToIgnore, EDrawDebugTrace::Type DrawDebugType, FHitResult& OutHit, bool bIgnoreSelf);
 	
 	/**
 	 * Does a collision trace along the given line and returns all hits encountered up to and including the first blocking hit.
@@ -673,7 +673,7 @@ class UKismetSystemLibrary : public UBlueprintFunctionLibrary
 	 * @return				True if there was a hit, false otherwise.
 	 */
 	UFUNCTION(BlueprintCallable, Category="Collision", meta=(bIgnoreSelf="true", WorldContext="WorldContextObject", AutoCreateRefTerm="ActorsToIgnore", FriendlyName = "MultiLineTraceByChannel", Keywords="raycast"))
-	static ENGINE_API bool LineTraceMulti_NEW(UObject* WorldContextObject, const FVector Start, const FVector End, ETraceTypeQuery TraceChannel, bool bTraceComplex, const TArray<AActor*>& ActorsToIgnore, EDrawDebugTrace::Type DrawDebugType, TArray<FHitResult>& OutHits, bool bIgnoreSelf);
+	static bool LineTraceMulti_NEW(UObject* WorldContextObject, const FVector Start, const FVector End, ETraceTypeQuery TraceChannel, bool bTraceComplex, const TArray<AActor*>& ActorsToIgnore, EDrawDebugTrace::Type DrawDebugType, TArray<FHitResult>& OutHits, bool bIgnoreSelf);
 
 	/**
 	 * Sweeps a sphere along the given line and returns the first blocking hit encountered.
@@ -688,7 +688,7 @@ class UKismetSystemLibrary : public UBlueprintFunctionLibrary
 	 * @return				True if there was a hit, false otherwise.
 	 */
 	UFUNCTION(BlueprintCallable, Category="Collision", meta=(bIgnoreSelf="true", WorldContext="WorldContextObject", AutoCreateRefTerm="ActorsToIgnore", FriendlyName = "SphereTraceByChannel", Keywords="sweep"))
-	static ENGINE_API bool SphereTraceSingle_NEW(UObject* WorldContextObject, const FVector Start, const FVector End, float Radius, ETraceTypeQuery TraceChannel, bool bTraceComplex, const TArray<AActor*>& ActorsToIgnore, EDrawDebugTrace::Type DrawDebugType, FHitResult& OutHit, bool bIgnoreSelf);
+	static bool SphereTraceSingle_NEW(UObject* WorldContextObject, const FVector Start, const FVector End, float Radius, ETraceTypeQuery TraceChannel, bool bTraceComplex, const TArray<AActor*>& ActorsToIgnore, EDrawDebugTrace::Type DrawDebugType, FHitResult& OutHit, bool bIgnoreSelf);
 
 	/**
 	 * Sweeps a sphere along the given line and returns all hits encountered up to and including the first blocking hit.
@@ -704,7 +704,7 @@ class UKismetSystemLibrary : public UBlueprintFunctionLibrary
 	 * @return				True if there was a hit, false otherwise.
 	 */
 	UFUNCTION(BlueprintCallable, Category="Collision", meta=(bIgnoreSelf="true", WorldContext="WorldContextObject", AutoCreateRefTerm="ActorsToIgnore", FriendlyName = "MultiSphereTraceByChannel", Keywords="sweep"))
-	static ENGINE_API bool SphereTraceMulti_NEW(UObject* WorldContextObject, const FVector Start, const FVector End, float Radius, ETraceTypeQuery TraceChannel, bool bTraceComplex, const TArray<AActor*>& ActorsToIgnore, EDrawDebugTrace::Type DrawDebugType, TArray<FHitResult>& OutHits, bool bIgnoreSelf);
+	static bool SphereTraceMulti_NEW(UObject* WorldContextObject, const FVector Start, const FVector End, float Radius, ETraceTypeQuery TraceChannel, bool bTraceComplex, const TArray<AActor*>& ActorsToIgnore, EDrawDebugTrace::Type DrawDebugType, TArray<FHitResult>& OutHits, bool bIgnoreSelf);
 
 	/**
 	* Sweeps a box along the given line and returns the first blocking hit encountered.
@@ -720,7 +720,7 @@ class UKismetSystemLibrary : public UBlueprintFunctionLibrary
 	* @return				True if there was a hit, false otherwise.
 	*/
 	UFUNCTION(BlueprintCallable, Category = "Collision", meta = (bIgnoreSelf = "true", HidePin = "WorldContextObject", DefaultToSelf = "WorldContextObject", AutoCreateRefTerm = "ActorsToIgnore", FriendlyName = "BoxTraceByChannel", Keywords="sweep"))
-	static ENGINE_API bool BoxTraceSingle(UObject* WorldContextObject, const FVector Start, const FVector End, const FVector HalfSize, const FRotator Orientation, ETraceTypeQuery TraceChannel, bool bTraceComplex, const TArray<AActor*>& ActorsToIgnore, EDrawDebugTrace::Type DrawDebugType, FHitResult& OutHit, bool bIgnoreSelf);
+	static bool BoxTraceSingle(UObject* WorldContextObject, const FVector Start, const FVector End, const FVector HalfSize, const FRotator Orientation, ETraceTypeQuery TraceChannel, bool bTraceComplex, const TArray<AActor*>& ActorsToIgnore, EDrawDebugTrace::Type DrawDebugType, FHitResult& OutHit, bool bIgnoreSelf);
 
 	/**
 	* Sweeps a box along the given line and returns all hits encountered.
@@ -736,7 +736,7 @@ class UKismetSystemLibrary : public UBlueprintFunctionLibrary
 	* @return				True if there was a hit, false otherwise.
 	*/
 	UFUNCTION(BlueprintCallable, Category = "Collision", meta = (bIgnoreSelf = "true", HidePin = "WorldContextObject", DefaultToSelf = "WorldContextObject", AutoCreateRefTerm = "ActorsToIgnore", FriendlyName = "MultiBoxTraceByChannel", Keywords="sweep"))
-	static ENGINE_API bool BoxTraceMulti(UObject* WorldContextObject, const FVector Start, const FVector End, FVector HalfSize, const FRotator Orientation, ETraceTypeQuery TraceChannel, bool bTraceComplex, const TArray<AActor*>& ActorsToIgnore, EDrawDebugTrace::Type DrawDebugType, TArray<FHitResult>& OutHits, bool bIgnoreSelf);
+	static bool BoxTraceMulti(UObject* WorldContextObject, const FVector Start, const FVector End, FVector HalfSize, const FRotator Orientation, ETraceTypeQuery TraceChannel, bool bTraceComplex, const TArray<AActor*>& ActorsToIgnore, EDrawDebugTrace::Type DrawDebugType, TArray<FHitResult>& OutHits, bool bIgnoreSelf);
 
 	
 	/**
@@ -754,7 +754,7 @@ class UKismetSystemLibrary : public UBlueprintFunctionLibrary
 	 * @return				True if there was a hit, false otherwise.
 	 */
 	UFUNCTION(BlueprintCallable, Category="Collision", meta=(bIgnoreSelf="true", WorldContext="WorldContextObject", AutoCreateRefTerm="ActorsToIgnore", FriendlyName = "CapsuleTraceByChannel", Keywords="sweep"))
-	static ENGINE_API bool CapsuleTraceSingle_NEW(UObject* WorldContextObject, const FVector Start, const FVector End, float Radius, float HalfHeight, ETraceTypeQuery TraceChannel, bool bTraceComplex, const TArray<AActor*>& ActorsToIgnore, EDrawDebugTrace::Type DrawDebugType, FHitResult& OutHit, bool bIgnoreSelf);
+	static bool CapsuleTraceSingle_NEW(UObject* WorldContextObject, const FVector Start, const FVector End, float Radius, float HalfHeight, ETraceTypeQuery TraceChannel, bool bTraceComplex, const TArray<AActor*>& ActorsToIgnore, EDrawDebugTrace::Type DrawDebugType, FHitResult& OutHit, bool bIgnoreSelf);
 
 	/**
 	 * Sweeps a capsule along the given line and returns all hits encountered up to and including the first blocking hit.
@@ -771,7 +771,7 @@ class UKismetSystemLibrary : public UBlueprintFunctionLibrary
 	 * @return				True if there was a hit, false otherwise.
 	 */
 	UFUNCTION(BlueprintCallable, Category="Collision", meta=(bIgnoreSelf="true", WorldContext="WorldContextObject", AutoCreateRefTerm="ActorsToIgnore", FriendlyName = "MultiCapsuleTraceByChannel", Keywords="sweep"))
-	static ENGINE_API bool CapsuleTraceMulti_NEW(UObject* WorldContextObject, const FVector Start, const FVector End, float Radius, float HalfHeight, ETraceTypeQuery TraceChannel, bool bTraceComplex, const TArray<AActor*>& ActorsToIgnore, EDrawDebugTrace::Type DrawDebugType, TArray<FHitResult>& OutHits, bool bIgnoreSelf);
+	static bool CapsuleTraceMulti_NEW(UObject* WorldContextObject, const FVector Start, const FVector End, float Radius, float HalfHeight, ETraceTypeQuery TraceChannel, bool bTraceComplex, const TArray<AActor*>& ActorsToIgnore, EDrawDebugTrace::Type DrawDebugType, TArray<FHitResult>& OutHits, bool bIgnoreSelf);
 
 	/**
 	 * Does a collision trace along the given line and returns the first hit encountered.
@@ -786,7 +786,7 @@ class UKismetSystemLibrary : public UBlueprintFunctionLibrary
 	 * @return				True if there was a hit, false otherwise.
 	 */
 	UFUNCTION(BlueprintCallable, Category="Collision", meta=(bIgnoreSelf="true", WorldContext="WorldContextObject", AutoCreateRefTerm="ActorsToIgnore", FriendlyName = "LineTraceForObjects", Keywords="raycast"))
-	static ENGINE_API bool LineTraceSingleForObjects(UObject* WorldContextObject, const FVector Start, const FVector End, const TArray<TEnumAsByte<EObjectTypeQuery> > & ObjectTypes, bool bTraceComplex, const TArray<AActor*>& ActorsToIgnore, EDrawDebugTrace::Type DrawDebugType, FHitResult& OutHit, bool bIgnoreSelf);
+	static bool LineTraceSingleForObjects(UObject* WorldContextObject, const FVector Start, const FVector End, const TArray<TEnumAsByte<EObjectTypeQuery> > & ObjectTypes, bool bTraceComplex, const TArray<AActor*>& ActorsToIgnore, EDrawDebugTrace::Type DrawDebugType, FHitResult& OutHit, bool bIgnoreSelf);
 	
 	/**
 	 * Does a collision trace along the given line and returns all hits encountered.
@@ -801,7 +801,7 @@ class UKismetSystemLibrary : public UBlueprintFunctionLibrary
 	 * @return				True if there was a hit, false otherwise.
 	 */
 	UFUNCTION(BlueprintCallable, Category="Collision", meta=(bIgnoreSelf="true", WorldContext="WorldContextObject", AutoCreateRefTerm="ActorsToIgnore", FriendlyName = "MultiLineTraceForObjects", Keywords="raycast"))
-	static ENGINE_API bool LineTraceMultiForObjects(UObject* WorldContextObject, const FVector Start, const FVector End, const TArray<TEnumAsByte<EObjectTypeQuery> > & ObjectTypes, bool bTraceComplex, const TArray<AActor*>& ActorsToIgnore, EDrawDebugTrace::Type DrawDebugType, TArray<FHitResult>& OutHits, bool bIgnoreSelf);
+	static bool LineTraceMultiForObjects(UObject* WorldContextObject, const FVector Start, const FVector End, const TArray<TEnumAsByte<EObjectTypeQuery> > & ObjectTypes, bool bTraceComplex, const TArray<AActor*>& ActorsToIgnore, EDrawDebugTrace::Type DrawDebugType, TArray<FHitResult>& OutHits, bool bIgnoreSelf);
 
 	/**
 	 * Sweeps a sphere along the given line and returns the first hit encountered.
@@ -816,7 +816,7 @@ class UKismetSystemLibrary : public UBlueprintFunctionLibrary
 	 * @return				True if there was a hit, false otherwise.
 	 */
 	UFUNCTION(BlueprintCallable, Category="Collision", meta=(bIgnoreSelf="true", WorldContext="WorldContextObject", AutoCreateRefTerm="ActorsToIgnore", FriendlyName = "SphereTraceForObjects", Keywords="sweep"))
-	static ENGINE_API bool SphereTraceSingleForObjects(UObject* WorldContextObject, const FVector Start, const FVector End, float Radius, const TArray<TEnumAsByte<EObjectTypeQuery> > & ObjectTypes, bool bTraceComplex, const TArray<AActor*>& ActorsToIgnore, EDrawDebugTrace::Type DrawDebugType, FHitResult& OutHit, bool bIgnoreSelf);
+	static bool SphereTraceSingleForObjects(UObject* WorldContextObject, const FVector Start, const FVector End, float Radius, const TArray<TEnumAsByte<EObjectTypeQuery> > & ObjectTypes, bool bTraceComplex, const TArray<AActor*>& ActorsToIgnore, EDrawDebugTrace::Type DrawDebugType, FHitResult& OutHit, bool bIgnoreSelf);
 
 	/**
 	 * Sweeps a sphere along the given line and returns all hits encountered.
@@ -832,7 +832,7 @@ class UKismetSystemLibrary : public UBlueprintFunctionLibrary
 	 * @return				True if there was a hit, false otherwise.
 	 */
 	UFUNCTION(BlueprintCallable, Category="Collision", meta=(bIgnoreSelf="true", WorldContext="WorldContextObject", AutoCreateRefTerm="ActorsToIgnore", FriendlyName = "MultiSphereTraceForObjects", Keywords="sweep"))
-	static ENGINE_API bool SphereTraceMultiForObjects(UObject* WorldContextObject, const FVector Start, const FVector End, float Radius, const TArray<TEnumAsByte<EObjectTypeQuery> > & ObjectTypes, bool bTraceComplex, const TArray<AActor*>& ActorsToIgnore, EDrawDebugTrace::Type DrawDebugType, TArray<FHitResult>& OutHits, bool bIgnoreSelf);
+	static bool SphereTraceMultiForObjects(UObject* WorldContextObject, const FVector Start, const FVector End, float Radius, const TArray<TEnumAsByte<EObjectTypeQuery> > & ObjectTypes, bool bTraceComplex, const TArray<AActor*>& ActorsToIgnore, EDrawDebugTrace::Type DrawDebugType, TArray<FHitResult>& OutHits, bool bIgnoreSelf);
 	
 
 	/**
@@ -849,7 +849,7 @@ class UKismetSystemLibrary : public UBlueprintFunctionLibrary
 	* @return				True if there was a hit, false otherwise.
 	*/
 	UFUNCTION(BlueprintCallable, Category = "Collision", meta = (bIgnoreSelf = "true", HidePin = "WorldContextObject", DefaultToSelf = "WorldContextObject", AutoCreateRefTerm = "ActorsToIgnore", FriendlyName = "BoxTraceForObjects", Keywords="sweep"))
-	static ENGINE_API bool BoxTraceSingleForObjects(UObject* WorldContextObject, const FVector Start, const FVector End, const FVector HalfSize, const FRotator Orientation, const TArray<TEnumAsByte<EObjectTypeQuery> > & ObjectTypes, bool bTraceComplex, const TArray<AActor*>& ActorsToIgnore, EDrawDebugTrace::Type DrawDebugType, FHitResult& OutHit, bool bIgnoreSelf);
+	static bool BoxTraceSingleForObjects(UObject* WorldContextObject, const FVector Start, const FVector End, const FVector HalfSize, const FRotator Orientation, const TArray<TEnumAsByte<EObjectTypeQuery> > & ObjectTypes, bool bTraceComplex, const TArray<AActor*>& ActorsToIgnore, EDrawDebugTrace::Type DrawDebugType, FHitResult& OutHit, bool bIgnoreSelf);
 
 
 	/**
@@ -866,7 +866,7 @@ class UKismetSystemLibrary : public UBlueprintFunctionLibrary
 	* @return				True if there was a hit, false otherwise.
 	*/
 	UFUNCTION(BlueprintCallable, Category = "Collision", meta = (bIgnoreSelf = "true", HidePin = "WorldContextObject", DefaultToSelf = "WorldContextObject", AutoCreateRefTerm = "ActorsToIgnore", FriendlyName = "MultiBoxTraceForObjects", Keywords="sweep"))
-	static ENGINE_API bool BoxTraceMultiForObjects(UObject* WorldContextObject, const FVector Start, const FVector End, const FVector HalfSize, const FRotator Orientation, const TArray<TEnumAsByte<EObjectTypeQuery> > & ObjectTypes, bool bTraceComplex, const TArray<AActor*>& ActorsToIgnore, EDrawDebugTrace::Type DrawDebugType, TArray<FHitResult>& OutHits, bool bIgnoreSelf);
+	static bool BoxTraceMultiForObjects(UObject* WorldContextObject, const FVector Start, const FVector End, const FVector HalfSize, const FRotator Orientation, const TArray<TEnumAsByte<EObjectTypeQuery> > & ObjectTypes, bool bTraceComplex, const TArray<AActor*>& ActorsToIgnore, EDrawDebugTrace::Type DrawDebugType, TArray<FHitResult>& OutHits, bool bIgnoreSelf);
 
 	/**
 	 * Sweeps a capsule along the given line and returns the first hit encountered.
@@ -883,7 +883,7 @@ class UKismetSystemLibrary : public UBlueprintFunctionLibrary
 	 * @return				True if there was a hit, false otherwise.
 	 */
 	UFUNCTION(BlueprintCallable, Category="Collision", meta=(bIgnoreSelf="true", WorldContext="WorldContextObject", AutoCreateRefTerm="ActorsToIgnore", FriendlyName = "CapsuleTraceForObjects", Keywords="sweep"))
-	static ENGINE_API bool CapsuleTraceSingleForObjects(UObject* WorldContextObject, const FVector Start, const FVector End, float Radius, float HalfHeight, const TArray<TEnumAsByte<EObjectTypeQuery> > & ObjectTypes, bool bTraceComplex, const TArray<AActor*>& ActorsToIgnore, EDrawDebugTrace::Type DrawDebugType, FHitResult& OutHit, bool bIgnoreSelf);
+	static bool CapsuleTraceSingleForObjects(UObject* WorldContextObject, const FVector Start, const FVector End, float Radius, float HalfHeight, const TArray<TEnumAsByte<EObjectTypeQuery> > & ObjectTypes, bool bTraceComplex, const TArray<AActor*>& ActorsToIgnore, EDrawDebugTrace::Type DrawDebugType, FHitResult& OutHit, bool bIgnoreSelf);
 
 	/**
 	 * Sweeps a capsule along the given line and returns all hits encountered.
@@ -900,7 +900,7 @@ class UKismetSystemLibrary : public UBlueprintFunctionLibrary
 	 * @return				True if there was a hit, false otherwise.
 	 */
 	UFUNCTION(BlueprintCallable, Category="Collision", meta=(bIgnoreSelf="true", WorldContext="WorldContextObject", AutoCreateRefTerm="ActorsToIgnore", FriendlyName = "MultiCapsuleTraceForObjects", Keywords="sweep"))
-	static ENGINE_API bool CapsuleTraceMultiForObjects(UObject* WorldContextObject, const FVector Start, const FVector End, float Radius, float HalfHeight, const TArray<TEnumAsByte<EObjectTypeQuery> > & ObjectTypes, bool bTraceComplex, const TArray<AActor*>& ActorsToIgnore, EDrawDebugTrace::Type DrawDebugType, TArray<FHitResult>& OutHits, bool bIgnoreSelf);
+	static bool CapsuleTraceMultiForObjects(UObject* WorldContextObject, const FVector Start, const FVector End, float Radius, float HalfHeight, const TArray<TEnumAsByte<EObjectTypeQuery> > & ObjectTypes, bool bTraceComplex, const TArray<AActor*>& ActorsToIgnore, EDrawDebugTrace::Type DrawDebugType, TArray<FHitResult>& OutHits, bool bIgnoreSelf);
 
 	/**
 	 * Does a collision trace along the given line and returns the first blocking hit encountered.
@@ -915,7 +915,7 @@ class UKismetSystemLibrary : public UBlueprintFunctionLibrary
 	 * @return				True if there was a hit, false otherwise.
 	 */
 	UFUNCTION(BlueprintCallable, Category="Collision", meta=(DeprecatedFunction, DeprecationMessage = "Use new SingleLineTraceByChannel", bIgnoreSelf="true", WorldContext="WorldContextObject", AutoCreateRefTerm="ActorsToIgnore", FriendlyName = "SingleLineTraceByChannelDeprecated"))
-	static ENGINE_API bool LineTraceSingle_DEPRECATED(UObject* WorldContextObject, const FVector Start, const FVector End, ECollisionChannel TraceChannel, bool bTraceComplex, const TArray<AActor*>& ActorsToIgnore, EDrawDebugTrace::Type DrawDebugType, FHitResult& OutHit, bool bIgnoreSelf);
+	static bool LineTraceSingle_DEPRECATED(UObject* WorldContextObject, const FVector Start, const FVector End, ECollisionChannel TraceChannel, bool bTraceComplex, const TArray<AActor*>& ActorsToIgnore, EDrawDebugTrace::Type DrawDebugType, FHitResult& OutHit, bool bIgnoreSelf);
 	
 	/**
 	 * Does a collision trace along the given line and returns all hits encountered up to and including the first blocking hit.
@@ -930,7 +930,7 @@ class UKismetSystemLibrary : public UBlueprintFunctionLibrary
 	 * @return				True if there was a hit, false otherwise.
 	 */
 	UFUNCTION(BlueprintCallable, Category="Collision", meta=(DeprecatedFunction, DeprecationMessage = "Use new MultiLineTraceByChannel", bIgnoreSelf="true", WorldContext="WorldContextObject", AutoCreateRefTerm="ActorsToIgnore", FriendlyName = "MultiLineTraceByChannelDeprecated"))
-	static ENGINE_API bool LineTraceMulti_DEPRECATED(UObject* WorldContextObject, const FVector Start, const FVector End, ECollisionChannel TraceChannel, bool bTraceComplex, const TArray<AActor*>& ActorsToIgnore, EDrawDebugTrace::Type DrawDebugType, TArray<FHitResult>& OutHits, bool bIgnoreSelf);
+	static bool LineTraceMulti_DEPRECATED(UObject* WorldContextObject, const FVector Start, const FVector End, ECollisionChannel TraceChannel, bool bTraceComplex, const TArray<AActor*>& ActorsToIgnore, EDrawDebugTrace::Type DrawDebugType, TArray<FHitResult>& OutHits, bool bIgnoreSelf);
 
 	/**
 	 * Sweeps a sphere along the given line and returns the first blocking hit encountered.
@@ -945,7 +945,7 @@ class UKismetSystemLibrary : public UBlueprintFunctionLibrary
 	 * @return				True if there was a hit, false otherwise.
 	 */
 	UFUNCTION(BlueprintCallable, Category="Collision", meta=(DeprecatedFunction, DeprecationMessage = "Use new SingleSphereTraceByChannel", bIgnoreSelf="true", WorldContext="WorldContextObject", AutoCreateRefTerm="ActorsToIgnore", FriendlyName = "SingleSphereTraceByChannelDeprecated"))
-	static ENGINE_API bool SphereTraceSingle_DEPRECATED(UObject* WorldContextObject, const FVector Start, const FVector End, float Radius, ECollisionChannel TraceChannel, bool bTraceComplex, const TArray<AActor*>& ActorsToIgnore, EDrawDebugTrace::Type DrawDebugType, FHitResult& OutHit, bool bIgnoreSelf);
+	static bool SphereTraceSingle_DEPRECATED(UObject* WorldContextObject, const FVector Start, const FVector End, float Radius, ECollisionChannel TraceChannel, bool bTraceComplex, const TArray<AActor*>& ActorsToIgnore, EDrawDebugTrace::Type DrawDebugType, FHitResult& OutHit, bool bIgnoreSelf);
 
 	/**
 	 * Sweeps a sphere along the given line and returns all hits encountered up to and including the first blocking hit.
@@ -961,7 +961,7 @@ class UKismetSystemLibrary : public UBlueprintFunctionLibrary
 	 * @return				True if there was a hit, false otherwise.
 	 */
 	UFUNCTION(BlueprintCallable, Category="Collision", meta=(DeprecatedFunction, DeprecationMessage = "Use new MultiSphereTraceByChannel", bIgnoreSelf="true", WorldContext="WorldContextObject", AutoCreateRefTerm="ActorsToIgnore", FriendlyName = "MultiSphereTraceByChannelDeprecated"))
-	static ENGINE_API bool SphereTraceMulti_DEPRECATED(UObject* WorldContextObject, const FVector Start, const FVector End, float Radius, ECollisionChannel TraceChannel, bool bTraceComplex, const TArray<AActor*>& ActorsToIgnore, EDrawDebugTrace::Type DrawDebugType, TArray<FHitResult>& OutHits, bool bIgnoreSelf);
+	static bool SphereTraceMulti_DEPRECATED(UObject* WorldContextObject, const FVector Start, const FVector End, float Radius, ECollisionChannel TraceChannel, bool bTraceComplex, const TArray<AActor*>& ActorsToIgnore, EDrawDebugTrace::Type DrawDebugType, TArray<FHitResult>& OutHits, bool bIgnoreSelf);
 	
 	/**
 	 * Sweeps a capsule along the given line and returns the first blocking hit encountered.
@@ -978,7 +978,7 @@ class UKismetSystemLibrary : public UBlueprintFunctionLibrary
 	 * @return				True if there was a hit, false otherwise.
 	 */
 	UFUNCTION(BlueprintCallable, Category="Collision", meta=(DeprecatedFunction, DeprecationMessage = "Use new SingleCapsuleTraceByChannel", bIgnoreSelf="true", WorldContext="WorldContextObject", AutoCreateRefTerm="ActorsToIgnore", FriendlyName = "SingleCapsuleTraceByChannelDeprecated"))
-	static ENGINE_API bool CapsuleTraceSingle_DEPRECATED(UObject* WorldContextObject, const FVector Start, const FVector End, float Radius, float HalfHeight, ECollisionChannel TraceChannel, bool bTraceComplex, const TArray<AActor*>& ActorsToIgnore, EDrawDebugTrace::Type DrawDebugType, FHitResult& OutHit, bool bIgnoreSelf);
+	static bool CapsuleTraceSingle_DEPRECATED(UObject* WorldContextObject, const FVector Start, const FVector End, float Radius, float HalfHeight, ECollisionChannel TraceChannel, bool bTraceComplex, const TArray<AActor*>& ActorsToIgnore, EDrawDebugTrace::Type DrawDebugType, FHitResult& OutHit, bool bIgnoreSelf);
 
 	/**
 	 * Sweeps a capsule along the given line and returns all hits encountered up to and including the first blocking hit.
@@ -995,7 +995,7 @@ class UKismetSystemLibrary : public UBlueprintFunctionLibrary
 	 * @return				True if there was a hit, false otherwise.
 	 */
 	UFUNCTION(BlueprintCallable, Category="Collision", meta=(DeprecatedFunction, DeprecationMessage = "Use new MultiCapsuleTraceByChannel", bIgnoreSelf="true", WorldContext="WorldContextObject", AutoCreateRefTerm="ActorsToIgnore", FriendlyName = "MultiCapsuleTraceByChannelDeprecated"))
-	static ENGINE_API bool CapsuleTraceMulti_DEPRECATED(UObject* WorldContextObject, const FVector Start, const FVector End, float Radius, float HalfHeight, ECollisionChannel TraceChannel, bool bTraceComplex, const TArray<AActor*>& ActorsToIgnore, EDrawDebugTrace::Type DrawDebugType, TArray<FHitResult>& OutHits, bool bIgnoreSelf);
+	static bool CapsuleTraceMulti_DEPRECATED(UObject* WorldContextObject, const FVector Start, const FVector End, float Radius, float HalfHeight, ECollisionChannel TraceChannel, bool bTraceComplex, const TArray<AActor*>& ActorsToIgnore, EDrawDebugTrace::Type DrawDebugType, TArray<FHitResult>& OutHits, bool bIgnoreSelf);
 
 	/**
 	 * Does a collision trace along the given line and returns the first hit encountered.
@@ -1010,7 +1010,7 @@ class UKismetSystemLibrary : public UBlueprintFunctionLibrary
 	 * @return				True if there was a hit, false otherwise.
 	 */
 	UFUNCTION(BlueprintCallable, Category="Collision", meta=(DeprecatedFunction, DeprecationMessage = "Use new SingleLineTraceForObjects", bIgnoreSelf="true", WorldContext="WorldContextObject", AutoCreateRefTerm="ActorsToIgnore", FriendlyName = "SingleLineTraceByObjectDeprecated"))
-	static ENGINE_API bool LineTraceSingleByObject_DEPRECATED(UObject* WorldContextObject, const FVector Start, const FVector End, const TArray<TEnumAsByte<ECollisionChannel> > & ObjectsToTrace, bool bTraceComplex, const TArray<AActor*>& ActorsToIgnore, EDrawDebugTrace::Type DrawDebugType, FHitResult& OutHit, bool bIgnoreSelf);
+	static bool LineTraceSingleByObject_DEPRECATED(UObject* WorldContextObject, const FVector Start, const FVector End, const TArray<TEnumAsByte<ECollisionChannel> > & ObjectsToTrace, bool bTraceComplex, const TArray<AActor*>& ActorsToIgnore, EDrawDebugTrace::Type DrawDebugType, FHitResult& OutHit, bool bIgnoreSelf);
 	
 	/**
 	 * Does a collision trace along the given line and returns all hits encountered.
@@ -1025,7 +1025,7 @@ class UKismetSystemLibrary : public UBlueprintFunctionLibrary
 	 * @return				True if there was a hit, false otherwise.
 	 */
 	UFUNCTION(BlueprintCallable, Category="Collision", meta=(DeprecatedFunction, DeprecationMessage = "Use new MultiLineTraceForObjects", bIgnoreSelf="true", WorldContext="WorldContextObject", AutoCreateRefTerm="ActorsToIgnore", FriendlyName = "MultiLineTraceByObjectDeprecated"))
-	static ENGINE_API bool LineTraceMultiByObject_DEPRECATED(UObject* WorldContextObject, const FVector Start, const FVector End, const TArray<TEnumAsByte<ECollisionChannel> > & ObjectsToTrace, bool bTraceComplex, const TArray<AActor*>& ActorsToIgnore, EDrawDebugTrace::Type DrawDebugType, TArray<FHitResult>& OutHits, bool bIgnoreSelf);
+	static bool LineTraceMultiByObject_DEPRECATED(UObject* WorldContextObject, const FVector Start, const FVector End, const TArray<TEnumAsByte<ECollisionChannel> > & ObjectsToTrace, bool bTraceComplex, const TArray<AActor*>& ActorsToIgnore, EDrawDebugTrace::Type DrawDebugType, TArray<FHitResult>& OutHits, bool bIgnoreSelf);
 
 	/**
 	 * Sweeps a sphere along the given line and returns the first hit encountered.
@@ -1040,7 +1040,7 @@ class UKismetSystemLibrary : public UBlueprintFunctionLibrary
 	 * @return				True if there was a hit, false otherwise.
 	 */
 	UFUNCTION(BlueprintCallable, Category="Collision", meta=(DeprecatedFunction, DeprecationMessage = "Use new SingleSphereTraceForObjects", bIgnoreSelf="true", WorldContext="WorldContextObject", AutoCreateRefTerm="ActorsToIgnore", FriendlyName = "SingleSphereTraceByObjectDeprecated"))
-	static ENGINE_API bool SphereTraceSingleByObject_DEPRECATED(UObject* WorldContextObject, const FVector Start, const FVector End, float Radius, const TArray<TEnumAsByte<ECollisionChannel> > & ObjectsToTrace, bool bTraceComplex, const TArray<AActor*>& ActorsToIgnore, EDrawDebugTrace::Type DrawDebugType, FHitResult& OutHit, bool bIgnoreSelf);
+	static bool SphereTraceSingleByObject_DEPRECATED(UObject* WorldContextObject, const FVector Start, const FVector End, float Radius, const TArray<TEnumAsByte<ECollisionChannel> > & ObjectsToTrace, bool bTraceComplex, const TArray<AActor*>& ActorsToIgnore, EDrawDebugTrace::Type DrawDebugType, FHitResult& OutHit, bool bIgnoreSelf);
 
 	/**
 	 * Sweeps a sphere along the given line and returns all hits encountered.
@@ -1056,7 +1056,7 @@ class UKismetSystemLibrary : public UBlueprintFunctionLibrary
 	 * @return				True if there was a hit, false otherwise.
 	 */
 	UFUNCTION(BlueprintCallable, Category="Collision", meta=(DeprecatedFunction, DeprecationMessage = "Use new MultiSphereTraceForObjects", bIgnoreSelf="true", WorldContext="WorldContextObject", AutoCreateRefTerm="ActorsToIgnore", FriendlyName = "MultiSphereTraceByObjectDeprecated"))
-	static ENGINE_API bool SphereTraceMultiByObject_DEPRECATED(UObject* WorldContextObject, const FVector Start, const FVector End, float Radius, const TArray<TEnumAsByte<ECollisionChannel> > & ObjectsToTrace, bool bTraceComplex, const TArray<AActor*>& ActorsToIgnore, EDrawDebugTrace::Type DrawDebugType, TArray<FHitResult>& OutHits, bool bIgnoreSelf);
+	static bool SphereTraceMultiByObject_DEPRECATED(UObject* WorldContextObject, const FVector Start, const FVector End, float Radius, const TArray<TEnumAsByte<ECollisionChannel> > & ObjectsToTrace, bool bTraceComplex, const TArray<AActor*>& ActorsToIgnore, EDrawDebugTrace::Type DrawDebugType, TArray<FHitResult>& OutHits, bool bIgnoreSelf);
 	
 	/**
 	 * Sweeps a capsule along the given line and returns the first hit encountered.
@@ -1073,7 +1073,7 @@ class UKismetSystemLibrary : public UBlueprintFunctionLibrary
 	 * @return				True if there was a hit, false otherwise.
 	 */
 	UFUNCTION(BlueprintCallable, Category="Collision", meta=(DeprecatedFunction, DeprecationMessage = "Use new SingleCapsuleTraceForObjects", bIgnoreSelf="true", WorldContext="WorldContextObject", AutoCreateRefTerm="ActorsToIgnore", FriendlyName = "SingleCapsuleTraceByObjectDeprecated"))
-	static ENGINE_API bool CapsuleTraceSingleByObject_DEPRECATED(UObject* WorldContextObject, const FVector Start, const FVector End, float Radius, float HalfHeight, const TArray<TEnumAsByte<ECollisionChannel> > & ObjectsToTrace, bool bTraceComplex, const TArray<AActor*>& ActorsToIgnore, EDrawDebugTrace::Type DrawDebugType, FHitResult& OutHit, bool bIgnoreSelf);
+	static bool CapsuleTraceSingleByObject_DEPRECATED(UObject* WorldContextObject, const FVector Start, const FVector End, float Radius, float HalfHeight, const TArray<TEnumAsByte<ECollisionChannel> > & ObjectsToTrace, bool bTraceComplex, const TArray<AActor*>& ActorsToIgnore, EDrawDebugTrace::Type DrawDebugType, FHitResult& OutHit, bool bIgnoreSelf);
 
 	/**
 	 * Sweeps a capsule along the given line and returns all hits encountered.
@@ -1090,7 +1090,7 @@ class UKismetSystemLibrary : public UBlueprintFunctionLibrary
 	 * @return				True if there was a hit, false otherwise.
 	 */
 	UFUNCTION(BlueprintCallable, Category="Collision", meta=(DeprecatedFunction, DeprecationMessage = "Use new MultiCapsuleTraceForObjects", bIgnoreSelf="true", WorldContext="WorldContextObject", AutoCreateRefTerm="ActorsToIgnore", FriendlyName = "MultiCapsuleTraceByObjectDeprecated"))
-	static ENGINE_API bool CapsuleTraceMultiByObject_DEPRECATED(UObject* WorldContextObject, const FVector Start, const FVector End, float Radius, float HalfHeight, const TArray<TEnumAsByte<ECollisionChannel> > & ObjectsToTrace, bool bTraceComplex, const TArray<AActor*>& ActorsToIgnore, EDrawDebugTrace::Type DrawDebugType, TArray<FHitResult>& OutHits, bool bIgnoreSelf);
+	static bool CapsuleTraceMultiByObject_DEPRECATED(UObject* WorldContextObject, const FVector Start, const FVector End, float Radius, float HalfHeight, const TArray<TEnumAsByte<ECollisionChannel> > & ObjectsToTrace, bool bTraceComplex, const TArray<AActor*>& ActorsToIgnore, EDrawDebugTrace::Type DrawDebugType, TArray<FHitResult>& OutHits, bool bIgnoreSelf);
 
 	/**
 	 * Returns an array of unique actors represented by the given list of components.
@@ -1099,68 +1099,68 @@ class UKismetSystemLibrary : public UBlueprintFunctionLibrary
 	 * @param OutActorList		Start of line segment.
 	 */
 	UFUNCTION(BlueprintCallable, Category="Utilities")
-	static ENGINE_API void GetActorListFromComponentList(const TArray<class UPrimitiveComponent*>& ComponentList, UClass* ActorClassFilter, TArray<class AActor*>& OutActorList);
+	static void GetActorListFromComponentList(const TArray<class UPrimitiveComponent*>& ComponentList, UClass* ActorClassFilter, TArray<class AActor*>& OutActorList);
 
 
 	// --- Debug drawing functions ------------------------------
 
 	/** Draw a debug line */
 	UFUNCTION(BlueprintCallable, Category="Rendering|Debug", meta=(WorldContext="WorldContextObject"))
-	static ENGINE_API void DrawDebugLine(UObject* WorldContextObject, const FVector LineStart, const FVector LineEnd, FLinearColor LineColor, float Duration=0.f, float Thickness = 0.f);
+	static void DrawDebugLine(UObject* WorldContextObject, const FVector LineStart, const FVector LineEnd, FLinearColor LineColor, float Duration=0.f, float Thickness = 0.f);
 
 	/** Draw a debug point */
 	UFUNCTION(BlueprintCallable, Category="Rendering|Debug", meta=(WorldContext="WorldContextObject"))
-	static ENGINE_API void DrawDebugPoint(UObject* WorldContextObject, const FVector Position, float Size, FLinearColor PointColor, float Duration=0.f);
+	static void DrawDebugPoint(UObject* WorldContextObject, const FVector Position, float Size, FLinearColor PointColor, float Duration=0.f);
 
 	/** Draw directional arrow, pointing from LineStart to LineEnd. */
 	UFUNCTION(BlueprintCallable, Category="Rendering|Debug", meta=(WorldContext="WorldContextObject"))
-	static ENGINE_API void DrawDebugArrow(UObject* WorldContextObject, const FVector LineStart, const FVector LineEnd, float ArrowSize, FLinearColor LineColor, float Duration=0.f);
+	static void DrawDebugArrow(UObject* WorldContextObject, const FVector LineStart, const FVector LineEnd, float ArrowSize, FLinearColor LineColor, float Duration=0.f);
 
 	/** Draw a debug box */
 	UFUNCTION(BlueprintCallable, Category="Rendering|Debug", meta=(WorldContext="WorldContextObject"))
-	static ENGINE_API void DrawDebugBox(UObject* WorldContextObject, const FVector Center, FVector Extent, FLinearColor LineColor, const FRotator Rotation=FRotator::ZeroRotator, float Duration=0.f);
+	static void DrawDebugBox(UObject* WorldContextObject, const FVector Center, FVector Extent, FLinearColor LineColor, const FRotator Rotation=FRotator::ZeroRotator, float Duration=0.f);
 
 	/** Draw a debug coordinate system. */
 	UFUNCTION(BlueprintCallable, Category="Rendering|Debug", meta=(WorldContext="WorldContextObject"))
-	static ENGINE_API void DrawDebugCoordinateSystem(UObject* WorldContextObject, const FVector AxisLoc, const FRotator AxisRot, float Scale=1.f, float Duration=0.f);
+	static void DrawDebugCoordinateSystem(UObject* WorldContextObject, const FVector AxisLoc, const FRotator AxisRot, float Scale=1.f, float Duration=0.f);
 
 	/** Draw a debug sphere */
 	UFUNCTION(BlueprintCallable, Category="Rendering|Debug", meta=(WorldContext="WorldContextObject"))
-	static ENGINE_API void DrawDebugSphere(UObject* WorldContextObject, const FVector Center, float Radius=100.f, int32 Segments=12, FLinearColor LineColor = FLinearColor::White, float Duration=0.f);
+	static void DrawDebugSphere(UObject* WorldContextObject, const FVector Center, float Radius=100.f, int32 Segments=12, FLinearColor LineColor = FLinearColor::White, float Duration=0.f);
 
 	/** Draw a debug cylinder */
 	UFUNCTION(BlueprintCallable, Category="Rendering|Debug", meta=(WorldContext="WorldContextObject"))
-	static ENGINE_API void DrawDebugCylinder(UObject* WorldContextObject, const FVector Start, const FVector End, float Radius=100.f, int32 Segments=12, FLinearColor LineColor = FLinearColor::White, float Duration=0.f);
+	static void DrawDebugCylinder(UObject* WorldContextObject, const FVector Start, const FVector End, float Radius=100.f, int32 Segments=12, FLinearColor LineColor = FLinearColor::White, float Duration=0.f);
 	
 	/** Draw a debug cone */
 	UFUNCTION(BlueprintCallable, Category="Rendering|Debug", meta=(WorldContext="WorldContextObject", DeprecatedFunction, DeprecationMessage="DrawDebugCone has been changed to use degrees for angles instead of radians. Place a new DrawDebugCone node and pass your angles as degrees."))
-	static ENGINE_API void DrawDebugCone(UObject* WorldContextObject, const FVector Origin, const FVector Direction, float Length, float AngleWidth, float AngleHeight, int32 NumSides, FLinearColor LineColor);
+	static void DrawDebugCone(UObject* WorldContextObject, const FVector Origin, const FVector Direction, float Length, float AngleWidth, float AngleHeight, int32 NumSides, FLinearColor LineColor);
 
 	/** 
 	 * Draw a debug cone 
 	 * Angles are specified in degrees
 	 */
 	UFUNCTION(BlueprintCallable, Category="Rendering|Debug", meta=(WorldContext="WorldContextObject", FriendlyName="DrawDebugCone"))
-	static ENGINE_API void DrawDebugConeInDegrees(UObject* WorldContextObject, const FVector Origin, const FVector Direction, float Length=100.f, float AngleWidth=45.f, float AngleHeight=45.f, int32 NumSides = 12, FLinearColor LineColor = FLinearColor::White, float Duration=0.f);
+	static void DrawDebugConeInDegrees(UObject* WorldContextObject, const FVector Origin, const FVector Direction, float Length=100.f, float AngleWidth=45.f, float AngleHeight=45.f, int32 NumSides = 12, FLinearColor LineColor = FLinearColor::White, float Duration=0.f);
 
 	/** Draw a debug capsule */
 	UFUNCTION(BlueprintCallable, Category="Rendering|Debug", meta=(WorldContext="WorldContextObject"))
-	static ENGINE_API void DrawDebugCapsule(UObject* WorldContextObject, const FVector Center, float HalfHeight, float Radius, const FRotator Rotation, FLinearColor LineColor = FLinearColor::White, float Duration=0.f);
+	static void DrawDebugCapsule(UObject* WorldContextObject, const FVector Center, float HalfHeight, float Radius, const FRotator Rotation, FLinearColor LineColor = FLinearColor::White, float Duration=0.f);
 
 	/** Draw a debug string at a 3d world location. */
 	UFUNCTION(BlueprintCallable, Category="Rendering|Debug", meta=(WorldContext="WorldContextObject"))
-	static ENGINE_API void DrawDebugString(UObject* WorldContextObject, const FVector TextLocation, const FString& Text, class AActor* TestBaseActor = NULL, FLinearColor TextColor = FLinearColor::White, float Duration=0.f);
+	static void DrawDebugString(UObject* WorldContextObject, const FVector TextLocation, const FString& Text, class AActor* TestBaseActor = NULL, FLinearColor TextColor = FLinearColor::White, float Duration=0.f);
 	/** 
 	 * Removes all debug strings. 
 	 *
 	 * @param WorldContext	World context
 	 */
 	UFUNCTION(BlueprintCallable, Category="Rendering|Debug", meta=(WorldContext="WorldContextObject"))
-	static ENGINE_API void FlushDebugStrings(UObject* WorldContextObject);
+	static void FlushDebugStrings(UObject* WorldContextObject);
 
 	/** Draws a debug plane. */
 	UFUNCTION(BlueprintCallable, Category="Rendering|Debug", meta=(WorldContext="WorldContextObject"))
-	static ENGINE_API void DrawDebugPlane(UObject* WorldContextObject, const FPlane& PlaneCoordinates, const FVector Location, float Size, FLinearColor PlaneColor = FLinearColor::White, float Duration=0.f);
+	static void DrawDebugPlane(UObject* WorldContextObject, const FPlane& PlaneCoordinates, const FVector Location, float Size, FLinearColor PlaneColor = FLinearColor::White, float Duration=0.f);
 
 	/** 
 	 * Flush all persistent debug lines and shapes.
@@ -1168,37 +1168,37 @@ class UKismetSystemLibrary : public UBlueprintFunctionLibrary
 	 * @param WorldContext	World context
 	 */
 	UFUNCTION(BlueprintCallable, Category="Rendering|Debug", meta=(WorldContext="WorldContextObject"))
-	static ENGINE_API void FlushPersistentDebugLines(UObject* WorldContextObject);
+	static void FlushPersistentDebugLines(UObject* WorldContextObject);
 
 	/** Draws a debug frustum. */
 	UFUNCTION(BlueprintCallable, Category="Rendering|Debug", meta=(WorldContext="WorldContextObject"))
-	static ENGINE_API void DrawDebugFrustum(UObject* WorldContextObject, const FTransform& FrustumTransform, FLinearColor FrustumColor = FLinearColor::White, float Duration=0.f);
+	static void DrawDebugFrustum(UObject* WorldContextObject, const FTransform& FrustumTransform, FLinearColor FrustumColor = FLinearColor::White, float Duration=0.f);
 
 	/** Draw a debug camera shape. */
 	UFUNCTION(BlueprintCallable, Category="Rendering|Debug")
-	static ENGINE_API void DrawDebugCamera(const ACameraActor* CameraActor, FLinearColor CameraColor = FLinearColor::White, float Duration=0.f);
+	static void DrawDebugCamera(const ACameraActor* CameraActor, FLinearColor CameraColor = FLinearColor::White, float Duration=0.f);
 
 	/* Draws a 2D Histogram of size 'DrawSize' based FDebugFloatHistory struct, using DrawTransform for the position in the world. */
 	UFUNCTION(BlueprintCallable, Category = "Rendering|Debug")
-	static ENGINE_API void DrawDebugFloatHistoryTransform(UObject* WorldContextObject, const FDebugFloatHistory& FloatHistory, const FTransform& DrawTransform, FVector2D DrawSize, FLinearColor DrawColor = FLinearColor::White, float Duration = 0.f);
+	static void DrawDebugFloatHistoryTransform(UObject* WorldContextObject, const FDebugFloatHistory& FloatHistory, const FTransform& DrawTransform, FVector2D DrawSize, FLinearColor DrawColor = FLinearColor::White, float Duration = 0.f);
 
 	/* Draws a 2D Histogram of size 'DrawSize' based FDebugFloatHistory struct, using DrawLocation for the location in the world, rotation will face camera of first player. */
 	UFUNCTION(BlueprintCallable, Category = "Rendering|Debug")
-	static ENGINE_API void DrawDebugFloatHistoryLocation(UObject* WorldContextObject, const FDebugFloatHistory& FloatHistory, FVector DrawLocation, FVector2D DrawSize, FLinearColor DrawColor = FLinearColor::White, float Duration = 0.f);
+	static void DrawDebugFloatHistoryLocation(UObject* WorldContextObject, const FDebugFloatHistory& FloatHistory, FVector DrawLocation, FVector2D DrawSize, FLinearColor DrawColor = FLinearColor::White, float Duration = 0.f);
 
 	UFUNCTION(BlueprintPure, Category = "Rendering|Debug")
-	static ENGINE_API FDebugFloatHistory AddFloatHistorySample(float Value, const FDebugFloatHistory& FloatHistory);
+	static FDebugFloatHistory AddFloatHistorySample(float Value, const FDebugFloatHistory& FloatHistory);
 	
 	/** Mark as modified. */
 	UFUNCTION(BlueprintCallable, Category="Development|Editor")
-	static ENGINE_API void CreateCopyForUndoBuffer(UObject* ObjectToModify);
+	static void CreateCopyForUndoBuffer(UObject* ObjectToModify);
 
 	/** Get bounds */
 	UFUNCTION(BlueprintPure, Category="Collision")
-	static ENGINE_API void GetComponentBounds(const USceneComponent* Component, FVector& Origin, FVector& BoxExtent, float& SphereRadius);
+	static void GetComponentBounds(const USceneComponent* Component, FVector& Origin, FVector& BoxExtent, float& SphereRadius);
 
 	UFUNCTION(BlueprintPure, Category="Collision", meta=(DeprecatedFunction))
-	static ENGINE_API void GetActorBounds(const AActor* Actor, FVector& Origin, FVector& BoxExtent);
+	static void GetActorBounds(const AActor* Actor, FVector& Origin, FVector& BoxExtent);
 
 	/**
 	 * Get the clamped state of r.DetailMode, see console variable help (allows for scalability, cannot be used in construction scripts)
diff --git a/Engine/Source/Runtime/Engine/Classes/Kismet/KismetTextLibrary.h b/Engine/Source/Runtime/Engine/Classes/Kismet/KismetTextLibrary.h
--- a/Engine/Source/Runtime/Engine/Classes/Kismet/KismetTextLibrary.h
+++ b/Engine/Source/Runtime/Engine/Classes/Kismet/KismetTextLibrary.h
@@ -40,8 +40,8 @@ struct FFormatTextArgument
 	FText TextValue;
 };
 
-UCLASS(MinimalAPI)
-class UKismetTextLibrary : public UBlueprintFunctionLibrary
+UCLASS()
+class ENGINE_API UKismetTextLibrary : public UBlueprintFunctionLibrary
 {
 	GENERATED_UCLASS_BODY()
 
