Documenting the return values of Array Add and AddUnique - the behavior of AddUnique will be particularly surprising to people (and may warrant changing - it would be more useful to always return the index of the item in the array)

diff --git a/Engine/Source/Runtime/Engine/Classes/Kismet/KismetArrayLibrary.h b/Engine/Source/Runtime/Engine/Classes/Kismet/KismetArrayLibrary.h
--- a/Engine/Source/Runtime/Engine/Classes/Kismet/KismetArrayLibrary.h
+++ b/Engine/Source/Runtime/Engine/Classes/Kismet/KismetArrayLibrary.h
@@ -12,6 +12,7 @@ class ENGINE_API UKismetArrayLibrary : public UBlueprintFunctionLibrary
 	 *
 	 *@param	TargetArray		The array to add item to
 	 *@param	NewItem			The item to add to the array
+	 *@return	The index of the newly added item
 	*/
 	UFUNCTION(BlueprintCallable, CustomThunk, meta=(FriendlyName = "Add", CompactNodeTitle = "ADD", ArrayParm = "TargetArray|ArrayProperty", ArrayTypeDependentParams = "NewItem", AutoCreateRefTerm = "NewItem"), Category="Utilities|Array")
 	static int32 Array_Add(const TArray<int32>& TargetArray, const UArrayProperty* ArrayProperty, const int32& NewItem);
@@ -19,8 +20,9 @@ class ENGINE_API UKismetArrayLibrary : public UBlueprintFunctionLibrary
 	/**
 	*Add item to array (unique)
 	*
-	*@param	TargetArray		The array to add item to
-	*@param	NewItem			The item to add to the array
+	*@param		TargetArray		The array to add item to
+	*@param		NewItem			The item to add to the array
+	*@return	The index of the newly added item, or INDEX_NONE if the item is already present in the array
 	*/
 	UFUNCTION(BlueprintCallable, CustomThunk, meta = (FriendlyName = "Add Unique", CompactNodeTitle = "ADDUNIQUE", ArrayParm = "TargetArray|ArrayProperty", ArrayTypeDependentParams = "NewItem", AutoCreateRefTerm = "NewItem"), Category = "Utilities|Array")
 	static int32 Array_AddUnique(const TArray<int32>& TargetArray, const UArrayProperty* ArrayProperty, const int32& NewItem);
