Fixed issue with Block not waiting till the shader map is complete before serialization, caused by removing outstanding compile id too early.

diff --git a/Engine/Source/Runtime/Engine/Private/ShaderCompiler/ShaderCompiler.cpp b/Engine/Source/Runtime/Engine/Private/ShaderCompiler/ShaderCompiler.cpp
--- a/Engine/Source/Runtime/Engine/Private/ShaderCompiler/ShaderCompiler.cpp
+++ b/Engine/Source/Runtime/Engine/Private/ShaderCompiler/ShaderCompiler.cpp
@@ -1485,10 +1485,6 @@ void FShaderCompilingManager::ProcessCompiledShaderMaps(
 				bShaderMapComplete = ShaderMap->ProcessCompilationResults(ResultArray, CompileResults.FinalizeJobIndex, TimeBudget);
 			}
 
-			for ( auto& Material : MaterialsArray )
-			{
-				Material->RemoveOutstandingCompileId( ShaderMap->CompilingId );
-			}
 
 			if (bShaderMapComplete)
 			{
@@ -1501,35 +1497,17 @@ void FShaderCompilingManager::ProcessCompiledShaderMaps(
 				for (int32 MaterialIndex = 0; MaterialIndex < MaterialsArray.Num(); MaterialIndex++)
 				{
 					FMaterial* Material = MaterialsArray[MaterialIndex];
-					FMaterialShaderMap* CompletedShaderMap = bSuccess ? ShaderMap : NULL;
-
-					if (CompletedShaderMap 
-						// Don't modify materials for which the compiled shader map is no longer complete
-						// This can happen if a material being compiled is edited, or if CheckMaterialUsage changes a flag and causes a recompile
-						&& CompletedShaderMap->IsComplete(Material, true))
-					{
-						MaterialsToUpdate.Add(Material, CompletedShaderMap);
+					FMaterialShaderMap* CompletedShaderMap = ShaderMap;
 
-						// Note: if !CompileResults.bApplyCompletedShaderMapForRendering, RenderingThreadShaderMap must be set elsewhere to match up with the new value of GameThreadShaderMap
-						if (CompileResults.bApplyCompletedShaderMapForRendering)
-						{
-							MaterialsToApplyToScene.Add(Material, CompletedShaderMap);
-						}
-					}
-				}
+					Material->RemoveOutstandingCompileId(ShaderMap->CompilingId);
 
-				if (!bSuccess)
-				{				
-					for (int32 MaterialIndex = 0; MaterialIndex < MaterialsArray.Num(); MaterialIndex++)
+					if ( !bSuccess )
 					{
-						FMaterial& CurrentMaterial = *MaterialsArray[MaterialIndex];
-
 						// Propagate error messages
-						CurrentMaterial.CompileErrors = Errors;
-
-						MaterialsToUpdate.Add( &CurrentMaterial, NULL );
+						Material->CompileErrors = Errors;
+						MaterialsToUpdate.Add( Material, NULL );
 
-						if (CurrentMaterial.IsDefaultMaterial())
+						if (Material->IsDefaultMaterial())
 						{
 							// Log the errors unsuppressed before the fatal error, so it's always obvious from the log what the compile error was
 							for (int32 ErrorIndex = 0; ErrorIndex < Errors.Num(); ErrorIndex++)
@@ -1537,11 +1515,11 @@ void FShaderCompilingManager::ProcessCompiledShaderMaps(
 								UE_LOG(LogShaderCompilers, Warning, TEXT("	%s"), *Errors[ErrorIndex]);
 							}
 							// Assert if a default material could not be compiled, since there will be nothing for other failed materials to fall back on.
-							UE_LOG(LogShaderCompilers, Fatal,TEXT("Failed to compile default material %s!"), *CurrentMaterial.GetBaseMaterialPathName());
+							UE_LOG(LogShaderCompilers, Fatal,TEXT("Failed to compile default material %s!"), *Material->GetBaseMaterialPathName());
 						}
 
 						UE_LOG(LogShaderCompilers, Warning, TEXT("Failed to compile Material %s for platform %s, Default Material will be used in game."), 
-							*CurrentMaterial.GetBaseMaterialPathName(), 
+							*Material->GetBaseMaterialPathName(), 
 							*LegacyShaderPlatformToShaderFormat(ShaderMap->GetShaderPlatform()).ToString());
 
 						for (int32 ErrorIndex = 0; ErrorIndex < Errors.Num(); ErrorIndex++)
@@ -1549,6 +1527,19 @@ void FShaderCompilingManager::ProcessCompiledShaderMaps(
 							UE_LOG(LogShaders, Warning, TEXT("	%s"), *Errors[ErrorIndex]);
 						}
 					}
+					// if we succeeded and our shader map is not complete this could be because the material was being edited quicker then the compile could be completed
+					// Don't modify materials for which the compiled shader map is no longer complete
+					// This can happen if a material being compiled is edited, or if CheckMaterialUsage changes a flag and causes a recompile
+					else if ( CompletedShaderMap->IsComplete(Material, true) )
+					{
+						MaterialsToUpdate.Add(Material, CompletedShaderMap);
+						// Note: if !CompileResults.bApplyCompletedShaderMapForRendering, RenderingThreadShaderMap must be set elsewhere to match up with the new value of GameThreadShaderMap
+						if (CompileResults.bApplyCompletedShaderMapForRendering)
+						{
+							MaterialsToApplyToScene.Add(Material, CompletedShaderMap);
+						}
+					}
+
 				}
 
 				// Cleanup shader jobs and compile tracking structures
