RayStartOffsetDepthScale comment

diff --git a/Engine/Source/Runtime/Engine/Classes/Components/LightComponent.h b/Engine/Source/Runtime/Engine/Classes/Components/LightComponent.h
--- a/Engine/Source/Runtime/Engine/Classes/Components/LightComponent.h
+++ b/Engine/Source/Runtime/Engine/Classes/Components/LightComponent.h
@@ -162,6 +162,10 @@ class ENGINE_API ULightComponent : public ULightComponentBase
 	UPROPERTY(EditAnywhere, BlueprintReadOnly, Category=DistanceFieldShadows, meta=(DisplayName = "RayTraced DistanceField Shadows"))
 	bool bUseRayTracedDistanceFieldShadows;
 
+	/** 
+	 * Controls how large of an offset ray traced shadows have from the receiving surface as the camera gets further away.  
+	 * This can be useful to hide self-shadowing artifacts from low resolution distance fields on huge static meshes.
+	 */
 	UPROPERTY(EditAnywhere, BlueprintReadOnly, Category=DistanceFieldShadows, meta=(UIMin = "0", UIMax = ".1"), AdvancedDisplay)
 	float RayStartOffsetDepthScale;
 
