Fixed bug with initialization

diff --git a/Engine/Source/Runtime/Engine/Private/AudioDevice.cpp b/Engine/Source/Runtime/Engine/Private/AudioDevice.cpp
--- a/Engine/Source/Runtime/Engine/Private/AudioDevice.cpp
+++ b/Engine/Source/Runtime/Engine/Private/AudioDevice.cpp
@@ -34,7 +34,8 @@ FAudioDevice::FAudioDevice()
 	, CurrentAudioVolume(nullptr)
 	, HighestPriorityReverb(nullptr)
 	, SpatializationPlugin(nullptr)
-	, bSpatializationExtensionEnabled(true)
+	, SpatializeProcessor(nullptr)
+	, bSpatializationExtensionEnabled(false)
 	, bHRTFEnabledForAll(false)
 	, bIsDeviceMuted(false)
 {
