Fixing build by adding forward declare on FAudioDevice

diff --git a/Engine/Source/Runtime/Engine/Public/IAudioExtensionPlugin.h b/Engine/Source/Runtime/Engine/Public/IAudioExtensionPlugin.h
--- a/Engine/Source/Runtime/Engine/Public/IAudioExtensionPlugin.h
+++ b/Engine/Source/Runtime/Engine/Public/IAudioExtensionPlugin.h
@@ -132,7 +132,7 @@ public:
 	*
 	* @return A new spatialization algorithm instance
 	*/
-	virtual IAudioSpatializationAlgorithm* GetNewSpatializationAlgorithm(FAudioDevice* AudioDevice)
+	virtual IAudioSpatializationAlgorithm* GetNewSpatializationAlgorithm(class FAudioDevice* AudioDevice)
 	{
 		return nullptr;
 	}
