Fix for linux build

including CoreUObject.h in SoundAttenuation.h to ensure that TStructOpsTypeTraits is defined

diff --git a/Engine/Source/Runtime/Engine/Classes/Sound/SoundAttenuation.h b/Engine/Source/Runtime/Engine/Classes/Sound/SoundAttenuation.h
--- a/Engine/Source/Runtime/Engine/Classes/Sound/SoundAttenuation.h
+++ b/Engine/Source/Runtime/Engine/Classes/Sound/SoundAttenuation.h
@@ -1,7 +1,9 @@
 // Copyright 1998-2015 Epic Games, Inc. All Rights Reserved.
 
 #pragma once
- 
+
+#include "CoreUObject.h"
+
 #include "SoundAttenuation.generated.h"
 
 UENUM()
