Fixing initialization order

diff --git a/Engine/Source/Runtime/Engine/Private/Audio.cpp b/Engine/Source/Runtime/Engine/Private/Audio.cpp
--- a/Engine/Source/Runtime/Engine/Private/Audio.cpp
+++ b/Engine/Source/Runtime/Engine/Private/Audio.cpp
@@ -357,13 +357,13 @@ FWaveInstance::FWaveInstance( FActiveSound* InActiveSound )
 ,	LFEBleed( 0.0f )
 ,	LoopingMode( LOOP_Never )
 ,	StartTime( -1.f )
-,   SpatializationAlgorithm(SPATIALIZATION_Default)
 ,	bApplyRadioFilter( false )
-,	bIsStarted( false )
+,	bIsStarted(false)
 ,	bIsFinished( false )
 ,	bAlreadyNotifiedHook( false )
 ,	bUseSpatialization( false )
-,	bEQFilterApplied( false )
+,	SpatializationAlgorithm(SPATIALIZATION_Default)
+,	bEQFilterApplied(false)
 ,	bIsUISound( false )
 ,	bIsMusic( false )
 ,	bReverb( true )
