Fixed object scaling when the pivot has been moved from the actor origin.

diff --git a/Engine/Source/Runtime/Engine/Private/ActorEditor.cpp b/Engine/Source/Runtime/Engine/Private/ActorEditor.cpp
--- a/Engine/Source/Runtime/Engine/Private/ActorEditor.cpp
+++ b/Engine/Source/Runtime/Engine/Private/ActorEditor.cpp
@@ -425,22 +425,30 @@ void AActor::EditorApplyScale( const FVector& DeltaScale, const FVector* PivotLo
 		const FVector CurrentScale = GetRootComponent()->RelativeScale3D;
 
 		// @todo: Remove this hack once we have decided on the scaling method to use.
+		FVector ScaleToApply;
+
 		if( AActor::bUsePercentageBasedScaling )
 		{
-			GetRootComponent()->SetRelativeScale3D(CurrentScale + DeltaScale * CurrentScale);
-
-			if (PivotLocation)
-			{
-				FVector Loc = GetActorLocation();
-				Loc -= *PivotLocation;
-				Loc += DeltaScale * Loc;
-				Loc += *PivotLocation;
-				GetRootComponent()->SetWorldLocation(Loc);
-			}
+			ScaleToApply = CurrentScale * (FVector(1.0f) + DeltaScale);
 		}
 		else
 		{
-			GetRootComponent()->SetRelativeScale3D(CurrentScale + DeltaScale * CurrentScale.GetSignVector());
+			ScaleToApply = CurrentScale + DeltaScale;
+		}
+
+		GetRootComponent()->SetRelativeScale3D(ScaleToApply);
+
+		if (PivotLocation)
+		{
+			const FVector CurrentScaleSafe(CurrentScale.X ? CurrentScale.X : 1.0f,
+										   CurrentScale.Y ? CurrentScale.Y : 1.0f,
+										   CurrentScale.Z ? CurrentScale.Z : 1.0f);
+
+			FVector Loc = GetActorLocation();
+			Loc -= *PivotLocation;
+			Loc *= (ScaleToApply / CurrentScaleSafe);
+			Loc += *PivotLocation;
+			GetRootComponent()->SetWorldLocation(Loc);
 		}
 	}
 	else
