Don't rename child actors when destroying their component during gameplay
Don't reset reloaders when renaming child actors during gameplay (for example by changing child actor class)
Fix crash when overriding child actor class on a per-instance basis

diff --git a/Engine/Source/Runtime/Engine/Classes/Components/ChildActorComponent.h b/Engine/Source/Runtime/Engine/Classes/Components/ChildActorComponent.h
--- a/Engine/Source/Runtime/Engine/Classes/Components/ChildActorComponent.h
+++ b/Engine/Source/Runtime/Engine/Classes/Components/ChildActorComponent.h
@@ -58,7 +58,7 @@ public:
 	ENGINE_API void CreateChildActor();
 
 	/** Kill any currently present child actor */
-	void DestroyChildActor();
+	void DestroyChildActor(const bool bRequiresRename = true);
 };
 
 
diff --git a/Engine/Source/Runtime/Engine/Private/Components/ChildActorComponent.cpp b/Engine/Source/Runtime/Engine/Private/Components/ChildActorComponent.cpp
--- a/Engine/Source/Runtime/Engine/Private/Components/ChildActorComponent.cpp
+++ b/Engine/Source/Runtime/Engine/Private/Components/ChildActorComponent.cpp
@@ -29,15 +29,21 @@ void UChildActorComponent::OnRegister()
 #if WITH_EDITOR
 void UChildActorComponent::PostEditChangeProperty(FPropertyChangedEvent& PropertyChangedEvent)
 {
-	Super::PostEditChangeProperty(PropertyChangedEvent);
-
 	static const FName NAME_ChildActorClass = GET_MEMBER_NAME_CHECKED(UChildActorComponent, ChildActorClass);
 
 	if (PropertyChangedEvent.Property && PropertyChangedEvent.Property->GetFName() == NAME_ChildActorClass)
 	{
-		DestroyChildActor();
-		CreateChildActor();
+		ChildActorName = NAME_None;
+
+		// If this was created by construction script, the post edit change super call will destroy it anyways
+		if (!IsCreatedByConstructionScript())
+		{
+			DestroyChildActor();
+			CreateChildActor();
+		}
 	}
+
+	Super::PostEditChangeProperty(PropertyChangedEvent);
 }
 
 void UChildActorComponent::PostEditUndo()
@@ -72,7 +78,7 @@ void UChildActorComponent::OnComponentDestroyed()
 {
 	Super::OnComponentDestroyed();
 
-	DestroyChildActor();
+	DestroyChildActor(!GetWorld()->IsGameWorld());
 }
 
 class FChildActorComponentInstanceData : public FSceneComponentInstanceData
@@ -263,7 +269,7 @@ void UChildActorComponent::CreateChildActor()
 	}
 }
 
-void UChildActorComponent::DestroyChildActor()
+void UChildActorComponent::DestroyChildActor(const bool bRequiresRename)
 {
 	// If we own an Actor, kill it now
 	if(ChildActor != nullptr && !GExitPurge)
@@ -291,8 +297,12 @@ void UChildActorComponent::DestroyChildActor()
 				// so we increment ClassUnique beyond our index to be certain of it.  This is ... a bit hacky.
 				ChildClass->ClassUnique = FMath::Max(ChildClass->ClassUnique, ChildActor->GetFName().GetNumber());
 
-				const FString ObjectBaseName = FString::Printf(TEXT("DESTROYED_%s_CHILDACTOR"), *ChildClass->GetName());
-				ChildActor->Rename(*MakeUniqueObjectName(ChildActor->GetOuter(), ChildClass, *ObjectBaseName).ToString(), nullptr, REN_DoNotDirty);
+				if (bRequiresRename)
+				{
+					const FString ObjectBaseName = FString::Printf(TEXT("DESTROYED_%s_CHILDACTOR"), *ChildClass->GetName());
+					const ERenameFlags RenameFlags = (GetWorld()->IsGameWorld() ? REN_DoNotDirty | REN_ForceNoResetLoaders : REN_DoNotDirty);
+					ChildActor->Rename(*MakeUniqueObjectName(ChildActor->GetOuter(), ChildClass, *ObjectBaseName).ToString(), nullptr, REN_DoNotDirty);
+				}
 				World->DestroyActor(ChildActor);
 			}
 		}
