Marked SetMasterPoseComponent to be unsafe during construction

diff --git a/Engine/Source/Runtime/Engine/Classes/Components/SkinnedMeshComponent.h b/Engine/Source/Runtime/Engine/Classes/Components/SkinnedMeshComponent.h
--- a/Engine/Source/Runtime/Engine/Classes/Components/SkinnedMeshComponent.h
+++ b/Engine/Source/Runtime/Engine/Classes/Components/SkinnedMeshComponent.h
@@ -655,7 +655,7 @@ public:
 	 *
 	 * @param NewMasterBoneComponent New MasterPoseComponent
 	 */
-	UFUNCTION(BlueprintCallable, Category="Components|SkinnedMesh")
+	UFUNCTION(BlueprintCallable, Category="Components|SkinnedMesh", meta=(UnsafeDuringActorConstruction="true"))
 	void SetMasterPoseComponent(USkinnedMeshComponent* NewMasterBoneComponent);
 
 	/** 
