A new advanced multiplayer option to enable audio devices for every player

- Rather than creating a separate audio device for every player in a multi-PIE session by default, this option will allow users to opt-in to the feature

- Multiple audio devices are still created for up to 2 players since this is a common use-case and shouldn't be an issue with performance. We can lower this to 1 trivially if needed/desired.

- The only shared audio device is the main audio device (used by the editor).

- When shared, the listener for the main audio device is only set for the in-focus viewport.

[REDACTED_HANDLE] marc.audy

diff --git a/Engine/Source/Editor/UnrealEd/Classes/Settings/LevelEditorPlaySettings.h b/Engine/Source/Editor/UnrealEd/Classes/Settings/LevelEditorPlaySettings.h
--- a/Engine/Source/Editor/UnrealEd/Classes/Settings/LevelEditorPlaySettings.h
+++ b/Engine/Source/Editor/UnrealEd/Classes/Settings/LevelEditorPlaySettings.h
@@ -222,6 +222,16 @@ private:
 	UPROPERTY(config, EditAnywhere, Category=MultiplayerOptions)
 	bool RouteGamepadToSecondWindow;
 
+	/** 
+	* If checked, a separate audio device is created for every player. 
+	
+	* If unchecked, a separate audio device is created for only the first two players and uses the main audio device for more than 2 players.
+	*
+	* Enabling this will allow rendering accurate audio from every player's perspective but will use more CPU. Keep this disabled on lower-perf machines.
+	*/
+	UPROPERTY(config, EditAnywhere, Category = MultiplayerOptions, meta=(EditCondition = "EnableSound"))
+	bool CreateAudioDeviceForEveryPlayer;
+
 	/** Height to use when spawning additional windows. */
 	UPROPERTY(config, EditAnywhere, Category=MultiplayerOptions)
 	int32 ClientWindowHeight;
@@ -269,7 +279,8 @@ public:
 	bool IsClientWindowSizeActive() const { return ((PlayNetMode == PIE_Standalone && RunUnderOneProcess) ? false : (PlayNumberOfClients >= 2)); }
 	bool GetClientWindowSize( FIntPoint &OutClientWindowSize ) const { OutClientWindowSize = FIntPoint(ClientWindowWidth, ClientWindowHeight); return IsClientWindowSizeActive(); }
 	EVisibility GetClientWindowSizeVisibility() const { return (RunUnderOneProcess ? EVisibility::Hidden : EVisibility::Visible); }
-	
+	bool IsCreateAudioDeviceForEveryPlayer() const { return CreateAudioDeviceForEveryPlayer; }
+
 public:
 
 	/** The last used height for multiple instance windows (in pixels). */
diff --git a/Engine/Source/Editor/UnrealEd/Private/PlayLevel.cpp b/Engine/Source/Editor/UnrealEd/Private/PlayLevel.cpp
--- a/Engine/Source/Editor/UnrealEd/Private/PlayLevel.cpp
+++ b/Engine/Source/Editor/UnrealEd/Private/PlayLevel.cpp
@@ -2651,8 +2651,12 @@ UGameInstance* UEditorEngine::CreatePIEGameInstance(int32 PIEInstance, bool bInS
 	
 	if (!PieWorldContext->RunAsDedicated)
 	{
+		const ULevelEditorPlaySettings* PlayInSettings = GetDefault<ULevelEditorPlaySettings>();
+		bool bCreateNewAudioDevice = PlayInSettings->IsCreateAudioDeviceForEveryPlayer();
+
 		ViewportClient = NewObject<UGameViewportClient>(this, GameViewportClientClass);
-		ViewportClient->Init(*PieWorldContext, GameInstance);
+		ViewportClient->Init(*PieWorldContext, GameInstance, bCreateNewAudioDevice);
+
 		GameViewport = ViewportClient;
 		GameViewport->bIsPlayInEditorViewport = true;
 		PieWorldContext->GameViewport = ViewportClient;
diff --git a/Engine/Source/Runtime/Engine/Classes/Engine/GameViewportClient.h b/Engine/Source/Runtime/Engine/Classes/Engine/GameViewportClient.h
--- a/Engine/Source/Runtime/Engine/Classes/Engine/GameViewportClient.h
+++ b/Engine/Source/Runtime/Engine/Classes/Engine/GameViewportClient.h
@@ -127,7 +127,7 @@ public:
 	/* Returns the relevant game instance for this viewport */
 	UGameInstance* GetGameInstance() const;
 
-	virtual void Init(struct FWorldContext& WorldContext, UGameInstance* OwningGameInstance);
+	virtual void Init(struct FWorldContext& WorldContext, UGameInstance* OwningGameInstance, bool bCreateNewAudioDevice = true);
 
 public:
 	// Begin UObject Interface
@@ -795,6 +795,9 @@ private:
 
 	/** Handle to the audio device created for this viewport. Each viewport (for multiple PIE) will have its own audio device. */
 	uint32 AudioDeviceHandle;
+
+	/** Whether or not this audio device is in audio-focus */
+	bool bHasAudioFocus;
 };
 
 
diff --git a/Engine/Source/Runtime/Engine/Private/AudioDevice.cpp b/Engine/Source/Runtime/Engine/Private/AudioDevice.cpp
--- a/Engine/Source/Runtime/Engine/Private/AudioDevice.cpp
+++ b/Engine/Source/Runtime/Engine/Private/AudioDevice.cpp
@@ -38,6 +38,7 @@ FAudioDevice::FAudioDevice()
 	, bSpatializationExtensionEnabled(false)
 	, bHRTFEnabledForAll(false)
 	, bIsDeviceMuted(false)
+	, bIsInitialized(false)
 {
 }
 
@@ -49,6 +50,11 @@ FAudioEffectsManager* FAudioDevice::CreateEffectsManager()
 
 bool FAudioDevice::Init()
 {
+	if (bIsInitialized)
+	{
+		return true;
+	}
+
 	bool bDeferStartupPrecache = false;
 	// initialize config variables
 	verify(GConfig->GetInt(TEXT("Audio"), TEXT("MaxChannels"), MaxChannels, GEngineIni));
@@ -108,6 +114,8 @@ bool FAudioDevice::Init()
 
 	UE_LOG(LogInit, Log, TEXT("FAudioDevice initialized." ));
 
+	bIsInitialized = true;
+
 	return true;
 }
 
diff --git a/Engine/Source/Runtime/Engine/Private/AudioDeviceManager.cpp b/Engine/Source/Runtime/Engine/Private/AudioDeviceManager.cpp
--- a/Engine/Source/Runtime/Engine/Private/AudioDeviceManager.cpp
+++ b/Engine/Source/Runtime/Engine/Private/AudioDeviceManager.cpp
@@ -17,6 +17,12 @@ static const uint16 AUDIO_DEVICE_MINIMUM_FREE_AUDIO_DEVICE_INDICES = 32;
 // Invalid handle used for initializing audio device handles
 static const uint32 AUDIO_DEVICE_HANDLE_INVALID			= ~0u;
 
+// The number of multiple audio devices allowed by default
+static const uint32 AUDIO_DEVICE_DEFAULT_ALLOWED_DEVICE_COUNT = 2;
+
+// The max number of audio devices allowed
+static const uint32 AUDIO_DEVICE_MAX_DEVICE_COUNT = 8;
+
 /*-----------------------------------------------------------------------------
 FAudioDeviceManager implementation.
 -----------------------------------------------------------------------------*/
@@ -25,6 +31,8 @@ FAudioDeviceManager::FAudioDeviceManager()
 	: AudioDeviceModule(nullptr)
 	, FreeIndicesSize(0)
 	, NumActiveAudioDevices(0)
+	, NumWorldsUsingMainAudioDevice(0)
+	, InFocusDeviceHandle(INDEX_NONE)
 	, NextResourceID(1)
 {
 }
@@ -48,7 +56,7 @@ void FAudioDeviceManager::RegisterAudioDeviceModule(IAudioDeviceModule* AudioDev
 	AudioDeviceModule = AudioDeviceModuleInput;
 }
 
-FAudioDevice* FAudioDeviceManager::CreateAudioDevice(uint32& HandleOut)
+FAudioDevice* FAudioDeviceManager::CreateAudioDevice(uint32& HandleOut, bool bCreateNewDevice)
 {
 	// If we don't have an audio device module, then we can't create new audio devices.
 	if (AudioDeviceModule == nullptr)
@@ -57,45 +65,61 @@ FAudioDevice* FAudioDeviceManager::CreateAudioDevice(uint32& HandleOut)
 		return nullptr;
 	}
 
-	// Create the new audio device and make sure it succeeded
-	FAudioDevice* NewAudioDevice = AudioDeviceModule->CreateAudioDevice();
-	if (NewAudioDevice == nullptr)
+	FAudioDevice* NewAudioDevice = nullptr;
+
+	if (NumActiveAudioDevices < AUDIO_DEVICE_DEFAULT_ALLOWED_DEVICE_COUNT || (bCreateNewDevice && NumActiveAudioDevices < AUDIO_DEVICE_MAX_DEVICE_COUNT))
 	{
-		HandleOut = AUDIO_DEVICE_HANDLE_INVALID;
-		return nullptr;
-	}
+		// Create the new audio device and make sure it succeeded
+		NewAudioDevice = AudioDeviceModule->CreateAudioDevice();
+		if (NewAudioDevice == nullptr)
+		{
+			HandleOut = AUDIO_DEVICE_HANDLE_INVALID;
+			return nullptr;
+		}
 
-	// Now generation a new audio device handle for the device and store the
-	// ptr to the new device in the array of audio devices.
+		// Now generation a new audio device handle for the device and store the
+		// ptr to the new device in the array of audio devices.
 
-	uint32 AudioDeviceIndex(AUDIO_DEVICE_HANDLE_INVALID);
+		uint32 AudioDeviceIndex(AUDIO_DEVICE_HANDLE_INVALID);
 
-	// First check to see if we should start recycling audio device indices, if not
-	// then we add a new entry to the Generation array and generate a new index
-	if (FreeIndicesSize > AUDIO_DEVICE_MINIMUM_FREE_AUDIO_DEVICE_INDICES)
-	{
-		FreeIndices.Dequeue(AudioDeviceIndex);
-		--FreeIndicesSize;
-		check(int32(AudioDeviceIndex) < Devices.Num());
-		check(Devices[AudioDeviceIndex] == nullptr);
-		Devices[AudioDeviceIndex] = NewAudioDevice;
+		// First check to see if we should start recycling audio device indices, if not
+		// then we add a new entry to the Generation array and generate a new index
+		if (FreeIndicesSize > AUDIO_DEVICE_MINIMUM_FREE_AUDIO_DEVICE_INDICES)
+		{
+			FreeIndices.Dequeue(AudioDeviceIndex);
+			--FreeIndicesSize;
+			check(int32(AudioDeviceIndex) < Devices.Num());
+			check(Devices[AudioDeviceIndex] == nullptr);
+			Devices[AudioDeviceIndex] = NewAudioDevice;
+		}
+		else
+		{
+			// Add a zeroth generation entry in the Generation array, get a brand new
+			// index and append the created device to the end of the Devices array
+
+			Generations.Add(0);
+			AudioDeviceIndex = Generations.Num() - 1;
+			check(AudioDeviceIndex < (1 << AUDIO_DEVICE_HANDLE_INDEX_BITS));
+			Devices.Add(NewAudioDevice);
+		}
+
+		HandleOut = CreateHandle(AudioDeviceIndex, Generations[AudioDeviceIndex]);
+
+		// Store the handle on the audio device itself
+		NewAudioDevice->DeviceHandle = HandleOut;
 	}
 	else
 	{
-		// Add a zeroth generation entry in the Generation array, get a brand new
-		// index and append the created device to the end of the Devices array
-
-		Generations.Add(0);
-		AudioDeviceIndex = Generations.Num() - 1;
-		check(AudioDeviceIndex < (1 << AUDIO_DEVICE_HANDLE_INDEX_BITS));
-		Devices.Add(NewAudioDevice);
+		++NumWorldsUsingMainAudioDevice;
+		FAudioDevice* MainAudioDevice = GEngine->GetMainAudioDevice();
+		if (MainAudioDevice)
+		{
+			HandleOut = MainAudioDevice->DeviceHandle;
+			NewAudioDevice = MainAudioDevice;
+		}
 	}
 
 	++NumActiveAudioDevices;
-	HandleOut = CreateHandle(AudioDeviceIndex, Generations[AudioDeviceIndex]);
-
-	// Store the handle on the audio device itself
-	NewAudioDevice->DeviceHandle = HandleOut;
 
 	return NewAudioDevice;
 }
@@ -127,12 +151,29 @@ bool FAudioDeviceManager::ShutdownAudioDevice(uint32 Handle)
 	check(NumActiveAudioDevices > 0);
 	--NumActiveAudioDevices;
 
-	// If we only have one audio device left, then set the active
-	// audio device to be the main audio device
-	if (NumActiveAudioDevices == 1)
+	// If there are more than 1 device active, check to see if this handle is the main audio device handle
+	if (NumActiveAudioDevices >= 1)
 	{
-		uint32 ActiveDeviceHandle = GEngine->GetAudioDeviceHandle();
-		SetActiveDevice(ActiveDeviceHandle);
+		uint32 MainDeviceHandle = GEngine->GetAudioDeviceHandle();
+
+		if (NumActiveAudioDevices == 1)
+		{
+			// If we only have one audio device left, then set the active
+			// audio device to be the main audio device
+			SetActiveDevice(MainDeviceHandle);
+		}
+
+		if (MainDeviceHandle == Handle)
+		{
+			// This was a handle to the main audio device using the main audio device
+			check(NumWorldsUsingMainAudioDevice > 0);
+			--NumWorldsUsingMainAudioDevice;
+
+			// If this is the main device handle, don't shut it down until it's the very last handle to get shut down
+			// this is because it's possible for some PIE sessions to be using the main audio device as a fallback to 
+			// preserve CPU performance on low-performance machines
+			return true;
+		}
 	}
 
 	uint32 Index = GetIndex(Handle);
@@ -172,6 +213,10 @@ bool FAudioDeviceManager::ShutdownAllAudioDevices()
 			ShutdownAudioDevice(AudioDevice->DeviceHandle);
 		}
 	}
+
+	check(NumActiveAudioDevices == 0);
+	check(NumWorldsUsingMainAudioDevice == 0);
+
 	return true;
 }
 
@@ -291,6 +336,11 @@ uint8 FAudioDeviceManager::GetNumActiveAudioDevices() const
 	return NumActiveAudioDevices;
 }
 
+uint8 FAudioDeviceManager::GetNumMainAudioDeviceWorlds() const
+{
+	return NumWorldsUsingMainAudioDevice;
+}
+
 uint32 FAudioDeviceManager::GetIndex(uint32 Handle) const
 {
 	return Handle & AUDIO_DEVICE_HANDLE_INDEX_MASK;
diff --git a/Engine/Source/Runtime/Engine/Private/GameViewportClient.cpp b/Engine/Source/Runtime/Engine/Private/GameViewportClient.cpp
--- a/Engine/Source/Runtime/Engine/Private/GameViewportClient.cpp
+++ b/Engine/Source/Runtime/Engine/Private/GameViewportClient.cpp
@@ -102,6 +102,8 @@ UGameViewportClient::UGameViewportClient(const FObjectInitializer& ObjectInitial
 	, bIgnoreInput(false)
 	, MouseCaptureMode(EMouseCaptureMode::CapturePermanently)
 	, bHideCursorDuringCapture(false)
+	, AudioDeviceHandle(INDEX_NONE)
+	, bHasAudioFocus(false)
 {
 
 	TitleSafeZone.MaxPercentX = 0.9f;
@@ -235,7 +237,7 @@ FString UGameViewportClient::ConsoleCommand( const FString& Command)
 	return *ConsoleOut;
 }
 
-void UGameViewportClient::Init(struct FWorldContext& WorldContext, UGameInstance* OwningGameInstance)
+void UGameViewportClient::Init(struct FWorldContext& WorldContext, UGameInstance* OwningGameInstance, bool bCreateNewAudioDevice)
 {
 	// set reference to world context
 	WorldContext.AddRef(World);
@@ -251,14 +253,14 @@ void UGameViewportClient::Init(struct FWorldContext& WorldContext, UGameInstance
 		FAudioDeviceManager* AudioDeviceManager = GEngine->GetAudioDeviceManager();
 		if (AudioDeviceManager)
 		{
-			FAudioDevice* NewAudioDevice = AudioDeviceManager->CreateAudioDevice(AudioDeviceHandle);
+			FAudioDevice* NewAudioDevice = AudioDeviceManager->CreateAudioDevice(AudioDeviceHandle, bCreateNewAudioDevice);
 			if (NewAudioDevice)
 			{
 				if (NewAudioDevice->Init())
 				{
 					// Set the base mix of the new device based on the world settings of the world
 					if (World)
-					{ 
+					{
 						NewAudioDevice->SetDefaultBaseSoundMix(World->GetWorldSettings()->DefaultBaseSoundMix);
 
 						// Set the world's audio device handle to use so that sounds which play in that world will use the correct audio device
@@ -936,31 +938,52 @@ void UGameViewportClient::Draw(FViewport* InViewport, FCanvas* SceneCanvas)
 						// Update the listener.
 						if (AudioDevice != NULL && PlayerController != NULL)
 						{
-							FVector Location;
-							FVector ProjFront;
-							FVector ProjRight;
-							PlayerController->GetAudioListenerPosition(/*out*/ Location, /*out*/ ProjFront, /*out*/ ProjRight);
+							bool bUpdateListenerPosition = true;
 
-							FTransform ListenerTransform(FRotationMatrix::MakeFromXY(ProjFront, ProjRight));
-							ListenerTransform.SetTranslation(Location);
-							ListenerTransform.NormalizeRotation();
+							// If the main audio device is used for multiple PIE viewport clients, we only
+							// want to update the main audio device listener position if it is in focus
+							if (GEngine)
+							{
+								FAudioDeviceManager* AudioDeviceManager = GEngine->GetAudioDeviceManager();
+
+								// If there is more than one world referencing the main audio device
+								if (AudioDeviceManager->GetNumMainAudioDeviceWorlds() > 1)
+								{
+									uint32 MainAudioDeviceHandle = GEngine->GetAudioDeviceHandle();
+									if (AudioDevice->DeviceHandle == MainAudioDeviceHandle && !bHasAudioFocus)
+									{
+										bUpdateListenerPosition = false;
+									}
+								}
+							}
 
-							bReverbSettingsFound = true;
+							if (bUpdateListenerPosition)
+							{
+								FVector Location;
+								FVector ProjFront;
+								FVector ProjRight;
+								PlayerController->GetAudioListenerPosition(/*out*/ Location, /*out*/ ProjFront, /*out*/ ProjRight);
 
-							FReverbSettings PlayerReverbSettings;
-							FInteriorSettings PlayerInteriorSettings;
-							class AAudioVolume* PlayerAudioVolume = GetWorld()->GetAudioSettings(Location, &PlayerReverbSettings, &PlayerInteriorSettings);
+								FTransform ListenerTransform(FRotationMatrix::MakeFromXY(ProjFront, ProjRight));
+								ListenerTransform.SetTranslation(Location);
+								ListenerTransform.NormalizeRotation();
 
-							if (AudioVolume == nullptr || (PlayerAudioVolume != nullptr && PlayerAudioVolume->Priority > AudioVolume->Priority))
-							{
-								AudioVolume = PlayerAudioVolume;
-								ReverbSettings = PlayerReverbSettings;
-							}
+								bReverbSettingsFound = true;
 
-							uint32 ViewportIndex = PlayerViewMap.Num() - 1;
-							AudioDevice->SetListener(ViewportIndex, ListenerTransform, (View->bCameraCut ? 0.f : GetWorld()->GetDeltaSeconds()), PlayerAudioVolume, PlayerInteriorSettings);
-						}
+								FReverbSettings PlayerReverbSettings;
+								FInteriorSettings PlayerInteriorSettings;
+								class AAudioVolume* PlayerAudioVolume = GetWorld()->GetAudioSettings(Location, &PlayerReverbSettings, &PlayerInteriorSettings);
 
+								if (AudioVolume == nullptr || (PlayerAudioVolume != nullptr && PlayerAudioVolume->Priority > AudioVolume->Priority))
+								{
+									AudioVolume = PlayerAudioVolume;
+									ReverbSettings = PlayerReverbSettings;
+								}
+
+								uint32 ViewportIndex = PlayerViewMap.Num() - 1;
+								AudioDevice->SetListener(ViewportIndex, ListenerTransform, (View->bCameraCut ? 0.f : GetWorld()->GetDeltaSeconds()), PlayerAudioVolume, PlayerInteriorSettings);
+							}
+						}
 					}
 
 					// Add view information for resource streaming.
@@ -1386,6 +1409,11 @@ void UGameViewportClient::LostFocus(FViewport* InViewport)
 			}
 		}
 	}
+
+	if (GEngine && GEngine->GetAudioDeviceManager())
+	{
+		bHasAudioFocus = false;
+	}
 }
 
 void UGameViewportClient::ReceivedFocus(FViewport* InViewport)
@@ -1398,6 +1426,7 @@ void UGameViewportClient::ReceivedFocus(FViewport* InViewport)
 	if (GEngine && GEngine->GetAudioDeviceManager())
 	{ 
 		GEngine->GetAudioDeviceManager()->SetActiveDevice(AudioDeviceHandle);
+		bHasAudioFocus = true;
 	}
 }
 
diff --git a/Engine/Source/Runtime/Engine/Private/UnrealEngine.cpp b/Engine/Source/Runtime/Engine/Private/UnrealEngine.cpp
--- a/Engine/Source/Runtime/Engine/Private/UnrealEngine.cpp
+++ b/Engine/Source/Runtime/Engine/Private/UnrealEngine.cpp
@@ -1687,7 +1687,7 @@ bool UEngine::InitializeAudioDeviceManager()
 					bool bSucceeded = false;
 
 					// Create a new audio device.
-					FAudioDevice* NewAudioDevice = AudioDeviceManager->CreateAudioDevice(MainAudioDeviceHandle);
+					FAudioDevice* NewAudioDevice = AudioDeviceManager->CreateAudioDevice(MainAudioDeviceHandle, true);
 					if (NewAudioDevice)
 					{
 						// Initialize the audio device
diff --git a/Engine/Source/Runtime/Engine/Public/AudioDevice.h b/Engine/Source/Runtime/Engine/Public/AudioDevice.h
--- a/Engine/Source/Runtime/Engine/Public/AudioDevice.h
+++ b/Engine/Source/Runtime/Engine/Public/AudioDevice.h
@@ -806,6 +806,9 @@ public:
 	/** Whether the audio device is active (current audio device in-focus in PIE) */
 	uint32 bIsDeviceMuted:1;
 
+	/** Whether the audio device has been initialized */
+	uint32 bIsInitialized:1;
+
 private:
 
 	TArray<struct FActiveSound*> ActiveSounds;
diff --git a/Engine/Source/Runtime/Engine/Public/AudioDeviceManager.h b/Engine/Source/Runtime/Engine/Public/AudioDeviceManager.h
--- a/Engine/Source/Runtime/Engine/Public/AudioDeviceManager.h
+++ b/Engine/Source/Runtime/Engine/Public/AudioDeviceManager.h
@@ -29,7 +29,7 @@ public:
 	* Creates and audio device instance internally and returns a
 	* handle to the audio device. Returns true on success.
 	*/
-	class FAudioDevice* CreateAudioDevice(uint32 & HandleOut);
+	class FAudioDevice* CreateAudioDevice(uint32 & HandleOut, bool bCreateNewDevice);
 
 	/**
 	* Returns whether the audio device handle is valid (i.e. points to
@@ -57,6 +57,9 @@ public:
 	/** Returns the current number of active audio devices. */
 	uint8 GetNumActiveAudioDevices() const;
 
+	/** Returns the number of worlds (e.g. PIE viewports) using the main audio device. */
+	uint8 GetNumMainAudioDeviceWorlds() const;
+
 	/** Updates all active audio devices */
 	void UpdateActiveAudioDevices(bool bGameTicking);
 
@@ -125,25 +128,28 @@ private:
 	/** Array for generation counts of audio devices in indices */
 	TArray<uint8> Generations;
 
-    /** Audio device module which creates audio devices. */
-    IAudioDeviceModule* AudioDeviceModule;
+	/** Audio device module which creates audio devices. */
+	IAudioDeviceModule* AudioDeviceModule;
+
+	/** Count of the number of free slots in the audio device array. */
+	uint32 FreeIndicesSize;
 
-    /** Count of the number of free slots in the audio device array. */
-    uint32 FreeIndicesSize;
+	/** Number of actively created audio device instances. */
+	uint8 NumActiveAudioDevices;
 
-    /** Number of actively created audio device instances. */
-    uint8 NumActiveAudioDevices;
+	/** Number of worlds using the main audio device instance. */
+	uint8 NumWorldsUsingMainAudioDevice;
 
-    /** Queue for free indices */
-    TQueue<uint32> FreeIndices;
+	/** Queue for free indices */
+	TQueue<uint32> FreeIndices;
 
-    /**
+	/**
 	* Array of audio device pointers. If the device has been free, then
 	* the device ptr at the array index will be null.
 	*/
 	TArray<FAudioDevice* > Devices;
 
-    /** The audio device handle currently in focus */
+	/** The audio device handle currently in focus */
 	int32 InFocusDeviceHandle;
 
 	/** Next resource ID to assign out to a wave/buffer */
