Pass in more accurate start/end time values for part of stream that must be in memory when we fast forward

diff --git a/Engine/Source/Runtime/Engine/Private/DemoNetDriver.cpp b/Engine/Source/Runtime/Engine/Private/DemoNetDriver.cpp
--- a/Engine/Source/Runtime/Engine/Private/DemoNetDriver.cpp
+++ b/Engine/Source/Runtime/Engine/Private/DemoNetDriver.cpp
@@ -1144,14 +1144,18 @@ void UDemoNetDriver::TickDemoPlayback( float DeltaSeconds )
 	
 	if ( TimeToSkip > 0.0f )
 	{
-		ReplayStreamer->SetHighPriorityTimeRange( 0, DemoTotalTime );		// HACK for now, download the entire demo until we get support to download to any time
+		const uint32 TimeInMSToCheck = FMath::Clamp( GetDemoCurrentTimeInMS() + (uint32)( TimeToSkip * 1000 ), (uint32)0, TotalDemoTimeInMS );
 
-		if ( !ReplayStreamer->IsDataAvailableForTimeRange( 0, DemoTotalTime ) )
+		ReplayStreamer->SetHighPriorityTimeRange( TimeInMSToCheck, TimeInMSToCheck );
+
+		if ( !ReplayStreamer->IsDataAvailableForTimeRange( TimeInMSToCheck, TimeInMSToCheck ) )
 		{
 			PauseChannels( true );
 			return;
 		}
 
+		ReplayStreamer->SetHighPriorityTimeRange( 0, 0 );
+
 		PauseChannels( false );
 
 		DemoCurrentTime += TimeToSkip;
