Actually use RenameFlags

diff --git a/Engine/Source/Runtime/Engine/Private/Components/ChildActorComponent.cpp b/Engine/Source/Runtime/Engine/Private/Components/ChildActorComponent.cpp
--- a/Engine/Source/Runtime/Engine/Private/Components/ChildActorComponent.cpp
+++ b/Engine/Source/Runtime/Engine/Private/Components/ChildActorComponent.cpp
@@ -301,7 +301,7 @@ void UChildActorComponent::DestroyChildActor(const bool bRequiresRename)
 				{
 					const FString ObjectBaseName = FString::Printf(TEXT("DESTROYED_%s_CHILDACTOR"), *ChildClass->GetName());
 					const ERenameFlags RenameFlags = (GetWorld()->IsGameWorld() ? REN_DoNotDirty | REN_ForceNoResetLoaders : REN_DoNotDirty);
-					ChildActor->Rename(*MakeUniqueObjectName(ChildActor->GetOuter(), ChildClass, *ObjectBaseName).ToString(), nullptr, REN_DoNotDirty);
+					ChildActor->Rename(*MakeUniqueObjectName(ChildActor->GetOuter(), ChildClass, *ObjectBaseName).ToString(), nullptr, RenameFlags);
 				}
 				World->DestroyActor(ChildActor);
 			}
