#ue4 - Optimize usage of TWeakObjectPointer in FBodyInstance. Optimize TArray usage in FBodyInstance.

diff --git a/Engine/Source/Runtime/Engine/Private/PhysicsEngine/BodyInstance.cpp b/Engine/Source/Runtime/Engine/Private/PhysicsEngine/BodyInstance.cpp
--- a/Engine/Source/Runtime/Engine/Private/PhysicsEngine/BodyInstance.cpp
+++ b/Engine/Source/Runtime/Engine/Private/PhysicsEngine/BodyInstance.cpp
@@ -640,7 +640,8 @@ void FBodyInstance::CreateDOFLock()
 ECollisionEnabled::Type FBodyInstance::GetCollisionEnabled() const
 {
 	// Check actor override
-	AActor* Owner = OwnerComponent.IsValid() ? OwnerComponent.Get()->GetOwner() : NULL;
+	const UPrimitiveComponent* OwnerComponentInst = OwnerComponent.Get();
+	AActor* Owner = OwnerComponentInst ? OwnerComponentInst->GetOwner() : NULL;
 	if ((Owner != NULL) && !Owner->GetActorEnableCollision())
 	{
 		return ECollisionEnabled::NoCollision;
@@ -679,7 +680,7 @@ void FBodyInstance::UpdatePhysicsShapeFilterData(uint32 SkelMeshCompID, bool bUs
 			bool bUseNotify = bNotifyOverride ? *bNotifyOverride : BI->bNotifyRigidBodyCollision;
 
 
-			UPrimitiveComponent* OwnerPrimitiveComponent = BI->OwnerComponent.IsValid() ? BI->OwnerComponent.Get() : nullptr;
+			UPrimitiveComponent* OwnerPrimitiveComponent = BI->OwnerComponent.Get();
 			int32 CompID = (OwnerPrimitiveComponent != nullptr) ? OwnerPrimitiveComponent->GetUniqueID() : 0;
 
 			// Create the filterdata structs
@@ -817,8 +818,9 @@ void FBodyInstance::UpdatePhysicsFilterData(bool bForceSimpleAsComplex)
 	}
 
 	// Figure out if we are static
-	AActor* Owner = OwnerComponent.IsValid() ? OwnerComponent.Get()->GetOwner() : NULL;
-	const bool bPhysicsStatic = !OwnerComponent.IsValid() || OwnerComponent.Get()->IsWorldGeometry();
+	UPrimitiveComponent* OwnerComponentInst = OwnerComponent.Get();
+	AActor* Owner = OwnerComponentInst ? OwnerComponentInst->GetOwner() : NULL;
+	const bool bPhysicsStatic = !OwnerComponentInst || OwnerComponentInst->IsWorldGeometry();
 
 	// Grab collision setting from body instance
 	TEnumAsByte<ECollisionEnabled::Type> UseCollisionEnabled = GetCollisionEnabled(); // this checks actor override
@@ -831,7 +833,7 @@ void FBodyInstance::UpdatePhysicsFilterData(bool bForceSimpleAsComplex)
 
 	// Get skelmeshcomp ID
 	uint32 SkelMeshCompID = 0;
-	if (USkeletalMeshComponent* SkelMeshComp = Cast<USkeletalMeshComponent>(OwnerComponent.Get()))
+	if (USkeletalMeshComponent* SkelMeshComp = Cast<USkeletalMeshComponent>(OwnerComponentInst))
 	{
 		SkelMeshCompID = SkelMeshComp->GetUniqueID();
 
@@ -870,12 +872,12 @@ void FBodyInstance::UpdatePhysicsFilterData(bool bForceSimpleAsComplex)
 
 #if WITH_EDITOR
 	// if no collision, but if world wants to enable trace collision for components, allow it
-	if ((UseCollisionEnabled == ECollisionEnabled::NoCollision) && Owner && (Owner->IsA(AVolume::StaticClass()) == false))
+	if ((UseCollisionEnabled == ECollisionEnabled::NoCollision) && Owner && (Cast<AVolume>(Owner) == nullptr))
 	{
 		UWorld* World = Owner->GetWorld();
-		UPrimitiveComponent* PrimComp = OwnerComponent.Get();
+		UPrimitiveComponent* PrimComp = OwnerComponentInst;
 		if (World && World->bEnableTraceCollision && 
-			(PrimComp->IsA(UStaticMeshComponent::StaticClass()) || PrimComp->IsA(USkeletalMeshComponent::StaticClass()) || PrimComp->IsA(UBrushComponent::StaticClass())))
+			(Cast<UStaticMeshComponent>(PrimComp) || Cast<USkeletalMeshComponent>(PrimComp) || Cast<UBrushComponent>(PrimComp)))
 		{
 			//UE_LOG(LogPhysics, Warning, TEXT("Enabling collision %s : %s"), *GetNameSafe(Owner), *GetNameSafe(OwnerComponent.Get()));
 			// clear all other channel just in case other people using those channels to do something
@@ -2241,42 +2243,47 @@ bool FBodyInstance::IsDynamic() const
 
 void FBodyInstance::SetInstanceSimulatePhysics(bool bSimulate, bool bMaintainPhysicsBlending)
 {
-#if !(UE_BUILD_SHIPPING || UE_BUILD_TEST)
-	if (bSimulate && OwnerComponent.IsValid())
+	if (bSimulate)
 	{
-		if (!IsValidBodyInstance())
-		{
-			FMessageLog("PIE").Warning(FText::Format(LOCTEXT("SimPhysNoBody", "Trying to simulate physics on ''{0}'' but no physics body."),
-				FText::FromString(GetPathNameSafe(OwnerComponent.Get()))));
-		}
-		else if (!IsDynamic())
+		UPrimitiveComponent* OwnerComponentInst = OwnerComponent.Get();
+
+#if !(UE_BUILD_SHIPPING || UE_BUILD_TEST)
+		if (OwnerComponentInst)
 		{
-			FMessageLog("PIE").Warning(FText::Format(LOCTEXT("SimPhysStatic", "Trying to simulate physics on ''{0}'' but it is static."),
-				FText::FromString(GetPathNameSafe(OwnerComponent.Get()))));
+			if (!IsValidBodyInstance())
+			{
+				FMessageLog("PIE").Warning(FText::Format(LOCTEXT("SimPhysNoBody", "Trying to simulate physics on ''{0}'' but no physics body."),
+					FText::FromString(GetPathNameSafe(OwnerComponentInst))));
+			}
+			else if (!IsDynamic())
+			{
+				FMessageLog("PIE").Warning(FText::Format(LOCTEXT("SimPhysStatic", "Trying to simulate physics on ''{0}'' but it is static."),
+					FText::FromString(GetPathNameSafe(OwnerComponentInst))));
+			}
 		}
-	}
 #endif
 
-	// If we are enabling simulation, and we are the root body of our component, we detach the component 
-	if ((bSimulate == true) && OwnerComponent.IsValid() && OwnerComponent->IsRegistered() && OwnerComponent->GetBodyInstance() == this)
-	{
-		if (OwnerComponent->AttachParent)
+		// If we are enabling simulation, and we are the root body of our component, we detach the component 
+		if (OwnerComponentInst && OwnerComponentInst->IsRegistered() && OwnerComponentInst->GetBodyInstance() == this)
 		{
-			OwnerComponent->DetachFromParent(true);
-		}
-		else if (bSimulatePhysics == false)	//if we're switching from kinematic to simulated
-		{
-			//if we are the root in hierarchy we must make sure to update children that are welded
-			TArray<FBodyInstance*> ChildrenBodies;
-			TArray<FName> ChildrenLabels;
-			OwnerComponent->GetWeldedBodies(ChildrenBodies, ChildrenLabels);
-
-			for (int32 ChildIdx = 0; ChildIdx < ChildrenBodies.Num(); ++ChildIdx)
+			if (OwnerComponentInst->AttachParent)
+			{
+				OwnerComponentInst->DetachFromParent(true);
+			}
+			else if (bSimulatePhysics == false)	//if we're switching from kinematic to simulated
 			{
-				FBodyInstance* ChildBI = ChildrenBodies[ChildIdx];
-				if (ChildBI != this)
+				//if we are the root in hierarchy we must make sure to update children that are welded
+				TArray<FBodyInstance*> ChildrenBodies;
+				TArray<FName> ChildrenLabels;
+				OwnerComponentInst->GetWeldedBodies(ChildrenBodies, ChildrenLabels);
+
+				for (int32 ChildIdx = 0; ChildIdx < ChildrenBodies.Num(); ++ChildIdx)
 				{
-					Weld(ChildBI, ChildBI->OwnerComponent->GetSocketTransform(ChildrenLabels[ChildIdx]));
+					FBodyInstance* ChildBI = ChildrenBodies[ChildIdx];
+					if (ChildBI != this)
+					{
+						Weld(ChildBI, ChildBI->OwnerComponent->GetSocketTransform(ChildrenLabels[ChildIdx]));
+					}
 				}
 			}
 		}
@@ -2411,11 +2418,12 @@ void FBodyInstance::SetBodyTransform(const FTransform& NewTransform, bool bTelep
 		// STATIC
 		else
 		{
-			const bool bIsGame = !GIsEditor || (OwnerComponent != NULL && OwnerComponent->GetWorld()->IsGameWorld());
+			UPrimitiveComponent* OwnerComponentInst = OwnerComponent.Get();
+			const bool bIsGame = !GIsEditor || (OwnerComponentInst != NULL && OwnerComponentInst->GetWorld()->IsGameWorld());
 			// Do NOT move static actors in-game, give a warning but let it happen
 			if (bIsGame)
 			{
-				const FString ComponentPathName = (OwnerComponent != NULL) ? OwnerComponent->GetPathName() : TEXT("NONE");
+				const FString ComponentPathName = (OwnerComponentInst != NULL) ? OwnerComponentInst->GetPathName() : TEXT("NONE");
 				UE_LOG(LogPhysics, Warning, TEXT("MoveFixedBody: Trying to move component'%s' with a non-Movable Mobility."), *ComponentPathName);
 			}
 			// In EDITOR, go ahead and move it with no warning, we are editing the level
@@ -3424,10 +3432,11 @@ FString FBodyInstance::GetBodyDebugName() const
 {
 	FString DebugName;
 
-	if (OwnerComponent != NULL)
+	UPrimitiveComponent* OwnerComponentInst = OwnerComponent.Get();
+	if (OwnerComponentInst != NULL)
 	{
-		DebugName = OwnerComponent->GetPathName();
-		if (const UObject* StatObject = OwnerComponent->AdditionalStatObject())
+		DebugName = OwnerComponentInst->GetPathName();
+		if (const UObject* StatObject = OwnerComponentInst->AdditionalStatObject())
 		{
 			DebugName += TEXT(" ");
 			StatObject->AppendName(DebugName);
@@ -3470,7 +3479,7 @@ bool FBodyInstance::LineTrace(struct FHitResult& OutHit, const FVector& Start, c
 			float BestHitDistance = BIG_NUMBER;
 
 			// Get all the shapes from the actor
-			TArray<PxShape*> PShapes;
+			TArray<PxShape*, TInlineAllocator<16>> PShapes;
 			PShapes.AddZeroed(RigidBody->getNbShapes());
 			int32 NumShapes = RigidBody->getShapes(PShapes.GetData(), PShapes.Num());
 
@@ -3521,8 +3530,8 @@ bool FBodyInstance::LineTrace(struct FHitResult& OutHit, const FVector& Start, c
 				QueryFilter.word2 = 0xFFFFF;
 
 				PxTransform PStartTM(U2PVector(Start));
-				const UPrimitiveComponent* Owner = OwnerComponent.Get();
-				ConvertQueryImpactHit(Owner ? Owner->GetWorld() : nullptr, BestHit, OutHit, DeltaMag, QueryFilter, Start, End, NULL, PStartTM, false, bReturnPhysicalMaterial);
+				const UPrimitiveComponent* OwnerComponentInst = OwnerComponent.Get();
+				ConvertQueryImpactHit(OwnerComponentInst ? OwnerComponentInst->GetWorld() : nullptr, BestHit, OutHit, DeltaMag, QueryFilter, Start, End, NULL, PStartTM, false, bReturnPhysicalMaterial);
 				bHitSomething = true;
 			}
 		}
@@ -3582,15 +3591,16 @@ bool FBodyInstance::InternalSweepPhysX(struct FHitResult& OutHit, const FVector&
 		
 		PxQuat PGeomRot = PxQuat::createIdentity();
 
+		UPrimitiveComponent* OwnerComponentInst = OwnerComponent.Get();
 		PxTransform PStartTM(U2PVector(Start), PGeomRot);
-		PxTransform PCompTM(U2PTransform(OwnerComponent->ComponentToWorld));
+		PxTransform PCompTM(U2PTransform(OwnerComponentInst->ComponentToWorld));
 
 		PxVec3 PDir = U2PVector(Delta/DeltaMag);
 
 		PxSweepHit PHit;
 
 		// Get all the shapes from the actor
-		TArray<PxShape*> PShapes;
+		TArray<PxShape*, TInlineAllocator<16>> PShapes;
 		PShapes.AddZeroed(RigidBody->getNbShapes());
 		int32 NumShapes = RigidBody->getShapes(PShapes.GetData(), PShapes.Num());
 
@@ -3624,8 +3634,7 @@ bool FBodyInstance::InternalSweepPhysX(struct FHitResult& OutHit, const FVector&
 						PHit.shape = PShape;
 						PHit.actor = PShape->getActor();
 						PxTransform PStartTransform(U2PVector(Start));
-						const UPrimitiveComponent* Owner = OwnerComponent.Get();
-						ConvertQueryImpactHit(Owner ? Owner->GetWorld() : nullptr, PHit, OutHit, DeltaMag, QueryFilter, Start, End, NULL, PStartTransform, false, false);
+						ConvertQueryImpactHit(OwnerComponentInst->GetWorld(), PHit, OutHit, DeltaMag, QueryFilter, Start, End, NULL, PStartTransform, false, false);
 						return true;
 					}
 				}
@@ -3648,7 +3657,7 @@ float FBodyInstance::GetDistanceToBody(const FVector& Point, FVector& OutPointOn
 	}
 
 	// Get all the shapes from the actor
-	TArray<PxShape*> PShapes;
+	TArray<PxShape*, TInlineAllocator<16>> PShapes;
 	PShapes.AddZeroed(RigidBody->getNbShapes());
 	int32 NumShapes = RigidBody->getShapes(PShapes.GetData(), PShapes.Num());
 
@@ -3722,12 +3731,16 @@ bool FBodyInstance::OverlapTest(const FVector& Position, const FQuat& Rotation,
 
 FTransform RootSpaceToWeldedSpace(const FBodyInstance* BI, const FTransform& RootTM)
 {
-	if (BI->WeldParent && BI->OwnerComponent.IsValid())
+	if (BI->WeldParent)
 	{
-		FTransform RootToWelded = BI->OwnerComponent->GetRelativeTransform().Inverse();
-		RootToWelded.ScaleTranslation(BI->Scale3D);
+		UPrimitiveComponent* BIOwnerComponentInst = BI->OwnerComponent.Get();
+		if (BIOwnerComponentInst)
+		{
+			FTransform RootToWelded = BIOwnerComponentInst->GetRelativeTransform().Inverse();
+			RootToWelded.ScaleTranslation(BI->Scale3D);
 
-		return RootToWelded * RootTM;
+			return RootToWelded * RootTM;
+		}
 	}
 
 	return RootTM;
@@ -3776,7 +3789,7 @@ bool FBodyInstance::OverlapMulti(TArray<struct FOverlapResult>& InOutOverlaps, c
 		{
 			SCOPED_SCENE_READ_LOCK(PRigidActor->getScene());
 
-			TArray<PxShape*, TInlineAllocator<8>> PShapes;
+			TArray<PxShape*, TInlineAllocator<16>> PShapes;
 			PShapes.AddZeroed(PRigidActor->getNbShapes());
 			int32 NumShapes = PRigidActor->getShapes(PShapes.GetData(), PShapes.Num());
 
@@ -3834,7 +3847,7 @@ bool FBodyInstance::OverlapPhysX(const PxGeometry& PGeom, const PxTransform& Sha
 	}
 
 	// Get all the shapes from the actor
-	TArray<PxShape*> PShapes;
+	TArray<PxShape*, TInlineAllocator<16>> PShapes;
 	PShapes.AddZeroed(RigidBody->getNbShapes());
 	int32 NumShapes = RigidBody->getShapes(PShapes.GetData(), PShapes.Num());
 
@@ -3977,9 +3990,10 @@ void FBodyInstance::FixupData(class UObject* Loader)
 bool FBodyInstance::UseAsyncScene() const
 {
 	bool bHasAsyncScene = true;
-	if (OwnerComponent.IsValid())
+	UPrimitiveComponent* OwnerComponentInst = OwnerComponent.Get();
+	if (OwnerComponentInst)
 	{
-		bHasAsyncScene = OwnerComponent->GetWorld()->GetPhysicsScene()->HasAsyncScene();
+		bHasAsyncScene = OwnerComponentInst->GetWorld()->GetPhysicsScene()->HasAsyncScene();
 	}
 	return bUseAsyncScene && bHasAsyncScene;
 }
@@ -4001,7 +4015,7 @@ void FBodyInstance::ApplyMaterialToShape(PxShape* PShape, PxMaterial* PSimpleMat
 	// If a triangle mesh, need to get array of materials...
 	if(PShape->getGeometryType() == PxGeometryType::eTRIANGLEMESH)
 	{
-		TArray<PxMaterial*> PComplexMats;
+		TArray<PxMaterial*, TInlineAllocator<16>> PComplexMats;
 		PComplexMats.AddUninitialized(ComplexPhysMats.Num());
 		for(int MatIdx = 0; MatIdx < ComplexPhysMats.Num(); MatIdx++)
 		{
@@ -4131,8 +4145,9 @@ void FBodyInstance::GetFilterData(FShapeData& ShapeData, bool bForceSimpleAsComp
 	}
 
 	// Figure out if we are static
-	AActor* Owner = OwnerComponent.IsValid() ? OwnerComponent.Get()->GetOwner() : NULL;
-	const bool bPhysicsStatic = !OwnerComponent.IsValid() || OwnerComponent.Get()->IsWorldGeometry();
+	UPrimitiveComponent* OwnerComponentInst = OwnerComponent.Get();
+	AActor* Owner = OwnerComponentInst ? OwnerComponentInst->GetOwner() : NULL;
+	const bool bPhysicsStatic = !OwnerComponentInst || OwnerComponentInst->IsWorldGeometry();
 
 	// Grab collision setting from body instance
 	ShapeData.CollisionEnabled = GetCollisionEnabled(); // this checks actor override
@@ -4145,7 +4160,7 @@ void FBodyInstance::GetFilterData(FShapeData& ShapeData, bool bForceSimpleAsComp
 
 	// Get skelmeshcomp ID
 	uint32 SkelMeshCompID = 0;
-	if(USkeletalMeshComponent* SkelMeshComp = Cast<USkeletalMeshComponent>(OwnerComponent.Get()))
+	if(USkeletalMeshComponent* SkelMeshComp = Cast<USkeletalMeshComponent>(OwnerComponentInst))
 	{
 		SkelMeshCompID = SkelMeshComp->GetUniqueID();
 
@@ -4187,7 +4202,7 @@ void FBodyInstance::GetFilterData(FShapeData& ShapeData, bool bForceSimpleAsComp
 	if((ShapeData.CollisionEnabled == ECollisionEnabled::NoCollision) && Owner && (Owner->IsA(AVolume::StaticClass()) == false))
 	{
 		UWorld* World = Owner->GetWorld();
-		UPrimitiveComponent* PrimComp = OwnerComponent.Get();
+		UPrimitiveComponent* PrimComp = OwnerComponentInst;
 		if(World && World->bEnableTraceCollision &&
 		   (PrimComp->IsA(UStaticMeshComponent::StaticClass()) || PrimComp->IsA(USkeletalMeshComponent::StaticClass()) || PrimComp->IsA(UBrushComponent::StaticClass())))
 		{
@@ -4213,7 +4228,7 @@ void FBodyInstance::GetFilterData(FShapeData& ShapeData, bool bForceSimpleAsComp
 			PxFilterData PSimFilterData;
 			PxFilterData PSimpleQueryData;
 			PxFilterData PComplexQueryData;
-			int32 CompID = (OwnerComponent != nullptr) ? OwnerComponent->GetUniqueID() : 0;
+			int32 CompID = (OwnerComponentInst != nullptr) ? OwnerComponentInst->GetUniqueID() : 0;
 			CreateShapeFilterData(ObjectType, CompID, UseResponse, SkelMeshCompID, InstanceBodyIndex, PSimpleQueryData, PSimFilterData, bUseCCD && !bPhysicsStatic, bUseNotifyRBCollision, bPhysicsStatic);	//CCD is determined by root body in case of welding
 			PComplexQueryData = PSimpleQueryData;
 			
@@ -4356,8 +4371,9 @@ void FBodyInstance::SetShapeFlags(TEnumAsByte<ECollisionEnabled::Type> UseCollis
 	bool bUpdateMassProperties = false;
 	if(UseCollisionEnabled != ECollisionEnabled::NoCollision)
 	{
-		AActor* Owner = OwnerComponent.IsValid() ? OwnerComponent.Get()->GetOwner() : NULL;
-		const bool bPhysicsStatic = !OwnerComponent.IsValid() || OwnerComponent.Get()->IsWorldGeometry();
+		UPrimitiveComponent* OwnerComponentInst = OwnerComponent.Get();
+		AActor* Owner = OwnerComponentInst ? OwnerComponentInst->GetOwner() : NULL;
+		const bool bPhysicsStatic = !OwnerComponentInst || OwnerComponentInst->IsWorldGeometry();
 		
 		// Only perform scene queries in the synchronous scene for static shapes
 		if(bPhysicsStatic)
@@ -4386,7 +4402,7 @@ void FBodyInstance::SetShapeFlags(TEnumAsByte<ECollisionEnabled::Type> UseCollis
 				PShape->setFlag(PxShapeFlag::eSIMULATION_SHAPE, true);
 			}
 
-			if(OwnerComponent == NULL || !OwnerComponent->IsA(UModelComponent::StaticClass()))
+			if(OwnerComponentInst == NULL || !OwnerComponentInst->IsA(UModelComponent::StaticClass()))
 			{
 				PShape->setFlag(PxShapeFlag::eVISUALIZATION, false); // dont draw the tri mesh, we can see it anyway, and its slow
 			}
@@ -4450,11 +4466,12 @@ void FBodyInstance::GetShapeFlags(FShapeData& ShapeData, TEnumAsByte<ECollisionE
 	ShapeData.SimpleShapeFlags |= PxShapeFlag::eVISUALIZATION;
 	ShapeData.ComplexShapeFlags |= PxShapeFlag::eVISUALIZATION;
 
-	AActor* Owner = OwnerComponent.IsValid() ? OwnerComponent.Get()->GetOwner() : NULL;
+	UPrimitiveComponent* OwnerComponentInst = OwnerComponent.Get();
+	AActor* Owner = OwnerComponentInst ? OwnerComponentInst->GetOwner() : NULL;
 
 	if(ShapeData.CollisionEnabled != ECollisionEnabled::NoCollision)
 	{
-		const bool bPhysicsStatic = !OwnerComponent.IsValid() || OwnerComponent.Get()->IsWorldGeometry();
+		const bool bPhysicsStatic = !OwnerComponentInst || OwnerComponentInst->IsWorldGeometry();
 
 		// Only perform scene queries in the synchronous scene for static shapes
 		if(bPhysicsStatic)
@@ -4471,7 +4488,7 @@ void FBodyInstance::GetShapeFlags(FShapeData& ShapeData, TEnumAsByte<ECollisionE
 			ShapeData.ComplexShapeFlags |= PxShapeFlag::eSIMULATION_SHAPE;
 		}
 		
-		if(OwnerComponent == NULL || !OwnerComponent->IsA(UModelComponent::StaticClass()))
+		if(OwnerComponentInst == NULL || !OwnerComponentInst->IsA(UModelComponent::StaticClass()))
 		{
 			ShapeData.ComplexShapeFlags.clear(PxShapeFlag::eVISUALIZATION);
 		}
