#ue4 - Renamed UProjectileMovementComponent::HandleHitWall() to new virtual HandleBlockingHit() which has a new enum return value. Refactored Deflection/Sliding determination to new HandleDeflection() virtual.

- HandleBlockingHit() returns an enum indicating how simulation should continue: Deflect, AdvanceNextSubstep, Abort.
- StopSimulating() is now virtual.
- Added documentation to a few other methods.

UE-10106

diff --git a/Engine/Source/Runtime/Engine/Classes/GameFramework/ProjectileMovementComponent.h b/Engine/Source/Runtime/Engine/Classes/GameFramework/ProjectileMovementComponent.h
--- a/Engine/Source/Runtime/Engine/Classes/GameFramework/ProjectileMovementComponent.h
+++ b/Engine/Source/Runtime/Engine/Classes/GameFramework/ProjectileMovementComponent.h
@@ -112,6 +112,7 @@ class ENGINE_API UProjectileMovementComponent : public UMovementComponent
 	/**
 	 * If velocity is below this threshold after a bounce, stops simulating and triggers the OnProjectileStop event.
 	 * Ignored if bShouldBounce is false, in which case the projectile stops simulating on the first impact.
+	 * @see StopSimulating(), OnProjectileStop
 	 */
 	UPROPERTY(EditAnywhere, BlueprintReadWrite, Category=ProjectileBounces)
 	float BounceVelocityStopSimulatingThreshold;
@@ -174,9 +175,9 @@ class ENGINE_API UProjectileMovementComponent : public UMovementComponent
 	 */
 	virtual FVector ComputeVelocity(FVector InitialVelocity, float DeltaTime) const;
 
-	/** Clears the reference to UpdatedComponent, fires stop event, and stops ticking. */
+	/** Clears the reference to UpdatedComponent, fires stop event (OnProjectileStop), and stops ticking (if bAutoUpdateTickRegistration is true). */
 	UFUNCTION(BlueprintCallable, Category="Game|Components|ProjectileMovement")
-	void StopSimulating(const FHitResult& HitResult);
+	virtual void StopSimulating(const FHitResult& HitResult);
 
 	bool HasStoppedSimulation() { return UpdatedComponent == NULL; }
 
@@ -222,18 +223,56 @@ class ENGINE_API UProjectileMovementComponent : public UMovementComponent
 
 protected:
 
-	/** @return true if the simulation should stop. */
-	bool HandleHitWall(const FHitResult& Hit, float TimeTick, const FVector& MoveDelta);
+	// Enum indicating how simulation should proceed after HandleBlockingHit() is called.
+	enum class EHandleBlockingHitResult
+	{
+		Deflect,				/** Assume velocity has been deflected, and trigger HandleDeflection(). This is the default return value of HandleBlockingHit(). */
+		AdvanceNextSubstep,		/** Advance to the next simulation update. Typically used when additional slide/multi-bounce logic can be ignored,
+								    such as when an object that blocked the projectile is destroyed and movement should continue. */
+		Abort,					/** Abort all further simulation. Typically used when components have been invalidated or simulation should stop. */
+	};
 
-	/** Applies bounce logic (if enabled) to affect velocity upon impact, or stops the projectile if bounces are not enabled (or velocity is below threshold). Fires applicable events. */
+	/**
+	 * Handle blocking hit during simulation update. Checks that simulation remains valid after collision.
+	 * If simulating then calls HandleImpact(), and returns EHandleHitWallResult::Deflect by default to enable multi-bounce and sliding support through HandleDeflection().
+	 * If no longer simulating then returns EHandleHitWallResult::Abort, which aborts attempts at further simulation.
+	 *
+	 * @param  Hit						Blocking hit that occurred.
+	 * @param  TimeTick					Time delta of last move that resulted in the blocking hit.
+	 * @param  MoveDelta				Movement delta for the current sub-step.
+	 * @param  SubTickTimeRemaining		How much time to continue simulating in the current sub-step, which may change as a result of this function.
+	 *									Initial default value is: TimeTick * (1.f - Hit.Time)
+	 * @return Result indicating how simulation should proceed.
+	 * @see EHandleHitWallResult, HandleImpact()
+	 */
+	 virtual EHandleBlockingHitResult HandleBlockingHit(const FHitResult& Hit, float TimeTick, const FVector& MoveDelta, float& SubTickTimeRemaining);
+
+	/**
+	 * Applies bounce logic if enabled to affect velocity upon impact (using ComputeBounceResult()),
+	 * or stops the projectile if bounces are not enabled or velocity is below BounceVelocityStopSimulatingThreshold.
+	 * Triggers applicable events (OnProjectileBounce).
+	 */
 	virtual void HandleImpact(const FHitResult& Hit, float TimeSlice=0.f, const FVector& MoveDelta = FVector::ZeroVector) override;
 
+	/**
+	 * Handle a blocking hit after HandleBlockingHit() returns a result indicating that deflection occured.
+	 * Default implementation increments NumBounces, checks conditions that could indicate a slide, and calls HandleSliding() if necessary.
+	 * 
+	 * @param  Hit					Blocking hit that occurred. May be changed to indicate the last hit result of further movement.
+	 * @param  OldVelocity			Velocity at the start of the simulation update sub-step. Current Velocity may differ (as a result of a bounce).
+	 * @param  NumBounces			Number of bounces that have occurred thus far in the tick.
+	 * @param  SubTickTimeRemaining	Time remaining in the simulation sub-step. May be changed to indicate change to remaining time.
+	 * @return True if simulation of the projectile should continue, false otherwise.
+	 * @see HandleSliding()
+	 */
+	 virtual bool HandleDeflection(FHitResult& Hit, const FVector& OldVelocity, const uint32 NumBounces, float& SubTickTimeRemaining);
+
 	/**
 	 * Handle case where projectile is sliding along a surface.
 	 * Velocity will be parallel to the impact surface upon entry to this method.
 	 * 
-	 * @param  InitialHit Hit result of impact causing slide. May be modified by this function to reflect any subsequent movement.
-	 * @param  SubTickTimeRemaining Time remaining in the tick. This function may update this time with any reduction to the simulation time requested.
+	 * @param  InitialHit				Hit result of impact causing slide. May be modified by this function to reflect any subsequent movement.
+	 * @param  SubTickTimeRemaining		Time remaining in the tick. This function may update this time with any reduction to the simulation time requested.
 	 * @return True if simulation of the projectile should continue, false otherwise.
 	 */
 	virtual bool HandleSliding(FHitResult& Hit, float& SubTickTimeRemaining);
diff --git a/Engine/Source/Runtime/Engine/Private/Components/ProjectileMovementComponent.cpp b/Engine/Source/Runtime/Engine/Private/Components/ProjectileMovementComponent.cpp
--- a/Engine/Source/Runtime/Engine/Private/Components/ProjectileMovementComponent.cpp
+++ b/Engine/Source/Runtime/Engine/Private/Components/ProjectileMovementComponent.cpp
@@ -93,7 +93,7 @@ void UProjectileMovementComponent::TickComponent(float DeltaTime, enum ELevelTic
 	QUICK_SCOPE_CYCLE_COUNTER( STAT_ProjectileMovementComponent_TickComponent );
 
 	// skip if don't want component updated when not rendered or updated component can't move
-	if (ShouldSkipUpdate(DeltaTime))
+	if (HasStoppedSimulation() || ShouldSkipUpdate(DeltaTime))
 	{
 		return;
 	}
@@ -112,11 +112,11 @@ void UProjectileMovementComponent::TickComponent(float DeltaTime, enum ELevelTic
 	}
 
 	float RemainingTime	= DeltaTime;
-	int32 NumBounces = 0;
+	uint32 NumBounces = 0;
 	int32 Iterations = 0;
 	FHitResult Hit(1.f);
 	
-	while( RemainingTime >= MIN_TICK_TIME && (Iterations < MaxSimulationIterations) && !ActorOwner->IsPendingKill() && UpdatedComponent )
+	while( RemainingTime >= MIN_TICK_TIME && (Iterations < MaxSimulationIterations) && !ActorOwner->IsPendingKill() && !HasStoppedSimulation() )
 	{
 		Iterations++;
 
@@ -126,9 +126,9 @@ void UProjectileMovementComponent::TickComponent(float DeltaTime, enum ELevelTic
 
 		Hit.Time = 1.f;
 		const FVector OldVelocity = Velocity;
-		FVector MoveDelta = ComputeMoveDelta(OldVelocity, TimeTick);
+		const FVector MoveDelta = ComputeMoveDelta(OldVelocity, TimeTick);
 
-		const FRotator NewRotation = (bRotationFollowsVelocity && !OldVelocity.IsNearlyZero()) ? OldVelocity.Rotation() : ActorOwner->GetActorRotation();
+		const FRotator NewRotation = (bRotationFollowsVelocity && !OldVelocity.IsNearlyZero(0.01f)) ? OldVelocity.Rotation() : ActorOwner->GetActorRotation();
 
 		// Move the component
 		if (bShouldBounce)
@@ -144,7 +144,7 @@ void UProjectileMovementComponent::TickComponent(float DeltaTime, enum ELevelTic
 		}
 		
 		// If we hit a trigger that destroyed us, abort.
-		if( ActorOwner->IsPendingKill() || !UpdatedComponent )
+		if( ActorOwner->IsPendingKill() || HasStoppedSimulation() )
 		{
 			return;
 		}
@@ -170,63 +170,31 @@ void UProjectileMovementComponent::TickComponent(float DeltaTime, enum ELevelTic
 				Velocity = (Hit.Time > KINDA_SMALL_NUMBER) ? ComputeVelocity(OldVelocity, TimeTick * Hit.Time) : OldVelocity;
 			}
 
-			if (HandleHitWall(Hit, TimeTick, MoveDelta))
+			// Handle blocking hit
+			float SubTickTimeRemaining = TimeTick * (1.f - Hit.Time);
+			const EHandleBlockingHitResult HandleBlockingResult = HandleBlockingHit(Hit, TimeTick, MoveDelta, SubTickTimeRemaining);
+			if (HandleBlockingResult == EHandleBlockingHitResult::Abort || HasStoppedSimulation())
 			{
 				break;
 			}
-
-			NumBounces++;
-			float SubTickTimeRemaining = TimeTick * (1.f - Hit.Time);
-			FVector Normal = ConstrainNormalToPlane(Hit.Normal);
-
-			// Multiple hits within very short time period?
-			const bool bMultiHit = (PreviousHitTime < 1.f && Hit.Time <= KINDA_SMALL_NUMBER);
-
-			// if velocity still into wall (after HandleHitWall() had a chance to adjust), slide along wall
-			const float DotTolerance = 0.01f;
-			bIsSliding = (bMultiHit && FVector::Coincident(PreviousHitNormal, Normal)) ||
-						 ((Velocity.GetSafeNormal() | Normal) <= DotTolerance);
-			
-			if (bIsSliding)
+			else if (HandleBlockingResult == EHandleBlockingHitResult::Deflect)
 			{
-				if (bMultiHit && (PreviousHitNormal | Normal) <= 0.f)
-				{
-					//90 degree or less corner, so use cross product for direction
-					FVector NewDir = (Normal ^ PreviousHitNormal);
-					NewDir = NewDir.GetSafeNormal();
-					Velocity = Velocity.ProjectOnToNormal(NewDir);
-					if ((OldVelocity | Velocity) < 0.f)
-					{
-						Velocity *= -1.f;
-					}
-					Velocity = ConstrainDirectionToPlane(Velocity);
-				}
-				else 
-				{
-					//adjust to move along new wall
-					Velocity = ComputeSlideVector(Velocity, 1.f, Normal, Hit);
-				}
-
-				// Check min velocity.
-				if (Velocity.SizeSquared() < FMath::Square(BounceVelocityStopSimulatingThreshold))
-				{
-					StopSimulating(Hit);
-					break;
-				}
-				
-				// Velocity is now parallel to the impact surface.
-				if (SubTickTimeRemaining > KINDA_SMALL_NUMBER)
-				{
-					if (!HandleSliding(Hit, SubTickTimeRemaining))
-					{
-						break;
-					}
-					Normal = ConstrainNormalToPlane(Hit.Normal);
-				}
+				NumBounces++;
+				HandleDeflection(Hit, OldVelocity, NumBounces, SubTickTimeRemaining);
+				PreviousHitTime = Hit.Time;
+				PreviousHitNormal = ConstrainNormalToPlane(Hit.Normal);
 			}
-
-			PreviousHitTime = Hit.Time;
-			PreviousHitNormal = Normal;
+			else if (HandleBlockingResult == EHandleBlockingHitResult::AdvanceNextSubstep)
+			{
+				// Reset deflection logic to ignore this hit
+				PreviousHitTime = 1.f;
+			}
+			else
+			{
+				// Unhandled EHandleBlockingHitResult
+				checkNoEntry();
+			}
+			
 			
 			// A few initial bounces should add more time and iterations to complete most of the simulation.
 			if (NumBounces <= 2 && SubTickTimeRemaining >= MIN_TICK_TIME)
@@ -241,6 +209,59 @@ void UProjectileMovementComponent::TickComponent(float DeltaTime, enum ELevelTic
 }
 
 
+bool UProjectileMovementComponent::HandleDeflection(FHitResult& Hit, const FVector& OldVelocity, const uint32 NumBounces, float& SubTickTimeRemaining)
+{
+	const FVector Normal = ConstrainNormalToPlane(Hit.Normal);
+
+	// Multiple hits within very short time period?
+	const bool bMultiHit = (PreviousHitTime < 1.f && Hit.Time <= KINDA_SMALL_NUMBER);
+
+	// if velocity still into wall (after HandleBlockingHit() had a chance to adjust), slide along wall
+	const float DotTolerance = 0.01f;
+	bIsSliding = (bMultiHit && FVector::Coincident(PreviousHitNormal, Normal)) ||
+					((Velocity.GetSafeNormal() | Normal) <= DotTolerance);
+
+	if (bIsSliding)
+	{
+		if (bMultiHit && (PreviousHitNormal | Normal) <= 0.f)
+		{
+			//90 degree or less corner, so use cross product for direction
+			FVector NewDir = (Normal ^ PreviousHitNormal);
+			NewDir = NewDir.GetSafeNormal();
+			Velocity = Velocity.ProjectOnToNormal(NewDir);
+			if ((OldVelocity | Velocity) < 0.f)
+			{
+				Velocity *= -1.f;
+			}
+			Velocity = ConstrainDirectionToPlane(Velocity);
+		}
+		else
+		{
+			//adjust to move along new wall
+			Velocity = ComputeSlideVector(Velocity, 1.f, Normal, Hit);
+		}
+
+		// Check min velocity.
+		if (Velocity.SizeSquared() < FMath::Square(BounceVelocityStopSimulatingThreshold))
+		{
+			StopSimulating(Hit);
+			return false;
+		}
+
+		// Velocity is now parallel to the impact surface.
+		if (SubTickTimeRemaining > KINDA_SMALL_NUMBER)
+		{
+			if (!HandleSliding(Hit, SubTickTimeRemaining))
+			{
+				return false;
+			}
+		}
+	}
+
+	return true;
+}
+
+
 bool UProjectileMovementComponent::HandleSliding(FHitResult& Hit, float& SubTickTimeRemaining)
 {
 	FHitResult InitialHit(Hit);
@@ -250,10 +271,19 @@ bool UProjectileMovementComponent::HandleSliding(FHitResult& Hit, float& SubTick
 	// Perform the move now, before adding gravity/accel again, so we don't just keep hitting the surface.
 	SafeMoveUpdatedComponent(Velocity * SubTickTimeRemaining, UpdatedComponent->GetComponentRotation(), true, Hit);
 
+	if (HasStoppedSimulation())
+	{
+		return false;
+	}
+
 	// A second hit can deflect the velocity (through the normal bounce code), for the next iteration.
 	if (Hit.bBlockingHit)
 	{
-		if (HandleHitWall(Hit, SubTickTimeRemaining, Velocity * SubTickTimeRemaining))
+		const float TimeTick = SubTickTimeRemaining;
+		SubTickTimeRemaining = TimeTick * (1.f - Hit.Time);
+		
+		if (HandleBlockingHit(Hit, TimeTick, Velocity * TimeTick, SubTickTimeRemaining) == EHandleBlockingHitResult::Abort ||
+			HasStoppedSimulation())
 		{
 			return false;
 		}
@@ -285,9 +315,10 @@ bool UProjectileMovementComponent::HandleSliding(FHitResult& Hit, float& SubTick
 			StopSimulating(InitialHit);
 			return false;
 		}
+
+		SubTickTimeRemaining = 0.f;
 	}
 
-	SubTickTimeRemaining = SubTickTimeRemaining * (1.f - Hit.Time);
 	return true;
 }
 
@@ -379,22 +410,23 @@ void UProjectileMovementComponent::StopSimulating(const FHitResult& HitResult)
 }
 
 
-bool UProjectileMovementComponent::HandleHitWall(const FHitResult& Hit, float TimeTick, const FVector& MoveDelta)
+UProjectileMovementComponent::EHandleBlockingHitResult UProjectileMovementComponent::HandleBlockingHit(const FHitResult& Hit, float TimeTick, const FVector& MoveDelta, float& SubTickTimeRemaining)
 {
 	AActor* ActorOwner = UpdatedComponent ? UpdatedComponent->GetOwner() : NULL;
-	if ( !CheckStillInWorld() || !ActorOwner || ActorOwner->IsPendingKill() )
+	if (!CheckStillInWorld() || !ActorOwner || ActorOwner->IsPendingKill())
 	{
-		return true;
+		return EHandleBlockingHitResult::Abort;
 	}
 	
 	HandleImpact(Hit, TimeTick, MoveDelta);
 	
-	if( ActorOwner->IsPendingKill() || !UpdatedComponent )
+	if (ActorOwner->IsPendingKill() || HasStoppedSimulation())
 	{
-		return true;
+		return EHandleBlockingHitResult::Abort;
 	}
 
-	return false;
+	SubTickTimeRemaining = TimeTick * (1.f - Hit.Time);
+	return EHandleBlockingHitResult::Deflect;
 }
  
 FVector UProjectileMovementComponent::ComputeBounceResult(const FHitResult& Hit, float TimeSlice, const FVector& MoveDelta)
