Engine: Allow Ctrl+B and Ctrl+E to find/edit the associated asset for Niagara actors
- Added ANiagaraActor::GetReferencedContentObjects and UNiagaraComponent::AdditionalStatObject implementations
- Fixed some coding standards violations

diff --git a/Engine/Source/Runtime/Engine/Classes/Components/NiagaraComponent.h b/Engine/Source/Runtime/Engine/Classes/Components/NiagaraComponent.h
--- a/Engine/Source/Runtime/Engine/Classes/Components/NiagaraComponent.h
+++ b/Engine/Source/Runtime/Engine/Classes/Components/NiagaraComponent.h
@@ -38,13 +38,13 @@ public:
 		CurrentBuffer = 0;
 	}
 
-	const FVector4 *GetAttributeData(const FNiagaraVariableInfo& AttrID) const
+	const FVector4* GetAttributeData(const FNiagaraVariableInfo& AttrID) const
 	{
 		const uint32* Offset = AttrMap.Find(AttrID);
 		return Offset ? ParticleBuffers[CurrentBuffer].GetData() + (*Offset * ParticleAllocation) : NULL;
 	}
 
-	FVector4 *GetAttributeDataWrite(const FNiagaraVariableInfo& AttrID)
+	FVector4* GetAttributeDataWrite(const FNiagaraVariableInfo& AttrID)
 	{
 		const uint32* Offset = AttrMap.Find(AttrID);
 		return Offset ? ParticleBuffers[CurrentBuffer].GetData() + (*Offset * ParticleAllocation) : NULL;
@@ -73,8 +73,8 @@ public:
 	void SetNumParticles(uint32 Num)	{ NumParticles = Num;  }
 	void SwapBuffers()					{ CurrentBuffer ^= 0x1; }
 
-	FVector4 *GetCurrentBuffer()		{ return ParticleBuffers[CurrentBuffer].GetData(); }
-	FVector4 *GetPreviousBuffer()		{ return ParticleBuffers[CurrentBuffer^0x1].GetData(); }
+	FVector4* GetCurrentBuffer()		{ return ParticleBuffers[CurrentBuffer].GetData(); }
+	FVector4* GetPreviousBuffer()		{ return ParticleBuffers[CurrentBuffer^0x1].GetData(); }
 
 	int GetBytesUsed()	{ return (ParticleBuffers[0].Num() + ParticleBuffers[1].Num()) * 16 + AttrMap.Num() * 4; }
 private:
@@ -96,9 +96,9 @@ class ENGINE_API UNiagaraComponent : public UPrimitiveComponent
 	GENERATED_UCLASS_BODY()
 private:
 	UPROPERTY()
-	//class UNiagaraEffect *Effect;
-	class UNiagaraEffect *Asset;
-	class FNiagaraEffectInstance *EffectInstance;
+	class UNiagaraEffect* Asset;
+
+	class FNiagaraEffectInstance* EffectInstance;
 
 	// Begin UActorComponent interface.
 protected:
@@ -106,7 +106,8 @@ protected:
 	virtual void OnUnregister()  override;
 	virtual void SendRenderDynamicData_Concurrent() override;
 public:
-	virtual void TickComponent(float DeltaTime, enum ELevelTick TickType, FActorComponentTickFunction *ThisTickFunction) override;
+	virtual void TickComponent(float DeltaTime, enum ELevelTick TickType, FActorComponentTickFunction* ThisTickFunction) override;
+	virtual const UObject* AdditionalStatObject() const override;
 	// End UActorComponent interface.
 
 	// Begin UPrimitiveComponent Interface
@@ -115,11 +116,11 @@ public:
 	virtual FPrimitiveSceneProxy* CreateSceneProxy() override;
 	// End UPrimitiveComponent Interface
 
-	void SetAsset(UNiagaraEffect *InAsset);
-	UNiagaraEffect *GetAsset() const { return Asset; }
+	void SetAsset(UNiagaraEffect* InAsset);
+	UNiagaraEffect* GetAsset() const { return Asset; }
 
-	FNiagaraEffectInstance *GetEffectInstance()	const { return EffectInstance; }
-	void SetEffectInstance(FNiagaraEffectInstance *InInstance)	{ EffectInstance = InInstance; }
+	FNiagaraEffectInstance* GetEffectInstance()	const { return EffectInstance; }
+	void SetEffectInstance(FNiagaraEffectInstance* InInstance)	{ EffectInstance = InInstance; }
 
 	// Begin UObject interface.
 #if WITH_EDITOR
@@ -148,9 +149,9 @@ public:
 
 	/** Called on render thread to assign new dynamic data */
 	void SetDynamicData_RenderThread(struct FNiagaraDynamicDataBase* NewDynamicData);
-	TArray<class NiagaraEffectRenderer*> &GetEffectRenderers() { return EffectRenderers; }
-	void AddEffectRenderer(NiagaraEffectRenderer *Renderer)	{ EffectRenderers.Add(Renderer); }
-	ENGINE_API void UpdateEffectRenderers(FNiagaraEffectInstance *InEffect);
+	TArray<class NiagaraEffectRenderer*>& GetEffectRenderers() { return EffectRenderers; }
+	void AddEffectRenderer(NiagaraEffectRenderer* Renderer)	{ EffectRenderers.Add(Renderer); }
+	ENGINE_API void UpdateEffectRenderers(FNiagaraEffectInstance* InEffect);
 
 private:
 	void ReleaseRenderThreadResources();
@@ -177,6 +178,6 @@ private:
 
 
 private:
-	//class NiagaraEffectRenderer *EffectRenderer;
-	TArray<class NiagaraEffectRenderer *>EffectRenderers;
+	//class NiagaraEffectRenderer* EffectRenderer;
+	TArray<class NiagaraEffectRenderer*>EffectRenderers;
 };
diff --git a/Engine/Source/Runtime/Engine/Classes/Engine/NiagaraActor.h b/Engine/Source/Runtime/Engine/Classes/Engine/NiagaraActor.h
--- a/Engine/Source/Runtime/Engine/Classes/Engine/NiagaraActor.h
+++ b/Engine/Source/Runtime/Engine/Classes/Engine/NiagaraActor.h
@@ -36,4 +36,11 @@ public:
 	/** Returns ArrowComponent subobject **/
 	class UArrowComponent* GetArrowComponent() const;
 #endif
+
+#if WITH_EDITOR
+	// AActor interface
+	virtual bool GetReferencedContentObjects(TArray<UObject*>& Objects) const override;
+	// End of AActor interface
+#endif // WITH_EDITOR
+
 };
diff --git a/Engine/Source/Runtime/Engine/Private/Components/NiagaraComponent.cpp b/Engine/Source/Runtime/Engine/Private/Components/NiagaraComponent.cpp
--- a/Engine/Source/Runtime/Engine/Private/Components/NiagaraComponent.cpp
+++ b/Engine/Source/Runtime/Engine/Private/Components/NiagaraComponent.cpp
@@ -24,7 +24,7 @@ FNiagaraSceneProxy::FNiagaraSceneProxy(const UNiagaraComponent* InComponent)
 	UpdateEffectRenderers(InComponent->GetEffectInstance());
 }
 
-void FNiagaraSceneProxy::UpdateEffectRenderers(FNiagaraEffectInstance *InEffect)
+void FNiagaraSceneProxy::UpdateEffectRenderers(FNiagaraEffectInstance* InEffect)
 {
 	EffectRenderers.Empty();
 	if (InEffect)
@@ -44,7 +44,7 @@ FNiagaraSceneProxy::~FNiagaraSceneProxy()
 /** Called on render thread to assign new dynamic data */
 void FNiagaraSceneProxy::SetDynamicData_RenderThread(FNiagaraDynamicDataBase* NewDynamicData)
 {
-	for (NiagaraEffectRenderer *Renderer : EffectRenderers)
+	for (NiagaraEffectRenderer* Renderer : EffectRenderers)
 	{
 		if (Renderer)
 		{
@@ -57,7 +57,7 @@ void FNiagaraSceneProxy::SetDynamicData_RenderThread(FNiagaraDynamicDataBase* Ne
 
 void FNiagaraSceneProxy::ReleaseRenderThreadResources()
 {
-	for (NiagaraEffectRenderer *Renderer : EffectRenderers)
+	for (NiagaraEffectRenderer* Renderer : EffectRenderers)
 	{
 		if (Renderer)
 		{
@@ -70,7 +70,7 @@ void FNiagaraSceneProxy::ReleaseRenderThreadResources()
 // FPrimitiveSceneProxy interface.
 void FNiagaraSceneProxy::CreateRenderThreadResources()
 {
-	for (NiagaraEffectRenderer *Renderer : EffectRenderers)
+	for (NiagaraEffectRenderer* Renderer : EffectRenderers)
 	{
 		if (Renderer)
 		{
@@ -95,7 +95,7 @@ FPrimitiveViewRelevance FNiagaraSceneProxy::GetViewRelevance(const FSceneView* V
 	FPrimitiveViewRelevance Relevance;
 	Relevance.bDynamicRelevance = true;
 
-	for (NiagaraEffectRenderer *Renderer : EffectRenderers)
+	for (NiagaraEffectRenderer* Renderer : EffectRenderers)
 	{
 		if (Renderer)
 		{
@@ -114,7 +114,7 @@ uint32 FNiagaraSceneProxy::GetMemoryFootprint() const
 uint32 FNiagaraSceneProxy::GetAllocatedSize() const
 { 
 	uint32 DynamicDataSize = 0;
-	for (NiagaraEffectRenderer *Renderer : EffectRenderers)
+	for (NiagaraEffectRenderer* Renderer : EffectRenderers)
 	{
 		if (Renderer)
 		{
@@ -127,7 +127,7 @@ uint32 FNiagaraSceneProxy::GetAllocatedSize() const
 
 void FNiagaraSceneProxy::GetDynamicMeshElements(const TArray<const FSceneView*>& Views, const FSceneViewFamily& ViewFamily, uint32 VisibilityMap, FMeshElementCollector& Collector) const
 {
-	for (NiagaraEffectRenderer *Renderer : EffectRenderers)
+	for (NiagaraEffectRenderer* Renderer : EffectRenderers)
 	{
 		if (Renderer)
 		{
@@ -146,7 +146,7 @@ UNiagaraComponent::UNiagaraComponent(const FObjectInitializer& ObjectInitializer
 }
 
 
-void UNiagaraComponent::TickComponent(float DeltaSeconds, enum ELevelTick TickType, FActorComponentTickFunction *ThisTickFunction)
+void UNiagaraComponent::TickComponent(float DeltaSeconds, enum ELevelTick TickType, FActorComponentTickFunction* ThisTickFunction)
 {
 //	EmitterAge += DeltaSeconds;
 
@@ -175,6 +175,11 @@ void UNiagaraComponent::TickComponent(float DeltaSeconds, enum ELevelTick TickTy
 	MarkRenderDynamicDataDirty();
 }
 
+const UObject* UNiagaraComponent::AdditionalStatObject() const
+{
+	return Asset;
+}
+
 void UNiagaraComponent::OnRegister()
 {
 	Super::OnRegister();
@@ -214,11 +219,11 @@ void UNiagaraComponent::SendRenderDynamicData_Concurrent()
 {
 	if (EffectInstance && SceneProxy)
 	{
-		FNiagaraSceneProxy *NiagaraProxy = static_cast<FNiagaraSceneProxy*>(SceneProxy);
+		FNiagaraSceneProxy* NiagaraProxy = static_cast<FNiagaraSceneProxy*>(SceneProxy);
 		for (int32 i = 0; i < EffectInstance->GetEmitters().Num(); i++)
 		{
-			TSharedPtr<FNiagaraSimulation>Emitter = EffectInstance->GetEmitters()[i];
-			NiagaraEffectRenderer *Renderer = Emitter->GetEffectRenderer();
+			TSharedPtr<FNiagaraSimulation> Emitter = EffectInstance->GetEmitters()[i];
+			NiagaraEffectRenderer* Renderer = Emitter->GetEffectRenderer();
 			if (Renderer)
 			{
 				FNiagaraDynamicDataBase* DynamicData = Renderer->GenerateVertexData(Emitter->GetData());
@@ -264,7 +269,7 @@ FBoxSphereBounds UNiagaraComponent::CalcBounds(const FTransform& LocalToWorld) c
 
 FPrimitiveSceneProxy* UNiagaraComponent::CreateSceneProxy()
 {
-	FNiagaraSceneProxy *Proxy = new FNiagaraSceneProxy(this);
+	FNiagaraSceneProxy* Proxy = new FNiagaraSceneProxy(this);
 	Proxy->UpdateEffectRenderers(EffectInstance);
 	return Proxy;
 }
@@ -280,7 +285,7 @@ void UNiagaraComponent::PostEditChangeProperty(FPropertyChangedEvent& PropertyCh
 
 
 
-void UNiagaraComponent::SetAsset(UNiagaraEffect *InAsset)
+void UNiagaraComponent::SetAsset(UNiagaraEffect* InAsset)
 {
 	Asset = InAsset;
 
diff --git a/Engine/Source/Runtime/Engine/Private/NiagaraActor.cpp b/Engine/Source/Runtime/Engine/Private/NiagaraActor.cpp
--- a/Engine/Source/Runtime/Engine/Private/NiagaraActor.cpp
+++ b/Engine/Source/Runtime/Engine/Private/NiagaraActor.cpp
@@ -4,6 +4,7 @@
 #include "EnginePrivate.h"
 #include "Engine/NiagaraActor.h"
 #include "Components/NiagaraComponent.h"
+#include "Engine/NiagaraEffect.h"
 
 ANiagaraActor::ANiagaraActor(const FObjectInitializer& ObjectInitializer)
 	: Super(ObjectInitializer)
@@ -68,3 +69,18 @@ UBillboardComponent* ANiagaraActor::GetSpriteComponent() const { return SpriteCo
 /** Returns ArrowComponent subobject **/
 UArrowComponent* ANiagaraActor::GetArrowComponent() const { return ArrowComponent; }
 #endif
+
+#if WITH_EDITOR
+bool ANiagaraActor::GetReferencedContentObjects(TArray<UObject*>& Objects) const
+{
+	Super::GetReferencedContentObjects(Objects);
+
+	if (UNiagaraEffect* Effect = NiagaraComponent->GetAsset())
+	{
+		Objects.Add(Effect);
+	}
+
+	return true;
+}
+#endif // WITH_EDITOR
+
