Engine: Improve stat display of the associated asset for audio components
- Implemented UAudioComponent::AdditionalStatObject

diff --git a/Engine/Source/Runtime/Engine/Classes/Components/AudioComponent.h b/Engine/Source/Runtime/Engine/Classes/Components/AudioComponent.h
--- a/Engine/Source/Runtime/Engine/Classes/Components/AudioComponent.h
+++ b/Engine/Source/Runtime/Engine/Classes/Components/AudioComponent.h
@@ -278,6 +278,7 @@ public:
 	virtual void OnRegister() override;
 #endif
 	virtual void OnUnregister() override;
+	virtual const UObject* AdditionalStatObject() const override;
 	// End ActorComponent interface.
 
 	/** Returns a pointer to the attenuation settings to be used (if any) for this audio component dependent on the SoundAttenuation asset or overrides set. */
diff --git a/Engine/Source/Runtime/Engine/Private/Components/AudioComponent.cpp b/Engine/Source/Runtime/Engine/Private/Components/AudioComponent.cpp
--- a/Engine/Source/Runtime/Engine/Private/Components/AudioComponent.cpp
+++ b/Engine/Source/Runtime/Engine/Private/Components/AudioComponent.cpp
@@ -69,6 +69,11 @@ void UAudioComponent::OnUnregister()
 	}
 }
 
+const UObject* UAudioComponent::AdditionalStatObject() const
+{
+	return Sound;
+}
+
 void UAudioComponent::SetSound( USoundBase* NewSound )
 {
 	const bool bPlay = IsPlaying();
