Engine: Allow Ctrl+B and Ctrl+E to find/edit the associated decal material for decal actors
- Added ADecalActor::GetReferencedContentObjects and UDecalComponent::AdditionalStatObject implementations
[UE-11049]

diff --git a/Engine/Source/Runtime/Engine/Classes/Components/DecalComponent.h b/Engine/Source/Runtime/Engine/Classes/Components/DecalComponent.h
--- a/Engine/Source/Runtime/Engine/Classes/Components/DecalComponent.h
+++ b/Engine/Source/Runtime/Engine/Classes/Components/DecalComponent.h
@@ -16,7 +16,6 @@ class ENGINE_API UDecalComponent : public USceneComponent
 	GENERATED_UCLASS_BODY()
 
 	/** Decal material. */
-	//editable(Decal) private{private} const MaterialInterface	DecalMaterial;
 	UPROPERTY(EditAnywhere, BlueprintReadOnly, Category=Decal)
 	class UMaterialInterface* DecalMaterial;
 
@@ -93,6 +92,7 @@ public:
 	virtual void CreateRenderState_Concurrent() override;
 	virtual void DestroyRenderState_Concurrent() override;
 	virtual void SendRenderTransform_Concurrent() override;
+	virtual const UObject* AdditionalStatObject() const override;
 	// End UActorComponent Interface
 
 	// Begin USceneComponent Interface
diff --git a/Engine/Source/Runtime/Engine/Classes/Engine/DecalActor.h b/Engine/Source/Runtime/Engine/Classes/Engine/DecalActor.h
--- a/Engine/Source/Runtime/Engine/Classes/Engine/DecalActor.h
+++ b/Engine/Source/Runtime/Engine/Classes/Engine/DecalActor.h
@@ -60,6 +60,7 @@ public:
 
 	// Begin AActor interface.
 	virtual void EditorApplyScale(const FVector& DeltaScale, const FVector* PivotLocation, bool bAltDown, bool bShiftDown, bool bCtrlDown) override;
+	virtual bool GetReferencedContentObjects(TArray<UObject*>& Objects) const override;
 	// End AActor interface.
 #endif // WITH_EDITOR
 
diff --git a/Engine/Source/Runtime/Engine/Private/Components/DecalComponent.cpp b/Engine/Source/Runtime/Engine/Private/Components/DecalComponent.cpp
--- a/Engine/Source/Runtime/Engine/Private/Components/DecalComponent.cpp
+++ b/Engine/Source/Runtime/Engine/Private/Components/DecalComponent.cpp
@@ -129,6 +129,11 @@ void UDecalComponent::SendRenderTransform_Concurrent()
 	Super::SendRenderTransform_Concurrent();
 }
 
+const UObject* UDecalComponent::AdditionalStatObject() const
+{
+	return DecalMaterial;
+}
+
 void UDecalComponent::DestroyRenderState_Concurrent()
 {
 	Super::DestroyRenderState_Concurrent();
diff --git a/Engine/Source/Runtime/Engine/Private/DecalActor.cpp b/Engine/Source/Runtime/Engine/Private/DecalActor.cpp
--- a/Engine/Source/Runtime/Engine/Private/DecalActor.cpp
+++ b/Engine/Source/Runtime/Engine/Private/DecalActor.cpp
@@ -115,6 +115,19 @@ void ADecalActor::EditorApplyScale(const FVector& DeltaScale, const FVector* Piv
 
 	Super::EditorApplyScale(ModifiedScale, PivotLocation, bAltDown, bShiftDown, bCtrlDown);
 }
+
+bool ADecalActor::GetReferencedContentObjects(TArray<UObject*>& Objects) const
+{
+	Super::GetReferencedContentObjects(Objects);
+
+	if (Decal->DecalMaterial != nullptr)
+	{
+		Objects.Add(Decal->DecalMaterial);
+	}
+
+	return true;
+}
+
 #endif // WITH_EDITOR
 
 void ADecalActor::SetDecalMaterial(class UMaterialInterface* NewDecalMaterial)
