Made navigation paths' update requests store path's weak pointer rather than a shared ref #UE4

this is part 2 of fixing paths being updated after AI's death

diff --git a/Engine/Source/Runtime/Engine/Classes/AI/Navigation/NavigationData.h b/Engine/Source/Runtime/Engine/Classes/AI/Navigation/NavigationData.h
--- a/Engine/Source/Runtime/Engine/Classes/AI/Navigation/NavigationData.h
+++ b/Engine/Source/Runtime/Engine/Classes/AI/Navigation/NavigationData.h
@@ -33,7 +33,7 @@ struct FSupportedAreaData
 
 struct FNavPathRecalculationRequest
 {
-	FNavPathSharedRef Path;
+	FNavPathWeakPtr Path;
 	ENavPathUpdateType::Type Reason;
 
 	FNavPathRecalculationRequest(const FNavPathSharedPtr& InPath, ENavPathUpdateType::Type InReason)
diff --git a/Engine/Source/Runtime/Engine/Private/AI/Navigation/NavigationData.cpp b/Engine/Source/Runtime/Engine/Private/AI/Navigation/NavigationData.cpp
--- a/Engine/Source/Runtime/Engine/Private/AI/Navigation/NavigationData.cpp
+++ b/Engine/Source/Runtime/Engine/Private/AI/Navigation/NavigationData.cpp
@@ -246,7 +246,13 @@ void ANavigationData::TickActor(float DeltaTime, enum ELevelTick TickType, FActo
 		for (auto RecalcRequest : RepathRequests)
 		{
 			// check if it can be updated right now
-			const UObject* PathQuerier = RecalcRequest.Path->GetQuerier();
+			FNavPathSharedPtr PinnedPath = RecalcRequest.Path.Pin();
+			if (PinnedPath.IsValid() == false)
+			{
+				continue;
+			}
+
+			const UObject* PathQuerier = PinnedPath->GetQuerier();
 			const INavAgentInterface* PathNavAgent = Cast<const INavAgentInterface>(PathQuerier);
 			if (PathNavAgent && PathNavAgent->ShouldPostponePathUpdates())
 			{
@@ -254,26 +260,26 @@ void ANavigationData::TickActor(float DeltaTime, enum ELevelTick TickType, FActo
 				continue;
 			}
 
-			FPathFindingQuery Query(RecalcRequest.Path);
+			FPathFindingQuery Query(PinnedPath.ToSharedRef());
 			// @todo consider supplying NavAgentPropertied from path's querier
-			const FPathFindingResult Result = FindPath(FNavAgentProperties(), Query.SetPathInstanceToUpdate(RecalcRequest.Path));
+			const FPathFindingResult Result = FindPath(FNavAgentProperties(), Query.SetPathInstanceToUpdate(PinnedPath));
 
 			// update time stamp to give observers any means of telling if it has changed
-			RecalcRequest.Path->SetTimeStamp(TimeStamp);
+			PinnedPath->SetTimeStamp(TimeStamp);
 
 			// partial paths are still valid and can change to full path when moving goal gets back on navmesh
 			if (Result.IsSuccessful() || Result.IsPartial())
 			{
-				RecalcRequest.Path->UpdateLastRepathGoalLocation();
-				RecalcRequest.Path->DoneUpdating(RecalcRequest.Reason);
+				PinnedPath->UpdateLastRepathGoalLocation();
+				PinnedPath->DoneUpdating(RecalcRequest.Reason);
 				if (RecalcRequest.Reason == ENavPathUpdateType::NavigationChanged)
 				{
-					RegisterActivePath(RecalcRequest.Path);
+					RegisterActivePath(PinnedPath);
 				}
 			}
 			else
 			{
-				RecalcRequest.Path->RePathFailed();
+				PinnedPath->RePathFailed();
 			}
 		}
 
