Made navigation paths "derivable" by expanding FNavPathType with a "parent type" #UE4

All one needs to do is to declare his path type, like so:

*.h:

struct FMyNavigationPath : public FNavMeshPath
{
    typedef FNavMeshPath Super;
    static const FNavPathType Type;
    FMyNavigationPath();
};

*.cpp:

const FNavPathType FMyNavigationPath::Type(&FMyNavigationPath::Super::Type);
FMyNavigationPath::FMyNavigationPath()
    : Super()
{
    PathType = FMyNavigationPath::Type;
}

Of course instead of FNavMeshPath any other path type can be used.

diff --git a/Engine/Source/Runtime/Engine/Classes/AI/Navigation/NavigationTypes.h b/Engine/Source/Runtime/Engine/Classes/AI/Navigation/NavigationTypes.h
--- a/Engine/Source/Runtime/Engine/Classes/AI/Navigation/NavigationTypes.h
+++ b/Engine/Source/Runtime/Engine/Classes/AI/Navigation/NavigationTypes.h
@@ -250,17 +250,19 @@ struct FNavPathPoint : public FNavLocation
 /** path type data */
 struct ENGINE_API FNavPathType
 {
-	FNavPathType() : Id(++NextUniqueId) {}
-	FNavPathType(const FNavPathType& Src) : Id(Src.Id) {}
+	explicit FNavPathType(const FNavPathType* Parent = nullptr) : Id(++NextUniqueId), ParentType(Parent) {}
+	FNavPathType(const FNavPathType& Src) : Id(Src.Id), ParentType(Src.ParentType) {}
 	
 	bool operator==(const FNavPathType& Other) const
 	{
-		return Id == Other.Id;
+		return Id == Other.Id || (ParentType != nullptr && *ParentType == Other)
+			|| (Other.ParentType != nullptr && *Other.ParentType == *this);
 	}
 
 private:
 	static uint32 NextUniqueId;
 	uint32 Id;
+	const FNavPathType* ParentType;
 };
 
 UENUM()
