Back out changelist 2469963 with Visual Logger circular buffer. It's too slow to use unfortunately (integrated CL #2489422 from Dev),
Removed usage of GWorld from Visual Logger (integrated CL #2490390 from Dev)

diff --git a/Engine/Source/Developer/LogVisualizer/Private/SVisualLogger.cpp b/Engine/Source/Developer/LogVisualizer/Private/SVisualLogger.cpp
--- a/Engine/Source/Developer/LogVisualizer/Private/SVisualLogger.cpp
+++ b/Engine/Source/Developer/LogVisualizer/Private/SVisualLogger.cpp
@@ -427,16 +427,6 @@ bool SVisualLogger::HandleStartRecordingCommandCanExecute() const
 
 void SVisualLogger::HandleStartRecordingCommandExecute()
 {
-	if (FVisualLogger::Get().IsRecording() == false && FVisualLogger::Get().IsUsingCircularBuffer())
-	{
-		TArray<FVisualLogDevice::FVisualLogEntryItem> OutLogs;
-		FVisualLogger::Get().GetCircularBuffer()->GetRecordedLogs(OutLogs);
-		for (auto& CurrentLog : OutLogs)
-		{
-			OnNewLogEntry(CurrentLog);
-		}
-		FVisualLogger::Get().GetCircularBuffer()->Cleanup(true);
-	}
 	FVisualLogger::Get().SetIsRecording(true);
 }
 
diff --git a/Engine/Source/Runtime/Engine/Classes/VisualLogger/VisualLoggerBinaryFileDevice.h b/Engine/Source/Runtime/Engine/Classes/VisualLogger/VisualLoggerBinaryFileDevice.h
--- a/Engine/Source/Runtime/Engine/Classes/VisualLogger/VisualLoggerBinaryFileDevice.h
+++ b/Engine/Source/Runtime/Engine/Classes/VisualLogger/VisualLoggerBinaryFileDevice.h
@@ -4,6 +4,8 @@
 
 #if ENABLE_VISUAL_LOG
 
+#define VISLOG_FILENAME_EXT TEXT("bvlog")
+
 class FVisualLoggerBinaryFileDevice : public FVisualLogDevice
 {
 public:
diff --git a/Engine/Source/Runtime/Engine/Classes/VisualLogger/VisualLoggerCircularBufferDevice.h b/Engine/Source/Runtime/Engine/Classes/VisualLogger/VisualLoggerCircularBufferDevice.h
deleted file mode 100644
diff --git a/Engine/Source/Runtime/Engine/Private/VisualLogger/VisualLogger.cpp b/Engine/Source/Runtime/Engine/Private/VisualLogger/VisualLogger.cpp
--- a/Engine/Source/Runtime/Engine/Private/VisualLogger/VisualLogger.cpp
+++ b/Engine/Source/Runtime/Engine/Private/VisualLogger/VisualLogger.cpp
@@ -4,7 +4,6 @@
 #include "VisualLogger/VisualLogger.h"
 #include "VisualLogger/VisualLogger.h"
 #include "VisualLogger/VisualLoggerBinaryFileDevice.h"
-#include "VisualLogger/VisualLoggerCircularBufferDevice.h"
 #if WITH_EDITOR
 #	include "Editor/UnrealEd/Public/EditorComponents.h"
 #	include "Editor/UnrealEd/Public/EditorReimportHandler.h"
@@ -27,10 +26,6 @@ FVisualLogger::FVisualLogger()
 {
 	BlockAllCategories(false);
 	AddDevice(&FVisualLoggerBinaryFileDevice::Get());
-	if (FVisualLoggerCircularBufferDevice::Get().GetCircularBufferSize() > 0)
-	{
-		AddDevice(&FVisualLoggerCircularBufferDevice::Get());
-	}
 	SetIsRecording(GEngine ? !!GEngine->bEnableVisualLogRecordingOnStart : false);
 	SetIsRecordingOnServer(false);
 
@@ -41,25 +36,6 @@ FVisualLogger::FVisualLogger()
 	}
 }
 
-bool FVisualLogger::IsUsingCircularBuffer()
-{
-	return GetDevices().Find(&FVisualLoggerCircularBufferDevice::Get()) != INDEX_NONE;
-}
-
-FVisualLogDevice* FVisualLogger::GetCircularBuffer()
-{
-	return FVisualLoggerCircularBufferDevice::Get().GetCircularBufferSize() > 0 ? & FVisualLoggerCircularBufferDevice::Get() : NULL;
-}
-
-bool FVisualLogger::DumpCircularBuffer()
-{
-	if (FVisualLoggerCircularBufferDevice::Get().GetCircularBufferSize() > 0)
-	{
-		FVisualLoggerCircularBufferDevice::Get().DumpBuffer();
-	}
-	return true;
-}
-
 UWorld* FVisualLogger::GetWorld(const class UObject* Object)
 {
 	UWorld* World = Object ? GEngine->GetWorldFromContextObject(Object, false) : NULL;
@@ -77,7 +53,7 @@ UWorld* FVisualLogger::GetWorld(const class UObject* Object)
 		World = GEngine->GetWorld();
 	}
 
-	return World != NULL ? World : GWorld;
+	return World;
 }
 
 void FVisualLogger::Shutdown()
@@ -175,8 +151,10 @@ void FVisualLogger::SetIsRecordingToFile(bool InIsRecording)
 		SetIsRecording(true);
 	}
 
+	UWorld* World = GEngine->GetWorld();
+
 	const FString BaseFileName = LogFileNameGetter.IsBound() ? LogFileNameGetter.Execute() : TEXT("VisualLog");
-	const FString MapName = GWorld.GetReference() ? GWorld->GetMapName() : TEXT("");
+	const FString MapName = World ? World->GetMapName() : TEXT("");
 
 	FString OutputFileName = FString::Printf(TEXT("%s_%s"), *BaseFileName, *MapName);
 
@@ -187,13 +165,13 @@ void FVisualLogger::SetIsRecordingToFile(bool InIsRecording)
 			if (Device->HasFlags(EVisualLoggerDeviceFlags::CanSaveToFile))
 			{
 				Device->SetFileName(OutputFileName);
-				Device->StopRecordingToFile(GWorld.GetReference() ? GWorld->TimeSeconds : StartRecordingToFileTime);
+				Device->StopRecordingToFile(World ? World->TimeSeconds : StartRecordingToFileTime);
 			}
 		}
 	}
 	else if (!bIsRecordingToFile && InIsRecording)
 	{
-		StartRecordingToFileTime = GWorld.GetReference() ? GWorld->TimeSeconds : 0;
+		StartRecordingToFileTime = World ? World->TimeSeconds : 0;
 		for (auto* Device : OutputDevices)
 		{
 			if (Device->HasFlags(EVisualLoggerDeviceFlags::CanSaveToFile))
@@ -255,13 +233,6 @@ public:
 #else
 			UE_LOG(LogVisual, Warning, TEXT("Unable to open LogVisualizer - logs are disabled"));
 #endif
-			} 
-		}
-		else if (FParse::Command(&Cmd, TEXT("dumpvislog")))
-		{
-			if (FVisualLogger::Get().IsUsingCircularBuffer())
-			{
-				FVisualLogger::Get().DumpCircularBuffer();
 			}
 		}
 		return false;
diff --git a/Engine/Source/Runtime/Engine/Private/VisualLogger/VisualLoggerAutomationTests.cpp b/Engine/Source/Runtime/Engine/Private/VisualLogger/VisualLoggerAutomationTests.cpp
--- a/Engine/Source/Runtime/Engine/Private/VisualLogger/VisualLoggerAutomationTests.cpp
+++ b/Engine/Source/Runtime/Engine/Private/VisualLogger/VisualLoggerAutomationTests.cpp
@@ -6,11 +6,32 @@
 #include "VisualLogger/VisualLoggerAutomationTests.h"
 #include "VisualLogger/VisualLoggerBinaryFileDevice.h"
 
-UVisualLoggerAutomationTests::UVisualLoggerAutomationTests(const FObjectInitializer& ObjectInitializer) : Super(ObjectInitializer)
+namespace
 {
+	UWorld* GetSimpleEngineAutomationTestWorld(const int32 TestFlags)
+	{
+		// Accessing the game world is only valid for game-only 
+		if (((TestFlags & EAutomationTestFlags::ATF_Editor) || (TestFlags & EAutomationTestFlags::ATF_Game)) == false)
+		{
+			return NULL;
+		}
+		if (GEngine->GetWorldContexts().Num() == 0)
+		{
+			return NULL;
+		}
+		if (GEngine->GetWorldContexts()[0].WorldType != EWorldType::Game && GEngine->GetWorldContexts()[0].WorldType != EWorldType::Editor)
+		{
+			return NULL;
+		}
 
+		return GEngine->GetWorldContexts()[0].World();
+	}
 }
 
+UVisualLoggerAutomationTests::UVisualLoggerAutomationTests(const FObjectInitializer& ObjectInitializer) : Super(ObjectInitializer)
+{
+
+}
 
 #if ENABLE_VISUAL_LOG
 
@@ -87,12 +108,15 @@ IMPLEMENT_SIMPLE_AUTOMATION_TEST(FVisualLogTest, "Engine.VisualLogger.Logging si
 
 bool FVisualLogTest::RunTest(const FString& Parameters)
 {
+	UWorld* World = GetSimpleEngineAutomationTestWorld(GetTestFlags());
+	CHECK_SUCCESS(World != NULL);
+
 	FTestDeviceContext<FVisualLoggerTestDevice> Context;
 
 	FVisualLogger::Get().SetIsRecording(false);
 	CHECK_FAIL(FVisualLogger::Get().IsRecording());
 
-	UE_VLOG(GWorld, LogVisual, Log, TEXT("Simple text line to test vlog"));
+	UE_VLOG(World, LogVisual, Log, TEXT("Simple text line to test vlog"));
 	CHECK_SUCCESS(Context.Device.LastObject == NULL);
 	CHECK_SUCCESS(Context.Device.LastEntry.TimeStamp == -1);
 
@@ -101,11 +125,11 @@ bool FVisualLogTest::RunTest(const FString& Parameters)
 
 	{
 		const FString TextToLog = TEXT("Simple text line to test if UE_VLOG_UELOG works fine");
-		float CurrentTimestamp = GWorld->TimeSeconds;
-		UE_VLOG_UELOG(GWorld, LogVisual, Log, TEXT("%s"), *TextToLog);
-		CHECK_SUCCESS(Context.Device.LastObject != GWorld);
+		float CurrentTimestamp = World->TimeSeconds;
+		UE_VLOG_UELOG(World, LogVisual, Log, TEXT("%s"), *TextToLog);
+		CHECK_SUCCESS(Context.Device.LastObject != World);
 		CHECK_SUCCESS(Context.Device.LastEntry.TimeStamp == -1);
-		FVisualLogEntry* CurrentEntry = FVisualLogger::Get().GetEntryToWrite(GWorld, GWorld->TimeSeconds, ECreateIfNeeded::DontCreate);
+		FVisualLogEntry* CurrentEntry = FVisualLogger::Get().GetEntryToWrite(World, World->TimeSeconds, ECreateIfNeeded::DontCreate);
 		
 		{
 			CHECK_SUCCESS(CurrentEntry != NULL);
@@ -115,7 +139,7 @@ bool FVisualLogTest::RunTest(const FString& Parameters)
 			CHECK_SUCCESS(CurrentEntry->LogLines[0].Line == TextToLog);
 
 			const float NewTimestamp = CurrentTimestamp + 0.1;
-			FVisualLogEntry* NewEntry = FVisualLogger::Get().GetEntryToWrite(GWorld, NewTimestamp); //generate new entry and serialize old one
+			FVisualLogEntry* NewEntry = FVisualLogger::Get().GetEntryToWrite(World, NewTimestamp); //generate new entry and serialize old one
 			CurrentEntry = &Context.Device.LastEntry;
 			CHECK_SUCCESS(CurrentEntry != NULL);
 			CHECK_SUCCESS(CurrentEntry->TimeStamp == CurrentTimestamp);
@@ -138,12 +162,15 @@ bool FVisualLogTest::RunTest(const FString& Parameters)
 IMPLEMENT_SIMPLE_AUTOMATION_TEST(FVisualLogSegmentsTest, "Engine.VisualLogger.Logging segment shape", EAutomationTestFlags::ATF_Game | EAutomationTestFlags::ATF_Editor)
 bool FVisualLogSegmentsTest::RunTest(const FString& Parameters)
 {
+	UWorld* World = GetSimpleEngineAutomationTestWorld(GetTestFlags());
+	CHECK_SUCCESS(World != NULL);
+
 	FTestDeviceContext<FVisualLoggerTestDevice> Context;
 
 	FVisualLogger::Get().SetIsRecording(false);
 	CHECK_FAIL(FVisualLogger::Get().IsRecording());
 
-	UE_VLOG_SEGMENT(GWorld, LogVisual, Log, FVector(0, 0, 0), FVector(1, 0, 0), FColor::Red, TEXT("Simple segment log to test vlog"));
+	UE_VLOG_SEGMENT(World, LogVisual, Log, FVector(0, 0, 0), FVector(1, 0, 0), FColor::Red, TEXT("Simple segment log to test vlog"));
 	CHECK_SUCCESS(Context.Device.LastObject == NULL);
 	CHECK_SUCCESS(Context.Device.LastEntry.TimeStamp == -1);
 
@@ -152,12 +179,12 @@ bool FVisualLogSegmentsTest::RunTest(const FString& Parameters)
 	{
 		const FVector StartPoint(0, 0, 0);
 		const FVector EndPoint(1, 0, 0);
-		UE_VLOG_SEGMENT(GWorld, LogVisual, Log, StartPoint, EndPoint, FColor::Red, TEXT("Simple segment log to test vlog"));
+		UE_VLOG_SEGMENT(World, LogVisual, Log, StartPoint, EndPoint, FColor::Red, TEXT("Simple segment log to test vlog"));
 		CHECK_SUCCESS(Context.Device.LastObject == NULL);
 		CHECK_SUCCESS(Context.Device.LastEntry.TimeStamp == -1);
-		FVisualLogEntry* CurrentEntry = FVisualLogger::Get().GetEntryToWrite(GWorld, GWorld->TimeSeconds, ECreateIfNeeded::DontCreate);
+		FVisualLogEntry* CurrentEntry = FVisualLogger::Get().GetEntryToWrite(World, World->TimeSeconds, ECreateIfNeeded::DontCreate);
 
-		float CurrentTimestamp = GWorld->TimeSeconds;
+		float CurrentTimestamp = World->TimeSeconds;
 		{
 			CHECK_SUCCESS(CurrentEntry != NULL);
 			CHECK_SUCCESS(CurrentEntry->TimeStamp == CurrentTimestamp);
@@ -168,7 +195,7 @@ bool FVisualLogSegmentsTest::RunTest(const FString& Parameters)
 			CHECK_SUCCESS(CurrentEntry->ElementsToDraw[0].Points[1] == EndPoint);
 
 			const float NewTimestamp = CurrentTimestamp + 0.1;
-			FVisualLogEntry* NewEntry = FVisualLogger::Get().GetEntryToWrite(GWorld, NewTimestamp); //generate new entry and serialize old one
+			FVisualLogEntry* NewEntry = FVisualLogger::Get().GetEntryToWrite(World, NewTimestamp); //generate new entry and serialize old one
 			CurrentEntry = &Context.Device.LastEntry;
 			CHECK_SUCCESS(CurrentEntry != NULL);
 			CHECK_SUCCESS(CurrentEntry->TimeStamp == CurrentTimestamp);
@@ -195,6 +222,9 @@ DEFINE_VLOG_EVENT(EventTest3, All, "Third simple event for vlog tests")
 IMPLEMENT_SIMPLE_AUTOMATION_TEST(FVisualLogEventsTest, "Engine.VisualLogger.Logging events", EAutomationTestFlags::ATF_Game | EAutomationTestFlags::ATF_Editor)
 bool FVisualLogEventsTest::RunTest(const FString& Parameters)
 {
+	UWorld* World = GetSimpleEngineAutomationTestWorld(GetTestFlags());
+	CHECK_SUCCESS(World != NULL);
+
 	FTestDeviceContext<FVisualLoggerTestDevice> Context;
 	FVisualLogger::Get().SetIsRecording(true);
 
@@ -207,15 +237,15 @@ bool FVisualLogEventsTest::RunTest(const FString& Parameters)
 	CHECK_SUCCESS(EventTest3.Name == TEXT("EventTest3"));
 	CHECK_SUCCESS(EventTest3.FriendlyDesc == TEXT("Third simple event for vlog tests"));
 
-	FVisualLogEntry* CurrentEntry = FVisualLogger::Get().GetEntryToWrite(GWorld, GWorld->TimeSeconds, ECreateIfNeeded::DontCreate);
-	float CurrentTimestamp = GWorld->TimeSeconds;
-	UE_VLOG_EVENTS(GWorld, NAME_None, EventTest);
+	FVisualLogEntry* CurrentEntry = FVisualLogger::Get().GetEntryToWrite(World, World->TimeSeconds, ECreateIfNeeded::DontCreate);
+	float CurrentTimestamp = World->TimeSeconds;
+	UE_VLOG_EVENTS(World, NAME_None, EventTest);
 	CHECK_SUCCESS(CurrentEntry != NULL);
 	CHECK_SUCCESS(CurrentEntry->TimeStamp == CurrentTimestamp);
 	CHECK_SUCCESS(CurrentEntry->Events.Num() == 1);
 	CHECK_SUCCESS(CurrentEntry->Events[0].Name == TEXT("EventTest"));
 
-	UE_VLOG_EVENTS(GWorld, NAME_None, EventTest, EventTest2);
+	UE_VLOG_EVENTS(World, NAME_None, EventTest, EventTest2);
 	CHECK_SUCCESS(CurrentEntry != NULL);
 	CHECK_SUCCESS(CurrentEntry->TimeStamp == CurrentTimestamp);
 	CHECK_SUCCESS(CurrentEntry->Events.Num() == 2);
@@ -224,8 +254,8 @@ bool FVisualLogEventsTest::RunTest(const FString& Parameters)
 	CHECK_SUCCESS(CurrentEntry->Events[1].Counter == 1);
 	CHECK_SUCCESS(CurrentEntry->Events[1].Name == TEXT("EventTest2"));
 
-	CurrentTimestamp = GWorld->TimeSeconds;
-	UE_VLOG_EVENTS(GWorld, NAME_None, EventTest, EventTest2, EventTest3);
+	CurrentTimestamp = World->TimeSeconds;
+	UE_VLOG_EVENTS(World, NAME_None, EventTest, EventTest2, EventTest3);
 
 	{
 		CHECK_SUCCESS(CurrentEntry != NULL);
@@ -243,7 +273,7 @@ bool FVisualLogEventsTest::RunTest(const FString& Parameters)
 		CHECK_SUCCESS(CurrentEntry->Events[2].UserFriendlyDesc == TEXT("Third simple event for vlog tests"));
 
 		const float NewTimestamp = CurrentTimestamp + 0.1;
-		FVisualLogEntry* NewEntry = FVisualLogger::Get().GetEntryToWrite(GWorld, NewTimestamp); //generate new entry and serialize old one
+		FVisualLogEntry* NewEntry = FVisualLogger::Get().GetEntryToWrite(World, NewTimestamp); //generate new entry and serialize old one
 		CurrentEntry = &Context.Device.LastEntry;
 		CHECK_SUCCESS(CurrentEntry != NULL);
 		CHECK_SUCCESS(CurrentEntry->TimeStamp == CurrentTimestamp);
@@ -268,9 +298,9 @@ bool FVisualLogEventsTest::RunTest(const FString& Parameters)
 	const FName EventTag2 = TEXT("ATLAS_C_1");
 	const FName EventTag3 = TEXT("ATLAS_C_2");
 
-	CurrentTimestamp = GWorld->TimeSeconds + 0.2;
-	CurrentEntry = FVisualLogger::Get().GetEntryToWrite(GWorld, CurrentTimestamp); //generate new entry and serialize old one
-	UE_VLOG_EVENT_WITH_DATA(GWorld, EventTest, EventTag1);
+	CurrentTimestamp = World->TimeSeconds + 0.2;
+	CurrentEntry = FVisualLogger::Get().GetEntryToWrite(World, CurrentTimestamp); //generate new entry and serialize old one
+	UE_VLOG_EVENT_WITH_DATA(World, EventTest, EventTag1);
 	CHECK_SUCCESS(CurrentEntry != NULL);
 	CHECK_SUCCESS(CurrentEntry->TimeStamp == CurrentTimestamp);
 	CHECK_SUCCESS(CurrentEntry->Events.Num() == 1);
@@ -278,10 +308,10 @@ bool FVisualLogEventsTest::RunTest(const FString& Parameters)
 	CHECK_SUCCESS(CurrentEntry->Events[0].EventTags.Num() == 1);
 	CHECK_SUCCESS(CurrentEntry->Events[0].EventTags[EventTag1] == 1);
 
-	CurrentTimestamp = GWorld->TimeSeconds + 0.3;
-	CurrentEntry = FVisualLogger::Get().GetEntryToWrite(GWorld, CurrentTimestamp); //generate new entry and serialize old one
-	UE_VLOG_EVENT_WITH_DATA(GWorld, EventTest, EventTag1, EventTag2, EventTag3);
-	UE_VLOG_EVENT_WITH_DATA(GWorld, EventTest, EventTag3);
+	CurrentTimestamp = World->TimeSeconds + 0.3;
+	CurrentEntry = FVisualLogger::Get().GetEntryToWrite(World, CurrentTimestamp); //generate new entry and serialize old one
+	UE_VLOG_EVENT_WITH_DATA(World, EventTest, EventTag1, EventTag2, EventTag3);
+	UE_VLOG_EVENT_WITH_DATA(World, EventTest, EventTag3);
 	CHECK_SUCCESS(CurrentEntry != NULL);
 	CHECK_SUCCESS(CurrentEntry->TimeStamp == CurrentTimestamp);
 	CHECK_SUCCESS(CurrentEntry->Events.Num() == 1);
diff --git a/Engine/Source/Runtime/Engine/Private/VisualLogger/VisualLoggerCircularBufferDevice.cpp b/Engine/Source/Runtime/Engine/Private/VisualLogger/VisualLoggerCircularBufferDevice.cpp
deleted file mode 100644
diff --git a/Engine/Source/Runtime/Engine/Public/VisualLogger/VisualLogger.h b/Engine/Source/Runtime/Engine/Public/VisualLogger/VisualLogger.h
--- a/Engine/Source/Runtime/Engine/Public/VisualLogger/VisualLogger.h
+++ b/Engine/Source/Runtime/Engine/Public/VisualLogger/VisualLogger.h
@@ -97,7 +97,6 @@
 #define TEXT_NULL TEXT("NULL")
 #define TEXT_TRUE TEXT("TRUE")
 #define TEXT_FALSE TEXT("FALSE")
-#define VISLOG_FILENAME_EXT TEXT("bvlog")
 #define TEXT_CONDITION(Condition) ((Condition) ? TEXT_TRUE : TEXT_FALSE)
 
 #if ENABLE_VISUAL_LOG
@@ -195,12 +194,6 @@ public:
 	void SetIsRecordingOnServer(bool IsRecording) { bIsRecordingOnServer = IsRecording; }
 	bool IsRecordingOnServer() { return !!bIsRecordingOnServer; }
 
-	bool IsUsingCircularBuffer();
-
-	FVisualLogDevice* GetCircularBuffer();
-
-	bool DumpCircularBuffer();
-
 	/** Add visual logger output device */
 	void AddDevice(FVisualLogDevice* InDevice) { OutputDevices.AddUnique(InDevice); }
 	/** Remove visual logger output device */
diff --git a/Engine/Source/Runtime/Engine/Public/VisualLogger/VisualLogger.inl b/Engine/Source/Runtime/Engine/Public/VisualLogger/VisualLogger.inl
--- a/Engine/Source/Runtime/Engine/Public/VisualLogger/VisualLogger.inl
+++ b/Engine/Source/Runtime/Engine/Public/VisualLogger/VisualLogger.inl
@@ -34,7 +34,7 @@ FORCEINLINE
 bool CheckVisualLogInputInternal(const class UObject* Object, const struct FLogCategoryBase& Category, ELogVerbosity::Type Verbosity, UWorld **World, FVisualLogEntry **CurrentEntry)
 {
 	FVisualLogger& VisualLogger = FVisualLogger::Get();
-	if (!Object || (GEngine && GEngine->bDisableAILogging) || (VisualLogger.IsRecording() == false && VisualLogger.IsUsingCircularBuffer() == false) || Object->HasAnyFlags(RF_ClassDefaultObject))
+	if (!Object || (GEngine && GEngine->bDisableAILogging) || VisualLogger.IsRecording() == false || Object->HasAnyFlags(RF_ClassDefaultObject))
 	{
 		return false; 
 	}
@@ -81,16 +81,9 @@ FVisualLogEntry* FVisualLogger::GetEntryToWrite(const class UObject* Object, flo
 				FVisualLogEntry* Entry = &CurrentPair.Value;
 				if (Entry->TimeStamp >= 0 && Entry->TimeStamp < TimeStamp)
 				{
-					if (IsRecording())
+					for (auto* Device : OutputDevices)
 					{
-						for (auto* Device : OutputDevices)
-						{
-							Device->Serialize(CurrentPair.Key, ObjectToNameMap[CurrentPair.Key], ObjectToClassNameMap[CurrentPair.Key], *Entry);
-						}
-					}
-					else if (IsUsingCircularBuffer())
-					{
-						GetCircularBuffer()->Serialize(CurrentPair.Key, ObjectToNameMap[CurrentPair.Key], ObjectToClassNameMap[CurrentPair.Key], *Entry);
+						Device->Serialize(CurrentPair.Key, ObjectToNameMap[CurrentPair.Key], ObjectToClassNameMap[CurrentPair.Key], *Entry);
 					}
 					Entry->Reset();
 				}
@@ -155,16 +148,9 @@ FVisualLogEntry* FVisualLogger::GetEntryToWrite(const class UObject* Object, flo
 				FVisualLogEntry* Entry = &CurrentPair.Value;
 				if (Entry->TimeStamp >= 0 && (!World.IsValid() || Entry->TimeStamp < World->GetTimeSeconds())) // CurrentEntry->TimeStamp == -1 means it's not initialized entry information
 				{
-					if (IsRecording())
-					{
-						for (auto* Device : OutputDevices)
-						{
-							Device->Serialize(CurrentPair.Key, ObjectToNameMap[CurrentPair.Key], ObjectToClassNameMap[CurrentPair.Key], *Entry);
-						}
-					}
-					else if (IsUsingCircularBuffer())
+					for (auto* Device : OutputDevices)
 					{
-						GetCircularBuffer()->Serialize(CurrentPair.Key, ObjectToNameMap[CurrentPair.Key], ObjectToClassNameMap[CurrentPair.Key], *Entry);
+						Device->Serialize(CurrentPair.Key, ObjectToNameMap[CurrentPair.Key], ObjectToClassNameMap[CurrentPair.Key], *Entry);
 					}
 					Entry->Reset();
 				}
