Fixed a wrong condition in FNavPathType that was making casts between unrelated path types possible #UE4

diff --git a/Engine/Source/Runtime/Engine/Classes/AI/Navigation/NavigationTypes.h b/Engine/Source/Runtime/Engine/Classes/AI/Navigation/NavigationTypes.h
--- a/Engine/Source/Runtime/Engine/Classes/AI/Navigation/NavigationTypes.h
+++ b/Engine/Source/Runtime/Engine/Classes/AI/Navigation/NavigationTypes.h
@@ -255,8 +255,7 @@ struct ENGINE_API FNavPathType
 	
 	bool operator==(const FNavPathType& Other) const
 	{
-		return Id == Other.Id || (ParentType != nullptr && *ParentType == Other)
-			|| (Other.ParentType != nullptr && *Other.ParentType == *this);
+		return Id == Other.Id || (ParentType != nullptr && *ParentType == Other);
 	}
 
 private:
