Fix issue with morphtarget not displaying in viewport due to re-registration

diff --git a/Engine/Source/Runtime/Engine/Classes/Components/SkeletalMeshComponent.h b/Engine/Source/Runtime/Engine/Classes/Components/SkeletalMeshComponent.h
--- a/Engine/Source/Runtime/Engine/Classes/Components/SkeletalMeshComponent.h
+++ b/Engine/Source/Runtime/Engine/Classes/Components/SkeletalMeshComponent.h
@@ -1170,6 +1170,7 @@ private:
 
 	bool ShouldBlendPhysicsBones();	
 	void ClearAnimScriptInstance();
+	virtual void RefreshActiveVertexAnims() override;
 
 	//Data for parallel evaluation of animation
 	FAnimationEvaluationContext AnimEvaluationContext;
diff --git a/Engine/Source/Runtime/Engine/Classes/Components/SkinnedMeshComponent.h b/Engine/Source/Runtime/Engine/Classes/Components/SkinnedMeshComponent.h
--- a/Engine/Source/Runtime/Engine/Classes/Components/SkinnedMeshComponent.h
+++ b/Engine/Source/Runtime/Engine/Classes/Components/SkinnedMeshComponent.h
@@ -890,6 +890,11 @@ private:
 	template <bool bExtraBoneInfluencesT, bool bCachedMatrices>
 	FVector GetTypedSkinnedVertexPosition(const FSkelMeshChunk& Chunk, const FSkeletalMeshVertexBuffer& VertexBufferGPUSkin, int32 VertIndex, bool bSoftVertex, const TArray<FMatrix> & RefToLocals = TArray<FMatrix>()) const;
 
+	/**
+	 * Gets called when register, verifies all ActiveVertexAnimation is still valid 
+	 */
+	virtual void RefreshActiveVertexAnims() {};
+
 	// Animation update rate control.
 public:
 	/** Animation Update Rate optimization parameters. */
diff --git a/Engine/Source/Runtime/Engine/Private/Components/SkeletalMeshComponent.cpp b/Engine/Source/Runtime/Engine/Private/Components/SkeletalMeshComponent.cpp
--- a/Engine/Source/Runtime/Engine/Private/Components/SkeletalMeshComponent.cpp
+++ b/Engine/Source/Runtime/Engine/Private/Components/SkeletalMeshComponent.cpp
@@ -2104,4 +2104,16 @@ void USkeletalMeshComponent::SetRootBodyIndex(int32 InBodyIndex)
 		ensure(false);
 		RootBodyData.TransformToRoot = FTransform::Identity;
 	}
+}
+
+void USkeletalMeshComponent::RefreshActiveVertexAnims()
+{
+	if (SkeletalMesh && AnimScriptInstance)
+	{
+		ActiveVertexAnims = UpdateActiveVertexAnims(SkeletalMesh, AnimScriptInstance->MorphTargetCurves, AnimScriptInstance->VertexAnims);
+	}
+	else
+	{
+		ActiveVertexAnims.Empty();
+	}
 }
\ No newline at end of file
diff --git a/Engine/Source/Runtime/Engine/Private/Components/SkinnedMeshComponent.cpp b/Engine/Source/Runtime/Engine/Private/Components/SkinnedMeshComponent.cpp
--- a/Engine/Source/Runtime/Engine/Private/Components/SkinnedMeshComponent.cpp
+++ b/Engine/Source/Runtime/Engine/Private/Components/SkinnedMeshComponent.cpp
@@ -283,9 +283,8 @@ void USkinnedMeshComponent::CreateRenderState_Concurrent()
 				UseLOD = FMath::Clamp(PredictedLODLevel, 0, MeshObject->GetSkeletalMeshResource().LODModels.Num()-1);
 			}
 
-			// We just recreated RenderState, that means it could be new mesh or different mesh
-			// Clear ActiveVertexAnims, so that it doesn't use from previous state
-			ActiveVertexAnims.Empty();
+			// verifies vertex animations are valid
+			RefreshActiveVertexAnims();
 			MeshObject->Update(UseLOD,this,ActiveVertexAnims);  // send to rendering thread
 		}
 
