Computeshader to properly copy all sizes of volumetextures

diff --git a/Engine/Shaders/UpdateTextureShaders.usf b/Engine/Shaders/UpdateTextureShaders.usf
--- a/Engine/Shaders/UpdateTextureShaders.usf
+++ b/Engine/Shaders/UpdateTextureShaders.usf
@@ -7,9 +7,10 @@
 #include "Common.usf"
 
 Buffer<int4> SrcBuffer;
-RWTexture2D<int4> DestTexture;
+RWTexture2D<int4> DestTexture;  
 int SrcPitch;
 
+
 uint4 DestPosSize; // xy = pos, zw = Size of the destination sub rect
 
 [numthreads(8,8,1)]
@@ -17,7 +18,21 @@ void UpdateTexture2DSubresourceCS( uint2 ThreadId : SV_DispatchThreadID )
 {
 	if( all( ThreadId.xy < DestPosSize.zw ) )
 	{
-		DestTexture[ ThreadId.xy + DestPosSize.xy ] = SrcBuffer[ ThreadId.x + ThreadId.y * SrcPitch ];
+		DestTexture[ThreadId.xy + DestPosSize.xy] = SrcBuffer[ThreadId.x + ThreadId.y * SrcPitch];
+	}
+}
+
+RWTexture3D<int4> DestTexture3D;
+int SrcDepthPitch;
+uint4 DestPos;
+uint4 DestSize;
+
+[numthreads(8, 8, 8)]
+void UpdateTexture3DSubresourceCS(uint3 ThreadId : SV_DispatchThreadID)
+{
+	if (all(ThreadId.xyz < DestSize.xyz))
+	{
+		DestTexture3D[ThreadId.xyz + DestPos.xyz] = SrcBuffer[ThreadId.x + (ThreadId.y * SrcPitch) + (ThreadId.z * SrcDepthPitch)];
 	}
 }
 
@@ -26,5 +41,5 @@ Texture2D<int4> SrcTexture;
 [numthreads(8,8,1)]
 void CopyTexture2DCS( uint3 ThreadId : SV_DispatchThreadID )
 {
-	DestTexture[ ThreadId.xy ] = SrcTexture.Load( int3( ThreadId.x, ThreadId.y, 0 ) );
+	DestTexture[ThreadId.xy] = SrcTexture.Load(int3(ThreadId.x, ThreadId.y, 0));
 }
\ No newline at end of file
diff --git a/Engine/Source/Runtime/Engine/Private/UpdateTextureShaders.cpp b/Engine/Source/Runtime/Engine/Private/UpdateTextureShaders.cpp
--- a/Engine/Source/Runtime/Engine/Private/UpdateTextureShaders.cpp
+++ b/Engine/Source/Runtime/Engine/Private/UpdateTextureShaders.cpp
@@ -5,4 +5,5 @@
 #include "UpdateTextureShaders.h"
 
 IMPLEMENT_SHADER_TYPE(,FUpdateTexture2DSubresouceCS,TEXT("UpdateTextureShaders"),TEXT("UpdateTexture2DSubresourceCS"),SF_Compute);
+IMPLEMENT_SHADER_TYPE(,FUpdateTexture3DSubresouceCS, TEXT("UpdateTextureShaders"), TEXT("UpdateTexture3DSubresourceCS"), SF_Compute);
 IMPLEMENT_SHADER_TYPE(,FCopyTexture2DCS,TEXT("UpdateTextureShaders"),TEXT("CopyTexture2DCS"),SF_Compute);
diff --git a/Engine/Source/Runtime/Engine/Public/UpdateTextureShaders.h b/Engine/Source/Runtime/Engine/Public/UpdateTextureShaders.h
--- a/Engine/Source/Runtime/Engine/Public/UpdateTextureShaders.h
+++ b/Engine/Source/Runtime/Engine/Public/UpdateTextureShaders.h
@@ -40,6 +40,49 @@ public:
 	FShaderResourceParameter DestTexture;
 };
 
+class FUpdateTexture3DSubresouceCS : public FGlobalShader
+{
+	DECLARE_SHADER_TYPE(FUpdateTexture3DSubresouceCS, Global);
+public:
+	FUpdateTexture3DSubresouceCS() {}
+	FUpdateTexture3DSubresouceCS(const ShaderMetaType::CompiledShaderInitializerType& Initializer)
+		: FGlobalShader(Initializer)
+	{
+		SrcPitchParameter.Bind(Initializer.ParameterMap, TEXT("SrcPitch"), SPF_Mandatory);
+		SrcDepthPitchParameter.Bind(Initializer.ParameterMap, TEXT("SrcDepthPitch"), SPF_Mandatory);
+
+		SrcBuffer.Bind(Initializer.ParameterMap, TEXT("SrcBuffer"), SPF_Mandatory);
+
+		DestPosParameter.Bind(Initializer.ParameterMap, TEXT("DestPos"), SPF_Mandatory);
+		DestSizeParameter.Bind(Initializer.ParameterMap, TEXT("DestSize"), SPF_Mandatory);
+
+		DestTexture3D.Bind(Initializer.ParameterMap, TEXT("DestTexture3D"), SPF_Mandatory);
+	}
+
+	// FShader interface.
+	virtual bool Serialize(FArchive& Ar)
+	{
+		bool bShaderHasOutdatedParameters = FGlobalShader::Serialize(Ar);
+		Ar << SrcPitchParameter << SrcDepthPitchParameter << SrcBuffer << DestPosParameter << DestSizeParameter  << DestTexture3D;
+		return bShaderHasOutdatedParameters;
+	}
+
+	static bool ShouldCache(EShaderPlatform Platform)
+	{
+		return IsFeatureLevelSupported(Platform, ERHIFeatureLevel::SM5);
+	}
+
+	//protected:
+	FShaderParameter SrcPitchParameter;
+	FShaderParameter SrcDepthPitchParameter;
+	FShaderResourceParameter SrcBuffer;
+
+	FShaderParameter DestPosParameter;
+	FShaderParameter DestSizeParameter;
+
+	FShaderResourceParameter DestTexture3D;
+};
+
 class FCopyTexture2DCS : public FGlobalShader
 {
 	DECLARE_SHADER_TYPE(FCopyTexture2DCS,Global);
