[UE-12569] Keeping DetermineUCSModifiedProperties() from being ran on component CDOs.

Also, keeping components with an null world from crashing during GC.

diff --git a/Engine/Source/Runtime/Engine/Private/Components/ActorComponent.cpp b/Engine/Source/Runtime/Engine/Private/Components/ActorComponent.cpp
--- a/Engine/Source/Runtime/Engine/Private/Components/ActorComponent.cpp
+++ b/Engine/Source/Runtime/Engine/Private/Components/ActorComponent.cpp
@@ -179,7 +179,7 @@ void UActorComponent::PostLoad()
 		}
 	}
 
-	if (GetLinkerUE4Version() < VER_UE4_TRACK_UCS_MODIFIED_PROPERTIES)
+	if ((GetLinkerUE4Version() < VER_UE4_TRACK_UCS_MODIFIED_PROPERTIES) && !HasAnyFlags(RF_ClassDefaultObject))
 	{
 		DetermineUCSModifiedProperties();
 	}
diff --git a/Engine/Source/Runtime/Engine/Private/Components/ChildActorComponent.cpp b/Engine/Source/Runtime/Engine/Private/Components/ChildActorComponent.cpp
--- a/Engine/Source/Runtime/Engine/Private/Components/ChildActorComponent.cpp
+++ b/Engine/Source/Runtime/Engine/Private/Components/ChildActorComponent.cpp
@@ -78,7 +78,7 @@ void UChildActorComponent::OnComponentDestroyed()
 {
 	Super::OnComponentDestroyed();
 
-	DestroyChildActor(!GetWorld()->IsGameWorld());
+	DestroyChildActor(GetWorld() && !GetWorld()->IsGameWorld());
 }
 
 class FChildActorComponentInstanceData : public FSceneComponentInstanceData
