#ue4 - Avoid creating temp FText strings every time in some PrimitiveComponent Physics methods. Instead use static const strings and only in non-shipping builds.

diff --git a/Engine/Source/Runtime/Engine/Classes/Components/PrimitiveComponent.h b/Engine/Source/Runtime/Engine/Classes/Components/PrimitiveComponent.h
--- a/Engine/Source/Runtime/Engine/Classes/Components/PrimitiveComponent.h
+++ b/Engine/Source/Runtime/Engine/Classes/Components/PrimitiveComponent.h
@@ -1550,7 +1550,7 @@ private:
 	bool CheckStaticMobilityAndWarn(const FText& ActionText) const;
 
 	/** Check if mobility is set to non-static. If BodyInstanceRequiresSimulation is non-null we check that it is simulated. Triggers a PIE warning if conditions fails */
-	void WarnInvalidPhysicsOperations(const FText& ActionText, const FBodyInstance* BodyInstanceRequiresSimulation = nullptr) const;
+	void WarnInvalidPhysicsOperations_Internal(const FText& ActionText, const FBodyInstance* BodyInstanceRequiresSimulation = nullptr) const;
 
 public:
 
diff --git a/Engine/Source/Runtime/Engine/Private/PrimitiveComponentPhysics.cpp b/Engine/Source/Runtime/Engine/Private/PrimitiveComponentPhysics.cpp
--- a/Engine/Source/Runtime/Engine/Private/PrimitiveComponentPhysics.cpp
+++ b/Engine/Source/Runtime/Engine/Private/PrimitiveComponentPhysics.cpp
@@ -8,6 +8,13 @@
 
 #define LOCTEXT_NAMESPACE "PrimitiveComponent"
 
+#if !(UE_BUILD_SHIPPING || UE_BUILD_TEST)
+	#define WarnInvalidPhysicsOperations(Text, BodyInstance)  { static const FText _WarnText(Text); WarnInvalidPhysicsOperations_Internal(_WarnText, BodyInstance); }
+#else
+	#define WarnInvalidPhysicsOperations(Text, BodyInstance)
+#endif
+
+
 bool UPrimitiveComponent::ApplyRigidBodyState(const FRigidBodyState& NewState, const FRigidBodyErrorCorrection& ErrorCorrection, FVector& OutDeltaPos, FName BoneName)
 {
 	bool bRestoredState = true;
@@ -152,7 +159,7 @@ void UPrimitiveComponent::SetWalkableSlopeOverride(const FWalkableSlopeOverride&
 	BodyInstance.SetWalkableSlopeOverride(NewOverride);
 }
 
-void UPrimitiveComponent::WarnInvalidPhysicsOperations(const FText& ActionText, const FBodyInstance* BI) const
+void UPrimitiveComponent::WarnInvalidPhysicsOperations_Internal(const FText& ActionText, const FBodyInstance* BI) const
 {
 #if !(UE_BUILD_SHIPPING || UE_BUILD_TEST)
 	if (!CheckStaticMobilityAndWarn(ActionText))	//all physics operations require non-static mobility
@@ -275,7 +282,7 @@ void UPrimitiveComponent::SetPhysicsLinearVelocity(FVector NewVel, bool bAddToCu
 {
 	if (FBodyInstance* BI = GetBodyInstance(BoneName))
 	{
-		WarnInvalidPhysicsOperations(LOCTEXT("SetPhysicsLinearVelocity", "SetPhysicsLinearVelocity"));
+		WarnInvalidPhysicsOperations(LOCTEXT("SetPhysicsLinearVelocity", "SetPhysicsLinearVelocity"), nullptr);
 		BI->SetLinearVelocity(NewVel, bAddToCurrent);
 	}
 }
@@ -299,7 +306,7 @@ void UPrimitiveComponent::SetPhysicsAngularVelocity(FVector NewAngVel, bool bAdd
 {
 	if (FBodyInstance* BI = GetBodyInstance(BoneName))
 	{
-		WarnInvalidPhysicsOperations(LOCTEXT("SetPhysicsAngularVelocity", "SetPhysicsAngularVelocity"));
+		WarnInvalidPhysicsOperations(LOCTEXT("SetPhysicsAngularVelocity", "SetPhysicsAngularVelocity"), nullptr);
 		BI->SetAngularVelocity(NewAngVel, bAddToCurrent);
 	}
 }
@@ -308,7 +315,7 @@ void UPrimitiveComponent::SetPhysicsMaxAngularVelocity(float NewMaxAngVel, bool
 {
 	if (FBodyInstance* BI = GetBodyInstance(BoneName))
 	{
-		WarnInvalidPhysicsOperations(LOCTEXT("SetPhysicsMaxAngularVelocity", "SetPhysicsMaxAngularVelocity"));
+		WarnInvalidPhysicsOperations(LOCTEXT("SetPhysicsMaxAngularVelocity", "SetPhysicsMaxAngularVelocity"), nullptr);
 		BI->SetMaxAngularVelocity(NewMaxAngVel, bAddToCurrent);
 	}
 }
