#ue4 - Fix comment on bEnableCameraRotationLag.

UE-12336

diff --git a/Engine/Source/Runtime/Engine/Classes/GameFramework/SpringArmComponent.h b/Engine/Source/Runtime/Engine/Classes/GameFramework/SpringArmComponent.h
--- a/Engine/Source/Runtime/Engine/Classes/GameFramework/SpringArmComponent.h
+++ b/Engine/Source/Runtime/Engine/Classes/GameFramework/SpringArmComponent.h
@@ -67,7 +67,7 @@ class ENGINE_API USpringArmComponent : public USceneComponent
 	uint32 bEnableCameraLag : 1;
 
 	/**
-	 * If true, camera lags behind target position to smooth its movement.
+	 * If true, camera lags behind target rotation to smooth its movement.
 	 * @see CameraRotationLagSpeed
 	 */
 	UPROPERTY(EditAnywhere, BlueprintReadWrite, Category=Lag)
