#ue4 - Consume movement input in Pawn when Unpossessed() or Restart() is called.

diff --git a/Engine/Source/Runtime/Engine/Private/Pawn.cpp b/Engine/Source/Runtime/Engine/Private/Pawn.cpp
--- a/Engine/Source/Runtime/Engine/Private/Pawn.cpp
+++ b/Engine/Source/Runtime/Engine/Private/Pawn.cpp
@@ -496,6 +496,8 @@ void APawn::UnPossessed()
 	{
 		ReceiveUnpossessed(OldController);
 	}
+
+	ConsumeMovementInputVector();
 }
 
 
@@ -656,6 +658,7 @@ void APawn::Restart()
 	{
 		GetMovementComponent()->StopMovementImmediately();
 	}
+	ConsumeMovementInputVector();
 	RecalculateBaseEyeHeight();
 }
 
