Fix for shutdown crash in standalone

- invalid assert in ShutdownAudioDevices
   - assert assumed an order of shutdown for multiple devices (i.e. that main audio device would be last to shutdown). Changed assert to conditional. Tested on several standalone games during shutdown.

https://jira.ol.epicgames.net/browse/UE-12594

diff --git a/Engine/Source/Runtime/Engine/Private/AudioDeviceManager.cpp b/Engine/Source/Runtime/Engine/Private/AudioDeviceManager.cpp
--- a/Engine/Source/Runtime/Engine/Private/AudioDeviceManager.cpp
+++ b/Engine/Source/Runtime/Engine/Private/AudioDeviceManager.cpp
@@ -65,6 +65,20 @@ FAudioDevice* FAudioDeviceManager::CreateAudioDevice(uint32& HandleOut, bool bCr
 		return nullptr;
 	}
 
+#if !WITH_EDITOR
+	// If we are running without the editor, we only need one audio device.
+	if (NumActiveAudioDevices == 1)
+	{
+		FAudioDevice* MainAudioDevice = GEngine->GetMainAudioDevice();
+		if (MainAudioDevice)
+		{
+			HandleOut = MainAudioDevice->DeviceHandle;
+			return MainAudioDevice;
+		}
+		return nullptr;
+	}
+#endif
+
 	FAudioDevice* NewAudioDevice = nullptr;
 
 	if (NumActiveAudioDevices < AUDIO_DEVICE_DEFAULT_ALLOWED_DEVICE_COUNT || (bCreateNewDevice && NumActiveAudioDevices < AUDIO_DEVICE_MAX_DEVICE_COUNT))
@@ -163,15 +177,14 @@ bool FAudioDeviceManager::ShutdownAudioDevice(uint32 Handle)
 			SetActiveDevice(MainDeviceHandle);
 		}
 
-		if (MainDeviceHandle == Handle)
+		// If this is the main device handle and there's more than one reference to the main device, 
+		// don't shut it down until it's the very last handle to get shut down
+		// this is because it's possible for some PIE sessions to be using the main audio device as a fallback to 
+		// preserve CPU performance on low-performance machines
+		if (NumWorldsUsingMainAudioDevice > 0 && MainDeviceHandle == Handle)
 		{
-			// This was a handle to the main audio device using the main audio device
-			check(NumWorldsUsingMainAudioDevice > 0);
 			--NumWorldsUsingMainAudioDevice;
 
-			// If this is the main device handle, don't shut it down until it's the very last handle to get shut down
-			// this is because it's possible for some PIE sessions to be using the main audio device as a fallback to 
-			// preserve CPU performance on low-performance machines
 			return true;
 		}
 	}
