PR #668: Added color temperature property for lights! (Contributed by [REDACTED_PERSON])

https://github.com/EpicGames/UnrealEngine/pull/668

diff --git a/Engine/Source/Editor/UnrealEd/Private/Lightmass/Lightmass.cpp b/Engine/Source/Editor/UnrealEd/Private/Lightmass/Lightmass.cpp
--- a/Engine/Source/Editor/UnrealEd/Private/Lightmass/Lightmass.cpp
+++ b/Engine/Source/Editor/UnrealEd/Private/Lightmass/Lightmass.cpp
@@ -153,6 +153,11 @@ void Copy( const ULightComponent* In, Lightmass::FLightData& Out )
 	Out.Position = In->GetLightPosition();
 	Out.Direction = In->GetDirection();
 
+	if( In->bUseTemperature )
+	{
+		Out.Color *= FLinearColor::MakeFromColorTemperature(In->Temperature);
+	}
+
 	FMemory::Memset(Out.LightProfileTextureData, 0xff, sizeof(Out.LightProfileTextureData));
 
 	if(In->IESTexture)
diff --git a/Engine/Source/Programs/UnrealLightmass/Public/SceneExport.h b/Engine/Source/Programs/UnrealLightmass/Public/SceneExport.h
--- a/Engine/Source/Programs/UnrealLightmass/Public/SceneExport.h
+++ b/Engine/Source/Programs/UnrealLightmass/Public/SceneExport.h
@@ -771,26 +771,26 @@ enum EDawnLightFlags
 
 struct FLightData
 {
-	FGuid		Guid;
+	FGuid			Guid;
 	/** Bit-wise combination of flags from EDawnLightFlags */
-	uint32		LightFlags;
+	uint32			LightFlags;
 	/** Homogeneous coordinates */
-	FVector4	Position;
-	FVector4	Direction;
-	FColor		Color;
-	float		Brightness;
+	FVector4		Position;
+	FVector4		Direction;
+	FLinearColor	Color;
+	float			Brightness;
 	/** The radius of the light's surface, not the light's influence. */
-	float		LightSourceRadius;
+	float			LightSourceRadius;
 	/** The length of the light source*/
-	float		LightSourceLength;
+	float			LightSourceLength;
 	/** Scale factor for the indirect lighting */
-	float		IndirectLightingScale;
+	float			IndirectLightingScale;
 	/** 0 will be completely desaturated, 1 will be unchanged, 2 will be completely saturated */
-	float		IndirectLightingSaturation;
+	float			IndirectLightingSaturation;
 	/** Controls the falloff of shadow penumbras */
-	float		ShadowExponent;
+	float			ShadowExponent;
 	//	only used if an LightProfile, 1d texture data 0:occluded, 255:not occluded
-	uint8		LightProfileTextureData[256];
+	uint8			LightProfileTextureData[256];
 
 	// @param DotProd dot product of light direction and (normalized vector to surface) -1..1
 	inline float ComputeLightProfileMultiplier(const float DotProd) const
diff --git a/Engine/Source/Runtime/Core/Private/Math/Color.cpp b/Engine/Source/Runtime/Core/Private/Math/Color.cpp
--- a/Engine/Source/Runtime/Core/Private/Math/Color.cpp
+++ b/Engine/Source/Runtime/Core/Private/Math/Color.cpp
@@ -89,9 +89,9 @@ FColor FLinearColor::ToFColor(const bool bSRGB) const
 
 	if(bSRGB)
 	{
-		FloatR = FMath::Pow(FloatR, 1.0f / 2.2f);
-		FloatG = FMath::Pow(FloatG, 1.0f / 2.2f);
-		FloatB = FMath::Pow(FloatB, 1.0f / 2.2f);
+		FloatR = FloatR <= 0.0031308f ? FloatR * 12.92f : FMath::Pow( FloatR, 1.0f / 2.4f ) * 1.055f - 0.055f;
+		FloatG = FloatG <= 0.0031308f ? FloatG * 12.92f : FMath::Pow( FloatG, 1.0f / 2.4f ) * 1.055f - 0.055f;
+		FloatB = FloatB <= 0.0031308f ? FloatB * 12.92f : FMath::Pow( FloatB, 1.0f / 2.4f ) * 1.055f - 0.055f;
 	}
 
 	FColor ret;
@@ -269,6 +269,35 @@ FColor FColor::MakeRandomColor()
 	return FColor(FLinearColor::MakeRandomColor());
 }
 
+FLinearColor FLinearColor::MakeFromColorTemperature( float Temp )
+{
+	Temp = FMath::Clamp( Temp, 1000.0f, 15000.0f );
+
+	// Approximate Planckian locus in CIE 1960 UCS
+	float u = ( 0.860117757f + 1.54118254e-4f * Temp + 1.28641212e-7f * Temp*Temp ) / ( 1.0f + 8.42420235e-4f * Temp + 7.08145163e-7f * Temp*Temp );
+	float v = ( 0.317398726f + 4.22806245e-5f * Temp + 4.20481691e-8f * Temp*Temp ) / ( 1.0f - 2.89741816e-5f * Temp + 1.61456053e-7f * Temp*Temp );
+
+	float x = 3.0f * u / ( 2.0f * u - 8.0f * v + 4.0f );
+	float y = 2.0f * v / ( 2.0f * u - 8.0f * v + 4.0f );
+	float z = 1.0f - x - y;
+
+	float Y = 1.0f;
+	float X = Y/y * x;
+	float Z = Y/y * z;
+
+	// XYZ to RGB with BT.709 primaries
+	float R =  3.2404542f * X + -1.5371385f * Y + -0.4985314f * Z;
+	float G = -0.9692660f * X +  1.8760108f * Y +  0.0415560f * Z;
+	float B =  0.0556434f * X + -0.2040259f * Y +  1.0572252f * Z;
+
+	return FLinearColor(R,G,B);
+}
+
+FColor FColor::MakeFromColorTemperature( float Temp )
+{
+	return FLinearColor::MakeFromColorTemperature( Temp ).ToFColor( true );
+}
+
 FColor FColor::MakeRedToGreenColorFromScalar(float Scalar)
 {
 	float RedSclr = FMath::Clamp<float>((1.0f - Scalar)/0.5f,0.f,1.f);
diff --git a/Engine/Source/Runtime/Core/Public/Math/Color.h b/Engine/Source/Runtime/Core/Public/Math/Color.h
--- a/Engine/Source/Runtime/Core/Public/Math/Color.h
+++ b/Engine/Source/Runtime/Core/Public/Math/Color.h
@@ -213,6 +213,11 @@ struct FLinearColor
 	*/
 	static CORE_API FLinearColor MakeRandomColor();
 
+	/**
+	* Converts temperature in Kelvins of a black body radiator to RGB chromaticity.
+	*/
+	static CORE_API FLinearColor MakeFromColorTemperature( float Temp );
+
 	/**
 	 * Euclidean distance between two points.
 	 */
@@ -436,6 +441,11 @@ public:
 	 */
 	static CORE_API FColor MakeRedToGreenColorFromScalar(float Scalar);
 
+	/**
+	* Converts temperature in Kelvins of a black body radiator to RGB chromaticity.
+	*/
+	static CORE_API FColor MakeFromColorTemperature( float Temp );
+
 	/**
 	 *	@return a new FColor based of this color with the new alpha value.
 	 *	Usage: const FColor& MyColor = FColorList::Green.WithAlpha(128);
diff --git a/Engine/Source/Runtime/Engine/Classes/Components/DirectionalLightComponent.h b/Engine/Source/Runtime/Engine/Classes/Components/DirectionalLightComponent.h
--- a/Engine/Source/Runtime/Engine/Classes/Components/DirectionalLightComponent.h
+++ b/Engine/Source/Runtime/Engine/Classes/Components/DirectionalLightComponent.h
@@ -119,7 +119,7 @@ class UDirectionalLightComponent : public ULightComponent
 	UPROPERTY(EditAnywhere, BlueprintReadOnly, Category=Light, meta=(ShowOnlyInnerProperties))
 	struct FLightmassDirectionalLightSettings LightmassSettings;
 
-	UPROPERTY(EditAnywhere, BlueprintReadOnly, Category=Light, meta=(DisplayName = "Atmosphere Sun Light"))
+	UPROPERTY(EditAnywhere, AdvancedDisplay, BlueprintReadOnly, Category=Light, meta=(DisplayName = "Atmosphere Sun Light"))
 	uint32 bUsedAsAtmosphereSunLight : 1;
 
 	UFUNCTION(BlueprintCallable, Category="Rendering|Lighting")
diff --git a/Engine/Source/Runtime/Engine/Classes/Components/LightComponent.h b/Engine/Source/Runtime/Engine/Classes/Components/LightComponent.h
--- a/Engine/Source/Runtime/Engine/Classes/Components/LightComponent.h
+++ b/Engine/Source/Runtime/Engine/Classes/Components/LightComponent.h
@@ -44,6 +44,17 @@ class ENGINE_API ULightComponent : public ULightComponentBase
 {
 	GENERATED_UCLASS_BODY()
 
+	/**
+	* Color temperature in Kelvin of the blackbody illuminant.
+	* White (D65) is 6500K.
+	*/
+	UPROPERTY(EditAnywhere, BlueprintReadOnly, interp, Category = Light, meta = (UIMin = "1700.0", UIMax = "12000.0"))
+	float Temperature;
+	
+	/** false: use white (D65) as illuminant. */
+	UPROPERTY(EditAnywhere, BlueprintReadOnly, Category=Light, meta=(DisplayName = "Use Temperature"))
+	uint32 bUseTemperature : 1;
+
 	/** 
 	 * Shadow map channel which is used to match up with the appropriate static shadowing during a deferred shading pass.
 	 * This is generated during a lighting build.
@@ -181,6 +192,9 @@ public:
 	UFUNCTION(BlueprintCallable, Category="Rendering|Components|Light")
 	void SetLightColor(FLinearColor NewLightColor);
 
+	UFUNCTION(BlueprintCallable, Category = "Rendering|Components|Light")
+	void SetTemperature(float NewTemperature);
+
 	UFUNCTION(BlueprintCallable, Category="Rendering|Components|Light")
 	void SetLightFunctionMaterial(UMaterialInterface* NewLightFunctionMaterial);
 
diff --git a/Engine/Source/Runtime/Engine/Private/Components/LightComponent.cpp b/Engine/Source/Runtime/Engine/Private/Components/LightComponent.cpp
--- a/Engine/Source/Runtime/Engine/Private/Components/LightComponent.cpp
+++ b/Engine/Source/Runtime/Engine/Private/Components/LightComponent.cpp
@@ -239,7 +239,12 @@ FLightSceneProxy::FLightSceneProxy(const ULightComponent* InLightComponent)
 			IESTexture = IESTextureObject;
 		}
 	}
+
 	Color = FLinearColor(InLightComponent->LightColor) * LightBrightness;
+	if( InLightComponent->bUseTemperature )
+	{
+		Color *= FLinearColor::MakeFromColorTemperature(InLightComponent->Temperature);
+	}
 
 	if(LightComponent->LightFunctionMaterial &&
 		LightComponent->LightFunctionMaterial->GetMaterial()->MaterialDomain == MD_LightFunction )
@@ -305,6 +310,8 @@ ULightComponentBase::ULightComponentBase(const FObjectInitializer& ObjectInitial
 ULightComponent::ULightComponent(const FObjectInitializer& ObjectInitializer)
 	: Super(ObjectInitializer)
 {
+	Temperature = 6500.0f;
+	bUseTemperature = false;
 	ShadowMapChannel = INDEX_NONE;
 	PreviewShadowMapChannel = INDEX_NONE;
 	IndirectLightingIntensity = 1.0f;
@@ -501,6 +508,11 @@ bool ULightComponent::CanEditChange(const UProperty* InProperty) const
 		{
 			return bEnableLightShaftBloom;
 		}
+
+		if (PropertyName == GET_MEMBER_NAME_STRING_CHECKED(ULightComponent, Temperature))
+		{
+			return bUseTemperature;
+		}
 	}
 
 	return Super::CanEditChange(InProperty);
@@ -560,7 +572,8 @@ void ULightComponent::PostEditChangeProperty(FPropertyChangedEvent& PropertyChan
 		PropertyName != GET_MEMBER_NAME_STRING_CHECKED(UDirectionalLightComponent, LightShaftOverrideDirection) &&
 		// Properties that should only unbuild lighting for a Static light (can be changed dynamically on a Stationary light)
 		(PropertyName != GET_MEMBER_NAME_STRING_CHECKED(ULightComponent, Intensity) || Mobility == EComponentMobility::Static) &&
-		(PropertyName != GET_MEMBER_NAME_STRING_CHECKED(ULightComponent, LightColor) || Mobility == EComponentMobility::Static))
+		(PropertyName != GET_MEMBER_NAME_STRING_CHECKED(ULightComponent, LightColor) || Mobility == EComponentMobility::Static) &&
+		(PropertyName != GET_MEMBER_NAME_STRING_CHECKED(ULightComponent, Temperature) || Mobility == EComponentMobility::Static) )
 	{
 		InvalidateLightingCache();
 	}
@@ -724,6 +737,24 @@ void ULightComponent::SetLightColor(FLinearColor NewLightColor)
 	}
 }
 
+/** Set color temperature of the light */
+void ULightComponent::SetTemperature(float NewTemperature)
+{
+	// Can't set color on a static light
+	if (AreDynamicDataChangesAllowed()
+		&& Temperature != NewTemperature)
+	{
+		Temperature = NewTemperature;
+
+		// Use lightweight color and brightness update 
+		if( World && World->Scene )
+		{
+			//@todo - remove from scene if brightness or color becomes 0
+			World->Scene->UpdateLightColorAndBrightness( this );
+		}
+	}
+}
+
 void ULightComponent::SetLightFunctionMaterial(UMaterialInterface* NewLightFunctionMaterial)
 {
 	// Can't set light function on a static light
@@ -996,12 +1027,14 @@ void ULightComponent::PostInterpChange(UProperty* PropertyThatChanged)
 	static FName IntensityName(TEXT("Intensity"));
 	static FName BrightnessName(TEXT("Brightness"));
 	static FName IndirectLightingIntensityName(TEXT("IndirectLightingIntensity"));
+	static FName TemperatureName(TEXT("Temperature"));
 
 	FName PropertyName = PropertyThatChanged->GetFName();
 	if (PropertyName == LightColorName
 		|| PropertyName == IntensityName
 		|| PropertyName == BrightnessName
-		|| PropertyName == IndirectLightingIntensityName)
+		|| PropertyName == IndirectLightingIntensityName
+		|| PropertyName == TemperatureName)
 	{
 		// Old brightness tracks will animate the deprecated value
 		if (PropertyName == BrightnessName)
diff --git a/Engine/Source/Runtime/Renderer/Private/RendererScene.cpp b/Engine/Source/Runtime/Renderer/Private/RendererScene.cpp
--- a/Engine/Source/Runtime/Renderer/Private/RendererScene.cpp
+++ b/Engine/Source/Runtime/Renderer/Private/RendererScene.cpp
@@ -1271,6 +1271,11 @@ void FScene::UpdateLightColorAndBrightness(ULightComponent* Light)
 		FUpdateLightColorParameters NewParameters;
 		NewParameters.NewColor = FLinearColor(Light->LightColor) * Light->ComputeLightBrightness();
 		NewParameters.NewIndirectLightingScale = Light->IndirectLightingIntensity;
+
+		if( Light->bUseTemperature )
+		{
+			 NewParameters.NewColor *= FLinearColor::MakeFromColorTemperature(Light->Temperature);
+		}
 	
 		ENQUEUE_UNIQUE_RENDER_COMMAND_THREEPARAMETER(
 			UpdateLightColorAndBrightness,
