Fixed: NavigationMesh streaming. Access violation editor crash
UE-12532

diff --git a/Engine/Source/Runtime/Engine/Classes/AI/Navigation/NavigationData.h b/Engine/Source/Runtime/Engine/Classes/AI/Navigation/NavigationData.h
--- a/Engine/Source/Runtime/Engine/Classes/AI/Navigation/NavigationData.h
+++ b/Engine/Source/Runtime/Engine/Classes/AI/Navigation/NavigationData.h
@@ -480,8 +480,8 @@ public:
 	virtual bool SupportsRuntimeGeneration() const;
 	virtual bool SupportsStreaming() const;
 	virtual void OnNavigationBoundsChanged();
-	virtual void OnStreamingLevelAdded(ULevel* InLevel) {};
-	virtual void OnStreamingLevelRemoved(ULevel* InLevel) {};
+	virtual void OnStreamingLevelAdded(ULevel* InLevel, UWorld* InWorld) {};
+	virtual void OnStreamingLevelRemoved(ULevel* InLevel, UWorld* InWorld) {};
 	
 	//----------------------------------------------------------------------//
 	// Generation & data access                                                      
diff --git a/Engine/Source/Runtime/Engine/Classes/AI/Navigation/RecastNavMesh.h b/Engine/Source/Runtime/Engine/Classes/AI/Navigation/RecastNavMesh.h
--- a/Engine/Source/Runtime/Engine/Classes/AI/Navigation/RecastNavMesh.h
+++ b/Engine/Source/Runtime/Engine/Classes/AI/Navigation/RecastNavMesh.h
@@ -694,8 +694,8 @@ public:
 	/** Called on world origin changes **/
 	virtual void ApplyWorldOffset(const FVector& InOffset, bool bWorldShift) override;
 
-	virtual void OnStreamingLevelAdded(ULevel* InLevel) override;
-	virtual void OnStreamingLevelRemoved(ULevel* InLevel) override;
+	virtual void OnStreamingLevelAdded(ULevel* InLevel, UWorld* InWorld) override;
+	virtual void OnStreamingLevelRemoved(ULevel* InLevel, UWorld* InWorld) override;
 	// End ANavigationData Interface
 
 protected:
diff --git a/Engine/Source/Runtime/Engine/Private/AI/Navigation/NavigationSystem.cpp b/Engine/Source/Runtime/Engine/Private/AI/Navigation/NavigationSystem.cpp
--- a/Engine/Source/Runtime/Engine/Private/AI/Navigation/NavigationSystem.cpp
+++ b/Engine/Source/Runtime/Engine/Private/AI/Navigation/NavigationSystem.cpp
@@ -619,7 +619,7 @@ void UNavigationSystem::OnWorldInitDone(FNavigationSystem::EMode Mode)
 				{
 					for (ANavigationData* NavData : NavDataSet)
 					{
-						NavData->OnStreamingLevelAdded(Level);
+						NavData->OnStreamingLevelAdded(Level, World);
 					}
 				}
 			}
@@ -3141,7 +3141,7 @@ void UNavigationSystem::OnLevelAddedToWorld(ULevel* InLevel, UWorld* InWorld)
 		{
 			for (ANavigationData* NavData : NavDataSet)
 			{
-				NavData->OnStreamingLevelAdded(InLevel);
+				NavData->OnStreamingLevelAdded(InLevel, InWorld);
 			}
 		}
 	}
@@ -3157,7 +3157,7 @@ void UNavigationSystem::OnLevelRemovedFromWorld(ULevel* InLevel, UWorld* InWorld
 		{
 			for (ANavigationData* NavData : NavDataSet)
 			{
-				NavData->OnStreamingLevelRemoved(InLevel);
+				NavData->OnStreamingLevelRemoved(InLevel, InWorld);
 			}
 		}
 	}
diff --git a/Engine/Source/Runtime/Engine/Private/AI/Navigation/RecastNavMesh.cpp b/Engine/Source/Runtime/Engine/Private/AI/Navigation/RecastNavMesh.cpp
--- a/Engine/Source/Runtime/Engine/Private/AI/Navigation/RecastNavMesh.cpp
+++ b/Engine/Source/Runtime/Engine/Private/AI/Navigation/RecastNavMesh.cpp
@@ -1550,13 +1550,13 @@ void ARecastNavMesh::ApplyWorldOffset(const FVector& InOffset, bool bWorldShift)
 	RequestDrawingUpdate();
 }
 
-void ARecastNavMesh::OnStreamingLevelAdded(ULevel* InLevel)
+void ARecastNavMesh::OnStreamingLevelAdded(ULevel* InLevel, UWorld* InWorld)
 {
 	QUICK_SCOPE_CYCLE_COUNTER(STAT_RecastNavMesh_OnStreamingLevelAdded);
 	
 	bWantsUpdate = true;
 
-	if (SupportsStreaming() && GetWorld()->IsGameWorld())
+	if (SupportsStreaming() && InWorld->IsGameWorld())
 	{
 		URecastNavMeshDataChunk* NavDataChunk = GetNavigationDataChunk(InLevel);
 		if (NavDataChunk)
@@ -1571,11 +1571,11 @@ void ARecastNavMesh::OnStreamingLevelAdded(ULevel* InLevel)
 	}
 }
 
-void ARecastNavMesh::OnStreamingLevelRemoved(ULevel* InLevel)
+void ARecastNavMesh::OnStreamingLevelRemoved(ULevel* InLevel, UWorld* InWorld)
 {
 	QUICK_SCOPE_CYCLE_COUNTER(STAT_RecastNavMesh_OnStreamingLevelRemoved);
 	
-	if (SupportsStreaming() && GetWorld()->IsGameWorld())
+	if (SupportsStreaming() && InWorld->IsGameWorld())
 	{
 		URecastNavMeshDataChunk* NavDataChunk = GetNavigationDataChunk(InLevel);
 		if (NavDataChunk)
