removed misplaced return in UBrushComponent::ComponentIsTouchingSelectionBox

diff --git a/Engine/Source/Runtime/Engine/Private/Components/BrushComponent.cpp b/Engine/Source/Runtime/Engine/Private/Components/BrushComponent.cpp
--- a/Engine/Source/Runtime/Engine/Private/Components/BrushComponent.cpp
+++ b/Engine/Source/Runtime/Engine/Private/Components/BrushComponent.cpp
@@ -686,10 +686,6 @@ bool UBrushComponent::ComponentIsTouchingSelectionBox(const FBox& InSelBBox, con
 						return false;
 					}
 				}
-
-				// If the selection box has to encompass all of the component and none of the component's verts failed the intersection test, this component
-				// is considered touching
-				return true;
 			}
 
 			// All points lay within the selection box
