Fix for incorrect body instance usage in InitArticulated when dedicated server does not use async scene but clients do.
#UE-8942

diff --git a/Engine/Source/Runtime/Engine/Private/SkeletalMeshComponentPhysics.cpp b/Engine/Source/Runtime/Engine/Private/SkeletalMeshComponentPhysics.cpp
--- a/Engine/Source/Runtime/Engine/Private/SkeletalMeshComponentPhysics.cpp
+++ b/Engine/Source/Runtime/Engine/Private/SkeletalMeshComponentPhysics.cpp
@@ -827,7 +827,7 @@ void USkeletalMeshComponent::InitArticulated(FPhysScene* PhysScene)
 	if(Aggregate && Aggregate->getNbActors() > 0 && PhysScene)
 	{
 		// Get the scene type from the SkeletalMeshComponent's BodyInstance
-		const uint32 SceneType = BodyInstance.UseAsyncScene() ? PST_Async : PST_Sync;
+		const uint32 SceneType = (bHasBodiesInAsyncScene && PhysScene->HasAsyncScene()) ? PST_Async : PST_Sync;
 		{
 			PxScene* PScene = PhysScene->GetPhysXScene(SceneType);
 			SCOPED_SCENE_WRITE_LOCK(PScene);
