Engine - Misc fix

diff --git a/Engine/Source/Runtime/Engine/Private/World.cpp b/Engine/Source/Runtime/Engine/Private/World.cpp
--- a/Engine/Source/Runtime/Engine/Private/World.cpp
+++ b/Engine/Source/Runtime/Engine/Private/World.cpp
@@ -3507,7 +3507,7 @@ ENGINE_API const FString GetMapNameStatic()
 	{
 		Retval = ContextToUse->World()->GetMapName();
 	}
-	else if ( Internal::GObjInitialized )
+	else if( UObjectInitialized() )
 	{
 		Retval = appGetStartupMap( FCommandLine::Get() );
 	}
