Fixed Canvas text shadows when using distance field fonts having different activation logic to non-distance field fonts

UDN - FCanvasTextItem Missing DropShadow

diff --git a/Engine/Source/Runtime/Engine/Private/UserInterface/CanvasItem.cpp b/Engine/Source/Runtime/Engine/Private/UserInterface/CanvasItem.cpp
--- a/Engine/Source/Runtime/Engine/Private/UserInterface/CanvasItem.cpp
+++ b/Engine/Source/Runtime/Engine/Private/UserInterface/CanvasItem.cpp
@@ -829,11 +829,11 @@ void FCanvasTextItem::Draw( class FCanvas* InCanvas )
 		return;
 	}
 
-	bool bHasShadow = ShadowOffset.Size() != 0.0f;
-	if( ( FontRenderInfo.bEnableShadow == true ) && ( bHasShadow == false ) )
+	bool bHasShadow = FontRenderInfo.bEnableShadow;
+	if( bHasShadow && ShadowOffset.Size() == 0.0f )
 	{
+		// EnableShadow will set a default ShadowOffset value
 		EnableShadow( FLinearColor::Black );
-		bHasShadow = true;
 	}
 	if (Font->ImportOptions.bUseDistanceFieldAlpha)
 	{
@@ -841,17 +841,12 @@ void FCanvasTextItem::Draw( class FCanvas* InCanvas )
 		switch(BlendMode)
 		{
 		case SE_BLEND_Translucent:
-			BlendMode = (FontRenderInfo.bEnableShadow) ? SE_BLEND_TranslucentDistanceFieldShadowed : SE_BLEND_TranslucentDistanceField;
+			BlendMode = (bHasShadow) ? SE_BLEND_TranslucentDistanceFieldShadowed : SE_BLEND_TranslucentDistanceField;
 			break;
 		case SE_BLEND_Masked:
-			BlendMode = (FontRenderInfo.bEnableShadow) ? SE_BLEND_MaskedDistanceFieldShadowed : SE_BLEND_MaskedDistanceField;
+			BlendMode = (bHasShadow) ? SE_BLEND_MaskedDistanceFieldShadowed : SE_BLEND_MaskedDistanceField;
 			break;
 		}
-		// Disable the show if the blend mode is not suitable
-		if (BlendMode != SE_BLEND_TranslucentDistanceFieldShadowed && BlendMode != SE_BLEND_MaskedDistanceFieldShadowed)
-		{
-			bHasShadow = false;
-		}
 	}
 	if (GetFontCacheType() == EFontCacheType::Runtime)
 	{
diff --git a/Engine/Source/Runtime/Engine/Public/CanvasItem.h b/Engine/Source/Runtime/Engine/Public/CanvasItem.h
--- a/Engine/Source/Runtime/Engine/Public/CanvasItem.h
+++ b/Engine/Source/Runtime/Engine/Public/CanvasItem.h
@@ -359,7 +359,8 @@ public:
 	void EnableShadow( const FLinearColor& InColor, const FVector2D& InOffset = FVector2D( 1.0f, 1.0f ) )
 	{
 		ShadowOffset = InOffset;
-		ShadowColor = InColor;		
+		ShadowColor = InColor;
+		FontRenderInfo.bEnableShadow = ShadowOffset.Size() != 0.0f;
 	}
 	
 	/** 
@@ -368,6 +369,7 @@ public:
 	void DisableShadow() 
 	{
 		ShadowOffset = FVector2D::ZeroVector;
+		FontRenderInfo.bEnableShadow = false;
 	}
 
 	/** 
