UE4 - add stats for possible mystery event showing up in profiles

diff --git a/Engine/Source/Runtime/Engine/Private/PhysicsEngine/PhysLevel.cpp b/Engine/Source/Runtime/Engine/Private/PhysicsEngine/PhysLevel.cpp
--- a/Engine/Source/Runtime/Engine/Private/PhysicsEngine/PhysLevel.cpp
+++ b/Engine/Source/Runtime/Engine/Private/PhysicsEngine/PhysLevel.cpp
@@ -217,6 +217,7 @@ void UWorld::StartClothSim()
 
 void FStartPhysicsTickFunction::ExecuteTick(float DeltaTime, enum ELevelTick TickType, ENamedThreads::Type CurrentThread, const FGraphEventRef& MyCompletionGraphEvent)
 {
+	QUICK_SCOPE_CYCLE_COUNTER(FStartPhysicsTickFunction_ExecuteTick);
 	check(Target);
 	Target->StartPhysicsSim();
 }
@@ -228,6 +229,8 @@ FString FStartPhysicsTickFunction::DiagnosticMessage()
 
 void FEndPhysicsTickFunction::ExecuteTick(float DeltaTime, enum ELevelTick TickType, ENamedThreads::Type CurrentThread, const FGraphEventRef& MyCompletionGraphEvent)
 {
+	QUICK_SCOPE_CYCLE_COUNTER(FEndPhysicsTickFunction_ExecuteTick);
+
 	check(Target);
 	FPhysScene* PhysScene = Target->GetPhysicsScene();
 	if (PhysScene == NULL)
@@ -263,6 +266,8 @@ FString FEndPhysicsTickFunction::DiagnosticMessage()
 
 void FStartClothSimulationFunction::ExecuteTick(float DeltaTime, enum ELevelTick TickType, ENamedThreads::Type CurrentThread, const FGraphEventRef& MyCompletionGraphEvent)
 {
+	QUICK_SCOPE_CYCLE_COUNTER(FStartClothSimulationFunction_ExecuteTick);
+
 	check(Target);
 	Target->StartClothSim();
 }
@@ -274,6 +279,8 @@ FString FStartClothSimulationFunction::DiagnosticMessage()
 
 void FEndClothSimulationFunction::ExecuteTick(float DeltaTime, enum ELevelTick TickType, ENamedThreads::Type CurrentThread, const FGraphEventRef& MyCompletionGraphEvent)
 {
+	QUICK_SCOPE_CYCLE_COUNTER(FEndClothSimulationFunction_ExecuteTick);
+
 	//We currently have nothing to do in this tick group, but we still want to wait on cloth simulation so that PostPhysics is ensured this is done
 	check(Target);
 	FPhysScene* PhysScene = Target->GetPhysicsScene();
diff --git a/Engine/Source/Runtime/Engine/Private/PhysicsEngine2D/Box2DIntegration.cpp b/Engine/Source/Runtime/Engine/Private/PhysicsEngine2D/Box2DIntegration.cpp
--- a/Engine/Source/Runtime/Engine/Private/PhysicsEngine2D/Box2DIntegration.cpp
+++ b/Engine/Source/Runtime/Engine/Private/PhysicsEngine2D/Box2DIntegration.cpp
@@ -217,6 +217,8 @@ bool FPhysicsScene2D::ShouldCollide(b2Fixture* FixtureA, b2Fixture* FixtureB)
 
 void FStartPhysics2DTickFunction::ExecuteTick(float DeltaTime, enum ELevelTick TickType, ENamedThreads::Type CurrentThread, const FGraphEventRef& MyCompletionGraphEvent)
 {
+	QUICK_SCOPE_CYCLE_COUNTER(FStartPhysics2DTickFunction_ExecuteTick);
+
 	if (Target->UnrealWorld->bShouldSimulatePhysics)
 	{
 		const int32 VelocityIterations = 8;
@@ -264,6 +266,7 @@ FString FStartPhysics2DTickFunction::DiagnosticMessage()
 
 void FEndPhysics2DTickFunction::ExecuteTick(float DeltaTime, enum ELevelTick TickType, ENamedThreads::Type CurrentThread, const FGraphEventRef& MyCompletionGraphEvent)
 {
+	QUICK_SCOPE_CYCLE_COUNTER(FEndPhysics2DTickFunction_ExecuteTick);
 }
 
 FString FEndPhysics2DTickFunction::DiagnosticMessage()
diff --git a/Engine/Source/Runtime/Engine/Private/SkeletalMeshComponentPhysics.cpp b/Engine/Source/Runtime/Engine/Private/SkeletalMeshComponentPhysics.cpp
--- a/Engine/Source/Runtime/Engine/Private/SkeletalMeshComponentPhysics.cpp
+++ b/Engine/Source/Runtime/Engine/Private/SkeletalMeshComponentPhysics.cpp
@@ -38,6 +38,8 @@
 
 void FSkeletalMeshComponentPreClothTickFunction::ExecuteTick(float DeltaTime, enum ELevelTick TickType, ENamedThreads::Type CurrentThread, const FGraphEventRef& MyCompletionGraphEvent)
 {
+	QUICK_SCOPE_CYCLE_COUNTER(FSkeletalMeshComponentPreClothTickFunction_ExecuteTick);
+
 	if ((TickType == LEVELTICK_All) && Target && !Target->HasAnyFlags(RF_PendingKill | RF_Unreachable))
 	{
 		Target->PreClothTick(DeltaTime);
