Viewport - The viewports find window function no longer assumes the viewport widget is valid always and properly checks for an invalid viewport.  Fixes a crash with VR Preview.

diff --git a/Engine/Source/Runtime/Engine/Private/Slate/SceneViewport.cpp b/Engine/Source/Runtime/Engine/Private/Slate/SceneViewport.cpp
--- a/Engine/Source/Runtime/Engine/Private/Slate/SceneViewport.cpp
+++ b/Engine/Source/Runtime/Engine/Private/Slate/SceneViewport.cpp
@@ -1010,7 +1010,13 @@ void FSceneViewport::SetViewportSize(uint32 NewViewportSizeX, uint32 NewViewport
 
 TSharedPtr<SWindow> FSceneViewport::FindWindow()
 {
-	return FSlateApplication::Get().FindWidgetWindow(ViewportWidget.Pin().ToSharedRef());
+	if ( ViewportWidget.IsValid() )
+	{
+		TSharedPtr<SViewport> PinnedViewportWidget = ViewportWidget.Pin();
+		return FSlateApplication::Get().FindWidgetWindow(PinnedViewportWidget.ToSharedRef());
+	}
+
+	return TSharedPtr<SWindow>();
 }
 
 bool FSceneViewport::IsStereoRenderingAllowed() const
