Change notify spawning behaviour during first update so that we dont get notifies from transitions we are "skipping"

diff --git a/Engine/Source/Runtime/Engine/Private/Animation/AnimNode_StateMachine.cpp b/Engine/Source/Runtime/Engine/Private/Animation/AnimNode_StateMachine.cpp
--- a/Engine/Source/Runtime/Engine/Private/Animation/AnimNode_StateMachine.cpp
+++ b/Engine/Source/Runtime/Engine/Private/Animation/AnimNode_StateMachine.cpp
@@ -206,8 +206,6 @@ void FAnimNode_StateMachine::Initialize(const FAnimationInitializeContext& Conte
 
 			// initialize first update
 			bFirstUpdate = true;
-			// Fire off any entry notifications the default state has
-			Context.AnimInstance->AddAnimNotifyFromGeneratedClass(GetStateInfo(Machine->InitialState).StartNotify);
 		}
 	}
 }
@@ -307,8 +305,11 @@ void FAnimNode_StateMachine::Update(const FAnimationUpdateContext& Context)
 			const int32 NextState = PotentialTransition.TargetState;
 
 			// Fire off Notifies for state transition
-			Context.AnimInstance->AddAnimNotifyFromGeneratedClass(GetStateInfo(PreviousState).EndNotify);
-			Context.AnimInstance->AddAnimNotifyFromGeneratedClass(GetStateInfo(NextState).StartNotify);
+			if (!bFirstUpdate)
+			{
+				Context.AnimInstance->AddAnimNotifyFromGeneratedClass(GetStateInfo(PreviousState).EndNotify);
+				Context.AnimInstance->AddAnimNotifyFromGeneratedClass(GetStateInfo(NextState).StartNotify);
+			}
 			
 			// Get the current weight of the next state, which may be non-zero
 			const float ExistingWeightOfNextState = GetStateWeight(NextState);
@@ -322,7 +323,10 @@ void FAnimNode_StateMachine::Update(const FAnimationUpdateContext& Context)
 			NewTransition->SourceTransitionIndices = PotentialTransition.SourceTransitionIndices;
 #endif
 
-			Context.AnimInstance->AddAnimNotifyFromGeneratedClass(NewTransition->StartNotify);
+			if (!bFirstUpdate)
+			{
+				Context.AnimInstance->AddAnimNotifyFromGeneratedClass(NewTransition->StartNotify);
+			}
 			
 			SetState(Context, NextState);
 
@@ -331,10 +335,12 @@ void FAnimNode_StateMachine::Update(const FAnimationUpdateContext& Context)
 	}
 	while (bFoundValidTransition && (TransitionCountThisFrame < MaxTransitionsPerFrame));
 
-	// in the first update, we don't like to transition from entry state
-	// so we throw out any transition data at the first update
 	if (bFirstUpdate)
 	{
+		//Handle enter notify for "first" (after initial transitions) state
+		Context.AnimInstance->AddAnimNotifyFromGeneratedClass(GetStateInfo().StartNotify);
+		// in the first update, we don't like to transition from entry state
+		// so we throw out any transition data at the first update
 		ActiveTransitionArray.Reset();
 		bFirstUpdate = false;
 	}
