Enable fast forward optimization when we don't have a checkpoint

diff --git a/Engine/Source/Runtime/Engine/Private/DemoNetDriver.cpp b/Engine/Source/Runtime/Engine/Private/DemoNetDriver.cpp
--- a/Engine/Source/Runtime/Engine/Private/DemoNetDriver.cpp
+++ b/Engine/Source/Runtime/Engine/Private/DemoNetDriver.cpp
@@ -1384,6 +1384,7 @@ void UDemoNetDriver::LoadCheckpoint()
 		if ( GotoCheckpointSkipExtraTimeInMS != -1 )
 		{
 			DemoCurrentTime += (float)GotoCheckpointSkipExtraTimeInMS / 1000;
+			bIsFastForwarding = true;
 		}
 
 		GotoCheckpointArchive			= NULL;
