We now flip bone transform buffers in RefreshBoneTransforms if we are not part of ticking. Hopefully avoid issues when RefreshBoneTransforms is manually called from some other system.

diff --git a/Engine/Source/Runtime/Engine/Private/Components/SkeletalMeshComponent.cpp b/Engine/Source/Runtime/Engine/Private/Components/SkeletalMeshComponent.cpp
--- a/Engine/Source/Runtime/Engine/Private/Components/SkeletalMeshComponent.cpp
+++ b/Engine/Source/Runtime/Engine/Private/Components/SkeletalMeshComponent.cpp
@@ -242,7 +242,6 @@ void USkeletalMeshComponent::InitAnim(bool bForceReinit)
         TickAnimation(0.f); 
 
 		RefreshBoneTransforms();
-		FlipEditableSpaceBases();
 		UpdateComponentToWorld();
 	}
 }
@@ -1102,6 +1101,12 @@ void USkeletalMeshComponent::RefreshBoneTransforms(FActorComponentTickFunction*
 
 		PostAnimEvaluation(AnimEvaluationContext);
 	}
+
+	if (TickFunction == NULL)
+	{
+		//Since we aren't doing this through the tick system, assume we want the buffer flipped now
+		FlipEditableSpaceBases();
+	}
 }
 
 void USkeletalMeshComponent::PostAnimEvaluation(FAnimationEvaluationContext& EvaluationContext)
