UE4 - automatically reduce lightmap sizes for instanced static meshes to avoid blowing out texture limits and crashing

diff --git a/Engine/Source/Runtime/Engine/Private/InstancedStaticMesh.cpp b/Engine/Source/Runtime/Engine/Private/InstancedStaticMesh.cpp
--- a/Engine/Source/Runtime/Engine/Private/InstancedStaticMesh.cpp
+++ b/Engine/Source/Runtime/Engine/Private/InstancedStaticMesh.cpp
@@ -1074,6 +1074,44 @@ void UInstancedStaticMeshComponent::GetStaticLightingInfo(FStaticLightingPrimiti
 		int32 LightMapHeight = 0;
 		GetLightMapResolution(LightMapWidth, LightMapHeight);
 
+		bool bFit = false;
+		bool bReduced = false;
+		while (1)
+		{
+			const int32 OneLessThanMaximumSupportedResolution = 1 << (GMaxTextureMipCount - 2);
+
+			const int32 MaxInstancesInMaxSizeLightmap = (OneLessThanMaximumSupportedResolution / LightMapWidth) * ((OneLessThanMaximumSupportedResolution / 2) / LightMapHeight);
+			if (PerInstanceSMData.Num() > MaxInstancesInMaxSizeLightmap)
+			{
+				if (LightMapWidth < 4 || LightMapHeight < 4)
+				{
+					break;
+				}
+				LightMapWidth /= 2;
+				LightMapHeight /= 2;
+				bReduced = true;
+			}
+			else
+			{
+				bFit = true;
+				break;
+			}
+		}
+
+		if (!bFit)
+		{
+			FMessageLog("LightingResults").Message(EMessageSeverity::Error)
+				->AddToken(FUObjectToken::Create(this))
+				->AddToken(FTextToken::Create(NSLOCTEXT("InstancedStaticMesh", "FailedStaticLightingWarning", "The total lightmap size for this InstancedStaticMeshComponent is too big no matter how much we reduce the per-instance size, the number of mesh instances in this component must be reduced")));
+			return;
+		}
+		if (bReduced)
+		{
+			FMessageLog("LightingResults").Message(EMessageSeverity::Warning)
+				->AddToken(FUObjectToken::Create(this))
+				->AddToken(FTextToken::Create(NSLOCTEXT("InstancedStaticMesh", "ReducedStaticLightingWarning", "The total lightmap size for this InstancedStaticMeshComponent was too big and it was automatically reduced. Consider reducing the component's lightmap resolution or number of mesh instances in this component")));
+		}
+
 		const int32 LightMapSize = GetWorld()->GetWorldSettings()->PackedLightAndShadowMapTextureSize;
 		const int32 MaxInstancesInDefaultSizeLightmap = (LightMapSize / LightMapWidth) * ((LightMapSize / 2) / LightMapHeight);
 		if (PerInstanceSMData.Num() > MaxInstancesInDefaultSizeLightmap)
