Data table improvements

UETOOL-281 - Improve data table workflows
UETOOL-282 - Data tables cannot be exported as CSV
UETOOL-289 - Data tables cannot represent arrays in their CSV data

Split up the util functions into their own file (DataTableUtils).
Split up the CSV exporting/importing into its own file (DataTableCSV).
Split up the JSON exporting/importing into its own file (DataTableJSON).
Added support for exporting/importing/displaying array properties.
Added the ability to export a data table as CSV.
Added the ability to import a data table from JSON.

diff --git a/Engine/Source/Developer/AssetTools/Private/AssetTypeActions/AssetTypeActions_CSVAssetBase.cpp b/Engine/Source/Developer/AssetTools/Private/AssetTypeActions/AssetTypeActions_CSVAssetBase.cpp
--- a/Engine/Source/Developer/AssetTools/Private/AssetTypeActions/AssetTypeActions_CSVAssetBase.cpp
+++ b/Engine/Source/Developer/AssetTools/Private/AssetTypeActions/AssetTypeActions_CSVAssetBase.cpp
@@ -4,7 +4,7 @@
 
 #define LOCTEXT_NAMESPACE "AssetTypeActions"
 
-void FAssetTypeActions_CSVAssetBase::ExecuteFindExcelFileInExplorer(TArray<FString> Filenames, TArray<FString> OverrideExtensions)
+void FAssetTypeActions_CSVAssetBase::ExecuteFindSourceFileInExplorer(TArray<FString> Filenames, TArray<FString> OverrideExtensions)
 {
 	for (TArray<FString>::TConstIterator FilenameIter(Filenames); FilenameIter; ++FilenameIter)
 	{
@@ -24,7 +24,7 @@ void FAssetTypeActions_CSVAssetBase::ExecuteFindExcelFileInExplorer(TArray<FStri
 	}
 }
 
-bool FAssetTypeActions_CSVAssetBase::CanExecuteFindExcelFileInExplorer(TArray<FString> Filenames, TArray<FString> OverrideExtensions) const
+bool FAssetTypeActions_CSVAssetBase::CanExecuteFindSourceFileInExplorer(TArray<FString> Filenames, TArray<FString> OverrideExtensions) const
 {
 	// Verify that extensions were provided
 	if (OverrideExtensions.Num() == 0)
diff --git a/Engine/Source/Developer/AssetTools/Private/AssetTypeActions/AssetTypeActions_CurveTable.cpp b/Engine/Source/Developer/AssetTools/Private/AssetTypeActions/AssetTypeActions_CurveTable.cpp
--- a/Engine/Source/Developer/AssetTools/Private/AssetTypeActions/AssetTypeActions_CurveTable.cpp
+++ b/Engine/Source/Developer/AssetTools/Private/AssetTypeActions/AssetTypeActions_CurveTable.cpp
@@ -22,16 +22,18 @@ void FAssetTypeActions_CurveTable::GetActions( const TArray<UObject*>& InObjects
 		}
 	}
 
-	TArray<FString> XLSExtensions;
-	XLSExtensions.Add(TEXT(".xls"));
-	XLSExtensions.Add(TEXT(".xlsm"));
+	TArray<FString> PotentialFileExtensions;
+	PotentialFileExtensions.Add(TEXT(".xls"));
+	PotentialFileExtensions.Add(TEXT(".xlsm"));
+	PotentialFileExtensions.Add(TEXT(".csv"));
+	PotentialFileExtensions.Add(TEXT(".json"));
 	MenuBuilder.AddMenuEntry(
-		LOCTEXT("CurveTable_OpenSourceXLS", "Open Source (.xls/.xlsm)"),
-		LOCTEXT("CurveTable_OpenSourceXLSTooltip", "Opens the curve table's source XLS/XLSM file in an external editor."),
+		LOCTEXT("CurveTable_OpenSourceData", "Open Source Data"),
+		LOCTEXT("CurveTable_OpenSourceDataTooltip", "Opens the data table's source data file in an external editor. It will search using the following extensions: .xls/.xlsm/.csv/.json"),
 		FSlateIcon(),
 		FUIAction(
-			FExecuteAction::CreateSP( this, &FAssetTypeActions_CurveTable::ExecuteFindExcelFileInExplorer, ImportPaths, XLSExtensions ),
-			FCanExecuteAction::CreateSP(this, &FAssetTypeActions_CurveTable::CanExecuteFindExcelFileInExplorer, ImportPaths, XLSExtensions)
+			FExecuteAction::CreateSP( this, &FAssetTypeActions_CurveTable::ExecuteFindSourceFileInExplorer, ImportPaths, PotentialFileExtensions ),
+			FCanExecuteAction::CreateSP(this, &FAssetTypeActions_CurveTable::CanExecuteFindSourceFileInExplorer, ImportPaths, PotentialFileExtensions)
 			)
 		);
 }
diff --git a/Engine/Source/Developer/AssetTools/Private/AssetTypeActions/AssetTypeActions_DataTable.cpp b/Engine/Source/Developer/AssetTools/Private/AssetTypeActions/AssetTypeActions_DataTable.cpp
--- a/Engine/Source/Developer/AssetTools/Private/AssetTypeActions/AssetTypeActions_DataTable.cpp
+++ b/Engine/Source/Developer/AssetTools/Private/AssetTypeActions/AssetTypeActions_DataTable.cpp
@@ -5,6 +5,8 @@
 #include "Editor/DataTableEditor/Public/DataTableEditorModule.h"
 #include "Editor/DataTableEditor/Public/IDataTableEditor.h"
 #include "Engine/DataTable.h"
+#include "DesktopPlatformModule.h"
+#include "IMainFrameModule.h"
 
 #define LOCTEXT_NAMESPACE "AssetTypeActions"
 
@@ -23,37 +25,119 @@ void FAssetTypeActions_DataTable::GetActions( const TArray<UObject*>& InObjects,
 	}
 
 	MenuBuilder.AddMenuEntry(
-		LOCTEXT("DataTable_JSON", "JSON"),
-		LOCTEXT("DataTable_JSONTooltip", "Creates a JSON version of the data table in the lock."),
+		LOCTEXT("DataTable_ExportAsCSV", "Export as CSV"),
+		LOCTEXT("DataTable_ExportAsCSVTooltip", "Export the data table as a file containing CSV data."),
 		FSlateIcon(),
 		FUIAction(
-			FExecuteAction::CreateSP( this, &FAssetTypeActions_DataTable::ExecuteJSON, Tables ),
+			FExecuteAction::CreateSP( this, &FAssetTypeActions_DataTable::ExecuteExportAsCSV, Tables ),
 			FCanExecuteAction()
 			)
 		);
 
-	TArray<FString> XLSExtensions;
-	XLSExtensions.Add(TEXT(".xls"));
-	XLSExtensions.Add(TEXT(".xlsm"));
 	MenuBuilder.AddMenuEntry(
-		LOCTEXT("DataTable_OpenSourceXLS", "Open Source (.xls/.xlsm)"),
-		LOCTEXT("DataTable_OpenSourceXLSTooltip", "Opens the data table's source XLS/XLSM file in an external editor."),
+		LOCTEXT("DataTable_ExportAsJSON", "Export as JSON"),
+		LOCTEXT("DataTable_ExportAsJSONTooltip", "Export the data table as a file containing JSON data."),
 		FSlateIcon(),
 		FUIAction(
-			FExecuteAction::CreateSP( this, &FAssetTypeActions_DataTable::ExecuteFindExcelFileInExplorer, ImportPaths, XLSExtensions ),
-			FCanExecuteAction::CreateSP(this, &FAssetTypeActions_DataTable::CanExecuteFindExcelFileInExplorer, ImportPaths, XLSExtensions)
+			FExecuteAction::CreateSP( this, &FAssetTypeActions_DataTable::ExecuteExportAsJSON, Tables ),
+			FCanExecuteAction()
+			)
+		);
+
+	TArray<FString> PotentialFileExtensions;
+	PotentialFileExtensions.Add(TEXT(".xls"));
+	PotentialFileExtensions.Add(TEXT(".xlsm"));
+	PotentialFileExtensions.Add(TEXT(".csv"));
+	PotentialFileExtensions.Add(TEXT(".json"));
+	MenuBuilder.AddMenuEntry(
+		LOCTEXT("DataTable_OpenSourceData", "Open Source Data"),
+		LOCTEXT("DataTable_OpenSourceDataTooltip", "Opens the data table's source data file in an external editor. It will search using the following extensions: .xls/.xlsm/.csv/.json"),
+		FSlateIcon(),
+		FUIAction(
+			FExecuteAction::CreateSP( this, &FAssetTypeActions_DataTable::ExecuteFindSourceFileInExplorer, ImportPaths, PotentialFileExtensions ),
+			FCanExecuteAction::CreateSP(this, &FAssetTypeActions_DataTable::CanExecuteFindSourceFileInExplorer, ImportPaths, PotentialFileExtensions)
 			)
 		);
 }
 
-void FAssetTypeActions_DataTable::ExecuteJSON(TArray< TWeakObjectPtr<UObject> > Objects)
+void FAssetTypeActions_DataTable::ExecuteExportAsCSV(TArray< TWeakObjectPtr<UObject> > Objects)
 {
+	IDesktopPlatform* DesktopPlatform = FDesktopPlatformModule::Get();
+
+	void* ParentWindowWindowHandle = nullptr;
+
+	IMainFrameModule& MainFrameModule = FModuleManager::LoadModuleChecked<IMainFrameModule>(TEXT("MainFrame"));
+	const TSharedPtr<SWindow>& MainFrameParentWindow = MainFrameModule.GetParentWindow();
+	if ( MainFrameParentWindow.IsValid() && MainFrameParentWindow->GetNativeWindow().IsValid() )
+	{
+		ParentWindowWindowHandle = MainFrameParentWindow->GetNativeWindow()->GetOSWindowHandle();
+	}
+
 	for (auto ObjIt = Objects.CreateConstIterator(); ObjIt; ++ObjIt)
 	{
-		auto Table = Cast<UDataTable>((*ObjIt).Get());
-		if (Table != NULL)
+		auto DataTable = Cast<UDataTable>((*ObjIt).Get());
+		if (DataTable)
+		{
+			const FText Title = FText::Format(LOCTEXT("DataTable_ExportCSVDialogTitle", "Export '{0}' as CSV..."), FText::FromString(*DataTable->GetName()));
+			const FString CurrentFilename = (DataTable->ImportPath.IsEmpty()) ? TEXT("") : FReimportManager::ResolveImportFilename(DataTable->ImportPath, DataTable);
+			const FString FileTypes = TEXT("Data Table CSV (*.csv)|*.csv");
+
+			TArray<FString> OutFilenames;
+			DesktopPlatform->SaveFileDialog(
+				ParentWindowWindowHandle,
+				Title.ToString(),
+				(CurrentFilename.IsEmpty()) ? TEXT("") : FPaths::GetPath(CurrentFilename),
+				(CurrentFilename.IsEmpty()) ? TEXT("") : FPaths::GetBaseFilename(CurrentFilename) + TEXT(".csv"),
+				FileTypes,
+				EFileDialogFlags::None,
+				OutFilenames
+				);
+
+			if (OutFilenames.Num() > 0)
+			{
+				FFileHelper::SaveStringToFile(DataTable->GetTableAsCSV(), *OutFilenames[0]);
+			}
+		}
+	}
+}
+
+void FAssetTypeActions_DataTable::ExecuteExportAsJSON(TArray< TWeakObjectPtr<UObject> > Objects)
+{
+	IDesktopPlatform* DesktopPlatform = FDesktopPlatformModule::Get();
+
+	void* ParentWindowWindowHandle = nullptr;
+
+	IMainFrameModule& MainFrameModule = FModuleManager::LoadModuleChecked<IMainFrameModule>(TEXT("MainFrame"));
+	const TSharedPtr<SWindow>& MainFrameParentWindow = MainFrameModule.GetParentWindow();
+	if ( MainFrameParentWindow.IsValid() && MainFrameParentWindow->GetNativeWindow().IsValid() )
+	{
+		ParentWindowWindowHandle = MainFrameParentWindow->GetNativeWindow()->GetOSWindowHandle();
+	}
+
+	for (auto ObjIt = Objects.CreateConstIterator(); ObjIt; ++ObjIt)
+	{
+		auto DataTable = Cast<UDataTable>((*ObjIt).Get());
+		if (DataTable)
 		{
-			UE_LOG(LogDataTable, Log,  TEXT("JSON DataTable : %s"), *Table->GetTableAsJSON());
+			const FText Title = FText::Format(LOCTEXT("DataTable_ExportCSVDialogTitle", "Export '{0}' as JSON..."), FText::FromString(*DataTable->GetName()));
+			const FString CurrentFilename = (DataTable->ImportPath.IsEmpty()) ? TEXT("") : FReimportManager::ResolveImportFilename(DataTable->ImportPath, DataTable);
+			const FString FileTypes = TEXT("Data Table JSON (*.json)|*.json");
+
+			TArray<FString> OutFilenames;
+			DesktopPlatform->SaveFileDialog(
+				ParentWindowWindowHandle,
+				Title.ToString(),
+				(CurrentFilename.IsEmpty()) ? TEXT("") : FPaths::GetPath(CurrentFilename),
+				(CurrentFilename.IsEmpty()) ? TEXT("") : FPaths::GetBaseFilename(CurrentFilename) + TEXT(".json"),
+				FileTypes,
+				EFileDialogFlags::None,
+				OutFilenames
+				);
+
+			if (OutFilenames.Num() > 0)
+			{
+				FFileHelper::SaveStringToFile(DataTable->GetTableAsJSON(), *OutFilenames[0]);
+			}
 		}
 	}
 }
diff --git a/Engine/Source/Developer/AssetTools/Private/AssetTypeActions/AssetTypeActions_DataTable.h b/Engine/Source/Developer/AssetTools/Private/AssetTypeActions/AssetTypeActions_DataTable.h
--- a/Engine/Source/Developer/AssetTools/Private/AssetTypeActions/AssetTypeActions_DataTable.h
+++ b/Engine/Source/Developer/AssetTools/Private/AssetTypeActions/AssetTypeActions_DataTable.h
@@ -16,6 +16,9 @@ public:
 	// End IAssetTypeActions
 
 private:
+	/** Handler for when CSV is selected */
+	void ExecuteExportAsCSV(TArray< TWeakObjectPtr<UObject> > Objects);
+
 	/** Handler for when JSON is selected */
-	void ExecuteJSON(TArray< TWeakObjectPtr<UObject> > Objects);
+	void ExecuteExportAsJSON(TArray< TWeakObjectPtr<UObject> > Objects);
 };
\ No newline at end of file
diff --git a/Engine/Source/Developer/AssetTools/Public/AssetTypeActions_CSVAssetBase.h b/Engine/Source/Developer/AssetTools/Public/AssetTypeActions_CSVAssetBase.h
--- a/Engine/Source/Developer/AssetTools/Public/AssetTypeActions_CSVAssetBase.h
+++ b/Engine/Source/Developer/AssetTools/Public/AssetTypeActions_CSVAssetBase.h
@@ -13,11 +13,11 @@ public:
 	virtual bool IsImportedAsset() const override { return true; }
 
 protected:
-	/** Handler for opening the xls/xlsm version of the source file */
-	void ExecuteFindExcelFileInExplorer(TArray<FString> Filenames, TArray<FString> OverrideExtensions);
+	/** Handler for opening the source file for this asset */
+	void ExecuteFindSourceFileInExplorer(TArray<FString> Filenames, TArray<FString> OverrideExtensions);
 
-	/** Determine whether the find excel file in explorer editor command can execute or not */
-	bool CanExecuteFindExcelFileInExplorer(TArray<FString> Filenames, TArray<FString> OverrideExtensions) const;
+	/** Determine whether the find source file in explorer editor command can execute or not */
+	bool CanExecuteFindSourceFileInExplorer(TArray<FString> Filenames, TArray<FString> OverrideExtensions) const;
 
 	/** Verify the specified filename exists */
 	bool VerifyFileExists(const FString& InFileName) const;
diff --git a/Engine/Source/Editor/DataTableEditor/Private/SRowEditor.cpp b/Engine/Source/Editor/DataTableEditor/Private/SRowEditor.cpp
--- a/Engine/Source/Editor/DataTableEditor/Private/SRowEditor.cpp
+++ b/Engine/Source/Editor/DataTableEditor/Private/SRowEditor.cpp
@@ -6,6 +6,7 @@
 #include "PropertyEditorModule.h"
 #include "IStructureDetailsView.h"
 #include "IDetailsView.h"
+#include "DataTableUtils.h"
 #include "DataTableEditorUtils.h"
 
 #define LOCTEXT_NAMESPACE "SRowEditor"
@@ -246,7 +247,7 @@ FReply SRowEditor::OnAddClicked()
 {
 	if (DataTable.IsValid())
 	{
-		FName NewName = DataTable->MakeValidName(TEXT("NewRow"));
+		FName NewName = DataTableUtils::MakeValidName(TEXT("NewRow"));
 		const TArray<FName> ExisitngNames = DataTable->GetRowNames();
 		while (ExisitngNames.Contains(NewName))
 		{
@@ -272,7 +273,7 @@ void SRowEditor::OnRowRenamed(const FText& Text, ETextCommit::Type CommitType)
 {
 	if (!GetCurrentNameAsText().EqualTo(Text) && DataTable.IsValid())
 	{
-		const FName NewName = DataTable->MakeValidName(Text.ToString());
+		const FName NewName = DataTableUtils::MakeValidName(Text.ToString());
 		for (auto Name : CachedRowNames)
 		{
 			if (Name.IsValid() && (*Name == NewName))
diff --git a/Engine/Source/Editor/UnrealEd/Classes/Factories/CSVImportFactory.h b/Engine/Source/Editor/UnrealEd/Classes/Factories/CSVImportFactory.h
--- a/Engine/Source/Editor/UnrealEd/Classes/Factories/CSVImportFactory.h
+++ b/Engine/Source/Editor/UnrealEd/Classes/Factories/CSVImportFactory.h
@@ -24,10 +24,14 @@ class UNREALED_API UCSVImportFactory : public UFactory
 	/* Reimport an object that was created based on a CSV */
 	bool ReimportCSV(UObject* Obj);
 
+protected:
+	virtual TArray<FString> DoImportDataTable(class UDataTable* TargetDataTable, const FString& DataToImport);
+	virtual TArray<FString> DoImportCurveTable(class UCurveTable* TargetCurveTable, const FString& DataToImport, const ERichCurveInterpMode ImportCurveInterpMode);
+	virtual TArray<FString> DoImportCurve(class UCurveBase* TargetCurve, const FString& DataToImport);
+
 private:
 	/* Reimport object from the given path*/
 	bool Reimport(UObject* Obj, const FString& Path  );
-
 };
 
 /** Enum to indicate what to import CSV as */
diff --git a/Engine/Source/Editor/UnrealEd/Classes/Factories/ReimportDataTableFactory.h b/Engine/Source/Editor/UnrealEd/Classes/Factories/ReimportDataTableFactory.h
--- a/Engine/Source/Editor/UnrealEd/Classes/Factories/ReimportDataTableFactory.h
+++ b/Engine/Source/Editor/UnrealEd/Classes/Factories/ReimportDataTableFactory.h
@@ -14,6 +14,10 @@ class UReimportDataTableFactory : public UCSVImportFactory, public FReimportHand
 	virtual EReimportResult::Type Reimport( UObject* Obj ) override;
 	virtual int32 GetPriority() const override;
 	// End FReimportHandler interface
+
+	// Begin UCSVImportFactory interface
+	virtual TArray<FString> DoImportDataTable(class UDataTable* TargetDataTable, const FString& DataToImport) override;
+	// End UCSVImportFactory interface
 };
 
 
diff --git a/Engine/Source/Editor/UnrealEd/Private/Factories/CSVImportFactory.cpp b/Engine/Source/Editor/UnrealEd/Private/Factories/CSVImportFactory.cpp
--- a/Engine/Source/Editor/UnrealEd/Private/Factories/CSVImportFactory.cpp
+++ b/Engine/Source/Editor/UnrealEd/Private/Factories/CSVImportFactory.cpp
@@ -157,7 +157,7 @@ UObject* UCSVImportFactory::FactoryCreateText( UClass* InClass, UObject* InParen
 			NewTable->RowStruct = ImportRowStruct;
 			NewTable->ImportPath = FReimportManager::SanitizeImportFilename(CurrentFilename, NewTable);
 			// Go ahead and create table from string
-			Problems = NewTable->CreateTableFromCSVString(String);
+			Problems = DoImportDataTable(NewTable, String);
 
 			// Print out
 			UE_LOG(LogCSVImportFactory, Log, TEXT("Imported DataTable '%s' - %d Problems"), *InName.ToString(), Problems.Num());
@@ -176,7 +176,7 @@ UObject* UCSVImportFactory::FactoryCreateText( UClass* InClass, UObject* InParen
 			NewTable->ImportPath = FReimportManager::SanitizeImportFilename(CurrentFilename, NewTable);
 
 			// Go ahead and create table from string
-			Problems = NewTable->CreateTableFromCSVString(String, ImportCurveInterpMode);
+			Problems = DoImportCurveTable(NewTable, String, ImportCurveInterpMode);
 
 			// Print out
 			UE_LOG(LogCSVImportFactory, Log, TEXT("Imported CurveTable '%s' - %d Problems"), *InName.ToString(), Problems.Num());
@@ -189,7 +189,7 @@ UObject* UCSVImportFactory::FactoryCreateText( UClass* InClass, UObject* InParen
 			// Create/reset curve
 			UCurveBase* NewCurve = NewObject<UCurveBase>(InParent, CurveClass, InName, Flags);
 
-			Problems = NewCurve->CreateCurveFromCSVString(String);
+			Problems = DoImportCurve(NewCurve, String);
 
 			UE_LOG(LogCSVImportFactory, Log, TEXT("Imported Curve '%s' - %d Problems"), *InName.ToString(), Problems.Num());
 			NewCurve->ImportPath = FReimportManager::SanitizeImportFilename(CurrentFilename, NewCurve);
@@ -254,11 +254,27 @@ bool UCSVImportFactory::Reimport( UObject* Obj, const FString& Path )
 	return false;
 }
 
+TArray<FString> UCSVImportFactory::DoImportDataTable(UDataTable* TargetDataTable, const FString& DataToImport)
+{
+	return TargetDataTable->CreateTableFromCSVString(DataToImport);
+}
+
+TArray<FString> UCSVImportFactory::DoImportCurveTable(UCurveTable* TargetCurveTable, const FString& DataToImport, const ERichCurveInterpMode ImportCurveInterpMode)
+{
+	return TargetCurveTable->CreateTableFromCSVString(DataToImport, ImportCurveInterpMode);
+}
+
+TArray<FString> UCSVImportFactory::DoImportCurve(UCurveBase* TargetCurve, const FString& DataToImport)
+{
+	return TargetCurve->CreateCurveFromCSVString(DataToImport);
+}
+
 //////////////////////////////////////////////////////////////////////////
 
 UReimportDataTableFactory::UReimportDataTableFactory(const FObjectInitializer& ObjectInitializer)
 	: Super(ObjectInitializer)
 {
+	Formats.Add(TEXT("json;JavaScript Object Notation"));
 }
 
 bool UReimportDataTableFactory::CanReimport( UObject* Obj, TArray<FString>& OutFilenames )
@@ -298,6 +314,19 @@ int32 UReimportDataTableFactory::GetPriority() const
 	return ImportPriority;
 }
 
+TArray<FString> UReimportDataTableFactory::DoImportDataTable(UDataTable* TargetDataTable, const FString& DataToImport)
+{
+	// Are we importing JSON data?
+	const bool bIsJSON = CurrentFilename.EndsWith(TEXT(".json"));
+	if (bIsJSON)
+	{
+		return TargetDataTable->CreateTableFromJSONString(DataToImport);
+	}
+
+	// Fallback to CSV
+	return UCSVImportFactory::DoImportDataTable(TargetDataTable, DataToImport);
+}
+
 ////////////////////////////////////////////////////////////////////////////
 //
 UReimportCurveTableFactory::UReimportCurveTableFactory(const FObjectInitializer& ObjectInitializer)
diff --git a/Engine/Source/Runtime/Engine/Classes/Engine/DataTable.h b/Engine/Source/Runtime/Engine/Classes/Engine/DataTable.h
--- a/Engine/Source/Runtime/Engine/Classes/Engine/DataTable.h
+++ b/Engine/Source/Runtime/Engine/Classes/Engine/DataTable.h
@@ -145,7 +145,10 @@ public:
 	ENGINE_API void RestoreAfterStructChange();
 
 	/** Output entire contents of table as a string */
-	ENGINE_API FString GetTableAsString();
+	ENGINE_API FString GetTableAsString() const;
+
+	/** Output entire contents of table as CSV */
+	ENGINE_API FString GetTableAsCSV() const;
 
 	/** Output entire contents of table as JSON */
 	ENGINE_API FString GetTableAsJSON() const;
@@ -179,16 +182,6 @@ public:
 	ENGINE_API TArray< TArray<FString> > GetTableData() const;
 #endif //WITH_EDITOR
 
-	ENGINE_API static FString AssignStringToProperty(const FString& InString, const UProperty* InProp, uint8* InData);
-
-	ENGINE_API static FString GetPropertyValueAsString(const UProperty* InProp, uint8* InData);
-
-	ENGINE_API static TArray<FName> GetStructPropertyNames(UStruct* InStruct);
-
-	ENGINE_API static FName MakeValidName(const FString& InString);
-
-	ENGINE_API static bool IsSupportedTableProperty(const UProperty* InProp);
-
 	// End UDataTable interface
 
 private:
diff --git a/Engine/Source/Runtime/Engine/Private/DataTable.cpp b/Engine/Source/Runtime/Engine/Private/DataTable.cpp
--- a/Engine/Source/Runtime/Engine/Private/DataTable.cpp
+++ b/Engine/Source/Runtime/Engine/Private/DataTable.cpp
@@ -1,10 +1,10 @@
 // Copyright 1998-2015 Epic Games, Inc. All Rights Reserved.
 
 #include "EnginePrivate.h"
-#include "Json.h"
-#include "CsvParser.h"
 #include "Engine/DataTable.h"
-#include "Engine/UserDefinedStruct.h"
+#include "DataTableUtils.h"
+#include "DataTableCSV.h"
+#include "DataTableJSON.h"
 
 DEFINE_LOG_CATEGORY(LogDataTable);
 
@@ -16,86 +16,6 @@ UDataTable::UDataTable(const FObjectInitializer& ObjectInitializer)
 {
 }
 
-/** Util that removes invalid chars and then make an FName */
-FName UDataTable::MakeValidName(const FString& InString)
-{
-	FString InvalidChars(INVALID_NAME_CHARACTERS);
-
-	FString FixedString;
-	TArray<TCHAR>& FixedCharArray = FixedString.GetCharArray();
-
-	// Iterate over input string characters
-	for(int32 CharIdx=0; CharIdx<InString.Len(); CharIdx++)
-	{
-		// See if this char occurs in the InvalidChars string
-		FString Char = InString.Mid( CharIdx, 1 );
-		if( !InvalidChars.Contains(Char) )
-		{
-			// Its ok, add to result
-			FixedCharArray.Add(Char[0]);
-		}
-	}
-	FixedCharArray.Add(0);
-
-	return FName(*FixedString);
-}
-
-/** Util to see if this property is supported in a row struct. */
-bool UDataTable::IsSupportedTableProperty(const UProperty* InProp)
-{
-	return(	InProp->IsA(UIntProperty::StaticClass()) || 
-			InProp->IsA(UFloatProperty::StaticClass()) ||
-			InProp->IsA(UNameProperty::StaticClass()) ||
-			InProp->IsA(UStrProperty::StaticClass()) ||
-			InProp->IsA(UBoolProperty::StaticClass()) ||
-			InProp->IsA(UObjectPropertyBase::StaticClass()) ||
-			InProp->IsA(UStructProperty::StaticClass()) ||
-			InProp->IsA(UByteProperty::StaticClass()) ||
-			InProp->IsA(UTextProperty::StaticClass()));
-}
-
-
-
-/** Util to assign a value (given as a string) to a struct property. */
-FString UDataTable::AssignStringToProperty(const FString& InString, const UProperty* InProp, uint8* InData)
-{
-	FStringOutputDevice ImportError;
-	if(InProp != NULL && IsSupportedTableProperty(InProp))
-	{
-		InProp->ImportText(*InString, InProp->ContainerPtrToValuePtr<uint8>(InData), PPF_None, NULL, &ImportError);
-	}
-
-	FString Error = ImportError;
-	return Error;
-}
-
-/** Util to assign get a property as a string */
-FString UDataTable::GetPropertyValueAsString(const UProperty* InProp, uint8* InData)
-{
-	FString Result(TEXT(""));
-
-	if(InProp != NULL && IsSupportedTableProperty(InProp))
-	{
-		InProp->ExportText_InContainer(0, Result, InData, InData, NULL, PPF_None);
-	}
-
-	return Result;
-}
-
-/** Util to get all property names from a struct */
-TArray<FName> UDataTable::GetStructPropertyNames(UStruct* InStruct)
-{
-	TArray<FName> PropNames;
-	for (TFieldIterator<UProperty> It(InStruct); It; ++It)
-	{
-		PropNames.Add(It->GetFName());
-	}
-	return PropNames;
-}
-
-
-//////////////////////////////////////////////////////////////////////////
-
 void UDataTable::LoadStructData(FArchive& Ar)
 {
 	UScriptStruct* LoadUsingStruct = RowStruct;
@@ -250,7 +170,7 @@ UProperty* UDataTable::FindTableProperty(const FName& PropertyName) const
 			break;
 		}
 	}
-	if (!IsSupportedTableProperty(Property))
+	if (!DataTableUtils::IsSupportedTableProperty(Property))
 	{
 		Property = NULL;
 	}
@@ -260,15 +180,6 @@ UProperty* UDataTable::FindTableProperty(const FName& PropertyName) const
 
 #if WITH_EDITOR
 
-struct FPropertyDisplayNameHelper
-{
-	static inline FString Get(const UProperty* Prop, const FString& DefaultName)
-	{
-		static const FName DisplayNameKey(TEXT("DisplayName"));
-		return (Prop && Prop->HasMetaData(DisplayNameKey)) ? Prop->GetMetaData(DisplayNameKey) : DefaultName;
-	}
-};
-
 void UDataTable::CleanBeforeStructChange()
 {
 	RowsSerializedWithTags.Reset();
@@ -320,7 +231,7 @@ void UDataTable::RestoreAfterStructChange()
 	RowsSerializedWithTags.Empty();
 }
 
-FString UDataTable::GetTableAsString()
+FString UDataTable::GetTableAsString() const
 {
 	FString Result;
 
@@ -347,7 +258,7 @@ FString UDataTable::GetTableAsString()
 		Result += TEXT("\n");
 
 		// Now iterate over rows
-		for ( auto RowIt = RowMap.CreateIterator(); RowIt; ++RowIt )
+		for ( auto RowIt = RowMap.CreateConstIterator(); RowIt; ++RowIt )
 		{
 			FName RowName = RowIt.Key();
 			Result += RowName.ToString();
@@ -356,7 +267,7 @@ FString UDataTable::GetTableAsString()
 			for(int32 PropIdx=0; PropIdx<StructProps.Num(); PropIdx++)
 			{
 				Result += TEXT(",");
-				Result += GetPropertyValueAsString(StructProps[PropIdx], RowData);
+				Result += DataTableUtils::GetPropertyValueAsString(StructProps[PropIdx], RowData);
 			}
 			Result += TEXT("\n");			
 		}
@@ -368,80 +279,34 @@ FString UDataTable::GetTableAsString()
 	return Result;
 }
 
-FString UDataTable::GetTableAsJSON() const
+FString UDataTable::GetTableAsCSV() const
 {
-	// use the pretty print policy since these values are usually getting dumpped for check-in to P4 (or for inspection)
 	FString Result;
-	TSharedRef< TJsonWriter<TCHAR, TPrettyJsonPrintPolicy<TCHAR> > > JsonWriter = TJsonWriterFactory<TCHAR, TPrettyJsonPrintPolicy<TCHAR> >::Create(&Result);
-	if (!WriteTableAsJSON(JsonWriter))
+	if (!FDataTableExporterCSV(*this, Result).WriteTable())
 	{
-		return TEXT("Missing RowStruct!\n");
+		Result = TEXT("Missing RowStruct!\n");
 	}
-	JsonWriter->Close();
 	return Result;
 }
 
-bool UDataTable::WriteRowAsJSON(const TSharedRef< TJsonWriter<TCHAR, TPrettyJsonPrintPolicy<TCHAR> > >& JsonWriter, const void* RowData) const
+FString UDataTable::GetTableAsJSON() const
 {
-	if (RowStruct == NULL)
-	{
-		return false;
-	}
-	for (TFieldIterator<UProperty> It(RowStruct); It; ++It)
+	FString Result;
+	if (!FDataTableExporterJSON(*this, Result).WriteTable())
 	{
-		UProperty* BaseProp = *It;
-		check(BaseProp != NULL);
-		const void* Data = BaseProp->ContainerPtrToValuePtr<void>(RowData, 0);
-		if (UNumericProperty *NumProp = Cast<UNumericProperty>(BaseProp))
-		{
-			if (NumProp->IsInteger())
-			{
-				JsonWriter->WriteValue(BaseProp->GetName(), NumProp->GetSignedIntPropertyValue(Data));
-			}
-			else
-			{
-				JsonWriter->WriteValue(BaseProp->GetName(), NumProp->GetFloatingPointPropertyValue(Data));
-			}
-		}
-		else if (UBoolProperty* BoolProp = Cast<UBoolProperty>(BaseProp))
-		{
-			JsonWriter->WriteValue(BaseProp->GetName(), BoolProp->GetPropertyValue(Data));
-		}
-		else
-		{
-			FString PropertyValue = GetPropertyValueAsString(BaseProp, (uint8*)RowData);
-			JsonWriter->WriteValue(BaseProp->GetName(), PropertyValue);
-		}
+		Result = TEXT("Missing RowStruct!\n");
 	}
-	return true;
+	return Result;
 }
 
-bool UDataTable::WriteTableAsJSON(const TSharedRef< TJsonWriter<TCHAR, TPrettyJsonPrintPolicy<TCHAR> > >& JsonWriter) const
+bool UDataTable::WriteRowAsJSON(const TSharedRef< TJsonWriter<TCHAR, TPrettyJsonPrintPolicy<TCHAR> > >& JsonWriter, const void* RowData) const
 {
-	if (RowStruct == NULL)
-	{
-		return false;
-	}
-	JsonWriter->WriteArrayStart();
-
-	// Iterate over rows
-	for (auto RowIt = RowMap.CreateConstIterator(); RowIt; ++RowIt)
-	{
-		JsonWriter->WriteObjectStart();
-		{
-			//RowName
-			FName RowName = RowIt.Key();
-			JsonWriter->WriteValue(TEXT("Name"), RowName.ToString());
-
-			//Now the values
-			uint8* RowData = RowIt.Value();
-			WriteRowAsJSON(JsonWriter, RowData);
-		}
-		JsonWriter->WriteObjectEnd();
-	}
+	return FDataTableExporterJSON(*this, JsonWriter).WriteRow(RowData);
+}
 
-	JsonWriter->WriteArrayEnd();
-	return true;
+bool UDataTable::WriteTableAsJSON(const TSharedRef< TJsonWriter<TCHAR, TPrettyJsonPrintPolicy<TCHAR> > >& JsonWriter) const
+{
+	return FDataTableExporterJSON(*this, JsonWriter).WriteTable();
 }
 
 /** Get array of UProperties that corresponds to columns in the table */
@@ -450,7 +315,7 @@ TArray<UProperty*> UDataTable::GetTablePropertyArray(const TArray<const TCHAR*>&
 	TArray<UProperty*> ColumnProps;
 
 	// Get list of all expected properties from the struct
-	TArray<FName> ExpectedPropNames = GetStructPropertyNames(InRowStruct);
+	TArray<FName> ExpectedPropNames = DataTableUtils::GetStructPropertyNames(InRowStruct);
 
 	// Need at least 2 columns, first column is skipped, will contain row names
 	if(Cells.Num() > 1)
@@ -463,7 +328,7 @@ TArray<UProperty*> UDataTable::GetTablePropertyArray(const TArray<const TCHAR*>&
 		{
 			const TCHAR* ColumnValue = Cells[ColIdx];
 
-			FName PropName = MakeValidName(ColumnValue);
+			FName PropName = DataTableUtils::MakeValidName(ColumnValue);
 			if(PropName == NAME_None)
 			{
 				OutProblems.Add(FString::Printf(TEXT("Missing name for column %d."), ColIdx));
@@ -474,7 +339,7 @@ TArray<UProperty*> UDataTable::GetTablePropertyArray(const TArray<const TCHAR*>&
 
 				for (TFieldIterator<UProperty> It(InRowStruct); It && !ColumnProp; ++It)
 				{
-					const auto DisplayName = FPropertyDisplayNameHelper::Get(*It, FString());
+					const auto DisplayName = DataTableUtils::GetPropertyDisplayName(*It, FString());
 					ColumnProp = (!DisplayName.IsEmpty() && DisplayName == ColumnValue) ? *It : NULL;
 				}
 
@@ -492,7 +357,7 @@ TArray<UProperty*> UDataTable::GetTablePropertyArray(const TArray<const TCHAR*>&
 						OutProblems.Add(FString::Printf(TEXT("Duplicate column '%s'."), *ColumnProp->GetName()));
 					}
 					// Check we support this property type
-					else if( !IsSupportedTableProperty(ColumnProp) )
+					else if( !DataTableUtils::IsSupportedTableProperty(ColumnProp) )
 					{
 						OutProblems.Add(FString::Printf(TEXT("Unsupported Property type for struct member '%s'."), *ColumnProp->GetName()));
 					}
@@ -513,7 +378,7 @@ TArray<UProperty*> UDataTable::GetTablePropertyArray(const TArray<const TCHAR*>&
 	for(int32 PropIdx=0; PropIdx < ExpectedPropNames.Num(); PropIdx++)
 	{
 		const UProperty* const ColumnProp = FindField<UProperty>(InRowStruct, ExpectedPropNames[PropIdx]);
-		const FString DisplayName = FPropertyDisplayNameHelper::Get(ColumnProp, ExpectedPropNames[PropIdx].ToString());
+		const FString DisplayName = DataTableUtils::GetPropertyDisplayName(ColumnProp, ExpectedPropNames[PropIdx].ToString());
 		OutProblems.Add(FString::Printf(TEXT("Expected column '%s' not found in input."), *DisplayName));
 	}
 
@@ -525,125 +390,21 @@ TArray<FString> UDataTable::CreateTableFromCSVString(const FString& InString)
 	// Array used to store problems about table creation
 	TArray<FString> OutProblems;
 
-	// Check we have a RowStruct specified
-	if(RowStruct == NULL)
-	{
-		OutProblems.Add(FString(TEXT("No RowStruct specified.")));
-		return OutProblems;
-	}
-	if (InString.IsEmpty())
-	{
-		OutProblems.Add(FString(TEXT("Input data is empty.")));
-		return OutProblems;
-	}
-
-	const FCsvParser Parser(InString);
-	const auto& Rows = Parser.GetRows();
-
-	// Must have at least 2 rows (column names + data)
-	if(Rows.Num() <= 1)
-	{
-		OutProblems.Add(FString(TEXT("Too few rows.")));
-		return OutProblems;
-	}
-
-	// Find property for each column
-	TArray<UProperty*> ColumnProps = GetTablePropertyArray(Rows[0], RowStruct, OutProblems);
-
-	// Empty existing data
-	EmptyTable();
+	FDataTableImporterCSV(*this, InString, OutProblems).ReadTable();
 
-	// Iterate over rows
-	for(int32 RowIdx=1; RowIdx<Rows.Num(); RowIdx++)
-	{
-		const TArray<const TCHAR*>& Cells = Rows[RowIdx];
-
-		// Need at least 1 cells (row name)
-		if(Cells.Num() < 1)
-		{
-			OutProblems.Add(FString::Printf(TEXT("Row '%d' has too few cells."), RowIdx));
-			continue;
-		}
-
-		// Need enough columns in the properties!
-		if( ColumnProps.Num() < Cells.Num() )
-		{
-			OutProblems.Add(FString::Printf(TEXT("Row '%d' has more cells than properties, is there a malformed string?"), RowIdx));
-			continue;
-		}
-
-		// Get row name
-		FName RowName = MakeValidName(Cells[0]);
-
-		// Check its not 'none'
-		if(RowName == NAME_None)
-		{
-			OutProblems.Add(FString::Printf(TEXT("Row '%d' missing a name."), RowIdx));
-			continue;
-		}
-
-		// Check its not a duplicate
-		if(RowMap.Find(RowName) != NULL)
-		{
-			OutProblems.Add(FString::Printf(TEXT("Duplicate row name '%s'."), *RowName.ToString()));
-			continue;
-		}
-
-		// Allocate data to store information, using UScriptStruct to know its size
-		uint8* RowData = (uint8*)FMemory::Malloc(RowStruct->PropertiesSize);
-		RowStruct->InitializeStruct(RowData);
-		// And be sure to call DestroyScriptStruct later
-
-#if WITH_EDITOR
-		if (auto UDStruct = Cast<const UUserDefinedStruct>(RowStruct))
-		{
-			UDStruct->InitializeDefaultValue(RowData);
-		}
-#endif // WITH_EDITOR
-
-		// Add to row map
-		RowMap.Add(RowName, RowData);
-
-		// Now iterate over cells (skipping first cell, that was row name)
-		for(int32 CellIdx=1; CellIdx<Cells.Num(); CellIdx++)
-		{
-			// Try and assign string to data using the column property
-			UProperty* ColumnProp = ColumnProps[CellIdx];
-			const FString CellValue = Cells[CellIdx];
-			FString Error = AssignStringToProperty(CellValue, ColumnProp, RowData);
-
-			// If we failed, output a problem string
-			if(Error.Len() > 0)
-			{
-				FString ColumnName = (ColumnProp != NULL) 
-					? FPropertyDisplayNameHelper::Get(ColumnProp, ColumnProp->GetName())
-					: FString(TEXT("NONE"));
-				OutProblems.Add(FString::Printf(TEXT("Problem assigning string '%s' to property '%s' on row '%s' : %s"), *CellValue, *ColumnName, *RowName.ToString(), *Error));
-			}
-		}
-
-		// Problem if we didn't have enough cells on this row
-		if(Cells.Num() < ColumnProps.Num())
-		{
-			OutProblems.Add(FString::Printf(TEXT("Too few cells on row '%s'."), *RowName.ToString()));			
-		}
-	}
-
-	Modify(true);
 	return OutProblems;
 }
 
-
-//NOT READY YET, PLACEHOLDER
 TArray<FString> UDataTable::CreateTableFromJSONString(const FString& InString)
 {
 	// Array used to store problems about table creation
 	TArray<FString> OutProblems;
 
+	FDataTableImporterJSON(*this, InString, OutProblems).ReadTable();
+
 	return OutProblems;
 }
 
-
 TArray<FString> UDataTable::GetColumnTitles() const
 {
 	TArray<FString> Result;
@@ -652,7 +413,7 @@ TArray<FString> UDataTable::GetColumnTitles() const
 	{
 		UProperty* Prop = *It;
 		check(Prop != NULL);
-		const FString DisplayName = FPropertyDisplayNameHelper::Get(Prop, Prop->GetName());
+		const FString DisplayName = DataTableUtils::GetPropertyDisplayName(Prop, Prop->GetName());
 		Result.Add(DisplayName);
 	}
 	return Result;
@@ -683,7 +444,7 @@ TArray< TArray<FString> > UDataTable::GetTableData() const
 		 uint8* RowData = RowIt.Value();
 		 for(int32 PropIdx=0; PropIdx<StructProps.Num(); PropIdx++)
 		 {
-			 RowResult.Add(GetPropertyValueAsString(StructProps[PropIdx], RowData));
+			 RowResult.Add(DataTableUtils::GetPropertyValueAsString(StructProps[PropIdx], RowData));
 		 }
 		 Result.Add(RowResult);
 	 }
@@ -698,4 +459,4 @@ TArray<FName> UDataTable::GetRowNames() const
 	TArray<FName> Keys;
 	RowMap.GetKeys(Keys);
 	return Keys;
-}
\ No newline at end of file
+}
diff --git a/Engine/Source/Runtime/Engine/Private/DataTableCSV.cpp b/Engine/Source/Runtime/Engine/Private/DataTableCSV.cpp
new file mode 100644
--- /dev/null
+++ b/Engine/Source/Runtime/Engine/Private/DataTableCSV.cpp
@@ -0,0 +1,209 @@
+// Copyright 1998-2015 Epic Games, Inc. All Rights Reserved.
+
+#include "EnginePrivate.h"
+#include "DataTableCSV.h"
+#include "CsvParser.h"
+#include "Engine/DataTable.h"
+#include "Engine/UserDefinedStruct.h"
+#include "DataTableUtils.h"
+
+#if WITH_EDITOR
+
+FDataTableExporterCSV::FDataTableExporterCSV(const UDataTable& InDataTable, FString& OutExportText)
+	: DataTable(&InDataTable)
+	, ExportedText(OutExportText)
+{
+}
+
+FDataTableExporterCSV::~FDataTableExporterCSV()
+{
+}
+
+bool FDataTableExporterCSV::WriteTable()
+{
+	if (!DataTable->RowStruct)
+	{
+		return false;
+	}
+
+	// Write the header (column titles)
+	ExportedText += TEXT("---");
+	for (TFieldIterator<UProperty> It(DataTable->RowStruct); It; ++It)
+	{
+		UProperty* BaseProp = *It;
+		check(BaseProp);
+
+		ExportedText += TEXT(",");
+		ExportedText += BaseProp->GetName();
+	}
+	ExportedText += TEXT("\n");
+
+	// Write each row
+	for (auto RowIt = DataTable->RowMap.CreateConstIterator(); RowIt; ++RowIt)
+	{
+		FName RowName = RowIt.Key();
+		ExportedText += RowName.ToString();
+
+		uint8* RowData = RowIt.Value();
+		WriteRow(RowData);
+
+		ExportedText += TEXT("\n");
+	}
+
+	return true;
+}
+
+bool FDataTableExporterCSV::WriteRow(const void* InRowData)
+{
+	if (!DataTable->RowStruct)
+	{
+		return false;
+	}
+
+	for (TFieldIterator<UProperty> It(DataTable->RowStruct); It; ++It)
+	{
+		UProperty* BaseProp = *It;
+		check(BaseProp);
+
+		const void* Data = BaseProp->ContainerPtrToValuePtr<void>(InRowData, 0);
+		WriteStructEntry(InRowData, BaseProp, Data);
+	}
+
+	return true;
+}
+
+bool FDataTableExporterCSV::WriteStructEntry(const void* InRowData, UProperty* InProperty, const void* InPropertyData)
+{
+	ExportedText += TEXT(",");
+
+	const FString PropertyValue = DataTableUtils::GetPropertyValueAsString(InProperty, (uint8*)InRowData);
+	ExportedText += TEXT("\"");
+	ExportedText += PropertyValue.Replace(TEXT("\""), TEXT("\"\""));
+	ExportedText += TEXT("\"");
+
+	return true;
+}
+
+
+FDataTableImporterCSV::FDataTableImporterCSV(UDataTable& InDataTable, FString InCSVData, TArray<FString>& OutProblems)
+	: DataTable(&InDataTable)
+	, CSVData(MoveTemp(InCSVData))
+	, ImportProblems(OutProblems)
+{
+}
+
+FDataTableImporterCSV::~FDataTableImporterCSV()
+{
+}
+
+bool FDataTableImporterCSV::ReadTable()
+{
+	if (CSVData.IsEmpty())
+	{
+		ImportProblems.Add(TEXT("Input data is empty."));
+		return false;
+	}
+
+	// Check we have a RowStruct specified
+	if (!DataTable->RowStruct)
+	{
+		ImportProblems.Add(TEXT("No RowStruct specified."));
+		return false;
+	}
+
+	const FCsvParser Parser(CSVData);
+	const auto& Rows = Parser.GetRows();
+
+	// Must have at least 2 rows (column names + data)
+	if(Rows.Num() <= 1)
+	{
+		ImportProblems.Add(TEXT("Too few rows."));
+		return false;
+	}
+
+	// Find property for each column
+	TArray<UProperty*> ColumnProps = DataTable->GetTablePropertyArray(Rows[0], DataTable->RowStruct, ImportProblems);
+
+	// Empty existing data
+	DataTable->EmptyTable();
+
+	// Iterate over rows
+	for(int32 RowIdx=1; RowIdx<Rows.Num(); RowIdx++)
+	{
+		const TArray<const TCHAR*>& Cells = Rows[RowIdx];
+
+		// Need at least 1 cells (row name)
+		if(Cells.Num() < 1)
+		{
+			ImportProblems.Add(FString::Printf(TEXT("Row '%d' has too few cells."), RowIdx));
+			continue;
+		}
+
+		// Need enough columns in the properties!
+		if( ColumnProps.Num() < Cells.Num() )
+		{
+			ImportProblems.Add(FString::Printf(TEXT("Row '%d' has more cells than properties, is there a malformed string?"), RowIdx));
+			continue;
+		}
+
+		// Get row name
+		FName RowName = DataTableUtils::MakeValidName(Cells[0]);
+
+		// Check its not 'none'
+		if(RowName == NAME_None)
+		{
+			ImportProblems.Add(FString::Printf(TEXT("Row '%d' missing a name."), RowIdx));
+			continue;
+		}
+
+		// Check its not a duplicate
+		if(DataTable->RowMap.Find(RowName) != NULL)
+		{
+			ImportProblems.Add(FString::Printf(TEXT("Duplicate row name '%s'."), *RowName.ToString()));
+			continue;
+		}
+
+		// Allocate data to store information, using UScriptStruct to know its size
+		uint8* RowData = (uint8*)FMemory::Malloc(DataTable->RowStruct->PropertiesSize);
+		DataTable->RowStruct->InitializeStruct(RowData);
+		// And be sure to call DestroyScriptStruct later
+
+		if (auto UDStruct = Cast<const UUserDefinedStruct>(DataTable->RowStruct))
+		{
+			UDStruct->InitializeDefaultValue(RowData);
+		}
+
+		// Add to row map
+		DataTable->RowMap.Add(RowName, RowData);
+
+		// Now iterate over cells (skipping first cell, that was row name)
+		for(int32 CellIdx=1; CellIdx<Cells.Num(); CellIdx++)
+		{
+			// Try and assign string to data using the column property
+			UProperty* ColumnProp = ColumnProps[CellIdx];
+			const FString CellValue = Cells[CellIdx];
+			FString Error = DataTableUtils::AssignStringToProperty(CellValue, ColumnProp, RowData);
+
+			// If we failed, output a problem string
+			if(Error.Len() > 0)
+			{
+				FString ColumnName = (ColumnProp != NULL) 
+					? DataTableUtils::GetPropertyDisplayName(ColumnProp, ColumnProp->GetName())
+					: FString(TEXT("NONE"));
+				ImportProblems.Add(FString::Printf(TEXT("Problem assigning string '%s' to property '%s' on row '%s' : %s"), *CellValue, *ColumnName, *RowName.ToString(), *Error));
+			}
+		}
+
+		// Problem if we didn't have enough cells on this row
+		if(Cells.Num() < ColumnProps.Num())
+		{
+			ImportProblems.Add(FString::Printf(TEXT("Too few cells on row '%s'."), *RowName.ToString()));			
+		}
+	}
+
+	DataTable->Modify(true);
+
+	return true;
+}
+
+#endif // WITH_EDITOR
diff --git a/Engine/Source/Runtime/Engine/Private/DataTableCSV.h b/Engine/Source/Runtime/Engine/Private/DataTableCSV.h
new file mode 100644
--- /dev/null
+++ b/Engine/Source/Runtime/Engine/Private/DataTableCSV.h
@@ -0,0 +1,42 @@
+// Copyright 1998-2015 Epic Games, Inc. All Rights Reserved.
+
+#pragma once
+
+class UDataTable;
+
+#if WITH_EDITOR
+
+class FDataTableExporterCSV
+{
+public:
+	FDataTableExporterCSV(const UDataTable& InDataTable, FString& OutExportText);
+
+	~FDataTableExporterCSV();
+
+	bool WriteTable();
+
+	bool WriteRow(const void* InRowData);
+
+private:
+	bool WriteStructEntry(const void* InRowData, UProperty* InProperty, const void* InPropertyData);
+
+	const UDataTable* DataTable;
+	FString& ExportedText;
+};
+
+class FDataTableImporterCSV
+{
+public:
+	FDataTableImporterCSV(UDataTable& InDataTable, FString InCSVData, TArray<FString>& OutProblems);
+
+	~FDataTableImporterCSV();
+
+	bool ReadTable();
+
+private:
+	UDataTable* DataTable;
+	FString CSVData;
+	TArray<FString>& ImportProblems;
+};
+
+#endif // WITH_EDITOR
diff --git a/Engine/Source/Runtime/Engine/Private/DataTableJSON.cpp b/Engine/Source/Runtime/Engine/Private/DataTableJSON.cpp
new file mode 100644
--- /dev/null
+++ b/Engine/Source/Runtime/Engine/Private/DataTableJSON.cpp
@@ -0,0 +1,440 @@
+// Copyright	 1998-2015 Epic Games, Inc. All Rights Reserved.
+
+#include "EnginePrivate.h"
+#include "DataTableJSON.h"
+#include "Json.h"
+#include "Engine/DataTable.h"
+#include "Engine/UserDefinedStruct.h"
+#include "DataTableUtils.h"
+
+#if WITH_EDITOR
+
+namespace
+{
+	const TCHAR* JSONTypeToString(const EJson InType)
+	{
+		switch(InType)
+		{
+		case EJson::None:
+			return TEXT("None");
+		case EJson::Null:
+			return TEXT("Null");
+		case EJson::String:
+			return TEXT("String");
+		case EJson::Number:
+			return TEXT("Number");
+		case EJson::Boolean:
+			return TEXT("Boolean");
+		case EJson::Array:
+			return TEXT("Array");
+		case EJson::Object:
+			return TEXT("Object");
+		default:
+			return TEXT("Unknown");
+		}
+	}
+}
+
+FDataTableExporterJSON::FDataTableExporterJSON(const UDataTable& InDataTable, FString& OutExportText)
+	: DataTable(&InDataTable)
+	, JsonWriter(TJsonWriterFactory<TCHAR, TPrettyJsonPrintPolicy<TCHAR> >::Create(&OutExportText))
+	, bJsonWriterNeedsClose(true)
+{
+}
+
+FDataTableExporterJSON::FDataTableExporterJSON(const UDataTable& InDataTable, TSharedRef< TJsonWriter<TCHAR, TPrettyJsonPrintPolicy<TCHAR> > > InJsonWriter)
+	: DataTable(&InDataTable)
+	, JsonWriter(InJsonWriter)
+	, bJsonWriterNeedsClose(false)
+{
+}
+
+FDataTableExporterJSON::~FDataTableExporterJSON()
+{
+	if (bJsonWriterNeedsClose)
+	{
+		JsonWriter->Close();
+	}
+}
+
+bool FDataTableExporterJSON::WriteTable()
+{
+	if (!DataTable->RowStruct)
+	{
+		return false;
+	}
+
+	JsonWriter->WriteArrayStart();
+
+	// Iterate over rows
+	for (auto RowIt = DataTable->RowMap.CreateConstIterator(); RowIt; ++RowIt)
+	{
+		JsonWriter->WriteObjectStart();
+		{
+			// RowName
+			const FName RowName = RowIt.Key();
+			JsonWriter->WriteValue(TEXT("Name"), RowName.ToString());
+
+			// Now the values
+			uint8* RowData = RowIt.Value();
+			WriteRow(RowData);
+		}
+		JsonWriter->WriteObjectEnd();
+	}
+
+	JsonWriter->WriteArrayEnd();
+
+	return true;
+}
+
+bool FDataTableExporterJSON::WriteRow(const void* InRowData)
+{
+	if (!DataTable->RowStruct)
+	{
+		return false;
+	}
+
+	for (TFieldIterator<UProperty> It(DataTable->RowStruct); It; ++It)
+	{
+		UProperty* BaseProp = *It;
+		check(BaseProp);
+
+		const void* Data = BaseProp->ContainerPtrToValuePtr<void>(InRowData, 0);
+		WriteStructEntry(InRowData, BaseProp, Data);
+	}
+
+	return true;
+}
+
+bool FDataTableExporterJSON::WriteStructEntry(const void* InRowData, const UProperty* InProperty, const void* InPropertyData)
+{
+	if (const UNumericProperty *NumProp = Cast<const UNumericProperty>(InProperty))
+	{
+		if (NumProp->IsInteger())
+		{
+			const int64 PropertyValue = NumProp->GetSignedIntPropertyValue(InPropertyData);
+			JsonWriter->WriteValue(InProperty->GetName(), PropertyValue);
+		}
+		else
+		{
+			const double PropertyValue = NumProp->GetFloatingPointPropertyValue(InPropertyData);
+			JsonWriter->WriteValue(InProperty->GetName(), PropertyValue);
+		}
+	}
+	else if (const UBoolProperty* BoolProp = Cast<const UBoolProperty>(InProperty))
+	{
+		const bool PropertyValue = BoolProp->GetPropertyValue(InPropertyData);
+		JsonWriter->WriteValue(InProperty->GetName(), PropertyValue);
+	}
+	else if (const UArrayProperty* ArrayProp = Cast<const UArrayProperty>(InProperty))
+	{
+		JsonWriter->WriteArrayStart(InProperty->GetName());
+
+		FScriptArrayHelper ArrayHelper(ArrayProp, InPropertyData);
+		for (int32 ArrayEntryIndex = 0; ArrayEntryIndex < ArrayHelper.Num(); ++ArrayEntryIndex)
+		{
+			const void* ArrayEntryData = ArrayHelper.GetRawPtr(ArrayEntryIndex);
+			WriteArrayEntry(ArrayProp->Inner, ArrayEntryData);
+		}
+
+		JsonWriter->WriteArrayEnd();
+	}
+	else
+	{
+		const FString PropertyValue = DataTableUtils::GetPropertyValueAsString(InProperty, (uint8*)InRowData);
+		JsonWriter->WriteValue(InProperty->GetName(), PropertyValue);
+	}
+
+	return true;
+}
+
+bool FDataTableExporterJSON::WriteArrayEntry(const UProperty* InProperty, const void* InPropertyData)
+{
+	if (const UNumericProperty *NumProp = Cast<const UNumericProperty>(InProperty))
+	{
+		if (NumProp->IsInteger())
+		{
+			const int64 PropertyValue = NumProp->GetSignedIntPropertyValue(InPropertyData);
+			JsonWriter->WriteValue(PropertyValue);
+		}
+		else
+		{
+			const double PropertyValue = NumProp->GetFloatingPointPropertyValue(InPropertyData);
+			JsonWriter->WriteValue(PropertyValue);
+		}
+	}
+	else if (const UBoolProperty* BoolProp = Cast<const UBoolProperty>(InProperty))
+	{
+		const bool PropertyValue = BoolProp->GetPropertyValue(InPropertyData);
+		JsonWriter->WriteValue(PropertyValue);
+	}
+	else if (const UArrayProperty* ArrayProp = Cast<const UArrayProperty>(InProperty))
+	{
+		// Cannot nest arrays
+		return false;
+	}
+	else
+	{
+		const FString PropertyValue = DataTableUtils::GetPropertyValueAsString(InProperty, (uint8*)InPropertyData);
+		JsonWriter->WriteValue(PropertyValue);
+	}
+
+	return true;
+}
+
+
+FDataTableImporterJSON::FDataTableImporterJSON(UDataTable& InDataTable, const FString& InJSONData, TArray<FString>& OutProblems)
+	: DataTable(&InDataTable)
+	, JSONData(InJSONData)
+	, ImportProblems(OutProblems)
+{
+}
+
+FDataTableImporterJSON::~FDataTableImporterJSON()
+{
+}
+
+bool FDataTableImporterJSON::ReadTable()
+{
+	if (JSONData.IsEmpty())
+	{
+		ImportProblems.Add(TEXT("Input data is empty."));
+		return false;
+	}
+
+	// Check we have a RowStruct specified
+	if (!DataTable->RowStruct)
+	{
+		ImportProblems.Add(TEXT("No RowStruct specified."));
+		return false;
+	}
+
+	TArray< TSharedPtr<FJsonValue> > ParsedTableRows;
+	{
+		const TSharedRef< TJsonReader<TCHAR> > JsonReader = TJsonReaderFactory<TCHAR>::Create(JSONData);
+		if (!FJsonSerializer::Deserialize(JsonReader, ParsedTableRows) || ParsedTableRows.Num() == 0)
+		{
+			ImportProblems.Add(FString::Printf(TEXT("Failed to parse the JSON data. Error: %s"), *JsonReader->GetErrorMessage()));
+			return false;
+		}
+	}
+
+	// Empty existing data
+	DataTable->EmptyTable();
+
+	// Iterate over rows
+	for (int32 RowIdx = 0; RowIdx < ParsedTableRows.Num(); ++RowIdx)
+	{
+		const TSharedPtr<FJsonValue>& ParsedTableRowValue = ParsedTableRows[RowIdx];
+		TSharedPtr<FJsonObject> ParsedTableRowObject = ParsedTableRowValue->AsObject();
+		if (!ParsedTableRowObject.IsValid())
+		{
+			ImportProblems.Add(FString::Printf(TEXT("Row '%d' is not a valid JSON object."), RowIdx));
+			continue;
+		}
+
+		ReadRow(ParsedTableRowObject.ToSharedRef(), RowIdx);
+	}
+
+	DataTable->Modify(true);
+
+	return true;
+}
+
+bool FDataTableImporterJSON::ReadRow(const TSharedRef<FJsonObject>& ParsedTableRowObject, const int32 RowIdx)
+{
+	// Get row name
+	FName RowName = DataTableUtils::MakeValidName(ParsedTableRowObject->GetStringField(TEXT("Name")));
+
+	// Check its not 'none'
+	if (RowName.IsNone())
+	{
+		ImportProblems.Add(FString::Printf(TEXT("Row '%d' missing a name."), RowIdx));
+		return false;
+	}
+
+	// Check its not a duplicate
+	if (DataTable->RowMap.Find(RowName) != nullptr)
+	{
+		ImportProblems.Add(FString::Printf(TEXT("Duplicate row name '%s'."), *RowName.ToString()));
+		return false;
+	}
+
+	// Allocate data to store information, using UScriptStruct to know its size
+	uint8* RowData = (uint8*)FMemory::Malloc(DataTable->RowStruct->PropertiesSize);
+	DataTable->RowStruct->InitializeStruct(RowData);
+	// And be sure to call DestroyScriptStruct later
+
+	if (auto UDStruct = Cast<const UUserDefinedStruct>(DataTable->RowStruct))
+	{
+		UDStruct->InitializeDefaultValue(RowData);
+	}
+
+	// Add to row map
+	DataTable->RowMap.Add(RowName, RowData);
+
+	// Now read in each property
+	for (TFieldIterator<UProperty> It(DataTable->RowStruct); It; ++It)
+	{
+		UProperty* BaseProp = *It;
+		check(BaseProp);
+
+		const FString PropertyName = BaseProp->GetName();
+		TSharedPtr<FJsonValue> ParsedPropertyValue = ParsedTableRowObject->TryGetField(PropertyName);
+		if (!ParsedPropertyValue.IsValid())
+		{
+			ImportProblems.Add(FString::Printf(TEXT("Row '%s' is missing an entry for '%s'."), *RowName.ToString(), *PropertyName));
+			continue;
+		}
+
+		void* Data = BaseProp->ContainerPtrToValuePtr<void>(RowData, 0);
+		ReadStructEntry(ParsedPropertyValue.ToSharedRef(), RowName, RowData, BaseProp, Data);
+	}
+
+	return true;
+}
+
+bool FDataTableImporterJSON::ReadStructEntry(const TSharedRef<FJsonValue>& ParsedPropertyValue, const FName InRowName, const void* InRowData, UProperty* InProperty, void* InPropertyData)
+{
+	const FString ColumnName = DataTableUtils::GetPropertyDisplayName(InProperty, InProperty->GetName());
+	const TCHAR* const ParsedPropertyType = JSONTypeToString(ParsedPropertyValue->Type);
+
+	if (UNumericProperty *NumProp = Cast<UNumericProperty>(InProperty))
+	{
+		if (NumProp->IsInteger())
+		{
+			int64 PropertyValue = 0;
+			if (!ParsedPropertyValue->TryGetNumber(PropertyValue))
+			{
+				ImportProblems.Add(FString::Printf(TEXT("Property '%s' on row '%s' is the incorrect type. Expected Integer, got %s."), *ColumnName, *InRowName.ToString(), ParsedPropertyType));
+				return false;
+			}
+
+			NumProp->SetIntPropertyValue(InPropertyData, PropertyValue);
+		}
+		else
+		{
+			double PropertyValue = 0.0;
+			if (!ParsedPropertyValue->TryGetNumber(PropertyValue))
+			{
+				ImportProblems.Add(FString::Printf(TEXT("Property '%s' on row '%s' is the incorrect type. Expected Double, got %s."), *ColumnName, *InRowName.ToString(), ParsedPropertyType));
+				return false;
+			}
+
+			NumProp->SetFloatingPointPropertyValue(InPropertyData, PropertyValue);
+		}
+	}
+	else if (UBoolProperty* BoolProp = Cast<UBoolProperty>(InProperty))
+	{
+		bool PropertyValue = false;
+		if (!ParsedPropertyValue->TryGetBool(PropertyValue))
+		{
+			ImportProblems.Add(FString::Printf(TEXT("Property '%s' on row '%s' is the incorrect type. Expected Boolean, got %s."), *ColumnName, *InRowName.ToString(), ParsedPropertyType));
+			return false;
+		}
+
+		BoolProp->SetPropertyValue(InPropertyData, PropertyValue);
+	}
+	else if (UArrayProperty* ArrayProp = Cast<UArrayProperty>(InProperty))
+	{
+		const TArray< TSharedPtr<FJsonValue> >* PropertyValuesPtr;
+		if (!ParsedPropertyValue->TryGetArray(PropertyValuesPtr))
+		{
+			ImportProblems.Add(FString::Printf(TEXT("Property '%s' on row '%s' is the incorrect type. Expected Array, got %s."), *ColumnName, *InRowName.ToString(), ParsedPropertyType));
+			return false;
+		}
+
+		FScriptArrayHelper ArrayHelper(ArrayProp, InPropertyData);
+		for (const TSharedPtr<FJsonValue>& PropertyValueEntry : *PropertyValuesPtr)
+		{
+			const int32 NewEntryIndex = ArrayHelper.AddValue();
+			void* ArrayEntryData = ArrayHelper.GetRawPtr(NewEntryIndex);
+			ReadArrayEntry(PropertyValueEntry.ToSharedRef(), InRowName, ColumnName, NewEntryIndex, ArrayProp->Inner, ArrayEntryData);
+		}
+	}
+	else
+	{
+		FString PropertyValue;
+		if (!ParsedPropertyValue->TryGetString(PropertyValue))
+		{
+			ImportProblems.Add(FString::Printf(TEXT("Property '%s' on row '%s' is the incorrect type. Expected String, got %s."), *ColumnName, *InRowName.ToString(), ParsedPropertyType));
+			return false;
+		}
+
+		const FString Error = DataTableUtils::AssignStringToProperty(PropertyValue, InProperty, (uint8*)InRowData);
+		if(Error.Len() > 0)
+		{
+			ImportProblems.Add(FString::Printf(TEXT("Problem assigning string '%s' to property '%s' on row '%s' : %s"), *PropertyValue, *ColumnName, *InRowName.ToString(), *Error));
+			return false;
+		}
+	}
+
+	return true;
+}
+
+bool FDataTableImporterJSON::ReadArrayEntry(const TSharedRef<FJsonValue>& ParsedPropertyValue, const FName InRowName, const FString& InColumnName, const int32 InArrayEntryIndex, UProperty* InProperty, void* InPropertyData)
+{
+	const TCHAR* const ParsedPropertyType = JSONTypeToString(ParsedPropertyValue->Type);
+
+	if (UNumericProperty *NumProp = Cast<UNumericProperty>(InProperty))
+	{
+		if (NumProp->IsInteger())
+		{
+			int64 PropertyValue = 0;
+			if (!ParsedPropertyValue->TryGetNumber(PropertyValue))
+			{
+				ImportProblems.Add(FString::Printf(TEXT("Entry %d on property '%s' on row '%s' is the incorrect type. Expected Integer, got %s."), InArrayEntryIndex, *InColumnName, *InRowName.ToString(), ParsedPropertyType));
+				return false;
+			}
+
+			NumProp->SetIntPropertyValue(InPropertyData, PropertyValue);
+		}
+		else
+		{
+			double PropertyValue = 0.0;
+			if (!ParsedPropertyValue->TryGetNumber(PropertyValue))
+			{
+				ImportProblems.Add(FString::Printf(TEXT("Entry %d on property '%s' on row '%s' is the incorrect type. Expected Double, got %s."), InArrayEntryIndex, *InColumnName, *InRowName.ToString(), ParsedPropertyType));
+				return false;
+			}
+
+			NumProp->SetFloatingPointPropertyValue(InPropertyData, PropertyValue);
+		}
+	}
+	else if (UBoolProperty* BoolProp = Cast<UBoolProperty>(InProperty))
+	{
+		bool PropertyValue = false;
+		if (!ParsedPropertyValue->TryGetBool(PropertyValue))
+		{
+			ImportProblems.Add(FString::Printf(TEXT("Entry %d on property '%s' on row '%s' is the incorrect type. Expected Boolean, got %s."), InArrayEntryIndex, *InColumnName, *InRowName.ToString(), ParsedPropertyType));
+			return false;
+		}
+
+		BoolProp->SetPropertyValue(InPropertyData, PropertyValue);
+	}
+	else if (UArrayProperty* ArrayProp = Cast<UArrayProperty>(InProperty))
+	{
+		// Cannot nest arrays
+		return false;
+	}
+	else
+	{
+		FString PropertyValue;
+		if (!ParsedPropertyValue->TryGetString(PropertyValue))
+		{
+			ImportProblems.Add(FString::Printf(TEXT("Entry %d on property '%s' on row '%s' is the incorrect type. Expected String, got %s."), InArrayEntryIndex, *InColumnName, *InRowName.ToString(), ParsedPropertyType));
+			return false;
+		}
+
+		const FString Error = DataTableUtils::AssignStringToProperty(PropertyValue, InProperty, (uint8*)InPropertyData);
+		if(Error.Len() > 0)
+		{
+			ImportProblems.Add(FString::Printf(TEXT("Problem assigning string '%s' to entry %d on property '%s' on row '%s' : %s"), InArrayEntryIndex, *PropertyValue, *InColumnName, *InRowName.ToString(), *Error));
+			return false;
+		}
+	}
+
+	return true;
+}
+
+#endif // WITH_EDITOR
diff --git a/Engine/Source/Runtime/Engine/Private/DataTableJSON.h b/Engine/Source/Runtime/Engine/Private/DataTableJSON.h
new file mode 100644
--- /dev/null
+++ b/Engine/Source/Runtime/Engine/Private/DataTableJSON.h
@@ -0,0 +1,61 @@
+// Copyright 1998-2015 Epic Games, Inc. All Rights Reserved.
+
+#pragma once
+
+class UDataTable;
+class FJsonValue;
+class FJsonObject;
+
+// forward declare JSON writer
+template <class CharType>
+struct TPrettyJsonPrintPolicy;
+template <class CharType, class PrintPolicy>
+class TJsonWriter;
+
+#if WITH_EDITOR
+
+class FDataTableExporterJSON
+{
+public:
+	FDataTableExporterJSON(const UDataTable& InDataTable, FString& OutExportText);
+
+	FDataTableExporterJSON(const UDataTable& InDataTable, TSharedRef< TJsonWriter<TCHAR, TPrettyJsonPrintPolicy<TCHAR> > > InJsonWriter);
+
+	~FDataTableExporterJSON();
+
+	bool WriteTable();
+
+	bool WriteRow(const void* InRowData);
+
+private:
+	bool WriteStructEntry(const void* InRowData, const UProperty* InProperty, const void* InPropertyData);
+
+	bool WriteArrayEntry(const UProperty* InProperty, const void* InPropertyData);
+
+	const UDataTable* DataTable;
+	TSharedRef< TJsonWriter<TCHAR, TPrettyJsonPrintPolicy<TCHAR> > > JsonWriter;
+	bool bJsonWriterNeedsClose;
+};
+
+class FDataTableImporterJSON
+{
+public:
+	FDataTableImporterJSON(UDataTable& InDataTable, const FString& InJSONData, TArray<FString>& OutProblems);
+
+	~FDataTableImporterJSON();
+
+	bool ReadTable();
+
+private:
+	bool ReadRow(const TSharedRef<FJsonObject>& ParsedTableRowObject, const int32 RowIdx);
+
+	bool ReadStructEntry(const TSharedRef<FJsonValue>& ParsedPropertyValue, const FName InRowName, const void* InRowData, UProperty* InProperty, void* InPropertyData);
+
+	bool ReadArrayEntry(const TSharedRef<FJsonValue>& ParsedPropertyValue, const FName InRowName, const FString& InColumnName, const int32 InArrayEntryIndex, UProperty* InProperty, void* InPropertyData);
+
+	UDataTable* DataTable;
+	const FString& JSONData;
+	TArray<FString>& ImportProblems;
+};
+
+#endif // WITH_EDITOR
diff --git a/Engine/Source/Runtime/Engine/Private/DataTableUtils.cpp b/Engine/Source/Runtime/Engine/Private/DataTableUtils.cpp
new file mode 100644
--- /dev/null
+++ b/Engine/Source/Runtime/Engine/Private/DataTableUtils.cpp
@@ -0,0 +1,86 @@
+// Copyright 1998-2015 Epic Games, Inc. All Rights Reserved.
+
+#include "EnginePrivate.h"
+#include "DataTableUtils.h"
+
+FString DataTableUtils::AssignStringToProperty(const FString& InString, const UProperty* InProp, uint8* InData)
+{
+	FStringOutputDevice ImportError;
+	if(InProp != NULL && IsSupportedTableProperty(InProp))
+	{
+		InProp->ImportText(*InString, InProp->ContainerPtrToValuePtr<uint8>(InData), PPF_None, NULL, &ImportError);
+	}
+
+	FString Error = ImportError;
+	return Error;
+}
+
+FString DataTableUtils::GetPropertyValueAsString(const UProperty* InProp, uint8* InData)
+{
+	FString Result(TEXT(""));
+
+	if(InProp != NULL && IsSupportedTableProperty(InProp))
+	{
+		InProp->ExportText_InContainer(0, Result, InData, InData, NULL, PPF_None);
+	}
+
+	return Result;
+}
+
+TArray<FName> DataTableUtils::GetStructPropertyNames(UStruct* InStruct)
+{
+	TArray<FName> PropNames;
+	for (TFieldIterator<UProperty> It(InStruct); It; ++It)
+	{
+		PropNames.Add(It->GetFName());
+	}
+	return PropNames;
+}
+
+FName DataTableUtils::MakeValidName(const FString& InString)
+{
+	FString InvalidChars(INVALID_NAME_CHARACTERS);
+
+	FString FixedString;
+	TArray<TCHAR>& FixedCharArray = FixedString.GetCharArray();
+
+	// Iterate over input string characters
+	for(int32 CharIdx=0; CharIdx<InString.Len(); CharIdx++)
+	{
+		// See if this char occurs in the InvalidChars string
+		FString Char = InString.Mid( CharIdx, 1 );
+		if( !InvalidChars.Contains(Char) )
+		{
+			// Its ok, add to result
+			FixedCharArray.Add(Char[0]);
+		}
+	}
+	FixedCharArray.Add(0);
+
+	return FName(*FixedString);
+}
+
+bool DataTableUtils::IsSupportedTableProperty(const UProperty* InProp)
+{
+	return(	InProp->IsA(UIntProperty::StaticClass()) || 
+			InProp->IsA(UFloatProperty::StaticClass()) ||
+			InProp->IsA(UNameProperty::StaticClass()) ||
+			InProp->IsA(UStrProperty::StaticClass()) ||
+			InProp->IsA(UBoolProperty::StaticClass()) ||
+			InProp->IsA(UObjectPropertyBase::StaticClass()) ||
+			InProp->IsA(UStructProperty::StaticClass()) ||
+			InProp->IsA(UByteProperty::StaticClass()) ||
+			InProp->IsA(UTextProperty::StaticClass()) ||
+			InProp->IsA(UArrayProperty::StaticClass())
+			);
+}
+
+FString DataTableUtils::GetPropertyDisplayName(const UProperty* Prop, const FString& DefaultName)
+{
+#if WITH_EDITOR
+	static const FName DisplayNameKey(TEXT("DisplayName"));
+	return (Prop && Prop->HasMetaData(DisplayNameKey)) ? Prop->GetMetaData(DisplayNameKey) : DefaultName;
+#else  // WITH_EDITOR
+	return DefaultName;
+#endif // WITH_EDITOR
+}
diff --git a/Engine/Source/Runtime/Engine/Public/DataTableUtils.h b/Engine/Source/Runtime/Engine/Public/DataTableUtils.h
new file mode 100644
--- /dev/null
+++ b/Engine/Source/Runtime/Engine/Public/DataTableUtils.h
@@ -0,0 +1,24 @@
+// Copyright 1998-2015 Epic Games, Inc. All Rights Reserved.
+
+#pragma once
+
+namespace DataTableUtils
+{
+	/** Util to assign a value (given as a string) to a struct property. */
+	ENGINE_API FString AssignStringToProperty(const FString& InString, const UProperty* InProp, uint8* InData);
+
+	/** Util to assign get a property as a string */
+	ENGINE_API FString GetPropertyValueAsString(const UProperty* InProp, uint8* InData);
+
+	/** Util to get all property names from a struct */
+	ENGINE_API TArray<FName> GetStructPropertyNames(UStruct* InStruct);
+
+	/** Util that removes invalid chars and then make an FName */
+	ENGINE_API FName MakeValidName(const FString& InString);
+
+	/** Util to see if this property is supported in a row struct. */
+	ENGINE_API bool IsSupportedTableProperty(const UProperty* InProp);
+
+	/** Util to get the friendly display name of a given property */
+	ENGINE_API FString GetPropertyDisplayName(const UProperty* Prop, const FString& DefaultName);
+}
