Blueprint Call Function nodes will now display localized function names.

FriendlyName metadata is now deprecated and all uses have been removed, a metadata redirector has been setup.

Changed most/all current instances of FriendlyName metadata over to DisplayName.

diff --git a/Engine/Config/BaseEngine.ini b/Engine/Config/BaseEngine.ini
--- a/Engine/Config/BaseEngine.ini
+++ b/Engine/Config/BaseEngine.ini
@@ -1213,6 +1213,7 @@ bUseStreamingPause=false
 +MetadataRedirects=(OldKey="K2CompactNode", NewKey="CompactNodeTitle")
 +MetadataRedirects=(OldKey="MenuCategory", NewKey="Category")
 +MetadataRedirects=(OldKey="ArrayPointerParm", NewKey="TargetArrayParm")
++MetadataRedirects=(OldKey="FriendlyName", NewKey="DisplayName")
 
 [PlatformInterface]
 CloudStorageInterfaceClassName=
diff --git a/Engine/Source/Editor/BlueprintGraph/Classes/EdGraphSchema_K2.h b/Engine/Source/Editor/BlueprintGraph/Classes/EdGraphSchema_K2.h
--- a/Engine/Source/Editor/BlueprintGraph/Classes/EdGraphSchema_K2.h
+++ b/Engine/Source/Editor/BlueprintGraph/Classes/EdGraphSchema_K2.h
@@ -700,7 +700,7 @@ public:
 	static bool CanKismetOverrideFunction(const UFunction* Function);
 
 	/** returns friendly signature name if possible or Removes any mangling to get the unmangled signature name of the function */
-	static FString GetFriendlySignatureName(const UFunction* Function);
+	static FText GetFriendlySignatureName(const UFunction* Function);
 
 	static bool IsAllowableBlueprintVariableType(const class UEnum* InEnum);
 	static bool IsAllowableBlueprintVariableType(const class UClass* InClass);
diff --git a/Engine/Source/Editor/BlueprintGraph/Classes/K2Node_CallFunction.h b/Engine/Source/Editor/BlueprintGraph/Classes/K2Node_CallFunction.h
--- a/Engine/Source/Editor/BlueprintGraph/Classes/K2Node_CallFunction.h
+++ b/Engine/Source/Editor/BlueprintGraph/Classes/K2Node_CallFunction.h
@@ -137,7 +137,7 @@ public:
 	virtual void PostParameterPinCreated(UEdGraphPin *Pin) {}
 
 	/** Gets the user-facing name for the function */
-	static FString GetUserFacingFunctionName(const UFunction* Function);
+	static FText GetUserFacingFunctionName(const UFunction* Function);
 
 	/** Set up a pins tooltip from a function's tooltip */
 	static void GeneratePinTooltipFromFunction(UEdGraphPin& Pin, const UFunction* Function);
diff --git a/Engine/Source/Editor/BlueprintGraph/Private/BlueprintEventNodeSpawner.cpp b/Engine/Source/Editor/BlueprintGraph/Private/BlueprintEventNodeSpawner.cpp
--- a/Engine/Source/Editor/BlueprintGraph/Private/BlueprintEventNodeSpawner.cpp
+++ b/Engine/Source/Editor/BlueprintGraph/Private/BlueprintEventNodeSpawner.cpp
@@ -95,8 +95,8 @@ UBlueprintEventNodeSpawner* UBlueprintEventNodeSpawner::Create(UFunction const*
 	NodeSpawner->NodeClass = UK2Node_Event::StaticClass();
 
 	FBlueprintActionUiSpec& MenuSignature = NodeSpawner->DefaultMenuSignature;
-	FString const FuncName = UEdGraphSchema_K2::GetFriendlySignatureName(EventFunc);
-	MenuSignature.MenuName = FText::Format(LOCTEXT("EventWithSignatureName", "Event {0}"), FText::FromString(FuncName));
+	FText const FuncName = UEdGraphSchema_K2::GetFriendlySignatureName(EventFunc);
+	MenuSignature.MenuName = FText::Format(LOCTEXT("EventWithSignatureName", "Event {0}"), FuncName);
 	FString const FuncCategory = UK2Node_CallFunction::GetDefaultCategoryForFunction(EventFunc, TEXT(""));
 	MenuSignature.Category = FText::FromString(LOCTEXT("AddEventCategory", "Add Event").ToString() + TEXT("|") + FuncCategory);
 	//MenuSignature.Tooltip, will be pulled from the node template
diff --git a/Engine/Source/Editor/BlueprintGraph/Private/BlueprintFunctionNodeSpawner.cpp b/Engine/Source/Editor/BlueprintGraph/Private/BlueprintFunctionNodeSpawner.cpp
--- a/Engine/Source/Editor/BlueprintGraph/Private/BlueprintFunctionNodeSpawner.cpp
+++ b/Engine/Source/Editor/BlueprintGraph/Private/BlueprintFunctionNodeSpawner.cpp
@@ -242,7 +242,7 @@ UBlueprintFunctionNodeSpawner* UBlueprintFunctionNodeSpawner::Create(TSubclassOf
 	//--------------------------------------
 
 	FBlueprintActionUiSpec& MenuSignature = NodeSpawner->DefaultMenuSignature;
-	MenuSignature.MenuName = FText::FromString( UK2Node_CallFunction::GetUserFacingFunctionName(Function) );
+	MenuSignature.MenuName = UK2Node_CallFunction::GetUserFacingFunctionName(Function);
 	MenuSignature.Category = FText::FromString( UK2Node_CallFunction::GetDefaultCategoryForFunction(Function, TEXT("")) );
 	MenuSignature.Tooltip  = FText::FromString( UK2Node_CallFunction::GetDefaultTooltipForFunction(Function) );
 	// add at least one character, so that PrimeDefaultMenuSignature() doesn't attempt to query the template node
diff --git a/Engine/Source/Editor/BlueprintGraph/Private/EdGraphSchema_K2.cpp b/Engine/Source/Editor/BlueprintGraph/Private/EdGraphSchema_K2.cpp
--- a/Engine/Source/Editor/BlueprintGraph/Private/EdGraphSchema_K2.cpp
+++ b/Engine/Source/Editor/BlueprintGraph/Private/EdGraphSchema_K2.cpp
@@ -784,7 +784,7 @@ bool UEdGraphSchema_K2::FunctionCanBeUsedInDelegate(const UFunction* InFunction)
 	return true;
 }
 
-FString UEdGraphSchema_K2::GetFriendlySignatureName(const UFunction* Function)
+FText UEdGraphSchema_K2::GetFriendlySignatureName(const UFunction* Function)
 {
 	return UK2Node_CallFunction::GetUserFacingFunctionName( Function );
 }
diff --git a/Engine/Source/Editor/BlueprintGraph/Private/K2ActionMenuBuilder.cpp b/Engine/Source/Editor/BlueprintGraph/Private/K2ActionMenuBuilder.cpp
--- a/Engine/Source/Editor/BlueprintGraph/Private/K2ActionMenuBuilder.cpp
+++ b/Engine/Source/Editor/BlueprintGraph/Private/K2ActionMenuBuilder.cpp
@@ -605,8 +605,8 @@ void FK2ActionMenuBuilder::AddEventForFunction(FGraphActionListBuilderBase& Acti
 	if (Function != NULL)
 	{
 		check(!Function->HasAnyFunctionFlags(FUNC_Delegate));
-		const FString SigName   = UEdGraphSchema_K2::GetFriendlySignatureName(Function);
-		const FString EventName = FText::Format( LOCTEXT("EventWithSignatureName", "Event {0}"), FText::FromString(SigName)).ToString();
+		const FText SigName   = UEdGraphSchema_K2::GetFriendlySignatureName(Function);
+		const FString EventName = FText::Format( LOCTEXT("EventWithSignatureName", "Event {0}"), SigName).ToString();
 
 		UK2Node_Event* EventNodeTemplate = ActionMenuBuilder.CreateTemplateNode<UK2Node_Event>();
 		EventNodeTemplate->EventReference.SetFromField<UFunction>(Function, false);
@@ -641,7 +641,7 @@ void FK2ActionMenuBuilder::AttachMacroGraphAction(FGraphActionListBuilderBase& A
 void FK2ActionMenuBuilder::AddSpawnInfoForFunction(const UFunction* Function, bool bIsParentContext, const FFunctionTargetInfo& TargetInfo, const FMemberReference& CallOnMember, const FString& BaseCategory, int32 const ActionGrouping, FGraphActionListBuilderBase& ListBuilder, bool bCalledForEach, bool bPlaceOnTop)
 {
 	FString Tooltip = UK2Node_CallFunction::GetDefaultTooltipForFunction(Function);
-	FString Description = UK2Node_CallFunction::GetUserFacingFunctionName(Function);
+	FString Description = UK2Node_CallFunction::GetUserFacingFunctionName(Function).ToString();
 	FString Keywords = UK2Node_CallFunction::GetKeywordsForFunction(Function);
 
 	if (bIsParentContext)
diff --git a/Engine/Source/Editor/BlueprintGraph/Private/K2Node_BaseAsyncTask.cpp b/Engine/Source/Editor/BlueprintGraph/Private/K2Node_BaseAsyncTask.cpp
--- a/Engine/Source/Editor/BlueprintGraph/Private/K2Node_BaseAsyncTask.cpp
+++ b/Engine/Source/Editor/BlueprintGraph/Private/K2Node_BaseAsyncTask.cpp
@@ -28,8 +28,8 @@ FText UK2Node_BaseAsyncTask::GetTooltipText() const
 
 FText UK2Node_BaseAsyncTask::GetNodeTitle(ENodeTitleType::Type TitleType) const
 {
-	const FString FunctionToolTipText = UK2Node_CallFunction::GetUserFacingFunctionName(GetFactoryFunction());
-	return FText::FromString(FunctionToolTipText);
+	const FText FunctionToolTipText = UK2Node_CallFunction::GetUserFacingFunctionName(GetFactoryFunction());
+	return FunctionToolTipText;
 }
 
 bool UK2Node_BaseAsyncTask::IsCompatibleWithGraph(const UEdGraph* TargetGraph) const
diff --git a/Engine/Source/Editor/BlueprintGraph/Private/K2Node_CallFunction.cpp b/Engine/Source/Editor/BlueprintGraph/Private/K2Node_CallFunction.cpp
--- a/Engine/Source/Editor/BlueprintGraph/Private/K2Node_CallFunction.cpp
+++ b/Engine/Source/Editor/BlueprintGraph/Private/K2Node_CallFunction.cpp
@@ -455,7 +455,7 @@ FString UK2Node_CallFunction::GetFunctionContextString() const
 
 FText UK2Node_CallFunction::GetNodeTitle(ENodeTitleType::Type TitleType) const
 {
-	FString FunctionName;
+	FText FunctionName;
 	FString ContextString;
 	FString RPCString;
 
@@ -467,24 +467,24 @@ FText UK2Node_CallFunction::GetNodeTitle(ENodeTitleType::Type TitleType) const
 	}
 	else
 	{
-		FunctionName = FunctionReference.GetMemberName().ToString();
+		FunctionName = FText::FromName(FunctionReference.GetMemberName());
 		if ((GEditor != NULL) && (GetDefault<UEditorStyleSettings>()->bShowFriendlyNames))
 		{
-			FunctionName = FName::NameToDisplayString(FunctionName, false);
+			FunctionName = FText::FromString(FName::NameToDisplayString(FunctionName.ToString(), false));
 		}
 	}
 
 	if(TitleType == ENodeTitleType::FullTitle)
 	{
 		FFormatNamedArguments Args;
-		Args.Add(TEXT("FunctionName"), FText::FromString(FunctionName));
+		Args.Add(TEXT("FunctionName"), FunctionName);
 		Args.Add(TEXT("ContextString"), FText::FromString(ContextString));
 		Args.Add(TEXT("RPCString"), FText::FromString(RPCString));
 		return FText::Format(LOCTEXT("CallFunction_FullTitle", "{FunctionName}{ContextString}{RPCString}"), Args);
 	}
 	else
 	{
-		return FText::FromString(FunctionName);
+		return FunctionName;
 	}
 }
 
@@ -1260,20 +1260,9 @@ void UK2Node_CallFunction::GeneratePinTooltipFromFunction(UEdGraphPin& Pin, cons
 	GetDefault<UEdGraphSchema_K2>()->ConstructBasicPinTooltip(Pin, FText::FromString(Pin.PinToolTip), Pin.PinToolTip);
 }
 
-FString UK2Node_CallFunction::GetUserFacingFunctionName(const UFunction* Function)
+FText UK2Node_CallFunction::GetUserFacingFunctionName(const UFunction* Function)
 {
-	FString FunctionName = Function->GetMetaData(TEXT("FriendlyName"));
-	
-	if (FunctionName.IsEmpty())
-	{
-		FunctionName = Function->GetName();
-	}
-
-	if( GEditor && GetDefault<UEditorStyleSettings>()->bShowFriendlyNames )
-	{
-		FunctionName = FName::NameToDisplayString(FunctionName, false);
-	}
-	return FunctionName;
+	return Function->GetDisplayNameText();
 }
 
 FString UK2Node_CallFunction::GetDefaultTooltipForFunction(const UFunction* Function)
@@ -1309,7 +1298,7 @@ FString UK2Node_CallFunction::GetDefaultTooltipForFunction(const UFunction* Func
 	}
 	else
 	{
-		return GetUserFacingFunctionName(Function);
+		return GetUserFacingFunctionName(Function).ToString();
 	}
 }
 
@@ -1340,7 +1329,7 @@ FString UK2Node_CallFunction::GetKeywordsForFunction(const UFunction* Function)
 {
 	// If the friendly name and real function name do not match add the real function name friendly name as a keyword.
 	FString Keywords;
-	if( Function->GetName() != GetUserFacingFunctionName(Function) )
+	if( Function->GetName() != GetUserFacingFunctionName(Function).ToString() )
 	{
 		Keywords = Function->GetName();
 	}
diff --git a/Engine/Source/Editor/BlueprintGraph/Private/K2Node_CallParentFunction.cpp b/Engine/Source/Editor/BlueprintGraph/Private/K2Node_CallParentFunction.cpp
--- a/Engine/Source/Editor/BlueprintGraph/Private/K2Node_CallParentFunction.cpp
+++ b/Engine/Source/Editor/BlueprintGraph/Private/K2Node_CallParentFunction.cpp
@@ -16,15 +16,19 @@ UK2Node_CallParentFunction::UK2Node_CallParentFunction(const FObjectInitializer&
 FText UK2Node_CallParentFunction::GetNodeTitle(ENodeTitleType::Type TitleType) const
 {
 	UFunction* Function = FunctionReference.ResolveMember<UFunction>(GetBlueprintClassFromNode());
-	FString FunctionName = Function ? GetUserFacingFunctionName( Function ) : FunctionReference.GetMemberName().ToString();
+	FText FunctionName;
 
-	if( GEditor && GetDefault<UEditorStyleSettings>()->bShowFriendlyNames )
+	if (Function)
 	{
-		FunctionName = FName::NameToDisplayString(FunctionName, false);
+		FunctionName = GetUserFacingFunctionName( Function );
+	}
+	else if ( GEditor && GetDefault<UEditorStyleSettings>()->bShowFriendlyNames )
+	{
+		FunctionName = FText::FromString(FName::NameToDisplayString(FunctionReference.GetMemberName().ToString(), false));
 	}
 
 	FFormatNamedArguments Args;
-	Args.Add(TEXT("FunctionName"),  FText::FromString( FunctionName ));
+	Args.Add(TEXT("FunctionName"), FunctionName);
 	return FText::Format( LOCTEXT( "CallSuperFunction", "Parent: {FunctionName}" ), Args);
 }
 
diff --git a/Engine/Source/Editor/BlueprintGraph/Private/K2Node_Event.cpp b/Engine/Source/Editor/BlueprintGraph/Private/K2Node_Event.cpp
--- a/Engine/Source/Editor/BlueprintGraph/Private/K2Node_Event.cpp
+++ b/Engine/Source/Editor/BlueprintGraph/Private/K2Node_Event.cpp
@@ -74,7 +74,7 @@ FText UK2Node_Event::GetNodeTitle(ENodeTitleType::Type TitleType) const
 {
 	if (bOverrideFunction || (CustomFunctionName == NAME_None))
 	{
-		FString FunctionName = EventReference.GetMemberName().ToString(); // If we fail to find the function, still want to write something on the node.
+		FText FunctionName = FText::FromName(EventReference.GetMemberName()); // If we fail to find the function, still want to write something on the node.
 
 		if (UFunction* Function = EventReference.ResolveMember<UFunction>(GetBlueprintClassFromNode()))
 		{
@@ -82,7 +82,7 @@ FText UK2Node_Event::GetNodeTitle(ENodeTitleType::Type TitleType) const
 		}
 
 		FFormatNamedArguments Args;
-		Args.Add(TEXT("FunctionName"), FText::FromString(FunctionName));
+		Args.Add(TEXT("FunctionName"), FunctionName);
 		FText Title = FText::Format(NSLOCTEXT("K2Node", "Event_Name", "Event {FunctionName}"), Args);
 
 		if(TitleType == ENodeTitleType::FullTitle && EventReference.GetMemberParentClass(GetBlueprintClassFromNode()) != nullptr && EventReference.GetMemberParentClass(GetBlueprintClassFromNode())->IsChildOf(UInterface::StaticClass()))
diff --git a/Engine/Source/Editor/BlueprintGraph/Private/K2Node_Message.cpp b/Engine/Source/Editor/BlueprintGraph/Private/K2Node_Message.cpp
--- a/Engine/Source/Editor/BlueprintGraph/Private/K2Node_Message.cpp
+++ b/Engine/Source/Editor/BlueprintGraph/Private/K2Node_Message.cpp
@@ -21,7 +21,7 @@ FText UK2Node_Message::GetNodeTitle(ENodeTitleType::Type TitleType) const
 	{
 		if (!CachedNodeTitles.IsTitleCached(TitleType))
 		{
-			FText NodeNameText = FText::FromString(UK2Node_CallFunction::GetUserFacingFunctionName(Function));
+			FText NodeNameText = UK2Node_CallFunction::GetUserFacingFunctionName(Function);
 			if (TitleType == ENodeTitleType::MenuTitle)
 			{
 				// FText::Format() is slow, so we cache this to save on performance
diff --git a/Engine/Source/Editor/Kismet/Private/SMyBlueprint.cpp b/Engine/Source/Editor/Kismet/Private/SMyBlueprint.cpp
--- a/Engine/Source/Editor/Kismet/Private/SMyBlueprint.cpp
+++ b/Engine/Source/Editor/Kismet/Private/SMyBlueprint.cpp
@@ -1202,7 +1202,7 @@ void SMyBlueprint::CollectAllActions(FGraphActionListBuilderBase& OutAllActions)
 		if ( UEdGraphSchema_K2::CanKismetOverrideFunction(Function) && !OverridableFunctionNames.Contains(FunctionName) && !ImplementedFunctionCache.Contains(FunctionName) && !FBlueprintEditorUtils::FindOverrideForFunction(BlueprintObj, CastChecked<UClass>(Function->GetOuter()), Function->GetFName()) )
 		{
 			FString FunctionTooltip = UK2Node_CallFunction::GetDefaultTooltipForFunction(Function);
-			FText FunctionDesc = FText::FromString(Function->GetMetaData(TEXT("FriendlyName")));
+			FText FunctionDesc = K2Schema->GetFriendlySignatureName(Function);
 			if ( FunctionDesc.IsEmpty() )
 			{
 				FunctionDesc = FText::FromString(Function->GetName());
@@ -1266,12 +1266,11 @@ void SMyBlueprint::CollectAllActions(FGraphActionListBuilderBase& OutAllActions)
 					if ( UEdGraphSchema_K2::CanKismetOverrideFunction(Function) && !ImplementedFunctionCache.Contains(FunctionName) && !UEdGraphSchema_K2::FunctionCanBePlacedAsEvent(Function) )
 					{
 						FString FunctionTooltip = Function->GetToolTipText().ToString();
-						FString FunctionDesc = Function->GetMetaData(TEXT("FriendlyName"));
-						if (FunctionDesc.IsEmpty()) {FunctionDesc = Function->GetName();}
+						FText FunctionDesc = K2Schema->GetFriendlySignatureName(Function);
 
 						FString FunctionCategory = Function->GetMetaData(FBlueprintMetadata::MD_FunctionCategory);
 
-						TSharedPtr<FEdGraphSchemaAction_K2Graph> NewFuncAction = MakeShareable(new FEdGraphSchemaAction_K2Graph(EEdGraphSchemaAction_K2Graph::Function, FunctionCategory, FText::FromString(FunctionDesc), FunctionTooltip, 1));
+						TSharedPtr<FEdGraphSchemaAction_K2Graph> NewFuncAction = MakeShareable(new FEdGraphSchemaAction_K2Graph(EEdGraphSchemaAction_K2Graph::Function, FunctionCategory, FunctionDesc, FunctionTooltip, 1));
 						NewFuncAction->FuncName = FunctionName;
 						OutAllActions.AddAction(NewFuncAction);
 					}
diff --git a/Engine/Source/Editor/UnrealEd/Private/Editor.cpp b/Engine/Source/Editor/UnrealEd/Private/Editor.cpp
--- a/Engine/Source/Editor/UnrealEd/Private/Editor.cpp
+++ b/Engine/Source/Editor/UnrealEd/Private/Editor.cpp
@@ -4705,11 +4705,6 @@ FString UEditorEngine::GetFriendlyName( const UProperty* Property, UStruct* Owne
 	if ( !DidFindText )
 	{
 		FString DefaultFriendlyName = Property->GetMetaData(TEXT("DisplayName"));
-		if (DefaultFriendlyName.IsEmpty())
-		{
-			// Fallback for old blueprint variables that were saved as friendly name
-			DefaultFriendlyName = Property->GetMetaData(TEXT("FriendlyName"));
-		}
 		if ( DefaultFriendlyName.IsEmpty() )
 		{
 			const bool bIsBool = Cast<const UBoolProperty>(Property) != NULL;
diff --git a/Engine/Source/Runtime/AIModule/Classes/AIController.h b/Engine/Source/Runtime/AIModule/Classes/AIController.h
--- a/Engine/Source/Runtime/AIModule/Classes/AIController.h
+++ b/Engine/Source/Runtime/AIModule/Classes/AIController.h
@@ -254,11 +254,11 @@ public:
 	virtual FVector GetFocalPointOnActor(const AActor *Actor) const;
 
 	/** Set the position that controller should be looking at. */
-	UFUNCTION(BlueprintCallable, Category="AI", meta=(FriendlyName="SetFocalPoint"))
+	UFUNCTION(BlueprintCallable, Category="AI", meta=(DisplayName="SetFocalPoint"))
 	void K2_SetFocalPoint(FVector FP);
 
 	/** Set Focus for actor, will set FocalPoint as a result. */
-	UFUNCTION(BlueprintCallable, Category="AI", meta=(FriendlyName="SetFocus"))
+	UFUNCTION(BlueprintCallable, Category="AI", meta=(DisplayName="SetFocus"))
 	void K2_SetFocus(AActor* NewFocus);
 
 	/** Get the focused actor. */
@@ -268,7 +268,7 @@ public:
 	FORCEINLINE AActor* GetFocusActorForPriority(EAIFocusPriority::Type InPriority) const {  return FocusInformation.Priorities.IsValidIndex(InPriority) ? FocusInformation.Priorities[InPriority].Actor.Get() : nullptr; }
 
 	/** Clears Focus, will also clear FocalPoint as a result */
-	UFUNCTION(BlueprintCallable, Category="AI", meta=(FriendlyName="ClearFocus"))
+	UFUNCTION(BlueprintCallable, Category="AI", meta=(DisplayName="ClearFocus"))
 	void K2_ClearFocus();
 
 
diff --git a/Engine/Source/Runtime/AIModule/Classes/Actions/PawnActionsComponent.h b/Engine/Source/Runtime/AIModule/Classes/Actions/PawnActionsComponent.h
--- a/Engine/Source/Runtime/AIModule/Classes/Actions/PawnActionsComponent.h
+++ b/Engine/Source/Runtime/AIModule/Classes/Actions/PawnActionsComponent.h
@@ -99,7 +99,7 @@ public:
 	// blueprint interface
 	//----------------------------------------------------------------------//
 
-	UFUNCTION(BlueprintCallable, Category = "AI|PawnActions", meta = (FriendlyName = "PerformAction"))
+	UFUNCTION(BlueprintCallable, Category = "AI|PawnActions", meta = (DisplayName = "PerformAction"))
 	static bool K2_PerformAction(APawn* Pawn, UPawnAction* Action, TEnumAsByte<EAIRequestPriority::Type> Priority = EAIRequestPriority::HardScript);
 	static bool PerformAction(APawn& Pawn, UPawnAction& Action, TEnumAsByte<EAIRequestPriority::Type> Priority = EAIRequestPriority::HardScript);
 
@@ -117,17 +117,17 @@ public:
 
 	bool OnEvent(UPawnAction& Action, EPawnActionEventType::Type Event);
 
-	UFUNCTION(BlueprintCallable, Category = PawnAction, meta = (FriendlyName = "PushAction"))
+	UFUNCTION(BlueprintCallable, Category = PawnAction, meta = (DisplayName = "PushAction"))
 	bool K2_PushAction(UPawnAction* NewAction, EAIRequestPriority::Type Priority, UObject* Instigator = NULL);
 	bool PushAction(UPawnAction& NewAction, EAIRequestPriority::Type Priority, UObject* Instigator = NULL);	
 
 	/** Aborts given action instance */
-	UFUNCTION(BlueprintCallable, Category = PawnAction, meta = (FriendlyName = "AbortAction"))
+	UFUNCTION(BlueprintCallable, Category = PawnAction, meta = (DisplayName = "AbortAction"))
 	EPawnActionAbortState::Type K2_AbortAction(UPawnAction* ActionToAbort);
 	EPawnActionAbortState::Type AbortAction(UPawnAction& ActionToAbort);
 
 	/** Aborts given action instance */
-	UFUNCTION(BlueprintCallable, Category = PawnAction, meta = (FriendlyName = "ForceAbortAction"))
+	UFUNCTION(BlueprintCallable, Category = PawnAction, meta = (DisplayName = "ForceAbortAction"))
 	EPawnActionAbortState::Type K2_ForceAbortAction(UPawnAction* ActionToAbort);
 	EPawnActionAbortState::Type ForceAbortAction(UPawnAction& ActionToAbort);
 
diff --git a/Engine/Source/Runtime/AIModule/Classes/BehaviorTree/Decorators/BTDecorator_TagCooldown.h b/Engine/Source/Runtime/AIModule/Classes/BehaviorTree/Decorators/BTDecorator_TagCooldown.h
--- a/Engine/Source/Runtime/AIModule/Classes/BehaviorTree/Decorators/BTDecorator_TagCooldown.h
+++ b/Engine/Source/Runtime/AIModule/Classes/BehaviorTree/Decorators/BTDecorator_TagCooldown.h
@@ -32,7 +32,7 @@ class AIMODULE_API UBTDecorator_TagCooldown : public UBTDecorator
 	bool bAddToExistingDuration;
 
 	/** Whether or not we are adding/setting to the cooldown tag's value when the decorator deactivates. */
-	UPROPERTY(Category = Decorator, EditAnywhere, meta = (FriendlyName = "Adds/Sets Cooldown on Deactivation"))
+	UPROPERTY(Category = Decorator, EditAnywhere, meta = (DisplayName = "Adds/Sets Cooldown on Deactivation"))
 	bool bActivatesCooldown;
 
 	// Begin UObject Interface
diff --git a/Engine/Source/Runtime/Engine/Classes/Components/ActorComponent.h b/Engine/Source/Runtime/Engine/Classes/Components/ActorComponent.h
--- a/Engine/Source/Runtime/Engine/Classes/Components/ActorComponent.h
+++ b/Engine/Source/Runtime/Engine/Classes/Components/ActorComponent.h
@@ -343,7 +343,7 @@ public:
 	 * Blueprint implementable event for when the component is beginning play, called before its Owner's BeginPlay on Actor BeginPlay 
 	 * or when the component is dynamically created if the Actor has already BegunPlay. 
 	 */
-	UFUNCTION(BlueprintImplementableEvent, meta=(FriendlyName = "Begin Play"))
+	UFUNCTION(BlueprintImplementableEvent, meta=(DisplayName = "Begin Play"))
 	void ReceiveBeginPlay();
 
 	/**
@@ -359,7 +359,7 @@ public:
 	virtual void UninitializeComponent();
 
 	/** Blueprint implementable event for when the component ends play, generally via destruction or its Actor's EndPlay. */
-	UFUNCTION(BlueprintImplementableEvent, meta=(Keywords = "delete", FriendlyName = "End Play"))
+	UFUNCTION(BlueprintImplementableEvent, meta=(Keywords = "delete", DisplayName = "End Play"))
 	void ReceiveEndPlay(EEndPlayReason::Type EndPlayReason);
 	
 	/**
@@ -587,7 +587,7 @@ public:
 	/**
 	 * Unregister and mark for pending kill a component.  This may not be used to destroy a component is owned by an actor other than the one calling the function.
 	 */
-	UFUNCTION(BlueprintCallable, Category="Components", meta=(Keywords = "Delete", HidePin="Object", DefaultToSelf="Object", FriendlyName = "DestroyComponent"))
+	UFUNCTION(BlueprintCallable, Category="Components", meta=(Keywords = "Delete", HidePin="Object", DefaultToSelf="Object", DisplayName = "DestroyComponent"))
 	void K2_DestroyComponent(UObject* Object);
 
 	/** Unregisters and immediately re-registers component.  Handles bWillReregister properly. */
@@ -613,7 +613,7 @@ public:
 	virtual void RemoveTickPrerequisiteComponent(UActorComponent* PrerequisiteComponent);
 
 	/** Event called every frame */
-	UFUNCTION(BlueprintImplementableEvent, meta=(FriendlyName = "Tick"))
+	UFUNCTION(BlueprintImplementableEvent, meta=(DisplayName = "Tick"))
 	void ReceiveTick(float DeltaSeconds);
 	
 	/** 
diff --git a/Engine/Source/Runtime/Engine/Classes/Components/ArrowComponent.h b/Engine/Source/Runtime/Engine/Classes/Components/ArrowComponent.h
--- a/Engine/Source/Runtime/Engine/Classes/Components/ArrowComponent.h
+++ b/Engine/Source/Runtime/Engine/Classes/Components/ArrowComponent.h
@@ -53,7 +53,7 @@ class UArrowComponent : public UPrimitiveComponent
 	virtual void SetArrowColor_DEPRECATED(FColor NewColor);
 
 	/** Updates the arrow's colour, and tells it to refresh */
-	UFUNCTION(BlueprintCallable, FriendlyName="SetArrowColor", Category="Components|Arrow")
+	UFUNCTION(BlueprintCallable, DisplayName="SetArrowColor", Category="Components|Arrow")
 	virtual void SetArrowColor_New(FLinearColor NewColor);
 
 	// Begin UPrimitiveComponent interface.
diff --git a/Engine/Source/Runtime/Engine/Classes/Components/AudioComponent.h b/Engine/Source/Runtime/Engine/Classes/Components/AudioComponent.h
--- a/Engine/Source/Runtime/Engine/Classes/Components/AudioComponent.h
+++ b/Engine/Source/Runtime/Engine/Classes/Components/AudioComponent.h
@@ -205,11 +205,11 @@ class ENGINE_API UAudioComponent : public USceneComponent
 	void SetWaveParameter(FName InName, class USoundWave* InWave);
 
 	/** Set a boolean instance parameter for use in sound cues played by this audio component */
-	UFUNCTION(BlueprintCallable, Category="Audio|Components|Audio", meta=(FriendlyName="Set Boolean Parameter"))
+	UFUNCTION(BlueprintCallable, Category="Audio|Components|Audio", meta=(DisplayName="Set Boolean Parameter"))
 	void SetBoolParameter(FName InName, bool InBool);
 	
 	/** Set an integer instance parameter for use in sound cues played by this audio component */
-	UFUNCTION(BlueprintCallable, Category="Audio|Components|Audio", meta=(FriendlyName="Set Integer Parameter"))
+	UFUNCTION(BlueprintCallable, Category="Audio|Components|Audio", meta=(DisplayName="Set Integer Parameter"))
 	void SetIntParameter(FName InName, int32 InInt);
 
 	/** Set a new volume multiplier */
diff --git a/Engine/Source/Runtime/Engine/Classes/Components/LightComponent.h b/Engine/Source/Runtime/Engine/Classes/Components/LightComponent.h
--- a/Engine/Source/Runtime/Engine/Classes/Components/LightComponent.h
+++ b/Engine/Source/Runtime/Engine/Classes/Components/LightComponent.h
@@ -225,7 +225,7 @@ public:
 	UFUNCTION(BlueprintCallable, Category="Rendering|Components|Light")
 	void SetBloomTint(FColor NewValue);
 
-	UFUNCTION(BlueprintCallable, Category="Rendering|Components|Light", meta=(FriendlyName = "Set IES Texture"))
+	UFUNCTION(BlueprintCallable, Category="Rendering|Components|Light", meta=(DisplayName = "Set IES Texture"))
 	void SetIESTexture(UTextureLightProfile* NewValue);
 
 public:
diff --git a/Engine/Source/Runtime/Engine/Classes/Components/SceneComponent.h b/Engine/Source/Runtime/Engine/Classes/Components/SceneComponent.h
--- a/Engine/Source/Runtime/Engine/Classes/Components/SceneComponent.h
+++ b/Engine/Source/Runtime/Engine/Classes/Components/SceneComponent.h
@@ -242,18 +242,18 @@ protected:
 public:
 
 	/** Set the location of this component relative to its parent */
-	UFUNCTION(BlueprintCallable, Category="Utilities|Transformation", meta=(FriendlyName="SetRelativeLocation"))
+	UFUNCTION(BlueprintCallable, Category="Utilities|Transformation", meta=(DisplayName="SetRelativeLocation"))
 	void K2_SetRelativeLocation(FVector NewLocation, bool bSweep, FHitResult& SweepHitResult);
 	void SetRelativeLocation(FVector NewLocation, bool bSweep=false, FHitResult* OutSweepHitResult=nullptr);
 
 
 	/** Set the rotation of this component relative to its parent */
-	UFUNCTION(BlueprintCallable, Category="Utilities|Transformation", meta=(FriendlyName="SetRelativeRotation", AdvancedDisplay="bSweep,SweepHitResult"))
+	UFUNCTION(BlueprintCallable, Category="Utilities|Transformation", meta=(DisplayName="SetRelativeRotation", AdvancedDisplay="bSweep,SweepHitResult"))
 	void K2_SetRelativeRotation(FRotator NewRotation, bool bSweep, FHitResult& SweepHitResult);
 	void SetRelativeRotation(FRotator NewRotation, bool bSweep=false, FHitResult* OutSweepHitResult=nullptr);
 
 	/** Set the transform of this component relative to its parent */
-	UFUNCTION(BlueprintCallable, Category="Utilities|Transformation", meta=(FriendlyName="SetRelativeTransform"))
+	UFUNCTION(BlueprintCallable, Category="Utilities|Transformation", meta=(DisplayName="SetRelativeTransform"))
 	void K2_SetRelativeTransform(const FTransform& NewTransform, bool bSweep, FHitResult& SweepHitResult);
 	void SetRelativeTransform(const FTransform& NewTransform, bool bSweep=false, FHitResult* OutSweepHitResult=nullptr);
 
@@ -270,41 +270,41 @@ public:
 	virtual void SetRelativeScale3D(FVector NewScale3D);
 
 	/** Adds a delta to the translation of this component relative to its parent */
-	UFUNCTION(BlueprintCallable, Category="Utilities|Transformation", meta=(FriendlyName="AddRelativeLocation"))
+	UFUNCTION(BlueprintCallable, Category="Utilities|Transformation", meta=(DisplayName="AddRelativeLocation"))
 	void K2_AddRelativeLocation(FVector DeltaLocation, bool bSweep, FHitResult& SweepHitResult);
 	void AddRelativeLocation(FVector DeltaLocation, bool bSweep=false, FHitResult* OutSweepHitResult=nullptr);
 
 	/** Adds a delta the rotation of this component relative to its parent */
-	UFUNCTION(BlueprintCallable, Category="Utilities|Transformation", meta=(FriendlyName="AddRelativeRotation", AdvancedDisplay="bSweep,SweepHitResult"))
+	UFUNCTION(BlueprintCallable, Category="Utilities|Transformation", meta=(DisplayName="AddRelativeRotation", AdvancedDisplay="bSweep,SweepHitResult"))
 	void K2_AddRelativeRotation(FRotator DeltaRotation, bool bSweep, FHitResult& SweepHitResult);
 	void AddRelativeRotation(FRotator DeltaRotation, bool bSweep=false, FHitResult* OutSweepHitResult=nullptr);
 
 	/** Adds a delta to the location of this component in its local reference frame */
-	UFUNCTION(BlueprintCallable, Category="Utilities|Transformation", meta=(FriendlyName="AddLocalOffset", Keywords="location position"))
+	UFUNCTION(BlueprintCallable, Category="Utilities|Transformation", meta=(DisplayName="AddLocalOffset", Keywords="location position"))
 	void K2_AddLocalOffset(FVector DeltaLocation, bool bSweep, FHitResult& SweepHitResult);
 	void AddLocalOffset(FVector DeltaLocation, bool bSweep=false, FHitResult* OutSweepHitResult=nullptr);
 
 	
 	/** Adds a delta to the rotation of this component in its local reference frame */
-	UFUNCTION(BlueprintCallable, Category="Utilities|Transformation", meta=(FriendlyName="AddLocalRotation", AdvancedDisplay="bSweep,SweepHitResult"))
+	UFUNCTION(BlueprintCallable, Category="Utilities|Transformation", meta=(DisplayName="AddLocalRotation", AdvancedDisplay="bSweep,SweepHitResult"))
 	void K2_AddLocalRotation(FRotator DeltaRotation, bool bSweep, FHitResult& SweepHitResult);
 	void AddLocalRotation(FRotator DeltaRotation, bool bSweep=false, FHitResult* OutSweepHitResult=nullptr);
 
 
 	/** Adds a delta to the transform of this component in its local reference frame. Scale is unchanged. */
-	UFUNCTION(BlueprintCallable, Category="Utilities|Transformation", meta=(FriendlyName="AddLocalTransform"))
+	UFUNCTION(BlueprintCallable, Category="Utilities|Transformation", meta=(DisplayName="AddLocalTransform"))
 	void K2_AddLocalTransform(const FTransform& DeltaTransform, bool bSweep, FHitResult& SweepHitResult);
 	void AddLocalTransform(const FTransform& DeltaTransform, bool bSweep=false, FHitResult* OutSweepHitResult=nullptr);
 
 
 	/** Set the relative location of this component to put it at the supplied location in world space. */
-	UFUNCTION(BlueprintCallable, Category="Utilities|Transformation", meta=(FriendlyName="SetWorldLocation"))
+	UFUNCTION(BlueprintCallable, Category="Utilities|Transformation", meta=(DisplayName="SetWorldLocation"))
 	void K2_SetWorldLocation(FVector NewLocation, bool bSweep, FHitResult& SweepHitResult);
 	void SetWorldLocation(FVector NewLocation, bool bSweep=false, FHitResult* OutSweepHitResult=nullptr);
 
 
 	/** Set the relative rotation of this component to put it at the supplied orientation in world space. */
-	UFUNCTION(BlueprintCallable, Category="Utilities|Transformation", meta=(FriendlyName="SetWorldRotation", AdvancedDisplay="bSweep,SweepHitResult"))
+	UFUNCTION(BlueprintCallable, Category="Utilities|Transformation", meta=(DisplayName="SetWorldRotation", AdvancedDisplay="bSweep,SweepHitResult"))
 	void K2_SetWorldRotation(FRotator NewRotation, bool bSweep, FHitResult& SweepHitResult);
 	void SetWorldRotation(FRotator NewRotation, bool bSweep=false, FHitResult* OutSweepHitResult=nullptr);
 
@@ -314,41 +314,41 @@ public:
 	void SetWorldScale3D(FVector NewScale);
 
 	/** Set the transform of this component in world space. */
-	UFUNCTION(BlueprintCallable, Category="Utilities|Transformation", meta=(FriendlyName="SetWorldTransform"))
+	UFUNCTION(BlueprintCallable, Category="Utilities|Transformation", meta=(DisplayName="SetWorldTransform"))
 	void K2_SetWorldTransform(const FTransform& NewTransform, bool bSweep, FHitResult& SweepHitResult);
 	void SetWorldTransform(const FTransform& NewTransform, bool bSweep=false, FHitResult* OutSweepHitResult=nullptr);
 
 
 	/** Adds a delta to the location of this component in world space. */
-	UFUNCTION(BlueprintCallable, Category="Utilities|Transformation", meta=(FriendlyName="AddWorldOffset", Keywords="location position"))
+	UFUNCTION(BlueprintCallable, Category="Utilities|Transformation", meta=(DisplayName="AddWorldOffset", Keywords="location position"))
 	void K2_AddWorldOffset(FVector DeltaLocation, bool bSweep, FHitResult& SweepHitResult);
 	void AddWorldOffset(FVector DeltaLocation, bool bSweep=false, FHitResult* OutSweepHitResult=nullptr);
 
 
 	/** Adds a delta to the rotation of this component in world space. */
-	UFUNCTION(BlueprintCallable, Category="Utilities|Transformation", meta=(FriendlyName="AddWorldRotation", AdvancedDisplay="bSweep,SweepHitResult"))
+	UFUNCTION(BlueprintCallable, Category="Utilities|Transformation", meta=(DisplayName="AddWorldRotation", AdvancedDisplay="bSweep,SweepHitResult"))
 	void K2_AddWorldRotation(FRotator DeltaRotation, bool bSweep, FHitResult& SweepHitResult);
 	void AddWorldRotation(FRotator DeltaRotation, bool bSweep=false, FHitResult* OutSweepHitResult=nullptr);
 
 	/** Adds a delta to the transform of this component in world space. Scale is unchanged. */
-	UFUNCTION(BlueprintCallable, Category="Utilities|Transformation", meta=(FriendlyName="AddWorldTransform"))
+	UFUNCTION(BlueprintCallable, Category="Utilities|Transformation", meta=(DisplayName="AddWorldTransform"))
 	void K2_AddWorldTransform(const FTransform& DeltaTransform, bool bSweep, FHitResult& SweepHitResult);
 	void AddWorldTransform(const FTransform& DeltaTransform, bool bSweep=false, FHitResult* OutSweepHitResult=nullptr);
 
 	/** Return location of the component, in world space */
-	UFUNCTION(BlueprintCallable, meta=(FriendlyName = "GetWorldLocation"), Category="Utilities|Transformation")
+	UFUNCTION(BlueprintCallable, meta=(DisplayName = "GetWorldLocation"), Category="Utilities|Transformation")
 	FVector K2_GetComponentLocation() const;
 
 	/** Returns rotation of the component, in world space. */
-	UFUNCTION(BlueprintCallable, meta=(FriendlyName = "GetWorldRotation"), Category="Utilities|Transformation")
+	UFUNCTION(BlueprintCallable, meta=(DisplayName = "GetWorldRotation"), Category="Utilities|Transformation")
 	FRotator K2_GetComponentRotation() const;
 	
 	/** Returns scale of the component, in world space. */
-	UFUNCTION(BlueprintCallable, meta=(FriendlyName = "GetWorldScale"), Category="Utilities|Transformation")
+	UFUNCTION(BlueprintCallable, meta=(DisplayName = "GetWorldScale"), Category="Utilities|Transformation")
 	FVector K2_GetComponentScale() const;
 
 	/** Get the current component-to-world transform for this component */
-	UFUNCTION(BlueprintCallable, meta=(FriendlyName = "GetWorldTransform"), Category="Utilities|Transformation")
+	UFUNCTION(BlueprintCallable, meta=(DisplayName = "GetWorldTransform"), Category="Utilities|Transformation")
 	FTransform K2_GetComponentToWorld() const;
 
 	/** Get the forward (X) unit direction vector from this component, in world space.  */
@@ -405,7 +405,7 @@ public:
 	*   Attach this component to another scene component, optionally at a named socket. It is valid to call this on components whether or not they have been Registered.
 	*   @param bMaintainWorldTransform	If true, update the relative location/rotation of this component to keep its world position the same
 	*/
-	UFUNCTION(BlueprintCallable, Category = "Utilities|Transformation", meta = (FriendlyName = "AttachTo", AttachType = "KeepRelativeOffset"))
+	UFUNCTION(BlueprintCallable, Category = "Utilities|Transformation", meta = (DisplayName = "AttachTo", AttachType = "KeepRelativeOffset"))
 	void K2_AttachTo(class USceneComponent* InParent, FName InSocketName = NAME_None, EAttachLocation::Type AttachType = EAttachLocation::KeepRelativeOffset, bool bWeldSimulatedBodies = true);
 
 	/** Zeroes out the relative transform of this component, and calls AttachTo(). Useful for attaching directly to a scene component or socket location  */
@@ -720,7 +720,7 @@ protected:
 	
 public:
 	/** Set the location and rotation of this component relative to its parent */
-	UFUNCTION(BlueprintCallable, Category="Utilities|Transformation", meta=(FriendlyName="SetRelativeLocationAndRotation"))
+	UFUNCTION(BlueprintCallable, Category="Utilities|Transformation", meta=(DisplayName="SetRelativeLocationAndRotation"))
 	void K2_SetRelativeLocationAndRotation(FVector NewLocation, FRotator NewRotation, bool bSweep, FHitResult& SweepHitResult);
 	void SetRelativeLocationAndRotation(FVector NewLocation, FRotator NewRotation, bool bSweep=false, FHitResult* OutSweepHitResult=nullptr);
 
@@ -729,7 +729,7 @@ public:
 	void SetAbsolute(bool bNewAbsoluteLocation = false, bool bNewAbsoluteRotation = false, bool bNewAbsoluteScale = false);
 
 	/** Set the relative location & FRotator rotation of this component to put it at the supplied pose in world space. */
-	UFUNCTION(BlueprintCallable, Category="Utilities|Transformation", meta=(FriendlyName="SetWorldLocationAndRotation"))
+	UFUNCTION(BlueprintCallable, Category="Utilities|Transformation", meta=(DisplayName="SetWorldLocationAndRotation"))
 	void K2_SetWorldLocationAndRotation(FVector NewLocation, FRotator NewRotation, bool bSweep, FHitResult& SweepHitResult);
 	void SetWorldLocationAndRotation(FVector NewLocation, FRotator NewRotation, bool bSweep=false, FHitResult* OutSweepHitResult=nullptr);
 
diff --git a/Engine/Source/Runtime/Engine/Classes/Components/TextRenderComponent.h b/Engine/Source/Runtime/Engine/Classes/Components/TextRenderComponent.h
--- a/Engine/Source/Runtime/Engine/Classes/Components/TextRenderComponent.h
+++ b/Engine/Source/Runtime/Engine/Classes/Components/TextRenderComponent.h
@@ -92,14 +92,14 @@ class ENGINE_API UTextRenderComponent : public UPrimitiveComponent
 	 * The FString variant is deprecated in favor of the FText variant
 	 */
 	DEPRECATED(4.8, "Passing text as FString is deprecated, please use FText instead (likely via a LOCTEXT).")
-	UFUNCTION(BlueprintCallable, Category="Rendering|Components|TextRender", meta=(FriendlyName="Set Text (String)", DeprecatedFunction, DeprecationMessage="Use the SetText function taking an FText instead."))
+	UFUNCTION(BlueprintCallable, Category="Rendering|Components|TextRender", meta=(DisplayName="Set Text (String)", DeprecatedFunction, DeprecationMessage="Use the SetText function taking an FText instead."))
 	void SetText(const FString& Value);
 
 	/** Change the text value and signal the primitives to be rebuilt */
 	void SetText(const FText& Value);
 
 	/** Change the text value and signal the primitives to be rebuilt */
-	UFUNCTION(BlueprintCallable, Category="Rendering|Components|TextRender", meta=(FriendlyName="Set Text"))
+	UFUNCTION(BlueprintCallable, Category="Rendering|Components|TextRender", meta=(DisplayName="Set Text"))
 	void K2_SetText(const FText& Value);
 
 	/** Change the text material and signal the primitives to be rebuilt */
diff --git a/Engine/Source/Runtime/Engine/Classes/EdGraph/EdGraphNode_Documentation.h b/Engine/Source/Runtime/Engine/Classes/EdGraph/EdGraphNode_Documentation.h
--- a/Engine/Source/Runtime/Engine/Classes/EdGraph/EdGraphNode_Documentation.h
+++ b/Engine/Source/Runtime/Engine/Classes/EdGraph/EdGraphNode_Documentation.h
@@ -9,11 +9,11 @@ class UEdGraphNode_Documentation : public UEdGraphNode
 	GENERATED_UCLASS_BODY()
 
 	/** Documentation Link */
-	UPROPERTY(meta=(FriendlyName="Documentation Link"))
+	UPROPERTY(meta=(DisplayName="Documentation Link"))
 	FString Link;
 
 	/** Documentation Excerpt */
-	UPROPERTY(meta=(FriendlyName="Documentation Excerpt"))
+	UPROPERTY(meta=(DisplayName="Documentation Excerpt"))
 	FString Excerpt;
 
 public:
diff --git a/Engine/Source/Runtime/Engine/Classes/Engine/Canvas.h b/Engine/Source/Runtime/Engine/Classes/Engine/Canvas.h
--- a/Engine/Source/Runtime/Engine/Classes/Engine/Canvas.h
+++ b/Engine/Source/Runtime/Engine/Classes/Engine/Canvas.h
@@ -376,7 +376,7 @@ public:
 	 * @param Thickness				How many pixels thick this line should be.
 	 * @param RenderColor			Color to render the line.
 	 */
-	UFUNCTION(BlueprintCallable, Category=Canvas, meta=(FriendlyName="Draw Line"))
+	UFUNCTION(BlueprintCallable, Category=Canvas, meta=(DisplayName="Draw Line"))
 	void K2_DrawLine(FVector2D ScreenPositionA=FVector2D::ZeroVector, FVector2D ScreenPositionB=FVector2D::ZeroVector, float Thickness=1.0f, FLinearColor RenderColor=FLinearColor::White);
 
 	/**
@@ -392,7 +392,7 @@ public:
 	 * @param Rotation					Rotation, in degrees, to render the texture.
 	 * @param PivotPoint				Normalized pivot point to use when rotating the texture.
 	 */
-	UFUNCTION(BlueprintCallable, Category=Canvas, meta=(FriendlyName="Draw Texture"))
+	UFUNCTION(BlueprintCallable, Category=Canvas, meta=(DisplayName="Draw Texture"))
 	void K2_DrawTexture(UTexture* RenderTexture, FVector2D ScreenPosition, FVector2D ScreenSize, FVector2D CoordinatePosition, FVector2D CoordinateSize=FVector2D::UnitVector, FLinearColor RenderColor=FLinearColor::White, EBlendMode BlendMode=BLEND_Translucent, float Rotation=0.f, FVector2D PivotPoint=FVector2D(0.5f,0.5f));
 
 	/**
@@ -406,7 +406,7 @@ public:
 	 * @param Rotation					Rotation, in degrees, to render the texture.
 	 * @param PivotPoint				Normalized pivot point to use when rotating the texture.
 	 */
-	UFUNCTION(BlueprintCallable, Category=Canvas, meta=(FriendlyName="Draw Material"))
+	UFUNCTION(BlueprintCallable, Category=Canvas, meta=(DisplayName="Draw Material"))
 	void K2_DrawMaterial(UMaterialInterface* RenderMaterial, FVector2D ScreenPosition, FVector2D ScreenSize, FVector2D CoordinatePosition, FVector2D CoordinateSize=FVector2D::UnitVector, float Rotation=0.f, FVector2D PivotPoint=FVector2D(0.5f,0.5f));
 
 	/**
@@ -424,7 +424,7 @@ public:
 	 * @param bOutlined					If true, then the text should be rendered with an outline.
 	 * @param OutlineColor				Color to render the outline for the text.
 	 */	
-	UFUNCTION(BlueprintCallable, Category=Canvas, meta=(FriendlyName="Draw Text"))
+	UFUNCTION(BlueprintCallable, Category=Canvas, meta=(DisplayName="Draw Text"))
 	void K2_DrawText(UFont* RenderFont, const FString& RenderText, FVector2D ScreenPosition, FLinearColor RenderColor=FLinearColor::White, float Kerning=0.0f, FLinearColor ShadowColor=FLinearColor::Black, FVector2D ShadowOffset=FVector2D::UnitVector, bool bCentreX=false, bool bCentreY=false, bool bOutlined=false, FLinearColor OutlineColor=FLinearColor::Black);
 	
 	/**
@@ -447,7 +447,7 @@ public:
 	 * @param PivotPoint				Normalized pivot point to use when rotating the texture.
 	 * @param CornerSize				Frame corner size in percent of frame texture (should be < 0.5f).
 	 */
-	UFUNCTION(BlueprintCallable, Category=Canvas, meta=(FriendlyName="Draw Border"))
+	UFUNCTION(BlueprintCallable, Category=Canvas, meta=(DisplayName="Draw Border"))
 	void K2_DrawBorder(UTexture* BorderTexture, UTexture* BackgroundTexture, UTexture* LeftBorderTexture, UTexture* RightBorderTexture, UTexture* TopBorderTexture, UTexture* BottomBorderTexture, FVector2D ScreenPosition, FVector2D ScreenSize, FVector2D CoordinatePosition, FVector2D CoordinateSize=FVector2D::UnitVector, FLinearColor RenderColor=FLinearColor::White, FVector2D BorderScale=FVector2D(0.1f,0.1f), FVector2D BackgroundScale=FVector2D(0.1f,0.1f), float Rotation=0.0f, FVector2D PivotPoint=FVector2D(0.5f,0.5f), FVector2D CornerSize=FVector2D::ZeroVector);
 
 	/**
@@ -457,7 +457,7 @@ public:
 	 * @param ScreenSize				Screen space size to render the texture.
 	 * @param Thickness					How many pixels thick the box lines should be.
 	 */
-	UFUNCTION(BlueprintCallable, Category=Canvas, meta=(FriendlyName="Draw Box"))
+	UFUNCTION(BlueprintCallable, Category=Canvas, meta=(DisplayName="Draw Box"))
 	void K2_DrawBox(FVector2D ScreenPosition, FVector2D ScreenSize, float Thickness=1.0f);
 
 	/**
@@ -466,7 +466,7 @@ public:
 	 * @param RenderTexture				Texture to use when rendering the triangles. If no texture is set, then the default white texture is used.
 	 * @param Triangles					Triangles to render.
 	 */
-	UFUNCTION(BlueprintCallable, Category=Canvas, meta=(FriendlyName="Draw Triangles"))
+	UFUNCTION(BlueprintCallable, Category=Canvas, meta=(DisplayName="Draw Triangles"))
 	void K2_DrawTriangle(UTexture* RenderTexture, TArray<FCanvasUVTri> Triangles);
 
 	/**
@@ -475,7 +475,7 @@ public:
 	 * @param RenderMaterial			Material to use when rendering. Remember that only the emissive channel is able to be rendered as no lighting is performed when rendering to the Canvas.
 	 * @param Triangles					Triangles to render.
 	 */
-	UFUNCTION(BlueprintCallable, Category = Canvas, meta = (FriendlyName = "Draw Material Triangles"))
+	UFUNCTION(BlueprintCallable, Category = Canvas, meta = (DisplayName = "Draw Material Triangles"))
 	void K2_DrawMaterialTriangle(UMaterialInterface* RenderMaterial, TArray<FCanvasUVTri> Triangles);
 	/**
 	 * Draws a polygon on the Canvas.
@@ -486,7 +486,7 @@ public:
 	 * @param NumberOfSides				How many sides this polygon should have. This should be above or equal to three.
 	 * @param RenderColor				Color to tint the polygon.
 	 */
-	UFUNCTION(BlueprintCallable, Category=Canvas, meta=(FriendlyName="Draw Polygon"))
+	UFUNCTION(BlueprintCallable, Category=Canvas, meta=(DisplayName="Draw Polygon"))
 	void K2_DrawPolygon(UTexture* RenderTexture, FVector2D ScreenPosition, FVector2D Radius=FVector2D::UnitVector, int32 NumberOfSides=3, FLinearColor RenderColor=FLinearColor::White);
 
 	/**
@@ -495,7 +495,7 @@ public:
 	 * @param WorldLocation				World space location to project onto the Canvas rendering plane.
 	 * @return							Returns a vector where X, Y defines a screen space position representing the world space location.
 	 */
-	UFUNCTION(BlueprintCallable, Category=Canvas, meta=(FriendlyName="Project"))
+	UFUNCTION(BlueprintCallable, Category=Canvas, meta=(DisplayName="Project"))
 	FVector K2_Project(FVector WorldLocation);
 
 	/**
@@ -505,7 +505,7 @@ public:
 	 * @param WorldOrigin				Vector which is the world position of the screen space position.
 	 * @param WorldDirection			Vector which can be used in a trace to determine what is "behind" the screen space position. Useful for object picking.
 	 */
-	UFUNCTION(BlueprintCallable, Category=Canvas, meta=(FriendlyName="Deproject"))
+	UFUNCTION(BlueprintCallable, Category=Canvas, meta=(DisplayName="Deproject"))
 	void K2_Deproject(FVector2D ScreenPosition, FVector& WorldOrigin, FVector& WorldDirection);
 
 	/**
@@ -515,7 +515,7 @@ public:
 	 * @param RenderText				Text to determine the size of.
 	 * @return							Returns the screen space size of the text.
 	 */
-	UFUNCTION(BlueprintCallable, Category=Canvas, meta=(FriendlyName="Wrapped Text Size"))
+	UFUNCTION(BlueprintCallable, Category=Canvas, meta=(DisplayName="Wrapped Text Size"))
 	FVector2D K2_StrLen(UFont* RenderFont, const FString& RenderText);
 
 	/**
@@ -526,6 +526,6 @@ public:
 	 * @param Scale						Scale of the font to use when determining the size of the text.
 	 * @return							Returns the screen space size of the text.
 	 */
-	UFUNCTION(BlueprintCallable, Category=Canvas, meta=(FriendlyName="Clipped Text Size"))
+	UFUNCTION(BlueprintCallable, Category=Canvas, meta=(DisplayName="Clipped Text Size"))
 	FVector2D K2_TextSize(UFont* RenderFont, const FString& RenderText, FVector2D Scale=FVector2D::UnitVector);
 };
diff --git a/Engine/Source/Runtime/Engine/Classes/Engine/GameInstance.h b/Engine/Source/Runtime/Engine/Classes/Engine/GameInstance.h
--- a/Engine/Source/Runtime/Engine/Classes/Engine/GameInstance.h
+++ b/Engine/Source/Runtime/Engine/Classes/Engine/GameInstance.h
@@ -60,22 +60,22 @@ public:
 	virtual void Init();
 
 	/** Opportunity for blueprints to handle the game instance being initialized. */
-	UFUNCTION(BlueprintImplementableEvent, meta=(FriendlyName = "Init"))
+	UFUNCTION(BlueprintImplementableEvent, meta=(DisplayName = "Init"))
 	void ReceiveInit();
 
 	/** virtual function to allow custom GameInstances an opportunity to do cleanup when shutting down */
 	virtual void Shutdown();
 
 	/** Opportunity for blueprints to handle the game instance being shutdown. */
-	UFUNCTION(BlueprintImplementableEvent, meta=(FriendlyName = "Shutdown"))
+	UFUNCTION(BlueprintImplementableEvent, meta=(DisplayName = "Shutdown"))
 	void ReceiveShutdown();
 
 	/** Opportunity for blueprints to handle network errors. */
-	UFUNCTION(BlueprintImplementableEvent, meta=(FriendlyName = "NetworkError"))
+	UFUNCTION(BlueprintImplementableEvent, meta=(DisplayName = "NetworkError"))
 	virtual void HandleNetworkError(ENetworkFailure::Type FailureType, bool bIsServer);
 
 	/** Opportunity for blueprints to handle travel errors. */
-	UFUNCTION(BlueprintImplementableEvent, meta=(FriendlyName = "TravelError"))
+	UFUNCTION(BlueprintImplementableEvent, meta=(DisplayName = "TravelError"))
 	virtual void HandleTravelError(ETravelFailure::Type FailureType);
 
 	/* Called to initialize the game instance for standalone instances of the game */
diff --git a/Engine/Source/Runtime/Engine/Classes/Engine/Texture2D.h b/Engine/Source/Runtime/Engine/Classes/Engine/Texture2D.h
--- a/Engine/Source/Runtime/Engine/Classes/Engine/Texture2D.h
+++ b/Engine/Source/Runtime/Engine/Classes/Engine/Texture2D.h
@@ -411,13 +411,13 @@ public:
 	/**
 	 * Gets the X size of the texture, in pixels
 	 */
-	UFUNCTION(BlueprintCallable, meta=(FriendlyName = "GetSizeX"), Category="Rendering|Texture")
+	UFUNCTION(BlueprintCallable, meta=(DisplayName = "GetSizeX"), Category="Rendering|Texture")
 	int32 Blueprint_GetSizeX() const;
 
 	/**
 	 * Gets the Y size of the texture, in pixels
 	 */
-	UFUNCTION(BlueprintCallable, meta=(FriendlyName = "GetSizeY"), Category="Rendering|Texture")
+	UFUNCTION(BlueprintCallable, meta=(DisplayName = "GetSizeY"), Category="Rendering|Texture")
 	int32 Blueprint_GetSizeY() const;
 
 	/**
diff --git a/Engine/Source/Runtime/Engine/Classes/GameFramework/Actor.h b/Engine/Source/Runtime/Engine/Classes/GameFramework/Actor.h
--- a/Engine/Source/Runtime/Engine/Classes/GameFramework/Actor.h
+++ b/Engine/Source/Runtime/Engine/Classes/GameFramework/Actor.h
@@ -607,14 +607,14 @@ public:
 	 * Get the actor-to-world transform.
 	 * @return The transform that transforms from actor space to world space.
 	 */
-	UFUNCTION(BlueprintCallable, meta=(FriendlyName = "GetActorTransform"), Category="Utilities|Transformation")
+	UFUNCTION(BlueprintCallable, meta=(DisplayName = "GetActorTransform"), Category="Utilities|Transformation")
 	FTransform GetTransform() const;
 
 	/** Get the local-to-world transform of the RootComponent. Identical to GetTransform(). */
 	FTransform ActorToWorld() const;
 
 	/** Returns the location of the RootComponent of this Actor */
-	UFUNCTION(BlueprintCallable, meta=(FriendlyName = "GetActorLocation", Keywords="position"), Category="Utilities|Transformation")
+	UFUNCTION(BlueprintCallable, meta=(DisplayName = "GetActorLocation", Keywords="position"), Category="Utilities|Transformation")
 	FVector K2_GetActorLocation() const;
 
 	/** 
@@ -624,11 +624,11 @@ public:
 	 *  @param SweepHitResult	The hit result from the move if swept.
 	 *	@return	Whether the location was successfully set (if not swept), or whether movement occurred at all (if swept).
 	 */
-	UFUNCTION(BlueprintCallable, meta=(FriendlyName = "SetActorLocation", Keywords="position"), Category="Utilities|Transformation")
+	UFUNCTION(BlueprintCallable, meta=(DisplayName = "SetActorLocation", Keywords="position"), Category="Utilities|Transformation")
 	bool K2_SetActorLocation(FVector NewLocation, bool bSweep, FHitResult& SweepHitResult);
 
 	/** Returns rotation of the RootComponent of this Actor. */
-	UFUNCTION(BlueprintCallable, meta=(FriendlyName = "GetActorRotation"), Category="Utilities|Transformation")
+	UFUNCTION(BlueprintCallable, meta=(DisplayName = "GetActorRotation"), Category="Utilities|Transformation")
 	FRotator K2_GetActorRotation() const;
 
 	/** Get the forward (X) vector (length 1.0) from this Actor, in world space.  */
@@ -647,11 +647,11 @@ public:
 	 * Returns the bounding box of all components that make up this Actor.
 	 * @param	bOnlyCollidingComponents	If true, will only return the bounding box for components with collision enabled.
 	 */
-	UFUNCTION(BlueprintCallable, Category="Collision", meta=(FriendlyName = "GetActorBounds"))
+	UFUNCTION(BlueprintCallable, Category="Collision", meta=(DisplayName = "GetActorBounds"))
 	void GetActorBounds(bool bOnlyCollidingComponents, FVector& Origin, FVector& BoxExtent) const;
 
 	/** Returns the RootComponent of this Actor */
-	UFUNCTION(BlueprintCallable, meta=(FriendlyName = "GetRootComponent"), Category="Utilities|Transformation")
+	UFUNCTION(BlueprintCallable, meta=(DisplayName = "GetRootComponent"), Category="Utilities|Transformation")
 	class USceneComponent* K2_GetRootComponent() const;
 
 	/** Returns velocity (in cm/s (Unreal Units/second) of the rootcomponent if it is either using physics or has an associated MovementComponent */
@@ -686,7 +686,7 @@ public:
 	 * @param SweepHitResult		The hit result from the move if swept.
 	 * @return	Whether the rotation was successfully set.
 	 */
-	UFUNCTION(BlueprintCallable, Category="Utilities|Transformation", meta=(FriendlyName="SetActorLocationAndRotation"))
+	UFUNCTION(BlueprintCallable, Category="Utilities|Transformation", meta=(DisplayName="SetActorLocationAndRotation"))
 	bool K2_SetActorLocationAndRotation(FVector NewLocation, FRotator NewRotation, bool bSweep, FHitResult& SweepHitResult);
 	
 	/** 
@@ -735,7 +735,7 @@ public:
 	 * @param  bSweep				Whether to sweep to the destination location, triggering overlaps along the way and stopping at the first blocking hit.
 	 * @param  SweepHitResult		The hit result from the move if swept.
 	 */
-	UFUNCTION(BlueprintCallable, Category="Utilities|Transformation", meta=(FriendlyName="AddActorWorldOffset", Keywords="location position"))
+	UFUNCTION(BlueprintCallable, Category="Utilities|Transformation", meta=(DisplayName="AddActorWorldOffset", Keywords="location position"))
 	void K2_AddActorWorldOffset(FVector DeltaLocation, bool bSweep, FHitResult& SweepHitResult);
 
 	/**
@@ -755,13 +755,13 @@ public:
 	 * @param  bSweep				Whether to sweep to the target rotation (not currently supported).
 	 * @param  SweepHitResult		The hit result from the move if swept.
 	 */
-	UFUNCTION(BlueprintCallable, Category="Utilities|Transformation", meta=(FriendlyName="AddActorWorldRotation", AdvancedDisplay="bSweep,SweepHitResult"))
+	UFUNCTION(BlueprintCallable, Category="Utilities|Transformation", meta=(DisplayName="AddActorWorldRotation", AdvancedDisplay="bSweep,SweepHitResult"))
 	void K2_AddActorWorldRotation(FRotator DeltaRotation, bool bSweep, FHitResult& SweepHitResult);
 	void AddActorWorldRotation(FRotator DeltaRotation, bool bSweep=false, FHitResult* OutSweepHitResult=nullptr);
 
 
 	/** Adds a delta to the transform of this actor in world space. Scale is unchanged. */
-	UFUNCTION(BlueprintCallable, Category="Utilities|Transformation", meta=(FriendlyName="AddActorWorldTransform"))
+	UFUNCTION(BlueprintCallable, Category="Utilities|Transformation", meta=(DisplayName="AddActorWorldTransform"))
 	void K2_AddActorWorldTransform(const FTransform& DeltaTransform, bool bSweep, FHitResult& SweepHitResult);
 	void AddActorWorldTransform(const FTransform& DeltaTransform, bool bSweep=false, FHitResult* OutSweepHitResult=nullptr);
 
@@ -770,25 +770,25 @@ public:
 	 * Set the Actors transform to the specified one.
 	 * @param bSweep		Whether to sweep to the destination location, stopping short of the target if blocked by something.
 	 */
-	UFUNCTION(BlueprintCallable, Category="Utilities|Transformation", meta=(FriendlyName="SetActorTransform"))
+	UFUNCTION(BlueprintCallable, Category="Utilities|Transformation", meta=(DisplayName="SetActorTransform"))
 	bool K2_SetActorTransform(const FTransform& NewTransform, bool bSweep, FHitResult& SweepHitResult);
 	bool SetActorTransform(const FTransform& NewTransform, bool bSweep=false, FHitResult* OutSweepHitResult=nullptr);
 
 
 	/** Adds a delta to the location of this component in its local reference frame */
-	UFUNCTION(BlueprintCallable, Category="Utilities|Transformation", meta=(FriendlyName="AddActorLocalOffset", Keywords="location position"))
+	UFUNCTION(BlueprintCallable, Category="Utilities|Transformation", meta=(DisplayName="AddActorLocalOffset", Keywords="location position"))
 	void K2_AddActorLocalOffset(FVector DeltaLocation, bool bSweep, FHitResult& SweepHitResult);
 	void AddActorLocalOffset(FVector DeltaLocation, bool bSweep=false, FHitResult* OutSweepHitResult=nullptr);
 
 
 	/** Adds a delta to the rotation of this component in its local reference frame */
-	UFUNCTION(BlueprintCallable, Category="Utilities|Transformation", meta=(FriendlyName="AddActorLocalRotation", AdvancedDisplay="bSweep,SweepHitResult"))
+	UFUNCTION(BlueprintCallable, Category="Utilities|Transformation", meta=(DisplayName="AddActorLocalRotation", AdvancedDisplay="bSweep,SweepHitResult"))
 	void K2_AddActorLocalRotation(FRotator DeltaRotation, bool bSweep, FHitResult& SweepHitResult);
 	void AddActorLocalRotation(FRotator DeltaRotation, bool bSweep=false, FHitResult* OutSweepHitResult=nullptr);
 
 
 	/** Adds a delta to the transform of this component in its local reference frame */
-	UFUNCTION(BlueprintCallable, Category="Utilities|Transformation", meta=(FriendlyName="AddActorLocalTransform"))
+	UFUNCTION(BlueprintCallable, Category="Utilities|Transformation", meta=(DisplayName="AddActorLocalTransform"))
 	void K2_AddActorLocalTransform(const FTransform& NewTransform, bool bSweep, FHitResult& SweepHitResult);
 	void AddActorLocalTransform(const FTransform& NewTransform, bool bSweep=false, FHitResult* OutSweepHitResult=nullptr);
 
@@ -798,7 +798,7 @@ public:
 	 * @param NewRelativeLocation	New relative location to set the actor's RootComponent to
 	 * @param bSweep				Should we sweep to the destination location. If true, will stop short of the target if blocked by something
 	 */
-	UFUNCTION(BlueprintCallable, Category="Utilities|Transformation", meta=(FriendlyName="SetActorRelativeLocation"))
+	UFUNCTION(BlueprintCallable, Category="Utilities|Transformation", meta=(DisplayName="SetActorRelativeLocation"))
 	void K2_SetActorRelativeLocation(FVector NewRelativeLocation, bool bSweep, FHitResult& SweepHitResult);
 	void SetActorRelativeLocation(FVector NewRelativeLocation, bool bSweep=false, FHitResult* OutSweepHitResult=nullptr);
 
@@ -807,7 +807,7 @@ public:
 	 * @param NewRelativeRotation		New relative rotation to set the actor's RootComponent to
 	 * @param bSweep					Should we sweep to the destination rotation. If true, will stop short of the target if blocked by something
 	 */
-	UFUNCTION(BlueprintCallable, Category="Utilities|Transformation", meta=(FriendlyName="SetActorRelativeRotation", AdvancedDisplay="bSweep,SweepHitResult"))
+	UFUNCTION(BlueprintCallable, Category="Utilities|Transformation", meta=(DisplayName="SetActorRelativeRotation", AdvancedDisplay="bSweep,SweepHitResult"))
 	void K2_SetActorRelativeRotation(FRotator NewRelativeRotation, bool bSweep, FHitResult& SweepHitResult);
 	void SetActorRelativeRotation(FRotator NewRelativeRotation, bool bSweep=false, FHitResult* OutSweepHitResult=nullptr);
 
@@ -816,7 +816,7 @@ public:
 	 * @param NewRelativeTransform		New relative transform to set the actor's RootComponent to
 	 * @param bSweep					Should we sweep to the destination transform. If true, will stop short of the target if blocked by something
 	 */
-	UFUNCTION(BlueprintCallable, Category="Utilities|Transformation", meta=(FriendlyName="SetActorRelativeTransform"))
+	UFUNCTION(BlueprintCallable, Category="Utilities|Transformation", meta=(DisplayName="SetActorRelativeTransform"))
 	void K2_SetActorRelativeTransform(const FTransform& NewRelativeTransform, bool bSweep, FHitResult& SweepHitResult);
 	void SetActorRelativeTransform(const FTransform& NewRelativeTransform, bool bSweep=false, FHitResult* OutSweepHitResult=nullptr);
 
@@ -837,7 +837,7 @@ public:
 	 *	Sets the actor to be hidden in the game
 	 *	@param	bNewHidden	Whether or not to hide the actor and all its components
 	 */
-	UFUNCTION(BlueprintCallable, Category="Rendering", meta=( FriendlyName = "Set Actor Hidden In Game", Keywords = "Visible Hidden Show Hide" ))
+	UFUNCTION(BlueprintCallable, Category="Rendering", meta=( DisplayName = "Set Actor Hidden In Game", Keywords = "Visible Hidden Show Hide" ))
 	virtual void SetActorHiddenInGame(bool bNewHidden);
 
 	/** Allows enabling/disabling collision for the whole actor */
@@ -849,7 +849,7 @@ public:
 	bool GetActorEnableCollision();
 
 	/** Destroy the actor */
-	UFUNCTION(BlueprintCallable, Category="Utilities", meta=(Keywords = "Delete", FriendlyName = "DestroyActor"))
+	UFUNCTION(BlueprintCallable, Category="Utilities", meta=(Keywords = "Delete", DisplayName = "DestroyActor"))
 	virtual void K2_DestroyActor();
 
 	/** Returns whether this actor has network authority */
@@ -865,7 +865,7 @@ public:
 	class UActorComponent* AddComponent(FName TemplateName, bool bManualAttachment, const FTransform& RelativeTransform, const UObject* ComponentTemplateContext);
 
 	/** DEPRECATED - Use Component::DestroyComponent */
-	UFUNCTION(BlueprintCallable, meta=(DeprecatedFunction, DeprecationMessage = "Use Component.DestroyComponent instead", BlueprintProtected = "true", FriendlyName = "DestroyComponent"))
+	UFUNCTION(BlueprintCallable, meta=(DeprecatedFunction, DeprecationMessage = "Use Component.DestroyComponent instead", BlueprintProtected = "true", DisplayName = "DestroyComponent"))
 	void K2_DestroyComponent(UActorComponent* Component);
 
 	/** 
@@ -879,7 +879,7 @@ public:
 	*   @param AttachLocationType	Type of attachment, AbsoluteWorld to keep its world position, RelativeOffset to keep the object's relative offset and SnapTo to snap to the new parent.
 	*/
 
-	UFUNCTION(BlueprintCallable, meta = (FriendlyName = "AttachActorToComponent", AttachLocationType = "KeepRelativeOffset"), Category = "Utilities|Transformation")
+	UFUNCTION(BlueprintCallable, meta = (DisplayName = "AttachActorToComponent", AttachLocationType = "KeepRelativeOffset"), Category = "Utilities|Transformation")
 	void K2_AttachRootComponentTo(class USceneComponent* InParent, FName InSocketName = NAME_None, EAttachLocation::Type AttachLocationType = EAttachLocation::KeepRelativeOffset, bool bWeldSimulatedBodies = true);
 
 	/**
@@ -895,21 +895,21 @@ public:
 	*   @param AttachLocationType	Type of attachment, AbsoluteWorld to keep its world position, RelativeOffset to keep the object's relative offset and SnapTo to snap to the new parent.
 	*/
 
-	UFUNCTION(BlueprintCallable, meta = (FriendlyName = "AttachActorToActor", AttachLocationType = "KeepRelativeOffset"), Category = "Utilities|Transformation")
+	UFUNCTION(BlueprintCallable, meta = (DisplayName = "AttachActorToActor", AttachLocationType = "KeepRelativeOffset"), Category = "Utilities|Transformation")
 	void K2_AttachRootComponentToActor(AActor* InParentActor, FName InSocketName = NAME_None, EAttachLocation::Type AttachLocationType = EAttachLocation::KeepRelativeOffset, bool bWeldSimulatedBodies = true);
 
 	/** 
 	 *  Snap the RootComponent of this Actor to the supplied Actor's root component, optionally at a named socket. It is not valid to call this on components that are not Registered. 
 	 *  If InSocketName == NAME_None, it will attach to origin of the InParentActor. 
 	 */
-	UFUNCTION(BlueprintCallable, meta=(DeprecatedFunction, DeprecationMessage = "Use AttachRootComponentTo with EAttachLocation::SnapToTarget option instead", FriendlyName = "SnapActorTo"), Category="Utilities|Transformation")
+	UFUNCTION(BlueprintCallable, meta=(DeprecatedFunction, DeprecationMessage = "Use AttachRootComponentTo with EAttachLocation::SnapToTarget option instead", DisplayName = "SnapActorTo"), Category="Utilities|Transformation")
 	void SnapRootComponentTo(AActor* InParentActor, FName InSocketName = NAME_None);
 
 	/** 
 	 *  Detaches the RootComponent of this Actor from any SceneComponent it is currently attached to. 
 	 *   @param bMaintainWorldTransform	If true, update the relative location/rotation of this component to keep its world position the same
 	 */
-	UFUNCTION(BlueprintCallable, meta=(FriendlyName = "DetachActorFromActor"), Category="Utilities|Transformation")
+	UFUNCTION(BlueprintCallable, meta=(DisplayName = "DetachActorFromActor"), Category="Utilities|Transformation")
 	void DetachRootComponentFromParent(bool bMaintainWorldPosition = true);
 
 	/** 
@@ -996,7 +996,7 @@ public:
 	// Blueprint
 	
 	/** Event when play begins for this actor. */
-	UFUNCTION(BlueprintImplementableEvent, meta=(FriendlyName = "BeginPlay"))
+	UFUNCTION(BlueprintImplementableEvent, meta=(DisplayName = "BeginPlay"))
 	void ReceiveBeginPlay();
 
 	/** Event when play begins for this actor. */
@@ -1009,22 +1009,22 @@ public:
 	bool HasActorBegunPlay() const { return bActorHasBegunPlay; }
 
 	/** Event when this actor takes ANY damage */
-	UFUNCTION(BlueprintImplementableEvent, BlueprintAuthorityOnly, meta=(FriendlyName = "AnyDamage"), Category="Game|Damage")
+	UFUNCTION(BlueprintImplementableEvent, BlueprintAuthorityOnly, meta=(DisplayName = "AnyDamage"), Category="Game|Damage")
 	void ReceiveAnyDamage(float Damage, const class UDamageType* DamageType, class AController* InstigatedBy, AActor* DamageCauser);
 	
 	/** 
 	 * Event when this actor takes RADIAL damage 
 	 * @todo Pass it the full array of hits instead of just one?
 	 */
-	UFUNCTION(BlueprintImplementableEvent, BlueprintAuthorityOnly, meta=(FriendlyName = "RadialDamage"), Category="Game|Damage")
+	UFUNCTION(BlueprintImplementableEvent, BlueprintAuthorityOnly, meta=(DisplayName = "RadialDamage"), Category="Game|Damage")
 	void ReceiveRadialDamage(float DamageReceived, const class UDamageType* DamageType, FVector Origin, const struct FHitResult& HitInfo, class AController* InstigatedBy, AActor* DamageCauser);
 
 	/** Event when this actor takes POINT damage */
-	UFUNCTION(BlueprintImplementableEvent, BlueprintAuthorityOnly, meta=(FriendlyName = "PointDamage"), Category="Game|Damage")
+	UFUNCTION(BlueprintImplementableEvent, BlueprintAuthorityOnly, meta=(DisplayName = "PointDamage"), Category="Game|Damage")
 	void ReceivePointDamage(float Damage, const class UDamageType* DamageType, FVector HitLocation, FVector HitNormal, class UPrimitiveComponent* HitComponent, FName BoneName, FVector ShotFromDirection, class AController* InstigatedBy, AActor* DamageCauser);
 
 	/** Event called every frame */
-	UFUNCTION(BlueprintImplementableEvent, meta=(FriendlyName = "Tick"))
+	UFUNCTION(BlueprintImplementableEvent, meta=(DisplayName = "Tick"))
 	void ReceiveTick(float DeltaSeconds);
 
 	/** 
@@ -1032,46 +1032,46 @@ public:
 	 *	For events when objects have a blocking collision, for example a player hitting a wall, see 'Hit' events.
 	 *	@note Components on both this and the other Actor must have bGenerateOverlapEvents set to true to generate overlap events.
 	 */
-	UFUNCTION(BlueprintImplementableEvent, meta=(FriendlyName = "ActorBeginOverlap"), Category="Collision")
+	UFUNCTION(BlueprintImplementableEvent, meta=(DisplayName = "ActorBeginOverlap"), Category="Collision")
 	virtual void ReceiveActorBeginOverlap(AActor* OtherActor);
 
 	/** 
 	 *	Event when an actor no longer overlaps another actor, and they have separated. 
 	 *	@note Components on both this and the other Actor must have bGenerateOverlapEvents set to true to generate overlap events.
 	 */
-	UFUNCTION(BlueprintImplementableEvent, meta=(FriendlyName = "ActorEndOverlap"), Category="Collision")
+	UFUNCTION(BlueprintImplementableEvent, meta=(DisplayName = "ActorEndOverlap"), Category="Collision")
 	virtual void ReceiveActorEndOverlap(AActor* OtherActor);
 
 	/** Event when this actor has the mouse moved over it with the clickable interface. */
-	UFUNCTION(BlueprintImplementableEvent, meta=(FriendlyName = "ActorBeginCursorOver"), Category="Mouse Input")
+	UFUNCTION(BlueprintImplementableEvent, meta=(DisplayName = "ActorBeginCursorOver"), Category="Mouse Input")
 	virtual void ReceiveActorBeginCursorOver();
 
 	/** Event when this actor has the mouse moved off of it with the clickable interface. */
-	UFUNCTION(BlueprintImplementableEvent, meta=(FriendlyName = "ActorEndCursorOver"), Category="Mouse Input")
+	UFUNCTION(BlueprintImplementableEvent, meta=(DisplayName = "ActorEndCursorOver"), Category="Mouse Input")
 	virtual void ReceiveActorEndCursorOver();
 
 	/** Event when this actor is clicked by the mouse when using the clickable interface. */
-	UFUNCTION(BlueprintImplementableEvent, meta=(FriendlyName = "ActorOnClicked"), Category="Mouse Input")
+	UFUNCTION(BlueprintImplementableEvent, meta=(DisplayName = "ActorOnClicked"), Category="Mouse Input")
 	virtual void ReceiveActorOnClicked();
 
 	/** Event when this actor is under the mouse when left mouse button is released while using the clickable interface. */
-	UFUNCTION(BlueprintImplementableEvent, meta=(FriendlyName = "ActorOnReleased"), Category="Mouse Input")
+	UFUNCTION(BlueprintImplementableEvent, meta=(DisplayName = "ActorOnReleased"), Category="Mouse Input")
 	virtual void ReceiveActorOnReleased();
 
 	/** Event when this actor is touched when click events are enabled. */
-	UFUNCTION(BlueprintImplementableEvent, meta=(FriendlyName = "BeginInputTouch"), Category="Touch Input")
+	UFUNCTION(BlueprintImplementableEvent, meta=(DisplayName = "BeginInputTouch"), Category="Touch Input")
 	virtual void ReceiveActorOnInputTouchBegin(const ETouchIndex::Type FingerIndex);
 
 	/** Event when this actor is under the finger when untouched when click events are enabled. */
-	UFUNCTION(BlueprintImplementableEvent, meta=(FriendlyName = "EndInputTouch"), Category="Touch Input")
+	UFUNCTION(BlueprintImplementableEvent, meta=(DisplayName = "EndInputTouch"), Category="Touch Input")
 	virtual void ReceiveActorOnInputTouchEnd(const ETouchIndex::Type FingerIndex);
 
 	/** Event when this actor has a finger moved over it with the clickable interface. */
-	UFUNCTION(BlueprintImplementableEvent, meta=(FriendlyName = "TouchEnter"), Category="Touch Input")
+	UFUNCTION(BlueprintImplementableEvent, meta=(DisplayName = "TouchEnter"), Category="Touch Input")
 	virtual void ReceiveActorOnInputTouchEnter(const ETouchIndex::Type FingerIndex);
 
 	/** Event when this actor has a finger moved off of it with the clickable interface. */
-	UFUNCTION(BlueprintImplementableEvent, meta=(FriendlyName = "TouchLeave"), Category="Touch Input")
+	UFUNCTION(BlueprintImplementableEvent, meta=(DisplayName = "TouchLeave"), Category="Touch Input")
 	virtual void ReceiveActorOnInputTouchLeave(const ETouchIndex::Type FingerIndex);
 
 	/** 
@@ -1095,7 +1095,7 @@ public:
 	 * @note When receiving a hit from another object's movement (bSelfMoved is false), the directions of 'Hit.Normal' and 'Hit.ImpactNormal'
 	 * will be adjusted to indicate force from the other object against this object.
 	 */
-	UFUNCTION(BlueprintImplementableEvent, meta=(FriendlyName = "Hit"), Category="Collision")
+	UFUNCTION(BlueprintImplementableEvent, meta=(DisplayName = "Hit"), Category="Collision")
 	virtual void ReceiveHit(class UPrimitiveComponent* MyComp, AActor* Other, class UPrimitiveComponent* OtherComp, bool bSelfMoved, FVector HitLocation, FVector HitNormal, FVector NormalImpulse, const FHitResult& Hit);
 
 	/** Set the lifespan of this actor. When it expires the object will be destroyed. If requested lifespan is 0, the timer is cleared and the actor will not be destroyed. */
@@ -1112,7 +1112,7 @@ public:
 	 * @param	Location	The location.
 	 * @param	Rotation	The rotation.
 	 */
-	UFUNCTION(BlueprintImplementableEvent, meta=(BlueprintInternalUseOnly = "true", FriendlyName = "Construction Script"))
+	UFUNCTION(BlueprintImplementableEvent, meta=(BlueprintInternalUseOnly = "true", DisplayName = "Construction Script"))
 	void UserConstructionScript();
 
 	/**
@@ -1124,7 +1124,7 @@ public:
 	 */
 	bool Destroy(bool bNetForce = false, bool bShouldModifyLevel = true );
 
-	UFUNCTION(BlueprintImplementableEvent, meta = (Keywords = "delete", FriendlyName = "Destroyed"))
+	UFUNCTION(BlueprintImplementableEvent, meta = (Keywords = "delete", DisplayName = "Destroyed"))
 	void ReceiveDestroyed();
 
 	/** Event triggered when the actor is destroyed. */
@@ -1133,7 +1133,7 @@ public:
 
 
 	/** Event to notify blueprints this actor is about to be deleted. */
-	UFUNCTION(BlueprintImplementableEvent, meta=(Keywords = "delete", FriendlyName = "End Play"))
+	UFUNCTION(BlueprintImplementableEvent, meta=(Keywords = "delete", DisplayName = "End Play"))
 	void ReceiveEndPlay(EEndPlayReason::Type EndPlayReason);
 
 	/** Event triggered when the actor is being removed from a level. */
@@ -1823,7 +1823,7 @@ public:
 	 * @param DestRotation The target rotation at the destination
 	 * @return true if the actor has been successfully moved, or false if it couldn't fit.
 	 */
-	UFUNCTION(BlueprintCallable, meta=( FriendlyName="Teleport", Keywords = "Move Position" ), Category="Utilities|Transformation")
+	UFUNCTION(BlueprintCallable, meta=( DisplayName="Teleport", Keywords = "Move Position" ), Category="Utilities|Transformation")
 	bool K2_TeleportTo( FVector DestLocation, FRotator DestRotation );
 
 	/** Called from TeleportTo() when teleport succeeds */
@@ -2076,11 +2076,11 @@ public:
 	virtual void EndViewTarget( class APlayerController* PC );
 
 	/** Event called when this Actor becomes the view target for the given PlayerController. */
-	UFUNCTION(BlueprintImplementableEvent, meta=(FriendlyName="OnBecomeViewTarget", Keywords="Activate Camera"), Category=Actor)
+	UFUNCTION(BlueprintImplementableEvent, meta=(DisplayName="OnBecomeViewTarget", Keywords="Activate Camera"), Category=Actor)
 	void K2_OnBecomeViewTarget( class APlayerController* PC );
 
 	/** Event called when this Actor is no longer the view target for the given PlayerController. */
-	UFUNCTION(BlueprintImplementableEvent, meta=(FriendlyName="OnEndViewTarget", Keywords="Deactivate Camera"), Category=Actor)
+	UFUNCTION(BlueprintImplementableEvent, meta=(DisplayName="OnEndViewTarget", Keywords="Deactivate Camera"), Category=Actor)
 	void K2_OnEndViewTarget( class APlayerController* PC );
 
 	/**
diff --git a/Engine/Source/Runtime/Engine/Classes/GameFramework/Character.h b/Engine/Source/Runtime/Engine/Classes/GameFramework/Character.h
--- a/Engine/Source/Runtime/Engine/Classes/GameFramework/Character.h
+++ b/Engine/Source/Runtime/Engine/Classes/GameFramework/Character.h
@@ -425,7 +425,7 @@ protected:
 	 * @Return Whether the character can jump in the current state. 
 	 */
 
-	UFUNCTION(BlueprintNativeEvent, Category="Pawn|Character|InternalEvents", meta=(FriendlyName="CanJump"))
+	UFUNCTION(BlueprintNativeEvent, Category="Pawn|Character|InternalEvents", meta=(DisplayName="CanJump"))
 	bool CanJumpInternal() const;
 	virtual bool CanJumpInternal_Implementation() const;
 
@@ -559,7 +559,7 @@ public:
 	 * @param	HalfHeightAdjust		difference between default collision half-height, and actual crouched capsule half-height.
 	 * @param	ScaledHalfHeightAdjust	difference after component scale is taken in to account.
 	 */
-	UFUNCTION(BlueprintImplementableEvent, meta=(FriendlyName = "OnEndCrouch"))
+	UFUNCTION(BlueprintImplementableEvent, meta=(DisplayName = "OnEndCrouch"))
 	void K2_OnEndCrouch(float HalfHeightAdjust, float ScaledHalfHeightAdjust);
 
 	/**
@@ -574,7 +574,7 @@ public:
 	 * @param	HalfHeightAdjust		difference between default collision half-height, and actual crouched capsule half-height.
 	 * @param	ScaledHalfHeightAdjust	difference after component scale is taken in to account.
 	 */
-	UFUNCTION(BlueprintImplementableEvent, meta=(FriendlyName = "OnStartCrouch"))
+	UFUNCTION(BlueprintImplementableEvent, meta=(DisplayName = "OnStartCrouch"))
 	void K2_OnStartCrouch(float HalfHeightAdjust, float ScaledHalfHeightAdjust);
 
 	/**
@@ -593,7 +593,7 @@ public:
 	 * @param	PrevCustomMode		Custom mode before the change (applicable if PrevMovementMode is Custom)
 	 * @param	NewCustomMode		New custom mode (applicable if NewMovementMode is Custom)
 	 */
-	UFUNCTION(BlueprintImplementableEvent, meta=(FriendlyName = "OnMovementModeChanged"))
+	UFUNCTION(BlueprintImplementableEvent, meta=(DisplayName = "OnMovementModeChanged"))
 	void K2_OnMovementModeChanged(EMovementMode PrevMovementMode, EMovementMode NewMovementMode, uint8 PrevCustomMode, uint8 NewCustomMode);
 
 	/** Event for implementing custom character movement mode. Called by CharacterMovement if MovementMode is set to Custom. */
diff --git a/Engine/Source/Runtime/Engine/Classes/GameFramework/CharacterMovementComponent.h b/Engine/Source/Runtime/Engine/Classes/GameFramework/CharacterMovementComponent.h
--- a/Engine/Source/Runtime/Engine/Classes/GameFramework/CharacterMovementComponent.h
+++ b/Engine/Source/Runtime/Engine/Classes/GameFramework/CharacterMovementComponent.h
@@ -987,7 +987,7 @@ public:
 	virtual float GetModifiedMaxAcceleration() const;
 	
 	/** @return Maximum acceleration for the current state, based on MaxAcceleration and any additional modifiers. */
-	UFUNCTION(BlueprintCallable, Category="Pawn|Components|CharacterMovement", meta=(DeprecatedFunction, FriendlyName="GetModifiedMaxAcceleration", DeprecationMessage="GetModifiedMaxAcceleration() is deprecated, apply your own modifiers to GetMaxAcceleration() if desired."))
+	UFUNCTION(BlueprintCallable, Category="Pawn|Components|CharacterMovement", meta=(DeprecatedFunction, DisplayName="GetModifiedMaxAcceleration", DeprecationMessage="GetModifiedMaxAcceleration() is deprecated, apply your own modifiers to GetMaxAcceleration() if desired."))
 	virtual float K2_GetModifiedMaxAcceleration() const;
 
 	/** @return Maximum acceleration for the current state. */
@@ -1278,7 +1278,7 @@ public:
 	FORCEINLINE float GetWalkableFloorAngle() const { return WalkableFloorAngle; }
 
 	/** Get the max angle in degrees of a walkable surface for the character. */
-	UFUNCTION(BlueprintCallable, Category="Pawn|Components|CharacterMovement", meta=(FriendlyName = "GetWalkableFloorAngle"))
+	UFUNCTION(BlueprintCallable, Category="Pawn|Components|CharacterMovement", meta=(DisplayName = "GetWalkableFloorAngle"))
 	float K2_GetWalkableFloorAngle() const;
 
 	/** Set the max angle in degrees of a walkable surface for the character. Also computes WalkableFloorZ. */
@@ -1289,7 +1289,7 @@ public:
 	FORCEINLINE float GetWalkableFloorZ() const { return WalkableFloorZ; }
 
 	/** Get the Z component of the normal of the steepest walkable surface for the character. Any lower than this and it is not walkable. */
-	UFUNCTION(BlueprintCallable, Category="Pawn|Components|CharacterMovement", meta=(FriendlyName = "GetWalkableFloorZ"))
+	UFUNCTION(BlueprintCallable, Category="Pawn|Components|CharacterMovement", meta=(DisplayName = "GetWalkableFloorZ"))
 	float K2_GetWalkableFloorZ() const;
 
 	/** Set the Z component of the normal of the steepest walkable surface for the character. Also computes WalkableFloorAngle. */
diff --git a/Engine/Source/Runtime/Engine/Classes/GameFramework/Controller.h b/Engine/Source/Runtime/Engine/Classes/GameFramework/Controller.h
--- a/Engine/Source/Runtime/Engine/Classes/GameFramework/Controller.h
+++ b/Engine/Source/Runtime/Engine/Classes/GameFramework/Controller.h
@@ -161,7 +161,7 @@ public:
 	void ClientSetRotation(FRotator NewRotation, bool bResetCamera = false);
 
 	/** Return the Pawn that is currently 'controlled' by this PlayerController */
-	UFUNCTION(BlueprintCallable, Category="Pawn", meta=(FriendlyName="Get Controlled Pawn"))
+	UFUNCTION(BlueprintCallable, Category="Pawn", meta=(DisplayName="Get Controlled Pawn"))
 	class APawn* K2_GetPawn() const;
 
 	DEPRECATED(4.4, "Use GetPawn() instead of GetControlledPawn()")
diff --git a/Engine/Source/Runtime/Engine/Classes/GameFramework/GameMode.h b/Engine/Source/Runtime/Engine/Classes/GameFramework/GameMode.h
--- a/Engine/Source/Runtime/Engine/Classes/GameFramework/GameMode.h
+++ b/Engine/Source/Runtime/Engine/Classes/GameFramework/GameMode.h
@@ -436,7 +436,7 @@ public:
 	virtual void PostLogin( APlayerController* NewPlayer );
 
 	/** Notification that a player has successfully logged in, and has been given a player controller */
-	UFUNCTION(BlueprintImplementableEvent, Category="Game", meta=(FriendlyName="PostLogin"))
+	UFUNCTION(BlueprintImplementableEvent, Category="Game", meta=(DisplayName="PostLogin"))
 	void K2_PostLogin( APlayerController* NewPlayer );
 
 	/** 
@@ -523,7 +523,7 @@ public:
 	 */
 	virtual class AActor* FindPlayerStart( AController* Player, const FString& IncomingName = TEXT("") );
 
-	UFUNCTION(BlueprintPure, Category=Game, meta=(FriendlyName="FindPlayerStart"))
+	UFUNCTION(BlueprintPure, Category=Game, meta=(DisplayName="FindPlayerStart"))
 	AActor* K2_FindPlayerStart(AController* Player);
 
 	/** 
diff --git a/Engine/Source/Runtime/Engine/Classes/GameFramework/MovementComponent.h b/Engine/Source/Runtime/Engine/Classes/GameFramework/MovementComponent.h
--- a/Engine/Source/Runtime/Engine/Classes/GameFramework/MovementComponent.h
+++ b/Engine/Source/Runtime/Engine/Classes/GameFramework/MovementComponent.h
@@ -177,7 +177,7 @@ public:
 	virtual float GetModifiedMaxSpeed() const;
 
 	/** @return the result of GetMaxSpeed() * GetMaxSpeedModifier(). */
-	UFUNCTION(BlueprintCallable, Category="Components|Movement", meta=(DeprecatedFunction, DisplayName="GetModifiedMaxSpeed", DeprecationMessage="GetModifiedMaxSpeed() is deprecated, apply your own modifiers to GetMaxSpeed() if desired."))
+	UFUNCTION(BlueprintCallable, Category="Components|Movement", meta=(DisplayName="GetModifiedMaxSpeed", DeprecationMessage="GetModifiedMaxSpeed() is deprecated, apply your own modifiers to GetMaxSpeed() if desired."))
 	virtual float K2_GetModifiedMaxSpeed() const;
 
 	/**
@@ -252,7 +252,7 @@ public:
 	 * Respects the plane constraint, if enabled.
 	 * @return True if some movement occurred, false if no movement occurred. Result of any impact will be stored in OutHit.
 	 */
-	UFUNCTION(BlueprintCallable, Category="Components|Movement", meta=(FriendlyName = "MoveUpdatedComponent"))
+	UFUNCTION(BlueprintCallable, Category="Components|Movement", meta=(DisplayName = "MoveUpdatedComponent"))
 	bool K2_MoveUpdatedComponent(FVector Delta, FRotator NewRotation, FHitResult& OutHit, bool bSweep = true);
 
 	/**
diff --git a/Engine/Source/Runtime/Engine/Classes/GameFramework/Pawn.h b/Engine/Source/Runtime/Engine/Classes/GameFramework/Pawn.h
--- a/Engine/Source/Runtime/Engine/Classes/GameFramework/Pawn.h
+++ b/Engine/Source/Runtime/Engine/Classes/GameFramework/Pawn.h
@@ -266,11 +266,11 @@ public:
 	virtual void PossessedBy(AController* NewController);
 
 	/** Event called when the Pawn is possessed by a Controller (normally only occurs on the server/standalone). */
-	UFUNCTION(BlueprintImplementableEvent, meta=(FriendlyName = "Possessed"))
+	UFUNCTION(BlueprintImplementableEvent, meta=(DisplayName = "Possessed"))
 	void ReceivePossessed(AController* NewController);
 
 	/** Event called when the Pawn is no longer possessed by a Controller. */
-	UFUNCTION(BlueprintImplementableEvent, meta=(FriendlyName = "Unpossessed"))
+	UFUNCTION(BlueprintImplementableEvent, meta=(DisplayName = "Unpossessed"))
 	void ReceiveUnpossessed(AController* OldController);
 
 	/** @return true if controlled by a local (not network) Controller.	 */
@@ -442,7 +442,7 @@ public:
 	FVector GetMovementInputVector() const;
 
 	/** (Deprecated) Return the input vector in world space. */
-	UFUNCTION(BlueprintCallable, Category="Pawn|Input", meta=(DeprecatedFunction, FriendlyName="GetMovementInputVector", DeprecationMessage="GetMovementInputVector has been deprecated, use either GetPendingMovementInputVector or GetLastMovementInputVector"))
+	UFUNCTION(BlueprintCallable, Category="Pawn|Input", meta=(DeprecatedFunction, DisplayName="GetMovementInputVector", DeprecationMessage="GetMovementInputVector has been deprecated, use either GetPendingMovementInputVector or GetLastMovementInputVector"))
 	FVector K2_GetMovementInputVector() const;	
 };
 
diff --git a/Engine/Source/Runtime/Engine/Classes/GameFramework/PawnMovementComponent.h b/Engine/Source/Runtime/Engine/Classes/GameFramework/PawnMovementComponent.h
--- a/Engine/Source/Runtime/Engine/Classes/GameFramework/PawnMovementComponent.h
+++ b/Engine/Source/Runtime/Engine/Classes/GameFramework/PawnMovementComponent.h
@@ -87,6 +87,6 @@ public:
 	FVector GetInputVector() const;
 
 	/** (Deprecated) Return the input vector in world space. */
-	UFUNCTION(BlueprintCallable, Category="Pawn|Components|PawnMovement", meta=(DeprecatedFunction, FriendlyName="GetInputVector", DeprecationMessage="GetInputVector has been deprecated, use either GetPendingInputVector or GetLastInputVector"))
+	UFUNCTION(BlueprintCallable, Category="Pawn|Components|PawnMovement", meta=(DeprecatedFunction, DisplayName="GetInputVector", DeprecationMessage="GetInputVector has been deprecated, use either GetPendingInputVector or GetLastInputVector"))
 	FVector K2_GetInputVector() const;
 };
diff --git a/Engine/Source/Runtime/Engine/Classes/GameFramework/PlayerController.h b/Engine/Source/Runtime/Engine/Classes/GameFramework/PlayerController.h
--- a/Engine/Source/Runtime/Engine/Classes/GameFramework/PlayerController.h
+++ b/Engine/Source/Runtime/Engine/Classes/GameFramework/PlayerController.h
@@ -438,18 +438,18 @@ public:
 	bool GetHitResultUnderFingerForObjects(ETouchIndex::Type FingerIndex, const  TArray<TEnumAsByte<EObjectTypeQuery> > & ObjectTypes, bool bTraceComplex, FHitResult& HitResult) const;
 
 	/** Convert current mouse 2D position to World Space 3D position and direction. Returns false if unable to determine value. **/
-	UFUNCTION(BlueprintCallable, Category="Game|Player", meta=(FriendlyName="ConvertMouseLocationToWorldSpace"))
+	UFUNCTION(BlueprintCallable, Category="Game|Player", meta=(DisplayName="ConvertMouseLocationToWorldSpace"))
 	bool DeprojectMousePositionToWorld(FVector& WorldLocation, FVector& WorldDirection) const;
 
 	/** Convert current mouse 2D position to World Space 3D position and direction. Returns false if unable to determine value. **/
-	UFUNCTION(BlueprintCallable, Category = "Game|Player", meta = (FriendlyName = "ConvertScreenLocationToWorldSpace"))
+	UFUNCTION(BlueprintCallable, Category = "Game|Player", meta = (DisplayName = "ConvertScreenLocationToWorldSpace"))
 	bool DeprojectScreenPositionToWorld(float ScreenX, float ScreenY, FVector& WorldLocation, FVector& WorldDirection) const;
 
 	/**
 	 * Convert a World Space 3D position into a 2D Screen Space position.
 	 * @return true if the world coordinate was successfully projected to the screen.
 	 */
-	UFUNCTION(BlueprintCallable, Category = "Game|Player", meta = ( FriendlyName = "ConvertWorldLocationToScreenLocation" ))
+	UFUNCTION(BlueprintCallable, Category = "Game|Player", meta = ( DisplayName = "ConvertWorldLocationToScreenLocation" ))
 	bool ProjectWorldLocationToScreen(FVector WorldLocation, FVector2D& ScreenLocation) const;
 	
 	/**
diff --git a/Engine/Source/Runtime/Engine/Classes/Kismet/GameplayStatics.h b/Engine/Source/Runtime/Engine/Classes/Kismet/GameplayStatics.h
--- a/Engine/Source/Runtime/Engine/Classes/Kismet/GameplayStatics.h
+++ b/Engine/Source/Runtime/Engine/Classes/Kismet/GameplayStatics.h
@@ -143,7 +143,7 @@ class ENGINE_API UGameplayStatics : public UBlueprintFunctionLibrary
 	UFUNCTION(BlueprintPure, Category="Game", meta=(WorldContext="WorldContextObject") )
 	static class AGameState* GetGameState(UObject* WorldContextObject);
 
-	UFUNCTION(BlueprintPure, meta=(FriendlyName = "GetClass"), Category="Utilities")
+	UFUNCTION(BlueprintPure, meta=(DisplayName = "GetClass"), Category="Utilities")
 	static class UClass *GetObjectClass(const UObject *Object);
 
 	/**
@@ -520,7 +520,7 @@ class ENGINE_API UGameplayStatics : public UBlueprintFunctionLibrary
 	 * @param bDrawDebug		When true, a debug arc is drawn (red for an invalid arc, green for a valid arc)
 	 * @return					Returns false if there is no valid solution or the valid solutions are blocked.  Returns true otherwise.
 	 */
-	UFUNCTION(BlueprintCallable, Category="Game|Components|ProjectileMovement", FriendlyName="SuggestProjectileVelocity", meta=(WorldContext="WorldContextObject"))
+	UFUNCTION(BlueprintCallable, Category="Game|Components|ProjectileMovement", DisplayName="SuggestProjectileVelocity", meta=(WorldContext="WorldContextObject"))
 	static bool BlueprintSuggestProjectileVelocity(UObject* WorldContextObject, FVector& TossVelocity, FVector StartLocation, FVector EndLocation, float LaunchSpeed, float OverrideGravityZ, ESuggestProjVelocityTraceOption::Type TraceOption, float CollisionRadius, bool bFavorHighArc, bool bDrawDebug);
 
 	/** Native version, has more options than the Blueprint version. */
