Blueprint Call Function nodes will now display localized function names.

diff --git a/Engine/Source/Runtime/Engine/Classes/Kismet/KismetArrayLibrary.h b/Engine/Source/Runtime/Engine/Classes/Kismet/KismetArrayLibrary.h
--- a/Engine/Source/Runtime/Engine/Classes/Kismet/KismetArrayLibrary.h
+++ b/Engine/Source/Runtime/Engine/Classes/Kismet/KismetArrayLibrary.h
@@ -14,7 +14,7 @@ class ENGINE_API UKismetArrayLibrary : public UBlueprintFunctionLibrary
 	 *@param	NewItem			The item to add to the array
 	 *@return	The index of the newly added item
 	*/
-	UFUNCTION(BlueprintCallable, CustomThunk, meta=(FriendlyName = "Add", CompactNodeTitle = "ADD", ArrayParm = "TargetArray|ArrayProperty", ArrayTypeDependentParams = "NewItem", AutoCreateRefTerm = "NewItem"), Category="Utilities|Array")
+	UFUNCTION(BlueprintCallable, CustomThunk, meta=(DisplayName = "Add", CompactNodeTitle = "ADD", ArrayParm = "TargetArray|ArrayProperty", ArrayTypeDependentParams = "NewItem", AutoCreateRefTerm = "NewItem"), Category="Utilities|Array")
 	static int32 Array_Add(const TArray<int32>& TargetArray, const UArrayProperty* ArrayProperty, const int32& NewItem);
 
 	/**
@@ -24,7 +24,7 @@ class ENGINE_API UKismetArrayLibrary : public UBlueprintFunctionLibrary
 	*@param		NewItem			The item to add to the array
 	*@return	The index of the newly added item, or INDEX_NONE if the item is already present in the array
 	*/
-	UFUNCTION(BlueprintCallable, CustomThunk, meta = (FriendlyName = "Add Unique", CompactNodeTitle = "ADDUNIQUE", ArrayParm = "TargetArray|ArrayProperty", ArrayTypeDependentParams = "NewItem", AutoCreateRefTerm = "NewItem"), Category = "Utilities|Array")
+	UFUNCTION(BlueprintCallable, CustomThunk, meta = (DisplayName = "Add Unique", CompactNodeTitle = "ADDUNIQUE", ArrayParm = "TargetArray|ArrayProperty", ArrayTypeDependentParams = "NewItem", AutoCreateRefTerm = "NewItem"), Category = "Utilities|Array")
 	static int32 Array_AddUnique(const TArray<int32>& TargetArray, const UArrayProperty* ArrayProperty, const int32& NewItem);
 
 	/** 
@@ -32,7 +32,7 @@ class ENGINE_API UKismetArrayLibrary : public UBlueprintFunctionLibrary
 	 *
 	 *@param	TargetArray		The array to shuffle
 	*/
-	UFUNCTION(BlueprintCallable, CustomThunk, meta=(FriendlyName = "Shuffle", CompactNodeTitle = "SHUFFLE", ArrayParm = "TargetArray|ArrayProperty"), Category="Utilities|Array")
+	UFUNCTION(BlueprintCallable, CustomThunk, meta=(DisplayName = "Shuffle", CompactNodeTitle = "SHUFFLE", ArrayParm = "TargetArray|ArrayProperty"), Category="Utilities|Array")
 	static void Array_Shuffle(const TArray<int32>& TargetArray, const UArrayProperty* ArrayProperty);
 
 	/** 
@@ -41,7 +41,7 @@ class ENGINE_API UKismetArrayLibrary : public UBlueprintFunctionLibrary
 	 *@param	TargetArray		The array to add the source array to
 	 *@param	SourceArray		The array to add to the target array
 	*/
-	UFUNCTION(BlueprintCallable, CustomThunk, meta=(FriendlyName = "Append Array", CompactNodeTitle = "APPEND", ArrayParm = "TargetArray|TargetArrayProperty,SourceArray|SourceArrayProperty", ArrayTypeDependentParams = "SourceArray"), Category="Utilities|Array")
+	UFUNCTION(BlueprintCallable, CustomThunk, meta=(DisplayName = "Append Array", CompactNodeTitle = "APPEND", ArrayParm = "TargetArray|TargetArrayProperty,SourceArray|SourceArrayProperty", ArrayTypeDependentParams = "SourceArray"), Category="Utilities|Array")
 	static void Array_Append(const TArray<int32>& TargetArray, const UArrayProperty* TargetArrayProperty, const TArray<int32>& SourceArray, const UArrayProperty* SourceArrayProperty);
 
 	/* 
@@ -51,7 +51,7 @@ class ENGINE_API UKismetArrayLibrary : public UBlueprintFunctionLibrary
 	 *@param	NewItem			The item to insert into the array
 	 *@param	Index			The index at which to insert the item into the array
 	*/
-	UFUNCTION(BlueprintCallable, CustomThunk, meta=(FriendlyName = "Insert", CompactNodeTitle = "INSERT", ArrayParm = "TargetArray|ArrayProperty", ArrayTypeDependentParams = "NewItem", AutoCreateRefTerm = "NewItem"), Category="Utilities|Array")
+	UFUNCTION(BlueprintCallable, CustomThunk, meta=(DisplayName = "Insert", CompactNodeTitle = "INSERT", ArrayParm = "TargetArray|ArrayProperty", ArrayTypeDependentParams = "NewItem", AutoCreateRefTerm = "NewItem"), Category="Utilities|Array")
 	static void Array_Insert(const TArray<int32>& TargetArray, const UArrayProperty* ArrayProperty, const int32& NewItem, int32 Index);
 
 
@@ -61,7 +61,7 @@ class ENGINE_API UKismetArrayLibrary : public UBlueprintFunctionLibrary
 	 *@param	TargetArray		The array to remove from
 	 *@param	IndexToRemove	The index into the array to remove from
 	*/
-	UFUNCTION(BlueprintCallable, CustomThunk, meta=(FriendlyName = "Remove Index", CompactNodeTitle = "REMOVE INDEX", ArrayParm = "TargetArray|ArrayProperty"), Category="Utilities|Array")
+	UFUNCTION(BlueprintCallable, CustomThunk, meta=(DisplayName = "Remove Index", CompactNodeTitle = "REMOVE INDEX", ArrayParm = "TargetArray|ArrayProperty"), Category="Utilities|Array")
 	static void Array_Remove(const TArray<int32>& TargetArray, const UArrayProperty* ArrayProperty, int32 IndexToRemove);
 
 	/* 
@@ -71,7 +71,7 @@ class ENGINE_API UKismetArrayLibrary : public UBlueprintFunctionLibrary
 	 *@param	Item			The item to remove from the array
 	 *@return	True if one or more items were removed
 	*/
-	UFUNCTION(BlueprintCallable, CustomThunk, meta=(FriendlyName = "Remove Item", CompactNodeTitle = "REMOVE", ArrayParm = "TargetArray|ArrayProperty", ArrayTypeDependentParams = "Item", AutoCreateRefTerm = "Item"), Category="Utilities|Array")
+	UFUNCTION(BlueprintCallable, CustomThunk, meta=(DisplayName = "Remove Item", CompactNodeTitle = "REMOVE", ArrayParm = "TargetArray|ArrayProperty", ArrayTypeDependentParams = "Item", AutoCreateRefTerm = "Item"), Category="Utilities|Array")
 	static bool Array_RemoveItem(const TArray<int32>& TargetArray, const UArrayProperty* ArrayProperty, const int32 &Item);
 
 	/* 
@@ -79,7 +79,7 @@ class ENGINE_API UKismetArrayLibrary : public UBlueprintFunctionLibrary
 	 *
 	 *@param	TargetArray		The array to clear
 	*/
-	UFUNCTION(BlueprintCallable, CustomThunk, meta=(FriendlyName = "Clear", CompactNodeTitle = "CLEAR", ArrayParm = "TargetArray|ArrayProperty"), Category="Utilities|Array")
+	UFUNCTION(BlueprintCallable, CustomThunk, meta=(DisplayName = "Clear", CompactNodeTitle = "CLEAR", ArrayParm = "TargetArray|ArrayProperty"), Category="Utilities|Array")
 	static void Array_Clear(const TArray<int32>& TargetArray, const UArrayProperty* ArrayProperty);
 
 	/* 
@@ -88,7 +88,7 @@ class ENGINE_API UKismetArrayLibrary : public UBlueprintFunctionLibrary
 	 *@param	TargetArray		The array to resize
 	 *@param	Size			The new size of the array
 	*/
-	UFUNCTION(BlueprintCallable, CustomThunk, meta=(FriendlyName = "Resize", CompactNodeTitle = "RESIZE", ArrayParm = "TargetArray|ArrayProperty"), Category="Utilities|Array")
+	UFUNCTION(BlueprintCallable, CustomThunk, meta=(DisplayName = "Resize", CompactNodeTitle = "RESIZE", ArrayParm = "TargetArray|ArrayProperty"), Category="Utilities|Array")
 	static void Array_Resize(const TArray<int32>& TargetArray, const UArrayProperty* ArrayProperty, int32 Size);
 
 	/* 
@@ -97,7 +97,7 @@ class ENGINE_API UKismetArrayLibrary : public UBlueprintFunctionLibrary
 	 *@param	TargetArray		The array to get the length of
 	 *@return	The length of the array
 	*/
-	UFUNCTION(BlueprintPure, CustomThunk, meta=(FriendlyName = "Length", CompactNodeTitle = "LENGTH", ArrayParm = "TargetArray|ArrayProperty", Keywords = "num size count"), Category="Utilities|Array")
+	UFUNCTION(BlueprintPure, CustomThunk, meta=(DisplayName = "Length", CompactNodeTitle = "LENGTH", ArrayParm = "TargetArray|ArrayProperty", Keywords = "num size count"), Category="Utilities|Array")
 	static int32 Array_Length(const TArray<int32>& TargetArray, const UArrayProperty* ArrayProperty);
 
 
@@ -107,7 +107,7 @@ class ENGINE_API UKismetArrayLibrary : public UBlueprintFunctionLibrary
 	 *@param	TargetArray		The array to perform the operation on
 	 *@return	The last valid index of the array
 	*/
-	UFUNCTION(BlueprintPure, CustomThunk, meta=(FriendlyName = "Last Index", CompactNodeTitle = "LAST INDEX", ArrayParm = "TargetArray|ArrayProperty"), Category="Utilities|Array")
+	UFUNCTION(BlueprintPure, CustomThunk, meta=(DisplayName = "Last Index", CompactNodeTitle = "LAST INDEX", ArrayParm = "TargetArray|ArrayProperty"), Category="Utilities|Array")
 	static int32 Array_LastIndex(const TArray<int32>& TargetArray, const UArrayProperty* ArrayProperty);
 
 	/* 
@@ -117,7 +117,7 @@ class ENGINE_API UKismetArrayLibrary : public UBlueprintFunctionLibrary
 	 *@param	Index			The index in the array to get an item from
 	 *@return	The item stored at the index
 	*/
-	UFUNCTION(BlueprintPure, CustomThunk, meta=(FriendlyName = "Get", CompactNodeTitle = "GET", ArrayParm = "TargetArray|ArrayProperty", ArrayTypeDependentParams = "Item"), Category="Utilities|Array")
+	UFUNCTION(BlueprintPure, CustomThunk, meta=(DisplayName = "Get", CompactNodeTitle = "GET", ArrayParm = "TargetArray|ArrayProperty", ArrayTypeDependentParams = "Item"), Category="Utilities|Array")
 	static void Array_Get(const TArray<int32>& TargetArray, const UArrayProperty* ArrayProperty, int32 Index, int32& Item);
 
 	/* 
@@ -128,7 +128,7 @@ class ENGINE_API UKismetArrayLibrary : public UBlueprintFunctionLibrary
 	 *@param	Item			The item to assign to the index of the array
 	 *@param	bSizeToFit		If true, the array will expand if Index is greater than the current size of the array
 	*/
-	UFUNCTION(BlueprintCallable, CustomThunk, meta=(FriendlyName = "Set Array Elem", ArrayParm = "TargetArray|ArrayProperty", ArrayTypeDependentParams = "Item", AutoCreateRefTerm = "Item"), Category="Utilities|Array")
+	UFUNCTION(BlueprintCallable, CustomThunk, meta=(DisplayName = "Set Array Elem", ArrayParm = "TargetArray|ArrayProperty", ArrayTypeDependentParams = "Item", AutoCreateRefTerm = "Item"), Category="Utilities|Array")
 	static void Array_Set(const TArray<int32>& TargetArray, const UArrayProperty* ArrayProperty, int32 Index, const int32& Item, bool bSizeToFit);
 
 	/*  
@@ -138,7 +138,7 @@ class ENGINE_API UKismetArrayLibrary : public UBlueprintFunctionLibrary
 	 *@param	ItemToFind		The item to look for
 	 *@return	The index the item was found at, or -1 if not found
 	*/
-	UFUNCTION(BlueprintPure, CustomThunk, meta=(FriendlyName = "Find Item", CompactNodeTitle = "FIND", ArrayParm = "TargetArray|ArrayProperty", ArrayTypeDependentParams = "ItemToFind", AutoCreateRefTerm = "ItemToFind"), Category="Utilities|Array")
+	UFUNCTION(BlueprintPure, CustomThunk, meta=(DisplayName = "Find Item", CompactNodeTitle = "FIND", ArrayParm = "TargetArray|ArrayProperty", ArrayTypeDependentParams = "ItemToFind", AutoCreateRefTerm = "ItemToFind"), Category="Utilities|Array")
 	static int32 Array_Find(const TArray<int32>& TargetArray, const UArrayProperty* ArrayProperty, const int32& ItemToFind);
 
 	/*  
@@ -148,7 +148,7 @@ class ENGINE_API UKismetArrayLibrary : public UBlueprintFunctionLibrary
 	 *@param	ItemToFind		The item to look for
 	 *@return	True if the item was found within the array
 	*/
-	UFUNCTION(BlueprintPure, CustomThunk, meta=(FriendlyName = "Contains Item", CompactNodeTitle = "CONTAINS", ArrayParm = "TargetArray|ArrayProperty", ArrayTypeDependentParams = "ItemToFind", AutoCreateRefTerm = "ItemToFind"), Category="Utilities|Array")
+	UFUNCTION(BlueprintPure, CustomThunk, meta=(DisplayName = "Contains Item", CompactNodeTitle = "CONTAINS", ArrayParm = "TargetArray|ArrayProperty", ArrayTypeDependentParams = "ItemToFind", AutoCreateRefTerm = "ItemToFind"), Category="Utilities|Array")
 	static bool Array_Contains(const TArray<int32>& TargetArray, const UArrayProperty* ArrayProperty, const int32& ItemToFind);
 
 	/*  
@@ -158,7 +158,7 @@ class ENGINE_API UKismetArrayLibrary : public UBlueprintFunctionLibrary
 	 *@param	FilterClass		The Actor sub-class type that acts as the filter, only objects derived from it will be returned.
 	 *@return	An array containing only those objects which are derived from the class specified.
 	*/
-	UFUNCTION(BlueprintCallable, meta=(FriendlyName = "Filter Array"), Category="Utilities|Array")
+	UFUNCTION(BlueprintCallable, meta=(DisplayName = "Filter Array"), Category="Utilities|Array")
 	static void FilterArray(const TArray<AActor*>& TargetArray, TSubclassOf<class AActor> FilterClass, TArray<AActor*>& FilteredArray);
 
 	/** 
diff --git a/Engine/Source/Runtime/Engine/Classes/Kismet/KismetGuidLibrary.h b/Engine/Source/Runtime/Engine/Classes/Kismet/KismetGuidLibrary.h
--- a/Engine/Source/Runtime/Engine/Classes/Kismet/KismetGuidLibrary.h
+++ b/Engine/Source/Runtime/Engine/Classes/Kismet/KismetGuidLibrary.h
@@ -12,19 +12,19 @@ class ENGINE_API UKismetGuidLibrary
 	GENERATED_UCLASS_BODY()
 
 	/** Returns true if the values are equal (A == B) */
-	UFUNCTION(BlueprintPure, meta=(FriendlyName="Equal (Guid)", CompactNodeTitle="=="), Category="Guid")
+	UFUNCTION(BlueprintPure, meta=(DisplayName="Equal (Guid)", CompactNodeTitle="=="), Category="Guid")
 	static bool EqualEqual_GuidGuid( const FGuid& A, const FGuid& B );
 	
 	/** Returns true if the values are not equal (A != B) */
-	UFUNCTION(BlueprintPure, meta=(FriendlyName="NotEqual (Guid)", CompactNodeTitle="!="), Category="Guid")
+	UFUNCTION(BlueprintPure, meta=(DisplayName="NotEqual (Guid)", CompactNodeTitle="!="), Category="Guid")
 	static bool NotEqual_GuidGuid( const FGuid& A, const FGuid& B );
 
 	/** Checks whether the given GUID is valid */
-	UFUNCTION(BlueprintPure, meta=(FriendlyName="IsValid (Guid)", CompactNodeTitle="isValid?"), Category="Guid")
+	UFUNCTION(BlueprintPure, meta=(DisplayName="IsValid (Guid)", CompactNodeTitle="isValid?"), Category="Guid")
 	static bool IsValid_Guid( const FGuid& InGuid );
 
 	/** Invalidates the given GUID */
-	UFUNCTION(BlueprintCallable, meta=(FriendlyName="Invalidate_Guid (Guid)"), Category="Guid")
+	UFUNCTION(BlueprintCallable, meta=(DisplayName="Invalidate_Guid (Guid)"), Category="Guid")
 	static void Invalidate_Guid( UPARAM(ref) FGuid& InGuid );
 
 	/** Returns a new unique GUID */
@@ -32,10 +32,10 @@ class ENGINE_API UKismetGuidLibrary
 	static FGuid NewGuid();
 	
 	/** Converts a GUID value to a string, in the form 'A-B-C-D' */
-	UFUNCTION(BlueprintPure, meta=(FriendlyName="ToString (Guid)", CompactNodeTitle="->"), Category="Guid")
+	UFUNCTION(BlueprintPure, meta=(DisplayName="ToString (Guid)", CompactNodeTitle="->"), Category="Guid")
 	static FString Conv_GuidToString( const FGuid& InGuid );
 
 	/** Converts a String of format EGuidFormats to a Guid. Returns Guid OutGuid, Returns bool Success */
-	UFUNCTION(BlueprintPure, meta=(FriendlyName="Parse String to Guid"), Category="Guid")
+	UFUNCTION(BlueprintPure, meta=(DisplayName="Parse String to Guid"), Category="Guid")
 	static void Parse_StringToGuid( const FString& GuidString, FGuid& OutGuid, bool& Success );
 };
diff --git a/Engine/Source/Runtime/Engine/Classes/Kismet/KismetInputLibrary.h b/Engine/Source/Runtime/Engine/Classes/Kismet/KismetInputLibrary.h
--- a/Engine/Source/Runtime/Engine/Classes/Kismet/KismetInputLibrary.h
+++ b/Engine/Source/Runtime/Engine/Classes/Kismet/KismetInputLibrary.h
@@ -21,7 +21,7 @@ class ENGINE_API UKismetInputLibrary : public UBlueprintFunctionLibrary
 	 * @param B - The key to compare
 	 * @returns True if the key are equal, false otherwise
 	 */
-	UFUNCTION(BlueprintPure, meta=(FriendlyName = "Equal (Key)", CompactNodeTitle = "=="), Category="Utilities|Key")
+	UFUNCTION(BlueprintPure, meta=(DisplayName = "Equal (Key)", CompactNodeTitle = "=="), Category="Utilities|Key")
 	static bool EqualEqual_KeyKey(FKey A, FKey B);
 	
 	/**
@@ -71,7 +71,7 @@ class ENGINE_API UKismetInputLibrary : public UBlueprintFunctionLibrary
 	 *
 	 * @return  True if this character is a repeat
 	 */
-	UFUNCTION(BlueprintPure, meta=( FriendlyName = "IsRepeat" ), Category="Utilities|InputEvent")
+	UFUNCTION(BlueprintPure, meta=( DisplayName = "IsRepeat" ), Category="Utilities|InputEvent")
 	static bool InputEvent_IsRepeat(const FInputEvent& Input);
 
 	/**
@@ -79,7 +79,7 @@ class ENGINE_API UKismetInputLibrary : public UBlueprintFunctionLibrary
 	 *
 	 * @return  True if shift is pressed
 	 */
-	UFUNCTION(BlueprintPure, meta=( FriendlyName = "IsShiftDown" ), Category="Utilities|InputEvent")
+	UFUNCTION(BlueprintPure, meta=( DisplayName = "IsShiftDown" ), Category="Utilities|InputEvent")
 	static bool InputEvent_IsShiftDown(const FInputEvent& Input);
 
 	/**
@@ -87,7 +87,7 @@ class ENGINE_API UKismetInputLibrary : public UBlueprintFunctionLibrary
 	 *
 	 * @return True if left shift is pressed.
 	 */
-	UFUNCTION(BlueprintPure, meta=( FriendlyName = "IsLeftShiftDown" ), Category="Utilities|InputEvent")
+	UFUNCTION(BlueprintPure, meta=( DisplayName = "IsLeftShiftDown" ), Category="Utilities|InputEvent")
 	static bool InputEvent_IsLeftShiftDown(const FInputEvent& Input);
 
 	/**
@@ -95,7 +95,7 @@ class ENGINE_API UKismetInputLibrary : public UBlueprintFunctionLibrary
 	 *
 	 * @return True if right shift is pressed.
 	 */
-	UFUNCTION(BlueprintPure, meta=( FriendlyName = "IsRightShiftDown" ), Category="Utilities|InputEvent")
+	UFUNCTION(BlueprintPure, meta=( DisplayName = "IsRightShiftDown" ), Category="Utilities|InputEvent")
 	static bool InputEvent_IsRightShiftDown(const FInputEvent& Input);
 
 	/**
@@ -103,7 +103,7 @@ class ENGINE_API UKismetInputLibrary : public UBlueprintFunctionLibrary
 	 *
 	 * @return  True if control is pressed
 	 */
-	UFUNCTION(BlueprintPure, meta=( FriendlyName = "IsControlDown" ), Category="Utilities|InputEvent")
+	UFUNCTION(BlueprintPure, meta=( DisplayName = "IsControlDown" ), Category="Utilities|InputEvent")
 	static bool InputEvent_IsControlDown(const FInputEvent& Input);
 
 	/**
@@ -111,7 +111,7 @@ class ENGINE_API UKismetInputLibrary : public UBlueprintFunctionLibrary
 	 *
 	 * @return  True if left control is pressed
 	 */
-	UFUNCTION(BlueprintPure, meta=( FriendlyName = "IsLeftControlDown" ), Category="Utilities|InputEvent")
+	UFUNCTION(BlueprintPure, meta=( DisplayName = "IsLeftControlDown" ), Category="Utilities|InputEvent")
 	static bool InputEvent_IsLeftControlDown(const FInputEvent& Input);
 
 	/**
@@ -119,7 +119,7 @@ class ENGINE_API UKismetInputLibrary : public UBlueprintFunctionLibrary
 	 *
 	 * @return  True if left control is pressed
 	 */
-	UFUNCTION(BlueprintPure, meta=( FriendlyName = "IsRightControlDown" ), Category="Utilities|InputEvent")
+	UFUNCTION(BlueprintPure, meta=( DisplayName = "IsRightControlDown" ), Category="Utilities|InputEvent")
 	static bool InputEvent_IsRightControlDown(const FInputEvent& Input);
 
 	/**
@@ -127,7 +127,7 @@ class ENGINE_API UKismetInputLibrary : public UBlueprintFunctionLibrary
 	 *
 	 * @return  True if alt is pressed
 	 */
-	UFUNCTION(BlueprintPure, meta=( FriendlyName = "IsAltDown" ), Category="Utilities|InputEvent")
+	UFUNCTION(BlueprintPure, meta=( DisplayName = "IsAltDown" ), Category="Utilities|InputEvent")
 	static bool InputEvent_IsAltDown(const FInputEvent& Input);
 
 	/**
@@ -135,7 +135,7 @@ class ENGINE_API UKismetInputLibrary : public UBlueprintFunctionLibrary
 	 *
 	 * @return  True if left alt is pressed
 	 */
-	UFUNCTION(BlueprintPure, meta=( FriendlyName = "IsLeftAltDown" ), Category="Utilities|InputEvent")
+	UFUNCTION(BlueprintPure, meta=( DisplayName = "IsLeftAltDown" ), Category="Utilities|InputEvent")
 	static bool InputEvent_IsLeftAltDown(const FInputEvent& Input);
 
 	/**
@@ -143,7 +143,7 @@ class ENGINE_API UKismetInputLibrary : public UBlueprintFunctionLibrary
 	 *
 	 * @return  True if right alt is pressed
 	 */
-	UFUNCTION(BlueprintPure, meta=( FriendlyName = "IsRightAltDown" ), Category="Utilities|InputEvent")
+	UFUNCTION(BlueprintPure, meta=( DisplayName = "IsRightAltDown" ), Category="Utilities|InputEvent")
 	static bool InputEvent_IsRightAltDown(const FInputEvent& Input);
 
 	/**
@@ -151,7 +151,7 @@ class ENGINE_API UKismetInputLibrary : public UBlueprintFunctionLibrary
 	 *
 	 * @return  True if command is pressed
 	 */
-	UFUNCTION(BlueprintPure, meta=( FriendlyName = "IsCommandDown" ), Category="Utilities|InputEvent")
+	UFUNCTION(BlueprintPure, meta=( DisplayName = "IsCommandDown" ), Category="Utilities|InputEvent")
 	static bool InputEvent_IsCommandDown(const FInputEvent& Input);
 
 	/**
@@ -159,7 +159,7 @@ class ENGINE_API UKismetInputLibrary : public UBlueprintFunctionLibrary
 	 *
 	 * @return  True if left command is pressed
 	 */
-	UFUNCTION(BlueprintPure, meta=( FriendlyName = "IsLeftCommandDown" ), Category="Utilities|InputEvent")
+	UFUNCTION(BlueprintPure, meta=( DisplayName = "IsLeftCommandDown" ), Category="Utilities|InputEvent")
 	static bool InputEvent_IsLeftCommandDown(const FInputEvent& Input);
 
 	/**
@@ -167,7 +167,7 @@ class ENGINE_API UKismetInputLibrary : public UBlueprintFunctionLibrary
 	 *
 	 * @return  True if right command is pressed
 	 */
-	UFUNCTION(BlueprintPure, meta=( FriendlyName = "IsRightCommandDown" ), Category="Utilities|InputEvent")
+	UFUNCTION(BlueprintPure, meta=( DisplayName = "IsRightCommandDown" ), Category="Utilities|InputEvent")
 	static bool InputEvent_IsRightCommandDown(const FInputEvent& Input);
 
 
@@ -186,58 +186,58 @@ class ENGINE_API UKismetInputLibrary : public UBlueprintFunctionLibrary
 	static float GetAnalogValue(const FAnalogInputEvent& Input);
 
 	/** @return The position of the cursor in screen space */
-	UFUNCTION(BlueprintPure, meta=( FriendlyName = "GetScreenSpacePosition" ), Category="Utilities|PointerEvent")
+	UFUNCTION(BlueprintPure, meta=( DisplayName = "GetScreenSpacePosition" ), Category="Utilities|PointerEvent")
 	static FVector2D PointerEvent_GetScreenSpacePosition(const FPointerEvent& Input);
 
 	/** @return The position of the cursor in screen space last time we handled an input event */
-	UFUNCTION(BlueprintPure, meta=( FriendlyName = "GetLastScreenSpacePosition" ), Category="Utilities|PointerEvent")
+	UFUNCTION(BlueprintPure, meta=( DisplayName = "GetLastScreenSpacePosition" ), Category="Utilities|PointerEvent")
 	static FVector2D PointerEvent_GetLastScreenSpacePosition(const FPointerEvent& Input);
 
 	/** @return the distance the mouse traveled since the last event was handled. */
-	UFUNCTION(BlueprintPure, meta=( FriendlyName = "GetCursorDelta" ), Category="Utilities|PointerEvent")
+	UFUNCTION(BlueprintPure, meta=( DisplayName = "GetCursorDelta" ), Category="Utilities|PointerEvent")
 	static FVector2D PointerEvent_GetCursorDelta(const FPointerEvent& Input);
 
 	/** Mouse buttons that are currently pressed */
-	UFUNCTION(BlueprintPure, meta=( FriendlyName = "IsMouseButtonDown" ), Category="Utilities|PointerEvent")
+	UFUNCTION(BlueprintPure, meta=( DisplayName = "IsMouseButtonDown" ), Category="Utilities|PointerEvent")
 	static bool PointerEvent_IsMouseButtonDown(const FPointerEvent& Input, FKey MouseButton);
 
 	/** Mouse button that caused this event to be raised (possibly EB_None) */
-	UFUNCTION(BlueprintPure, meta=( FriendlyName = "GetEffectingButton" ), Category="Utilities|PointerEvent")
+	UFUNCTION(BlueprintPure, meta=( DisplayName = "GetEffectingButton" ), Category="Utilities|PointerEvent")
 	static FKey PointerEvent_GetEffectingButton(const FPointerEvent& Input);
 
 	/** How much did the mouse wheel turn since the last mouse event */
-	UFUNCTION(BlueprintPure, meta=( FriendlyName = "GetWheelDelta" ), Category="Utilities|PointerEvent")
+	UFUNCTION(BlueprintPure, meta=( DisplayName = "GetWheelDelta" ), Category="Utilities|PointerEvent")
 	static float PointerEvent_GetWheelDelta(const FPointerEvent& Input);
 
 	/** @return The index of the user that caused the event */
-	UFUNCTION(BlueprintPure, meta=( FriendlyName = "GetUserIndex" ), Category="Utilities|PointerEvent")
+	UFUNCTION(BlueprintPure, meta=( DisplayName = "GetUserIndex" ), Category="Utilities|PointerEvent")
 	static int32 PointerEvent_GetUserIndex(const FPointerEvent& Input);
 
 	/** @return The unique identifier of the pointer (e.g., finger index) */
-	UFUNCTION(BlueprintPure, meta=( FriendlyName = "GetPointerIndex" ), Category="Utilities|PointerEvent")
+	UFUNCTION(BlueprintPure, meta=( DisplayName = "GetPointerIndex" ), Category="Utilities|PointerEvent")
 	static int32 PointerEvent_GetPointerIndex(const FPointerEvent& Input);
 
 	/** @return The index of the touch pad that generated this event (for platforms with multiple touch pads per user) */
-	UFUNCTION(BlueprintPure, meta=( FriendlyName = "GetTouchpadIndex" ), Category="Utilities|PointerEvent")
+	UFUNCTION(BlueprintPure, meta=( DisplayName = "GetTouchpadIndex" ), Category="Utilities|PointerEvent")
 	static int32 PointerEvent_GetTouchpadIndex(const FPointerEvent& Input);
 
 	/** @return Is this event a result from a touch (as opposed to a mouse) */
-	UFUNCTION(BlueprintPure, meta=( FriendlyName = "IsTouchEvent" ), Category="Utilities|PointerEvent")
+	UFUNCTION(BlueprintPure, meta=( DisplayName = "IsTouchEvent" ), Category="Utilities|PointerEvent")
 	static bool PointerEvent_IsTouchEvent(const FPointerEvent& Input);
 
 	//TODO UMG Support GetGestureType()
 
 	///** @return The type of touch gesture */
-	//UFUNCTION(BlueprintPure, meta=( FriendlyName = "GetGestureType" ), Category="Utilities|PointerEvent")
+	//UFUNCTION(BlueprintPure, meta=( DisplayName = "GetGestureType" ), Category="Utilities|PointerEvent")
 	//static EGestureEvent::Type PointerEvent_GetGestureType(const FPointerEvent& Input);
 
 	/** @return The change in gesture value since the last gesture event of the same type. */
-	UFUNCTION(BlueprintPure, meta=( FriendlyName = "GetGestureDelta" ), Category="Utilities|PointerEvent")
+	UFUNCTION(BlueprintPure, meta=( DisplayName = "GetGestureDelta" ), Category="Utilities|PointerEvent")
 	static FVector2D PointerEvent_GetGestureDelta(const FPointerEvent& Input);
 
 
 	/** @return The controller button that caused this event */
-	UFUNCTION(BlueprintPure, meta = (DeprecatedFunction, DeprecationMessage = "Use GetKey() for KeyEvent instead", FriendlyName = "GetEffectingButton"), Category = "Utilities|ControllerEvent")
+	UFUNCTION(BlueprintPure, meta = (DeprecatedFunction, DeprecationMessage = "Use GetKey() for KeyEvent instead", DisplayName = "GetEffectingButton"), Category = "Utilities|ControllerEvent")
 	static FKey ControllerEvent_GetEffectingButton(const FControllerEvent& Input);
 
 	/** @return The index of the user that caused the event */
