Blueprint Call Function nodes will now display localized function names.

diff --git a/Engine/Source/Runtime/Engine/Classes/Kismet/KismetMathLibrary.h b/Engine/Source/Runtime/Engine/Classes/Kismet/KismetMathLibrary.h
--- a/Engine/Source/Runtime/Engine/Classes/Kismet/KismetMathLibrary.h
+++ b/Engine/Source/Runtime/Engine/Classes/Kismet/KismetMathLibrary.h
@@ -69,27 +69,27 @@ class ENGINE_API UKismetMathLibrary : public UBlueprintFunctionLibrary
 	static bool RandomBool();
 
 	/* Returns the logical complement of the Boolean value (NOT A) */
-	UFUNCTION(BlueprintPure, meta=(FriendlyName = "NOT Boolean", CompactNodeTitle = "NOT", Keywords = "! not"), Category="Math|Boolean")
+	UFUNCTION(BlueprintPure, meta=(DisplayName = "NOT Boolean", CompactNodeTitle = "NOT", Keywords = "! not"), Category="Math|Boolean")
 	static bool Not_PreBool(bool A);
 
 	/* Returns true if the values are equal (A == B) */
-	UFUNCTION(BlueprintPure, meta=(FriendlyName = "Equal Boolean", CompactNodeTitle = "==", Keywords = "== equal"), Category="Math|Boolean")
+	UFUNCTION(BlueprintPure, meta=(DisplayName = "Equal Boolean", CompactNodeTitle = "==", Keywords = "== equal"), Category="Math|Boolean")
 	static bool EqualEqual_BoolBool(bool A, bool B);
 
 	/* Returns true if the values are not equal (A != B) */
-	UFUNCTION(BlueprintPure, meta=(FriendlyName = "NotEqual Boolean", CompactNodeTitle = "!=", Keywords = "!= not equal"), Category="Math|Boolean")
+	UFUNCTION(BlueprintPure, meta=(DisplayName = "NotEqual Boolean", CompactNodeTitle = "!=", Keywords = "!= not equal"), Category="Math|Boolean")
 	static bool NotEqual_BoolBool(bool A, bool B);
 
 	/* Returns the logical AND of two values (A AND B) */
-	UFUNCTION(BlueprintPure, meta=(FriendlyName = "AND Boolean", CompactNodeTitle = "AND", Keywords = "& and", CommutativeAssociativeBinaryOperator = "true"), Category="Math|Boolean")
+	UFUNCTION(BlueprintPure, meta=(DisplayName = "AND Boolean", CompactNodeTitle = "AND", Keywords = "& and", CommutativeAssociativeBinaryOperator = "true"), Category="Math|Boolean")
 	static bool BooleanAND(bool A, bool B);
 
 	/* Returns the logical OR of two values (A OR B) */
-	UFUNCTION(BlueprintPure, meta=(FriendlyName = "OR Boolean", CompactNodeTitle = "OR", Keywords = "| or", CommutativeAssociativeBinaryOperator = "true"), Category="Math|Boolean")
+	UFUNCTION(BlueprintPure, meta=(DisplayName = "OR Boolean", CompactNodeTitle = "OR", Keywords = "| or", CommutativeAssociativeBinaryOperator = "true"), Category="Math|Boolean")
 	static bool BooleanOR(bool A, bool B);
 		
 	/* Returns the logical eXclusive OR of two values (A XOR B) */
-	UFUNCTION(BlueprintPure, meta=(FriendlyName = "XOR Boolean", CompactNodeTitle = "XOR", Keywords = "^ xor"), Category="Math|Boolean")
+	UFUNCTION(BlueprintPure, meta=(DisplayName = "XOR Boolean", CompactNodeTitle = "XOR", Keywords = "^ xor"), Category="Math|Boolean")
 	static bool BooleanXOR(bool A, bool B);
 
 	//
@@ -97,47 +97,47 @@ class ENGINE_API UKismetMathLibrary : public UBlueprintFunctionLibrary
 	//
 
 	/* Multiplication (A * B) */
-	UFUNCTION(BlueprintPure, meta=(FriendlyName = "Byte * Byte", CompactNodeTitle = "*", Keywords = "* multiply", CommutativeAssociativeBinaryOperator = "true"), Category="Math|Byte")
+	UFUNCTION(BlueprintPure, meta=(DisplayName = "Byte * Byte", CompactNodeTitle = "*", Keywords = "* multiply", CommutativeAssociativeBinaryOperator = "true"), Category="Math|Byte")
 	static uint8 Multiply_ByteByte(uint8 A, uint8 B);
 
 	/* Division (A / B) */
-	UFUNCTION(BlueprintPure, meta=(FriendlyName = "Byte / Byte", CompactNodeTitle = "/", Keywords = "/ divide division"), Category="Math|Byte")
+	UFUNCTION(BlueprintPure, meta=(DisplayName = "Byte / Byte", CompactNodeTitle = "/", Keywords = "/ divide division"), Category="Math|Byte")
 	static uint8 Divide_ByteByte(uint8 A, uint8 B = 1);
 
 	/* Modulo (A % B) */
-	UFUNCTION(BlueprintPure, meta=(FriendlyName = "% (Byte)", CompactNodeTitle = "%", Keywords = "% modulus"), Category="Math|Byte")
+	UFUNCTION(BlueprintPure, meta=(DisplayName = "% (Byte)", CompactNodeTitle = "%", Keywords = "% modulus"), Category="Math|Byte")
 	static uint8 Percent_ByteByte(uint8 A, uint8 B = 1);
 
 	/* Addition (A + B) */
-	UFUNCTION(BlueprintPure, meta=(FriendlyName = "Byte + Byte", CompactNodeTitle = "+", Keywords = "+ add plus", CommutativeAssociativeBinaryOperator = "true"), Category="Math|Byte")
+	UFUNCTION(BlueprintPure, meta=(DisplayName = "Byte + Byte", CompactNodeTitle = "+", Keywords = "+ add plus", CommutativeAssociativeBinaryOperator = "true"), Category="Math|Byte")
 	static uint8 Add_ByteByte(uint8 A, uint8 B = 1);
 
 	/* Subtraction (A - B) */
-	UFUNCTION(BlueprintPure, meta=(FriendlyName = "Byte - Byte", CompactNodeTitle = "-", Keywords = "- subtract minus"), Category="Math|Byte")
+	UFUNCTION(BlueprintPure, meta=(DisplayName = "Byte - Byte", CompactNodeTitle = "-", Keywords = "- subtract minus"), Category="Math|Byte")
 	static uint8 Subtract_ByteByte(uint8 A, uint8 B = 1);
 
 	/* Returns true if A is less than B (A < B) */
-	UFUNCTION(BlueprintPure, meta=(FriendlyName = "Byte < Byte", CompactNodeTitle = "<", Keywords = "< less"), Category="Math|Byte")
+	UFUNCTION(BlueprintPure, meta=(DisplayName = "Byte < Byte", CompactNodeTitle = "<", Keywords = "< less"), Category="Math|Byte")
 	static bool Less_ByteByte(uint8 A, uint8 B);
 
 	/* Returns true if A is greater than B (A > B) */
-	UFUNCTION(BlueprintPure, meta=(FriendlyName = "Byte > Byte", CompactNodeTitle = ">", Keywords = "> greater"), Category="Math|Byte")
+	UFUNCTION(BlueprintPure, meta=(DisplayName = "Byte > Byte", CompactNodeTitle = ">", Keywords = "> greater"), Category="Math|Byte")
 	static bool Greater_ByteByte(uint8 A, uint8 B);
 
 	/* Returns true if A is less than or equal to B (A <= B) */
-	UFUNCTION(BlueprintPure, meta=(FriendlyName = "Byte <= Byte", CompactNodeTitle = "<=", Keywords = "<= less"), Category="Math|Byte")
+	UFUNCTION(BlueprintPure, meta=(DisplayName = "Byte <= Byte", CompactNodeTitle = "<=", Keywords = "<= less"), Category="Math|Byte")
 	static bool LessEqual_ByteByte(uint8 A, uint8 B);
 
 	/* Returns true if A is greater than or equal to B (A >= B) */
-	UFUNCTION(BlueprintPure, meta=(FriendlyName = "Byte >= Byte", CompactNodeTitle = ">=", Keywords = ">= greater"), Category="Math|Byte")
+	UFUNCTION(BlueprintPure, meta=(DisplayName = "Byte >= Byte", CompactNodeTitle = ">=", Keywords = ">= greater"), Category="Math|Byte")
 	static bool GreaterEqual_ByteByte(uint8 A, uint8 B);
 
 	/* Returns true if A is equal to B (A == B) */
-	UFUNCTION(BlueprintPure, meta=(FriendlyName = "Equal (Byte)", CompactNodeTitle = "==", Keywords = "== equal"), Category="Math|Byte")
+	UFUNCTION(BlueprintPure, meta=(DisplayName = "Equal (Byte)", CompactNodeTitle = "==", Keywords = "== equal"), Category="Math|Byte")
 	static bool EqualEqual_ByteByte(uint8 A, uint8 B);
 
 	/* Returns true if A is not equal to B (A != B) */
-	UFUNCTION(BlueprintPure, meta=(FriendlyName = "NotEqual (Byte)", CompactNodeTitle = "!=", Keywords = "!= not equal"), Category="Math|Byte")
+	UFUNCTION(BlueprintPure, meta=(DisplayName = "NotEqual (Byte)", CompactNodeTitle = "!=", Keywords = "!= not equal"), Category="Math|Byte")
 	static bool NotEqual_ByteByte(uint8 A, uint8 B);
 
 	//
@@ -145,63 +145,63 @@ class ENGINE_API UKismetMathLibrary : public UBlueprintFunctionLibrary
 	//
 
 	/* Multiplication (A * B) */
-	UFUNCTION(BlueprintPure, meta=(FriendlyName = "integer * integer", CompactNodeTitle = "*", Keywords = "* multiply", CommutativeAssociativeBinaryOperator = "true"), Category="Math|Integer")
+	UFUNCTION(BlueprintPure, meta=(DisplayName = "integer * integer", CompactNodeTitle = "*", Keywords = "* multiply", CommutativeAssociativeBinaryOperator = "true"), Category="Math|Integer")
 	static int32 Multiply_IntInt(int32 A, int32 B);
 
 	/* Division (A / B) */
-	UFUNCTION(BlueprintPure, meta=(FriendlyName = "integer / integer", CompactNodeTitle = "/", Keywords = "/ divide division"), Category="Math|Integer")
+	UFUNCTION(BlueprintPure, meta=(DisplayName = "integer / integer", CompactNodeTitle = "/", Keywords = "/ divide division"), Category="Math|Integer")
 	static int32 Divide_IntInt(int32 A, int32 B = 1);
 
 	/* Modulo (A % B) */
-	UFUNCTION(BlueprintPure, meta=(FriendlyName = "% (integer)", CompactNodeTitle = "%", Keywords = "% modulus"), Category="Math|Integer")
+	UFUNCTION(BlueprintPure, meta=(DisplayName = "% (integer)", CompactNodeTitle = "%", Keywords = "% modulus"), Category="Math|Integer")
 	static int32 Percent_IntInt(int32 A, int32 B = 1);
 
 	/* Addition (A + B) */
-	UFUNCTION(BlueprintPure, meta=(FriendlyName = "integer + integer", CompactNodeTitle = "+", Keywords = "+ add plus", CommutativeAssociativeBinaryOperator = "true"), Category="Math|Integer")
+	UFUNCTION(BlueprintPure, meta=(DisplayName = "integer + integer", CompactNodeTitle = "+", Keywords = "+ add plus", CommutativeAssociativeBinaryOperator = "true"), Category="Math|Integer")
 	static int32 Add_IntInt(int32 A, int32 B = 1);
 
 	/* Subtraction (A - B) */
-	UFUNCTION(BlueprintPure, meta=(FriendlyName = "integer - integer", CompactNodeTitle = "-", Keywords = "- subtract minus"), Category="Math|Integer")
+	UFUNCTION(BlueprintPure, meta=(DisplayName = "integer - integer", CompactNodeTitle = "-", Keywords = "- subtract minus"), Category="Math|Integer")
 	static int32 Subtract_IntInt(int32 A, int32 B = 1);
 
 	/* Returns true if A is less than B (A < B) */
-	UFUNCTION(BlueprintPure, meta=(FriendlyName = "integer < integer", CompactNodeTitle = "<", Keywords = "< less"), Category="Math|Integer")
+	UFUNCTION(BlueprintPure, meta=(DisplayName = "integer < integer", CompactNodeTitle = "<", Keywords = "< less"), Category="Math|Integer")
 	static bool Less_IntInt(int32 A, int32 B);
 
 	/* Returns true if A is greater than B (A > B) */
-	UFUNCTION(BlueprintPure, meta=(FriendlyName = "integer > integer", CompactNodeTitle = ">", Keywords = "> greater"), Category="Math|Integer")
+	UFUNCTION(BlueprintPure, meta=(DisplayName = "integer > integer", CompactNodeTitle = ">", Keywords = "> greater"), Category="Math|Integer")
 	static bool Greater_IntInt(int32 A, int32 B);
 
 	/* Returns true if A is less than or equal to B (A <= B) */
-	UFUNCTION(BlueprintPure, meta=(FriendlyName = "integer <= integer", CompactNodeTitle = "<=", Keywords = "<= less"), Category="Math|Integer")
+	UFUNCTION(BlueprintPure, meta=(DisplayName = "integer <= integer", CompactNodeTitle = "<=", Keywords = "<= less"), Category="Math|Integer")
 	static bool LessEqual_IntInt(int32 A, int32 B);
 
 	/* Returns true if A is greater than or equal to B (A >= B) */
-	UFUNCTION(BlueprintPure, meta=(FriendlyName = "integer >= integer", CompactNodeTitle = ">=", Keywords = ">= greater"), Category="Math|Integer")
+	UFUNCTION(BlueprintPure, meta=(DisplayName = "integer >= integer", CompactNodeTitle = ">=", Keywords = ">= greater"), Category="Math|Integer")
 	static bool GreaterEqual_IntInt(int32 A, int32 B);
 
 	/* Returns true if A is equal to B (A == B) */
-	UFUNCTION(BlueprintPure, meta=(FriendlyName = "Equal (integer)", CompactNodeTitle = "==", Keywords = "== equal"), Category="Math|Integer")
+	UFUNCTION(BlueprintPure, meta=(DisplayName = "Equal (integer)", CompactNodeTitle = "==", Keywords = "== equal"), Category="Math|Integer")
 	static bool EqualEqual_IntInt(int32 A, int32 B);
 
 	/* Returns true if A is not equal to B (A != B) */
-	UFUNCTION(BlueprintPure, meta=(FriendlyName = "NotEqual (integer)", CompactNodeTitle = "!=", Keywords = "!= not equal"), Category="Math|Integer")
+	UFUNCTION(BlueprintPure, meta=(DisplayName = "NotEqual (integer)", CompactNodeTitle = "!=", Keywords = "!= not equal"), Category="Math|Integer")
 	static bool NotEqual_IntInt(int32 A, int32 B);
 
 	/* Bitwise AND (A & B) */
-	UFUNCTION(BlueprintPure, meta=(FriendlyName = "Bitwise AND", CompactNodeTitle = "&", Keywords = "& and", CommutativeAssociativeBinaryOperator = "true"), Category="Math|Integer")
+	UFUNCTION(BlueprintPure, meta=(DisplayName = "Bitwise AND", CompactNodeTitle = "&", Keywords = "& and", CommutativeAssociativeBinaryOperator = "true"), Category="Math|Integer")
 	static int32 And_IntInt(int32 A, int32 B);
 
 	/* Bitwise XOR (A ^ B) */
-	UFUNCTION(BlueprintPure, meta=(FriendlyName = "Bitwise XOR", CompactNodeTitle = "^", Keywords = "^ xor"), Category="Math|Integer")
+	UFUNCTION(BlueprintPure, meta=(DisplayName = "Bitwise XOR", CompactNodeTitle = "^", Keywords = "^ xor"), Category="Math|Integer")
 	static int32 Xor_IntInt(int32 A, int32 B);
 
 	/* Bitwise OR (A | B) */
-	UFUNCTION(BlueprintPure, meta=(FriendlyName = "Bitwise OR", CompactNodeTitle = "|", Keywords = "| or", CommutativeAssociativeBinaryOperator = "true"), Category="Math|Integer")
+	UFUNCTION(BlueprintPure, meta=(DisplayName = "Bitwise OR", CompactNodeTitle = "|", Keywords = "| or", CommutativeAssociativeBinaryOperator = "true"), Category="Math|Integer")
 	static int32 Or_IntInt(int32 A, int32 B);
 
 	/* Sign (integer, returns -1 if A < 0, 0 if A is zero, and +1 if A > 0) */
-	UFUNCTION(BlueprintPure, meta=(FriendlyName = "Sign (int)"), Category="Math|Integer")
+	UFUNCTION(BlueprintPure, meta=(DisplayName = "Sign (int)"), Category="Math|Integer")
 	static int32 SignOfInteger(int32 A);
 
 	/* Returns a uniformly distributed random number between 0 and Max - 1 */
@@ -213,19 +213,19 @@ class ENGINE_API UKismetMathLibrary : public UBlueprintFunctionLibrary
 	static int32 RandomIntegerInRange(int32 Min, int32 Max);
 
 	/* Returns the minimum value of A and B */
-	UFUNCTION(BlueprintPure, meta=(FriendlyName = "Min (int)", CompactNodeTitle = "MIN", CommutativeAssociativeBinaryOperator = "true"), Category="Math|Integer")
+	UFUNCTION(BlueprintPure, meta=(DisplayName = "Min (int)", CompactNodeTitle = "MIN", CommutativeAssociativeBinaryOperator = "true"), Category="Math|Integer")
 	static int32 Min(int32 A, int32 B);
 
 	/* Returns the maximum value of A and B */
-	UFUNCTION(BlueprintPure, meta=(FriendlyName = "Max (int)", CompactNodeTitle = "MAX", CommutativeAssociativeBinaryOperator = "true"), Category="Math|Integer")
+	UFUNCTION(BlueprintPure, meta=(DisplayName = "Max (int)", CompactNodeTitle = "MAX", CommutativeAssociativeBinaryOperator = "true"), Category="Math|Integer")
 	static int32 Max(int32 A, int32 B);
 
 	/* Returns Value clamped to be between A and B (inclusive) */
-	UFUNCTION(BlueprintPure, meta=(FriendlyName = "Clamp (int)"), Category="Math|Integer")
+	UFUNCTION(BlueprintPure, meta=(DisplayName = "Clamp (int)"), Category="Math|Integer")
 	static int32 Clamp(int32 Value, int32 Min, int32 Max);
 
 	/* Returns the absolute (positive) value of A */
-	UFUNCTION(BlueprintPure, meta=(FriendlyName = "Absolute (int)", CompactNodeTitle = "ABS"), Category="Math|Integer")
+	UFUNCTION(BlueprintPure, meta=(DisplayName = "Absolute (int)", CompactNodeTitle = "ABS"), Category="Math|Integer")
 	static int32 Abs_Int(int32 A);
 
 
@@ -233,23 +233,23 @@ class ENGINE_API UKismetMathLibrary : public UBlueprintFunctionLibrary
 	// Float functions.
 	//
 	/* Power (Base to the Exp-th power) */
-	UFUNCTION(BlueprintPure, meta=(FriendlyName = "Power" ), Category="Math|Float")
+	UFUNCTION(BlueprintPure, meta=(DisplayName = "Power" ), Category="Math|Float")
 	static float MultiplyMultiply_FloatFloat(float Base, float Exp);
 
 	/* Multiplication (A * B) */
-	UFUNCTION(BlueprintPure, meta=(FriendlyName = "float * float", CompactNodeTitle = "*", Keywords = "* multiply", CommutativeAssociativeBinaryOperator = "true"), Category="Math|Float")
+	UFUNCTION(BlueprintPure, meta=(DisplayName = "float * float", CompactNodeTitle = "*", Keywords = "* multiply", CommutativeAssociativeBinaryOperator = "true"), Category="Math|Float")
 	static float Multiply_FloatFloat(float A, float B);
 
 	/* Multiplication (A * B) */
-	UFUNCTION(BlueprintPure, meta=(FriendlyName = "int * float", CompactNodeTitle = "*", Keywords = "* multiply"), Category="Math|Float")
+	UFUNCTION(BlueprintPure, meta=(DisplayName = "int * float", CompactNodeTitle = "*", Keywords = "* multiply"), Category="Math|Float")
 	static float Multiply_IntFloat(int32 A, float B);
 
 	/* Division (A / B) */
-	UFUNCTION(BlueprintPure, meta=(FriendlyName = "float / float", CompactNodeTitle = "/", Keywords = "/ divide division"), Category="Math|Float")
+	UFUNCTION(BlueprintPure, meta=(DisplayName = "float / float", CompactNodeTitle = "/", Keywords = "/ divide division"), Category="Math|Float")
 	static float Divide_FloatFloat(float A, float B = 1.f);
 
 	/* Modulo (A % B) */
-	UFUNCTION(BlueprintPure, meta=(FriendlyName = "% (float)", CompactNodeTitle = "%", Keywords = "% modulus"), Category="Math|Float")
+	UFUNCTION(BlueprintPure, meta=(DisplayName = "% (float)", CompactNodeTitle = "%", Keywords = "% modulus"), Category="Math|Float")
 	static float Percent_FloatFloat(float A, float B = 1.f);
 
 	/** Returns the fractional part of a float. */
@@ -257,43 +257,43 @@ class ENGINE_API UKismetMathLibrary : public UBlueprintFunctionLibrary
 	static float Fraction(float A);
 
 	/* Addition (A + B) */
-	UFUNCTION(BlueprintPure, meta=(FriendlyName = "float + float", CompactNodeTitle = "+", Keywords = "+ add plus", CommutativeAssociativeBinaryOperator = "true"), Category="Math|Float")
+	UFUNCTION(BlueprintPure, meta=(DisplayName = "float + float", CompactNodeTitle = "+", Keywords = "+ add plus", CommutativeAssociativeBinaryOperator = "true"), Category="Math|Float")
 	static float Add_FloatFloat(float A, float B = 1.f);
 
 	/* Subtraction (A - B) */
-	UFUNCTION(BlueprintPure, meta=(FriendlyName = "float - float", CompactNodeTitle = "-", Keywords = "- subtract minus"), Category="Math|Float")
+	UFUNCTION(BlueprintPure, meta=(DisplayName = "float - float", CompactNodeTitle = "-", Keywords = "- subtract minus"), Category="Math|Float")
 	static float Subtract_FloatFloat(float A, float B = 1.f);
 
 	/*Returns true if A is Less than B (A < B) */
-	UFUNCTION(BlueprintPure, meta=(FriendlyName = "float < float", CompactNodeTitle = "<", Keywords = "< less"), Category="Math|Float")
+	UFUNCTION(BlueprintPure, meta=(DisplayName = "float < float", CompactNodeTitle = "<", Keywords = "< less"), Category="Math|Float")
 	static bool Less_FloatFloat(float A, float B);
 
 	/*Returns true if A is Greater than B (A > B) */
-	UFUNCTION(BlueprintPure, meta=(FriendlyName = "float > float", CompactNodeTitle = ">", Keywords = "> greater"), Category="Math|Float")
+	UFUNCTION(BlueprintPure, meta=(DisplayName = "float > float", CompactNodeTitle = ">", Keywords = "> greater"), Category="Math|Float")
 	static bool Greater_FloatFloat(float A, float B);
 
 	/* Returns true if A is Less than or equal to B (A <= B) */
-	UFUNCTION(BlueprintPure, meta=(FriendlyName = "float <= float", CompactNodeTitle = "<=", Keywords = "<= less"), Category="Math|Float")
+	UFUNCTION(BlueprintPure, meta=(DisplayName = "float <= float", CompactNodeTitle = "<=", Keywords = "<= less"), Category="Math|Float")
 	static bool LessEqual_FloatFloat(float A, float B);
 
 	/* Returns true if A is Greater than or equal to B (A >= B) */
-	UFUNCTION(BlueprintPure, meta=(FriendlyName = "float >= float", CompactNodeTitle = ">=", Keywords = ">= greater"), Category="Math|Float")
+	UFUNCTION(BlueprintPure, meta=(DisplayName = "float >= float", CompactNodeTitle = ">=", Keywords = ">= greater"), Category="Math|Float")
 	static bool GreaterEqual_FloatFloat(float A, float B);
 
 	/* Returns true if A is exactly equal to B (A == B)*/
-	UFUNCTION(BlueprintPure, meta=(FriendlyName = "Equal (float)", CompactNodeTitle = "==", Keywords = "== equal"), Category="Math|Float")
+	UFUNCTION(BlueprintPure, meta=(DisplayName = "Equal (float)", CompactNodeTitle = "==", Keywords = "== equal"), Category="Math|Float")
 	static bool EqualEqual_FloatFloat(float A, float B);
 
 	/* Returns true if A is nearly equal to B (|A - B| < ErrorTolerance) */
-	UFUNCTION(BlueprintPure, meta=(FriendlyName = "Nearly Equal (float)", Keywords = "== equal"), Category="Math|Float")
+	UFUNCTION(BlueprintPure, meta=(DisplayName = "Nearly Equal (float)", Keywords = "== equal"), Category="Math|Float")
 	static bool NearlyEqual_FloatFloat(float A, float B, float ErrorTolerance = 1.e-6f);
 
 	/* Returns true if A does not equal B (A != B)*/
-	UFUNCTION(BlueprintPure, meta=(FriendlyName = "NotEqual (float)", CompactNodeTitle = "!=", Keywords = "!= not equal"), Category="Math|Float")
+	UFUNCTION(BlueprintPure, meta=(DisplayName = "NotEqual (float)", CompactNodeTitle = "!=", Keywords = "!= not equal"), Category="Math|Float")
 	static bool NotEqual_FloatFloat(float A, float B);
 
 	/* Returns V clamped to be between A and B (inclusive) */
-	UFUNCTION(BlueprintPure, meta=(FriendlyName = "InRange (float)", Min="0.0", Max="1.0"), Category="Math|Float")
+	UFUNCTION(BlueprintPure, meta=(DisplayName = "InRange (float)", Min="0.0", Max="1.0"), Category="Math|Float")
 	static bool InRange_FloatFloat(float Value, float Min, float Max, bool InclusiveMin = true, bool InclusiveMax = true);
 
 	/** Snaps a value to the nearest grid multiple. E.g.,
@@ -301,39 +301,39 @@ class ENGINE_API UKismetMathLibrary : public UBlueprintFunctionLibrary
 	 * If GridSize is 0 Location is returned
 	 * if GridSize is very small precision issues may occur.
 	 */
-	UFUNCTION(BlueprintPure, meta = (FriendlyName = "Snap to grid (float)"), Category = "Math|Float")
+	UFUNCTION(BlueprintPure, meta = (DisplayName = "Snap to grid (float)"), Category = "Math|Float")
 	static float GridSnap_Float(float Location, float GridSize);
 
 	/* Returns the absolute (positive) value of A */
-	UFUNCTION(BlueprintPure, meta=(FriendlyName = "Absolute (float)", CompactNodeTitle = "ABS"), Category="Math|Float")
+	UFUNCTION(BlueprintPure, meta=(DisplayName = "Absolute (float)", CompactNodeTitle = "ABS"), Category="Math|Float")
 	static float Abs(float A);
 
 	/* Returns the sine of A (expects Radians)*/
-	UFUNCTION(BlueprintPure, meta=(FriendlyName = "Sin (Radians)", CompactNodeTitle = "SIN", Keywords = "sine"), Category="Math|Trig")
+	UFUNCTION(BlueprintPure, meta=(DisplayName = "Sin (Radians)", CompactNodeTitle = "SIN", Keywords = "sine"), Category="Math|Trig")
 	static float Sin(float A);
 
 	/* Returns the inverse sine (arcsin) of A (result is in Radians) */
-	UFUNCTION(BlueprintPure, meta=(FriendlyName = "Asin (Radians)", CompactNodeTitle = "ASIN", Keywords = "sine"), Category="Math|Trig")
+	UFUNCTION(BlueprintPure, meta=(DisplayName = "Asin (Radians)", CompactNodeTitle = "ASIN", Keywords = "sine"), Category="Math|Trig")
 	static float Asin(float A);
 
 	/* Returns the cosine of A (expects Radians)*/
-	UFUNCTION(BlueprintPure, meta=(FriendlyName = "Cos (Radians)", CompactNodeTitle = "COS"), Category="Math|Trig")
+	UFUNCTION(BlueprintPure, meta=(DisplayName = "Cos (Radians)", CompactNodeTitle = "COS"), Category="Math|Trig")
 	static float Cos(float A);
 
 	/* Returns the inverse cosine (arccos) of A (result is in Radians) */
-	UFUNCTION(BlueprintPure, meta=(FriendlyName = "Acos (Radians)", CompactNodeTitle = "ACOS"), Category="Math|Trig")
+	UFUNCTION(BlueprintPure, meta=(DisplayName = "Acos (Radians)", CompactNodeTitle = "ACOS"), Category="Math|Trig")
 	static float Acos(float A);
 
 	/* Returns the tan of A (expects Radians)*/
-	UFUNCTION(BlueprintPure, meta=(FriendlyName = "Tan (Radians)", CompactNodeTitle = "TAN"), Category="Math|Trig")
+	UFUNCTION(BlueprintPure, meta=(DisplayName = "Tan (Radians)", CompactNodeTitle = "TAN"), Category="Math|Trig")
 	static float Tan(float A);
 
 	/* Returns the inverse tan (atan) (result is in Radians)*/
-	UFUNCTION(BlueprintPure, meta=(FriendlyName = "Atan (Radians)"), Category="Math|Trig")
+	UFUNCTION(BlueprintPure, meta=(DisplayName = "Atan (Radians)"), Category="Math|Trig")
 	static float Atan(float A);
 
 	/* Returns the inverse tan (atan2) of A/B (result is in Radians)*/
-	UFUNCTION(BlueprintPure, meta=(FriendlyName = "Atan2 (Radians)"), Category="Math|Trig")
+	UFUNCTION(BlueprintPure, meta=(DisplayName = "Atan2 (Radians)"), Category="Math|Trig")
 	static float Atan2(float A, float B);
 
 	/* Returns exponential(e) to the power A (e^A)*/
@@ -361,43 +361,43 @@ class ENGINE_API UKismetMathLibrary : public UBlueprintFunctionLibrary
 	static float RandomFloatInRange(float Min, float Max);
 
 	/* Returns the value of PI */
-	UFUNCTION(BlueprintPure, meta=(FriendlyName = "Get PI", CompactNodeTitle = "PI"), Category="Math|Trig")
+	UFUNCTION(BlueprintPure, meta=(DisplayName = "Get PI", CompactNodeTitle = "PI"), Category="Math|Trig")
 	static float GetPI();
 
 	/* Returns radians value based on the input degrees */
-	UFUNCTION(BlueprintPure, meta=(FriendlyName = "Degrees To Radians", CompactNodeTitle = "D2R"), Category="Math|Trig")
+	UFUNCTION(BlueprintPure, meta=(DisplayName = "Degrees To Radians", CompactNodeTitle = "D2R"), Category="Math|Trig")
 	static float DegreesToRadians(float A);
 
 	/* Returns degrees value based on the input radians */
-	UFUNCTION(BlueprintPure, meta=(FriendlyName = "Radians To Degrees", CompactNodeTitle = "R2D"), Category="Math|Trig")
+	UFUNCTION(BlueprintPure, meta=(DisplayName = "Radians To Degrees", CompactNodeTitle = "R2D"), Category="Math|Trig")
 	static float RadiansToDegrees(float A);
 
 	/* Returns the sin of A (expects Degrees)*/
-	UFUNCTION(BlueprintPure, meta=(FriendlyName = "Sin (Degrees)", CompactNodeTitle = "SINd"), Category="Math|Trig")
+	UFUNCTION(BlueprintPure, meta=(DisplayName = "Sin (Degrees)", CompactNodeTitle = "SINd"), Category="Math|Trig")
 	static float DegSin(float A);
 
 	/* Returns the inverse sin (arcsin) of A (result is in Degrees) */
-	UFUNCTION(BlueprintPure, meta=(FriendlyName = "Asin (Degrees)", CompactNodeTitle = "ASINd"), Category="Math|Trig")
+	UFUNCTION(BlueprintPure, meta=(DisplayName = "Asin (Degrees)", CompactNodeTitle = "ASINd"), Category="Math|Trig")
 	static float DegAsin(float A);
 
 	/* Returns the cos of A (expects Degrees)*/
-	UFUNCTION(BlueprintPure, meta=(FriendlyName = "Cos (Degrees)", CompactNodeTitle = "COSd"), Category="Math|Trig")
+	UFUNCTION(BlueprintPure, meta=(DisplayName = "Cos (Degrees)", CompactNodeTitle = "COSd"), Category="Math|Trig")
 	static float DegCos(float A);
 
 	/* Returns the inverse cos (arccos) of A (result is in Degrees) */
-	UFUNCTION(BlueprintPure, meta=(FriendlyName = "Acos (Degrees)", CompactNodeTitle = "ACOSd"), Category="Math|Trig")
+	UFUNCTION(BlueprintPure, meta=(DisplayName = "Acos (Degrees)", CompactNodeTitle = "ACOSd"), Category="Math|Trig")
 	static float DegAcos(float A);
 
 	/* Returns the tan of A (expects Degrees)*/
-	UFUNCTION(BlueprintPure, meta=(FriendlyName = "Tan (Degrees)", CompactNodeTitle = "TANd"), Category="Math|Trig")
+	UFUNCTION(BlueprintPure, meta=(DisplayName = "Tan (Degrees)", CompactNodeTitle = "TANd"), Category="Math|Trig")
 	static float DegTan(float A);
 
 	/* Returns the inverse tan (atan) (result is in Degrees)*/
-	UFUNCTION(BlueprintPure, meta=(FriendlyName = "Atan (Degrees)"), Category="Math|Trig")
+	UFUNCTION(BlueprintPure, meta=(DisplayName = "Atan (Degrees)"), Category="Math|Trig")
 	static float DegAtan(float A);
 
 	/* Returns the inverse tan (atan2) of A/B (result is in Degrees)*/
-	UFUNCTION(BlueprintPure, meta=(FriendlyName = "Atan2 (Degrees)"), Category="Math|Trig")
+	UFUNCTION(BlueprintPure, meta=(DisplayName = "Atan2 (Degrees)"), Category="Math|Trig")
 	static float DegAtan2(float A, float B);
 
 	/** 
@@ -407,19 +407,19 @@ class ENGINE_API UKismetMathLibrary : public UBlueprintFunctionLibrary
 	 * @param MaxAngleDegrees	"to" angle that defines the end of the range of valid angles
 	 * @return Returns clamped angle in the range -180..180.
 	 */
-	UFUNCTION(BlueprintPure, meta=(FriendlyName = "Clamp Angle"), Category="Math|Float")
+	UFUNCTION(BlueprintPure, meta=(DisplayName = "Clamp Angle"), Category="Math|Float")
 	static float ClampAngle(float AngleDegrees, float MinAngleDegrees, float MaxAngleDegrees);
 
 	/* Returns the minimum value of A and B */
-	UFUNCTION(BlueprintPure, meta=(FriendlyName = "Min (float)", CompactNodeTitle = "MIN", CommutativeAssociativeBinaryOperator = "true"), Category="Math|Float")
+	UFUNCTION(BlueprintPure, meta=(DisplayName = "Min (float)", CompactNodeTitle = "MIN", CommutativeAssociativeBinaryOperator = "true"), Category="Math|Float")
 	static float FMin(float A, float B);
 
 	/* Returns the maximum value of A and B */
-	UFUNCTION(BlueprintPure, meta=(FriendlyName = "Max (float)", CompactNodeTitle = "MAX", CommutativeAssociativeBinaryOperator = "true"), Category="Math|Float")
+	UFUNCTION(BlueprintPure, meta=(DisplayName = "Max (float)", CompactNodeTitle = "MAX", CommutativeAssociativeBinaryOperator = "true"), Category="Math|Float")
 	static float FMax(float A, float B);
 
 	/* Returns V clamped to be between A and B (inclusive) */
-	UFUNCTION(BlueprintPure, meta=(FriendlyName = "Clamp (float)", Min="0.0", Max="1.0"), Category="Math|Float")
+	UFUNCTION(BlueprintPure, meta=(DisplayName = "Clamp (float)", Min="0.0", Max="1.0"), Category="Math|Float")
 	static float FClamp(float Value, float Min, float Max);
 
 	/** Returns max of all array entries and the index at which it was found. Returns value of 0 and index of -1 if the supplied array is empty. */
@@ -464,7 +464,7 @@ class ENGINE_API UKismetMathLibrary : public UBlueprintFunctionLibrary
 	static float InverseLerp(float A, float B, float Value);
 
 	/* Easeing  between A and B using a specified easing function */
-	UFUNCTION(BlueprintPure, meta = (FriendlyName = "Ease", BlueprintInternalUseOnly = "true"), Category = "Math|Interpolation")
+	UFUNCTION(BlueprintPure, meta = (DisplayName = "Ease", BlueprintInternalUseOnly = "true"), Category = "Math|Interpolation")
 	static float Ease(float A, float B, float Alpha, TEnumAsByte<EEasingFunc::Type> EasingFunc, float BlendExp = 2, int32 Steps = 2);
 
 	/* Rounds A to the nearest integer */
@@ -472,11 +472,11 @@ class ENGINE_API UKismetMathLibrary : public UBlueprintFunctionLibrary
 	static int32 Round(float A);
 
 	/* Rounds A to the largest previous integer */
-	UFUNCTION(BlueprintPure, meta=(FriendlyName = "Floor"), Category="Math|Float")
+	UFUNCTION(BlueprintPure, meta=(DisplayName = "Floor"), Category="Math|Float")
 	static int32 FFloor(float A);
 	
 	/* Rounds A to an integer with truncation towards zero.  (e.g. -1.7 truncated to -1, 2.8 truncated to 2) */
-	UFUNCTION(BlueprintPure, meta=(FriendlyName = "Truncate"), Category="Math|Float")
+	UFUNCTION(BlueprintPure, meta=(DisplayName = "Truncate"), Category="Math|Float")
 	static int32 FTrunc(float A);
 	
 	/* Rounds A to the smallest following integer */
@@ -484,11 +484,11 @@ class ENGINE_API UKismetMathLibrary : public UBlueprintFunctionLibrary
 	static int32 FCeil(float A);
 
 	/* Returns the number of times Divisor will go into Dividend (i.e., Dividend divided by Divisor), as well as the remainder */
-	UFUNCTION(BlueprintPure, meta=(FriendlyName = "Division (whole and remainder)"), Category="Math|Float")
+	UFUNCTION(BlueprintPure, meta=(DisplayName = "Division (whole and remainder)"), Category="Math|Float")
 	static int32 FMod(float Dividend, float Divisor, float& Remainder);
 
 	/* Sign (float, returns -1 if A < 0, 0 if A is zero, and +1 if A > 0) */
-	UFUNCTION(BlueprintPure, meta=(FriendlyName = "Sign (float)"), Category="Math|Float")
+	UFUNCTION(BlueprintPure, meta=(DisplayName = "Sign (float)"), Category="Math|Float")
 	static float SignOfFloat(float A);
 
 	/** Returns Value normalized to the given range.  (e.g. 20 normalized to the range 10->50 would result in 0.25) */
@@ -525,101 +525,101 @@ class ENGINE_API UKismetMathLibrary : public UBlueprintFunctionLibrary
 	//
 
 	/* Scales Vector A by B */
-	UFUNCTION(BlueprintPure, meta=(FriendlyName = "vector * float", CompactNodeTitle = "*", Keywords = "* multiply"), Category="Math|Vector")
+	UFUNCTION(BlueprintPure, meta=(DisplayName = "vector * float", CompactNodeTitle = "*", Keywords = "* multiply"), Category="Math|Vector")
 	static FVector Multiply_VectorFloat(FVector A, float B);
 
 	/* Element-wise Vector multiplication (Result = {A.x*B.x, A.y*B.y, A.z*B.z}) */
-	UFUNCTION(BlueprintPure, meta=(FriendlyName = "vector * vector", CompactNodeTitle = "*", Keywords = "* multiply", CommutativeAssociativeBinaryOperator = "true"), Category="Math|Vector")
+	UFUNCTION(BlueprintPure, meta=(DisplayName = "vector * vector", CompactNodeTitle = "*", Keywords = "* multiply", CommutativeAssociativeBinaryOperator = "true"), Category="Math|Vector")
 	static FVector Multiply_VectorVector(FVector A, FVector B);
 
 	/* Vector divide by a float */
-	UFUNCTION(BlueprintPure, meta=(FriendlyName = "vector / float", CompactNodeTitle = "/", Keywords = "/ divide division"), Category="Math|Vector")
+	UFUNCTION(BlueprintPure, meta=(DisplayName = "vector / float", CompactNodeTitle = "/", Keywords = "/ divide division"), Category="Math|Vector")
 	static FVector Divide_VectorFloat(FVector A, float B = 1.f);
 
 	/* Vector divide by vector */
-	UFUNCTION(BlueprintPure, meta=(FriendlyName = "vector / vector", CompactNodeTitle = "/", Keywords = "/ divide division"), Category="Math|Vector")
+	UFUNCTION(BlueprintPure, meta=(DisplayName = "vector / vector", CompactNodeTitle = "/", Keywords = "/ divide division"), Category="Math|Vector")
 	static FVector Divide_VectorVector(FVector A, FVector B = FVector(1.f,1.f,1.f));
 
 	/* Vector addition */
-	UFUNCTION(BlueprintPure, meta=(FriendlyName = "vector + vector", CompactNodeTitle = "+", Keywords = "+ add plus", CommutativeAssociativeBinaryOperator = "true"), Category="Math|Vector")
+	UFUNCTION(BlueprintPure, meta=(DisplayName = "vector + vector", CompactNodeTitle = "+", Keywords = "+ add plus", CommutativeAssociativeBinaryOperator = "true"), Category="Math|Vector")
 	static FVector Add_VectorVector(FVector A, FVector B);
 
-	UFUNCTION(BlueprintPure, meta=(FriendlyName = "vector + float", CompactNodeTitle = "+", Keywords = "+ add plus"), Category="Math|Vector")
+	UFUNCTION(BlueprintPure, meta=(DisplayName = "vector + float", CompactNodeTitle = "+", Keywords = "+ add plus"), Category="Math|Vector")
 	static FVector Add_VectorFloat(FVector A, float B);
 
 	/* Vector subtraction */
-	UFUNCTION(BlueprintPure, meta=(FriendlyName = "vector - vector", CompactNodeTitle = "-", Keywords = "- subtract minus"), Category="Math|Vector")
+	UFUNCTION(BlueprintPure, meta=(DisplayName = "vector - vector", CompactNodeTitle = "-", Keywords = "- subtract minus"), Category="Math|Vector")
 	static FVector Subtract_VectorVector(FVector A, FVector B);
 
-	UFUNCTION(BlueprintPure, meta=(FriendlyName = "vector - float", CompactNodeTitle = "-", Keywords = "- subtract minus"), Category="Math|Vector")
+	UFUNCTION(BlueprintPure, meta=(DisplayName = "vector - float", CompactNodeTitle = "-", Keywords = "- subtract minus"), Category="Math|Vector")
 	static FVector Subtract_VectorFloat(FVector A, float B);
 
 	/* Returns result of vector A rotated by the inverse of Rotator B */
-	UFUNCTION(BlueprintPure, meta=(FriendlyName = "UnrotateVector"), Category="Math|Vector")
+	UFUNCTION(BlueprintPure, meta=(DisplayName = "UnrotateVector"), Category="Math|Vector")
 	static FVector LessLess_VectorRotator(FVector A, FRotator B);
 
 	/* Returns result of vector A rotated by Rotator B */
-	UFUNCTION(BlueprintPure, meta=(FriendlyName = "RotateVector"), Category="Math|Vector")
+	UFUNCTION(BlueprintPure, meta=(DisplayName = "RotateVector"), Category="Math|Vector")
 	static FVector GreaterGreater_VectorRotator(FVector A, FRotator B);
 
 	/* Returns result of vector A rotated by AngleDeg around Axis */
-	UFUNCTION(BlueprintPure, meta=(FriendlyName = "RotateVectorAroundAxis"), Category="Math|Vector")
+	UFUNCTION(BlueprintPure, meta=(DisplayName = "RotateVectorAroundAxis"), Category="Math|Vector")
 	static FVector RotateAngleAxis(FVector InVect, float AngleDeg, FVector Axis);
 
 	/* Returns true if vector A is equal to vector B (A == B) within a specified error tolerance */
-	UFUNCTION(BlueprintPure, meta=(FriendlyName = "Equal (vector)", CompactNodeTitle = "==", Keywords = "== equal"), Category="Math|Vector")
+	UFUNCTION(BlueprintPure, meta=(DisplayName = "Equal (vector)", CompactNodeTitle = "==", Keywords = "== equal"), Category="Math|Vector")
 	static bool EqualEqual_VectorVector(FVector A, FVector B, float ErrorTolerance = 1.e-4f);
 
 	/* Returns true if vector A is not equal to vector B (A != B) within a specified error tolerance */
-	UFUNCTION(BlueprintPure, meta=(FriendlyName = "Not Equal (vector)", CompactNodeTitle = "!=", Keywords = "!= not equal"), Category="Math|Vector")
+	UFUNCTION(BlueprintPure, meta=(DisplayName = "Not Equal (vector)", CompactNodeTitle = "!=", Keywords = "!= not equal"), Category="Math|Vector")
 	static bool NotEqual_VectorVector(FVector A, FVector B, float ErrorTolerance = 1.e-4f);
 
 	/* Returns the dot product of two 3d vectors */
-	UFUNCTION(BlueprintPure, meta=(FriendlyName = "Dot Product", CompactNodeTitle = "."), Category="Math|Vector" )
+	UFUNCTION(BlueprintPure, meta=(DisplayName = "Dot Product", CompactNodeTitle = "."), Category="Math|Vector" )
 	static float Dot_VectorVector(FVector A, FVector B);
 
 	/* Returns the cross product of two 3d vectors */
-	UFUNCTION(BlueprintPure, meta=(FriendlyName = "Cross Product", CompactNodeTitle = "x"), Category="Math|Vector" )
+	UFUNCTION(BlueprintPure, meta=(DisplayName = "Cross Product", CompactNodeTitle = "x"), Category="Math|Vector" )
 	static FVector Cross_VectorVector(FVector A, FVector B);
 
 	/* Returns the dot product of two 2d vectors */
-	UFUNCTION(BlueprintPure, meta=(FriendlyName="Dot Product (2D)", CompactNodeTitle="."), Category="Math|Vector")
+	UFUNCTION(BlueprintPure, meta=(DisplayName="Dot Product (2D)", CompactNodeTitle="."), Category="Math|Vector")
 	static float DotProduct2D(FVector2D A, FVector2D B);
 
 	/* Returns the cross product of two 2d vectors */
-	UFUNCTION(BlueprintPure, meta=(FriendlyName="Cross Product (2D)", CompactNodeTitle="x"), Category="Math|Vector")
+	UFUNCTION(BlueprintPure, meta=(DisplayName="Cross Product (2D)", CompactNodeTitle="x"), Category="Math|Vector")
 	static float CrossProduct2D(FVector2D A, FVector2D B);
 
 	/* Returns the length of the FVector */
-	UFUNCTION(BlueprintPure, meta=(FriendlyName = "VectorLength"), Category="Math|Vector")
+	UFUNCTION(BlueprintPure, meta=(DisplayName = "VectorLength"), Category="Math|Vector")
 	static float VSize(FVector A);
 
 	/** Returns the length of a 2d FVector. */
-	UFUNCTION(BlueprintPure, meta=(FriendlyName = "Vector2dLength"), Category="Math|Vector2D")
+	UFUNCTION(BlueprintPure, meta=(DisplayName = "Vector2dLength"), Category="Math|Vector2D")
 	static float VSize2D(FVector2D A);
 
 	/* Returns the squared length of the FVector */
-	UFUNCTION(BlueprintPure, meta=(FriendlyName = "VectorLengthSquared"), Category="Math|Vector")
+	UFUNCTION(BlueprintPure, meta=(DisplayName = "VectorLengthSquared"), Category="Math|Vector")
 	static float VSizeSquared(FVector A);
 
 	/** Returns the squared length of a 2d FVector. */
-	UFUNCTION(BlueprintPure, meta=(FriendlyName = "Vector2dLengthSquared"), Category="Math|Vector2D")
+	UFUNCTION(BlueprintPure, meta=(DisplayName = "Vector2dLengthSquared"), Category="Math|Vector2D")
 	static float VSize2DSquared(FVector2D A);
 
 	/* Returns a unit normal version of the FVector A */
-	UFUNCTION(BlueprintPure, meta=(FriendlyName = "Normalize"), Category="Math|Vector")
+	UFUNCTION(BlueprintPure, meta=(DisplayName = "Normalize"), Category="Math|Vector")
 	static FVector Normal(FVector A);
 
 	/* Returns a unit normal version of the vector2d A */
-	UFUNCTION(BlueprintPure, meta=(FriendlyName = "Normalize2D"), Category="Math|Vector2D")
+	UFUNCTION(BlueprintPure, meta=(DisplayName = "Normalize2D"), Category="Math|Vector2D")
 	static FVector2D Normal2D(FVector2D A);
 
 	/* Linearly interpolates between A and B based on Alpha (100% of A when Alpha=0 and 100% of B when Alpha=1) */
-	UFUNCTION(BlueprintPure, meta=(FriendlyName = "Lerp (vector)"), Category="Math|Vector")
+	UFUNCTION(BlueprintPure, meta=(DisplayName = "Lerp (vector)"), Category="Math|Vector")
 	static FVector VLerp(FVector A, FVector B, float Alpha);
 
 	/* Easeing  between A and B using a specified easing function */
-	UFUNCTION(BlueprintPure, meta = (FriendlyName = "Ease (vector)", BlueprintInternalUseOnly = "true"), Category = "Math|Interpolation")
+	UFUNCTION(BlueprintPure, meta = (DisplayName = "Ease (vector)", BlueprintInternalUseOnly = "true"), Category = "Math|Interpolation")
 	static FVector VEase(FVector A, FVector B, float Alpha, TEnumAsByte<EEasingFunc::Type> EasingFunc, float BlendExp = 2, int32 Steps = 2);
 
 	/* Returns a random vector with length of 1 */
@@ -703,19 +703,19 @@ class ENGINE_API UKismetMathLibrary : public UBlueprintFunctionLibrary
 	//
 
 	//Returns true if rotator A is equal to rotator B (A == B) within a specified error tolerance
-	UFUNCTION(BlueprintPure, meta=(FriendlyName = "Equal (Rotator)", CompactNodeTitle = "==", Keywords = "== equal"), Category="Math|Rotator")
+	UFUNCTION(BlueprintPure, meta=(DisplayName = "Equal (Rotator)", CompactNodeTitle = "==", Keywords = "== equal"), Category="Math|Rotator")
 	static bool EqualEqual_RotatorRotator(FRotator A, FRotator B, float ErrorTolerance = 1.e-4f);
 
 	//Returns true if rotator A is not equal to rotator B (A != B) within a specified error tolerance
-	UFUNCTION(BlueprintPure, meta=(FriendlyName = "Not Equal (Rotator)", CompactNodeTitle = "!=", Keywords = "!= not equal"), Category="Math|Rotator")
+	UFUNCTION(BlueprintPure, meta=(DisplayName = "Not Equal (Rotator)", CompactNodeTitle = "!=", Keywords = "!= not equal"), Category="Math|Rotator")
 	static bool NotEqual_RotatorRotator(FRotator A, FRotator B, float ErrorTolerance = 1.e-4f);
 
 	//Returns rotator representing rotator A scaled by B 
-	UFUNCTION(BlueprintPure, meta=(FriendlyName = "ScaleRotator", CompactNodeTitle = "*", Keywords = "* multiply rotate rotation"), Category="Math|Rotator")
+	UFUNCTION(BlueprintPure, meta=(DisplayName = "ScaleRotator", CompactNodeTitle = "*", Keywords = "* multiply rotate rotation"), Category="Math|Rotator")
 	static FRotator Multiply_RotatorFloat(FRotator A, float B);
 
 	/** Combine 2 rotations to give you the resulting rotation */
-	UFUNCTION(BlueprintPure, meta=(FriendlyName = "CombineRotators", Keywords="rotate rotation add"), Category="Math|Rotator")
+	UFUNCTION(BlueprintPure, meta=(DisplayName = "CombineRotators", Keywords="rotate rotation add"), Category="Math|Rotator")
 	static FRotator ComposeRotators(FRotator A, FRotator B);
 
 	/** Negate a rotator*/
@@ -731,15 +731,15 @@ class ENGINE_API UKismetMathLibrary : public UBlueprintFunctionLibrary
 	static FRotator RandomRotator(bool bRoll = false);
 
 	/* Linearly interpolates between A and B based on Alpha (100% of A when Alpha=0 and 100% of B when Alpha=1) */
-	UFUNCTION(BlueprintPure, meta=(FriendlyName = "Lerp (Rotator)"), Category="Math|Rotator")
+	UFUNCTION(BlueprintPure, meta=(DisplayName = "Lerp (Rotator)"), Category="Math|Rotator")
 	static FRotator RLerp(FRotator A, FRotator B, float Alpha, bool bShortestPath);
 
 	/* Easeing  between A and B using a specified easing function */
-	UFUNCTION(BlueprintPure, meta = (FriendlyName = "Ease (Rotator)", BlueprintInternalUseOnly = "true"), Category = "Math|Interpolation")
+	UFUNCTION(BlueprintPure, meta = (DisplayName = "Ease (Rotator)", BlueprintInternalUseOnly = "true"), Category = "Math|Interpolation")
 	static FRotator REase(FRotator A, FRotator B, float Alpha, bool bShortestPath, TEnumAsByte<EEasingFunc::Type> EasingFunc, float BlendExp = 2, int32 Steps = 2);
 
 	/** Normalized A-B */
-	UFUNCTION(BlueprintPure, meta=(FriendlyName = "Delta (Rotator)"), Category="Math|Rotator")
+	UFUNCTION(BlueprintPure, meta=(DisplayName = "Delta (Rotator)"), Category="Math|Rotator")
 	static FRotator NormalizedDeltaRotator(FRotator A, FRotator B);
 
 	/** Create a rotation from an axis and and angle (in degrees) */
@@ -751,15 +751,15 @@ class ENGINE_API UKismetMathLibrary : public UBlueprintFunctionLibrary
 	//
 
 	/* Linearly interpolates between A and B based on Alpha (100% of A when Alpha=0 and 100% of B when Alpha=1) */
-	UFUNCTION(BlueprintPure, meta=(FriendlyName = "Lerp (LinearColor)"), Category="Math|Color")
+	UFUNCTION(BlueprintPure, meta=(DisplayName = "Lerp (LinearColor)"), Category="Math|Color")
 	static FLinearColor LinearColorLerp(FLinearColor A, FLinearColor B, float Alpha);
 
 	/* Element-wise multiplication of two linear colors (R*R, G*G, B*B, A*A) */
-	UFUNCTION(BlueprintPure, meta=(FriendlyName = "LinearColor * (LinearColor)", CompactNodeTitle = "*"), Category="Math|Color")
+	UFUNCTION(BlueprintPure, meta=(DisplayName = "LinearColor * (LinearColor)", CompactNodeTitle = "*"), Category="Math|Color")
 	static FLinearColor Multiply_LinearColorLinearColor(FLinearColor A, FLinearColor B);
 
 	/* Element-wise multiplication of a linear color by a float (F*R, F*G, F*B, F*A) */
-	UFUNCTION(BlueprintPure, meta=(FriendlyName = "LinearColor * Float", CompactNodeTitle = "*", Keywords = "* multiply"), Category="Math|Color")
+	UFUNCTION(BlueprintPure, meta=(DisplayName = "LinearColor * Float", CompactNodeTitle = "*", Keywords = "* multiply"), Category="Math|Color")
 	static FLinearColor Multiply_LinearColorFloat(FLinearColor A, float B);
 
 	//
@@ -767,119 +767,119 @@ class ENGINE_API UKismetMathLibrary : public UBlueprintFunctionLibrary
 	//
 
 	/* Addition (A + B) */
-	UFUNCTION(BlueprintPure, meta=(FriendlyName="DateTime + Timespan", CompactNodeTitle="+", Keywords="+ add plus"), Category="Math|DateTime")
+	UFUNCTION(BlueprintPure, meta=(DisplayName="DateTime + Timespan", CompactNodeTitle="+", Keywords="+ add plus"), Category="Math|DateTime")
 	static FDateTime Add_DateTimeTimespan( FDateTime A, FTimespan B );
 
 	/* Subtraction (A - B) */
-	UFUNCTION(BlueprintPure, meta=(FriendlyName="DateTime - Timespan", CompactNodeTitle="-", Keywords="- subtract minus"), Category="Math|DateTime")
+	UFUNCTION(BlueprintPure, meta=(DisplayName="DateTime - Timespan", CompactNodeTitle="-", Keywords="- subtract minus"), Category="Math|DateTime")
 	static FDateTime Subtract_DateTimeTimespan( FDateTime A, FTimespan B );
 
 	/* Returns true if the values are equal (A == B) */
-	UFUNCTION(BlueprintPure, meta=(FriendlyName="Equal (DateTime)", CompactNodeTitle="==", Keywords="== equal"), Category="Math|DateTime")
+	UFUNCTION(BlueprintPure, meta=(DisplayName="Equal (DateTime)", CompactNodeTitle="==", Keywords="== equal"), Category="Math|DateTime")
 	static bool EqualEqual_DateTimeDateTime( FDateTime A, FDateTime B );
 
 	/* Returns true if the values are not equal (A != B) */
-	UFUNCTION(BlueprintPure, meta=(FriendlyName="NotEqual (DateTime)", CompactNodeTitle="!=", Keywords="!= not equal"), Category="Math|DateTime")
+	UFUNCTION(BlueprintPure, meta=(DisplayName="NotEqual (DateTime)", CompactNodeTitle="!=", Keywords="!= not equal"), Category="Math|DateTime")
 	static bool NotEqual_DateTimeDateTime( FDateTime A, FDateTime B );
 
 	/* Returns true if A is Greater than B (A > B) */
-	UFUNCTION(BlueprintPure, meta=(FriendlyName="DateTime > DateTime", CompactNodeTitle=">", Keywords="> greater"), Category="Math|DateTime")
+	UFUNCTION(BlueprintPure, meta=(DisplayName="DateTime > DateTime", CompactNodeTitle=">", Keywords="> greater"), Category="Math|DateTime")
 	static bool Greater_DateTimeDateTime( FDateTime A, FDateTime B );
 
 	/* Returns true if A is Greater than B (A >= B) */
-	UFUNCTION(BlueprintPure, meta=(FriendlyName="DateTime >= DateTime", CompactNodeTitle=">=", Keywords=">= greater"), Category="Math|DateTime")
+	UFUNCTION(BlueprintPure, meta=(DisplayName="DateTime >= DateTime", CompactNodeTitle=">=", Keywords=">= greater"), Category="Math|DateTime")
 	static bool GreaterEqual_DateTimeDateTime( FDateTime A, FDateTime B );
 
 	/* Returns true if A is Greater than B (A < B) */
-	UFUNCTION(BlueprintPure, meta=(FriendlyName="DateTime < DateTime", CompactNodeTitle="<", Keywords="< less"), Category="Math|DateTime")
+	UFUNCTION(BlueprintPure, meta=(DisplayName="DateTime < DateTime", CompactNodeTitle="<", Keywords="< less"), Category="Math|DateTime")
 	static bool Less_DateTimeDateTime( FDateTime A, FDateTime B );
 
 	/* Returns true if A is Greater than B (A <= B) */
-	UFUNCTION(BlueprintPure, meta=(FriendlyName="DateTime <= DateTime", CompactNodeTitle="<=", Keywords="<= less"), Category="Math|DateTime")
+	UFUNCTION(BlueprintPure, meta=(DisplayName="DateTime <= DateTime", CompactNodeTitle="<=", Keywords="<= less"), Category="Math|DateTime")
 	static bool LessEqual_DateTimeDateTime( FDateTime A, FDateTime B );
 
 	/* Returns the date component of A */
-	UFUNCTION(BlueprintPure, meta=(FriendlyName="GetDate"), Category="Math|DateTime")
+	UFUNCTION(BlueprintPure, meta=(DisplayName="GetDate"), Category="Math|DateTime")
 	static FDateTime GetDate( FDateTime A );
 
 	/* Returns the day component of A (1 to 31) */
-	UFUNCTION(BlueprintPure, meta=(FriendlyName="GetDay"), Category="Math|DateTime")
+	UFUNCTION(BlueprintPure, meta=(DisplayName="GetDay"), Category="Math|DateTime")
 	static int32 GetDay( FDateTime A );
 
 	/* Returns the day of year of A */
-	UFUNCTION(BlueprintPure, meta=(FriendlyName="GetDayOfYear"), Category="Math|DateTime")
+	UFUNCTION(BlueprintPure, meta=(DisplayName="GetDayOfYear"), Category="Math|DateTime")
 	static int32 GetDayOfYear( FDateTime A );
 
 	/* Returns the hour component of A (24h format) */
-	UFUNCTION(BlueprintPure, meta=(FriendlyName="GetHour"), Category="Math|DateTime")
+	UFUNCTION(BlueprintPure, meta=(DisplayName="GetHour"), Category="Math|DateTime")
 	static int32 GetHour( FDateTime A );
 
 	/* Returns the hour component of A (12h format) */
-	UFUNCTION(BlueprintPure, meta=(FriendlyName="GetHour12"), Category="Math|DateTime")
+	UFUNCTION(BlueprintPure, meta=(DisplayName="GetHour12"), Category="Math|DateTime")
 	static int32 GetHour12( FDateTime A );
 
 	/* Returns the millisecond component of A */
-	UFUNCTION(BlueprintPure, meta=(FriendlyName="GetMillisecond"), Category="Math|DateTime")
+	UFUNCTION(BlueprintPure, meta=(DisplayName="GetMillisecond"), Category="Math|DateTime")
 	static int32 GetMillisecond( FDateTime A );
 
 	/* Returns the minute component of A */
-	UFUNCTION(BlueprintPure, meta=(FriendlyName="GetMinute"), Category="Math|DateTime")
+	UFUNCTION(BlueprintPure, meta=(DisplayName="GetMinute"), Category="Math|DateTime")
 	static int32 GetMinute( FDateTime A );
 
 	/* Returns the month component of A */
-	UFUNCTION(BlueprintPure, meta=(FriendlyName="GetMonth"), Category="Math|DateTime")
+	UFUNCTION(BlueprintPure, meta=(DisplayName="GetMonth"), Category="Math|DateTime")
 	static int32 GetMonth( FDateTime A );
 
 	/* Returns the second component of A */
-	UFUNCTION(BlueprintPure, meta=(FriendlyName="GetSecond"), Category="Math|DateTime")
+	UFUNCTION(BlueprintPure, meta=(DisplayName="GetSecond"), Category="Math|DateTime")
 	static int32 GetSecond( FDateTime A );
 
 	/* Returns the time elapsed since midnight of A */
-	UFUNCTION(BlueprintPure, meta=(FriendlyName="GetTimeOfDay"), Category="Math|DateTime")
+	UFUNCTION(BlueprintPure, meta=(DisplayName="GetTimeOfDay"), Category="Math|DateTime")
 	static FTimespan GetTimeOfDay( FDateTime A );
 
 	/* Returns the year component of A */
-	UFUNCTION(BlueprintPure, meta=(FriendlyName="GetYear"), Category="Math|DateTime")
+	UFUNCTION(BlueprintPure, meta=(DisplayName="GetYear"), Category="Math|DateTime")
 	static int32 GetYear( FDateTime A );
 
 	/* Returns whether A's time is in the afternoon */
-	UFUNCTION(BlueprintPure, meta=(FriendlyName="IsAfternoon"), Category="Math|DateTime")
+	UFUNCTION(BlueprintPure, meta=(DisplayName="IsAfternoon"), Category="Math|DateTime")
 	static bool IsAfternoon( FDateTime A );
 
 	/* Returns whether A's time is in the morning */
-	UFUNCTION(BlueprintPure, meta=(FriendlyName="IsMorning"), Category="Math|DateTime")
+	UFUNCTION(BlueprintPure, meta=(DisplayName="IsMorning"), Category="Math|DateTime")
 	static bool IsMorning( FDateTime A );
 
 	/* Returns the number of days in the given year and month */
-	UFUNCTION(BlueprintPure, meta=(FriendlyName="DaysInMonth"), Category="Math|DateTime")
+	UFUNCTION(BlueprintPure, meta=(DisplayName="DaysInMonth"), Category="Math|DateTime")
 	static int32 DaysInMonth( int32 Year, int32 Month );
 
 	/* Returns the number of days in the given year */
-	UFUNCTION(BlueprintPure, meta=(FriendlyName="DaysInYear"), Category="Math|DateTime")
+	UFUNCTION(BlueprintPure, meta=(DisplayName="DaysInYear"), Category="Math|DateTime")
 	static int32 DaysInYear( int32 Year );
 
 	/* Returns whether given year is a leap year */
-	UFUNCTION(BlueprintPure, meta=(FriendlyName="IsLeapYear"), Category="Math|DateTime")
+	UFUNCTION(BlueprintPure, meta=(DisplayName="IsLeapYear"), Category="Math|DateTime")
 	static bool IsLeapYear( int32 Year );
 
 	/* Returns the maximum date and time value */
-	UFUNCTION(BlueprintPure, meta=(FriendlyName="MaxValue"), Category="Math|DateTime")
+	UFUNCTION(BlueprintPure, meta=(DisplayName="MaxValue"), Category="Math|DateTime")
 	static FDateTime DateTimeMaxValue( );
 
 	/* Returns the minimum date and time value */
-	UFUNCTION(BlueprintPure, meta=(FriendlyName="MinValue"), Category="Math|DateTime")
+	UFUNCTION(BlueprintPure, meta=(DisplayName="MinValue"), Category="Math|DateTime")
 	static FDateTime DateTimeMinValue( );
 
 	/* Returns the local date and time on this computer */
-	UFUNCTION(BlueprintPure, meta=(FriendlyName="Now"), Category="Math|DateTime")
+	UFUNCTION(BlueprintPure, meta=(DisplayName="Now"), Category="Math|DateTime")
 	static FDateTime Now( );
 
 	/* Returns the local date on this computer */
-	UFUNCTION(BlueprintPure, meta=(FriendlyName="Today"), Category="Math|DateTime")
+	UFUNCTION(BlueprintPure, meta=(DisplayName="Today"), Category="Math|DateTime")
 	static FDateTime Today( );
 
 	/* Returns the UTC date and time on this computer */
-	UFUNCTION(BlueprintPure, meta=(FriendlyName="UtcNow"), Category="Math|DateTime")
+	UFUNCTION(BlueprintPure, meta=(DisplayName="UtcNow"), Category="Math|DateTime")
 	static FDateTime UtcNow( );
 
 	//
@@ -887,189 +887,189 @@ class ENGINE_API UKismetMathLibrary : public UBlueprintFunctionLibrary
 	//
 
 	/* Addition (A + B) */
-	UFUNCTION(BlueprintPure, meta=(FriendlyName="Timespan + Timespan", CompactNodeTitle="+", Keywords="+ add plus"), Category="Math|Timespan")
+	UFUNCTION(BlueprintPure, meta=(DisplayName="Timespan + Timespan", CompactNodeTitle="+", Keywords="+ add plus"), Category="Math|Timespan")
 	static FTimespan Add_TimespanTimespan( FTimespan A, FTimespan B );
 
 	/* Subtraction (A - B) */
-	UFUNCTION(BlueprintPure, meta=(FriendlyName="Timespan - Timespan", CompactNodeTitle="-", Keywords="- subtract minus"), Category="Math|Timespan")
+	UFUNCTION(BlueprintPure, meta=(DisplayName="Timespan - Timespan", CompactNodeTitle="-", Keywords="- subtract minus"), Category="Math|Timespan")
 	static FTimespan Subtract_TimespanTimespan( FTimespan A, FTimespan B );
 
 	/* Scalar multiplication (A * s) */
-	UFUNCTION(BlueprintPure, meta=(FriendlyName="Timespan * float", CompactNodeTitle="*", Keywords="* multiply"), Category="Math|Timespan")
+	UFUNCTION(BlueprintPure, meta=(DisplayName="Timespan * float", CompactNodeTitle="*", Keywords="* multiply"), Category="Math|Timespan")
 	static FTimespan Multiply_TimespanFloat( FTimespan A, float Scalar );
 
 	/* Returns true if the values are equal (A == B) */
-	UFUNCTION(BlueprintPure, meta=(FriendlyName="Equal (Timespan)", CompactNodeTitle="==", Keywords="== equal"), Category="Math|Timespan")
+	UFUNCTION(BlueprintPure, meta=(DisplayName="Equal (Timespan)", CompactNodeTitle="==", Keywords="== equal"), Category="Math|Timespan")
 	static bool EqualEqual_TimespanTimespan( FTimespan A, FTimespan B );
 
 	/* Returns true if the values are not equal (A != B) */
-	UFUNCTION(BlueprintPure, meta=(FriendlyName="NotEqual (Timespan)", CompactNodeTitle="!=", Keywords="!= not equal"), Category="Math|Timespan")
+	UFUNCTION(BlueprintPure, meta=(DisplayName="NotEqual (Timespan)", CompactNodeTitle="!=", Keywords="!= not equal"), Category="Math|Timespan")
 	static bool NotEqual_TimespanTimespan( FTimespan A, FTimespan B );
 
 	/* Returns true if A is Greater than B (A > B) */
-	UFUNCTION(BlueprintPure, meta=(FriendlyName="Timespan > Timespan", CompactNodeTitle=">", Keywords="> greater"), Category="Math|Timespan")
+	UFUNCTION(BlueprintPure, meta=(DisplayName="Timespan > Timespan", CompactNodeTitle=">", Keywords="> greater"), Category="Math|Timespan")
 	static bool Greater_TimespanTimespan( FTimespan A, FTimespan B );
 
 	/* Returns true if A is Greater than B (A >= B) */
-	UFUNCTION(BlueprintPure, meta=(FriendlyName="Timespan >= Timespan", CompactNodeTitle=">=", Keywords=">= greater"), Category="Math|Timespan")
+	UFUNCTION(BlueprintPure, meta=(DisplayName="Timespan >= Timespan", CompactNodeTitle=">=", Keywords=">= greater"), Category="Math|Timespan")
 	static bool GreaterEqual_TimespanTimespan( FTimespan A, FTimespan B );
 
 	/* Returns true if A is Greater than B (A < B) */
-	UFUNCTION(BlueprintPure, meta=(FriendlyName="Timespan < Timespan", CompactNodeTitle="<", Keywords="< less"), Category="Math|Timespan")
+	UFUNCTION(BlueprintPure, meta=(DisplayName="Timespan < Timespan", CompactNodeTitle="<", Keywords="< less"), Category="Math|Timespan")
 	static bool Less_TimespanTimespan( FTimespan A, FTimespan B );
 
 	/* Returns true if A is Greater than B (A <= B) */
-	UFUNCTION(BlueprintPure, meta=(FriendlyName="Timespan <= Timespan", CompactNodeTitle="<=", Keywords="<= less"), Category="Math|Timespan")
+	UFUNCTION(BlueprintPure, meta=(DisplayName="Timespan <= Timespan", CompactNodeTitle="<=", Keywords="<= less"), Category="Math|Timespan")
 	static bool LessEqual_TimespanTimespan( FTimespan A, FTimespan B );
 
 	/* Returns the days component of A */
-	UFUNCTION(BlueprintPure, meta=(FriendlyName="GetDays"), Category="Math|Timespan")
+	UFUNCTION(BlueprintPure, meta=(DisplayName="GetDays"), Category="Math|Timespan")
 	static int32 GetDays( FTimespan A );
 
 	/* Returns the absolute value of A */
-	UFUNCTION(BlueprintPure, meta=(FriendlyName="GetDuration"), Category="Math|Timespan")
+	UFUNCTION(BlueprintPure, meta=(DisplayName="GetDuration"), Category="Math|Timespan")
 	static FTimespan GetDuration( FTimespan A );
 
 	/* Returns the hours component of A */
-	UFUNCTION(BlueprintPure, meta=(FriendlyName="GetHours"), Category="Math|Timespan")
+	UFUNCTION(BlueprintPure, meta=(DisplayName="GetHours"), Category="Math|Timespan")
 	static int32 GetHours( FTimespan A );
 
 	/* Returns the milliseconds component of A */
-	UFUNCTION(BlueprintPure, meta=(FriendlyName="GetMilliseconds"), Category="Math|Timespan")
+	UFUNCTION(BlueprintPure, meta=(DisplayName="GetMilliseconds"), Category="Math|Timespan")
 	static int32 GetMilliseconds( FTimespan A );
 
 	/* Returns the minutes component of A */
-	UFUNCTION(BlueprintPure, meta=(FriendlyName="GetMinutes"), Category="Math|Timespan")
+	UFUNCTION(BlueprintPure, meta=(DisplayName="GetMinutes"), Category="Math|Timespan")
 	static int32 GetMinutes( FTimespan A );
 
 	/* Returns the seconds component of A */
-	UFUNCTION(BlueprintPure, meta=(FriendlyName="GetSeconds"), Category="Math|Timespan")
+	UFUNCTION(BlueprintPure, meta=(DisplayName="GetSeconds"), Category="Math|Timespan")
 	static int32 GetSeconds( FTimespan A );
 
 	/* Returns the total number of days in A */
-	UFUNCTION(BlueprintPure, meta=(FriendlyName="GetTotalDays"), Category="Math|Timespan")
+	UFUNCTION(BlueprintPure, meta=(DisplayName="GetTotalDays"), Category="Math|Timespan")
 	static float GetTotalDays( FTimespan A );
 
 	/* Returns the total number of hours in A */
-	UFUNCTION(BlueprintPure, meta=(FriendlyName="GetTotalHours"), Category="Math|Timespan")
+	UFUNCTION(BlueprintPure, meta=(DisplayName="GetTotalHours"), Category="Math|Timespan")
 	static float GetTotalHours( FTimespan A );
 
 	/* Returns the total number of milliseconds in A */
-	UFUNCTION(BlueprintPure, meta=(FriendlyName="GetTotalMilliseconds"), Category="Math|Timespan")
+	UFUNCTION(BlueprintPure, meta=(DisplayName="GetTotalMilliseconds"), Category="Math|Timespan")
 	static float GetTotalMilliseconds( FTimespan A );
 
 	/* Returns the total number of minutes in A */
-	UFUNCTION(BlueprintPure, meta=(FriendlyName="GetTotalMinutes"), Category="Math|Timespan")
+	UFUNCTION(BlueprintPure, meta=(DisplayName="GetTotalMinutes"), Category="Math|Timespan")
 	static float GetTotalMinutes( FTimespan A );
 
 	/* Returns the total number of seconds in A */
-	UFUNCTION(BlueprintPure, meta=(FriendlyName="GetTotalSeconds"), Category="Math|Timespan")
+	UFUNCTION(BlueprintPure, meta=(DisplayName="GetTotalSeconds"), Category="Math|Timespan")
 	static float GetTotalSeconds( FTimespan A );
 
 	/* Returns a time span that represents the specified number of days */
-	UFUNCTION(BlueprintPure, meta=(FriendlyName="FromDays"), Category="Math|Timespan")
+	UFUNCTION(BlueprintPure, meta=(DisplayName="FromDays"), Category="Math|Timespan")
 	static FTimespan FromDays( float Days );
 
 	/* Returns a time span that represents the specified number of hours */
-	UFUNCTION(BlueprintPure, meta=(FriendlyName="FromHours"), Category="Math|Timespan")
+	UFUNCTION(BlueprintPure, meta=(DisplayName="FromHours"), Category="Math|Timespan")
 	static FTimespan FromHours( float Hours );
 
 	/* Returns a time span that represents the specified number of milliseconds */
-	UFUNCTION(BlueprintPure, meta=(FriendlyName="FromMilliseconds"), Category="Math|Timespan")
+	UFUNCTION(BlueprintPure, meta=(DisplayName="FromMilliseconds"), Category="Math|Timespan")
 	static FTimespan FromMilliseconds( float Milliseconds );
 
 	/* Returns a time span that represents the specified number of minutes */
-	UFUNCTION(BlueprintPure, meta=(FriendlyName="FromMinutes"), Category="Math|Timespan")
+	UFUNCTION(BlueprintPure, meta=(DisplayName="FromMinutes"), Category="Math|Timespan")
 	static FTimespan FromMinutes( float Minutes );
 
 	/* Returns a time span that represents the specified number of seconds */
-	UFUNCTION(BlueprintPure, meta=(FriendlyName="FromSeconds"), Category="Math|Timespan")
+	UFUNCTION(BlueprintPure, meta=(DisplayName="FromSeconds"), Category="Math|Timespan")
 	static FTimespan FromSeconds( float Seconds );
 
 	/* Returns the maximum time span value */
-	UFUNCTION(BlueprintPure, meta=(FriendlyName="MaxValue"), Category="Math|Timespan")
+	UFUNCTION(BlueprintPure, meta=(DisplayName="MaxValue"), Category="Math|Timespan")
 	static FTimespan TimespanMaxValue( );
 
 	/* Returns the minimum time span value */
-	UFUNCTION(BlueprintPure, meta=(FriendlyName="MinValue"), Category="Math|Timespan")
+	UFUNCTION(BlueprintPure, meta=(DisplayName="MinValue"), Category="Math|Timespan")
 	static FTimespan TimespanMinValue( );
 
 	/* Returns the ratio between two time spans (A / B), handles zero values */
-	UFUNCTION(BlueprintPure, meta=(FriendlyName="TimespanRatio"), Category="Math|Timespan")
+	UFUNCTION(BlueprintPure, meta=(DisplayName="TimespanRatio"), Category="Math|Timespan")
 	static float TimespanRatio( FTimespan A, FTimespan B );
 
 	/* Returns a zero time span value */
-	UFUNCTION(BlueprintPure, meta=(FriendlyName="ZeroValue"), Category="Math|Timespan")
+	UFUNCTION(BlueprintPure, meta=(DisplayName="ZeroValue"), Category="Math|Timespan")
 	static FTimespan TimespanZeroValue( );
 
 	// -- Begin K2 utilities
 
 	/** Converts a byte to a float */
-	UFUNCTION(BlueprintPure, meta=(FriendlyName = "ToFloat (byte)", CompactNodeTitle = "->", Keywords="cast convert"), Category="Math|Conversions")
+	UFUNCTION(BlueprintPure, meta=(DisplayName = "ToFloat (byte)", CompactNodeTitle = "->", Keywords="cast convert"), Category="Math|Conversions")
 	static float Conv_ByteToFloat(uint8 InByte);
 
 	/** Converts an integer to a float */
-	UFUNCTION(BlueprintPure, meta=(FriendlyName = "ToFloat (int)", CompactNodeTitle = "->", Keywords="cast convert"), Category="Math|Conversions")
+	UFUNCTION(BlueprintPure, meta=(DisplayName = "ToFloat (int)", CompactNodeTitle = "->", Keywords="cast convert"), Category="Math|Conversions")
 	static float Conv_IntToFloat(int32 InInt);
 
 	/** Converts an integer to a byte (if the integer is too large, returns the low 8 bits) */
-	UFUNCTION(BlueprintPure, meta=(FriendlyName = "ToByte (int)", CompactNodeTitle = "->", Keywords="cast convert"), Category="Math|Conversions")
+	UFUNCTION(BlueprintPure, meta=(DisplayName = "ToByte (int)", CompactNodeTitle = "->", Keywords="cast convert"), Category="Math|Conversions")
 	static uint8 Conv_IntToByte(int32 InInt);
 
 	/** Converts a int to a bool*/
-	UFUNCTION(BlueprintPure, meta=(FriendlyName = "ToBool (int)", CompactNodeTitle = "->", Keywords="cast convert"), Category="Math|Conversions")
+	UFUNCTION(BlueprintPure, meta=(DisplayName = "ToBool (int)", CompactNodeTitle = "->", Keywords="cast convert"), Category="Math|Conversions")
 	static bool Conv_IntToBool(int32 InInt);
 
 	/** Converts a bool to an int */
-	UFUNCTION(BlueprintPure, meta=(FriendlyName = "ToInt (bool)", CompactNodeTitle = "->", Keywords="cast convert"), Category="Math|Conversions")
+	UFUNCTION(BlueprintPure, meta=(DisplayName = "ToInt (bool)", CompactNodeTitle = "->", Keywords="cast convert"), Category="Math|Conversions")
 	static int32 Conv_BoolToInt(bool InBool);
 
 	/** Converts a bool to a float (0.0f or 1.0f) */
-	UFUNCTION(BlueprintPure, meta=(FriendlyName = "ToFloat (bool)", CompactNodeTitle = "->", Keywords="cast convert"), Category="Math|Conversions")
+	UFUNCTION(BlueprintPure, meta=(DisplayName = "ToFloat (bool)", CompactNodeTitle = "->", Keywords="cast convert"), Category="Math|Conversions")
 	static float Conv_BoolToFloat(bool InBool);
 
 	/** Converts a bool to a byte */
-	UFUNCTION(BlueprintPure, meta=(FriendlyName = "ToByte (bool)", CompactNodeTitle = "->", Keywords="cast convert"), Category="Math|Conversions")
+	UFUNCTION(BlueprintPure, meta=(DisplayName = "ToByte (bool)", CompactNodeTitle = "->", Keywords="cast convert"), Category="Math|Conversions")
 	static uint8 Conv_BoolToByte(bool InBool);
 	
 	/** Converts a byte to an integer */
-	UFUNCTION(BlueprintPure, meta=(FriendlyName = "ToInt (byte)", CompactNodeTitle = "->", Keywords="cast convert"), Category="Math|Conversions")
+	UFUNCTION(BlueprintPure, meta=(DisplayName = "ToInt (byte)", CompactNodeTitle = "->", Keywords="cast convert"), Category="Math|Conversions")
 	static int32 Conv_ByteToInt(uint8 InByte);
 
 	/** Converts a vector to LinearColor */
-	UFUNCTION(BlueprintPure, meta=(FriendlyName = "ToLinearColor (vector)", CompactNodeTitle = "->", Keywords="cast convert"), Category="Math|Conversions")
+	UFUNCTION(BlueprintPure, meta=(DisplayName = "ToLinearColor (vector)", CompactNodeTitle = "->", Keywords="cast convert"), Category="Math|Conversions")
 	static FLinearColor Conv_VectorToLinearColor(FVector InVec);
 
 	/** Converts a LinearColor to a vector */
-	UFUNCTION(BlueprintPure, meta=(FriendlyName = "ToVector (linear color)", CompactNodeTitle = "->", Keywords="cast convert"), Category="Math|Conversions")
+	UFUNCTION(BlueprintPure, meta=(DisplayName = "ToVector (linear color)", CompactNodeTitle = "->", Keywords="cast convert"), Category="Math|Conversions")
 	static FVector Conv_LinearColorToVector(FLinearColor InLinearColor);
 
 	/** Converts a color to LinearColor */
-	UFUNCTION(BlueprintPure, meta=(FriendlyName = "ToLinearColor (color)", CompactNodeTitle = "->", Keywords="cast convert"), Category="Math|Conversions")
+	UFUNCTION(BlueprintPure, meta=(DisplayName = "ToLinearColor (color)", CompactNodeTitle = "->", Keywords="cast convert"), Category="Math|Conversions")
 	static FLinearColor Conv_ColorToLinearColor(FColor InColor);
 
 	/** Converts a LinearColor to a color*/
-	UFUNCTION(BlueprintPure, meta=(FriendlyName = "ToColor (linear color)", CompactNodeTitle = "->", Keywords="cast convert"), Category="Math|Conversions")
+	UFUNCTION(BlueprintPure, meta=(DisplayName = "ToColor (linear color)", CompactNodeTitle = "->", Keywords="cast convert"), Category="Math|Conversions")
 	static FColor Conv_LinearColorToColor(FLinearColor InLinearColor);
 
 	/** Convert a vector to a transform. Uses vector as location */
-	UFUNCTION(BlueprintPure, meta=(FriendlyName = "ToTransform (vector)", CompactNodeTitle = "->", Keywords="cast convert"), Category="Math|Conversions")
+	UFUNCTION(BlueprintPure, meta=(DisplayName = "ToTransform (vector)", CompactNodeTitle = "->", Keywords="cast convert"), Category="Math|Conversions")
 	static FTransform Conv_VectorToTransform(FVector InLocation);
 	
 	/** Convert a Vector to a Vector2D */
-	UFUNCTION(BlueprintPure, meta=(FriendlyName = "ToVector2D (Vector)", CompactNodeTitle = "->", Keywords="cast convert"), Category="Math|Conversions")
+	UFUNCTION(BlueprintPure, meta=(DisplayName = "ToVector2D (Vector)", CompactNodeTitle = "->", Keywords="cast convert"), Category="Math|Conversions")
 	static FVector2D Conv_VectorToVector2D(FVector InVector);
 
 	/** Convert a Vector2D to a Vector */
-	UFUNCTION(BlueprintPure, meta=(FriendlyName = "ToVector (Vector2D)", CompactNodeTitle = "->", Keywords="cast convert"), Category="Math|Conversions")
+	UFUNCTION(BlueprintPure, meta=(DisplayName = "ToVector (Vector2D)", CompactNodeTitle = "->", Keywords="cast convert"), Category="Math|Conversions")
 	static FVector Conv_Vector2DToVector(FVector2D InVector2D, float Z = 0);
 
 	/** Convert a float into a vector, where each element is that float */
-	UFUNCTION(BlueprintPure, meta=(FriendlyName = "ToVector (float)", CompactNodeTitle = "->", Keywords="cast convert"), Category="Math|Conversions")
+	UFUNCTION(BlueprintPure, meta=(DisplayName = "ToVector (float)", CompactNodeTitle = "->", Keywords="cast convert"), Category="Math|Conversions")
 	static FVector Conv_FloatToVector(float InFloat);
 
 	/** Convert a float into a LinearColor, where each element is that float */
-	UFUNCTION(BlueprintPure, meta=(FriendlyName = "ToLinearColor (float)", CompactNodeTitle = "->", Keywords="cast convert"), Category="Math|Conversions")
+	UFUNCTION(BlueprintPure, meta=(DisplayName = "ToLinearColor (float)", CompactNodeTitle = "->", Keywords="cast convert"), Category="Math|Conversions")
 	static FLinearColor Conv_FloatToLinearColor(float InFloat);
 
 	/** Makes a vector {X, Y, Z} */
@@ -1177,19 +1177,19 @@ class ENGINE_API UKismetMathLibrary : public UBlueprintFunctionLibrary
 	static void BreakColor(const FLinearColor InColor, float& R, float& G, float& B, float& A);
 
 	/** Make a color from individual color components (HSV space) */
-	UFUNCTION(BlueprintPure, Category="Math|Color", meta=(FriendlyName = "HSV to RGB"))
+	UFUNCTION(BlueprintPure, Category="Math|Color", meta=(DisplayName = "HSV to RGB"))
 	static FLinearColor HSVToRGB(float H, float S, float V, float A = 1.0f);
 
 	/** Breaks apart a color into individual HSV components (as well as alpha) */
-	UFUNCTION(BlueprintPure, Category="Math|Color", meta=(FriendlyName = "RGB to HSV"))
+	UFUNCTION(BlueprintPure, Category="Math|Color", meta=(DisplayName = "RGB to HSV"))
 	static void RGBToHSV(const FLinearColor InColor, float& H, float& S, float& V, float& A);
 
 	// Converts a HSV linear color (where H is in R, S is in G, and V is in B) to RGB
-	UFUNCTION(BlueprintPure, Category="Math|Color", meta=(FriendlyName = "HSV to RGB (vector)", Keywords="cast convert"))
+	UFUNCTION(BlueprintPure, Category="Math|Color", meta=(DisplayName = "HSV to RGB (vector)", Keywords="cast convert"))
 	static void HSVToRGB_Vector(const FLinearColor HSV, FLinearColor& RGB);
 
 	// Converts a RGB linear color to HSV (where H is in R, S is in G, and V is in B)
-	UFUNCTION(BlueprintPure, Category="Math|Color", meta=(FriendlyName = "RGB to HSV (vector)", Keywords="cast convert"))
+	UFUNCTION(BlueprintPure, Category="Math|Color", meta=(DisplayName = "RGB to HSV (vector)", Keywords="cast convert"))
 	static void RGBToHSV_Vector(const FLinearColor RGB, FLinearColor& HSV);
 
 	/** If bPickA is true, A is returned, otherwise B is */
@@ -1233,11 +1233,11 @@ class ENGINE_API UKismetMathLibrary : public UBlueprintFunctionLibrary
 	static FRotator MakeRotationFromAxes(FVector Forward, FVector Right, FVector Up);
 
 	/** Create a rotator which orients X along the supplied direction vector */
-	UFUNCTION(BlueprintPure, meta=(FriendlyName = "RotationFromXVector", Keywords="rotation rotate cast convert"), Category="Math|Rotator")
+	UFUNCTION(BlueprintPure, meta=(DisplayName = "RotationFromXVector", Keywords="rotation rotate cast convert"), Category="Math|Rotator")
 	static FRotator Conv_VectorToRotator(FVector InVec);
 
 	/** Get the X direction vector after this rotation */
-	UFUNCTION(BlueprintPure, meta=(FriendlyName = "GetRotationXVector", Keywords="rotation rotate cast convert"), Category="Math|Rotator")
+	UFUNCTION(BlueprintPure, meta=(DisplayName = "GetRotationXVector", Keywords="rotation rotate cast convert"), Category="Math|Rotator")
 	static FVector Conv_RotatorToVector(FRotator InRot);
 
 
@@ -1246,11 +1246,11 @@ class ENGINE_API UKismetMathLibrary : public UBlueprintFunctionLibrary
 	//
 	
 	/* Returns true if A and B are equal (A == B) */
-	UFUNCTION(BlueprintPure, meta=(FriendlyName = "Equal (Object)", CompactNodeTitle = "==", Keywords = "== equal"), Category="Utilities")
+	UFUNCTION(BlueprintPure, meta=(DisplayName = "Equal (Object)", CompactNodeTitle = "==", Keywords = "== equal"), Category="Utilities")
 	static bool EqualEqual_ObjectObject(class UObject* A, class UObject* B);
 
 	/* Returns true if A and B are not equal (A != B) */
-	UFUNCTION(BlueprintPure, meta=(FriendlyName = "NotEqual (Object)", CompactNodeTitle = "!=", Keywords = "!= not equal"), Category="Utilities")
+	UFUNCTION(BlueprintPure, meta=(DisplayName = "NotEqual (Object)", CompactNodeTitle = "!=", Keywords = "!= not equal"), Category="Utilities")
 	static bool NotEqual_ObjectObject(class UObject* A, class UObject* B);
 
 	//
@@ -1258,11 +1258,11 @@ class ENGINE_API UKismetMathLibrary : public UBlueprintFunctionLibrary
 	//
 
 	/* Returns true if A and B are equal (A == B) */
-	UFUNCTION(BlueprintPure, meta=(FriendlyName = "Equal (Class)", CompactNodeTitle = "==", Keywords = "== equal"), Category="Utilities")
+	UFUNCTION(BlueprintPure, meta=(DisplayName = "Equal (Class)", CompactNodeTitle = "==", Keywords = "== equal"), Category="Utilities")
 	static bool EqualEqual_ClassClass(class UClass* A, class UClass* B);
 
 	/* Returns true if A and B are not equal (A != B) */
-	UFUNCTION(BlueprintPure, meta=(FriendlyName = "NotEqual (Class)", CompactNodeTitle = "!=", Keywords = "!= not equal"), Category="Utilities")
+	UFUNCTION(BlueprintPure, meta=(DisplayName = "NotEqual (Class)", CompactNodeTitle = "!=", Keywords = "!= not equal"), Category="Utilities")
 	static bool NotEqual_ClassClass(class UClass* A, class UClass* B);
 
 	/**
@@ -1279,11 +1279,11 @@ class ENGINE_API UKismetMathLibrary : public UBlueprintFunctionLibrary
 	//
 	
 	/* Returns true if A and B are equal (A == B) */
-	UFUNCTION(BlueprintPure, meta=(FriendlyName = "Equal (Name)", CompactNodeTitle = "==", Keywords = "== equal"), Category="Utilities|Name")
+	UFUNCTION(BlueprintPure, meta=(DisplayName = "Equal (Name)", CompactNodeTitle = "==", Keywords = "== equal"), Category="Utilities|Name")
 	static bool EqualEqual_NameName(FName A, FName B);
 
 	/* Returns true if A and B are not equal (A != B) */
-	UFUNCTION(BlueprintPure, meta=(FriendlyName = "NotEqual (Name)", CompactNodeTitle = "!=", Keywords = "!= not equal"), Category="Utilities|Name")
+	UFUNCTION(BlueprintPure, meta=(DisplayName = "NotEqual (Name)", CompactNodeTitle = "!=", Keywords = "!= not equal"), Category="Utilities|Name")
 	static bool NotEqual_NameName(FName A, FName B);
 
 	//
@@ -1342,11 +1342,11 @@ class ENGINE_API UKismetMathLibrary : public UBlueprintFunctionLibrary
 	static FTransform ConvertTransformToRelative(const FTransform& WorldTransform, const FTransform& LocalTransform);
 
 	/* Linearly interpolates between A and B based on Alpha (100% of A when Alpha=0 and 100% of B when Alpha=1). */
-	UFUNCTION(BlueprintPure, meta=(FriendlyName = "Lerp (Transform)"), Category="Math|Transform")
+	UFUNCTION(BlueprintPure, meta=(DisplayName = "Lerp (Transform)"), Category="Math|Transform")
 	static FTransform TLerp(const FTransform& A, const FTransform& B, float Alpha);
 
 	/** Ease between A and B using a specified easing function. */
-	UFUNCTION(BlueprintPure, meta = (FriendlyName = "Ease (Transform)", BlueprintInternalUseOnly = "true"), Category = "Math|Interpolation")
+	UFUNCTION(BlueprintPure, meta = (DisplayName = "Ease (Transform)", BlueprintInternalUseOnly = "true"), Category = "Math|Interpolation")
 	static FTransform TEase(const FTransform& A, const FTransform& B, float Alpha, TEnumAsByte<EEasingFunc::Type> EasingFunc, float BlendExp = 2, int32 Steps = 2);
 
 	/** Tries to reach a target transform. */
@@ -1358,27 +1358,27 @@ class ENGINE_API UKismetMathLibrary : public UBlueprintFunctionLibrary
 	//
 
 	/** Returns addition of Vector A and Vector B (A + B) */
-	UFUNCTION(BlueprintPure, meta=(FriendlyName = "vector2d + vector2d", CompactNodeTitle = "+", Keywords = "+ add plus", CommutativeAssociativeBinaryOperator = "true"), Category="Math|Vector2D")
+	UFUNCTION(BlueprintPure, meta=(DisplayName = "vector2d + vector2d", CompactNodeTitle = "+", Keywords = "+ add plus", CommutativeAssociativeBinaryOperator = "true"), Category="Math|Vector2D")
 	static FVector2D Add_Vector2DVector2D(FVector2D A, FVector2D B);
 
 	/** Returns subtraction of Vector B from Vector A (A - B) */
-	UFUNCTION(BlueprintPure, meta=(FriendlyName = "vector2d - vector2d", CompactNodeTitle = "-", Keywords = "- subtract minus"), Category="Math|Vector2D")
+	UFUNCTION(BlueprintPure, meta=(DisplayName = "vector2d - vector2d", CompactNodeTitle = "-", Keywords = "- subtract minus"), Category="Math|Vector2D")
 	static FVector2D Subtract_Vector2DVector2D(FVector2D A, FVector2D B);
 
 	/** Returns Vector A scaled by B */
-	UFUNCTION(BlueprintPure, meta=(FriendlyName = "vector2d * float", CompactNodeTitle = "*", Keywords = "* multiply"), Category="Math|Vector2D")
+	UFUNCTION(BlueprintPure, meta=(DisplayName = "vector2d * float", CompactNodeTitle = "*", Keywords = "* multiply"), Category="Math|Vector2D")
 	static FVector2D Multiply_Vector2DFloat(FVector2D A, float B);
 
 	/** Returns Vector A divided by B */
-	UFUNCTION(BlueprintPure, meta=(FriendlyName = "vector2d / float", CompactNodeTitle = "/", Keywords = "/ divide division"), Category="Math|Vector2D")
+	UFUNCTION(BlueprintPure, meta=(DisplayName = "vector2d / float", CompactNodeTitle = "/", Keywords = "/ divide division"), Category="Math|Vector2D")
 	static FVector2D Divide_Vector2DFloat(FVector2D A, float B = 1.f);
 
 	/** Returns Vector A added by B */
-	UFUNCTION(BlueprintPure, meta=(FriendlyName = "vector2d + float", CompactNodeTitle = "+", Keywords = "+ add plus"), Category="Math|Vector2D")
+	UFUNCTION(BlueprintPure, meta=(DisplayName = "vector2d + float", CompactNodeTitle = "+", Keywords = "+ add plus"), Category="Math|Vector2D")
 	static FVector2D Add_Vector2DFloat(FVector2D A, float B);
 
 	/** Returns Vector A subtracted by B */
-	UFUNCTION(BlueprintPure, meta=(FriendlyName = "vector2d - float", CompactNodeTitle = "-", Keywords = "- subtract minus"), Category="Math|Vector2D")
+	UFUNCTION(BlueprintPure, meta=(DisplayName = "vector2d - float", CompactNodeTitle = "-", Keywords = "- subtract minus"), Category="Math|Vector2D")
 	static FVector2D Subtract_Vector2DFloat(FVector2D A, float B);
 
 	/**
@@ -1580,6 +1580,6 @@ class ENGINE_API UKismetMathLibrary : public UBlueprintFunctionLibrary
 	 * @param		Intersection - The point of intersection between the line and the plane
 	 * @return		True if the intersection test was successful.
 	 */
-	UFUNCTION(BlueprintCallable, Category = "Math|Intersection", meta = (FriendlyName = "Line Plane Intersection (Origin & Normal)"))
+	UFUNCTION(BlueprintCallable, Category = "Math|Intersection", meta = (DisplayName = "Line Plane Intersection (Origin & Normal)"))
 	static bool LinePlaneIntersection_OriginNormal(const FVector& LineStart, const FVector& LineEnd, FVector PlaneOrigin, FVector PlaneNormal, float& T, FVector& Intersection);
 };
diff --git a/Engine/Source/Runtime/Engine/Classes/Kismet/KismetStringLibrary.h b/Engine/Source/Runtime/Engine/Classes/Kismet/KismetStringLibrary.h
--- a/Engine/Source/Runtime/Engine/Classes/Kismet/KismetStringLibrary.h
+++ b/Engine/Source/Runtime/Engine/Classes/Kismet/KismetStringLibrary.h
@@ -10,59 +10,59 @@ class ENGINE_API UKismetStringLibrary : public UBlueprintFunctionLibrary
 	GENERATED_UCLASS_BODY()
 
 	/** Converts a float value to a string */
-	UFUNCTION(BlueprintPure, meta=(FriendlyName = "ToString (float)", CompactNodeTitle = "->"), Category="Utilities|String")
+	UFUNCTION(BlueprintPure, meta=(DisplayName = "ToString (float)", CompactNodeTitle = "->"), Category="Utilities|String")
 	static FString Conv_FloatToString(float InFloat);
 
 	/** Converts an integer value to a string */
-	UFUNCTION(BlueprintPure, meta=(FriendlyName = "ToString (int)", CompactNodeTitle = "->"), Category="Utilities|String")
+	UFUNCTION(BlueprintPure, meta=(DisplayName = "ToString (int)", CompactNodeTitle = "->"), Category="Utilities|String")
 	static FString Conv_IntToString(int32 InInt);
 
 	/** Converts a byte value to a string */
-	UFUNCTION(BlueprintPure, meta=(FriendlyName = "ToString (byte)", CompactNodeTitle = "->"), Category="Utilities|String")
+	UFUNCTION(BlueprintPure, meta=(DisplayName = "ToString (byte)", CompactNodeTitle = "->"), Category="Utilities|String")
 	static FString Conv_ByteToString(uint8 InByte);
 
 	/** Converts a boolean value to a string, either 'true' or 'false' */
-	UFUNCTION(BlueprintPure, meta=(FriendlyName = "ToString (boolean)", CompactNodeTitle = "->"), Category="Utilities|String")
+	UFUNCTION(BlueprintPure, meta=(DisplayName = "ToString (boolean)", CompactNodeTitle = "->"), Category="Utilities|String")
 	static FString Conv_BoolToString(bool InBool);
 
 	/** Converts a vector value to a string, in the form 'X= Y= Z=' */
-	UFUNCTION(BlueprintPure, meta=(FriendlyName = "ToString (Vector)", CompactNodeTitle = "->"), Category="Utilities|String")
+	UFUNCTION(BlueprintPure, meta=(DisplayName = "ToString (Vector)", CompactNodeTitle = "->"), Category="Utilities|String")
 	static FString Conv_VectorToString(FVector InVec);
 
 	/** Converts a vector2d value to a string, in the form 'X= Y=' */
-	UFUNCTION(BlueprintPure, meta=(FriendlyName = "ToString (vector2d)", CompactNodeTitle = "->"), Category="Utilities|String")
+	UFUNCTION(BlueprintPure, meta=(DisplayName = "ToString (vector2d)", CompactNodeTitle = "->"), Category="Utilities|String")
 	static FString Conv_Vector2dToString(FVector2D InVec);
 
 	/** Converts a rotator value to a string, in the form 'P= Y= R=' */
-	UFUNCTION(BlueprintPure, meta=(FriendlyName = "ToString (rotator)", CompactNodeTitle = "->"), Category="Utilities|String")
+	UFUNCTION(BlueprintPure, meta=(DisplayName = "ToString (rotator)", CompactNodeTitle = "->"), Category="Utilities|String")
 	static FString Conv_RotatorToString(FRotator InRot);
 
 	/** Converts a transform value to a string, in the form 'Translation: X= Y= Z= Rotation: P= Y= R= Scale: X= Y= Z=' */
-	UFUNCTION(BlueprintPure, meta=(FriendlyName = "ToString (transform)", CompactNodeTitle = "->"), Category="Utilities|String")
+	UFUNCTION(BlueprintPure, meta=(DisplayName = "ToString (transform)", CompactNodeTitle = "->"), Category="Utilities|String")
 	static FString Conv_TransformToString(const FTransform& InTrans);
 
 	/** Converts a UObject value to a string by calling the object's GetName method */
-	UFUNCTION(BlueprintPure, meta=(FriendlyName = "ToString (object)", CompactNodeTitle = "->"), Category="Utilities|String")
+	UFUNCTION(BlueprintPure, meta=(DisplayName = "ToString (object)", CompactNodeTitle = "->"), Category="Utilities|String")
 	static FString Conv_ObjectToString(class UObject* InObj);
 
 	/** Converts a linear color value to a string, in the form '(R=,G=,B=,A=)' */
-	UFUNCTION(BlueprintPure, meta=(FriendlyName = "ToString (linear color)", CompactNodeTitle = "->"), Category="Utilities|String")
+	UFUNCTION(BlueprintPure, meta=(DisplayName = "ToString (linear color)", CompactNodeTitle = "->"), Category="Utilities|String")
 	static FString Conv_ColorToString(FLinearColor InColor);
 
 	/** Converts a name value to a string */
-	UFUNCTION(BlueprintPure, meta=(FriendlyName = "ToString (name)", CompactNodeTitle = "->"), Category="Utilities|String")
+	UFUNCTION(BlueprintPure, meta=(DisplayName = "ToString (name)", CompactNodeTitle = "->"), Category="Utilities|String")
 	static FString Conv_NameToString(FName InName);
 
 	/** Converts a string to a name value */
-	UFUNCTION(BlueprintPure, meta=(FriendlyName = "String To Name", CompactNodeTitle = "->"), Category="Utilities|String")
+	UFUNCTION(BlueprintPure, meta=(DisplayName = "String To Name", CompactNodeTitle = "->"), Category="Utilities|String")
 	static FName Conv_StringToName(const FString& InString);
 
 	/** Converts a string to a int value */
-	UFUNCTION(BlueprintPure, meta=(FriendlyName = "String To Int", CompactNodeTitle = "->"), Category="Utilities|String")
+	UFUNCTION(BlueprintPure, meta=(DisplayName = "String To Int", CompactNodeTitle = "->"), Category="Utilities|String")
 	static int32 Conv_StringToInt(const FString& InString);
 
 	/** Converts a string to a float value */
-	UFUNCTION(BlueprintPure, meta=(FriendlyName = "String To Float", CompactNodeTitle = "->"), Category="Utilities|String")
+	UFUNCTION(BlueprintPure, meta=(DisplayName = "String To Float", CompactNodeTitle = "->"), Category="Utilities|String")
 	static float Conv_StringToFloat(const FString& InString);
 	
 
@@ -74,7 +74,7 @@ class ENGINE_API UKismetStringLibrary : public UBlueprintFunctionLibrary
 	 * @param Suffix - A suffix to append to the end of the conversion string
 	 * @return A new string built from the passed parameters
 	 */
-	UFUNCTION(BlueprintPure, meta=(FriendlyName = "BuildString (float)"), Category="Utilities|String")
+	UFUNCTION(BlueprintPure, meta=(DisplayName = "BuildString (float)"), Category="Utilities|String")
 	static FString BuildString_Float(const FString& AppendTo, const FString& Prefix, float InFloat, const FString& Suffix);
 
 	/** 
@@ -85,7 +85,7 @@ class ENGINE_API UKismetStringLibrary : public UBlueprintFunctionLibrary
 	 * @param Suffix - A suffix to append to the end of the conversion string
 	 * @return A new string built from the passed parameters
 	 */
-	UFUNCTION(BlueprintPure, meta=(FriendlyName = "BuildString (int)"), Category="Utilities|String")
+	UFUNCTION(BlueprintPure, meta=(DisplayName = "BuildString (int)"), Category="Utilities|String")
 	static FString BuildString_Int(const FString& AppendTo, const FString& Prefix, int32 InInt, const FString& Suffix);
 
 	/** 
@@ -96,7 +96,7 @@ class ENGINE_API UKismetStringLibrary : public UBlueprintFunctionLibrary
 	 * @param Suffix - A suffix to append to the end of the conversion string
 	 * @return A new string built from the passed parameters
 	 */
-	UFUNCTION(BlueprintPure, meta=(FriendlyName = "BuildString (boolean)"), Category="Utilities|String")
+	UFUNCTION(BlueprintPure, meta=(DisplayName = "BuildString (boolean)"), Category="Utilities|String")
 	static FString BuildString_Bool(const FString& AppendTo, const FString& Prefix, bool InBool, const FString& Suffix);
 
 	/** 
@@ -107,7 +107,7 @@ class ENGINE_API UKismetStringLibrary : public UBlueprintFunctionLibrary
 	 * @param Suffix - A suffix to append to the end of the conversion string
 	 * @return A new string built from the passed parameters
 	 */
-	UFUNCTION(BlueprintPure, meta=(FriendlyName = "BuildString (vector)"), Category="Utilities|String")
+	UFUNCTION(BlueprintPure, meta=(DisplayName = "BuildString (vector)"), Category="Utilities|String")
 	static FString BuildString_Vector(const FString& AppendTo, const FString& Prefix, FVector InVector, const FString& Suffix);
 
 	/** 
@@ -118,7 +118,7 @@ class ENGINE_API UKismetStringLibrary : public UBlueprintFunctionLibrary
 	 * @param Suffix - A suffix to append to the end of the conversion string
 	 * @return A new string built from the passed parameters
 	 */
-	UFUNCTION(BlueprintPure, meta=(FriendlyName = "BuildString (vector2d)"), Category="Utilities|String")
+	UFUNCTION(BlueprintPure, meta=(DisplayName = "BuildString (vector2d)"), Category="Utilities|String")
 	static FString BuildString_Vector2d(const FString& AppendTo, const FString& Prefix, FVector2D InVector2d, const FString& Suffix);
 
 	/** 
@@ -129,7 +129,7 @@ class ENGINE_API UKismetStringLibrary : public UBlueprintFunctionLibrary
 	 * @param Suffix - A suffix to append to the end of the conversion string
 	 * @return A new string built from the passed parameters
 	 */
-	UFUNCTION(BlueprintPure, meta=(FriendlyName = "BuildString (rotator)"), Category="Utilities|String")
+	UFUNCTION(BlueprintPure, meta=(DisplayName = "BuildString (rotator)"), Category="Utilities|String")
 	static FString BuildString_Rotator(const FString& AppendTo, const FString& Prefix, FRotator InRot, const FString& Suffix);
 
 	/** 
@@ -140,7 +140,7 @@ class ENGINE_API UKismetStringLibrary : public UBlueprintFunctionLibrary
 	 * @param Suffix - A suffix to append to the end of the conversion string
 	 * @return A new string built from the passed parameters
 	 */
-	UFUNCTION(BlueprintPure, meta=(FriendlyName = "BuildString (object)"), Category="Utilities|String")
+	UFUNCTION(BlueprintPure, meta=(DisplayName = "BuildString (object)"), Category="Utilities|String")
 	static FString BuildString_Object(const FString& AppendTo, const FString& Prefix, class UObject* InObj, const FString& Suffix);
 
 	/** 
@@ -151,7 +151,7 @@ class ENGINE_API UKismetStringLibrary : public UBlueprintFunctionLibrary
 	 * @param Suffix - A suffix to append to the end of the conversion string
 	 * @return A new string built from the passed parameters
 	 */
-	UFUNCTION(BlueprintPure, meta=(FriendlyName = "BuildString (color)"), Category="Utilities|String")
+	UFUNCTION(BlueprintPure, meta=(DisplayName = "BuildString (color)"), Category="Utilities|String")
 	static FString BuildString_Color(const FString& AppendTo, const FString& Prefix, FLinearColor InColor, const FString& Suffix);
 
 	/** 
@@ -162,7 +162,7 @@ class ENGINE_API UKismetStringLibrary : public UBlueprintFunctionLibrary
 	 * @param Suffix - A suffix to append to the end of the conversion string
 	 * @return A new string built from the passed parameters
 	 */
-	UFUNCTION(BlueprintPure, meta=(FriendlyName = "BuildString (name)"), Category="Utilities|String")
+	UFUNCTION(BlueprintPure, meta=(DisplayName = "BuildString (name)"), Category="Utilities|String")
 	static FString BuildString_Name(const FString& AppendTo, const FString& Prefix, FName InName, const FString& Suffix);
 
 
@@ -176,7 +176,7 @@ class ENGINE_API UKismetStringLibrary : public UBlueprintFunctionLibrary
 	 * @param B - The string to append to A
 	 * @returns A new string which is the concatenation of A+B
 	 */
-	UFUNCTION(BlueprintPure, meta=(FriendlyName = "Append", CommutativeAssociativeBinaryOperator = "true"), Category="Utilities|String")
+	UFUNCTION(BlueprintPure, meta=(DisplayName = "Append", CommutativeAssociativeBinaryOperator = "true"), Category="Utilities|String")
 	static FString Concat_StrStr(const FString& A, const FString& B);
 
 	/**
@@ -185,7 +185,7 @@ class ENGINE_API UKismetStringLibrary : public UBlueprintFunctionLibrary
 	 * @param B - The string to compare
 	 * @returns True if the strings are equal, false otherwise
 	 */
-	UFUNCTION(BlueprintPure, meta=(FriendlyName = "Equal (string)", CompactNodeTitle = "=="), Category="Utilities|String")
+	UFUNCTION(BlueprintPure, meta=(DisplayName = "Equal (string)", CompactNodeTitle = "=="), Category="Utilities|String")
 	static bool EqualEqual_StrStr(const FString& A, const FString& B);
 
 	/**
@@ -194,7 +194,7 @@ class ENGINE_API UKismetStringLibrary : public UBlueprintFunctionLibrary
 	 * @param B - The string to compare
 	 * @returns True if the strings are equal, false otherwise
 	 */
-	UFUNCTION(BlueprintPure, meta=(FriendlyName = "Equal, Case Insensitive (string)", CompactNodeTitle = "=="), Category="Utilities|String")
+	UFUNCTION(BlueprintPure, meta=(DisplayName = "Equal, Case Insensitive (string)", CompactNodeTitle = "=="), Category="Utilities|String")
 	static bool EqualEqual_StriStri(const FString& A, const FString& B);
 
 	/** 
@@ -203,7 +203,7 @@ class ENGINE_API UKismetStringLibrary : public UBlueprintFunctionLibrary
 	 * @param B - The string to compare
 	 * @return Returns true if the input strings are not equal, false if they are equal
 	 */
-	UFUNCTION(BlueprintPure, meta=(FriendlyName = "NotEqual (string)", CompactNodeTitle = "!="), Category="Utilities|String")
+	UFUNCTION(BlueprintPure, meta=(DisplayName = "NotEqual (string)", CompactNodeTitle = "!="), Category="Utilities|String")
 	static bool NotEqual_StrStr(const FString& A, const FString& B);
 
 	/** Test if the input string are not equal (A != B), ignoring case differences
@@ -211,7 +211,7 @@ class ENGINE_API UKismetStringLibrary : public UBlueprintFunctionLibrary
 	 * @param B - The string to compare
 	 * @return Returns true if the input strings are not equal, false if they are equal
 	 */
-	UFUNCTION(BlueprintPure, meta=(FriendlyName = "NotEqual, case insenstive (string)", CompactNodeTitle = "!="), Category="Utilities|String")
+	UFUNCTION(BlueprintPure, meta=(DisplayName = "NotEqual, case insenstive (string)", CompactNodeTitle = "!="), Category="Utilities|String")
 	static bool NotEqual_StriStri(const FString& A, const FString& B);
 
 	/** 
