Blueprint Call Function nodes will now display localized function names.

diff --git a/Engine/Source/Runtime/Engine/Classes/Kismet/KismetSystemLibrary.h b/Engine/Source/Runtime/Engine/Classes/Kismet/KismetSystemLibrary.h
--- a/Engine/Source/Runtime/Engine/Classes/Kismet/KismetSystemLibrary.h
+++ b/Engine/Source/Runtime/Engine/Classes/Kismet/KismetSystemLibrary.h
@@ -84,7 +84,7 @@ class ENGINE_API UKismetSystemLibrary : public UBlueprintFunctionLibrary
 	static FString GetDisplayName(const UObject* Object);
 
 	// Returns the display name of a class
-	UFUNCTION(BlueprintPure, Category = "Utilities", meta = (FriendlyName = "Get Display Name"))
+	UFUNCTION(BlueprintPure, Category = "Utilities", meta = (DisplayName = "Get Display Name"))
 	static FString GetClassDisplayName(UClass* Class);
 
 	// Engine build number, for displaying to end users.
@@ -126,7 +126,7 @@ class ENGINE_API UKismetSystemLibrary : public UBlueprintFunctionLibrary
 	UFUNCTION(BlueprintPure, Category="Utilities|Platform")
 	static FString GetUniqueDeviceId();
 
-	UFUNCTION(BlueprintPure, meta=(FriendlyName = "ToObject (interface)", CompactNodeTitle = "->"), Category="Utilities")
+	UFUNCTION(BlueprintPure, meta=(DisplayName = "ToObject (interface)", CompactNodeTitle = "->"), Category="Utilities")
 	static UObject* Conv_InterfaceToObject(const FScriptInterface& Interface); 
 
 	/**
@@ -276,63 +276,63 @@ class ENGINE_API UKismetSystemLibrary : public UBlueprintFunctionLibrary
 	 * @param Time			How long to wait before executing the delegate, in seconds. Setting a timer to <= 0 seconds will clear it if it is set.
 	 * @param bLooping		true to keep executing the delegate every Time seconds, false to execute delegate only once.
 	 */
-	UFUNCTION(BlueprintCallable, meta=(FriendlyName = "SetTimer_Delegate"), Category="Utilities|Time")
+	UFUNCTION(BlueprintCallable, meta=(DisplayName = "SetTimer_Delegate"), Category="Utilities|Time")
 	static void K2_SetTimerDelegate(FTimerDynamicDelegate Delegate, float Time, bool bLooping);
 
 	/**
 	 * Clears a set timer.
 	 * @param Delegate		Delegate. Can be a K2 function or a Custom Event.
 	 */
-	UFUNCTION(BlueprintCallable, meta=(FriendlyName = "ClearTimer_Delegate", DefaultToSelf = "Object"), Category="Utilities|Time")
+	UFUNCTION(BlueprintCallable, meta=(DisplayName = "ClearTimer_Delegate", DefaultToSelf = "Object"), Category="Utilities|Time")
 	static void K2_ClearTimerDelegate(FTimerDynamicDelegate Delegate);
 
 	/**
 	 * Pauses a set timer at its current elapsed time.
 	 * @param Delegate		Delegate. Can be a K2 function or a Custom Event.
 	 */
-	UFUNCTION(BlueprintCallable, meta=(FriendlyName = "PauseTimer_Delegate", DefaultToSelf = "Object"), Category="Utilities|Time")
+	UFUNCTION(BlueprintCallable, meta=(DisplayName = "PauseTimer_Delegate", DefaultToSelf = "Object"), Category="Utilities|Time")
 	static void K2_PauseTimerDelegate(FTimerDynamicDelegate Delegate);
 
 	/**
 	 * Resumes a paused timer from its current elapsed time.
 	 * @param Delegate		Delegate. Can be a K2 function or a Custom Event.
 	 */
-	UFUNCTION(BlueprintCallable, meta=(FriendlyName = "UnPauseTimer_Delegate", DefaultToSelf = "Object"), Category="Utilities|Time")
+	UFUNCTION(BlueprintCallable, meta=(DisplayName = "UnPauseTimer_Delegate", DefaultToSelf = "Object"), Category="Utilities|Time")
 	static void K2_UnPauseTimerDelegate(FTimerDynamicDelegate Delegate);
 
 	/**
 	 * Returns true is a timer exists and is active for the given delegate, false otherwise.
 	 * @param Delegate		Delegate. Can be a K2 function or a Custom Event.
 	 */
-	UFUNCTION(BlueprintPure, meta=(FriendlyName = "IsTimerActive_Delegate", DefaultToSelf = "Object"), Category="Utilities|Time")
+	UFUNCTION(BlueprintPure, meta=(DisplayName = "IsTimerActive_Delegate", DefaultToSelf = "Object"), Category="Utilities|Time")
 	static bool K2_IsTimerActiveDelegate(FTimerDynamicDelegate Delegate);
 
 	/**
 	 * Returns true is a timer exists and is paused for the given delegate, false otherwise.
 	 * @param Delegate		Delegate. Can be a K2 function or a Custom Event.
 	 */
-	UFUNCTION(BlueprintPure, meta = (FriendlyName = "IsTimerPaused_Delegate", DefaultToSelf = "Object"), Category = "Utilities|Time")
+	UFUNCTION(BlueprintPure, meta = (DisplayName = "IsTimerPaused_Delegate", DefaultToSelf = "Object"), Category = "Utilities|Time")
 	static bool K2_IsTimerPausedDelegate(FTimerDynamicDelegate Delegate);
 
 	/**
 	 * Returns true is a timer for the given delegate exists, false otherwise.
 	 * @param Delegate		Delegate. Can be a K2 function or a Custom Event.
 	 */
-	UFUNCTION(BlueprintPure, meta = (FriendlyName = "TimerExists_Delegate", DefaultToSelf = "Object"), Category = "Utilities|Time")
+	UFUNCTION(BlueprintPure, meta = (DisplayName = "TimerExists_Delegate", DefaultToSelf = "Object"), Category = "Utilities|Time")
 	static bool K2_TimerExistsDelegate(FTimerDynamicDelegate Delegate);
 	
 	/**
 	 * Returns elapsed time for the given delegate (time since current countdown iteration began).
 	 * @param Delegate		Delegate. Can be a K2 function or a Custom Event.
 	 */
-	UFUNCTION(BlueprintPure, meta=(FriendlyName = "GetTimerElapsedTime_Delegate", DefaultToSelf = "Object"), Category="Utilities|Time")
+	UFUNCTION(BlueprintPure, meta=(DisplayName = "GetTimerElapsedTime_Delegate", DefaultToSelf = "Object"), Category="Utilities|Time")
 	static float K2_GetTimerElapsedTimeDelegate(FTimerDynamicDelegate Delegate);
 
 	/**
 	 * Returns time until the timer will next execute its delegate.
 	 * @param Delegate		Delegate. Can be a K2 function or a Custom Event.
 	 */
-	UFUNCTION(BlueprintPure, meta=(FriendlyName = "GetTimerRemainingTime_Delegate", DefaultToSelf = "Object"), Category="Utilities|Time")
+	UFUNCTION(BlueprintPure, meta=(DisplayName = "GetTimerRemainingTime_Delegate", DefaultToSelf = "Object"), Category="Utilities|Time")
 	static float K2_GetTimerRemainingTimeDelegate(FTimerDynamicDelegate Delegate);
 
 	// --- Timer functions ------------------------------
@@ -344,7 +344,7 @@ class ENGINE_API UKismetSystemLibrary : public UBlueprintFunctionLibrary
 	 * @param Time			How long to wait before executing the delegate, in seconds. Setting a timer to <= 0 seconds will clear it if it is set.
 	 * @param bLooping		true to keep executing the delegate every Time seconds, false to execute delegate only once.
 	 */
-	UFUNCTION(BlueprintCallable, meta=(FriendlyName = "SetTimer", DefaultToSelf = "Object"), Category="Utilities|Time")
+	UFUNCTION(BlueprintCallable, meta=(DisplayName = "SetTimer", DefaultToSelf = "Object"), Category="Utilities|Time")
 	static void K2_SetTimer(UObject* Object, FString FunctionName, float Time, bool bLooping);
 
 	/**
@@ -352,7 +352,7 @@ class ENGINE_API UKismetSystemLibrary : public UBlueprintFunctionLibrary
 	 * @param Object		Object that implements the delegate function. Defaults to self (this blueprint)
 	 * @param FunctionName	Delegate function name. Can be a K2 function or a Custom Event.
 	 */
-	UFUNCTION(BlueprintCallable, meta=(FriendlyName = "ClearTimer", DefaultToSelf = "Object"), Category="Utilities|Time")
+	UFUNCTION(BlueprintCallable, meta=(DisplayName = "ClearTimer", DefaultToSelf = "Object"), Category="Utilities|Time")
 	static void K2_ClearTimer(UObject* Object, FString FunctionName);
 
 	/**
@@ -360,7 +360,7 @@ class ENGINE_API UKismetSystemLibrary : public UBlueprintFunctionLibrary
 	 * @param Object		Object that implements the delegate function. Defaults to self (this blueprint)
 	 * @param FunctionName	Delegate function name. Can be a K2 function or a Custom Event.
 	 */
-	UFUNCTION(BlueprintCallable, meta=(FriendlyName = "PauseTimer", DefaultToSelf = "Object"), Category="Utilities|Time")
+	UFUNCTION(BlueprintCallable, meta=(DisplayName = "PauseTimer", DefaultToSelf = "Object"), Category="Utilities|Time")
 	static void K2_PauseTimer(UObject* Object, FString FunctionName);
 
 	/**
@@ -368,7 +368,7 @@ class ENGINE_API UKismetSystemLibrary : public UBlueprintFunctionLibrary
 	 * @param Object		Object that implements the delegate function. Defaults to self (this blueprint)
 	 * @param FunctionName	Delegate function name. Can be a K2 function or a Custom Event.
 	 */
-	UFUNCTION(BlueprintCallable, meta=(FriendlyName = "UnPauseTimer", DefaultToSelf = "Object"), Category="Utilities|Time")
+	UFUNCTION(BlueprintCallable, meta=(DisplayName = "UnPauseTimer", DefaultToSelf = "Object"), Category="Utilities|Time")
 	static void K2_UnPauseTimer(UObject* Object, FString FunctionName);
 
 	/**
@@ -376,7 +376,7 @@ class ENGINE_API UKismetSystemLibrary : public UBlueprintFunctionLibrary
 	 * @param Object		Object that implements the delegate function. Defaults to self (this blueprint)
 	 * @param FunctionName	Delegate function name. Can be a K2 function or a Custom Event.
 	 */
-	UFUNCTION(BlueprintPure, meta=(FriendlyName = "IsTimerActive", DefaultToSelf = "Object"), Category="Utilities|Time")
+	UFUNCTION(BlueprintPure, meta=(DisplayName = "IsTimerActive", DefaultToSelf = "Object"), Category="Utilities|Time")
 	static bool K2_IsTimerActive(UObject* Object, FString FunctionName);
 
 	/**
@@ -384,7 +384,7 @@ class ENGINE_API UKismetSystemLibrary : public UBlueprintFunctionLibrary
 	* @param Object		Object that implements the delegate function. Defaults to self (this blueprint)
 	* @param FunctionName	Delegate function name. Can be a K2 function or a Custom Event.
 	*/
-	UFUNCTION(BlueprintPure, meta = (FriendlyName = "IsTimerPaused", DefaultToSelf = "Object"), Category = "Utilities|Time")
+	UFUNCTION(BlueprintPure, meta = (DisplayName = "IsTimerPaused", DefaultToSelf = "Object"), Category = "Utilities|Time")
 	static bool K2_IsTimerPaused(UObject* Object, FString FunctionName);
 
 	/**
@@ -392,7 +392,7 @@ class ENGINE_API UKismetSystemLibrary : public UBlueprintFunctionLibrary
 	* @param Object		Object that implements the delegate function. Defaults to self (this blueprint)
 	* @param FunctionName	Delegate function name. Can be a K2 function or a Custom Event.
 	*/
-	UFUNCTION(BlueprintPure, meta = (FriendlyName = "TimerExists", DefaultToSelf = "Object"), Category = "Utilities|Time")
+	UFUNCTION(BlueprintPure, meta = (DisplayName = "TimerExists", DefaultToSelf = "Object"), Category = "Utilities|Time")
 	static bool K2_TimerExists(UObject* Object, FString FunctionName);
 	
 	/**
@@ -400,7 +400,7 @@ class ENGINE_API UKismetSystemLibrary : public UBlueprintFunctionLibrary
 	 * @param Object		Object that implements the delegate function. Defaults to self (this blueprint)
 	 * @param FunctionName	Delegate function name. Can be a K2 function or a Custom Event.
 	 */
-	UFUNCTION(BlueprintPure, meta=(FriendlyName = "GetTimerElapsedTime", DefaultToSelf = "Object"), Category="Utilities|Time")
+	UFUNCTION(BlueprintPure, meta=(DisplayName = "GetTimerElapsedTime", DefaultToSelf = "Object"), Category="Utilities|Time")
 	static float K2_GetTimerElapsedTime(UObject* Object, FString FunctionName);
 
 	/**
@@ -408,7 +408,7 @@ class ENGINE_API UKismetSystemLibrary : public UBlueprintFunctionLibrary
 	 * @param Object		Object that implements the delegate function. Defaults to self (this blueprint)
 	 * @param FunctionName	Delegate function name. Can be a K2 function or a Custom Event.
 	 */
-	UFUNCTION(BlueprintPure, meta=(FriendlyName = "GetTimerRemainingTime", DefaultToSelf = "Object"), Category="Utilities|Time")
+	UFUNCTION(BlueprintPure, meta=(DisplayName = "GetTimerRemainingTime", DefaultToSelf = "Object"), Category="Utilities|Time")
 	static float K2_GetTimerRemainingTime(UObject* Object, FString FunctionName);
 
 
@@ -518,7 +518,7 @@ class ENGINE_API UKismetSystemLibrary : public UBlueprintFunctionLibrary
 	 * @param OutActors		Returned array of actors. Unsorted.
 	 * @return				true if there was an overlap that passed the filters, false otherwise.
 	 */
-	UFUNCTION(BlueprintCallable, Category="Collision", meta=(WorldContext="WorldContextObject", AutoCreateRefTerm="ActorsToIgnore", FriendlyName = "SphereOverlapActors"))
+	UFUNCTION(BlueprintCallable, Category="Collision", meta=(WorldContext="WorldContextObject", AutoCreateRefTerm="ActorsToIgnore", DisplayName = "SphereOverlapActors"))
 	static bool SphereOverlapActors_NEW(UObject* WorldContextObject, const FVector SpherePos, float SphereRadius, const TArray<TEnumAsByte<EObjectTypeQuery> > & ObjectTypes, UClass* ActorClassFilter, const TArray<AActor*>& ActorsToIgnore, TArray<class AActor*>& OutActors);
 
 	UFUNCTION(BlueprintCallable, Category="Collision", meta=(DeprecatedFunction, DeprecationMessage = "Use new SphereOverlapActors", WorldContext="WorldContextObject", AutoCreateRefTerm="ActorsToIgnore"))
@@ -535,7 +535,7 @@ class ENGINE_API UKismetSystemLibrary : public UBlueprintFunctionLibrary
 	 * @param OutActors		Returned array of actors. Unsorted.
 	 * @return				true if there was an overlap that passed the filters, false otherwise.
 	 */
-	UFUNCTION(BlueprintCallable, Category="Collision", meta=(WorldContext="WorldContextObject", AutoCreateRefTerm="ActorsToIgnore", FriendlyName="SphereOverlapComponents"))
+	UFUNCTION(BlueprintCallable, Category="Collision", meta=(WorldContext="WorldContextObject", AutoCreateRefTerm="ActorsToIgnore", DisplayName="SphereOverlapComponents"))
 	static bool SphereOverlapComponents_NEW(UObject* WorldContextObject, const FVector SpherePos, float SphereRadius, const TArray<TEnumAsByte<EObjectTypeQuery> > & ObjectTypes, UClass* ComponentClassFilter, const TArray<AActor*>& ActorsToIgnore, TArray<class UPrimitiveComponent*>& OutComponents);
 
 	UFUNCTION(BlueprintCallable, Category="Collision", meta=(DeprecatedFunction, DeprecationMessage = "Use new SphereOverlapComponents", WorldContext="WorldContextObject", AutoCreateRefTerm="ActorsToIgnore"))
@@ -553,7 +553,7 @@ class ENGINE_API UKismetSystemLibrary : public UBlueprintFunctionLibrary
 	 * @param OutActors		Returned array of actors. Unsorted.
 	 * @return				true if there was an overlap that passed the filters, false otherwise.
 	 */
-	UFUNCTION(BlueprintCallable, Category="Collision", meta=(WorldContext="WorldContextObject", AutoCreateRefTerm="ActorsToIgnore", FriendlyName="BoxOverlapActors"))
+	UFUNCTION(BlueprintCallable, Category="Collision", meta=(WorldContext="WorldContextObject", AutoCreateRefTerm="ActorsToIgnore", DisplayName="BoxOverlapActors"))
 	static bool BoxOverlapActors_NEW(UObject* WorldContextObject, const FVector BoxPos, FVector BoxExtent, const TArray<TEnumAsByte<EObjectTypeQuery> > & ObjectTypes, UClass* ActorClassFilter, const TArray<AActor*>& ActorsToIgnore, TArray<class AActor*>& OutActors);
 
 	UFUNCTION(BlueprintCallable, Category="Collision", meta=(DeprecatedFunction, DeprecationMessage = "Use new BoxOverlapActors", WorldContext="WorldContextObject", AutoCreateRefTerm="ActorsToIgnore"))
@@ -570,7 +570,7 @@ class ENGINE_API UKismetSystemLibrary : public UBlueprintFunctionLibrary
 	 * @param OutActors		Returned array of actors. Unsorted.
 	 * @return				true if there was an overlap that passed the filters, false otherwise.
 	 */
-	UFUNCTION(BlueprintCallable, Category="Collision", meta=(WorldContext="WorldContextObject", AutoCreateRefTerm="ActorsToIgnore", FriendlyName="BoxOverlapComponents"))
+	UFUNCTION(BlueprintCallable, Category="Collision", meta=(WorldContext="WorldContextObject", AutoCreateRefTerm="ActorsToIgnore", DisplayName="BoxOverlapComponents"))
 	static bool BoxOverlapComponents_NEW(UObject* WorldContextObject, const FVector BoxPos, FVector Extent, const TArray<TEnumAsByte<EObjectTypeQuery> > & ObjectTypes, UClass* ComponentClassFilter, const TArray<AActor*>& ActorsToIgnore, TArray<class UPrimitiveComponent*>& OutComponents);
 
 	UFUNCTION(BlueprintCallable, Category="Collision", meta=(DeprecatedFunction, DeprecationMessage = "Use new BoxOverlapComponents", WorldContext="WorldContextObject", AutoCreateRefTerm="ActorsToIgnore"))
@@ -589,7 +589,7 @@ class ENGINE_API UKismetSystemLibrary : public UBlueprintFunctionLibrary
 	 * @param OutActors		Returned array of actors. Unsorted.
 	 * @return				true if there was an overlap that passed the filters, false otherwise.
 	 */
-	UFUNCTION(BlueprintCallable, Category="Collision", meta=(WorldContext="WorldContextObject", AutoCreateRefTerm="ActorsToIgnore", FriendlyName="CapsuleOverlapActors"))
+	UFUNCTION(BlueprintCallable, Category="Collision", meta=(WorldContext="WorldContextObject", AutoCreateRefTerm="ActorsToIgnore", DisplayName="CapsuleOverlapActors"))
 	static bool CapsuleOverlapActors_NEW(UObject* WorldContextObject, const FVector CapsulePos, float Radius, float HalfHeight, const TArray<TEnumAsByte<EObjectTypeQuery> > & ObjectTypes, UClass* ActorClassFilter, const TArray<AActor*>& ActorsToIgnore, TArray<class AActor*>& OutActors);
 
 	UFUNCTION(BlueprintCallable, Category="Collision", meta=(DeprecatedFunction, DeprecationMessage = "Use new CapsuleOverlapActors", WorldContext="WorldContextObject", AutoCreateRefTerm="ActorsToIgnore"))
@@ -606,7 +606,7 @@ class ENGINE_API UKismetSystemLibrary : public UBlueprintFunctionLibrary
 	 * @param OutActors		Returned array of actors. Unsorted.
 	 * @return				true if there was an overlap that passed the filters, false otherwise.
 	 */
-	UFUNCTION(BlueprintCallable, Category="Collision", meta=(WorldContext="WorldContextObject", AutoCreateRefTerm="ActorsToIgnore", FriendlyName="CapsuleOverlapComponents") )
+	UFUNCTION(BlueprintCallable, Category="Collision", meta=(WorldContext="WorldContextObject", AutoCreateRefTerm="ActorsToIgnore", DisplayName="CapsuleOverlapComponents") )
 	static bool CapsuleOverlapComponents_NEW(UObject* WorldContextObject, const FVector CapsulePos, float Radius, float HalfHeight, const TArray<TEnumAsByte<EObjectTypeQuery> > & ObjectTypes, UClass* ComponentClassFilter, const TArray<AActor*>& ActorsToIgnore, TArray<class UPrimitiveComponent*>& OutComponents);
 
 	UFUNCTION(BlueprintCallable, Category="Collision", meta=(DeprecatedFunction, DeprecationMessage = "Use new CapsuleOverlapComponents", WorldContext="WorldContextObject", AutoCreateRefTerm="ActorsToIgnore") )
@@ -623,7 +623,7 @@ class ENGINE_API UKismetSystemLibrary : public UBlueprintFunctionLibrary
 	 * @param OutActors				Returned array of actors. Unsorted.
 	 * @return						true if there was an overlap that passed the filters, false otherwise.
 	 */
-	UFUNCTION(BlueprintCallable, Category="Collision", meta=(AutoCreateRefTerm="ActorsToIgnore", FriendlyName="ComponentOverlapActors"))
+	UFUNCTION(BlueprintCallable, Category="Collision", meta=(AutoCreateRefTerm="ActorsToIgnore", DisplayName="ComponentOverlapActors"))
 	static bool ComponentOverlapActors_NEW(UPrimitiveComponent* Component, const FTransform& ComponentTransform, const TArray<TEnumAsByte<EObjectTypeQuery> > & ObjectTypes, UClass* ActorClassFilter, const TArray<AActor*>& ActorsToIgnore, TArray<class AActor*>& OutActors);
 
 	UFUNCTION(BlueprintCallable, Category="Collision", meta=(DeprecatedFunction, DeprecationMessage = "Use new ComponentOverlapActors", AutoCreateRefTerm="ActorsToIgnore"))
@@ -639,7 +639,7 @@ class ENGINE_API UKismetSystemLibrary : public UBlueprintFunctionLibrary
 	 * @param OutActors				Returned array of actors. Unsorted.
 	 * @return						true if there was an overlap that passed the filters, false otherwise.
 	 */
-	UFUNCTION(BlueprintCallable, Category="Collision", meta=(AutoCreateRefTerm="ActorsToIgnore", FriendlyName="ComponentOverlapComponents"))
+	UFUNCTION(BlueprintCallable, Category="Collision", meta=(AutoCreateRefTerm="ActorsToIgnore", DisplayName="ComponentOverlapComponents"))
 	static bool ComponentOverlapComponents_NEW(UPrimitiveComponent* Component, const FTransform& ComponentTransform, const TArray<TEnumAsByte<EObjectTypeQuery> > & ObjectTypes, UClass* ComponentClassFilter, const TArray<AActor*>& ActorsToIgnore, TArray<class UPrimitiveComponent*>& OutComponents);
 
 	UFUNCTION(BlueprintCallable, Category="Collision", meta=(DeprecatedFunction, DeprecationMessage = "Use new ComponentOverlapComponents", AutoCreateRefTerm="ActorsToIgnore"))
@@ -657,7 +657,7 @@ class ENGINE_API UKismetSystemLibrary : public UBlueprintFunctionLibrary
 	 * @param OutHit		Properties of the trace hit.
 	 * @return				True if there was a hit, false otherwise.
 	 */
-	UFUNCTION(BlueprintCallable, Category="Collision", meta=(bIgnoreSelf="true", WorldContext="WorldContextObject", AutoCreateRefTerm="ActorsToIgnore", FriendlyName="LineTraceByChannel", Keywords="raycast"))
+	UFUNCTION(BlueprintCallable, Category="Collision", meta=(bIgnoreSelf="true", WorldContext="WorldContextObject", AutoCreateRefTerm="ActorsToIgnore", DisplayName="LineTraceByChannel", Keywords="raycast"))
 	static bool LineTraceSingle_NEW(UObject* WorldContextObject, const FVector Start, const FVector End, ETraceTypeQuery TraceChannel, bool bTraceComplex, const TArray<AActor*>& ActorsToIgnore, EDrawDebugTrace::Type DrawDebugType, FHitResult& OutHit, bool bIgnoreSelf);
 	
 	/**
@@ -672,7 +672,7 @@ class ENGINE_API UKismetSystemLibrary : public UBlueprintFunctionLibrary
 	 * @param OutHit		Properties of the trace hit.
 	 * @return				True if there was a hit, false otherwise.
 	 */
-	UFUNCTION(BlueprintCallable, Category="Collision", meta=(bIgnoreSelf="true", WorldContext="WorldContextObject", AutoCreateRefTerm="ActorsToIgnore", FriendlyName = "MultiLineTraceByChannel", Keywords="raycast"))
+	UFUNCTION(BlueprintCallable, Category="Collision", meta=(bIgnoreSelf="true", WorldContext="WorldContextObject", AutoCreateRefTerm="ActorsToIgnore", DisplayName = "MultiLineTraceByChannel", Keywords="raycast"))
 	static bool LineTraceMulti_NEW(UObject* WorldContextObject, const FVector Start, const FVector End, ETraceTypeQuery TraceChannel, bool bTraceComplex, const TArray<AActor*>& ActorsToIgnore, EDrawDebugTrace::Type DrawDebugType, TArray<FHitResult>& OutHits, bool bIgnoreSelf);
 
 	/**
@@ -687,7 +687,7 @@ class ENGINE_API UKismetSystemLibrary : public UBlueprintFunctionLibrary
 	 * @param OutHit		Properties of the trace hit.
 	 * @return				True if there was a hit, false otherwise.
 	 */
-	UFUNCTION(BlueprintCallable, Category="Collision", meta=(bIgnoreSelf="true", WorldContext="WorldContextObject", AutoCreateRefTerm="ActorsToIgnore", FriendlyName = "SphereTraceByChannel", Keywords="sweep"))
+	UFUNCTION(BlueprintCallable, Category="Collision", meta=(bIgnoreSelf="true", WorldContext="WorldContextObject", AutoCreateRefTerm="ActorsToIgnore", DisplayName = "SphereTraceByChannel", Keywords="sweep"))
 	static bool SphereTraceSingle_NEW(UObject* WorldContextObject, const FVector Start, const FVector End, float Radius, ETraceTypeQuery TraceChannel, bool bTraceComplex, const TArray<AActor*>& ActorsToIgnore, EDrawDebugTrace::Type DrawDebugType, FHitResult& OutHit, bool bIgnoreSelf);
 
 	/**
@@ -703,7 +703,7 @@ class ENGINE_API UKismetSystemLibrary : public UBlueprintFunctionLibrary
 	 * @param OutHits		A list of hits, sorted along the trace from start to finish.  The blocking hit will be the last hit, if there was one.
 	 * @return				True if there was a hit, false otherwise.
 	 */
-	UFUNCTION(BlueprintCallable, Category="Collision", meta=(bIgnoreSelf="true", WorldContext="WorldContextObject", AutoCreateRefTerm="ActorsToIgnore", FriendlyName = "MultiSphereTraceByChannel", Keywords="sweep"))
+	UFUNCTION(BlueprintCallable, Category="Collision", meta=(bIgnoreSelf="true", WorldContext="WorldContextObject", AutoCreateRefTerm="ActorsToIgnore", DisplayName = "MultiSphereTraceByChannel", Keywords="sweep"))
 	static bool SphereTraceMulti_NEW(UObject* WorldContextObject, const FVector Start, const FVector End, float Radius, ETraceTypeQuery TraceChannel, bool bTraceComplex, const TArray<AActor*>& ActorsToIgnore, EDrawDebugTrace::Type DrawDebugType, TArray<FHitResult>& OutHits, bool bIgnoreSelf);
 
 	/**
@@ -719,7 +719,7 @@ class ENGINE_API UKismetSystemLibrary : public UBlueprintFunctionLibrary
 	* @param OutHit			Properties of the trace hit.
 	* @return				True if there was a hit, false otherwise.
 	*/
-	UFUNCTION(BlueprintCallable, Category = "Collision", meta = (bIgnoreSelf = "true", HidePin = "WorldContextObject", DefaultToSelf = "WorldContextObject", AutoCreateRefTerm = "ActorsToIgnore", FriendlyName = "BoxTraceByChannel", Keywords="sweep"))
+	UFUNCTION(BlueprintCallable, Category = "Collision", meta = (bIgnoreSelf = "true", HidePin = "WorldContextObject", DefaultToSelf = "WorldContextObject", AutoCreateRefTerm = "ActorsToIgnore", DisplayName = "BoxTraceByChannel", Keywords="sweep"))
 	static bool BoxTraceSingle(UObject* WorldContextObject, const FVector Start, const FVector End, const FVector HalfSize, const FRotator Orientation, ETraceTypeQuery TraceChannel, bool bTraceComplex, const TArray<AActor*>& ActorsToIgnore, EDrawDebugTrace::Type DrawDebugType, FHitResult& OutHit, bool bIgnoreSelf);
 
 	/**
@@ -735,7 +735,7 @@ class ENGINE_API UKismetSystemLibrary : public UBlueprintFunctionLibrary
 	* @param OutHits		A list of hits, sorted along the trace from start to finish. The blocking hit will be the last hit, if there was one.
 	* @return				True if there was a hit, false otherwise.
 	*/
-	UFUNCTION(BlueprintCallable, Category = "Collision", meta = (bIgnoreSelf = "true", HidePin = "WorldContextObject", DefaultToSelf = "WorldContextObject", AutoCreateRefTerm = "ActorsToIgnore", FriendlyName = "MultiBoxTraceByChannel", Keywords="sweep"))
+	UFUNCTION(BlueprintCallable, Category = "Collision", meta = (bIgnoreSelf = "true", HidePin = "WorldContextObject", DefaultToSelf = "WorldContextObject", AutoCreateRefTerm = "ActorsToIgnore", DisplayName = "MultiBoxTraceByChannel", Keywords="sweep"))
 	static bool BoxTraceMulti(UObject* WorldContextObject, const FVector Start, const FVector End, FVector HalfSize, const FRotator Orientation, ETraceTypeQuery TraceChannel, bool bTraceComplex, const TArray<AActor*>& ActorsToIgnore, EDrawDebugTrace::Type DrawDebugType, TArray<FHitResult>& OutHits, bool bIgnoreSelf);
 
 	
@@ -753,7 +753,7 @@ class ENGINE_API UKismetSystemLibrary : public UBlueprintFunctionLibrary
 	 * @param OutHit		Properties of the trace hit.
 	 * @return				True if there was a hit, false otherwise.
 	 */
-	UFUNCTION(BlueprintCallable, Category="Collision", meta=(bIgnoreSelf="true", WorldContext="WorldContextObject", AutoCreateRefTerm="ActorsToIgnore", FriendlyName = "CapsuleTraceByChannel", Keywords="sweep"))
+	UFUNCTION(BlueprintCallable, Category="Collision", meta=(bIgnoreSelf="true", WorldContext="WorldContextObject", AutoCreateRefTerm="ActorsToIgnore", DisplayName = "CapsuleTraceByChannel", Keywords="sweep"))
 	static bool CapsuleTraceSingle_NEW(UObject* WorldContextObject, const FVector Start, const FVector End, float Radius, float HalfHeight, ETraceTypeQuery TraceChannel, bool bTraceComplex, const TArray<AActor*>& ActorsToIgnore, EDrawDebugTrace::Type DrawDebugType, FHitResult& OutHit, bool bIgnoreSelf);
 
 	/**
@@ -770,7 +770,7 @@ class ENGINE_API UKismetSystemLibrary : public UBlueprintFunctionLibrary
 	 * @param OutHits		A list of hits, sorted along the trace from start to finish.  The blocking hit will be the last hit, if there was one.
 	 * @return				True if there was a hit, false otherwise.
 	 */
-	UFUNCTION(BlueprintCallable, Category="Collision", meta=(bIgnoreSelf="true", WorldContext="WorldContextObject", AutoCreateRefTerm="ActorsToIgnore", FriendlyName = "MultiCapsuleTraceByChannel", Keywords="sweep"))
+	UFUNCTION(BlueprintCallable, Category="Collision", meta=(bIgnoreSelf="true", WorldContext="WorldContextObject", AutoCreateRefTerm="ActorsToIgnore", DisplayName = "MultiCapsuleTraceByChannel", Keywords="sweep"))
 	static bool CapsuleTraceMulti_NEW(UObject* WorldContextObject, const FVector Start, const FVector End, float Radius, float HalfHeight, ETraceTypeQuery TraceChannel, bool bTraceComplex, const TArray<AActor*>& ActorsToIgnore, EDrawDebugTrace::Type DrawDebugType, TArray<FHitResult>& OutHits, bool bIgnoreSelf);
 
 	/**
@@ -785,7 +785,7 @@ class ENGINE_API UKismetSystemLibrary : public UBlueprintFunctionLibrary
 	 * @param OutHit		Properties of the trace hit.
 	 * @return				True if there was a hit, false otherwise.
 	 */
-	UFUNCTION(BlueprintCallable, Category="Collision", meta=(bIgnoreSelf="true", WorldContext="WorldContextObject", AutoCreateRefTerm="ActorsToIgnore", FriendlyName = "LineTraceForObjects", Keywords="raycast"))
+	UFUNCTION(BlueprintCallable, Category="Collision", meta=(bIgnoreSelf="true", WorldContext="WorldContextObject", AutoCreateRefTerm="ActorsToIgnore", DisplayName = "LineTraceForObjects", Keywords="raycast"))
 	static bool LineTraceSingleForObjects(UObject* WorldContextObject, const FVector Start, const FVector End, const TArray<TEnumAsByte<EObjectTypeQuery> > & ObjectTypes, bool bTraceComplex, const TArray<AActor*>& ActorsToIgnore, EDrawDebugTrace::Type DrawDebugType, FHitResult& OutHit, bool bIgnoreSelf);
 	
 	/**
@@ -800,7 +800,7 @@ class ENGINE_API UKismetSystemLibrary : public UBlueprintFunctionLibrary
 	 * @param OutHit		Properties of the trace hit.
 	 * @return				True if there was a hit, false otherwise.
 	 */
-	UFUNCTION(BlueprintCallable, Category="Collision", meta=(bIgnoreSelf="true", WorldContext="WorldContextObject", AutoCreateRefTerm="ActorsToIgnore", FriendlyName = "MultiLineTraceForObjects", Keywords="raycast"))
+	UFUNCTION(BlueprintCallable, Category="Collision", meta=(bIgnoreSelf="true", WorldContext="WorldContextObject", AutoCreateRefTerm="ActorsToIgnore", DisplayName = "MultiLineTraceForObjects", Keywords="raycast"))
 	static bool LineTraceMultiForObjects(UObject* WorldContextObject, const FVector Start, const FVector End, const TArray<TEnumAsByte<EObjectTypeQuery> > & ObjectTypes, bool bTraceComplex, const TArray<AActor*>& ActorsToIgnore, EDrawDebugTrace::Type DrawDebugType, TArray<FHitResult>& OutHits, bool bIgnoreSelf);
 
 	/**
@@ -815,7 +815,7 @@ class ENGINE_API UKismetSystemLibrary : public UBlueprintFunctionLibrary
 	 * @param OutHit		Properties of the trace hit.
 	 * @return				True if there was a hit, false otherwise.
 	 */
-	UFUNCTION(BlueprintCallable, Category="Collision", meta=(bIgnoreSelf="true", WorldContext="WorldContextObject", AutoCreateRefTerm="ActorsToIgnore", FriendlyName = "SphereTraceForObjects", Keywords="sweep"))
+	UFUNCTION(BlueprintCallable, Category="Collision", meta=(bIgnoreSelf="true", WorldContext="WorldContextObject", AutoCreateRefTerm="ActorsToIgnore", DisplayName = "SphereTraceForObjects", Keywords="sweep"))
 	static bool SphereTraceSingleForObjects(UObject* WorldContextObject, const FVector Start, const FVector End, float Radius, const TArray<TEnumAsByte<EObjectTypeQuery> > & ObjectTypes, bool bTraceComplex, const TArray<AActor*>& ActorsToIgnore, EDrawDebugTrace::Type DrawDebugType, FHitResult& OutHit, bool bIgnoreSelf);
 
 	/**
@@ -831,7 +831,7 @@ class ENGINE_API UKismetSystemLibrary : public UBlueprintFunctionLibrary
 	 * @param OutHits		A list of hits, sorted along the trace from start to finish.  The blocking hit will be the last hit, if there was one.
 	 * @return				True if there was a hit, false otherwise.
 	 */
-	UFUNCTION(BlueprintCallable, Category="Collision", meta=(bIgnoreSelf="true", WorldContext="WorldContextObject", AutoCreateRefTerm="ActorsToIgnore", FriendlyName = "MultiSphereTraceForObjects", Keywords="sweep"))
+	UFUNCTION(BlueprintCallable, Category="Collision", meta=(bIgnoreSelf="true", WorldContext="WorldContextObject", AutoCreateRefTerm="ActorsToIgnore", DisplayName = "MultiSphereTraceForObjects", Keywords="sweep"))
 	static bool SphereTraceMultiForObjects(UObject* WorldContextObject, const FVector Start, const FVector End, float Radius, const TArray<TEnumAsByte<EObjectTypeQuery> > & ObjectTypes, bool bTraceComplex, const TArray<AActor*>& ActorsToIgnore, EDrawDebugTrace::Type DrawDebugType, TArray<FHitResult>& OutHits, bool bIgnoreSelf);
 	
 
@@ -848,7 +848,7 @@ class ENGINE_API UKismetSystemLibrary : public UBlueprintFunctionLibrary
 	* @param OutHit			Properties of the trace hit.
 	* @return				True if there was a hit, false otherwise.
 	*/
-	UFUNCTION(BlueprintCallable, Category = "Collision", meta = (bIgnoreSelf = "true", HidePin = "WorldContextObject", DefaultToSelf = "WorldContextObject", AutoCreateRefTerm = "ActorsToIgnore", FriendlyName = "BoxTraceForObjects", Keywords="sweep"))
+	UFUNCTION(BlueprintCallable, Category = "Collision", meta = (bIgnoreSelf = "true", HidePin = "WorldContextObject", DefaultToSelf = "WorldContextObject", AutoCreateRefTerm = "ActorsToIgnore", DisplayName = "BoxTraceForObjects", Keywords="sweep"))
 	static bool BoxTraceSingleForObjects(UObject* WorldContextObject, const FVector Start, const FVector End, const FVector HalfSize, const FRotator Orientation, const TArray<TEnumAsByte<EObjectTypeQuery> > & ObjectTypes, bool bTraceComplex, const TArray<AActor*>& ActorsToIgnore, EDrawDebugTrace::Type DrawDebugType, FHitResult& OutHit, bool bIgnoreSelf);
 
 
@@ -865,7 +865,7 @@ class ENGINE_API UKismetSystemLibrary : public UBlueprintFunctionLibrary
 	* @param OutHits		A list of hits, sorted along the trace from start to finish.  The blocking hit will be the last hit, if there was one.
 	* @return				True if there was a hit, false otherwise.
 	*/
-	UFUNCTION(BlueprintCallable, Category = "Collision", meta = (bIgnoreSelf = "true", HidePin = "WorldContextObject", DefaultToSelf = "WorldContextObject", AutoCreateRefTerm = "ActorsToIgnore", FriendlyName = "MultiBoxTraceForObjects", Keywords="sweep"))
+	UFUNCTION(BlueprintCallable, Category = "Collision", meta = (bIgnoreSelf = "true", HidePin = "WorldContextObject", DefaultToSelf = "WorldContextObject", AutoCreateRefTerm = "ActorsToIgnore", DisplayName = "MultiBoxTraceForObjects", Keywords="sweep"))
 	static bool BoxTraceMultiForObjects(UObject* WorldContextObject, const FVector Start, const FVector End, const FVector HalfSize, const FRotator Orientation, const TArray<TEnumAsByte<EObjectTypeQuery> > & ObjectTypes, bool bTraceComplex, const TArray<AActor*>& ActorsToIgnore, EDrawDebugTrace::Type DrawDebugType, TArray<FHitResult>& OutHits, bool bIgnoreSelf);
 
 	/**
@@ -882,7 +882,7 @@ class ENGINE_API UKismetSystemLibrary : public UBlueprintFunctionLibrary
 	 * @param OutHit		Properties of the trace hit.
 	 * @return				True if there was a hit, false otherwise.
 	 */
-	UFUNCTION(BlueprintCallable, Category="Collision", meta=(bIgnoreSelf="true", WorldContext="WorldContextObject", AutoCreateRefTerm="ActorsToIgnore", FriendlyName = "CapsuleTraceForObjects", Keywords="sweep"))
+	UFUNCTION(BlueprintCallable, Category="Collision", meta=(bIgnoreSelf="true", WorldContext="WorldContextObject", AutoCreateRefTerm="ActorsToIgnore", DisplayName = "CapsuleTraceForObjects", Keywords="sweep"))
 	static bool CapsuleTraceSingleForObjects(UObject* WorldContextObject, const FVector Start, const FVector End, float Radius, float HalfHeight, const TArray<TEnumAsByte<EObjectTypeQuery> > & ObjectTypes, bool bTraceComplex, const TArray<AActor*>& ActorsToIgnore, EDrawDebugTrace::Type DrawDebugType, FHitResult& OutHit, bool bIgnoreSelf);
 
 	/**
@@ -899,7 +899,7 @@ class ENGINE_API UKismetSystemLibrary : public UBlueprintFunctionLibrary
 	 * @param OutHits		A list of hits, sorted along the trace from start to finish.  The blocking hit will be the last hit, if there was one.
 	 * @return				True if there was a hit, false otherwise.
 	 */
-	UFUNCTION(BlueprintCallable, Category="Collision", meta=(bIgnoreSelf="true", WorldContext="WorldContextObject", AutoCreateRefTerm="ActorsToIgnore", FriendlyName = "MultiCapsuleTraceForObjects", Keywords="sweep"))
+	UFUNCTION(BlueprintCallable, Category="Collision", meta=(bIgnoreSelf="true", WorldContext="WorldContextObject", AutoCreateRefTerm="ActorsToIgnore", DisplayName = "MultiCapsuleTraceForObjects", Keywords="sweep"))
 	static bool CapsuleTraceMultiForObjects(UObject* WorldContextObject, const FVector Start, const FVector End, float Radius, float HalfHeight, const TArray<TEnumAsByte<EObjectTypeQuery> > & ObjectTypes, bool bTraceComplex, const TArray<AActor*>& ActorsToIgnore, EDrawDebugTrace::Type DrawDebugType, TArray<FHitResult>& OutHits, bool bIgnoreSelf);
 
 	/**
@@ -914,7 +914,7 @@ class ENGINE_API UKismetSystemLibrary : public UBlueprintFunctionLibrary
 	 * @param OutHit		Properties of the trace hit.
 	 * @return				True if there was a hit, false otherwise.
 	 */
-	UFUNCTION(BlueprintCallable, Category="Collision", meta=(DeprecatedFunction, DeprecationMessage = "Use new SingleLineTraceByChannel", bIgnoreSelf="true", WorldContext="WorldContextObject", AutoCreateRefTerm="ActorsToIgnore", FriendlyName = "SingleLineTraceByChannelDeprecated"))
+	UFUNCTION(BlueprintCallable, Category="Collision", meta=(DeprecatedFunction, DeprecationMessage = "Use new SingleLineTraceByChannel", bIgnoreSelf="true", WorldContext="WorldContextObject", AutoCreateRefTerm="ActorsToIgnore", DisplayName = "SingleLineTraceByChannelDeprecated"))
 	static bool LineTraceSingle_DEPRECATED(UObject* WorldContextObject, const FVector Start, const FVector End, ECollisionChannel TraceChannel, bool bTraceComplex, const TArray<AActor*>& ActorsToIgnore, EDrawDebugTrace::Type DrawDebugType, FHitResult& OutHit, bool bIgnoreSelf);
 	
 	/**
@@ -929,7 +929,7 @@ class ENGINE_API UKismetSystemLibrary : public UBlueprintFunctionLibrary
 	 * @param OutHit		Properties of the trace hit.
 	 * @return				True if there was a hit, false otherwise.
 	 */
-	UFUNCTION(BlueprintCallable, Category="Collision", meta=(DeprecatedFunction, DeprecationMessage = "Use new MultiLineTraceByChannel", bIgnoreSelf="true", WorldContext="WorldContextObject", AutoCreateRefTerm="ActorsToIgnore", FriendlyName = "MultiLineTraceByChannelDeprecated"))
+	UFUNCTION(BlueprintCallable, Category="Collision", meta=(DeprecatedFunction, DeprecationMessage = "Use new MultiLineTraceByChannel", bIgnoreSelf="true", WorldContext="WorldContextObject", AutoCreateRefTerm="ActorsToIgnore", DisplayName = "MultiLineTraceByChannelDeprecated"))
 	static bool LineTraceMulti_DEPRECATED(UObject* WorldContextObject, const FVector Start, const FVector End, ECollisionChannel TraceChannel, bool bTraceComplex, const TArray<AActor*>& ActorsToIgnore, EDrawDebugTrace::Type DrawDebugType, TArray<FHitResult>& OutHits, bool bIgnoreSelf);
 
 	/**
@@ -944,7 +944,7 @@ class ENGINE_API UKismetSystemLibrary : public UBlueprintFunctionLibrary
 	 * @param OutHit		Properties of the trace hit.
 	 * @return				True if there was a hit, false otherwise.
 	 */
-	UFUNCTION(BlueprintCallable, Category="Collision", meta=(DeprecatedFunction, DeprecationMessage = "Use new SingleSphereTraceByChannel", bIgnoreSelf="true", WorldContext="WorldContextObject", AutoCreateRefTerm="ActorsToIgnore", FriendlyName = "SingleSphereTraceByChannelDeprecated"))
+	UFUNCTION(BlueprintCallable, Category="Collision", meta=(DeprecatedFunction, DeprecationMessage = "Use new SingleSphereTraceByChannel", bIgnoreSelf="true", WorldContext="WorldContextObject", AutoCreateRefTerm="ActorsToIgnore", DisplayName = "SingleSphereTraceByChannelDeprecated"))
 	static bool SphereTraceSingle_DEPRECATED(UObject* WorldContextObject, const FVector Start, const FVector End, float Radius, ECollisionChannel TraceChannel, bool bTraceComplex, const TArray<AActor*>& ActorsToIgnore, EDrawDebugTrace::Type DrawDebugType, FHitResult& OutHit, bool bIgnoreSelf);
 
 	/**
@@ -960,7 +960,7 @@ class ENGINE_API UKismetSystemLibrary : public UBlueprintFunctionLibrary
 	 * @param OutHits		A list of hits, sorted along the trace from start to finish.  The blocking hit will be the last hit, if there was one.
 	 * @return				True if there was a hit, false otherwise.
 	 */
-	UFUNCTION(BlueprintCallable, Category="Collision", meta=(DeprecatedFunction, DeprecationMessage = "Use new MultiSphereTraceByChannel", bIgnoreSelf="true", WorldContext="WorldContextObject", AutoCreateRefTerm="ActorsToIgnore", FriendlyName = "MultiSphereTraceByChannelDeprecated"))
+	UFUNCTION(BlueprintCallable, Category="Collision", meta=(DeprecatedFunction, DeprecationMessage = "Use new MultiSphereTraceByChannel", bIgnoreSelf="true", WorldContext="WorldContextObject", AutoCreateRefTerm="ActorsToIgnore", DisplayName = "MultiSphereTraceByChannelDeprecated"))
 	static bool SphereTraceMulti_DEPRECATED(UObject* WorldContextObject, const FVector Start, const FVector End, float Radius, ECollisionChannel TraceChannel, bool bTraceComplex, const TArray<AActor*>& ActorsToIgnore, EDrawDebugTrace::Type DrawDebugType, TArray<FHitResult>& OutHits, bool bIgnoreSelf);
 	
 	/**
@@ -977,7 +977,7 @@ class ENGINE_API UKismetSystemLibrary : public UBlueprintFunctionLibrary
 	 * @param OutHit		Properties of the trace hit.
 	 * @return				True if there was a hit, false otherwise.
 	 */
-	UFUNCTION(BlueprintCallable, Category="Collision", meta=(DeprecatedFunction, DeprecationMessage = "Use new SingleCapsuleTraceByChannel", bIgnoreSelf="true", WorldContext="WorldContextObject", AutoCreateRefTerm="ActorsToIgnore", FriendlyName = "SingleCapsuleTraceByChannelDeprecated"))
+	UFUNCTION(BlueprintCallable, Category="Collision", meta=(DeprecatedFunction, DeprecationMessage = "Use new SingleCapsuleTraceByChannel", bIgnoreSelf="true", WorldContext="WorldContextObject", AutoCreateRefTerm="ActorsToIgnore", DisplayName = "SingleCapsuleTraceByChannelDeprecated"))
 	static bool CapsuleTraceSingle_DEPRECATED(UObject* WorldContextObject, const FVector Start, const FVector End, float Radius, float HalfHeight, ECollisionChannel TraceChannel, bool bTraceComplex, const TArray<AActor*>& ActorsToIgnore, EDrawDebugTrace::Type DrawDebugType, FHitResult& OutHit, bool bIgnoreSelf);
 
 	/**
@@ -994,7 +994,7 @@ class ENGINE_API UKismetSystemLibrary : public UBlueprintFunctionLibrary
 	 * @param OutHits		A list of hits, sorted along the trace from start to finish.  The blocking hit will be the last hit, if there was one.
 	 * @return				True if there was a hit, false otherwise.
 	 */
-	UFUNCTION(BlueprintCallable, Category="Collision", meta=(DeprecatedFunction, DeprecationMessage = "Use new MultiCapsuleTraceByChannel", bIgnoreSelf="true", WorldContext="WorldContextObject", AutoCreateRefTerm="ActorsToIgnore", FriendlyName = "MultiCapsuleTraceByChannelDeprecated"))
+	UFUNCTION(BlueprintCallable, Category="Collision", meta=(DeprecatedFunction, DeprecationMessage = "Use new MultiCapsuleTraceByChannel", bIgnoreSelf="true", WorldContext="WorldContextObject", AutoCreateRefTerm="ActorsToIgnore", DisplayName = "MultiCapsuleTraceByChannelDeprecated"))
 	static bool CapsuleTraceMulti_DEPRECATED(UObject* WorldContextObject, const FVector Start, const FVector End, float Radius, float HalfHeight, ECollisionChannel TraceChannel, bool bTraceComplex, const TArray<AActor*>& ActorsToIgnore, EDrawDebugTrace::Type DrawDebugType, TArray<FHitResult>& OutHits, bool bIgnoreSelf);
 
 	/**
@@ -1009,7 +1009,7 @@ class ENGINE_API UKismetSystemLibrary : public UBlueprintFunctionLibrary
 	 * @param OutHit		Properties of the trace hit.
 	 * @return				True if there was a hit, false otherwise.
 	 */
-	UFUNCTION(BlueprintCallable, Category="Collision", meta=(DeprecatedFunction, DeprecationMessage = "Use new SingleLineTraceForObjects", bIgnoreSelf="true", WorldContext="WorldContextObject", AutoCreateRefTerm="ActorsToIgnore", FriendlyName = "SingleLineTraceByObjectDeprecated"))
+	UFUNCTION(BlueprintCallable, Category="Collision", meta=(DeprecatedFunction, DeprecationMessage = "Use new SingleLineTraceForObjects", bIgnoreSelf="true", WorldContext="WorldContextObject", AutoCreateRefTerm="ActorsToIgnore", DisplayName = "SingleLineTraceByObjectDeprecated"))
 	static bool LineTraceSingleByObject_DEPRECATED(UObject* WorldContextObject, const FVector Start, const FVector End, const TArray<TEnumAsByte<ECollisionChannel> > & ObjectsToTrace, bool bTraceComplex, const TArray<AActor*>& ActorsToIgnore, EDrawDebugTrace::Type DrawDebugType, FHitResult& OutHit, bool bIgnoreSelf);
 	
 	/**
@@ -1024,7 +1024,7 @@ class ENGINE_API UKismetSystemLibrary : public UBlueprintFunctionLibrary
 	 * @param OutHit		Properties of the trace hit.
 	 * @return				True if there was a hit, false otherwise.
 	 */
-	UFUNCTION(BlueprintCallable, Category="Collision", meta=(DeprecatedFunction, DeprecationMessage = "Use new MultiLineTraceForObjects", bIgnoreSelf="true", WorldContext="WorldContextObject", AutoCreateRefTerm="ActorsToIgnore", FriendlyName = "MultiLineTraceByObjectDeprecated"))
+	UFUNCTION(BlueprintCallable, Category="Collision", meta=(DeprecatedFunction, DeprecationMessage = "Use new MultiLineTraceForObjects", bIgnoreSelf="true", WorldContext="WorldContextObject", AutoCreateRefTerm="ActorsToIgnore", DisplayName = "MultiLineTraceByObjectDeprecated"))
 	static bool LineTraceMultiByObject_DEPRECATED(UObject* WorldContextObject, const FVector Start, const FVector End, const TArray<TEnumAsByte<ECollisionChannel> > & ObjectsToTrace, bool bTraceComplex, const TArray<AActor*>& ActorsToIgnore, EDrawDebugTrace::Type DrawDebugType, TArray<FHitResult>& OutHits, bool bIgnoreSelf);
 
 	/**
@@ -1039,7 +1039,7 @@ class ENGINE_API UKismetSystemLibrary : public UBlueprintFunctionLibrary
 	 * @param OutHit		Properties of the trace hit.
 	 * @return				True if there was a hit, false otherwise.
 	 */
-	UFUNCTION(BlueprintCallable, Category="Collision", meta=(DeprecatedFunction, DeprecationMessage = "Use new SingleSphereTraceForObjects", bIgnoreSelf="true", WorldContext="WorldContextObject", AutoCreateRefTerm="ActorsToIgnore", FriendlyName = "SingleSphereTraceByObjectDeprecated"))
+	UFUNCTION(BlueprintCallable, Category="Collision", meta=(DeprecatedFunction, DeprecationMessage = "Use new SingleSphereTraceForObjects", bIgnoreSelf="true", WorldContext="WorldContextObject", AutoCreateRefTerm="ActorsToIgnore", DisplayName = "SingleSphereTraceByObjectDeprecated"))
 	static bool SphereTraceSingleByObject_DEPRECATED(UObject* WorldContextObject, const FVector Start, const FVector End, float Radius, const TArray<TEnumAsByte<ECollisionChannel> > & ObjectsToTrace, bool bTraceComplex, const TArray<AActor*>& ActorsToIgnore, EDrawDebugTrace::Type DrawDebugType, FHitResult& OutHit, bool bIgnoreSelf);
 
 	/**
@@ -1055,7 +1055,7 @@ class ENGINE_API UKismetSystemLibrary : public UBlueprintFunctionLibrary
 	 * @param OutHits		A list of hits, sorted along the trace from start to finish.  The blocking hit will be the last hit, if there was one.
 	 * @return				True if there was a hit, false otherwise.
 	 */
-	UFUNCTION(BlueprintCallable, Category="Collision", meta=(DeprecatedFunction, DeprecationMessage = "Use new MultiSphereTraceForObjects", bIgnoreSelf="true", WorldContext="WorldContextObject", AutoCreateRefTerm="ActorsToIgnore", FriendlyName = "MultiSphereTraceByObjectDeprecated"))
+	UFUNCTION(BlueprintCallable, Category="Collision", meta=(DeprecatedFunction, DeprecationMessage = "Use new MultiSphereTraceForObjects", bIgnoreSelf="true", WorldContext="WorldContextObject", AutoCreateRefTerm="ActorsToIgnore", DisplayName = "MultiSphereTraceByObjectDeprecated"))
 	static bool SphereTraceMultiByObject_DEPRECATED(UObject* WorldContextObject, const FVector Start, const FVector End, float Radius, const TArray<TEnumAsByte<ECollisionChannel> > & ObjectsToTrace, bool bTraceComplex, const TArray<AActor*>& ActorsToIgnore, EDrawDebugTrace::Type DrawDebugType, TArray<FHitResult>& OutHits, bool bIgnoreSelf);
 	
 	/**
@@ -1072,7 +1072,7 @@ class ENGINE_API UKismetSystemLibrary : public UBlueprintFunctionLibrary
 	 * @param OutHit		Properties of the trace hit.
 	 * @return				True if there was a hit, false otherwise.
 	 */
-	UFUNCTION(BlueprintCallable, Category="Collision", meta=(DeprecatedFunction, DeprecationMessage = "Use new SingleCapsuleTraceForObjects", bIgnoreSelf="true", WorldContext="WorldContextObject", AutoCreateRefTerm="ActorsToIgnore", FriendlyName = "SingleCapsuleTraceByObjectDeprecated"))
+	UFUNCTION(BlueprintCallable, Category="Collision", meta=(DeprecatedFunction, DeprecationMessage = "Use new SingleCapsuleTraceForObjects", bIgnoreSelf="true", WorldContext="WorldContextObject", AutoCreateRefTerm="ActorsToIgnore", DisplayName = "SingleCapsuleTraceByObjectDeprecated"))
 	static bool CapsuleTraceSingleByObject_DEPRECATED(UObject* WorldContextObject, const FVector Start, const FVector End, float Radius, float HalfHeight, const TArray<TEnumAsByte<ECollisionChannel> > & ObjectsToTrace, bool bTraceComplex, const TArray<AActor*>& ActorsToIgnore, EDrawDebugTrace::Type DrawDebugType, FHitResult& OutHit, bool bIgnoreSelf);
 
 	/**
@@ -1089,7 +1089,7 @@ class ENGINE_API UKismetSystemLibrary : public UBlueprintFunctionLibrary
 	 * @param OutHits		A list of hits, sorted along the trace from start to finish.  The blocking hit will be the last hit, if there was one.
 	 * @return				True if there was a hit, false otherwise.
 	 */
-	UFUNCTION(BlueprintCallable, Category="Collision", meta=(DeprecatedFunction, DeprecationMessage = "Use new MultiCapsuleTraceForObjects", bIgnoreSelf="true", WorldContext="WorldContextObject", AutoCreateRefTerm="ActorsToIgnore", FriendlyName = "MultiCapsuleTraceByObjectDeprecated"))
+	UFUNCTION(BlueprintCallable, Category="Collision", meta=(DeprecatedFunction, DeprecationMessage = "Use new MultiCapsuleTraceForObjects", bIgnoreSelf="true", WorldContext="WorldContextObject", AutoCreateRefTerm="ActorsToIgnore", DisplayName = "MultiCapsuleTraceByObjectDeprecated"))
 	static bool CapsuleTraceMultiByObject_DEPRECATED(UObject* WorldContextObject, const FVector Start, const FVector End, float Radius, float HalfHeight, const TArray<TEnumAsByte<ECollisionChannel> > & ObjectsToTrace, bool bTraceComplex, const TArray<AActor*>& ActorsToIgnore, EDrawDebugTrace::Type DrawDebugType, TArray<FHitResult>& OutHits, bool bIgnoreSelf);
 
 	/**
@@ -1140,7 +1140,7 @@ class ENGINE_API UKismetSystemLibrary : public UBlueprintFunctionLibrary
 	 * Draw a debug cone 
 	 * Angles are specified in degrees
 	 */
-	UFUNCTION(BlueprintCallable, Category="Rendering|Debug", meta=(WorldContext="WorldContextObject", FriendlyName="DrawDebugCone"))
+	UFUNCTION(BlueprintCallable, Category="Rendering|Debug", meta=(WorldContext="WorldContextObject", DisplayName="DrawDebugCone"))
 	static void DrawDebugConeInDegrees(UObject* WorldContextObject, const FVector Origin, const FVector Direction, float Length=100.f, float AngleWidth=45.f, float AngleHeight=45.f, int32 NumSides = 12, FLinearColor LineColor = FLinearColor::White, float Duration=0.f);
 
 	/** Draw a debug capsule */
diff --git a/Engine/Source/Runtime/Engine/Classes/Kismet/KismetTextLibrary.h b/Engine/Source/Runtime/Engine/Classes/Kismet/KismetTextLibrary.h
--- a/Engine/Source/Runtime/Engine/Classes/Kismet/KismetTextLibrary.h
+++ b/Engine/Source/Runtime/Engine/Classes/Kismet/KismetTextLibrary.h
@@ -46,15 +46,15 @@ class ENGINE_API UKismetTextLibrary : public UBlueprintFunctionLibrary
 	GENERATED_UCLASS_BODY()
 
 	/** Converts localizable text to the string */
-	UFUNCTION(BlueprintPure, meta=(FriendlyName = "ToString (text)", CompactNodeTitle = "->"), Category="Utilities|String")
+	UFUNCTION(BlueprintPure, meta=(DisplayName = "ToString (text)", CompactNodeTitle = "->"), Category="Utilities|String")
 	static FString Conv_TextToString(const FText& InText);
 
 	/** Converts string to localizable text */
-	UFUNCTION(BlueprintPure, meta=(FriendlyName = "ToText (string)", CompactNodeTitle = "->"), Category="Utilities|Text")
+	UFUNCTION(BlueprintPure, meta=(DisplayName = "ToText (string)", CompactNodeTitle = "->"), Category="Utilities|Text")
 	static FText Conv_StringToText(const FString& InString);
 
 	/** Converts string to localizable text */
-	UFUNCTION(BlueprintPure, meta=(FriendlyName = "ToText (name)", CompactNodeTitle = "->"), Category="Utilities|Text")
+	UFUNCTION(BlueprintPure, meta=(DisplayName = "ToText (name)", CompactNodeTitle = "->"), Category="Utilities|Text")
 	static FText Conv_NameToText(FName InName);
 
 	/* Returns true if text is empty. */
@@ -90,68 +90,68 @@ class ENGINE_API UKismetTextLibrary : public UBlueprintFunctionLibrary
 	static bool FindTextInLocalizationTable(const FString& Namespace, const FString& Key, FText& OutText);
 
 	/* Returns true if A and B are linguistically equal (A == B). */
-	UFUNCTION(BlueprintPure, meta=(FriendlyName = "Equal (text)", CompactNodeTitle = "=="), Category="Utilities|Text")
+	UFUNCTION(BlueprintPure, meta=(DisplayName = "Equal (text)", CompactNodeTitle = "=="), Category="Utilities|Text")
 	static bool EqualEqual_TextText(const FText& A, const FText& B);
 
 	/* Returns true if A and B are linguistically equal (A == B), ignoring case. */
-	UFUNCTION(BlueprintPure, meta=(FriendlyName = "Equal, Case Insensitive (text)", CompactNodeTitle = "=="), Category="Utilities|Text")
+	UFUNCTION(BlueprintPure, meta=(DisplayName = "Equal, Case Insensitive (text)", CompactNodeTitle = "=="), Category="Utilities|Text")
 	static bool EqualEqual_IgnoreCase_TextText(const FText& A, const FText& B);
 				
 	/* Returns true if A and B are linguistically not equal (A != B). */
-	UFUNCTION(BlueprintPure, meta=(FriendlyName = "NotEqual (text)", CompactNodeTitle = "!="), Category="Utilities|Text")
+	UFUNCTION(BlueprintPure, meta=(DisplayName = "NotEqual (text)", CompactNodeTitle = "!="), Category="Utilities|Text")
 	static bool NotEqual_TextText(const FText& A, const FText& B);
 
 	/* Returns true if A and B are linguistically not equal (A != B), ignoring case. */
-	UFUNCTION(BlueprintPure, meta=(FriendlyName = "NotEqual, Case Insensitive (text)", CompactNodeTitle = "!="), Category="Utilities|Text")
+	UFUNCTION(BlueprintPure, meta=(DisplayName = "NotEqual, Case Insensitive (text)", CompactNodeTitle = "!="), Category="Utilities|Text")
 	static bool NotEqual_IgnoreCase_TextText(const FText& A, const FText& B);
 
 	/** Converts a boolean value to text, either 'true' or 'false' */
-	UFUNCTION(BlueprintPure, meta=(FriendlyName = "ToText (boolean)", CompactNodeTitle = "->"), Category="Utilities|Text")
+	UFUNCTION(BlueprintPure, meta=(DisplayName = "ToText (boolean)", CompactNodeTitle = "->"), Category="Utilities|Text")
 	static FText Conv_BoolToText(bool InBool);
 
 	/** Converts a byte value to text */
-	UFUNCTION(BlueprintPure, meta=(FriendlyName = "ToText (byte)", CompactNodeTitle = "->"), Category="Utilities|Text")
+	UFUNCTION(BlueprintPure, meta=(DisplayName = "ToText (byte)", CompactNodeTitle = "->"), Category="Utilities|Text")
 	static FText Conv_ByteToText(uint8 Value);
 
 	// Default values are duplicated from FNumberFormattingOptions and should be replicated in all functions and in the struct when changed!
 	/* Converts a passed in integer to a text based on formatting options */
-	UFUNCTION(BlueprintPure, meta=(FriendlyName = "ToText (int)", AdvancedDisplay = "1"), Category="Utilities|Text")
+	UFUNCTION(BlueprintPure, meta=(DisplayName = "ToText (int)", AdvancedDisplay = "1"), Category="Utilities|Text")
 	static FText Conv_IntToText(int32 Value, bool bUseGrouping = true, int32 MinimumIntegralDigits = 1, int32 MaximumIntegralDigits = 324);
 
 	// Default values are duplicated from FNumberFormattingOptions and should be replicated in all functions and in the struct when changed!
 	/* Converts a passed in float to a text based on formatting options */
-	UFUNCTION(BlueprintPure, meta=(FriendlyName = "ToText (float)", AdvancedDisplay = "1"), Category="Utilities|Text")
+	UFUNCTION(BlueprintPure, meta=(DisplayName = "ToText (float)", AdvancedDisplay = "1"), Category="Utilities|Text")
 	static FText Conv_FloatToText(float Value, TEnumAsByte<ERoundingMode> RoundingMode, bool bUseGrouping = true, int32 MinimumIntegralDigits = 1, int32 MaximumIntegralDigits = 324, int32 MinimumFractionalDigits = 0, int32 MaximumFractionalDigits = 3);
 
 	// Default values are duplicated from FNumberFormattingOptions and should be replicated in all functions and in the struct when changed!
 	/* Converts a passed in integer to a text formatted as a currency */
-	UFUNCTION(BlueprintPure, meta=(FriendlyName = "AsCurrency (int)", AdvancedDisplay = "1"), Category="Utilities|Text")
+	UFUNCTION(BlueprintPure, meta=(DisplayName = "AsCurrency (int)", AdvancedDisplay = "1"), Category="Utilities|Text")
 	static FText AsCurrency_Integer(int32 Value, TEnumAsByte<ERoundingMode> RoundingMode, bool bUseGrouping = true, int32 MinimumIntegralDigits = 1, int32 MaximumIntegralDigits = 324, int32 MinimumFractionalDigits = 0, int32 MaximumFractionalDigits = 3, const FString& CurrencyCode = TEXT(""));
 
 	// Default values are duplicated from FNumberFormattingOptions and should be replicated in all functions and in the struct when changed!
 	/* Converts a passed in float to a text formatted as a currency */
-	UFUNCTION(BlueprintPure, meta=(FriendlyName = "AsCurrency (float)", AdvancedDisplay = "1"), Category="Utilities|Text")
+	UFUNCTION(BlueprintPure, meta=(DisplayName = "AsCurrency (float)", AdvancedDisplay = "1"), Category="Utilities|Text")
 	static FText AsCurrency_Float(float Value, TEnumAsByte<ERoundingMode> RoundingMode, bool bUseGrouping = true, int32 MinimumIntegralDigits = 1, int32 MaximumIntegralDigits = 324, int32 MinimumFractionalDigits = 0, int32 MaximumFractionalDigits = 3, const FString& CurrencyCode = TEXT(""));
 
 	// Default values are duplicated from FNumberFormattingOptions and should be replicated in all functions and in the struct when changed!
 	/* Converts a passed in float to a text, formatted as a percent */
-	UFUNCTION(BlueprintPure, meta=(FriendlyName = "AsPercent", AdvancedDisplay = "1"), Category="Utilities|Text")
+	UFUNCTION(BlueprintPure, meta=(DisplayName = "AsPercent", AdvancedDisplay = "1"), Category="Utilities|Text")
 	static FText AsPercent_Float(float Value, TEnumAsByte<ERoundingMode> RoundingMode, bool bUseGrouping = true, int32 MinimumIntegralDigits = 1, int32 MaximumIntegralDigits = 324, int32 MinimumFractionalDigits = 0, int32 MaximumFractionalDigits = 3);
 
 	/* Converts a passed in date & time to a text, formatted as a date */
-	UFUNCTION(BlueprintPure, meta=(FriendlyName = "AsDate", AdvancedDisplay = "1"), Category="Utilities|Text")
+	UFUNCTION(BlueprintPure, meta=(DisplayName = "AsDate", AdvancedDisplay = "1"), Category="Utilities|Text")
 	static FText AsDate_DateTime(const FDateTime& InDateTime);
 
 	/* Converts a passed in date & time to a text, formatted as a date & time */
-	UFUNCTION(BlueprintPure, meta=(FriendlyName = "AsDateTime", AdvancedDisplay = "1"), Category="Utilities|Text")
+	UFUNCTION(BlueprintPure, meta=(DisplayName = "AsDateTime", AdvancedDisplay = "1"), Category="Utilities|Text")
 	static FText AsDateTime_DateTime(const FDateTime& In);
 
 	/* Converts a passed in date & time to a text, formatted as a time */
-	UFUNCTION(BlueprintPure, meta=(FriendlyName = "AsTime", AdvancedDisplay = "1"), Category="Utilities|Text")
+	UFUNCTION(BlueprintPure, meta=(DisplayName = "AsTime", AdvancedDisplay = "1"), Category="Utilities|Text")
 	static FText AsTime_DateTime(const FDateTime& In);
 
 	/* Converts a passed in time span to a text, formatted as a time span */
-	UFUNCTION(BlueprintPure, meta=(FriendlyName = "AsTimespan", AdvancedDisplay = "1"), Category="Utilities|Text")
+	UFUNCTION(BlueprintPure, meta=(DisplayName = "AsTimespan", AdvancedDisplay = "1"), Category="Utilities|Text")
 	static FText AsTimespan_Timespan(const FTimespan& InTimespan);
 
 	/* Used for formatting text using the FText::Format function and utilized by the UK2Node_FormatText */
diff --git a/Engine/Source/Runtime/Engine/Classes/Materials/MaterialInstanceDynamic.h b/Engine/Source/Runtime/Engine/Classes/Materials/MaterialInstanceDynamic.h
--- a/Engine/Source/Runtime/Engine/Classes/Materials/MaterialInstanceDynamic.h
+++ b/Engine/Source/Runtime/Engine/Classes/Materials/MaterialInstanceDynamic.h
@@ -15,7 +15,7 @@ class UMaterialInstanceDynamic : public UMaterialInstance
 	ENGINE_API void SetScalarParameterValue(FName ParameterName, float Value);
 
 	/** Get the current scalar (float) parameter value from an MID */
-	UFUNCTION(BlueprintCallable, meta=(FriendlyName = "GetScalarParameterValue", Keywords = "GetFloatParameterValue"), Category="Rendering|Material")
+	UFUNCTION(BlueprintCallable, meta=(DisplayName = "GetScalarParameterValue", Keywords = "GetFloatParameterValue"), Category="Rendering|Material")
 	float K2_GetScalarParameterValue(FName ParameterName);
 
 	/** Set an MID texture parameter value */
@@ -23,7 +23,7 @@ class UMaterialInstanceDynamic : public UMaterialInstance
 	ENGINE_API void SetTextureParameterValue(FName ParameterName, class UTexture* Value);
 
 	/** Get the current MID texture parameter value */
-	UFUNCTION(BlueprintCallable, meta=(FriendlyName = "GetTextureParameterValue"), Category="Rendering|Material")
+	UFUNCTION(BlueprintCallable, meta=(DisplayName = "GetTextureParameterValue"), Category="Rendering|Material")
 	class UTexture* K2_GetTextureParameterValue(FName ParameterName);
 
 	/** Set an MID vector parameter value */
@@ -31,17 +31,17 @@ class UMaterialInstanceDynamic : public UMaterialInstance
 	ENGINE_API void SetVectorParameterValue(FName ParameterName, FLinearColor Value);
 
 	/** Get the current MID vector parameter value */
-	UFUNCTION(BlueprintCallable, meta=(FriendlyName = "GetVectorParameterValue", Keywords = "GetColorParameterValue"), Category="Rendering|Material")
+	UFUNCTION(BlueprintCallable, meta=(DisplayName = "GetVectorParameterValue", Keywords = "GetColorParameterValue"), Category="Rendering|Material")
 	FLinearColor K2_GetVectorParameterValue(FName ParameterName);
 	
 	/**
 	 * Interpolates the scalar and vector parameters of this material instance based on two other material instances, and an alpha blending factor
 	 */
-	UFUNCTION(BlueprintCallable, meta=(FriendlyName = "InterpolateMaterialInstanceParameters"), Category="Rendering|Material")
+	UFUNCTION(BlueprintCallable, meta=(DisplayName = "InterpolateMaterialInstanceParameters"), Category="Rendering|Material")
 	void K2_InterpolateMaterialInstanceParams(UMaterialInstance* MaterialA, UMaterialInstance* MaterialB, float Alpha);
 
 	/** Copies over parameters given a material interface */
-	UFUNCTION(BlueprintCallable, meta=(FriendlyName = "CopyMaterialInstanceParameters"), Category="Rendering|Material")
+	UFUNCTION(BlueprintCallable, meta=(DisplayName = "CopyMaterialInstanceParameters"), Category="Rendering|Material")
 	void K2_CopyMaterialInstanceParameters(UMaterialInterface* SourceMaterialToCopyFrom);
 
 	/**
@@ -64,7 +64,7 @@ class UMaterialInstanceDynamic : public UMaterialInstance
 	 * Copy parameter values from another material instance. This will copy only
 	 * parameters explicitly overridden in that material instance!!
 	 */
-	UFUNCTION(BlueprintCallable, meta=(FriendlyName = "CopyParameterOverrides"), Category="Rendering|Material")
+	UFUNCTION(BlueprintCallable, meta=(DisplayName = "CopyParameterOverrides"), Category="Rendering|Material")
 	ENGINE_API void CopyParameterOverrides(UMaterialInstance* MaterialInstance);
 
 };
diff --git a/Engine/Source/Runtime/Engine/Classes/Vehicles/WheeledVehicleMovementComponent4W.h b/Engine/Source/Runtime/Engine/Classes/Vehicles/WheeledVehicleMovementComponent4W.h
--- a/Engine/Source/Runtime/Engine/Classes/Vehicles/WheeledVehicleMovementComponent4W.h
+++ b/Engine/Source/Runtime/Engine/Classes/Vehicles/WheeledVehicleMovementComponent4W.h
@@ -113,7 +113,7 @@ struct FVehicleTransmissionData
 {
 	GENERATED_USTRUCT_BODY()
 	/** Whether to use automatic transmission */
-	UPROPERTY(EditAnywhere, Category = VehicleSetup, meta=(FriendlyName = "Automatic Transmission"))
+	UPROPERTY(EditAnywhere, Category = VehicleSetup, meta=(DisplayName = "Automatic Transmission"))
 	bool bUseGearAutoBox;
 
 	/** Time it takes to switch gears (seconds) */
diff --git a/Engine/Source/Runtime/Engine/Classes/VisualLogger/VisualLoggerKismetLibrary.h b/Engine/Source/Runtime/Engine/Classes/VisualLogger/VisualLoggerKismetLibrary.h
--- a/Engine/Source/Runtime/Engine/Classes/VisualLogger/VisualLoggerKismetLibrary.h
+++ b/Engine/Source/Runtime/Engine/Classes/VisualLogger/VisualLoggerKismetLibrary.h
@@ -10,14 +10,14 @@ class UVisualLoggerKismetLibrary : public UBlueprintFunctionLibrary
 	GENERATED_UCLASS_BODY()
 
 	/** Logs simple text string with Visual Logger - recording for Visual Logs has to be enabled to record this data */
-	UFUNCTION(BlueprintCallable, Category = "Debug|VisualLogger", meta = (HidePin = "WorldContextObject", DefaultToSelf = "WorldContextObject", FriendlyName = "LogText"))
+	UFUNCTION(BlueprintCallable, Category = "Debug|VisualLogger", meta = (HidePin = "WorldContextObject", DefaultToSelf = "WorldContextObject", DisplayName = "LogText"))
 	static void LogText(UObject* WorldContextObject, FString Text, FName LogCategory = TEXT("Blueprint Log"));
 
 	/** Logs location as sphere with given radius - recording for Visual Logs has to be enabled to record this data */
-	UFUNCTION(BlueprintCallable, Category = "Debug|VisualLogger", meta = (HidePin = "WorldContextObject", DefaultToSelf = "WorldContextObject", FriendlyName = "LogLocation"))
+	UFUNCTION(BlueprintCallable, Category = "Debug|VisualLogger", meta = (HidePin = "WorldContextObject", DefaultToSelf = "WorldContextObject", DisplayName = "LogLocation"))
 	static void LogLocation(UObject* WorldContextObject, FVector Location, FString Text, FLinearColor ObjectColor = FLinearColor::White, float Radius = 10, FName LogCategory = TEXT("Blueprint Log"));
 
 	/** Logs box shape - recording for Visual Logs has to be enabled to record this data */
-	UFUNCTION(BlueprintCallable, Category = "Debug|VisualLogger", meta = (HidePin = "WorldContextObject", DefaultToSelf = "WorldContextObject", FriendlyName = "Log Box Shape"))
+	UFUNCTION(BlueprintCallable, Category = "Debug|VisualLogger", meta = (HidePin = "WorldContextObject", DefaultToSelf = "WorldContextObject", DisplayName = "Log Box Shape"))
 	static void LogBox(UObject* WorldContextObject, FBox BoxShape, FString Text, FLinearColor ObjectColor = FLinearColor::White, FName LogCategory = TEXT("Blueprint Log"));
 };
diff --git a/Engine/Source/Runtime/GameplayAbilities/Public/Abilities/GameplayAbility.h b/Engine/Source/Runtime/GameplayAbilities/Public/Abilities/GameplayAbility.h
--- a/Engine/Source/Runtime/GameplayAbilities/Public/Abilities/GameplayAbility.h
+++ b/Engine/Source/Runtime/GameplayAbilities/Public/Abilities/GameplayAbility.h
@@ -275,7 +275,7 @@ protected:
 	// --------------------------------------
 
 	/** Returns true if this ability can be activated right now. Has no side effects */
-	UFUNCTION(BlueprintImplementableEvent, Category = Ability, FriendlyName = "ShouldAbilityRespondToEvent")
+	UFUNCTION(BlueprintImplementableEvent, Category = Ability, DisplayName = "ShouldAbilityRespondToEvent")
 	bool K2_ShouldAbilityRespondToEvent(FGameplayAbilityActorInfo ActorInfo, FGameplayEventData Payload) const;
 
 	bool bHasBlueprintShouldAbilityRespondToEvent;
@@ -285,7 +285,7 @@ protected:
 	// --------------------------------------
 	
 	/** Returns true if this ability can be activated right now. Has no side effects */
-	UFUNCTION(BlueprintImplementableEvent, Category = Ability, FriendlyName="CanActivateAbility")
+	UFUNCTION(BlueprintImplementableEvent, Category = Ability, DisplayName="CanActivateAbility")
 	bool K2_CanActivateAbility(FGameplayAbilityActorInfo ActorInfo, FGameplayTagContainer& RelevantTags) const;
 
 	bool bHasBlueprintCanUse;
@@ -305,10 +305,10 @@ protected:
 	 *  will only occur when latent/async actions are pending. When K2_ActivateAbility logically finishes, then we will expect Commit/End to have been called.
 	 *  
 	 */
-	UFUNCTION(BlueprintImplementableEvent, Category = Ability, FriendlyName = "ActivateAbility")
+	UFUNCTION(BlueprintImplementableEvent, Category = Ability, DisplayName = "ActivateAbility")
 	void K2_ActivateAbility();
 
-	UFUNCTION(BlueprintImplementableEvent, Category = Ability, FriendlyName = "ActivateAbilityFromEvent")
+	UFUNCTION(BlueprintImplementableEvent, Category = Ability, DisplayName = "ActivateAbilityFromEvent")
 	void K2_ActivateAbilityFromEvent(const FGameplayEventData& EventData);
 
 	virtual void ActivateAbility(const FGameplayAbilitySpecHandle Handle, const FGameplayAbilityActorInfo* ActorInfo, const FGameplayAbilityActivationInfo ActivationInfo, const FGameplayEventData* TriggerEventData);
@@ -334,23 +334,23 @@ public:
 	 * Attempts to commit the ability (spend resources, etc). This our last chance to fail.
 	 *	-Child classes that override ActivateAbility must call this themselves!
 	 */
-	UFUNCTION(BlueprintCallable, Category = Ability, FriendlyName = "CommitAbility")
+	UFUNCTION(BlueprintCallable, Category = Ability, DisplayName = "CommitAbility")
 	virtual bool K2_CommitAbility();
 
 	/** Attempts to commit the ability's cooldown only. If BroadcastCommitEvent is true, it will broadcast the commit event that tasks like WaitAbilityCommit are listening for. */
-	UFUNCTION(BlueprintCallable, Category = Ability, FriendlyName = "CommitAbilityCooldown")
+	UFUNCTION(BlueprintCallable, Category = Ability, DisplayName = "CommitAbilityCooldown")
 	virtual bool K2_CommitAbilityCooldown(bool BroadcastCommitEvent=false);
 
 	/** Attempts to commit the ability's cost only. If BroadcastCommitEvent is true, it will broadcast the commit event that tasks like WaitAbilityCommit are listening for. */
-	UFUNCTION(BlueprintCallable, Category = Ability, FriendlyName = "CommitAbilityCost")
+	UFUNCTION(BlueprintCallable, Category = Ability, DisplayName = "CommitAbilityCost")
 	virtual bool K2_CommitAbilityCost(bool BroadcastCommitEvent=false);
 
 	/** Checks the ability's cooldown, but does not apply it.*/
-	UFUNCTION(BlueprintCallable, Category = Ability, FriendlyName = "CheckAbilityCooldown")
+	UFUNCTION(BlueprintCallable, Category = Ability, DisplayName = "CheckAbilityCooldown")
 	virtual bool K2_CheckAbilityCooldown();
 
 	/** Checks the ability's cost, but does not apply it. */
-	UFUNCTION(BlueprintCallable, Category = Ability, FriendlyName = "CheckAbilityCost")
+	UFUNCTION(BlueprintCallable, Category = Ability, DisplayName = "CheckAbilityCost")
 	virtual bool K2_CheckAbilityCost();
 
 	virtual bool CommitAbility(const FGameplayAbilitySpecHandle Handle, const FGameplayAbilityActorInfo* ActorInfo, const FGameplayAbilityActivationInfo ActivationInfo);
@@ -367,7 +367,7 @@ public:
 	//	CommitExecute
 	// --------------------------------------
 
-	UFUNCTION(BlueprintImplementableEvent, Category = Ability, FriendlyName="CommitExecute")
+	UFUNCTION(BlueprintImplementableEvent, Category = Ability, DisplayName="CommitExecute")
 	void K2_CommitExecute();
 
 	/** Does the commit atomically (consume resources, do cooldowns, etc) */
@@ -403,11 +403,11 @@ protected:
 	// -------------------------------------
 
 	/** Call from kismet to end the ability naturally */
-	UFUNCTION(BlueprintCallable, Category = Ability, FriendlyName="EndAbility")
+	UFUNCTION(BlueprintCallable, Category = Ability, DisplayName="EndAbility")
 	virtual void K2_EndAbility();
 
 	/** Kismet event, will be called if an ability ends normally or abnormally */
-	UFUNCTION(BlueprintImplementableEvent, Category = Ability, FriendlyName = "OnEndAbility")
+	UFUNCTION(BlueprintImplementableEvent, Category = Ability, DisplayName = "OnEndAbility")
 	void K2_OnEndAbility();
 
 	/** Native function, called if an ability ends normally or abnormally. If bReplicate is set to true, try to replicate the ending to the client/server */
@@ -422,16 +422,16 @@ protected:
 	// Apply Self
 	// ---------
 
-	UFUNCTION(BlueprintCallable, Category = Ability, FriendlyName="ApplyGameplayEffectToOwner")
+	UFUNCTION(BlueprintCallable, Category = Ability, DisplayName="ApplyGameplayEffectToOwner")
 	FActiveGameplayEffectHandle BP_ApplyGameplayEffectToOwner(TSubclassOf<UGameplayEffect> GameplayEffectClass, int32 GameplayEffectLevel = 1);
 
-	UFUNCTION(BlueprintCallable, Category = Ability, FriendlyName="ApplyGameplayEffectToOwner", meta=(DeprecatedFunction, DeprecationMessage = "Use new ApplyGameplayEffectToOwner"))
+	UFUNCTION(BlueprintCallable, Category = Ability, DisplayName="ApplyGameplayEffectToOwner", meta=(DeprecatedFunction, DeprecationMessage = "Use new ApplyGameplayEffectToOwner"))
 	FActiveGameplayEffectHandle K2_ApplyGameplayEffectToOwner(const UGameplayEffect* GameplayEffect, int32 GameplayEffectLevel = 1);
 
 	/** Non blueprintcallable, safe to call on CDO/NonInstance abilities */
 	FActiveGameplayEffectHandle ApplyGameplayEffectToOwner(const FGameplayAbilitySpecHandle Handle, const FGameplayAbilityActorInfo* ActorInfo, const FGameplayAbilityActivationInfo ActivationInfo, const UGameplayEffect* GameplayEffect, float GameplayEffectLevel) const;
 
-	UFUNCTION(BlueprintCallable, Category = Ability, FriendlyName = "ApplyGameplayEffectSpecToOwner")
+	UFUNCTION(BlueprintCallable, Category = Ability, DisplayName = "ApplyGameplayEffectSpecToOwner")
 	FActiveGameplayEffectHandle K2_ApplyGameplayEffectSpecToOwner(const FGameplayEffectSpecHandle EffectSpecHandle);
 
 	FActiveGameplayEffectHandle ApplyGameplayEffectSpecToOwner(const FGameplayAbilitySpecHandle AbilityHandle, const FGameplayAbilityActorInfo* ActorInfo, const FGameplayAbilityActivationInfo ActivationInfo, const FGameplayEffectSpecHandle SpecHandle) const;
@@ -440,16 +440,16 @@ protected:
 	// Apply Target
 	// ---------
 
-	UFUNCTION(BlueprintCallable, Category = Ability, FriendlyName = "ApplyGameplayEffectToTarget")
+	UFUNCTION(BlueprintCallable, Category = Ability, DisplayName = "ApplyGameplayEffectToTarget")
 	TArray<FActiveGameplayEffectHandle> BP_ApplyGameplayEffectToTarget(FGameplayAbilityTargetDataHandle TargetData, TSubclassOf<UGameplayEffect> GameplayEffectClass, int32 GameplayEffectLevel = 1);
 
-	UFUNCTION(BlueprintCallable, Category = Ability, FriendlyName = "ApplyGameplayEffectToTarget", meta=(DeprecatedFunction, DeprecationMessage = "Use new ApplyGameplayEffectToTarget"))
+	UFUNCTION(BlueprintCallable, Category = Ability, DisplayName = "ApplyGameplayEffectToTarget", meta=(DeprecatedFunction, DeprecationMessage = "Use new ApplyGameplayEffectToTarget"))
 	TArray<FActiveGameplayEffectHandle> K2_ApplyGameplayEffectToTarget(FGameplayAbilityTargetDataHandle TargetData, const UGameplayEffect* GameplayEffect, int32 GameplayEffectLevel = 1);
 
 	/** Non blueprintcallable, safe to call on CDO/NonInstance abilities */
 	TArray<FActiveGameplayEffectHandle> ApplyGameplayEffectToTarget(const FGameplayAbilitySpecHandle Handle, const FGameplayAbilityActorInfo* ActorInfo, const FGameplayAbilityActivationInfo ActivationInfo, FGameplayAbilityTargetDataHandle Target, TSubclassOf<UGameplayEffect> GameplayEffectClass, float GameplayEffectLevel) const;
 
-	UFUNCTION(BlueprintCallable, Category = Ability, FriendlyName = "ApplyGameplayEffectSpecToTarget")
+	UFUNCTION(BlueprintCallable, Category = Ability, DisplayName = "ApplyGameplayEffectSpecToTarget")
 	TArray<FActiveGameplayEffectHandle> K2_ApplyGameplayEffectSpecToTarget(const FGameplayEffectSpecHandle EffectSpecHandle, FGameplayAbilityTargetDataHandle TargetData);
 
 	TArray<FActiveGameplayEffectHandle> ApplyGameplayEffectSpecToTarget(const FGameplayAbilitySpecHandle AbilityHandle, const FGameplayAbilityActorInfo* ActorInfo, const FGameplayAbilityActivationInfo ActivationInfo, const FGameplayEffectSpecHandle SpecHandle, FGameplayAbilityTargetDataHandle TargetData) const;
@@ -459,10 +459,10 @@ protected:
 	// ---------
 	
 	/** Removes GAmeplayEffects from owner which match the given asset level tags. */
-	UFUNCTION(BlueprintCallable, Category = Ability, FriendlyName="RemoveGameplayEffectFromOwnerWithAssetTags")
+	UFUNCTION(BlueprintCallable, Category = Ability, DisplayName="RemoveGameplayEffectFromOwnerWithAssetTags")
 	void BP_RemoveGameplayEffectFromOwnerWithAssetTags(FGameplayTagContainer WithAssetTags, int32 StacksToRemove = -1);
 
-	UFUNCTION(BlueprintCallable, Category = Ability, FriendlyName="RemoveGameplayEffectFromOwnerWithGrantedTags")
+	UFUNCTION(BlueprintCallable, Category = Ability, DisplayName="RemoveGameplayEffectFromOwnerWithGrantedTags")
 	void BP_RemoveGameplayEffectFromOwnerWithGrantedTags(FGameplayTagContainer WithGrantedTags, int32 StacksToRemove = -1);
 
 	// -------------------------------------
@@ -470,13 +470,13 @@ protected:
 	//	Abilities can invoke GameplayCues without having to create GameplayEffects
 	// -------------------------------------
 	
-	UFUNCTION(BlueprintCallable, Category = Ability, meta=(GameplayTagFilter="GameplayCue"), FriendlyName="ExecuteGameplayCue")
+	UFUNCTION(BlueprintCallable, Category = Ability, meta=(GameplayTagFilter="GameplayCue"), DisplayName="ExecuteGameplayCue")
 	virtual void K2_ExecuteGameplayCue(FGameplayTag GameplayCueTag, FGameplayEffectContextHandle Context);
 
-	UFUNCTION(BlueprintCallable, Category = Ability, meta=(GameplayTagFilter="GameplayCue"), FriendlyName="AddGameplayCue")
+	UFUNCTION(BlueprintCallable, Category = Ability, meta=(GameplayTagFilter="GameplayCue"), DisplayName="AddGameplayCue")
 	virtual void K2_AddGameplayCue(FGameplayTag GameplayCueTag, FGameplayEffectContextHandle Context, bool bRemoveOnAbilityEnd = true);
 
-	UFUNCTION(BlueprintCallable, Category = Ability, meta=(GameplayTagFilter="GameplayCue"), FriendlyName="RemoveGameplayCue")
+	UFUNCTION(BlueprintCallable, Category = Ability, meta=(GameplayTagFilter="GameplayCue"), DisplayName="RemoveGameplayCue")
 	virtual void K2_RemoveGameplayCue(FGameplayTag GameplayCueTag);
 
 	/** Generates a GameplayEffectContextHandle from our owner and an optional TargetData.*/
@@ -532,7 +532,7 @@ public:
 	AActor* GetAvatarActorFromActorInfo() const;
 
 	/** Convenience method for abilities to get skeletal mesh component - useful for aiming abilities */
-	UFUNCTION(BlueprintCallable, FriendlyName = "GetSkeletalMeshComponentFromActorInfo", Category = Ability)
+	UFUNCTION(BlueprintCallable, DisplayName = "GetSkeletalMeshComponentFromActorInfo", Category = Ability)
 	USkeletalMeshComponent* GetOwningComponentFromActorInfo() const;
 
 	/** Convenience method for abilities to get outgoing gameplay effect specs (for example, to pass on to projectiles to apply to whoever they hit) */
diff --git a/Engine/Source/Runtime/GameplayAbilities/Public/Abilities/Tasks/AbilityTask_PlayMontageAndWait.h b/Engine/Source/Runtime/GameplayAbilities/Public/Abilities/Tasks/AbilityTask_PlayMontageAndWait.h
--- a/Engine/Source/Runtime/GameplayAbilities/Public/Abilities/Tasks/AbilityTask_PlayMontageAndWait.h
+++ b/Engine/Source/Runtime/GameplayAbilities/Public/Abilities/Tasks/AbilityTask_PlayMontageAndWait.h
@@ -25,7 +25,7 @@ class UAbilityTask_PlayMontageAndWait : public UAbilityTask
 	UFUNCTION()
 	void OnMontageInterrupted();
 
-	UFUNCTION(BlueprintCallable, Category="Ability|Tasks", meta = (FriendlyName="PlayMontageAndWait",
+	UFUNCTION(BlueprintCallable, Category="Ability|Tasks", meta = (DisplayName="PlayMontageAndWait",
 		HidePin = "WorldContextObject", DefaultToSelf = "WorldContextObject", BlueprintInternalUseOnly = "TRUE"))
 	static UAbilityTask_PlayMontageAndWait* CreatePlayMontageAndWaitProxy(UObject* WorldContextObject,
 		FName TaskInstanceName, UAnimMontage *MontageToPlay, float Rate = 1.f, FName StartSection = NAME_None);
diff --git a/Engine/Source/Runtime/GameplayAbilities/Public/Abilities/Tasks/AbilityTask_WaitAbilityCommit.h b/Engine/Source/Runtime/GameplayAbilities/Public/Abilities/Tasks/AbilityTask_WaitAbilityCommit.h
--- a/Engine/Source/Runtime/GameplayAbilities/Public/Abilities/Tasks/AbilityTask_WaitAbilityCommit.h
+++ b/Engine/Source/Runtime/GameplayAbilities/Public/Abilities/Tasks/AbilityTask_WaitAbilityCommit.h
@@ -26,7 +26,7 @@ class UAbilityTask_WaitAbilityCommit : public UAbilityTask
 	void OnAbilityCommit(UGameplayAbility *ActivatedAbility);
 
 	/** Wait until a new ability (of the same or different type) is commited. Used to gracefully interrupt abilities in specific ways. */
-	UFUNCTION(BlueprintCallable, Category="Ability|Tasks", meta = (HidePin = "WorldContextObject", DefaultToSelf = "WorldContextObject", BlueprintInternalUseOnly = "TRUE", FriendlyName = "Wait For New Ability Commit"))
+	UFUNCTION(BlueprintCallable, Category="Ability|Tasks", meta = (HidePin = "WorldContextObject", DefaultToSelf = "WorldContextObject", BlueprintInternalUseOnly = "TRUE", DisplayName = "Wait For New Ability Commit"))
 	static UAbilityTask_WaitAbilityCommit* WaitForAbilityCommit(UObject* WorldContextObject, FGameplayTag WithTag, FGameplayTag WithoutTage);
 
 	FGameplayTag WithTag;
diff --git a/Engine/Source/Runtime/GameplayAbilities/Public/Abilities/Tasks/AbilityTask_WaitConfirmCancel.h b/Engine/Source/Runtime/GameplayAbilities/Public/Abilities/Tasks/AbilityTask_WaitConfirmCancel.h
--- a/Engine/Source/Runtime/GameplayAbilities/Public/Abilities/Tasks/AbilityTask_WaitConfirmCancel.h
+++ b/Engine/Source/Runtime/GameplayAbilities/Public/Abilities/Tasks/AbilityTask_WaitConfirmCancel.h
@@ -34,7 +34,7 @@ class UAbilityTask_WaitConfirmCancel : public UAbilityTask
 	UFUNCTION()
 	void OnLocalCancelCallback();
 
-	UFUNCTION(BlueprintCallable, meta=(HidePin = "WorldContextObject", DefaultToSelf = "WorldContextObject", BlueprintInternalUseOnly = "true", FriendlyName="Wait for Confirm Input"), Category="Ability|Tasks")
+	UFUNCTION(BlueprintCallable, meta=(HidePin = "WorldContextObject", DefaultToSelf = "WorldContextObject", BlueprintInternalUseOnly = "true", DisplayName="Wait for Confirm Input"), Category="Ability|Tasks")
 	static UAbilityTask_WaitConfirmCancel* WaitConfirmCancel(UObject* WorldContextObject);	
 
 	virtual void Activate() override;
diff --git a/Engine/Source/Runtime/GameplayAbilities/Public/Abilities/Tasks/AbilityTask_WaitMovementModeChange.h b/Engine/Source/Runtime/GameplayAbilities/Public/Abilities/Tasks/AbilityTask_WaitMovementModeChange.h
--- a/Engine/Source/Runtime/GameplayAbilities/Public/Abilities/Tasks/AbilityTask_WaitMovementModeChange.h
+++ b/Engine/Source/Runtime/GameplayAbilities/Public/Abilities/Tasks/AbilityTask_WaitMovementModeChange.h
@@ -23,7 +23,7 @@ public:
 	EMovementMode	RequiredMode;
 
 	/** Wait until movement mode changes (E.g., landing) */
-	UFUNCTION(BlueprintCallable, Category="Ability|Tasks", meta = (HidePin = "WorldContextObject", DefaultToSelf = "WorldContextObject", BlueprintInternalUseOnly = "TRUE", FriendlyName="WaitMovementModeChange"))
+	UFUNCTION(BlueprintCallable, Category="Ability|Tasks", meta = (HidePin = "WorldContextObject", DefaultToSelf = "WorldContextObject", BlueprintInternalUseOnly = "TRUE", DisplayName="WaitMovementModeChange"))
 	static UAbilityTask_WaitMovementModeChange* CreateWaitMovementModeChange(UObject* WorldContextObject, EMovementMode NewMode);
 
 	virtual void Activate() override;
diff --git a/Engine/Source/Runtime/GameplayAbilities/Public/Abilities/Tasks/AbilityTask_WaitVelocityChange.h b/Engine/Source/Runtime/GameplayAbilities/Public/Abilities/Tasks/AbilityTask_WaitVelocityChange.h
--- a/Engine/Source/Runtime/GameplayAbilities/Public/Abilities/Tasks/AbilityTask_WaitVelocityChange.h
+++ b/Engine/Source/Runtime/GameplayAbilities/Public/Abilities/Tasks/AbilityTask_WaitVelocityChange.h
@@ -18,7 +18,7 @@ class UAbilityTask_WaitVelocityChange: public UAbilityTask
 	virtual void TickTask(float DeltaTime) override;
 
 	/** Wait for the actor's movement component velocity to be of minimum magnitude when projected along given direction */
-	UFUNCTION(BlueprintCallable, Category = "Ability|Tasks", meta = (FriendlyName="WaitVelocityChange",HidePin = "WorldContextObject", DefaultToSelf = "WorldContextObject", BlueprintInternalUseOnly = "TRUE"))
+	UFUNCTION(BlueprintCallable, Category = "Ability|Tasks", meta = (DisplayName="WaitVelocityChange",HidePin = "WorldContextObject", DefaultToSelf = "WorldContextObject", BlueprintInternalUseOnly = "TRUE"))
 	static UAbilityTask_WaitVelocityChange* CreateWaitVelocityChange(UObject* WorldContextObject, FVector Direction, float MinimumMagnitude);
 		
 	virtual void Activate() override;
diff --git a/Engine/Source/Runtime/GameplayAbilities/Public/AbilitySystemBlueprintLibrary.h b/Engine/Source/Runtime/GameplayAbilities/Public/AbilitySystemBlueprintLibrary.h
--- a/Engine/Source/Runtime/GameplayAbilities/Public/AbilitySystemBlueprintLibrary.h
+++ b/Engine/Source/Runtime/GameplayAbilities/Public/AbilitySystemBlueprintLibrary.h
@@ -102,29 +102,29 @@ class GAMEPLAYABILITIES_API UAbilitySystemBlueprintLibrary : public UBlueprintFu
 	//		GameplayEffectContext
 	// -------------------------------------------------------------------------------
 
-	UFUNCTION(BlueprintPure, Category = "Ability|EffectContext", Meta = (FriendlyName = "IsInstigatorLocallyControlled"))
+	UFUNCTION(BlueprintPure, Category = "Ability|EffectContext", Meta = (DisplayName = "IsInstigatorLocallyControlled"))
 	static bool EffectContextIsInstigatorLocallyControlled(FGameplayEffectContextHandle EffectContext);
 
-	UFUNCTION(BlueprintPure, Category = "Ability|EffectContext", Meta = (FriendlyName = "GetHitResult"))
+	UFUNCTION(BlueprintPure, Category = "Ability|EffectContext", Meta = (DisplayName = "GetHitResult"))
 	static FHitResult EffectContextGetHitResult(FGameplayEffectContextHandle EffectContext);
 
-	UFUNCTION(BlueprintPure, Category = "Ability|EffectContext", Meta = (FriendlyName = "HasHitResult"))
+	UFUNCTION(BlueprintPure, Category = "Ability|EffectContext", Meta = (DisplayName = "HasHitResult"))
 	static bool EffectContextHasHitResult(FGameplayEffectContextHandle EffectContext);
 
 	/** Gets the location the effect originated from */
-	UFUNCTION(BlueprintPure, Category = "Ability|EffectContext", Meta = (FriendlyName = "GetOrigin"))
+	UFUNCTION(BlueprintPure, Category = "Ability|EffectContext", Meta = (DisplayName = "GetOrigin"))
 	static FVector EffectContextGetOrigin(FGameplayEffectContextHandle EffectContext);
 
 	/** Gets the instigating actor (that holds the ability system component) of the EffectContext */
-	UFUNCTION(BlueprintPure, Category = "Ability|EffectContext", Meta = (FriendlyName = "GetInstigatorActor"))
+	UFUNCTION(BlueprintPure, Category = "Ability|EffectContext", Meta = (DisplayName = "GetInstigatorActor"))
 	static AActor* EffectContextGetInstigatorActor(FGameplayEffectContextHandle EffectContext);
 
 	/** Gets the original instigator actor that started the chain of events to cause this effect */
-	UFUNCTION(BlueprintPure, Category = "Ability|EffectContext", Meta = (FriendlyName = "GetOriginalInstigatorActor"))
+	UFUNCTION(BlueprintPure, Category = "Ability|EffectContext", Meta = (DisplayName = "GetOriginalInstigatorActor"))
 	static AActor* EffectContextGetOriginalInstigatorActor(FGameplayEffectContextHandle EffectContext);
 
 	/** Gets the physical actor that caused the effect, possibly a projectile or weapon */
-	UFUNCTION(BlueprintPure, Category = "Ability|EffectContext", Meta = (FriendlyName = "GetEffectCauser"))
+	UFUNCTION(BlueprintPure, Category = "Ability|EffectContext", Meta = (DisplayName = "GetEffectCauser"))
 	static AActor* EffectContextGetEffectCauser(FGameplayEffectContextHandle EffectContext);
 
 	// -------------------------------------------------------------------------------
diff --git a/Engine/Source/Runtime/GameplayAbilities/Public/AbilitySystemComponent.h b/Engine/Source/Runtime/GameplayAbilities/Public/AbilitySystemComponent.h
--- a/Engine/Source/Runtime/GameplayAbilities/Public/AbilitySystemComponent.h
+++ b/Engine/Source/Runtime/GameplayAbilities/Public/AbilitySystemComponent.h
@@ -114,7 +114,7 @@ class GAMEPLAYABILITIES_API UAbilitySystemComponent : public UActorComponent, pu
 
 	const UAttributeSet* InitStats(TSubclassOf<class UAttributeSet> Attributes, const UDataTable* DataTable);
 
-	UFUNCTION(BlueprintCallable, Category="Skills", meta=(FriendlyName="InitStats"))
+	UFUNCTION(BlueprintCallable, Category="Skills", meta=(DisplayName="InitStats"))
 	void K2_InitStats(TSubclassOf<class UAttributeSet> Attributes, const UDataTable* DataTable);
 		
 
@@ -406,18 +406,18 @@ class GAMEPLAYABILITIES_API UAbilitySystemComponent : public UActorComponent, pu
 	FOnActiveGameplayEffectRemoved* OnGameplayEffectRemovedDelegate(FActiveGameplayEffectHandle Handle);
 	FOnActiveGameplayEffectRemoved& OnAnyGameplayEffectRemovedDelegate();
 
-	UFUNCTION(BlueprintCallable, Category = GameplayEffects, meta=(FriendlyName = "ApplyGameplayEffectToTarget"))
+	UFUNCTION(BlueprintCallable, Category = GameplayEffects, meta=(DisplayName = "ApplyGameplayEffectToTarget"))
 	FActiveGameplayEffectHandle BP_ApplyGameplayEffectToTarget(TSubclassOf<UGameplayEffect> GameplayEffectClass, UAbilitySystemComponent *Target, float Level, FGameplayEffectContextHandle Context);
 
-	UFUNCTION(BlueprintCallable, Category = GameplayEffects, meta=(FriendlyName = "ApplyGameplayEffectToTarget"), meta=(DeprecatedFunction, DeprecationMessage = "Use new ApplyGameplayEffectToSelf"))
+	UFUNCTION(BlueprintCallable, Category = GameplayEffects, meta=(DisplayName = "ApplyGameplayEffectToTarget"), meta=(DeprecatedFunction, DeprecationMessage = "Use new ApplyGameplayEffectToSelf"))
 	FActiveGameplayEffectHandle K2_ApplyGameplayEffectToTarget(UGameplayEffect *GameplayEffect, UAbilitySystemComponent *Target, float Level, FGameplayEffectContextHandle Context);
 
 	FActiveGameplayEffectHandle ApplyGameplayEffectToTarget(UGameplayEffect *GameplayEffect, UAbilitySystemComponent *Target, float Level = UGameplayEffect::INVALID_LEVEL, FGameplayEffectContextHandle Context = FGameplayEffectContextHandle(), FPredictionKey PredictionKey = FPredictionKey());
 
-	UFUNCTION(BlueprintCallable, Category = GameplayEffects, meta=(FriendlyName = "ApplyGameplayEffectToSelf"))
+	UFUNCTION(BlueprintCallable, Category = GameplayEffects, meta=(DisplayName = "ApplyGameplayEffectToSelf"))
 	FActiveGameplayEffectHandle BP_ApplyGameplayEffectToSelf(TSubclassOf<UGameplayEffect> GameplayEffectClass, float Level, FGameplayEffectContextHandle EffectContext);
 
-	UFUNCTION(BlueprintCallable, Category = GameplayEffects, meta=(FriendlyName = "ApplyGameplayEffectToSelf"), meta=(DeprecatedFunction, DeprecationMessage = "Use new ApplyGameplayEffectToSelf"))
+	UFUNCTION(BlueprintCallable, Category = GameplayEffects, meta=(DisplayName = "ApplyGameplayEffectToSelf"), meta=(DeprecatedFunction, DeprecationMessage = "Use new ApplyGameplayEffectToSelf"))
 	FActiveGameplayEffectHandle K2_ApplyGameplayEffectToSelf(const UGameplayEffect *GameplayEffect, float Level, FGameplayEffectContextHandle EffectContext);
 	
 	FActiveGameplayEffectHandle ApplyGameplayEffectToSelf(const UGameplayEffect *GameplayEffect, float Level, const FGameplayEffectContextHandle& EffectContext, FPredictionKey PredictionKey = FPredictionKey());
diff --git a/Engine/Source/Runtime/GameplayAbilities/Public/GameplayCueNotify_Actor.h b/Engine/Source/Runtime/GameplayAbilities/Public/GameplayCueNotify_Actor.h
--- a/Engine/Source/Runtime/GameplayAbilities/Public/GameplayCueNotify_Actor.h
+++ b/Engine/Source/Runtime/GameplayAbilities/Public/GameplayCueNotify_Actor.h
@@ -40,7 +40,7 @@ class GAMEPLAYABILITIES_API AGameplayCueNotify_Actor : public AActor
 #endif // WITH_EDITOR
 
 	/** Generic Event Graph event that will get called for every event type */
-	UFUNCTION(BlueprintImplementableEvent, Category = "GameplayCueNotify", FriendlyName = "HandleGameplayCue")
+	UFUNCTION(BlueprintImplementableEvent, Category = "GameplayCueNotify", DisplayName = "HandleGameplayCue")
 	void K2_HandleGameplayCue(TWeakObjectPtr<AActor> MyTarget, EGameplayCueEvent::Type EventType, FGameplayCueParameters Parameters);
 
 	UFUNCTION(BlueprintNativeEvent, Category = "GameplayCueNotify")
diff --git a/Engine/Source/Runtime/GameplayAbilities/Public/GameplayCueNotify_Static.h b/Engine/Source/Runtime/GameplayAbilities/Public/GameplayCueNotify_Static.h
--- a/Engine/Source/Runtime/GameplayAbilities/Public/GameplayCueNotify_Static.h
+++ b/Engine/Source/Runtime/GameplayAbilities/Public/GameplayCueNotify_Static.h
@@ -40,7 +40,7 @@ class GAMEPLAYABILITIES_API UGameplayCueNotify_Static : public UObject
 #endif // WITH_EDITOR
 
 	/** Generic Event Graph event that will get called for every event type */
-	UFUNCTION(BlueprintImplementableEvent, Category = "GameplayCueNotify", FriendlyName = "HandleGameplayCue")
+	UFUNCTION(BlueprintImplementableEvent, Category = "GameplayCueNotify", DisplayName = "HandleGameplayCue")
 	void K2_HandleGameplayCue(TWeakObjectPtr<AActor> MyTarget, EGameplayCueEvent::Type EventType, FGameplayCueParameters Parameters) const;
 
 	UFUNCTION(BlueprintNativeEvent, BlueprintPure, Category = "GameplayCueNotify")
diff --git a/Engine/Source/Runtime/Online/OnlineSubsystemUtils/Classes/InAppPurchaseCallbackProxy.h b/Engine/Source/Runtime/Online/OnlineSubsystemUtils/Classes/InAppPurchaseCallbackProxy.h
--- a/Engine/Source/Runtime/Online/OnlineSubsystemUtils/Classes/InAppPurchaseCallbackProxy.h
+++ b/Engine/Source/Runtime/Online/OnlineSubsystemUtils/Classes/InAppPurchaseCallbackProxy.h
@@ -20,7 +20,7 @@ class UInAppPurchaseCallbackProxy : public UObject
 	FInAppPurchaseResult OnFailure;
 
 	// Kicks off a transaction for the provided product identifier
-	UFUNCTION(BlueprintCallable, meta = (FriendlyName="Make an In-App Purchase"), Category="Online|InAppPurchase")
+	UFUNCTION(BlueprintCallable, meta = (DisplayName="Make an In-App Purchase"), Category="Online|InAppPurchase")
 	static UInAppPurchaseCallbackProxy* CreateProxyObjectForInAppPurchase(class APlayerController* PlayerController, const FInAppPurchaseProductRequest& ProductRequest);
 
 public:
diff --git a/Engine/Source/Runtime/Online/OnlineSubsystemUtils/Classes/InAppPurchaseQueryCallbackProxy.h b/Engine/Source/Runtime/Online/OnlineSubsystemUtils/Classes/InAppPurchaseQueryCallbackProxy.h
--- a/Engine/Source/Runtime/Online/OnlineSubsystemUtils/Classes/InAppPurchaseQueryCallbackProxy.h
+++ b/Engine/Source/Runtime/Online/OnlineSubsystemUtils/Classes/InAppPurchaseQueryCallbackProxy.h
@@ -20,7 +20,7 @@ class UInAppPurchaseQueryCallbackProxy : public UObject
 	FInAppPurchaseQueryResult OnFailure;
 
 	// Queries a InAppPurchase for an integer value
-	UFUNCTION(BlueprintCallable, meta = (FriendlyName="Read In App Purchase Information"), Category="Online|InAppPurchase")
+	UFUNCTION(BlueprintCallable, meta = (DisplayName="Read In App Purchase Information"), Category="Online|InAppPurchase")
 	static UInAppPurchaseQueryCallbackProxy* CreateProxyObjectForInAppPurchaseQuery(class APlayerController* PlayerController, const TArray<FString>& ProductIdentifiers);
 
 public:
diff --git a/Engine/Source/Runtime/Online/OnlineSubsystemUtils/Classes/LeaderboardQueryCallbackProxy.h b/Engine/Source/Runtime/Online/OnlineSubsystemUtils/Classes/LeaderboardQueryCallbackProxy.h
--- a/Engine/Source/Runtime/Online/OnlineSubsystemUtils/Classes/LeaderboardQueryCallbackProxy.h
+++ b/Engine/Source/Runtime/Online/OnlineSubsystemUtils/Classes/LeaderboardQueryCallbackProxy.h
@@ -20,7 +20,7 @@ class ULeaderboardQueryCallbackProxy : public UObject
 	FLeaderboardQueryResult OnFailure;
 
 	// Queries a leaderboard for an integer value
-	UFUNCTION(BlueprintCallable, meta = (BlueprintInternalUseOnly = "true", FriendlyName="Read Leaderboard Integer"), Category="Online|Leaderboard")
+	UFUNCTION(BlueprintCallable, meta = (BlueprintInternalUseOnly = "true", DisplayName="Read Leaderboard Integer"), Category="Online|Leaderboard")
 	static ULeaderboardQueryCallbackProxy* CreateProxyObjectForIntQuery(class APlayerController* PlayerController, FName StatName);
 
 public:
diff --git a/Engine/Source/Runtime/UMG/Public/Blueprint/WidgetBlueprintLibrary.h b/Engine/Source/Runtime/UMG/Public/Blueprint/WidgetBlueprintLibrary.h
--- a/Engine/Source/Runtime/UMG/Public/Blueprint/WidgetBlueprintLibrary.h
+++ b/Engine/Source/Runtime/UMG/Public/Blueprint/WidgetBlueprintLibrary.h
@@ -18,7 +18,7 @@ class UWidgetBlueprintLibrary : public UBlueprintFunctionLibrary
 public:
 
 	/** Creates a widget */
-	UFUNCTION(BlueprintCallable, BlueprintCosmetic, meta=( WorldContext="WorldContextObject", FriendlyName="Create Widget", BlueprintInternalUseOnly="true" ), Category="Widget")
+	UFUNCTION(BlueprintCallable, BlueprintCosmetic, meta=( WorldContext="WorldContextObject", DisplayName="Create Widget", BlueprintInternalUseOnly="true" ), Category="Widget")
 	static class UUserWidget* Create(UObject* WorldContextObject, TSubclassOf<class UUserWidget> WidgetType, APlayerController* OwningPlayer);
 
 	/**
diff --git a/Engine/Source/Runtime/UMG/Public/Components/EditableText.h b/Engine/Source/Runtime/UMG/Public/Components/EditableText.h
--- a/Engine/Source/Runtime/UMG/Public/Components/EditableText.h
+++ b/Engine/Source/Runtime/UMG/Public/Components/EditableText.h
@@ -116,7 +116,7 @@ public:
 	 * Gets the widget text
 	 * @return The widget text
 	 */
-	UFUNCTION(BlueprintCallable, Category = "Widget", meta=(FriendlyName="GetText (Editable Text)"))
+	UFUNCTION(BlueprintCallable, Category = "Widget", meta=(DisplayName="GetText (Editable Text)"))
 	FText GetText() const;
 
 	/**
@@ -124,7 +124,7 @@ public:
 	 * Warning: This will wipe any binding created for the Text property!
 	 * @param InText The text to assign to the widget
 	 */
-	UFUNCTION(BlueprintCallable, Category = "Widget", meta=(FriendlyName="SetText (Editable Text)"))
+	UFUNCTION(BlueprintCallable, Category = "Widget", meta=(DisplayName="SetText (Editable Text)"))
 	void SetText(FText InText);
 
 	UFUNCTION(BlueprintCallable, Category = "Widget")
diff --git a/Engine/Source/Runtime/UMG/Public/Components/EditableTextBox.h b/Engine/Source/Runtime/UMG/Public/Components/EditableTextBox.h
--- a/Engine/Source/Runtime/UMG/Public/Components/EditableTextBox.h
+++ b/Engine/Source/Runtime/UMG/Public/Components/EditableTextBox.h
@@ -115,14 +115,14 @@ public:
 public:
 
 	/**  */
-	UFUNCTION(BlueprintCallable, Category="Widget", meta=(FriendlyName="GetText (Text Box)"))
+	UFUNCTION(BlueprintCallable, Category="Widget", meta=(DisplayName="GetText (Text Box)"))
 	FText GetText() const;
 
 	/**  */
-	UFUNCTION(BlueprintCallable, Category="Widget", meta=(FriendlyName="SetText (Text Box)"))
+	UFUNCTION(BlueprintCallable, Category="Widget", meta=(DisplayName="SetText (Text Box)"))
 	void SetText(FText InText);
 
-	UFUNCTION(BlueprintCallable, Category="Widget",  meta=(FriendlyName="SetError (Text Box)"))
+	UFUNCTION(BlueprintCallable, Category="Widget",  meta=(DisplayName="SetError (Text Box)"))
 	void SetError(FText InError);
 
 	UFUNCTION(BlueprintCallable, Category="Widget")
diff --git a/Engine/Source/Runtime/UMG/Public/Components/MultiLineEditableText.h b/Engine/Source/Runtime/UMG/Public/Components/MultiLineEditableText.h
--- a/Engine/Source/Runtime/UMG/Public/Components/MultiLineEditableText.h
+++ b/Engine/Source/Runtime/UMG/Public/Components/MultiLineEditableText.h
@@ -55,11 +55,11 @@ public:
 public:
 
 	/**  */
-	UFUNCTION(BlueprintCallable, Category="Widget", meta=(FriendlyName="GetText (Multi-Line Editable Text)"))
+	UFUNCTION(BlueprintCallable, Category="Widget", meta=(DisplayName="GetText (Multi-Line Editable Text)"))
 	FText GetText() const;
 
 	/**  */
-	UFUNCTION(BlueprintCallable, Category="Widget", meta=(FriendlyName="SetText (Multi-Line Editable Text)"))
+	UFUNCTION(BlueprintCallable, Category="Widget", meta=(DisplayName="SetText (Multi-Line Editable Text)"))
 	void SetText(FText InText);
 	
 	// UWidget interface
