Collect object references from object annotations

diff --git a/Engine/Source/Editor/UnrealEd/Private/EditorTransaction.cpp b/Engine/Source/Editor/UnrealEd/Private/EditorTransaction.cpp
--- a/Engine/Source/Editor/UnrealEd/Private/EditorTransaction.cpp
+++ b/Engine/Source/Editor/UnrealEd/Private/EditorTransaction.cpp
@@ -197,6 +197,10 @@ void FTransaction::FObjectRecord::AddReferencedObjects( FReferenceCollector& Col
 	{
 		ReferencedObject.AddReferencedObjects(Collector);
 	}
+	if (ObjectAnnotation.IsValid())
+	{
+		ObjectAnnotation->AddReferencedObjects(Collector);
+	}
 }
 
 void FTransaction::AddReferencedObjects( FReferenceCollector& Collector )
diff --git a/Engine/Source/Runtime/Core/Public/Misc/ITransaction.h b/Engine/Source/Runtime/Core/Public/Misc/ITransaction.h
--- a/Engine/Source/Runtime/Core/Public/Misc/ITransaction.h
+++ b/Engine/Source/Runtime/Core/Public/Misc/ITransaction.h
@@ -17,7 +17,11 @@ typedef void(*STRUCT_DTOR)( void* TPtr );					// destruct
  * on the UObject that a modification was performed on, but it does not see other changes that may have
  * to be remembered in order to properly restore the object internals.
  */
-class ITransactionObjectAnnotation { };
+class ITransactionObjectAnnotation 
+{ 
+public:
+	virtual void AddReferencedObjects(class FReferenceCollector& Collector) = 0;
+};
 
 
 /**
diff --git a/Engine/Source/Runtime/Engine/Classes/Components/SceneComponent.h b/Engine/Source/Runtime/Engine/Classes/Components/SceneComponent.h
--- a/Engine/Source/Runtime/Engine/Classes/Components/SceneComponent.h
+++ b/Engine/Source/Runtime/Engine/Classes/Components/SceneComponent.h
@@ -833,6 +833,7 @@ public:
 
 	virtual void ApplyToComponent(UActorComponent* Component, const ECacheApplyPhase CacheApplyPhase) override;
 	virtual void FindAndReplaceInstances(const TMap<UObject*, UObject*>& OldToNewInstanceMap) override;
+	virtual void AddReferencedObjects(FReferenceCollector& Collector) override;
 
 	TArray<USceneComponent*> AttachedInstanceComponents;
 };
diff --git a/Engine/Source/Runtime/Engine/Classes/GameFramework/Actor.h b/Engine/Source/Runtime/Engine/Classes/GameFramework/Actor.h
--- a/Engine/Source/Runtime/Engine/Classes/GameFramework/Actor.h
+++ b/Engine/Source/Runtime/Engine/Classes/GameFramework/Actor.h
@@ -1187,6 +1187,8 @@ public:
 	public:
 		FActorTransactionAnnotation(const AActor* Actor);
 
+		virtual void AddReferencedObjects(FReferenceCollector& Collector) override;
+
 		bool HasInstanceData() const;
 
 		FComponentInstanceDataCache ComponentInstanceData;
diff --git a/Engine/Source/Runtime/Engine/Private/ActorEditor.cpp b/Engine/Source/Runtime/Engine/Private/ActorEditor.cpp
--- a/Engine/Source/Runtime/Engine/Private/ActorEditor.cpp
+++ b/Engine/Source/Runtime/Engine/Private/ActorEditor.cpp
@@ -299,6 +299,11 @@ AActor::FActorTransactionAnnotation::FActorTransactionAnnotation(const AActor* A
 	}
 }
 
+void AActor::FActorTransactionAnnotation::AddReferencedObjects(FReferenceCollector& Collector)
+{
+	ComponentInstanceData.AddReferencedObjects(Collector);
+}
+
 bool AActor::FActorTransactionAnnotation::HasInstanceData() const
 {
 	return (bRootComponentDataCached || ComponentInstanceData.HasInstanceData());
diff --git a/Engine/Source/Runtime/Engine/Private/ComponentInstanceDataCache.cpp b/Engine/Source/Runtime/Engine/Private/ComponentInstanceDataCache.cpp
--- a/Engine/Source/Runtime/Engine/Private/ComponentInstanceDataCache.cpp
+++ b/Engine/Source/Runtime/Engine/Private/ComponentInstanceDataCache.cpp
@@ -126,6 +126,11 @@ void FActorComponentInstanceData::ApplyToComponent(UActorComponent* Component, c
 	}
 }
 
+void FActorComponentInstanceData::AddReferencedObjects(FReferenceCollector& Collector)
+{
+	Collector.AddReferencedObject(SourceComponentClass);
+}
+
 FComponentInstanceDataCache::FComponentInstanceDataCache(const AActor* Actor)
 {
 	if(Actor != NULL)
@@ -224,7 +229,7 @@ void FComponentInstanceDataCache::FindAndReplaceInstances(const TMap<UObject*, U
 		}
 	}
 	TArray<USceneComponent*> SceneComponents;
-	InstanceComponentTransformToRootMap.GetKeys(SceneComponents);
+	InstanceComponentTransformToRootMap.GenerateKeyArray(SceneComponents);
 
 	for (USceneComponent* SceneComponent : SceneComponents)
 	{
@@ -240,4 +245,20 @@ void FComponentInstanceDataCache::FindAndReplaceInstances(const TMap<UObject*, U
 			}
 		}
 	}
+}
+
+void FComponentInstanceDataCache::AddReferencedObjects(FReferenceCollector& Collector)
+{
+	TArray<USceneComponent*> SceneComponents;
+	InstanceComponentTransformToRootMap.GenerateKeyArray(SceneComponents);
+
+	Collector.AddReferencedObjects(SceneComponents);
+
+	for (auto ComponentInstanceDataPair : TypeToDataMap)
+	{
+		if (ComponentInstanceDataPair.Value)
+		{
+			ComponentInstanceDataPair.Value->AddReferencedObjects(Collector);
+		}
+	}
 }
\ No newline at end of file
diff --git a/Engine/Source/Runtime/Engine/Private/Components/SceneComponent.cpp b/Engine/Source/Runtime/Engine/Private/Components/SceneComponent.cpp
--- a/Engine/Source/Runtime/Engine/Private/Components/SceneComponent.cpp
+++ b/Engine/Source/Runtime/Engine/Private/Components/SceneComponent.cpp
@@ -1205,6 +1205,12 @@ void FSceneComponentInstanceData::ApplyToComponent(UActorComponent* Component, c
 	}
 }
 
+void FSceneComponentInstanceData::AddReferencedObjects(FReferenceCollector& Collector)
+{
+	FActorComponentInstanceData::AddReferencedObjects(Collector);
+	Collector.AddReferencedObjects(AttachedInstanceComponents);
+}
+
 void FSceneComponentInstanceData::FindAndReplaceInstances(const TMap<UObject*, UObject*>& OldToNewInstanceMap)
 {
 	for (USceneComponent*& ChildComponent : AttachedInstanceComponents)
diff --git a/Engine/Source/Runtime/Engine/Private/Components/StaticMeshComponent.cpp b/Engine/Source/Runtime/Engine/Private/Components/StaticMeshComponent.cpp
--- a/Engine/Source/Runtime/Engine/Private/Components/StaticMeshComponent.cpp
+++ b/Engine/Source/Runtime/Engine/Private/Components/StaticMeshComponent.cpp
@@ -39,6 +39,24 @@ public:
 		CastChecked<UStaticMeshComponent>(Component)->ApplyComponentInstanceData(this);
 	}
 
+	virtual void AddReferencedObjects(FReferenceCollector& Collector) override
+	{
+		FSceneComponentInstanceData::AddReferencedObjects(Collector);
+
+		Collector.AddReferencedObject(StaticMesh);
+		if (bHasCachedStaticLighting)
+		{
+			for (FLightMapRef& LightMapRef : CachedStaticLighting.LODDataLightMap)
+			{
+				LightMapRef->AddReferencedObjects(Collector);
+			}
+			for (FShadowMapRef& ShadowMapRef : CachedStaticLighting.LODDataShadowMap)
+			{
+				ShadowMapRef->AddReferencedObjects(Collector);
+			}		
+		}
+	}
+
 	/** Add vertex color data for a specified LOD before RerunConstructionScripts is called */
 	void AddVertexColorData(const struct FStaticMeshComponentLODInfo& LODInfo, uint32 LODIndex)
 	{
diff --git a/Engine/Source/Runtime/Engine/Private/InstancedStaticMesh.cpp b/Engine/Source/Runtime/Engine/Private/InstancedStaticMesh.cpp
--- a/Engine/Source/Runtime/Engine/Private/InstancedStaticMesh.cpp
+++ b/Engine/Source/Runtime/Engine/Private/InstancedStaticMesh.cpp
@@ -733,6 +733,24 @@ public:
 		CastChecked<UInstancedStaticMeshComponent>(Component)->ApplyComponentInstanceData(this);
 	}
 
+	virtual void AddReferencedObjects(FReferenceCollector& Collector) override
+	{
+		FSceneComponentInstanceData::AddReferencedObjects(Collector);
+
+		Collector.AddReferencedObject(StaticMesh);
+		if (bHasCachedStaticLighting)
+		{
+			for (FLightMapRef& LightMapRef : CachedStaticLighting.LODDataLightMap)
+			{
+				LightMapRef->AddReferencedObjects(Collector);
+			}
+			for (FShadowMapRef& ShadowMapRef : CachedStaticLighting.LODDataShadowMap)
+			{
+				ShadowMapRef->AddReferencedObjects(Collector);
+			}		
+		}
+	}
+
 	/** Used to store lightmap data during RerunConstructionScripts */
 	struct FLightMapInstanceData
 	{
diff --git a/Engine/Source/Runtime/Engine/Public/ComponentInstanceDataCache.h b/Engine/Source/Runtime/Engine/Public/ComponentInstanceDataCache.h
--- a/Engine/Source/Runtime/Engine/Public/ComponentInstanceDataCache.h
+++ b/Engine/Source/Runtime/Engine/Public/ComponentInstanceDataCache.h
@@ -35,6 +35,8 @@ public:
 	/** Replaces any references to old instances during Actor reinstancing */
 	virtual void FindAndReplaceInstances(const TMap<UObject*, UObject*>& OldToNewInstanceMap) { };
 
+	virtual void AddReferencedObjects(FReferenceCollector& Collector);
+
 	bool ContainsSavedProperties() const { return SavedProperties.Num() > 0; }
 
 protected:
@@ -74,6 +76,8 @@ public:
 
 	bool HasInstanceData() const { return TypeToDataMap.Num() > 0; }
 
+	void AddReferencedObjects(FReferenceCollector& Collector);
+
 private:
 	/** Map of data type name to data of that type */
 	TMultiMap< FName, FActorComponentInstanceData* >	TypeToDataMap;
diff --git a/Engine/Source/Runtime/UMG/Private/Components/WidgetComponent.cpp b/Engine/Source/Runtime/UMG/Private/Components/WidgetComponent.cpp
--- a/Engine/Source/Runtime/UMG/Private/Components/WidgetComponent.cpp
+++ b/Engine/Source/Runtime/UMG/Private/Components/WidgetComponent.cpp
@@ -636,6 +636,14 @@ public:
 		CastChecked<UWidgetComponent>(Component)->ApplyComponentInstanceData(this);
 	}
 
+	virtual void AddReferencedObjects(FReferenceCollector& Collector) override
+	{
+		FSceneComponentInstanceData::AddReferencedObjects(Collector);
+		UClass* WidgetUClass = *WidgetClass;
+		Collector.AddReferencedObject(WidgetUClass);
+		Collector.AddReferencedObject(RenderTarget);
+	}
+
 public:
 	TSubclassOf<UUserWidget> WidgetClass;
 	UTextureRenderTarget2D* RenderTarget;
