Correct call to APEXToDebrisDepth

diff --git a/Engine/Source/Runtime/Engine/Private/DestructibleMesh.cpp b/Engine/Source/Runtime/Engine/Private/DestructibleMesh.cpp
--- a/Engine/Source/Runtime/Engine/Private/DestructibleMesh.cpp
+++ b/Engine/Source/Runtime/Engine/Private/DestructibleMesh.cpp
@@ -286,7 +286,7 @@ void FDestructibleSpecialHierarchyDepths::LoadDefaultDestructibleParametersFromA
 	verify(NxParameterized::getParamU32(*Params, "destructibleParameters.minimumFractureDepth", PMinimumFractureDepth));
 	MinimumFractureDepth = (int32)PMinimumFractureDepth;
 	verify(NxParameterized::getParamI32(*Params, "destructibleParameters.debrisDepth", DebrisDepth));
-	DebrisDepthToAPEX(bEnableDebris, DebrisDepth);
+	APEXToDebrisDepth(bEnableDebris, DebrisDepth);
 
 	physx::PxU32 PEssentialDepth;
 	verify(NxParameterized::getParamU32(*Params, "destructibleParameters.essentialDepth", PEssentialDepth));
