Removal of virtual from BlueprintImplementableEvents

diff --git a/Engine/Source/Runtime/Engine/Classes/Animation/AnimInstance.h b/Engine/Source/Runtime/Engine/Classes/Animation/AnimInstance.h
--- a/Engine/Source/Runtime/Engine/Classes/Animation/AnimInstance.h
+++ b/Engine/Source/Runtime/Engine/Classes/Animation/AnimInstance.h
@@ -444,11 +444,11 @@ public:
 
 	/** Executed when the Animation is initialized */
 	UFUNCTION(BlueprintImplementableEvent)
-	virtual void BlueprintInitializeAnimation();
+	void BlueprintInitializeAnimation();
 
 	/** Executed when the Animation is updated */
 	UFUNCTION(BlueprintImplementableEvent)
-	virtual void BlueprintUpdateAnimation(float DeltaTimeX);
+	void BlueprintUpdateAnimation(float DeltaTimeX);
 
 	bool CanTransitionSignature() const;
 	
diff --git a/Engine/Source/Runtime/Engine/Classes/Animation/AnimNotifies/AnimNotify.h b/Engine/Source/Runtime/Engine/Classes/Animation/AnimNotifies/AnimNotify.h
--- a/Engine/Source/Runtime/Engine/Classes/Animation/AnimNotifies/AnimNotify.h
+++ b/Engine/Source/Runtime/Engine/Classes/Animation/AnimNotifies/AnimNotify.h
@@ -21,7 +21,7 @@ class ENGINE_API UAnimNotify : public UObject
 	FString GetNotifyName() const;
 
 	UFUNCTION(BlueprintImplementableEvent)
-	virtual bool Received_Notify(USkeletalMeshComponent* MeshComp, UAnimSequenceBase* Animation) const;
+	bool Received_Notify(USkeletalMeshComponent* MeshComp, UAnimSequenceBase* Animation) const;
 
 #if WITH_EDITORONLY_DATA
 	/** Color of Notify in editor */
diff --git a/Engine/Source/Runtime/Engine/Classes/Animation/AnimNotifies/AnimNotifyState.h b/Engine/Source/Runtime/Engine/Classes/Animation/AnimNotifies/AnimNotifyState.h
--- a/Engine/Source/Runtime/Engine/Classes/Animation/AnimNotifies/AnimNotifyState.h
+++ b/Engine/Source/Runtime/Engine/Classes/Animation/AnimNotifies/AnimNotifyState.h
@@ -21,13 +21,13 @@ class ENGINE_API UAnimNotifyState : public UObject
 	FString GetNotifyName() const;
 
 	UFUNCTION(BlueprintImplementableEvent)
-	virtual bool Received_NotifyBegin(USkeletalMeshComponent * MeshComp, UAnimSequenceBase * Animation, float TotalDuration) const;
+	bool Received_NotifyBegin(USkeletalMeshComponent * MeshComp, UAnimSequenceBase * Animation, float TotalDuration) const;
 	
 	UFUNCTION(BlueprintImplementableEvent)
-	virtual bool Received_NotifyTick(USkeletalMeshComponent * MeshComp, UAnimSequenceBase * Animation, float FrameDeltaTime) const;
+	bool Received_NotifyTick(USkeletalMeshComponent * MeshComp, UAnimSequenceBase * Animation, float FrameDeltaTime) const;
 
 	UFUNCTION(BlueprintImplementableEvent)
-	virtual bool Received_NotifyEnd(USkeletalMeshComponent * MeshComp, UAnimSequenceBase * Animation) const;
+	bool Received_NotifyEnd(USkeletalMeshComponent * MeshComp, UAnimSequenceBase * Animation) const;
 
 #if WITH_EDITORONLY_DATA
 	/** Color of Notify in editor */
diff --git a/Engine/Source/Runtime/Engine/Classes/Animation/AnimNotifies/AnimNotifyState_Trail.h b/Engine/Source/Runtime/Engine/Classes/Animation/AnimNotifies/AnimNotifyState_Trail.h
--- a/Engine/Source/Runtime/Engine/Classes/Animation/AnimNotifies/AnimNotifyState_Trail.h
+++ b/Engine/Source/Runtime/Engine/Classes/Animation/AnimNotifies/AnimNotifyState_Trail.h
@@ -16,7 +16,7 @@ class UAnimNotifyState_Trail : public UAnimNotifyState
 	UParticleSystem* PSTemplate;
 
 	UFUNCTION(BlueprintImplementableEvent)
-	virtual UParticleSystem* OverridePSTemplate(class USkeletalMeshComponent * MeshComp, class UAnimSequenceBase * Animation) const;
+	UParticleSystem* OverridePSTemplate(class USkeletalMeshComponent * MeshComp, class UAnimSequenceBase * Animation) const;
 
 	/** Name of the first socket defining this trail. */
 	UPROPERTY(EditAnywhere, BlueprintReadWrite, Category = Trail)
