#ue4 - SpectatorPawn now replicates by default, however those spawned by PlayerController will continue not to replicate since they are client-only.

- Fixed PlayerController view relevancy checks for spectators. Added sync of spectator rotation in addition to currently synced location.
- Added ServerSetSpectatorWaiting / ClientSetSpectatorWaiting RPCs to make it easier to allow spectators to be able to join/respawn.

diff --git a/Engine/Source/Runtime/Engine/Classes/GameFramework/PlayerController.h b/Engine/Source/Runtime/Engine/Classes/GameFramework/PlayerController.h
--- a/Engine/Source/Runtime/Engine/Classes/GameFramework/PlayerController.h
+++ b/Engine/Source/Runtime/Engine/Classes/GameFramework/PlayerController.h
@@ -207,6 +207,14 @@ class ENGINE_API APlayerController : public AController
 	UPROPERTY()
 	float LastSpectatorStateSynchTime;
 
+	/** Last location synced on the server for a spectator. */
+	UPROPERTY(Transient)
+	FVector LastSpectatorSyncLocation;
+
+	/** Last rotation synced on the server for a spectator. */
+	UPROPERTY(Transient)
+	FRotator LastSpectatorSyncRotation;
+
 	/** Cap set by server on bandwidth from client to server in bytes/sec (only has impact if >=2600) */
 	UPROPERTY()
 	int32 ClientCap;
@@ -237,9 +245,18 @@ class ENGINE_API APlayerController : public AController
 	/** Whether this controller is using streaming volumes.  **/
 	uint32 bIsUsingStreamingVolumes:1;
 
-	/** Only valid in Spectating state. True if PlayerController is currently waiting for the match to start */
+	/** True if PlayerController is currently waiting for the match to start or to respawn. Only valid in Spectating state. */
+	UPROPERTY(VisibleInstanceOnly, BlueprintReadOnly, Category=PlayerController)
 	uint32 bPlayerIsWaiting:1;
 
+	/** Indicate that the Spectator is waiting to join/respawn. */
+	UFUNCTION(server, reliable, WithValidation, Category=PlayerController)
+	void ServerSetSpectatorWaiting(bool bWaiting);
+
+	/** Indicate that the Spectator is waiting to join/respawn. */
+	UFUNCTION(client, reliable, Category=PlayerController)
+	void ClientSetSpectatorWaiting(bool bWaiting);
+
 	/** index identifying players using the same base connection (splitscreen clients)
 	 * Used by netcode to match replicated PlayerControllers to the correct splitscreen viewport and child connection
 	 * replicated via special internal code, not through normal variable replication
@@ -919,9 +936,9 @@ public:
 	UFUNCTION(reliable, server, WithValidation)
 	void ServerRestartPlayer();
 
-	/** When spectating, pings the server to make sure spectating should continue. */
+	/** When spectating, updates spectator location/rotation and pings the server to make sure spectating should continue. */
 	UFUNCTION(unreliable, server, WithValidation)
-	void ServerSetSpectatorLocation(FVector NewLoc);
+	void ServerSetSpectatorLocation(FVector NewLoc, FRotator NewRot);
 
 	/** Calls ServerSetSpectatorLocation but throttles it to reduce bandwidth and only calls it when necessary. */
 	void SafeServerUpdateSpectatorState();
diff --git a/Engine/Source/Runtime/Engine/Private/PlayerController.cpp b/Engine/Source/Runtime/Engine/Private/PlayerController.cpp
--- a/Engine/Source/Runtime/Engine/Private/PlayerController.cpp
+++ b/Engine/Source/Runtime/Engine/Private/PlayerController.cpp
@@ -860,7 +860,13 @@ void APlayerController::CalcCamera(float DeltaTime, FMinimalViewInfo& OutResult)
 
 void APlayerController::GetPlayerViewPoint( FVector& out_Location, FRotator& out_Rotation ) const
 {
-	if (PlayerCameraManager != NULL && 
+	if (IsInState(NAME_Spectating) && HasAuthority() && !IsLocalController())
+	{
+		// Server uses the synced location from clients. Important for view relevancy checks.
+		out_Location = LastSpectatorSyncLocation;
+		out_Rotation = LastSpectatorSyncRotation;
+	}
+	else if (PlayerCameraManager != NULL && 
 		PlayerCameraManager->CameraCache.TimeStamp > 0.f) // Whether camera was updated at least once)
 	{
 		PlayerCameraManager->GetCameraViewPoint(out_Location, out_Rotation);
@@ -1516,6 +1522,7 @@ void APlayerController::UpdatePing(float InPing)
 void APlayerController::SetSpawnLocation(const FVector& NewLocation)
 {
 	SpawnLocation = NewLocation;
+	LastSpectatorSyncLocation = NewLocation;
 }
 
 
@@ -2542,21 +2549,23 @@ void APlayerController::SafeServerUpdateSpectatorState()
 	{
 		if (GetWorld()->TimeSince(LastSpectatorStateSynchTime) > RetryServerCheckSpectatorThrottleTime)
 		{
-			ServerSetSpectatorLocation(GetFocalLocation());
+			ServerSetSpectatorLocation(GetFocalLocation(), GetControlRotation());
 			LastSpectatorStateSynchTime = GetWorld()->TimeSeconds;
 		}
 	}
 }
 
-bool APlayerController::ServerSetSpectatorLocation_Validate(FVector NewLoc)
+bool APlayerController::ServerSetSpectatorLocation_Validate(FVector NewLoc, FRotator NewRot)
 {
 	return true;
 }
 
-void APlayerController::ServerSetSpectatorLocation_Implementation(FVector NewLoc)
+void APlayerController::ServerSetSpectatorLocation_Implementation(FVector NewLoc, FRotator NewRot)
 {
 	if ( IsInState(NAME_Spectating) )
 	{
+		LastSpectatorSyncLocation = NewLoc;
+		LastSpectatorSyncRotation = NewRot;
 		if ( GetWorld()->TimeSeconds - LastSpectatorStateSynchTime > 2.f )
 		{
 			ClientGotoState(GetStateName());
@@ -2580,6 +2589,29 @@ void APlayerController::ServerSetSpectatorLocation_Implementation(FVector NewLoc
 	}
 }
 
+
+bool APlayerController::ServerSetSpectatorWaiting_Validate(bool bWaiting)
+{
+	return true;
+}
+
+void APlayerController::ServerSetSpectatorWaiting_Implementation(bool bWaiting)
+{
+	if (IsInState(NAME_Spectating))
+	{
+		bPlayerIsWaiting = true;
+	}
+}
+
+void APlayerController::ClientSetSpectatorWaiting_Implementation(bool bWaiting)
+{
+	if (IsInState(NAME_Spectating))
+	{
+		bPlayerIsWaiting = true;
+	}
+}
+
+
 bool APlayerController::ServerViewNextPlayer_Validate()
 {
 	return true;
@@ -3937,6 +3969,15 @@ void APlayerController::SetSpectatorPawn(class ASpectatorPawn* NewSpectatorPawn)
 		SpectatorPawn = NewSpectatorPawn;
 		AttachToPawn(SpectatorPawn);
 		AddPawnTickDependency(SpectatorPawn);
+
+		if (NewSpectatorPawn)
+		{
+			AutoManageActiveCameraTarget(NewSpectatorPawn);
+		}
+		else
+		{
+			AutoManageActiveCameraTarget(this);
+		}
 	}
 }
 
@@ -3957,6 +3998,7 @@ ASpectatorPawn* APlayerController::SpawnSpectatorPawn()
 			SpawnedSpectator = GetWorld()->SpawnActor<ASpectatorPawn>(GameState->SpectatorClass, GetSpawnLocation(), GetControlRotation(), SpawnParams);
 			if (SpawnedSpectator)
 			{
+				SpawnedSpectator->SetReplicates(false); // Client-side only
 				SpawnedSpectator->PossessedBy(this);
 				SpawnedSpectator->PawnClientRestart();
 				SpawnedSpectator->SetActorTickEnabled(true);
diff --git a/Engine/Source/Runtime/Engine/Private/SpectatorPawn.cpp b/Engine/Source/Runtime/Engine/Private/SpectatorPawn.cpp
--- a/Engine/Source/Runtime/Engine/Private/SpectatorPawn.cpp
+++ b/Engine/Source/Runtime/Engine/Private/SpectatorPawn.cpp
@@ -13,8 +13,8 @@ ASpectatorPawn::ASpectatorPawn(const FObjectInitializer& ObjectInitializer)
 	)
 {
 	bCanBeDamaged = false;
-	SetRemoteRoleForBackwardsCompat(ROLE_None);
-	bReplicates = false;
+	//SetRemoteRoleForBackwardsCompat(ROLE_SimulatedProxy);
+	//bReplicates = true;
 
 	BaseEyeHeight = 0.0f;
 	bCollideWhenPlacing = false;
@@ -26,13 +26,21 @@ ASpectatorPawn::ASpectatorPawn(const FObjectInitializer& ObjectInitializer)
 
 void ASpectatorPawn::PossessedBy(class AController* NewController)
 {
-	AController* const OldController = Controller;
-	Controller = NewController;
-
-	// dispatch Blueprint event if necessary
-	if (OldController != NewController)
+	if (bReplicates)
+	{
+		Super::PossessedBy(NewController);
+	}
+	else
 	{
-		ReceivePossessed(Controller);
+		// We don't want the automatic changing of net role in Pawn code since we don't replicate, so don't call Super.
+		AController* const OldController = Controller;
+		Controller = NewController;
+
+		// dispatch Blueprint event if necessary
+		if (OldController != NewController)
+		{
+			ReceivePossessed(Controller);
+		}
 	}
 }
 
