#ue4 - Character::OnWalkingOffLedge() event now takes multiple parameters to make it easier to compute a custom velocity.

- BPs will automatically update. C++ has deprecated the old empty functions (they will not be called) in favor of the new ones.

UE-6035

diff --git a/Engine/Source/Runtime/Engine/Classes/GameFramework/Character.h b/Engine/Source/Runtime/Engine/Classes/GameFramework/Character.h
--- a/Engine/Source/Runtime/Engine/Classes/GameFramework/Character.h
+++ b/Engine/Source/Runtime/Engine/Classes/GameFramework/Character.h
@@ -516,11 +516,27 @@ public:
 
 	/**
 	 * Event fired when the Character is walking off a surface and is about to fall because CharacterMovement->CurrentFloor became unwalkable.
-	 * If CharacterMovement->MovementMode does not change (from Walking) during this event then the character will start falling.
+	 * If CharacterMovement->MovementMode does not change during this event then the character will automatically start falling afterwards.
+	 * @note Z velocity is zero during walking movement, and will be here as well. Another velocity can be computed here if desired and will be used when starting to fall.
+	 *
+	 * @param  PreviousFloorImpactNormal Normal of the previous walkable floor.
+	 * @param  PreviousFloorContactNormal Normal of the contact with the previous walkable floor.
+	 * @param  PreviousLocation	Previous character location before movement off the ledge.
+	 * @param  TimeTick	Time delta of movement update resulting in moving off the ledge.
 	 */
 	UFUNCTION(BlueprintNativeEvent, Category="Pawn|Character")
-	void OnWalkingOffLedge();
-	virtual void OnWalkingOffLedge_Implementation();
+	void OnWalkingOffLedge(const FVector& PreviousFloorImpactNormal, const FVector& PreviousFloorContactNormal, const FVector& PreviousLocation, float TimeDelta);
+	virtual void OnWalkingOffLedge_Implementation(const FVector& PreviousFloorImpactNormal, const FVector& PreviousFloorContactNormal, const FVector& PreviousLocation, float TimeDelta);
+
+
+	// Deprecated, use the new version that takes multiple parameters.
+	DEPRECATED(4.8, "OnWalkingOffLedge() is deprecated and will not be called, use the new version that takes multiple parameters.")
+	void OnWalkingOffLedge() {}
+
+	// Deprecated, use the new version that takes multiple parameters.
+	DEPRECATED(4.8, "OnWalkingOffLedge_Implementation() is deprecated and will not be called, use the new version that takes multiple parameters.")
+	virtual void OnWalkingOffLedge_Implementation() {}
+
 
 	/** Called when pawn's movement is blocked
 		@PARAM Impact describes the blocking hit. */
diff --git a/Engine/Source/Runtime/Engine/Classes/GameFramework/CharacterMovementComponent.h b/Engine/Source/Runtime/Engine/Classes/GameFramework/CharacterMovementComponent.h
--- a/Engine/Source/Runtime/Engine/Classes/GameFramework/CharacterMovementComponent.h
+++ b/Engine/Source/Runtime/Engine/Classes/GameFramework/CharacterMovementComponent.h
@@ -357,11 +357,11 @@ public:
 	UPROPERTY()
 	uint32 bShrinkProxyCapsule:1;
 
-	/** if true, Character can walk off a ledge. */
+	/** If true, Character can walk off a ledge. */
 	UPROPERTY(Category="Character Movement", EditAnywhere, BlueprintReadWrite)
 	uint32 bCanWalkOffLedges:1;
 
-	/** if true, Character can walk off a ledge when crouching. */
+	/** If true, Character can walk off a ledge when crouching. */
 	UPROPERTY(Category="Character Movement", EditAnywhere, BlueprintReadWrite)
 	uint32 bCanWalkOffLedgesWhenCrouching:1;
 
@@ -577,7 +577,7 @@ public:
 	uint32 bImpartBaseAngularVelocity:1;
 
 	/** Used by movement code to determine if a change in position is based on normal movement or a teleport. If not a teleport, velocity can be recomputed based on the change in position. */
-	UPROPERTY(Transient)
+	UPROPERTY(Category="Character Movement", Transient, VisibleInstanceOnly, BlueprintReadWrite)
 	uint32 bJustTeleported:1;
 
 	/** True when a network replication update is received for simulated proxies. */
@@ -1185,7 +1185,7 @@ public:
 	FVector GetLedgeMove(const FVector& OldLocation, const FVector& Delta, const FVector& GravDir);
 
 	/** Check if pawn is falling */
-	bool CheckFall(FHitResult &Hit, FVector Delta, FVector subLoc, float remainingTime, float timeTick, int32 Iterations, bool bMustJump);
+	bool CheckFall(const FFindFloorResult& OldFloor, const FHitResult& Hit, const FVector& Delta, const FVector& OldLocation, float remainingTime, float timeTick, int32 Iterations, bool bMustJump);
 	
 	/** 
 	 *  Revert to previous position OldLocation, return to being based on OldBase.
diff --git a/Engine/Source/Runtime/Engine/Private/Character.cpp b/Engine/Source/Runtime/Engine/Private/Character.cpp
--- a/Engine/Source/Runtime/Engine/Private/Character.cpp
+++ b/Engine/Source/Runtime/Engine/Private/Character.cpp
@@ -209,7 +209,7 @@ UActorComponent* ACharacter::FindComponentByClass(const TSubclassOf<UActorCompon
 	return Super::FindComponentByClass(ComponentClass);
 }
 
-void ACharacter::OnWalkingOffLedge_Implementation()
+void ACharacter::OnWalkingOffLedge_Implementation(const FVector& PreviousFloorImpactNormal, const FVector& PreviousFloorContactNormal, const FVector& PreviousLocation, float TimeDelta)
 {
 }
 
diff --git a/Engine/Source/Runtime/Engine/Private/Components/CharacterMovementComponent.cpp b/Engine/Source/Runtime/Engine/Private/Components/CharacterMovementComponent.cpp
--- a/Engine/Source/Runtime/Engine/Private/Components/CharacterMovementComponent.cpp
+++ b/Engine/Source/Runtime/Engine/Private/Components/CharacterMovementComponent.cpp
@@ -3364,7 +3364,7 @@ bool UCharacterMovementComponent::CanWalkOffLedges() const
 	return bCanWalkOffLedges;
 }
 
-bool UCharacterMovementComponent::CheckFall(FHitResult &Hit, FVector Delta, FVector subLoc, float remainingTime, float timeTick, int32 Iterations, bool bMustJump)
+bool UCharacterMovementComponent::CheckFall(const FFindFloorResult& OldFloor, const FHitResult& Hit, const FVector& Delta, const FVector& OldLocation, float remainingTime, float timeTick, int32 Iterations, bool bMustJump)
 {
 	if (!HasValidData())
 	{
@@ -3373,11 +3373,11 @@ bool UCharacterMovementComponent::CheckFall(FHitResult &Hit, FVector Delta, FVec
 
 	if (bMustJump || CanWalkOffLedges())
 	{
-		CharacterOwner->OnWalkingOffLedge();
+		CharacterOwner->OnWalkingOffLedge(OldFloor.HitResult.ImpactNormal, OldFloor.HitResult.Normal, OldLocation, timeTick);
 		if (IsMovingOnGround())
 		{
 			// If still walking, then fall. If not, assume the user set a different mode they want to keep.
-			StartFalling(Iterations, remainingTime, timeTick, Delta, subLoc);
+			StartFalling(Iterations, remainingTime, timeTick, Delta, OldLocation);
 		}
 		return true;
 	}
@@ -3393,7 +3393,6 @@ void UCharacterMovementComponent::StartFalling(int32 Iterations, float remaining
 					? 0.f
 					: remainingTime + timeTick * (1.f - FMath::Min(1.f,ActualDist/DesiredDist));
 
-	Velocity.Z = 0.f;			
 	if ( IsMovingOnGround() )
 	{
 		// This is to catch cases where the first frame of PIE is executed, and the
@@ -3699,7 +3698,7 @@ void UCharacterMovementComponent::PhysWalking(float deltaTime, int32 Iterations)
 				// see if it is OK to jump
 				// @todo collision : only thing that can be problem is that oldbase has world collision on
 				bool bMustJump = bZeroDelta || (OldBase == NULL || (!OldBase->IsCollisionEnabled() && MovementBaseUtility::IsDynamicBase(OldBase)));
-				if ( (bMustJump || !bCheckedFall) && CheckFall(CurrentFloor.HitResult, Delta, OldLocation, remainingTime, timeTick, Iterations, bMustJump) )
+				if ( (bMustJump || !bCheckedFall) && CheckFall(OldFloor, CurrentFloor.HitResult, Delta, OldLocation, remainingTime, timeTick, Iterations, bMustJump) )
 				{
 					return;
 				}
@@ -3718,7 +3717,7 @@ void UCharacterMovementComponent::PhysWalking(float deltaTime, int32 Iterations)
 			{
 				if (ShouldCatchAir(OldFloor, CurrentFloor))
 				{
-					CharacterOwner->OnWalkingOffLedge();
+					CharacterOwner->OnWalkingOffLedge(OldFloor.HitResult.ImpactNormal, OldFloor.HitResult.Normal, OldLocation, timeTick);
 					if (IsMovingOnGround())
 					{
 						// If still walking, then fall. If not, assume the user set a different mode they want to keep.
@@ -3751,7 +3750,7 @@ void UCharacterMovementComponent::PhysWalking(float deltaTime, int32 Iterations)
 			if (!CurrentFloor.IsWalkableFloor() && !CurrentFloor.HitResult.bStartPenetrating)
 			{
 				const bool bMustJump = bJustTeleported || bZeroDelta || (OldBase == NULL || (!OldBase->IsCollisionEnabled() && MovementBaseUtility::IsDynamicBase(OldBase)));
-				if ((bMustJump || !bCheckedFall) && CheckFall(CurrentFloor.HitResult, Delta, OldLocation, remainingTime, timeTick, Iterations, bMustJump) )
+				if ((bMustJump || !bCheckedFall) && CheckFall(OldFloor, CurrentFloor.HitResult, Delta, OldLocation, remainingTime, timeTick, Iterations, bMustJump) )
 				{
 					return;
 				}
