Fixes for use of an enum as a flag, we were doing a lot if if( foo & 0 ) in this file...

diff --git a/Engine/Source/Runtime/Engine/Private/DebugRenderSceneProxy.cpp b/Engine/Source/Runtime/Engine/Private/DebugRenderSceneProxy.cpp
--- a/Engine/Source/Runtime/Engine/Private/DebugRenderSceneProxy.cpp
+++ b/Engine/Source/Runtime/Engine/Private/DebugRenderSceneProxy.cpp
@@ -148,12 +148,12 @@ void FDebugRenderSceneProxy::GetDynamicMeshElements(const TArray<const FSceneVie
 			// Draw Cylinders
 			for(const auto& Cylinder : Cylinders)
 			{
-				if (DrawType & SolidAndWireMeshes || DrawType & WireMesh)
+				if (DrawType == SolidAndWireMeshes || DrawType == WireMesh)
 				{
-					DrawWireCylinder(PDI, Cylinder.Base, FVector(1, 0, 0), FVector(0, 1, 0), FVector(0, 0, 1), Cylinder.Color, Cylinder.Radius, Cylinder.HalfHeight, (DrawType & SolidAndWireMeshes) ? 9 : 16, SDPG_World, DrawType & SolidAndWireMeshes ? 2 : 0, 0, true);
+					DrawWireCylinder(PDI, Cylinder.Base, FVector(1, 0, 0), FVector(0, 1, 0), FVector(0, 0, 1), Cylinder.Color, Cylinder.Radius, Cylinder.HalfHeight, (DrawType == SolidAndWireMeshes) ? 9 : 16, SDPG_World, DrawType == SolidAndWireMeshes ? 2 : 0, 0, true);
 				}
 
-				if (DrawType & SolidAndWireMeshes || DrawType & SolidMesh)
+				if (DrawType == SolidAndWireMeshes || DrawType == SolidMesh)
 				{
 					GetCylinderMesh(Cylinder.Base, FVector(1, 0, 0), FVector(0, 1, 0), FVector(0, 0, 1), Cylinder.Radius, Cylinder.HalfHeight, 16, MaterialCache[Cylinder.Color.WithAlpha(DrawAlpha)], SDPG_World, ViewIndex, Collector);
 				}
@@ -162,11 +162,11 @@ void FDebugRenderSceneProxy::GetDynamicMeshElements(const TArray<const FSceneVie
 			// Draw Boxes
 			for(const auto& Box :  Boxes)
 			{
-				if (DrawType & SolidAndWireMeshes || DrawType & WireMesh)
+				if (DrawType == SolidAndWireMeshes || DrawType == WireMesh)
 				{
-					DrawWireBox(PDI, Box.Transform.ToMatrixWithScale(), Box.Box, Box.Color, SDPG_World, DrawType & SolidAndWireMeshes ? 2 : 0, 0, true);
+					DrawWireBox(PDI, Box.Transform.ToMatrixWithScale(), Box.Box, Box.Color, SDPG_World, DrawType == SolidAndWireMeshes ? 2 : 0, 0, true);
 				}
-				if (DrawType & SolidAndWireMeshes || DrawType & SolidMesh)
+				if (DrawType == SolidAndWireMeshes || DrawType == SolidMesh)
 				{
 					GetBoxMesh(FTransform(Box.Box.GetCenter()).ToMatrixNoScale() * Box.Transform.ToMatrixWithScale(), Box.Box.GetExtent(), MaterialCache[Box.Color.WithAlpha(DrawAlpha)], SDPG_World, ViewIndex, Collector);
 				}
@@ -176,11 +176,11 @@ void FDebugRenderSceneProxy::GetDynamicMeshElements(const TArray<const FSceneVie
 			TArray<FVector> Verts;
 			for (auto& CurrentCone : Cones)
 			{
-				if (DrawType & SolidAndWireMeshes || DrawType & WireMesh)
+				if (DrawType == SolidAndWireMeshes || DrawType == WireMesh)
 				{
-					DrawWireCone(PDI, Verts, CurrentCone.ConeToWorld, 1, CurrentCone.Angle2, (DrawType & SolidAndWireMeshes) ? 9 : 16, CurrentCone.Color, SDPG_World, DrawType & SolidAndWireMeshes ? 2 : 0, 0, true);
+					DrawWireCone(PDI, Verts, CurrentCone.ConeToWorld, 1, CurrentCone.Angle2, (DrawType == SolidAndWireMeshes) ? 9 : 16, CurrentCone.Color, SDPG_World, DrawType == SolidAndWireMeshes ? 2 : 0, 0, true);
 				}
-				if (DrawType & SolidAndWireMeshes || DrawType & SolidMesh)
+				if (DrawType == SolidAndWireMeshes || DrawType == SolidMesh)
 				{
 					GetConeMesh(CurrentCone.ConeToWorld, CurrentCone.Angle1, CurrentCone.Angle2, 16, MaterialCache[CurrentCone.Color.WithAlpha(DrawAlpha)], SDPG_World, ViewIndex, Collector);
 				}
@@ -190,11 +190,11 @@ void FDebugRenderSceneProxy::GetDynamicMeshElements(const TArray<const FSceneVie
 			{
 				if (PointInView(It->Location, View))
 				{
-					if (DrawType & SolidAndWireMeshes || DrawType & WireMesh)
+					if (DrawType == SolidAndWireMeshes || DrawType == WireMesh)
 					{
-						DrawWireSphere(PDI, It->Location, It->Color.WithAlpha(255), It->Radius, 20, SDPG_World, DrawType & SolidAndWireMeshes ? 2 : 0, 0, true);
+						DrawWireSphere(PDI, It->Location, It->Color.WithAlpha(255), It->Radius, 20, SDPG_World, DrawType == SolidAndWireMeshes ? 2 : 0, 0, true);
 					}
-					if (DrawType & SolidAndWireMeshes || DrawType & SolidMesh)
+					if (DrawType == SolidAndWireMeshes || DrawType == SolidMesh)
 					{
 						GetSphereMesh(It->Location, FVector(It->Radius), 20, 7, MaterialCache[It->Color.WithAlpha(DrawAlpha)], SDPG_World, false, ViewIndex, Collector);
 					}
@@ -205,16 +205,16 @@ void FDebugRenderSceneProxy::GetDynamicMeshElements(const TArray<const FSceneVie
 			{
 				if (PointInView(It->Location, View))
 				{
-					if (DrawType & SolidAndWireMeshes || DrawType & WireMesh)
+					if (DrawType == SolidAndWireMeshes || DrawType == WireMesh)
 					{
 						const float HalfAxis = FMath::Max<float>(It->HalfHeight - It->Radius, 1.f);
 						const FVector BottomEnd = It->Location + It->Radius * It->Z;
 						const FVector TopEnd = BottomEnd + (2 * HalfAxis) * It->Z;
 						const float CylinderHalfHeight = (TopEnd - BottomEnd).Size() * 0.5;
 						const FVector CylinderLocation = BottomEnd + CylinderHalfHeight * It->Z;
-						DrawWireCapsule(PDI, CylinderLocation, It->X, It->Y, It->Z, It->Color, It->Radius, It->HalfHeight, (DrawType & SolidAndWireMeshes) ? 9 : 16, SDPG_World, DrawType & SolidAndWireMeshes ? 2 : 0, 0, true);
+						DrawWireCapsule(PDI, CylinderLocation, It->X, It->Y, It->Z, It->Color, It->Radius, It->HalfHeight, (DrawType == SolidAndWireMeshes) ? 9 : 16, SDPG_World, DrawType == SolidAndWireMeshes ? 2 : 0, 0, true);
 					}
-					if (DrawType & SolidAndWireMeshes || DrawType & SolidMesh)
+					if (DrawType == SolidAndWireMeshes || DrawType == SolidMesh)
 					{
 						GetCapsuleMesh(It->Location, It->X, It->Y, It->Z, It->Color, It->Radius, It->HalfHeight, 16, MaterialCache[It->Color.WithAlpha(DrawAlpha)], SDPG_World, false, ViewIndex, Collector);
 					}
