#ue4 - BlueprintImplementableEvent ACharacter::K2_UpdateCustomMovement() is no longer virtual.

- It was never intended to be overridden in C++. That is the purpose of UCharacterMovementComponent::PhysCustom().

diff --git a/Engine/Source/Runtime/Engine/Classes/GameFramework/Character.h b/Engine/Source/Runtime/Engine/Classes/GameFramework/Character.h
--- a/Engine/Source/Runtime/Engine/Classes/GameFramework/Character.h
+++ b/Engine/Source/Runtime/Engine/Classes/GameFramework/Character.h
@@ -612,9 +612,13 @@ public:
 	UFUNCTION(BlueprintImplementableEvent, meta=(DisplayName = "OnMovementModeChanged"))
 	void K2_OnMovementModeChanged(EMovementMode PrevMovementMode, EMovementMode NewMovementMode, uint8 PrevCustomMode, uint8 NewCustomMode);
 
-	/** Event for implementing custom character movement mode. Called by CharacterMovement if MovementMode is set to Custom. */
+	/**
+	 * Event for implementing custom character movement mode. Called by CharacterMovement if MovementMode is set to Custom.
+	 * @note C++ code should override UCharacterMovementComponent::PhysCustom() instead.
+	 * @see UCharacterMovementComponent::PhysCustom()
+	 */
 	UFUNCTION(BlueprintImplementableEvent, meta=(FriendlyName= "UpdateCustomMovement"))
-	virtual void K2_UpdateCustomMovement(float DeltaTime);
+	void K2_UpdateCustomMovement(float DeltaTime);
 
 	/**
 	 * Event triggered at the end of a CharacterMovementComponent movement update.
