#ue4 - Fix deprecated FriendlyName -> DisplayName.

diff --git a/Engine/Source/Runtime/Engine/Classes/GameFramework/Character.h b/Engine/Source/Runtime/Engine/Classes/GameFramework/Character.h
--- a/Engine/Source/Runtime/Engine/Classes/GameFramework/Character.h
+++ b/Engine/Source/Runtime/Engine/Classes/GameFramework/Character.h
@@ -617,7 +617,7 @@ public:
 	 * @note C++ code should override UCharacterMovementComponent::PhysCustom() instead.
 	 * @see UCharacterMovementComponent::PhysCustom()
 	 */
-	UFUNCTION(BlueprintImplementableEvent, meta=(FriendlyName= "UpdateCustomMovement"))
+	UFUNCTION(BlueprintImplementableEvent, meta=(DisplayName = "UpdateCustomMovement"))
 	void K2_UpdateCustomMovement(float DeltaTime);
 
 	/**
