#ue4 - Fix deprecated FriendlyName -> DisplayName.

diff --git a/Engine/Source/Runtime/Engine/Classes/Components/PrimitiveComponent.h b/Engine/Source/Runtime/Engine/Classes/Components/PrimitiveComponent.h
--- a/Engine/Source/Runtime/Engine/Classes/Components/PrimitiveComponent.h
+++ b/Engine/Source/Runtime/Engine/Classes/Components/PrimitiveComponent.h
@@ -671,14 +671,14 @@ public:
 	 * Creates a Dynamic Material Instance for the specified element index.  The parent of the instance is set to the material being replaced.
 	 * @param ElementIndex - The index of the skin to replace the material for.  If invalid, the material is unchanged and NULL is returned.
 	 */
-	UFUNCTION(BlueprintCallable, meta=(FriendlyName = "CreateMIDForElement", DeprecatedFunction, DeprecationMessage="Use CreateDynamicMaterialInstance instead."), Category="Rendering|Material")
+	UFUNCTION(BlueprintCallable, meta=(DisplayName = "CreateMIDForElement", DeprecatedFunction, DeprecationMessage="Use CreateDynamicMaterialInstance instead."), Category="Rendering|Material")
 	virtual class UMaterialInstanceDynamic* CreateAndSetMaterialInstanceDynamic(int32 ElementIndex);
 
 	/**
 	 * Creates a Dynamic Material Instance for the specified element index.  The parent of the instance is set to the material being replaced.
 	 * @param ElementIndex - The index of the skin to replace the material for.  If invalid, the material is unchanged and NULL is returned.
 	 */
-	UFUNCTION(BlueprintCallable, meta=(FriendlyName = "CreateMIDForElementFromMaterial", DeprecatedFunction, DeprecationMessage="Use CreateDynamicMaterialInstance instead."), Category="Rendering|Material")
+	UFUNCTION(BlueprintCallable, meta=(DisplayName = "CreateMIDForElementFromMaterial", DeprecatedFunction, DeprecationMessage="Use CreateDynamicMaterialInstance instead."), Category="Rendering|Material")
 	virtual class UMaterialInstanceDynamic* CreateAndSetMaterialInstanceDynamicFromMaterial(int32 ElementIndex, class UMaterialInterface* Parent);
 
 	/**
@@ -721,7 +721,7 @@ public:
 	 * Sets the constraint mode of the component.
 	 * @param ConstraintMode	The type of constraint to use.
 	 */
-	UFUNCTION(BlueprintCallable, meta = (FriendlyName = "Set Constraint Mode", Keywords = "set locked axis constraint physics"), Category = Physics)
+	UFUNCTION(BlueprintCallable, meta = (DisplayName = "Set Constraint Mode", Keywords = "set locked axis constraint physics"), Category = Physics)
 	virtual void SetConstraintMode(EDOFMode::Type ConstraintMode);
 
 	/**
@@ -934,7 +934,7 @@ public:
 	void SetCollisionObjectType(ECollisionChannel Channel);
 
 	/** Perform a line trace against a single component */
-	UFUNCTION(BlueprintCallable, Category="Collision", meta=(FriendlyName = "Line Trace Component", bTraceComplex="true"))	
+	UFUNCTION(BlueprintCallable, Category="Collision", meta=(DisplayName = "Line Trace Component", bTraceComplex="true"))	
 	bool K2_LineTraceComponent(FVector TraceStart, FVector TraceEnd, bool bTraceComplex, bool bShowTrace, FVector& HitLocation, FVector& HitNormal, FName& BoneName);
 
 	/** Sets the bRenderCustomDepth property and marks the render state dirty. */
@@ -1572,7 +1572,7 @@ public:
 	 *	Note that if physics is already running on this component, this will _not_ alter its mass/inertia etc,  
 	 *	it will only change its surface properties like friction.
 	 */
-	UFUNCTION(BlueprintCallable, Category="Physics", meta=(FriendlyName="Set PhysicalMaterial Override"))
+	UFUNCTION(BlueprintCallable, Category="Physics", meta=(DisplayName="Set PhysicalMaterial Override"))
 	virtual void SetPhysMaterialOverride(class UPhysicalMaterial* NewPhysMaterial);
 
 	/** 
@@ -1586,7 +1586,7 @@ public:
 	 * Changes the value of CullDistance.
 	 * @param NewCullDistance - The value to assign to CullDistance.
 	 */
-	UFUNCTION(BlueprintCallable, Category="LOD", meta=(FriendlyName="Set Max Draw Distance"))
+	UFUNCTION(BlueprintCallable, Category="LOD", meta=(DisplayName="Set Max Draw Distance"))
 	void SetCullDistance(float NewCullDistance);
 	
 	/**
