Mac build fix

diff --git a/Engine/Source/Runtime/Engine/Classes/GameFramework/Actor.h b/Engine/Source/Runtime/Engine/Classes/GameFramework/Actor.h
--- a/Engine/Source/Runtime/Engine/Classes/GameFramework/Actor.h
+++ b/Engine/Source/Runtime/Engine/Classes/GameFramework/Actor.h
@@ -1186,6 +1186,7 @@ public:
 	{
 	public:
 		FActorTransactionAnnotation(const AActor* Actor);
+		virtual ~FActorTransactionAnnotation() {}
 
 		virtual void AddReferencedObjects(FReferenceCollector& Collector) override;
 
