Make trace distance for player controller's object hit result customizable

diff --git a/Engine/Source/Runtime/Engine/Classes/GameFramework/PlayerController.h b/Engine/Source/Runtime/Engine/Classes/GameFramework/PlayerController.h
--- a/Engine/Source/Runtime/Engine/Classes/GameFramework/PlayerController.h
+++ b/Engine/Source/Runtime/Engine/Classes/GameFramework/PlayerController.h
@@ -207,14 +207,6 @@ class ENGINE_API APlayerController : public AController
 	UPROPERTY()
 	float LastSpectatorStateSynchTime;
 
-	/** Last location synced on the server for a spectator. */
-	UPROPERTY(Transient)
-	FVector LastSpectatorSyncLocation;
-
-	/** Last rotation synced on the server for a spectator. */
-	UPROPERTY(Transient)
-	FRotator LastSpectatorSyncRotation;
-
 	/** Cap set by server on bandwidth from client to server in bytes/sec (only has impact if >=2600) */
 	UPROPERTY()
 	int32 ClientCap;
@@ -245,18 +237,9 @@ class ENGINE_API APlayerController : public AController
 	/** Whether this controller is using streaming volumes.  **/
 	uint32 bIsUsingStreamingVolumes:1;
 
-	/** True if PlayerController is currently waiting for the match to start or to respawn. Only valid in Spectating state. */
-	UPROPERTY(VisibleInstanceOnly, BlueprintReadOnly, Category=PlayerController)
+	/** Only valid in Spectating state. True if PlayerController is currently waiting for the match to start */
 	uint32 bPlayerIsWaiting:1;
 
-	/** Indicate that the Spectator is waiting to join/respawn. */
-	UFUNCTION(server, reliable, WithValidation, Category=PlayerController)
-	void ServerSetSpectatorWaiting(bool bWaiting);
-
-	/** Indicate that the Spectator is waiting to join/respawn. */
-	UFUNCTION(client, reliable, Category=PlayerController)
-	void ClientSetSpectatorWaiting(bool bWaiting);
-
 	/** index identifying players using the same base connection (splitscreen clients)
 	 * Used by netcode to match replicated PlayerControllers to the correct splitscreen viewport and child connection
 	 * replicated via special internal code, not through normal variable replication
@@ -331,6 +314,9 @@ class ENGINE_API APlayerController : public AController
 	UPROPERTY(BlueprintReadWrite, Category=MouseInterface)
 	TEnumAsByte<ECollisionChannel> CurrentClickTraceChannel;
 
+	UPROPERTY(EditAnywhere, BlueprintReadWrite, Category=MouseInterface, meta=(DisplayName="Trace Distance"))
+	float HitResultTraceDistance;
+
 	FForceFeedbackValues ForceFeedbackValues;
 
 public:
@@ -936,9 +922,9 @@ public:
 	UFUNCTION(reliable, server, WithValidation)
 	void ServerRestartPlayer();
 
-	/** When spectating, updates spectator location/rotation and pings the server to make sure spectating should continue. */
+	/** When spectating, pings the server to make sure spectating should continue. */
 	UFUNCTION(unreliable, server, WithValidation)
-	void ServerSetSpectatorLocation(FVector NewLoc, FRotator NewRot);
+	void ServerSetSpectatorLocation(FVector NewLoc);
 
 	/** Calls ServerSetSpectatorLocation but throttles it to reduce bandwidth and only calls it when necessary. */
 	void SafeServerUpdateSpectatorState();
diff --git a/Engine/Source/Runtime/Engine/Private/PlayerController.cpp b/Engine/Source/Runtime/Engine/Private/PlayerController.cpp
--- a/Engine/Source/Runtime/Engine/Private/PlayerController.cpp
+++ b/Engine/Source/Runtime/Engine/Private/PlayerController.cpp
@@ -63,6 +63,7 @@ APlayerController::APlayerController(const FObjectInitializer& ObjectInitializer
 	LastRetryPlayerTime = 0.f;
 	DefaultMouseCursor = EMouseCursor::Default;
 	DefaultClickTraceChannel = ECollisionChannel::ECC_Visibility;
+	HitResultTraceDistance = 100000.f;
 
 	bCinemaDisableInputMove = false;
 	bCinemaDisableInputLook = false;
@@ -1948,7 +1949,7 @@ bool APlayerController::GetHitResultAtScreenPosition(const FVector2D ScreenPosit
 			FVector WorldDirection;
 			SceneView->DeprojectFVector2D(ScreenPosition, WorldOrigin, WorldDirection);
 
-			return GetWorld()->LineTraceSingleByChannel(HitResult, WorldOrigin, WorldOrigin + WorldDirection * 100000.f, TraceChannel, CollisionQueryParams);
+			return GetWorld()->LineTraceSingleByChannel(HitResult, WorldOrigin, WorldOrigin + WorldDirection * HitResultTraceDistance, TraceChannel, CollisionQueryParams);
 		}
 	}
 
@@ -2000,7 +2001,7 @@ bool APlayerController::GetHitResultAtScreenPosition(const FVector2D ScreenPosit
 
 			FCollisionObjectQueryParams ObjParam(ObjectTypes);
 			
-			return GetWorld()->LineTraceSingleByObjectType(HitResult, WorldOrigin, WorldOrigin + WorldDirection * 100000.f, ObjParam, FCollisionQueryParams("ClickableTrace", bTraceComplex));
+			return GetWorld()->LineTraceSingleByObjectType(HitResult, WorldOrigin, WorldOrigin + WorldDirection * HitResultTraceDistance, ObjParam, FCollisionQueryParams("ClickableTrace", bTraceComplex));
 		}
 	}
 
