#ue4 - Add 'bInitialOverlap' to BreakHitResult. Improve comments/documentation for initial overlap case.

diff --git a/Engine/Source/Runtime/Engine/Classes/Kismet/GameplayStatics.h b/Engine/Source/Runtime/Engine/Classes/Kismet/GameplayStatics.h
--- a/Engine/Source/Runtime/Engine/Classes/Kismet/GameplayStatics.h
+++ b/Engine/Source/Runtime/Engine/Classes/Kismet/GameplayStatics.h
@@ -403,10 +403,11 @@ class ENGINE_API UGameplayStatics : public UBlueprintFunctionLibrary
 	/** Extracts data from a HitResult. 
 	 * @param Hit			The source HitResult.
 	 * @param bBlockingHit	True if there was a blocking hit, false otherwise.
+	 * @param bInitialOverlap True if the hit started in an initial overlap. In this case some other values should be interpreted differently. Time will be 0, ImpactPoint will equal Location, and normals will be equal and indicate a depenetration vector.
 	 * @param Time			'Time' of impact along trace direction ranging from [0.0 to 1.0) if there is a hit, indicating time between start and end. Equals 1.0 if there is no hit.
 	 * @param Location		Location of the hit in world space. If this was a swept shape test, this is the location where we can place the shape in the world where it will not penetrate.
 	 * @param Normal		Normal of the hit in world space, for the object that was swept (e.g. for a sphere trace this points towards the sphere's center). Equal to ImpactNormal for line tests.
-	 * @param ImpactPoint	Location of the actual contact point of the trace shape with the surface of the hit object.
+	 * @param ImpactPoint	Location of the actual contact point of the trace shape with the surface of the hit object. Equal to Location in the case of an initial overlap.
 	 * @param ImpactNormal	Normal of the hit in world space, for the object that was hit by the sweep.
 	 * @param PhysMat		Physical material that was hit. Must set bReturnPhysicalMaterial to true in the query params for this to be returned.
 	 * @param HitActor		Actor hit by the trace.
@@ -415,7 +416,7 @@ class ENGINE_API UGameplayStatics : public UBlueprintFunctionLibrary
 	 * @param HitItem		Primitive-specific data recording which item in the primitive was hit
 	 */
 	UFUNCTION(BlueprintPure, Category = "Collision", meta=(NativeBreakFunc))
-	static void BreakHitResult(const struct FHitResult& Hit, bool& bBlockingHit, float& Time, FVector& Location, FVector& Normal, FVector& ImpactPoint, FVector& ImpactNormal, class UPhysicalMaterial*& PhysMat, class AActor*& HitActor, class UPrimitiveComponent*& HitComponent, FName& HitBoneName, int32& HitItem, FVector& TraceStart, FVector& TraceEnd);
+	static void BreakHitResult(const struct FHitResult& Hit, bool& bBlockingHit, bool& bInitialOverlap, float& Time, FVector& Location, FVector& Normal, FVector& ImpactPoint, FVector& ImpactNormal, class UPhysicalMaterial*& PhysMat, class AActor*& HitActor, class UPrimitiveComponent*& HitComponent, FName& HitBoneName, int32& HitItem, FVector& TraceStart, FVector& TraceEnd);
 
 	/** Returns the EPhysicalSurface type of the given Hit. 
 	 * To edit surface type for your project, use ProjectSettings/Physics/PhysicalSurface section
diff --git a/Engine/Source/Runtime/Engine/Private/GameplayStatics.cpp b/Engine/Source/Runtime/Engine/Private/GameplayStatics.cpp
--- a/Engine/Source/Runtime/Engine/Private/GameplayStatics.cpp
+++ b/Engine/Source/Runtime/Engine/Private/GameplayStatics.cpp
@@ -602,9 +602,10 @@ UParticleSystemComponent* UGameplayStatics::SpawnEmitterAttached(UParticleSystem
 	return PSC;
 }
 
-void UGameplayStatics::BreakHitResult(const FHitResult& Hit, bool& bBlockingHit, float& Time, FVector& Location, FVector& Normal, FVector& ImpactPoint, FVector& ImpactNormal, UPhysicalMaterial*& PhysMat, AActor*& HitActor, UPrimitiveComponent*& HitComponent, FName& HitBoneName, int32& HitItem, FVector& TraceStart, FVector& TraceEnd)
+void UGameplayStatics::BreakHitResult(const FHitResult& Hit, bool& bBlockingHit, bool& bInitialOverlap, float& Time, FVector& Location, FVector& Normal, FVector& ImpactPoint, FVector& ImpactNormal, UPhysicalMaterial*& PhysMat, AActor*& HitActor, UPrimitiveComponent*& HitComponent, FName& HitBoneName, int32& HitItem, FVector& TraceStart, FVector& TraceEnd)
 {
 	bBlockingHit = Hit.bBlockingHit;
+	bInitialOverlap = Hit.bStartPenetrating;
 	Time = Hit.Time;
 	Location = Hit.Location;
 	Normal = Hit.Normal;
