Replace changes caused by bad auto-resolve.

diff --git a/Engine/Source/Runtime/Engine/Classes/GameFramework/PlayerController.h b/Engine/Source/Runtime/Engine/Classes/GameFramework/PlayerController.h
--- a/Engine/Source/Runtime/Engine/Classes/GameFramework/PlayerController.h
+++ b/Engine/Source/Runtime/Engine/Classes/GameFramework/PlayerController.h
@@ -207,6 +207,14 @@ class ENGINE_API APlayerController : public AController
 	UPROPERTY()
 	float LastSpectatorStateSynchTime;
 
+	/** Last location synced on the server for a spectator. */
+	UPROPERTY(Transient)
+	FVector LastSpectatorSyncLocation;
+
+	/** Last rotation synced on the server for a spectator. */
+	UPROPERTY(Transient)
+	FRotator LastSpectatorSyncRotation;
+
 	/** Cap set by server on bandwidth from client to server in bytes/sec (only has impact if >=2600) */
 	UPROPERTY()
 	int32 ClientCap;
@@ -237,9 +245,18 @@ class ENGINE_API APlayerController : public AController
 	/** Whether this controller is using streaming volumes.  **/
 	uint32 bIsUsingStreamingVolumes:1;
 
-	/** Only valid in Spectating state. True if PlayerController is currently waiting for the match to start */
+	/** True if PlayerController is currently waiting for the match to start or to respawn. Only valid in Spectating state. */
+	UPROPERTY(VisibleInstanceOnly, BlueprintReadOnly, Category=PlayerController)
 	uint32 bPlayerIsWaiting:1;
 
+	/** Indicate that the Spectator is waiting to join/respawn. */
+	UFUNCTION(server, reliable, WithValidation, Category=PlayerController)
+	void ServerSetSpectatorWaiting(bool bWaiting);
+
+	/** Indicate that the Spectator is waiting to join/respawn. */
+	UFUNCTION(client, reliable, Category=PlayerController)
+	void ClientSetSpectatorWaiting(bool bWaiting);
+
 	/** index identifying players using the same base connection (splitscreen clients)
 	 * Used by netcode to match replicated PlayerControllers to the correct splitscreen viewport and child connection
 	 * replicated via special internal code, not through normal variable replication
@@ -922,9 +939,9 @@ public:
 	UFUNCTION(reliable, server, WithValidation)
 	void ServerRestartPlayer();
 
-	/** When spectating, pings the server to make sure spectating should continue. */
+	/** When spectating, updates spectator location/rotation and pings the server to make sure spectating should continue. */
 	UFUNCTION(unreliable, server, WithValidation)
-	void ServerSetSpectatorLocation(FVector NewLoc);
+	void ServerSetSpectatorLocation(FVector NewLoc, FRotator NewRot);
 
 	/** Calls ServerSetSpectatorLocation but throttles it to reduce bandwidth and only calls it when necessary. */
 	void SafeServerUpdateSpectatorState();
