PR #953: Extend DrawDebugLine Thickness parameter up to DrawDebugDirectionalArrow function (Contributed by [REDACTED_PERSON])
https://github.com/EpicGames/UnrealEngine/pull/953

diff --git a/Engine/Source/Runtime/Engine/Private/DrawDebugHelpers.cpp b/Engine/Source/Runtime/Engine/Private/DrawDebugHelpers.cpp
--- a/Engine/Source/Runtime/Engine/Private/DrawDebugHelpers.cpp
+++ b/Engine/Source/Runtime/Engine/Private/DrawDebugHelpers.cpp
@@ -50,7 +50,7 @@ void DrawDebugPoint(const UWorld* InWorld, FVector const& Position, float Size,
 	}
 }
 
-void DrawDebugDirectionalArrow(const UWorld* InWorld, FVector const& LineStart, FVector const& LineEnd, float ArrowSize, FColor const& Color, bool bPersistentLines, float LifeTime, uint8 DepthPriority)
+void DrawDebugDirectionalArrow(const UWorld* InWorld, FVector const& LineStart, FVector const& LineEnd, float ArrowSize, FColor const& Color, bool bPersistentLines, float LifeTime, uint8 DepthPriority, float Thickness)
 {
 	// no debug line drawing on dedicated server
 	if (GEngine->GetNetMode(InWorld) != NM_DedicatedServer)
@@ -60,7 +60,7 @@ void DrawDebugDirectionalArrow(const UWorld* InWorld, FVector const& LineStart,
 			ArrowSize = 10.f;
 		}
 
-		DrawDebugLine(InWorld, LineStart, LineEnd, Color, bPersistentLines, LifeTime);
+		DrawDebugLine(InWorld, LineStart, LineEnd, Color, bPersistentLines, LifeTime, DepthPriority, Thickness);
 
 		FVector Dir = (LineEnd-LineStart);
 		Dir.Normalize();
@@ -78,8 +78,8 @@ void DrawDebugDirectionalArrow(const UWorld* InWorld, FVector const& LineStart,
 		// since dir is x direction, my arrow will be pointing +y, -x and -y, -x
 		float ArrowSqrt = FMath::Sqrt(ArrowSize);
 		FVector ArrowPos;
-		DrawDebugLine(InWorld, LineEnd, LineEnd + TM.TransformPosition(FVector(-ArrowSqrt, ArrowSqrt, 0)), Color, bPersistentLines, LifeTime, DepthPriority);
-		DrawDebugLine(InWorld, LineEnd, LineEnd + TM.TransformPosition(FVector(-ArrowSqrt, -ArrowSqrt, 0)), Color, bPersistentLines, LifeTime, DepthPriority);
+		DrawDebugLine(InWorld, LineEnd, LineEnd + TM.TransformPosition(FVector(-ArrowSqrt, ArrowSqrt, 0)), Color, bPersistentLines, LifeTime, DepthPriority, Thickness);
+		DrawDebugLine(InWorld, LineEnd, LineEnd + TM.TransformPosition(FVector(-ArrowSqrt, -ArrowSqrt, 0)), Color, bPersistentLines, LifeTime, DepthPriority, Thickness);
 	}
 }
 
diff --git a/Engine/Source/Runtime/Engine/Public/DrawDebugHelpers.h b/Engine/Source/Runtime/Engine/Public/DrawDebugHelpers.h
--- a/Engine/Source/Runtime/Engine/Public/DrawDebugHelpers.h
+++ b/Engine/Source/Runtime/Engine/Public/DrawDebugHelpers.h
@@ -14,7 +14,7 @@ ENGINE_API void DrawDebugLine(const UWorld* InWorld, FVector const& LineStart, F
 /** Draw a debug point */
 ENGINE_API void DrawDebugPoint(const UWorld* InWorld, FVector const& Position, float Size, FColor const& PointColor, bool bPersistentLines = false, float LifeTime=-1.f, uint8 DepthPriority = 0);
 /** Draw directional arrow **/
-ENGINE_API void DrawDebugDirectionalArrow(const UWorld* InWorld, FVector const& LineStart, FVector const& LineEnd, float ArrowSize, FColor const& Color, bool bPersistentLines = false, float LifeTime=-1.f, uint8 DepthPriority = 0);
+ENGINE_API void DrawDebugDirectionalArrow(const UWorld* InWorld, FVector const& LineStart, FVector const& LineEnd, float ArrowSize, FColor const& Color, bool bPersistentLines = false, float LifeTime=-1.f, uint8 DepthPriority = 0, float Thickness = 0.f);
 /** Draw a debug box */
 ENGINE_API void DrawDebugBox(const UWorld* InWorld, FVector const& Center, FVector const& Extent, FColor const& Color, bool bPersistentLines = false, float LifeTime=-1.f, uint8 DepthPriority = 0);
 /** Draw a debug box with rotation */
