Engine - Added FRequestExit command

diff --git a/Engine/Source/Runtime/Engine/Private/Tests/AutomationCommon.cpp b/Engine/Source/Runtime/Engine/Private/Tests/AutomationCommon.cpp
--- a/Engine/Source/Runtime/Engine/Private/Tests/AutomationCommon.cpp
+++ b/Engine/Source/Runtime/Engine/Private/Tests/AutomationCommon.cpp
@@ -53,4 +53,10 @@ bool FLoadGameMapCommand::Update()
 	UE_LOG(LogEngineAutomationTests, Log, TEXT("Loading Map Now. '%s'"), *MapName);
 	GEngine->Exec(GEngine->GetWorldContexts()[0].World(), *FString::Printf(TEXT("Open %s"), *MapName));
 	return true;
+}
+
+bool FRequestExitCommand::Update()
+{
+	GIsRequestingExit = true;
+	return true;
 }
\ No newline at end of file
diff --git a/Engine/Source/Runtime/Engine/Public/Tests/AutomationCommon.h b/Engine/Source/Runtime/Engine/Public/Tests/AutomationCommon.h
--- a/Engine/Source/Runtime/Engine/Public/Tests/AutomationCommon.h
+++ b/Engine/Source/Runtime/Engine/Public/Tests/AutomationCommon.h
@@ -109,4 +109,9 @@ DEFINE_ENGINE_LATENT_AUTOMATION_COMMAND_ONE_PARAMETER(FTakeEditorScreenshotComma
 */
 DEFINE_LATENT_AUTOMATION_COMMAND_ONE_PARAMETER(FLoadGameMapCommand, FString, MapName);
 
+/**
+ * Latent command that requests exit
+ */
+DEFINE_ENGINE_LATENT_AUTOMATION_COMMAND( FRequestExitCommand );
+
 
