Added an experimental setting allowing for creation of navigation data on net clients #UE4

diff --git a/Engine/Source/Runtime/Engine/Classes/AI/Navigation/NavigationSystem.h b/Engine/Source/Runtime/Engine/Classes/AI/Navigation/NavigationSystem.h
--- a/Engine/Source/Runtime/Engine/Classes/AI/Navigation/NavigationSystem.h
+++ b/Engine/Source/Runtime/Engine/Classes/AI/Navigation/NavigationSystem.h
@@ -143,6 +143,9 @@ protected:
 	UPROPERTY(config, EditAnywhere, Category=NavigationSystem)
 	uint32 bAutoCreateNavigationData:1;
 
+	UPROPERTY(config, EditAnywhere, Category=NavigationSystem)
+	uint32 bAllowClientSizeNavigation_Experimental:1;
+
 	/** gets set to true if gathering navigation data (like in navoctree) is required due to the need of navigation generation 
 	 *	Is always true in Editor Mode. In other modes it depends on bRebuildAtRuntime of every required NavigationData class' CDO
 	 */
@@ -466,6 +469,8 @@ public:
 
 	bool ShouldGenerateNavigationEverywhere() const { return bWholeWorldNavigable; }
 
+	virtual bool ShouldLoadNavigationOnClient(ANavigationData* NavData = nullptr) const { return bAllowClientSizeNavigation_Experimental; }
+
 	FBox GetWorldBounds() const;
 	
 	FBox GetLevelBounds(ULevel* InLevel) const;
diff --git a/Engine/Source/Runtime/Engine/Private/AI/Navigation/NavigationData.cpp b/Engine/Source/Runtime/Engine/Private/AI/Navigation/NavigationData.cpp
--- a/Engine/Source/Runtime/Engine/Private/AI/Navigation/NavigationData.cpp
+++ b/Engine/Source/Runtime/Engine/Private/AI/Navigation/NavigationData.cpp
@@ -109,7 +109,7 @@ ANavigationData::ANavigationData(const FObjectInitializer& ObjectInitializer)
 	, NavDataUniqueID(GetNextUniqueID())
 {
 	PrimaryActorTick.bCanEverTick = true;
-	bNetLoadOnClient = false;
+	bNetLoadOnClient = (HasAnyFlags(RF_ClassDefaultObject) == false) && (*GEngine->NavigationSystemClass != nullptr) && (GEngine->NavigationSystemClass->GetDefaultObject<UNavigationSystem>()->ShouldLoadNavigationOnClient());
 	bCanBeDamaged = false;
 	DefaultQueryFilter = MakeShareable(new FNavigationQueryFilter());
 	ObservedPathsTickInterval = 0.5;
diff --git a/Engine/Source/Runtime/Engine/Private/AI/Navigation/NavigationSystem.cpp b/Engine/Source/Runtime/Engine/Private/AI/Navigation/NavigationSystem.cpp
--- a/Engine/Source/Runtime/Engine/Private/AI/Navigation/NavigationSystem.cpp
+++ b/Engine/Source/Runtime/Engine/Private/AI/Navigation/NavigationSystem.cpp
@@ -2000,7 +2000,8 @@ UNavigationSystem* UNavigationSystem::CreateNavigationSystem(UWorld* WorldOwner)
 
 #if WITH_SERVER_CODE || WITH_EDITOR
 	// create navigation system for editor and server targets, but remove it from game clients
-	if (WorldOwner && (WorldOwner->GetNetMode() != NM_Client) && (*GEngine->NavigationSystemClass != nullptr))
+	if (WorldOwner && (*GEngine->NavigationSystemClass != nullptr) 
+		&& (GEngine->NavigationSystemClass->GetDefaultObject<UNavigationSystem>()->bAllowClientSizeNavigation_Experimental || WorldOwner->GetNetMode() != NM_Client))
 	{
 		AWorldSettings* WorldSettings = WorldOwner->GetWorldSettings();
 		if (WorldSettings == NULL || WorldSettings->bEnableNavigationSystem)
