[UE-12125] Relative transform methods (e.g. AddRelativeLocation) now work when invoked on non-movable member component instances during the Construction Script in a Blueprint.

change summary:
- relocated UPrimitiveComponent::CheckStaticMobilityAndWarn() to USceneComponent to consolidate inconsistent mobility checks for movement
- modified the implementation of CheckStaticMobilityAndWarn() to use USceneComponent's previous check logic
- modified USceneComponent::MoveComponent() to call CheckStaticMobilityAndWarn() in place of its previous inline check
- removed the TODO comment in the UPrimitiveComponent::MoveComponent() override as the updated check logic should cover that case

diff --git a/Engine/Source/Runtime/Engine/Classes/Components/PrimitiveComponent.h b/Engine/Source/Runtime/Engine/Classes/Components/PrimitiveComponent.h
--- a/Engine/Source/Runtime/Engine/Classes/Components/PrimitiveComponent.h
+++ b/Engine/Source/Runtime/Engine/Classes/Components/PrimitiveComponent.h
@@ -1546,9 +1546,6 @@ private:
 	 */
 	bool ApplyRigidBodyState(const FRigidBodyState& NewState, const FRigidBodyErrorCorrection& ErrorCorrection, FVector& OutDeltaPos, FName BoneName = NAME_None);
 
-	/** Check if mobility is set to non-static. If it's static we trigger a PIE warning and return true*/
-	bool CheckStaticMobilityAndWarn(const FText& ActionText) const;
-
 	/** Check if mobility is set to non-static. If BodyInstanceRequiresSimulation is non-null we check that it is simulated. Triggers a PIE warning if conditions fails */
 	void WarnInvalidPhysicsOperations_Internal(const FText& ActionText, const FBodyInstance* BodyInstanceRequiresSimulation = nullptr) const;
 
diff --git a/Engine/Source/Runtime/Engine/Classes/Components/SceneComponent.h b/Engine/Source/Runtime/Engine/Classes/Components/SceneComponent.h
--- a/Engine/Source/Runtime/Engine/Classes/Components/SceneComponent.h
+++ b/Engine/Source/Runtime/Engine/Classes/Components/SceneComponent.h
@@ -541,6 +541,9 @@ protected:
 
 	virtual void OnUpdateTransform(bool bSkipPhysicsMove);
 
+	/** Check if mobility is set to non-static. If it's static we trigger a PIE warning and return true*/
+	bool CheckStaticMobilityAndWarn(const FText& ActionText) const;
+
 private:
 
 	void PropagateTransformUpdate(bool bTransformChanged, bool bSkipPhysicsMove = false);
diff --git a/Engine/Source/Runtime/Engine/Private/Components/PrimitiveComponent.cpp b/Engine/Source/Runtime/Engine/Private/Components/PrimitiveComponent.cpp
--- a/Engine/Source/Runtime/Engine/Private/Components/PrimitiveComponent.cpp
+++ b/Engine/Source/Runtime/Engine/Private/Components/PrimitiveComponent.cpp
@@ -1320,25 +1320,6 @@ FCollisionShape UPrimitiveComponent::GetCollisionShape(float Inflation) const
 	return FCollisionShape::MakeBox(Bounds.BoxExtent + Inflation);
 }
 
-bool UPrimitiveComponent::CheckStaticMobilityAndWarn(const FText& ActionText) const
-{
-	// static things can move before they are registered (e.g. immediately after streaming), but not after.
-	if (Mobility == EComponentMobility::Static)
-	{
-		AActor* Actor = GetOwner();
-		if (Actor && Actor->IsActorInitialized())
-		{
-			static const FText WarnText = LOCTEXT("InvalidStaticMove", "Mobility of {0} : {1} has to be 'Movable' if you'd like to {2}. ");
-			FMessageLog("PIE").Warning(FText::Format(WarnText,
-				FText::FromString(GetNameSafe(GetOwner())), FText::FromString(GetName()), ActionText));
-
-			return true;
-		}
-	}
-
-	return false;
-}
-
 bool UPrimitiveComponent::MoveComponent( const FVector& Delta, const FRotator& NewRotation, bool bSweep, FHitResult* OutHit, EMoveComponentFlags MoveFlags)
 {
 	SCOPE_CYCLE_COUNTER(STAT_MoveComponentTime);
@@ -1365,7 +1346,6 @@ bool UPrimitiveComponent::MoveComponent( const FVector& Delta, const FRotator& N
 	AActor* const Actor = GetOwner();
 
 	// static things can move before they are registered (e.g. immediately after streaming), but not after.
-	// TODO: Static components without an owner can move, should they be able to?
 	static const FText WarnText = LOCTEXT("InvalidMove", "move");
 	if (CheckStaticMobilityAndWarn(WarnText))
 	{
diff --git a/Engine/Source/Runtime/Engine/Private/Components/SceneComponent.cpp b/Engine/Source/Runtime/Engine/Private/Components/SceneComponent.cpp
--- a/Engine/Source/Runtime/Engine/Private/Components/SceneComponent.cpp
+++ b/Engine/Source/Runtime/Engine/Private/Components/SceneComponent.cpp
@@ -1639,10 +1639,8 @@ void USceneComponent::UpdateOverlaps(TArray<FOverlapInfo> const* PendingOverlaps
 	}
 }
 
-
-bool USceneComponent::MoveComponent( const FVector& Delta, const FRotator& NewRotation, bool bSweep, FHitResult* OutHit, EMoveComponentFlags MoveFlags )
+bool USceneComponent::CheckStaticMobilityAndWarn(const FText& ActionText) const
 {
-#if !(UE_BUILD_SHIPPING || UE_BUILD_TEST)
 	// make sure mobility is movable, otherwise you shouldn't try to move
 	if (Mobility != EComponentMobility::Movable)
 	{
@@ -1653,13 +1651,30 @@ bool USceneComponent::MoveComponent( const FVector& Delta, const FRotator& NewRo
 			// It's only a problem if we're in gameplay, and the owning level is visible
 			if (World->HasBegunPlay() && IsRegistered() && Level && Level->bIsVisible)
 			{
-				FMessageLog("PIE").Warning(FText::Format(LOCTEXT("InvalidMove", "Mobility of {0} : {1} has to be 'Movable' if you'd like to move. "),
-					FText::FromString(GetNameSafe(GetOwner())), FText::FromString(GetName())));
-				return false;
+#if !(UE_BUILD_SHIPPING || UE_BUILD_TEST)
+				FMessageLog("PIE").Warning(FText::Format(LOCTEXT("InvalidMove", "Mobility of {0} : {1} has to be 'Movable' if you'd like to {2}. "),
+					FText::FromString(GetNameSafe(GetOwner())), FText::FromString(GetName()), ActionText));
+#endif
+				return true;
 			}
 		}
 	}
-#endif
+
+	return false;
+}
+
+bool USceneComponent::MoveComponent( const FVector& Delta, const FRotator& NewRotation, bool bSweep, FHitResult* OutHit, EMoveComponentFlags MoveFlags )
+{
+	// static things can move before they are registered (e.g. immediately after streaming), but not after.
+	static const FText WarnText = LOCTEXT("InvalidMove", "move");
+	if (CheckStaticMobilityAndWarn(WarnText))
+	{
+		if (OutHit)
+		{
+			*OutHit = FHitResult();
+		}
+		return false;
+	}
 
 	// Fill in optional output param. SceneComponent doesn't sweep, so this is just an empty result.
 	if (OutHit)
