Fix issue with not calling update child transforms when we have blend physics but no bodies

diff --git a/Engine/Source/Runtime/Engine/Classes/Components/SkeletalMeshComponent.h b/Engine/Source/Runtime/Engine/Classes/Components/SkeletalMeshComponent.h
--- a/Engine/Source/Runtime/Engine/Classes/Components/SkeletalMeshComponent.h
+++ b/Engine/Source/Runtime/Engine/Classes/Components/SkeletalMeshComponent.h
@@ -1168,7 +1168,9 @@ private:
 
 	void RenderAxisGizmo(const FTransform& Transform, class UCanvas* Canvas) const;
 
-	bool ShouldBlendPhysicsBones();	
+	bool ShouldBlendPhysicsBones() const;
+	bool DoAnyPhysicsBodiesHaveWeight() const;
+
 	void ClearAnimScriptInstance();
 	virtual void RefreshActiveVertexAnims() override;
 
diff --git a/Engine/Source/Runtime/Engine/Private/Components/SkeletalMeshComponent.cpp b/Engine/Source/Runtime/Engine/Private/Components/SkeletalMeshComponent.cpp
--- a/Engine/Source/Runtime/Engine/Private/Components/SkeletalMeshComponent.cpp
+++ b/Engine/Source/Runtime/Engine/Private/Components/SkeletalMeshComponent.cpp
@@ -1154,7 +1154,7 @@ void USkeletalMeshComponent::PostAnimEvaluation(FAnimationEvaluationContext& Eva
 	// this causes issue where a half of frame, physics position is fixed with anim pose, and the other half is real simulated position
 	// if you enable physics in tick, since that's before physics update, you'll get animation pose dominating physics pose, which isn't what you want. (Or what you'll see)
 	// so do not update transform if physics is on. This problem will be solved by double buffer, when we keep one buffer for intermediate, and the other buffer for result query
-	if (!bBlendPhysics && !ShouldBlendPhysicsBones())
+	if (!ShouldBlendPhysicsBones())
 	{
 		SCOPE_CYCLE_COUNTER(STAT_UpdateLocalToWorldAndOverlaps);
 
diff --git a/Engine/Source/Runtime/Engine/Private/PhysicsEngine/PhysAnim.cpp b/Engine/Source/Runtime/Engine/Private/PhysicsEngine/PhysAnim.cpp
--- a/Engine/Source/Runtime/Engine/Private/PhysicsEngine/PhysAnim.cpp
+++ b/Engine/Source/Runtime/Engine/Private/PhysicsEngine/PhysAnim.cpp
@@ -194,11 +194,16 @@ void USkeletalMeshComponent::BlendPhysicsBones( TArray<FBoneIndexType>& InRequir
 
 
 
-bool USkeletalMeshComponent::ShouldBlendPhysicsBones()
+bool USkeletalMeshComponent::ShouldBlendPhysicsBones() const
 {
-	for (int32 BodyIndex = 0; BodyIndex < Bodies.Num(); ++BodyIndex)
+	return Bodies.Num() > 0 && (DoAnyPhysicsBodiesHaveWeight() || bBlendPhysics);
+}
+
+bool USkeletalMeshComponent::DoAnyPhysicsBodiesHaveWeight() const
+{
+	for (const FBodyInstance* Body : Bodies)
 	{
-		if (Bodies[BodyIndex]->PhysicsBlendWeight > 0.f)
+		if (Body->PhysicsBlendWeight > 0.f)
 		{
 			return true;
 		}
diff --git a/Engine/Source/Runtime/Engine/Private/SkeletalMeshComponentPhysics.cpp b/Engine/Source/Runtime/Engine/Private/SkeletalMeshComponentPhysics.cpp
--- a/Engine/Source/Runtime/Engine/Private/SkeletalMeshComponentPhysics.cpp
+++ b/Engine/Source/Runtime/Engine/Private/SkeletalMeshComponentPhysics.cpp
@@ -3150,8 +3150,7 @@ void USkeletalMeshComponent::PreClothTick(float DeltaTime)
 	// and run this if that is true or rendered
 	// that will at least reduce the chance of mismatch
 	// generally if you move your actor position, this has to happen to approximately match their bounds
-	bool bShouldBlendPhys = ShouldBlendPhysicsBones() || bBlendPhysics;
-	if (bShouldBlendPhys)
+	if (ShouldBlendPhysicsBones())
 	{
 		if (IsRegistered())
 		{
