Fixed misspelling in UNavigationSystem::bAllowClientSideNavigation_Experimental variable name #UE4

diff --git a/Engine/Source/Runtime/Engine/Classes/AI/Navigation/NavigationSystem.h b/Engine/Source/Runtime/Engine/Classes/AI/Navigation/NavigationSystem.h
--- a/Engine/Source/Runtime/Engine/Classes/AI/Navigation/NavigationSystem.h
+++ b/Engine/Source/Runtime/Engine/Classes/AI/Navigation/NavigationSystem.h
@@ -144,7 +144,7 @@ protected:
 	uint32 bAutoCreateNavigationData:1;
 
 	UPROPERTY(config, EditAnywhere, Category=NavigationSystem)
-	uint32 bAllowClientSizeNavigation_Experimental:1;
+	uint32 bAllowClientSideNavigation_Experimental:1;
 
 	/** gets set to true if gathering navigation data (like in navoctree) is required due to the need of navigation generation 
 	 *	Is always true in Editor Mode. In other modes it depends on bRebuildAtRuntime of every required NavigationData class' CDO
@@ -469,7 +469,7 @@ public:
 
 	bool ShouldGenerateNavigationEverywhere() const { return bWholeWorldNavigable; }
 
-	virtual bool ShouldLoadNavigationOnClient(ANavigationData* NavData = nullptr) const { return bAllowClientSizeNavigation_Experimental; }
+	virtual bool ShouldLoadNavigationOnClient(ANavigationData* NavData = nullptr) const { return bAllowClientSideNavigation_Experimental; }
 
 	FBox GetWorldBounds() const;
 	
diff --git a/Engine/Source/Runtime/Engine/Private/AI/Navigation/NavigationSystem.cpp b/Engine/Source/Runtime/Engine/Private/AI/Navigation/NavigationSystem.cpp
--- a/Engine/Source/Runtime/Engine/Private/AI/Navigation/NavigationSystem.cpp
+++ b/Engine/Source/Runtime/Engine/Private/AI/Navigation/NavigationSystem.cpp
@@ -2001,7 +2001,7 @@ UNavigationSystem* UNavigationSystem::CreateNavigationSystem(UWorld* WorldOwner)
 #if WITH_SERVER_CODE || WITH_EDITOR
 	// create navigation system for editor and server targets, but remove it from game clients
 	if (WorldOwner && (*GEngine->NavigationSystemClass != nullptr) 
-		&& (GEngine->NavigationSystemClass->GetDefaultObject<UNavigationSystem>()->bAllowClientSizeNavigation_Experimental || WorldOwner->GetNetMode() != NM_Client))
+		&& (GEngine->NavigationSystemClass->GetDefaultObject<UNavigationSystem>()->bAllowClientSideNavigation_Experimental || WorldOwner->GetNetMode() != NM_Client))
 	{
 		AWorldSettings* WorldSettings = WorldOwner->GetWorldSettings();
 		if (WorldSettings == NULL || WorldSettings->bEnableNavigationSystem)
