MaterialInstance corretly reports if CachedCookedPlatformData is loaded.

diff --git a/Engine/Source/Runtime/Engine/Private/Materials/MaterialInstance.cpp b/Engine/Source/Runtime/Engine/Private/Materials/MaterialInstance.cpp
--- a/Engine/Source/Runtime/Engine/Private/Materials/MaterialInstance.cpp
+++ b/Engine/Source/Runtime/Engine/Private/Materials/MaterialInstance.cpp
@@ -1558,7 +1558,7 @@ bool UMaterialInstance::IsCachedCookedPlatformDataLoaded( const ITargetPlatform*
 
 		return true;
 	}
-	return true; // this happens if we haven't started caching (begincache hasn't been called yet)
+	return false; // this happens if we haven't started caching (begincache hasn't been called yet)
 }
 
 
