Add some comments to the entries in EComponentCreationMethod

diff --git a/Engine/Source/Runtime/Engine/Classes/Components/ActorComponent.h b/Engine/Source/Runtime/Engine/Classes/Components/ActorComponent.h
--- a/Engine/Source/Runtime/Engine/Classes/Components/ActorComponent.h
+++ b/Engine/Source/Runtime/Engine/Classes/Components/ActorComponent.h
@@ -12,10 +12,10 @@ struct FReplicationFlags;
 UENUM()
 enum class EComponentCreationMethod : uint8
 {
-	Native,
-	SimpleConstructionScript,
-	UserConstructionScript,
-	Instance,
+	Native,						// A component that is part of a native class
+	SimpleConstructionScript,	// A component that is created from a template defined in the Components section of the Blueprint
+	UserConstructionScript,		// A dynamically created component, either from the UserConstructionScript or from a Add Component node in a Blueprint event graph
+	Instance,					// A component added to a single Actor instance via the Component section of the Actor's details panel
 };
 
 /**
