#ue4 const parameter in CustomThunk functions in array library (for safer native code , generated from bp)

diff --git a/Engine/Source/Runtime/Engine/Classes/Kismet/KismetArrayLibrary.h b/Engine/Source/Runtime/Engine/Classes/Kismet/KismetArrayLibrary.h
--- a/Engine/Source/Runtime/Engine/Classes/Kismet/KismetArrayLibrary.h
+++ b/Engine/Source/Runtime/Engine/Classes/Kismet/KismetArrayLibrary.h
@@ -177,11 +177,11 @@ class ENGINE_API UKismetArrayLibrary : public UBlueprintFunctionLibrary
 	static bool GenericArray_RemoveItem(void* TargetArray, const UArrayProperty* ArrayProp, const void* Item);
 	static void GenericArray_Clear(void* TargetArray, const UArrayProperty* ArrayProp);
 	static void GenericArray_Resize(void* TargetArray, const UArrayProperty* ArrayProp, int32 Size);
-	static int32 GenericArray_Length(void* TargetArray, const UArrayProperty* ArrayProp);
-	static int32 GenericArray_LastIndex(void* TargetArray, const UArrayProperty* ArrayProp);
+	static int32 GenericArray_Length(const void* TargetArray, const UArrayProperty* ArrayProp);
+	static int32 GenericArray_LastIndex(const void* TargetArray, const UArrayProperty* ArrayProp);
 	static void GenericArray_Get(void* TargetArray, const UArrayProperty* ArrayProp, int32 Index, void* Item);
 	static void GenericArray_Set(void* TargetArray, const UArrayProperty* ArrayProp, int32 Index, const void* NewItem, bool bSizeToFit);
-	static int32 GenericArray_Find(void* TargetArray, const UArrayProperty* ArrayProperty, const void* ItemToFind);
+	static int32 GenericArray_Find(const void* TargetArray, const UArrayProperty* ArrayProperty, const void* ItemToFind);
 	static void GenericArray_SetArrayPropertyByName(UObject* OwnerObject, FName ArrayPropertyName, const void* SrcArrayAddr);
 
 private:
diff --git a/Engine/Source/Runtime/Engine/Private/KismetArrayLibrary.cpp b/Engine/Source/Runtime/Engine/Private/KismetArrayLibrary.cpp
--- a/Engine/Source/Runtime/Engine/Private/KismetArrayLibrary.cpp
+++ b/Engine/Source/Runtime/Engine/Private/KismetArrayLibrary.cpp
@@ -177,7 +177,7 @@ void UKismetArrayLibrary::GenericArray_Resize(void* TargetArray, const UArrayPro
 	}
 }
 
-int32 UKismetArrayLibrary::GenericArray_Length(void* TargetArray, const UArrayProperty* ArrayProp)
+int32 UKismetArrayLibrary::GenericArray_Length(const void* TargetArray, const UArrayProperty* ArrayProp)
 {
 	if( TargetArray )
 	{
@@ -189,7 +189,7 @@ int32 UKismetArrayLibrary::GenericArray_Length(void* TargetArray, const UArrayPr
 	return 0;
 }
 
-int32 UKismetArrayLibrary::GenericArray_LastIndex(void* TargetArray, const UArrayProperty* ArrayProp)
+int32 UKismetArrayLibrary::GenericArray_LastIndex(const void* TargetArray, const UArrayProperty* ArrayProp)
 {
 	if( TargetArray )
 	{
@@ -245,7 +245,7 @@ void UKismetArrayLibrary::GenericArray_Set(void* TargetArray, const UArrayProper
 	}
 }
 
-int32 UKismetArrayLibrary::GenericArray_Find(void* TargetArray, const UArrayProperty* ArrayProperty, const void* ItemToFind)
+int32 UKismetArrayLibrary::GenericArray_Find(const void* TargetArray, const UArrayProperty* ArrayProperty, const void* ItemToFind)
 {
 	int32 ResultIndex = INDEX_NONE;
 
