Turn on sync scene lock warning. Fix lock semantics in several places.

diff --git a/Engine/Plugins/Experimental/Phya/Source/Phya/Classes/PhyaCollisionHandler.h b/Engine/Plugins/Experimental/Phya/Source/Phya/Classes/PhyaCollisionHandler.h
--- a/Engine/Plugins/Experimental/Phya/Source/Phya/Classes/PhyaCollisionHandler.h
+++ b/Engine/Plugins/Experimental/Phya/Source/Phya/Classes/PhyaCollisionHandler.h
@@ -32,7 +32,7 @@ class UPhyaCollisionHandler : public UPhysicsCollisionHandler
 
 	// Begin UPhysicsCollisionHandler interface
 	virtual void InitCollisionHandler() override;
-	virtual void HandlePhysicsCollisions(const TArray<FCollisionNotifyInfo>& PendingCollisionNotifies) override;
+	virtual void HandlePhysicsCollisions_AssumesLocked(const TArray<FCollisionNotifyInfo>& PendingCollisionNotifies) override;
 	// End UPhysicsCollisionHandler interface
 
 	void TestImpact();
diff --git a/Engine/Plugins/Experimental/Phya/Source/Phya/Private/PhyaCollisionHandler.cpp b/Engine/Plugins/Experimental/Phya/Source/Phya/Private/PhyaCollisionHandler.cpp
--- a/Engine/Plugins/Experimental/Phya/Source/Phya/Private/PhyaCollisionHandler.cpp
+++ b/Engine/Plugins/Experimental/Phya/Source/Phya/Private/PhyaCollisionHandler.cpp
@@ -178,7 +178,7 @@ void UPhyaCollisionHandler::TestImpact()
 
 
 
-void UPhyaCollisionHandler::HandlePhysicsCollisions(const TArray<FCollisionNotifyInfo>& PendingCollisionNotifies)
+void UPhyaCollisionHandler::HandlePhysicsCollisions_AssumesLocked(const TArray<FCollisionNotifyInfo>& PendingCollisionNotifies)
 {
 	if(GetWorld()->HasBegunPlay())
 	{
diff --git a/Engine/Source/Runtime/Core/Private/Misc/Core.cpp b/Engine/Source/Runtime/Core/Private/Misc/Core.cpp
--- a/Engine/Source/Runtime/Core/Private/Misc/Core.cpp
+++ b/Engine/Source/Runtime/Core/Private/Misc/Core.cpp
@@ -256,7 +256,8 @@ DEFINE_STAT(STAT_SkeletalMeshMotionBlurSkinningMemory);
 DEFINE_STAT(STAT_VertexShaderMemory);
 DEFINE_STAT(STAT_PixelShaderMemory);
 DEFINE_STAT(STAT_NavigationMemory);
-DEFINE_STAT(STAT_PhysSceneLock);
+DEFINE_STAT(STAT_PhysSceneReadLock);
+DEFINE_STAT(STAT_PhysSceneWriteLock);
 
 DEFINE_STAT(STAT_ReflectionCaptureTextureMemory);
 DEFINE_STAT(STAT_ReflectionCaptureMemory);
diff --git a/Engine/Source/Runtime/Core/Public/Misc/CoreStats.h b/Engine/Source/Runtime/Core/Public/Misc/CoreStats.h
--- a/Engine/Source/Runtime/Core/Public/Misc/CoreStats.h
+++ b/Engine/Source/Runtime/Core/Public/Misc/CoreStats.h
@@ -35,7 +35,8 @@ DECLARE_MEMORY_STAT_EXTERN(TEXT("Navigation Memory"),STAT_NavigationMemory,STATG
 DECLARE_MEMORY_STAT_EXTERN(TEXT("PhysX Memory Used"),STAT_MemoryPhysXTotalAllocationSize,STATGROUP_Memory, CORE_API);
 DECLARE_MEMORY_STAT_EXTERN(TEXT("ICU Memory Used"),STAT_MemoryICUTotalAllocationSize,STATGROUP_Memory, CORE_API);
 DECLARE_MEMORY_STAT_EXTERN(TEXT("ICU Data File Memory Used"),STAT_MemoryICUDataFileAllocationSize,STATGROUP_Memory, CORE_API);
-DECLARE_CYCLE_STAT_EXTERN(TEXT("PhysX Scene Lock"), STAT_PhysSceneLock, STATGROUP_Physics, CORE_API);
+DECLARE_CYCLE_STAT_EXTERN(TEXT("PhysX Scene ReadLock"), STAT_PhysSceneReadLock, STATGROUP_Physics, CORE_API);
+DECLARE_CYCLE_STAT_EXTERN(TEXT("PhysX Scene WriteLock"), STAT_PhysSceneWriteLock, STATGROUP_Physics, CORE_API);
 
 DECLARE_MEMORY_STAT_EXTERN(TEXT("StaticMesh Total Memory"),STAT_StaticMeshTotalMemory2,STATGROUP_MemoryStaticMesh, CORE_API);
 DECLARE_MEMORY_STAT_EXTERN(TEXT("StaticMesh Vertex Memory"),STAT_StaticMeshVertexMemory,STATGROUP_MemoryStaticMesh, CORE_API);
diff --git a/Engine/Source/Runtime/Engine/Classes/PhysicsEngine/BodyInstance.h b/Engine/Source/Runtime/Engine/Classes/PhysicsEngine/BodyInstance.h
--- a/Engine/Source/Runtime/Engine/Classes/PhysicsEngine/BodyInstance.h
+++ b/Engine/Source/Runtime/Engine/Classes/PhysicsEngine/BodyInstance.h
@@ -35,6 +35,7 @@ namespace physx
 	struct PxContactPair;
 	struct PxFilterData;
 }
+
 #endif // WITH_PHYSX
 
 UENUM(BlueprintType)
@@ -436,33 +437,155 @@ public:
 	 */
 	static void InitStaticBodies(TArray<FBodyInstance*>& Bodies, TArray<FTransform>& Transforms, class UBodySetup* BodySetup, class UPrimitiveComponent* PrimitiveComp, class FPhysScene* InRBScene);
 
+	/** Obtains the appropriate PhysX scene lock for READING and executes the passed in lambda. */
+	void ExecuteOnPhysicsReadOnly(TFunctionRef<void()> Func) const;
+
+	/** Obtains the appropriate PhysX scene lock for WRITING and executes the passed in lambda. */
+	void ExecuteOnPhysicsReadWrite(TFunctionRef<void()> Func) const;
+	
 #if WITH_PHYSX
+
+	/** If scene type is not specified return the first scene used. */
+	int32 GetSceneIndex(int32 SceneType = -1) const;
+
+
 	/** Initialise dynamic properties for this instance when using PhysX - this must be done after scene addition.
 	 */
-	void InitDynamicProperties();
+	DEPRECATED(4.8, "Please call InitDynamicProperties_AssumesLocked and make sure you obtain the appropriate PhysX scene locks")
+	void InitDynamicProperties()
+	{
+		InitDynamicProperties_AssumesLocked();
+	}
+
+	/** Initialise dynamic properties for this instance when using PhysX - this must be done after scene addition.
+	 *  Note: This function is not thread safe. Make sure to obtain the appropriate PhysX scene locks before calling this function
+	 */
+	void InitDynamicProperties_AssumesLocked();
 
 	/** Populate the filter data within the provided FShapeData with the correct filters for this instance
 	 *	@param ShapeData ShapeData to populate
 	 *	@param bForceSimpleAsComplex Whether to force simple colision as complex
 	 */
-	void GetFilterData(FShapeData& ShapeData, bool bForceSimpleAsComplex = false);
+	DEPRECATED(4.8, "Please call GetFilterData_AssumesLocked and make sure you obtain the appropriate PhysX scene locks")
+	void GetFilterData(FShapeData& ShapeData, bool bForceSimpleAsComplex = false)
+	{
+		GetFilterData_AssumesLocked(ShapeData, bForceSimpleAsComplex);
+	}
+
+	/** Populate the filter data within the provided FShapeData with the correct filters for this instance
+	 *	@param ShapeData ShapeData to populate
+	 *	@param bForceSimpleAsComplex Whether to force simple colision as complex
+	 *  Note: This function is not thread safe. Make sure to obtain the appropriate PhysX scene locks before calling this function
+	 */
+	void GetFilterData_AssumesLocked(FShapeData& ShapeData, bool bForceSimpleAsComplex = false);
+
+	/** Set initialisation flags on a provided shape
+	 *	@param UseCollisionEnabled Whether collision is enabled for the shape
+	 *	@param PShape The shape to set the flags on
+	 *	@param SceneType Which scene we are using
+	 *	@param bUseComplexAsSimple Whether to use complex collision as simple
+	 */
+	DEPRECATED(4.8, "Please call SetShapeFlags_AssumesLocked and make sure you obtain the appropriate PhysX scene locks")
+	void SetShapeFlags(TEnumAsByte<ECollisionEnabled::Type> UseCollisionEnabled, physx::PxShape* PShape, EPhysicsSceneType SceneType, const bool bUseComplexAsSimple = false)
+	{
+		SetShapeFlags_AssumesLocked(UseCollisionEnabled, PShape, SceneType, bUseComplexAsSimple);
+	}
 
 	/** Set initialisation flags on a provided shape
 	 *	@param UseCollisionEnabled Whether collision is enabled for the shape
 	 *	@param PShape The shape to set the flags on
 	 *	@param SceneType Which scene we are using
 	 *	@param bUseComplexAsSimple Whether to use complex collision as simple
+	 *  Note: This function is not thread safe. Make sure to obtain the appropriate PhysX scene locks before calling this function
 	 */
-	void SetShapeFlags(TEnumAsByte<ECollisionEnabled::Type> UseCollisionEnabled, physx::PxShape* PShape, EPhysicsSceneType SceneType, const bool bUseComplexAsSimple = false);
+	void SetShapeFlags_AssumesLocked(TEnumAsByte<ECollisionEnabled::Type> UseCollisionEnabled, physx::PxShape* PShape, EPhysicsSceneType SceneType, const bool bUseComplexAsSimple = false);
 
 	/** Populate the flag fields of the provided FShapeData with correct initialisation flags
 	 *	@param ShapeData ShapeData to populate
 	 *	@param UseCollisionEnabled Whether collision is enabled for this instance
 	 *	@param bUseComplexAsSimple Whether to use complex collision as simple
 	 */
-	void GetShapeFlags(FShapeData& ShapeData, TEnumAsByte<ECollisionEnabled::Type> UseCollisionEnabled, const bool bUseComplexAsSimple = false);
+	DEPRECATED(4.8, "Please call GetShapeFlags_AssumesLocked and make sure you obtain the appropriate PhysX scene locks")
+	void GetShapeFlags(FShapeData& ShapeData, TEnumAsByte<ECollisionEnabled::Type> UseCollisionEnabled, const bool bUseComplexAsSimple = false)
+	{
+		GetShapeFlags_AssumesLocked(ShapeData, UseCollisionEnabled, bUseComplexAsSimple);
+	}
 
-#endif // WITH_PHYSX
+	/** Populate the flag fields of the provided FShapeData with correct initialisation flags
+	 *	@param ShapeData ShapeData to populate
+	 *	@param UseCollisionEnabled Whether collision is enabled for this instance
+	 *	@param bUseComplexAsSimple Whether to use complex collision as simple
+	 *  Note: This function is not thread safe. Make sure to obtain the appropriate PhysX scene locks before calling this function
+	 */
+	void GetShapeFlags_AssumesLocked(FShapeData& ShapeData, TEnumAsByte<ECollisionEnabled::Type> UseCollisionEnabled, const bool bUseComplexAsSimple = false);
+
+	/**
+	 * Return the PxRigidActor from the given scene (see EPhysicsSceneType), if SceneType is in the range [0, PST_MAX).
+	 * If SceneType < 0, the PST_Sync actor is returned if it is not NULL, otherwise the PST_Async actor is returned.
+	 * Invalid scene types will cause NULL to be returned.
+	 */
+	DEPRECATED(4.8, "Please call GetPxRigidActor_AssumesLocked and make sure you obtain the appropriate PhysX scene locks")
+	physx::PxRigidActor* GetPxRigidActor(int32 SceneType = -1) const
+	{
+		return GetPxRigidActor_AssumesLocked(SceneType);
+	}
+
+	/**
+	 * Return the PxRigidActor from the given scene (see EPhysicsSceneType), if SceneType is in the range [0, PST_MAX).
+	 * If SceneType < 0, the PST_Sync actor is returned if it is not NULL, otherwise the PST_Async actor is returned.
+	 * Invalid scene types will cause NULL to be returned.
+	 * Note: Reading/writing from/to PxRigidActor is not thread safe. If you use the actor make sure to obtain the appropriate PhysX scene lock
+	 */
+	physx::PxRigidActor* GetPxRigidActor_AssumesLocked(int32 SceneType = -1) const;
+
+
+	/** Return the PxRigidDynamic if it exists in one of the scenes (NULL otherwise).  Currently a PxRigidDynamic can exist in only one of the two scenes. */
+	DEPRECATED(4.8, "Please call GetPxRigidDynamic_AssumesLocked and make sure you obtain the appropriate PhysX scene locks")
+	physx::PxRigidDynamic* GetPxRigidDynamic() const
+	{
+		return GetPxRigidDynamic_AssumesLocked();
+	}
+
+	/** Return the PxRigidDynamic if it exists in one of the scenes (NULL otherwise).  Currently a PxRigidDynamic can exist in only one of the two scenes.
+	 *  Note: This function assumes the appropriate scene lock has been obtained.
+	 *  @see ExecuteOnPxRigidDynamicReadOnly, ExecuteOnPxRigidDynamicReadWrite
+	 */
+	physx::PxRigidDynamic* GetPxRigidDynamic_AssumesLocked() const;
+
+	/** Return the PxRigidBody if it exists in one of the scenes (NULL otherwise).  Currently a PxRigidBody can exist in only one of the two scenes. */
+	DEPRECATED(4.8, "Please call GetPxRigidBody_AssumesLocked and make sure you obtain the appropriate PhysX scene locks")
+	physx::PxRigidBody* GetPxRigidBody() const
+	{
+		return GetPxRigidBody_AssumesLocked();
+	}
+
+	/** Return the PxRigidBody if it exists in one of the scenes (NULL otherwise).  Currently a PxRigidBody can exist in only one of the two scenes.
+	 *  Note: This function assumes the appropriate scene lock has been obtained.
+	 *  @see ExecuteOnPxRigidBodyReadOnly, ExecuteOnPxRigidBodyReadWrite
+	 */
+	physx::PxRigidBody* GetPxRigidBody_AssumesLocked() const;
+
+	
+
+	/** 
+	 *	Utility to get all the shapes from a FBodyInstance 
+	 *	Shapes belonging to sync actor are first, then async. Number of shapes belonging to sync actor is returned.
+	 */
+	DEPRECATED(4.8, "Please call GetAllShapes_AssumesLocked and make sure you obtain the appropriate PhysX scene locks")
+	TArray<physx::PxShape*> GetAllShapes(int32& OutNumSyncShapes) const
+	{
+		TArray<physx::PxShape*> Shapes;
+		OutNumSyncShapes = GetAllShapes_AssumesLocked(Shapes);
+		return Shapes;
+	}
+
+	/** 
+	 *	Utility to get all the shapes from a FBodyInstance 
+	 *	Shapes belonging to sync actor are first, then async. Number of shapes belonging to sync actor is returned.
+	 *	NOTE: This function is not thread safe. You must hold the PhysX scene lock while calling it and reading/writing from the shapes
+	 */
+	int32 GetAllShapes_AssumesLocked(TArray<physx::PxShape*>& OutShapes) const;
+#endif	//WITH_PHYSX
 
 #endif	//UE_WITH_PHYSICS
 
@@ -520,28 +643,12 @@ public:
 	/** Indicates whether this body should use the async scene. Must be called before body is init'd, will assert otherwise. Will have no affect if there is no async scene. */
 	void SetUseAsyncScene(bool bNewUseAsyncScene);
 
-#if WITH_PHYSX
-	/**
-	 * Return the PxRigidActor from the given scene (see EPhysicsSceneType), if SceneType is in the range [0, PST_MAX).
-	 * If SceneType < 0, the PST_Sync actor is returned if it is not NULL, otherwise the PST_Async actor is returned.
-	 * Invalid scene types will cause NULL to be returned.
-	 */
-	physx::PxRigidActor* GetPxRigidActor(int32 SceneType = -1) const;
-	/** Return the PxRigidDynamic if it exists in one of the scenes (NULL otherwise).  Currently a PxRigidDynamic can exist in only one of the two scenes. */
-	physx::PxRigidDynamic* GetPxRigidDynamic() const;
-	/** Return the PxRigidBody if it exists in one of the scenes (NULL otherwise).  Currently a PxRigidBody can exist in only one of the two scenes. */
-	physx::PxRigidBody* GetPxRigidBody() const;
-
-	/** 
-	 *	Utility to get all the shapes from a FBodyInstance 
-	 *	Shapes belonging to sync actor are first, then async. Number of shapes belonging to sync actor is returned.
-	 */
-	TArray<physx::PxShape*> GetAllShapes(int32& OutNumSyncShapes) const;
-#endif	//WITH_PHYSX
-
 	/** Returns true if the body is not static */
 	bool IsDynamic() const;
 
+	/** Returns true if the body is non kinematic*/
+	bool IsNonKinematic() const;
+
 	/** Returns the body's mass */
 	float GetBodyMass() const;
 	/** Return bounds of physics representation */
@@ -597,6 +704,9 @@ public:
 	/** Get current transform in world space from physics body. */
 	FTransform GetUnrealWorldTransform() const;
 
+	/** Get current transform in world space from physics body. */
+	FTransform GetUnrealWorldTransform_AssumesLocked() const;
+
 	/**
 	 *	Move the physics body to a new pose.
 	 *	@param	bTeleport	If true, no velocity is inferred on the kinematic body from this movement, but it moves right away.
@@ -606,12 +716,21 @@ public:
 	/** Get current velocity in world space from physics body. */
 	FVector GetUnrealWorldVelocity() const;
 
+	/** Get current velocity in world space from physics body. */
+	FVector GetUnrealWorldVelocity_AssumesLocked() const;
+
 	/** Get current angular velocity in world space from physics body. */
 	FVector GetUnrealWorldAngularVelocity() const;
 
+	/** Get current angular velocity in world space from physics body. */
+	FVector GetUnrealWorldAngularVelocity_AssumesLocked() const;
+
 	/** Get current velocity of a point on this physics body, in world space. Point is specified in world space. */
 	FVector GetUnrealWorldVelocityAtPoint(const FVector& Point) const;
 
+	/** Get current velocity of a point on this physics body, in world space. Point is specified in world space. */
+	FVector GetUnrealWorldVelocityAtPoint_AssumesLocked(const FVector& Point) const;
+
 	/** Set physical material override for this body */
 	void SetPhysMaterialOverride(class UPhysicalMaterial* NewPhysMaterial);
 
@@ -676,9 +795,23 @@ public:
 	void UpdatePhysicalMaterials();
 
 #if WITH_PHYSX
-	static void ApplyMaterialToShape(physx::PxShape* PShape, physx::PxMaterial* PSimpleMat, TArray<UPhysicalMaterial*>& ComplexPhysMats);
+	DEPRECATED(4.8, "Please call ApplyMaterialToShape_AssumesLocked and make sure you obtain the appropriate PhysX scene locks")
+	static void ApplyMaterialToShape(physx::PxShape* PShape, physx::PxMaterial* PSimpleMat, TArray<UPhysicalMaterial*>& ComplexPhysMats)
+	{
+		ApplyMaterialToShape_AssumesLocked(PShape, PSimpleMat, ComplexPhysMats);
+	}
+
+	/** Note: This function is not thread safe. Make sure you obtain the appropriate PhysX scene lock before calling it*/
+	static void ApplyMaterialToShape_AssumesLocked(physx::PxShape* PShape, physx::PxMaterial* PSimpleMat, TArray<UPhysicalMaterial*>& ComplexPhysMats);
+
+	DEPRECATED(4.8, "Please call ApplyMaterialToInstanceShapes_AssumesLocked and make sure you obtain the appropriate PhysX scene locks")
+	void ApplyMaterialToInstanceShapes(physx::PxMaterial* PSimpleMat, TArray<UPhysicalMaterial*>& ComplexPhysMats)
+	{
+		ApplyMaterialToInstanceShapes_AssumesLocked(PSimpleMat, ComplexPhysMats);
+	}
 
-	void ApplyMaterialToInstanceShapes(physx::PxMaterial* PSimpleMat, TArray<UPhysicalMaterial*>& ComplexPhysMats);
+	/** Note: This function is not thread safe. Make sure you obtain the appropriate PhysX scene lock before calling it*/
+	void ApplyMaterialToInstanceShapes_AssumesLocked(physx::PxMaterial* PSimpleMat, TArray<UPhysicalMaterial*>& ComplexPhysMats);
 #endif
 
 	/** Update the instances collision filtering data */
@@ -720,7 +853,23 @@ public:
 	 *  @param  OutMTD			The minimum translation direction needed to push the shape out of this BodyInstance. (Optional)
 	 *  @return true if PrimComp overlaps this component at the specified location/rotation
 	 */
-	bool OverlapPhysX(const physx::PxGeometry& Geom, const physx::PxTransform&  ShapePose, FMTDResult* OutMTD = nullptr) const;
+	DEPRECATED(4.8, "Please call OverlapPhysX_AssumesLocked and make sure you obtain the appropriate PhysX scene locks")
+	bool OverlapPhysX(const physx::PxGeometry& Geom, const physx::PxTransform&  ShapePose, FMTDResult* OutMTD = nullptr) const
+	{
+		return OverlapPhysX_AssumesLocked(Geom, ShapePose, OutMTD);
+	}
+
+	/**
+	 *  Test if the bodyinstance overlaps with the geometry in the Pos/Rot
+	 *
+	 *	@param	PGeom			Geometry it would like to test
+	 *  @param  ShapePose       Transform information in world. Use U2PTransform to convert from FTransform
+	 *  @param  OutMTD			The minimum translation direction needed to push the shape out of this BodyInstance. (Optional)
+	 *  @return true if PrimComp overlaps this component at the specified location/rotation
+	 *
+	 *	Note: This function is not thread safe. Make sure to obtain the appropriate PhysX scene locks before calling this function
+	 */
+	bool OverlapPhysX_AssumesLocked(const physx::PxGeometry& Geom, const physx::PxTransform&  ShapePose, FMTDResult* OutMTD = nullptr) const;
 #endif	//WITH_PHYSX
 
 	/**
@@ -734,6 +883,16 @@ public:
 	 */
 	bool OverlapTest(const FVector& Position, const FQuat& Rotation, const struct FCollisionShape& CollisionShape, FMTDResult* OutMTD = nullptr) const;
 
+	/**
+	 *  Test if the bodyinstance overlaps with the specified body instances
+	 *
+	 *  @param  Position		Position to place our shapes at before testing (shapes of this BodyInstance)
+	 *  @param  Rotation		Rotation to apply to our shapes before testing (shapes of this BodyInstance)
+	 *  @param  Bodies			The bodies we are testing for overlap with. These bodies will be in world space already
+	 *  @return true if any of the bodies passed in overlap with this
+	 */
+	bool OverlapTestForBodies(const FVector& Position, const FQuat& Rotation, const TArray<FBodyInstance*>& Bodies) const;
+
 	/**
 	 *  Determines the set of components that this body instance would overlap with at the supplied location/rotation
 	 *  @param  InOutOverlaps   Array of overlaps found between this component in specified pose and the world (new overlaps will be appended, the array is not cleared!)
@@ -836,7 +995,7 @@ private:
 
 #if WITH_BOX2D
 
-protected:
+public:
 	class b2Body* BodyInstancePtr;
 
 #endif	//WITH_BOX2D
diff --git a/Engine/Source/Runtime/Engine/Classes/PhysicsEngine/BodySetup.h b/Engine/Source/Runtime/Engine/Classes/PhysicsEngine/BodySetup.h
--- a/Engine/Source/Runtime/Engine/Classes/PhysicsEngine/BodySetup.h
+++ b/Engine/Source/Runtime/Engine/Classes/PhysicsEngine/BodySetup.h
@@ -217,14 +217,23 @@ public:
 	/** 
 	 *   Add the shapes defined by this body setup to the supplied PxRigidBody. 
 	 */
-	ENGINE_API void AddShapesToRigidActor(FBodyInstance* OwningInstance, physx::PxRigidActor* PDestActor, EPhysicsSceneType SceneType, FVector& Scale3D, physx::PxMaterial* SimpleMaterial, TArray<UPhysicalMaterial*>& ComplexMaterials, FShapeData& ShapeData, const FTransform& RelativeTM = FTransform::Identity, TArray<physx::PxShape*>* NewShapes = NULL);
+	DEPRECATED(4.8, "Please call AddShapesToRigidActor_AssumesLocked and make sure you obtain the appropriate PhysX scene locks")
+	ENGINE_API void AddShapesToRigidActor(FBodyInstance* OwningInstance, physx::PxRigidActor* PDestActor, EPhysicsSceneType SceneType, FVector& Scale3D, physx::PxMaterial* SimpleMaterial, TArray<UPhysicalMaterial*>& ComplexMaterials, FShapeData& ShapeData, const FTransform& RelativeTM = FTransform::Identity, TArray<physx::PxShape*>* NewShapes = NULL)
+	{
+		AddShapesToRigidActor_AssumesLocked(OwningInstance, PDestActor, SceneType, Scale3D, SimpleMaterial, ComplexMaterials, ShapeData, RelativeTM, NewShapes);
+	}
+
+	/** 
+	 *   Add the shapes defined by this body setup to the supplied PxRigidBody. 
+	 */
+	ENGINE_API void AddShapesToRigidActor_AssumesLocked(FBodyInstance* OwningInstance, physx::PxRigidActor* PDestActor, EPhysicsSceneType SceneType, FVector& Scale3D, physx::PxMaterial* SimpleMaterial, TArray<UPhysicalMaterial*>& ComplexMaterials, FShapeData& ShapeData, const FTransform& RelativeTM = FTransform::Identity, TArray<physx::PxShape*>* NewShapes = NULL);
 
 private:
-	void AddSpheresToRigidActor(physx::PxRigidActor* PDestActor, const FTransform& RelativeTM, float MinScale, float MinScaleAbs, TArray<physx::PxShape*>* NewShapes) const;
-	void AddBoxesToRigidActor(physx::PxRigidActor* PDestActor, const FTransform& RelativeTM, const FVector& Scale3D, const FVector& Scale3DAbs, TArray<physx::PxShape*>* NewShapes) const;
-	void AddSphylsToRigidActor(physx::PxRigidActor* PDestActor, const FTransform& RelativeTM, const FVector& Scale3D, const FVector& Scale3DAbs, TArray<physx::PxShape*>* NewShapes) const;
-	void AddConvexElemsToRigidActor(physx::PxRigidActor* PDestActor, const FTransform& RelativeTM, const FVector& Scale3D, const FVector& Scale3DAbs, TArray<physx::PxShape*>* NewShapes) const;
-	void AddTriMeshToRigidActor(physx::PxRigidActor* PDestActor, const FVector& Scale3D, const FVector& Scale3DAbs, TArray<physx::PxShape*>* NewShapes) const;
+	void AddSpheresToRigidActor_AssumesLocked(physx::PxRigidActor* PDestActor, const FTransform& RelativeTM, float MinScale, float MinScaleAbs, TArray<physx::PxShape*>* NewShapes) const;
+	void AddBoxesToRigidActor_AssumesLocked(physx::PxRigidActor* PDestActor, const FTransform& RelativeTM, const FVector& Scale3D, const FVector& Scale3DAbs, TArray<physx::PxShape*>* NewShapes) const;
+	void AddSphylsToRigidActor_AssumesLocked(physx::PxRigidActor* PDestActor, const FTransform& RelativeTM, const FVector& Scale3D, const FVector& Scale3DAbs, TArray<physx::PxShape*>* NewShapes) const;
+	void AddConvexElemsToRigidActor_AssumesLocked(physx::PxRigidActor* PDestActor, const FTransform& RelativeTM, const FVector& Scale3D, const FVector& Scale3DAbs, TArray<physx::PxShape*>* NewShapes) const;
+	void AddTriMeshToRigidActor_AssumesLocked(physx::PxRigidActor* PDestActor, const FVector& Scale3D, const FVector& Scale3DAbs, TArray<physx::PxShape*>* NewShapes) const;
 #endif // WITH_PHYSX
 
 };
diff --git a/Engine/Source/Runtime/Engine/Classes/PhysicsEngine/ConstraintInstance.h b/Engine/Source/Runtime/Engine/Classes/PhysicsEngine/ConstraintInstance.h
--- a/Engine/Source/Runtime/Engine/Classes/PhysicsEngine/ConstraintInstance.h
+++ b/Engine/Source/Runtime/Engine/Classes/PhysicsEngine/ConstraintInstance.h
@@ -498,10 +498,13 @@ public:
 
 private:
 #if WITH_PHYSX 
-	bool CreatePxJoint(physx::PxRigidActor* PActor1, physx::PxRigidActor* PActor2, physx::PxScene* PScene, const float Scale);
-	void UpdateConstraintFlags();
-	void UpdateAverageMass(const physx::PxRigidActor* PActor1, const physx::PxRigidActor* PActor2);
-	physx::PxD6Joint* GetUnbrokenJoint() const;
+	bool CreatePxJoint_AssumesLocked(physx::PxRigidActor* PActor1, physx::PxRigidActor* PActor2, physx::PxScene* PScene, const float Scale);
+	void UpdateConstraintFlags_AssumesLocked();
+	void UpdateAverageMass_AssumesLocked(const physx::PxRigidActor* PActor1, const physx::PxRigidActor* PActor2);
+	physx::PxD6Joint* GetUnbrokenJoint_AssumesLocked() const;
+
+	bool ExecuteOnUnbrokenJointReadOnly(TFunctionRef<void(const physx::PxD6Joint*)> Func) const;
+	bool ExecuteOnUnbrokenJointReadWrite(TFunctionRef<void(physx::PxD6Joint*)> Func) const;
 #endif
 
 	void UpdateBreakable();
diff --git a/Engine/Source/Runtime/Engine/Classes/PhysicsEngine/PhysicsCollisionHandler.h b/Engine/Source/Runtime/Engine/Classes/PhysicsEngine/PhysicsCollisionHandler.h
--- a/Engine/Source/Runtime/Engine/Classes/PhysicsEngine/PhysicsCollisionHandler.h
+++ b/Engine/Source/Runtime/Engine/Classes/PhysicsEngine/PhysicsCollisionHandler.h
@@ -35,10 +35,24 @@ class ENGINE_API UPhysicsCollisionHandler : public UObject
 	}
 
 	/** Gives game-specific ability to handle all physics collisions in one place. This is a good place to play sounds and spawn effects, as it does not require special object-specific code. */
-	virtual void HandlePhysicsCollisions(const TArray<FCollisionNotifyInfo>& PendingCollisionNotifies);
+	DEPRECATED(4.8, "Please call HandlePhysicsCollisions_AssumesLocked and make sure you obtain the appropriate PhysX scene locks")
+	virtual void HandlePhysicsCollisions(const TArray<FCollisionNotifyInfo>& PendingCollisionNotifies)
+	{
+		HandlePhysicsCollisions_AssumesLocked(PendingCollisionNotifies);
+	}
+
+	/** Gives game-specific ability to handle all physics collisions in one place. This is a good place to play sounds and spawn effects, as it does not require special object-specific code. */
+	virtual void HandlePhysicsCollisions_AssumesLocked(const TArray<FCollisionNotifyInfo>& PendingCollisionNotifies);
+
+	/** Handle a single */
+	DEPRECATED(4.8, "Please call DefaultHandleCollision_AssumesLocked and make sure you obtain the appropriate PhysX scene locks")
+	void DefaultHandleCollision(const FRigidBodyCollisionInfo& MyInfo, const FRigidBodyCollisionInfo& OtherInfo, const FCollisionImpactData& RigidCollisionData)
+	{
+		DefaultHandleCollision_AssumesLocked(MyInfo, OtherInfo, RigidCollisionData);
+	}
 
 	/** Handle a single */
-	void DefaultHandleCollision(const FRigidBodyCollisionInfo& MyInfo, const FRigidBodyCollisionInfo& OtherInfo, const FCollisionImpactData& RigidCollisionData);
+	void DefaultHandleCollision_AssumesLocked(const FRigidBodyCollisionInfo& MyInfo, const FRigidBodyCollisionInfo& OtherInfo, const FCollisionImpactData& RigidCollisionData);
 
 	/** Called after collision handler is allocated for a world. GetWorld should be valid inside this function. */
 	virtual void InitCollisionHandler() {}
diff --git a/Engine/Source/Runtime/Engine/Classes/PhysicsEngine/PhysicsSettings.h b/Engine/Source/Runtime/Engine/Classes/PhysicsEngine/PhysicsSettings.h
--- a/Engine/Source/Runtime/Engine/Classes/PhysicsEngine/PhysicsSettings.h
+++ b/Engine/Source/Runtime/Engine/Classes/PhysicsEngine/PhysicsSettings.h
@@ -114,6 +114,10 @@ class ENGINE_API UPhysicsSettings
 	UPROPERTY(config, EditAnywhere, AdvancedDisplay, Category=Simulation)
 	bool bEnableAsyncScene;
 
+	/** Whether to warn when physics locks are used incorrectly. Turning this off is not recommended and should only be used by very advanced users. */
+	UPROPERTY(config, EditAnywhere, AdvancedDisplay, Category = Simulation)
+	bool bWarnMissingLocks;
+
 	/** Can 2D physics be used (Box2D)? */
 	UPROPERTY(config, EditAnywhere, Category = Simulation)
 	bool bEnable2DPhysics;
diff --git a/Engine/Source/Runtime/Engine/Private/Collision/WorldCollision.cpp b/Engine/Source/Runtime/Engine/Private/Collision/WorldCollision.cpp
--- a/Engine/Source/Runtime/Engine/Private/Collision/WorldCollision.cpp
+++ b/Engine/Source/Runtime/Engine/Private/Collision/WorldCollision.cpp
@@ -503,11 +503,8 @@ bool UWorld::ComponentSweepMulti(TArray<struct FHitResult>& OutHits, class UPrim
 	bool bHaveBlockingHit = false;
 
 #if WITH_PHYSX
-	if (PxRigidActor* PRigidActor = PrimComp->BodyInstance.GetPxRigidActor())
+	ExecuteOnPxRigidActorReadOnly(&PrimComp->BodyInstance, [&] (const PxRigidActor* PRigidActor)
 	{
-		PxScene* const PScene = PRigidActor->getScene();
-		SCOPED_SCENE_READ_LOCK(PScene);
-
 		// Get all the shapes from the actor
 		TArray<PxShape*, TInlineAllocator<32>> PShapes;
 		{
@@ -548,7 +545,7 @@ bool UWorld::ComponentSweepMulti(TArray<struct FHitResult>& OutHits, class UPrim
 				OutHits.Append(Hits);
 			}
 		}
-	}
+	});
 #endif //WITH_PHYSX
 
 	//@TODO: BOX2D: Implement UWorld::ComponentSweepMulti
diff --git a/Engine/Source/Runtime/Engine/Private/Components/DestructibleComponent.cpp b/Engine/Source/Runtime/Engine/Private/Components/DestructibleComponent.cpp
--- a/Engine/Source/Runtime/Engine/Private/Components/DestructibleComponent.cpp
+++ b/Engine/Source/Runtime/Engine/Private/Components/DestructibleComponent.cpp
@@ -558,7 +558,7 @@ void UDestructibleComponent::AddRadialImpulse(FVector Origin, float Radius, floa
 					LockedScene->lockWrite();
 					LockedScene->lockRead();
 				}
-				AddRadialImpulseToPxRigidBody(*PActor, Origin, Radius, Strength, Falloff, bVelChange);
+				AddRadialImpulseToPxRigidBody_AssumesLocked(*PActor, Origin, Radius, Strength, Falloff, bVelChange);
 			}
 		}
 
@@ -601,7 +601,7 @@ void UDestructibleComponent::AddRadialForce(FVector Origin, float Radius, float
 					LockedScene->lockRead();
 				}
 
-				AddRadialForceToPxRigidBody(*PActor, Origin, Radius, Strength, Falloff, bAccelChange);
+				AddRadialForceToPxRigidBody_AssumesLocked(*PActor, Origin, Radius, Strength, Falloff, bAccelChange);
 			}
 
 			if (LockedScene)
diff --git a/Engine/Source/Runtime/Engine/Private/Components/PrimitiveComponent.cpp b/Engine/Source/Runtime/Engine/Private/Components/PrimitiveComponent.cpp
--- a/Engine/Source/Runtime/Engine/Private/Components/PrimitiveComponent.cpp
+++ b/Engine/Source/Runtime/Engine/Private/Components/PrimitiveComponent.cpp
@@ -1744,8 +1744,6 @@ bool UPrimitiveComponent::SweepComponent(struct FHitResult& OutHit, const FVecto
 
 bool UPrimitiveComponent::ComponentOverlapComponent(class UPrimitiveComponent* PrimComp, const FVector Pos, const FRotator Rot, const struct FCollisionQueryParams& Params)
 {
-	//@TODO: BOX2D: Implement UPrimitiveComponent::ComponentOverlapComponent
-
 	// if target is skeletalmeshcomponent and do not support singlebody physics
 	USkeletalMeshComponent * OtherComp = Cast<USkeletalMeshComponent>(PrimComp);
 	if (OtherComp)
@@ -1753,63 +1751,14 @@ bool UPrimitiveComponent::ComponentOverlapComponent(class UPrimitiveComponent* P
 		UE_LOG(LogCollision, Log, TEXT("ComponentOverlapMulti : (%s) Does not support skeletalmesh with Physics Asset"), *PrimComp->GetPathName());
 		return false;
 	}
-#if WITH_PHYSX
-	// will have to do per component - default single body or physicsinstance 
-	const PxRigidActor* TargetRigidBody = (PrimComp)? PrimComp->BodyInstance.GetPxRigidActor():NULL;
-	if (TargetRigidBody==NULL || TargetRigidBody->getNbShapes()==0)
-	{
-		return false;
-	}
-
-	// calculate the test global pose of the actor
-	PxTransform PTestGlobalPose = U2PTransform(FTransform(Rot, Pos));
-	// Get all the shapes from the actor
-	TArray<PxShape*> PTargetShapes;
-	PTargetShapes.AddZeroed(TargetRigidBody->getNbShapes());
-	int32 NumTargetShapes = TargetRigidBody->getShapes(PTargetShapes.GetData(), PTargetShapes.Num());
 
-	bool bHaveOverlap = false;
-
-	for (int32 TargetShapeIdx=0; TargetShapeIdx<PTargetShapes.Num(); ++TargetShapeIdx)
+	if(FBodyInstance* BI = PrimComp->GetBodyInstance())
 	{
-		const PxShape * PTargetShape = PTargetShapes[TargetShapeIdx];
-		check (PTargetShape);
-
-		// Calc shape global pose
-		PxTransform PShapeGlobalPose = PTestGlobalPose.transform(PTargetShape->getLocalPose());
-
-		GET_GEOMETRY_FROM_SHAPE(PGeom, PTargetShape);
-
-		if(PGeom != NULL)
-		{
-			bHaveOverlap = BodyInstance.OverlapPhysX(*PGeom, PShapeGlobalPose);
-
-#if !(UE_BUILD_SHIPPING || UE_BUILD_TEST)
-			if((GetWorld()->DebugDrawTraceTag != NAME_None) && (GetWorld()->DebugDrawTraceTag == Params.TraceTag))
-			{
-				TArray<FOverlapResult> Overlaps;
-				if (bHaveOverlap)
-				{
-					FOverlapResult Result;
-					Result.Actor = PrimComp->GetOwner();
-					Result.Component = PrimComp;
-					Result.bBlockingHit = true;
-					Overlaps.Add(Result);
-				}
-
-				DrawGeomOverlaps(GetWorld(), *PGeom, PShapeGlobalPose, Overlaps, DebugLineLifetime);
-			}
-#endif //!(UE_BUILD_SHIPPING || UE_BUILD_TEST)
-
-			if (bHaveOverlap)
-			{
-				break;
-			}
-		}
+		TArray<FBodyInstance*> Bodies;
+		Bodies.Add(GetBodyInstance());
+		return BI->OverlapTestForBodies(Pos, Rot.Quaternion(), Bodies);
 	}
 
-	return bHaveOverlap;
-#endif //WITH_PHYSX
 	return false;
 }
 
diff --git a/Engine/Source/Runtime/Engine/Private/PhysicsEngine/BodyInstance.cpp b/Engine/Source/Runtime/Engine/Private/PhysicsEngine/BodyInstance.cpp
--- a/Engine/Source/Runtime/Engine/Private/PhysicsEngine/BodyInstance.cpp
+++ b/Engine/Source/Runtime/Engine/Private/PhysicsEngine/BodyInstance.cpp
@@ -319,32 +319,29 @@ bool ShapeBoundToBody(const PxShape * PShape, const FBodyInstance* SubBody)
 	return (SubBody->WeldParent == NULL && BI == NULL) || (BI == SubBody && BI->WeldParent != NULL);
 }
 
-TArray<PxShape*> FBodyInstance::GetAllShapes(int32& OutNumSyncShapes) const
+int32 FBodyInstance::GetAllShapes_AssumesLocked(TArray<PxShape*>& OutShapes) const
 {
+	int32 NumSyncShapes = 0;
+	OutShapes.Empty();
 	// grab shapes from sync actor
-	TArray<PxShape*> PSyncShapes;
-	if( RigidActorSync != NULL )
+	if(RigidActorSync)
 	{
-		SCOPED_SCENE_READ_LOCK(RigidActorSync->getScene());
-		PSyncShapes.AddZeroed( RigidActorSync->getNbShapes() );
-		RigidActorSync->getShapes(PSyncShapes.GetData(), PSyncShapes.Num());
+		NumSyncShapes = RigidActorSync->getNbShapes();
+		OutShapes.AddZeroed(NumSyncShapes);
+		RigidActorSync->getShapes(OutShapes.GetData(), NumSyncShapes);
 	}
 
 	// grab shapes from async actor
-	TArray<PxShape*> PAsyncShapes;
 	if( RigidActorAsync != NULL )
 	{
-		SCOPED_SCENE_READ_LOCK( RigidActorAsync->getScene() );
-		PAsyncShapes.AddZeroed( RigidActorAsync->getNbShapes() );
-		RigidActorAsync->getShapes(PAsyncShapes.GetData(), PAsyncShapes.Num());
+		const int32 NumAsyncShapes = RigidActorAsync->getNbShapes();
+		OutShapes.AddZeroed(NumAsyncShapes);
+		RigidActorAsync->getShapes(OutShapes.GetData() + NumSyncShapes, NumAsyncShapes);
 	}
 
-	OutNumSyncShapes = PSyncShapes.Num();
-
-	TArray<PxShape*> PAllShapes = PSyncShapes;
-	PAllShapes.Append(PAsyncShapes);
-	return PAllShapes;
+	return NumSyncShapes;
 }
+
 #endif
 
 void FBodyInstance::UpdateTriMeshVertices(const TArray<FVector> & NewPositions)
@@ -352,15 +349,14 @@ void FBodyInstance::UpdateTriMeshVertices(const TArray<FVector> & NewPositions)
 #if WITH_PHYSX
 	if (BodySetup.IsValid())
 	{
-		SCOPED_SCENE_WRITE_LOCK((RigidActorSync ? RigidActorSync->getScene() : NULL));
-		SCOPED_SCENE_WRITE_LOCK((RigidActorAsync ? RigidActorAsync->getScene() : NULL));
-
+		ExecuteOnPhysicsReadWrite([&]
+		{
 		BodySetup->UpdateTriMeshVertices(NewPositions);
 
 		//after updating the vertices we must call setGeometry again to update any shapes referencing the mesh
 
-		int32 SyncShapeCount = 0;
-		TArray<PxShape *> PShapes = GetAllShapes(SyncShapeCount);
+			TArray<PxShape *> PShapes;
+			const int32 SyncShapeCount = GetAllShapes_AssumesLocked(PShapes);
 		PxTriangleMeshGeometry PTriangleMeshGeometry;
 		for (int32 ShapeIdx = 0; ShapeIdx < PShapes.Num(); ShapeIdx++)
 		{
@@ -371,6 +367,7 @@ void FBodyInstance::UpdateTriMeshVertices(const TArray<FVector> & NewPositions)
 				PShape->setGeometry(PTriangleMeshGeometry);
 			}
 		}
+		});
 	}
 #endif
 }
@@ -388,28 +385,10 @@ void FBodyInstance::UpdatePhysicalMaterials()
 #if WITH_PHYSX
 	PxMaterial* PSimpleMat = SimplePhysMat->GetPhysXMaterial();
 	check(PSimpleMat != NULL);
-
-	if (RigidActorSync != NULL)
-	{
-		SCENE_LOCK_WRITE(RigidActorSync->getScene());
-	}
-
-	if (RigidActorAsync != NULL)
-	{
-		SCENE_LOCK_WRITE(RigidActorAsync->getScene());
-	}
-
-	ApplyMaterialToInstanceShapes(PSimpleMat, ComplexPhysMats);
-
-	if (RigidActorSync != NULL)
+	ExecuteOnPhysicsReadWrite([&]()
 	{
-		SCENE_UNLOCK_WRITE(RigidActorSync->getScene());
-	}
-
-	if (RigidActorAsync != NULL)
-	{
-		SCENE_UNLOCK_WRITE(RigidActorAsync->getScene());
-	}
+		ApplyMaterialToInstanceShapes_AssumesLocked(PSimpleMat, ComplexPhysMats);
+	});
 #endif
 
 #if WITH_BOX2D
@@ -655,18 +634,14 @@ ECollisionEnabled::Type FBodyInstance::GetCollisionEnabled() const
 #if WITH_PHYSX
 void FBodyInstance::UpdatePhysicsShapeFilterData(uint32 SkelMeshCompID, bool bUseComplexAsSimple, bool bUseSimpleAsComplex, bool bPhysicsStatic, TEnumAsByte<ECollisionEnabled::Type> * CollisionEnabledOverride, FCollisionResponseContainer * ResponseOverride, bool * bNotifyOverride)
 {
-	if (PxRigidActor* PActor = GetPxRigidActor())
+	ExecuteOnPhysicsReadWrite([&]
 	{
-		// Iterate over all shapes and assign filterdata
-		int32 NumSyncShapes = 0;
-		TArray<PxShape*> AllShapes = GetAllShapes(NumSyncShapes);
-
-		PxScene* AsyncScene = (RigidActorAsync != NULL) ? RigidActorAsync->getScene() : NULL;
-		
-		if(AsyncScene)
+		if (PxRigidActor* PActor = GetPxRigidActor_AssumesLocked())
 		{
-			SCENE_LOCK_WRITE(AsyncScene);
-		}
+
+		// Iterate over all shapes and assign filterdata
+			TArray<PxShape*> AllShapes;
+			const int32 NumSyncShapes = GetAllShapes_AssumesLocked(AllShapes);
 
 		bool bUpdateMassProperties = false;
 		for (int32 ShapeIdx = 0; ShapeIdx < AllShapes.Num(); ShapeIdx++)
@@ -767,7 +742,7 @@ void FBodyInstance::UpdatePhysicsShapeFilterData(uint32 SkelMeshCompID, bool bUs
 					}
 
 					// enable swept bounds for CCD for this shape
-					PxRigidBody* PBody = GetPxRigidActor()->is<PxRigidBody>();
+						PxRigidBody* PBody = PActor->is<PxRigidBody>();
 					if (bSimCollision && !bPhysicsStatic && bUseCCD && PBody)
 					{
 						PBody->setRigidBodyFlag(PxRigidBodyFlag::eENABLE_CCD, true);
@@ -791,12 +766,8 @@ void FBodyInstance::UpdatePhysicsShapeFilterData(uint32 SkelMeshCompID, bool bUs
 		{
 			UpdateMassProperties();
 		}
-
-		if(AsyncScene)
-		{
-			SCENE_UNLOCK_WRITE(AsyncScene);
-		}
 	}
+	});
 }
 #endif
 
@@ -1273,7 +1244,7 @@ void FBodyInstance::InitBodies(TArray<FBodyInstance*>& Bodies, TArray<FTransform
 		bool bInitFail = false;
 
 		// Don't process if we've already got a body
-		if(Instance->GetPxRigidActor())
+		if(Instance->GetPxRigidActor_AssumesLocked())
 		{
 			Instance->OwnerComponent = nullptr;
 			Instance->BodySetup = nullptr;
@@ -1331,8 +1302,8 @@ void FBodyInstance::InitBodies(TArray<FBodyInstance*>& Bodies, TArray<FTransform
 		PxMaterial* PSimpleMat = SimplePhysMat->GetPhysXMaterial();
 
 		FShapeData ShapeData;
-		Instance->GetFilterData(ShapeData);
-		Instance->GetShapeFlags(ShapeData, Instance->CollisionEnabled, BodySetup->CollisionTraceFlag == CTF_UseComplexAsSimple);
+		Instance->GetFilterData_AssumesLocked(ShapeData);
+		Instance->GetShapeFlags_AssumesLocked(ShapeData, Instance->CollisionEnabled, BodySetup->CollisionTraceFlag == CTF_UseComplexAsSimple);
 
 		if(PNewDynamic)
 		{
@@ -1345,7 +1316,7 @@ void FBodyInstance::InitBodies(TArray<FBodyInstance*>& Bodies, TArray<FTransform
 
 		if(PNewActorSync)
 		{
-			BodySetup->AddShapesToRigidActor(Instance, PNewActorSync, PST_Sync, Instance->Scale3D, PSimpleMat, ComplexPhysMats, ShapeData);
+			BodySetup->AddShapesToRigidActor_AssumesLocked(Instance, PNewActorSync, PST_Sync, Instance->Scale3D, PSimpleMat, ComplexPhysMats, ShapeData);
 			bInitFail |= PNewActorSync->getNbShapes() == 0;
 			PNewActorSync->userData = &Instance->PhysxUserData;
 			PNewActorSync->setName(Instance->CharDebugName.IsValid() ? Instance->CharDebugName->GetData() : nullptr);
@@ -1359,7 +1330,7 @@ void FBodyInstance::InitBodies(TArray<FBodyInstance*>& Bodies, TArray<FTransform
 		if(PNewActorAsync)
 		{
 			check(PSceneAsync);
-			BodySetup->AddShapesToRigidActor(Instance, PNewActorAsync, PST_Async, Instance->Scale3D, PSimpleMat, ComplexPhysMats, ShapeData);
+			BodySetup->AddShapesToRigidActor_AssumesLocked(Instance, PNewActorAsync, PST_Async, Instance->Scale3D, PSimpleMat, ComplexPhysMats, ShapeData);
 			bInitFail |= PNewActorAsync->getNbShapes() == 0;
 			PNewActorAsync->userData = &Instance->PhysxUserData;
 			PNewActorAsync->setName(Instance->CharDebugName.IsValid() ? Instance->CharDebugName->GetData() : nullptr);
@@ -1445,6 +1416,15 @@ void FBodyInstance::InitBodies(TArray<FBodyInstance*>& Bodies, TArray<FTransform
 	if(!bCanDefer)
 	{
 		SCOPE_CYCLE_COUNTER(STAT_BulkSceneAdd);
+		SCOPED_SCENE_WRITE_LOCK(PSceneSync);
+		const bool bAddingToAsyncScene = PhysXAsyncActors.Num() && PSceneAsync;
+		
+		if(bAddingToAsyncScene)
+		{
+			PSceneAsync->lockWrite();
+		}
+
+
 		// Use the aggregate if it exists, has enough slots and is in the correct scene (or no scene)
 		if(InAggregate && PhysXDynamicActors.Num() > 0 && (InAggregate->getMaxNbActors() - InAggregate->getNbActors()) >= (uint32)PhysXDynamicActors.Num() && (InAggregate->getScene() == PSceneForNewDynamic || InAggregate->getScene() == NULL))
 		{
@@ -1455,7 +1435,7 @@ void FBodyInstance::InitBodies(TArray<FBodyInstance*>& Bodies, TArray<FTransform
 		}
 		else
 		{
-			SCOPED_SCENE_WRITE_LOCK(PSceneSync);
+			
 			for(PxActor* Actor : PhysXSyncActors)
 			{
 				PSceneSync->addActor(*Actor);
@@ -1463,7 +1443,6 @@ void FBodyInstance::InitBodies(TArray<FBodyInstance*>& Bodies, TArray<FTransform
 
 			if(PSceneAsync)
 			{
-				SCOPED_SCENE_WRITE_LOCK(PSceneAsync);
 				for (PxActor* Actor : PhysXAsyncActors)
 				{
 					PSceneAsync->addActor(*Actor);
@@ -1478,9 +1457,14 @@ void FBodyInstance::InitBodies(TArray<FBodyInstance*>& Bodies, TArray<FTransform
 			for(int32 BodyIdx = 0 ; BodyIdx < NumBodies ; ++BodyIdx)
 			{
 				FBodyInstance* Instance = Bodies[BodyIdx];
-				Instance->InitDynamicProperties();
+				Instance->InitDynamicProperties_AssumesLocked();
 			}
 		}
+
+		if(bAddingToAsyncScene)
+		{
+			PSceneAsync->unlockWrite();
+		}
 	}
 
 	
@@ -1557,7 +1541,7 @@ void TermBodyHelper(int32& SceneIndex, PxRigidActor*& PRigidActor, FBodyInstance
 				FPhysScene* PhysScene = FPhysxUserData::Get<FPhysScene>(PScene->userData);
 				if (PhysScene)
 				{
-					PhysScene->TermBody(BodyInstance);
+					PhysScene->TermBody_AssumesLocked(BodyInstance);
 				}
 
 				PRigidActor->release();
@@ -1663,27 +1647,27 @@ bool FBodyInstance::Weld(FBodyInstance* TheirBody, const FTransform& TheirTM)
 	PxScene* PSyncScene = RigidActorSync ? RigidActorSync->getScene() : NULL;
 	PxScene* PAsyncScene = RigidActorAsync ? RigidActorAsync->getScene() : NULL;
 
-	SCOPED_SCENE_WRITE_LOCK(PSyncScene);
-	SCOPED_SCENE_WRITE_LOCK(PAsyncScene);
+	ExecuteOnPhysicsReadWrite([&]
+	{
 	SCOPE_CYCLE_COUNTER(STAT_UpdatePhysMats);
 
 	UPhysicalMaterial* SimplePhysMat = GetSimplePhysicalMaterial();
 	TArray<UPhysicalMaterial*> ComplexPhysMats = GetComplexPhysicalMaterials();
 	PxMaterial* PSimpleMat = SimplePhysMat->GetPhysXMaterial();
-	
+
 	FShapeData ShapeData;
-	GetFilterData(ShapeData);
-	GetShapeFlags(ShapeData, CollisionEnabled, BodySetup->CollisionTraceFlag == CTF_UseComplexAsSimple);
+		GetFilterData_AssumesLocked(ShapeData);
+		GetShapeFlags_AssumesLocked(ShapeData, CollisionEnabled, BodySetup->CollisionTraceFlag == CTF_UseComplexAsSimple);
 
 	//child body gets placed into the same scenes as parent body
 	if (PxRigidActor* MyBody = RigidActorSync)
 	{
-		TheirBody->BodySetup->AddShapesToRigidActor(this, MyBody, PST_Sync, Scale3D, PSimpleMat, ComplexPhysMats, ShapeData, RelativeTM, &PNewShapes);
+			TheirBody->BodySetup->AddShapesToRigidActor_AssumesLocked(this, MyBody, PST_Sync, Scale3D, PSimpleMat, ComplexPhysMats, ShapeData, RelativeTM, &PNewShapes);
 	}
 
 	if (PxRigidActor* MyBody = RigidActorAsync)
 	{
-		TheirBody->BodySetup->AddShapesToRigidActor(this, MyBody, PST_Sync, Scale3D, PSimpleMat, ComplexPhysMats, ShapeData, RelativeTM, &PNewShapes);
+			TheirBody->BodySetup->AddShapesToRigidActor_AssumesLocked(this, MyBody, PST_Sync, Scale3D, PSimpleMat, ComplexPhysMats, ShapeData, RelativeTM, &PNewShapes);
 	}
 
 
@@ -1701,6 +1685,8 @@ bool FBodyInstance::Weld(FBodyInstance* TheirBody, const FTransform& TheirTM)
 	//remove their body from scenes
 	TermBodyHelper(TheirBody->SceneIndexSync, TheirBody->RigidActorSync, TheirBody);
 	TermBodyHelper(TheirBody->SceneIndexAsync, TheirBody->RigidActorAsync, TheirBody);
+	});
+	
 
 	return true;
 #endif
@@ -1714,17 +1700,13 @@ void FBodyInstance::UnWeld(FBodyInstance* TheirBI)
 
 #if WITH_PHYSX
 
-	int32 NumSyncShapes = 0;
-	TArray<PxShape *> PShapes = GetAllShapes(NumSyncShapes);
-
 	bool bShapesChanged = false;
 	bool bNeedsNotification = false;
 
-	PxScene * PSyncScene = RigidActorSync ? RigidActorSync->getScene() : NULL;
-	PxScene * PAsyncScene = RigidActorAsync ? RigidActorAsync->getScene() : NULL;
-
-	SCOPED_SCENE_WRITE_LOCK(PSyncScene);
-	SCOPED_SCENE_WRITE_LOCK(PAsyncScene);
+	ExecuteOnPhysicsReadWrite([&]
+	{
+		TArray<physx::PxShape *> PShapes;
+		const int32 NumSyncShapes = GetAllShapes_AssumesLocked(PShapes);
 
 	for (int32 ShapeIdx = 0; ShapeIdx < NumSyncShapes; ++ShapeIdx)
 	{
@@ -1757,6 +1739,7 @@ void FBodyInstance::UnWeld(FBodyInstance* TheirBI)
 			}
 		}
 	}
+	});
 
 	if (bShapesChanged)
 	{
@@ -1908,17 +1891,19 @@ bool FBodyInstance::UpdateBodyScale(const FVector& InScale3D)
 
 	// Determine the scaling mode
 	EScaleMode::Type ScaleMode = EScaleMode::Free;
-
+	FVector UpdatedScale3D;
 #if WITH_PHYSX
 	//Get all shapes
-	int32 TotalShapeCount = 0;
-	TArray<PxShape *> PShapes = GetAllShapes(TotalShapeCount);
+	ExecuteOnPhysicsReadWrite([&]
+	{
+		TArray<PxShape *> PShapes;
+		GetAllShapes_AssumesLocked(PShapes);
 
 	for (int32 ShapeIdx = 0; ShapeIdx < PShapes.Num(); ++ShapeIdx)
 	{
 		PxShape* PShape = PShapes[ShapeIdx];
 		PxGeometryType::Enum GeomType = PShape->getGeometryType();
-		
+
 		if (GeomType == PxGeometryType::eSPHERE)
 		{
 			ScaleMode = EScaleMode::LockedXYZ;	//sphere is most restrictive so we can stop
@@ -1939,7 +1924,6 @@ bool FBodyInstance::UpdateBodyScale(const FVector& InScale3D)
 
 	FVector RelativeScale3D;
 	FVector RelativeScale3DAbs;
-	FVector UpdatedScale3D;
 	ComputeScalingVectors(ScaleMode, InScale3DAdjusted, OldScale3D, RelativeScale3D, RelativeScale3DAbs, UpdatedScale3D);
 
 	// Apply scaling
@@ -1951,13 +1935,12 @@ bool FBodyInstance::UpdateBodyScale(const FVector& InScale3D)
 	PxConvexMeshGeometry PConvexGeom;
 	PxTriangleMeshGeometry PTriMeshGeom;
 
-	for(int32 ShapeIdx=0; ShapeIdx<PShapes.Num(); ShapeIdx++)
+		for (int32 ShapeIdx = 0; ShapeIdx < PShapes.Num(); ShapeIdx++)
 	{
 		bool bInvalid = false;	//we only mark invalid if actually found geom and it's invalid scale
 		PxGeometry* UpdatedGeometry = NULL;
 		PxShape* PShape = PShapes[ShapeIdx];
 		PxScene* PScene = PShape->getActor()->getScene();
-		SCENE_LOCK_READ(PScene);
 
 		PxTransform PLocalPose = PShape->getLocalPose();
 		PLocalPose.q.normalize();
@@ -1970,11 +1953,10 @@ bool FBodyInstance::UpdateBodyScale(const FVector& InScale3D)
 				ensure(ScaleMode == EScaleMode::LockedXYZ);
 
 				PShape->getSphereGeometry(PSphereGeom);
-				SCENE_UNLOCK_READ(PScene);
 
 				PSphereGeom.radius *= RelativeScale3DAbs.X;
 				PLocalPose.p *= RelativeScale3D.X;
-				
+
 				if (PSphereGeom.isValid())
 				{
 					UpdatedGeometry = &PSphereGeom;
@@ -1989,7 +1971,6 @@ bool FBodyInstance::UpdateBodyScale(const FVector& InScale3D)
 			case PxGeometryType::eBOX:
 			{
 				PShape->getBoxGeometry(PBoxGeom);
-				SCENE_UNLOCK_READ(PScene);
 
 				PBoxGeom.halfExtents.x *= RelativeScale3DAbs.X;
 				PBoxGeom.halfExtents.y *= RelativeScale3DAbs.Y;
@@ -2014,7 +1995,6 @@ bool FBodyInstance::UpdateBodyScale(const FVector& InScale3D)
 				ensure(ScaleMode == EScaleMode::LockedXY || ScaleMode == EScaleMode::LockedXYZ);
 
 				PShape->getCapsuleGeometry(PCapsuleGeom);
-				SCENE_UNLOCK_READ(PScene);
 
 				PCapsuleGeom.halfHeight *= RelativeScale3DAbs.Z;
 				PCapsuleGeom.radius *= RelativeScale3DAbs.X;
@@ -2038,7 +2018,6 @@ bool FBodyInstance::UpdateBodyScale(const FVector& InScale3D)
 			case PxGeometryType::eCONVEXMESH:
 			{
 				PShape->getConvexMeshGeometry(PConvexGeom);
-				SCENE_UNLOCK_READ(PScene);
 
 				// find which convex elems it is
 				// it would be nice to know if the order of PShapes array index is in the order of createShape
@@ -2079,13 +2058,12 @@ bool FBodyInstance::UpdateBodyScale(const FVector& InScale3D)
 						}
 					}
 				}
-				
+
 				break;
 			}
 			case PxGeometryType::eTRIANGLEMESH:
 			{
 				PShape->getTriangleMeshGeometry(PTriMeshGeom);
-				SCENE_UNLOCK_READ(PScene);
 
 				// Create tri-mesh shape
 				if (BodySetup.IsValid() && (BodySetup->TriMesh != NULL || BodySetup->TriMeshNegX != NULL))
@@ -2126,19 +2104,16 @@ bool FBodyInstance::UpdateBodyScale(const FVector& InScale3D)
 			case PxGeometryType::eHEIGHTFIELD:
 			{
 				// HeightField is only used by Landscape, which does different code path from other primitives
-				SCENE_UNLOCK_READ(PScene);
 				break;
 			}
 			default:
 			{
-					   SCENE_UNLOCK_READ(PScene);
 					   UE_LOG(LogPhysics, Error, TEXT("Unknown geom type."));
 			}
 		}// end switch
 
 		if (UpdatedGeometry)
 		{
-			SCOPED_SCENE_WRITE_LOCK(PScene);
 			PShape->setLocalPose(PLocalPose);
 			PShape->setGeometry(*UpdatedGeometry);
 		}
@@ -2147,6 +2122,8 @@ bool FBodyInstance::UpdateBodyScale(const FVector& InScale3D)
 			FMessageLog("PIE").Warning(FText::Format(LOCTEXT("PhysicsInvalidScale", "Scale ''{0}'' is not valid on object '{1}'."), FText::FromString(InScale3DAdjusted.ToString()), FText::FromString(GetBodyDebugName())));
 		}
 	}
+	});
+	
 #endif
 
 #if WITH_BOX2D
@@ -2178,13 +2155,12 @@ void FBodyInstance::UpdateInstanceSimulatePhysics()
 	const bool bUseSimulate = IsInstanceSimulatingPhysics();
 	bool bInitialized = false;
 #if WITH_PHYSX
-	if (PxRigidDynamic* PRigidDynamic = GetPxRigidDynamic())
+	ExecuteOnPxRigidDynamicReadWrite(this, [&](PxRigidDynamic* PRigidDynamic)
 	{
 		bInitialized = true;
 		// If we want it fixed, and it is currently not kinematic
 		bool bNewKinematic = (bUseSimulate == false);
 		{
-			SCOPED_SCENE_WRITE_LOCK(PRigidDynamic->getScene());
 			PRigidDynamic->setRigidDynamicFlag(PxRigidDynamicFlag::eKINEMATIC, bNewKinematic);
 			PRigidDynamic->setRigidBodyFlag(PxRigidBodyFlag::eENABLE_CCD, !bNewKinematic && bUseCCD);
 
@@ -2194,7 +2170,8 @@ void FBodyInstance::UpdateInstanceSimulatePhysics()
 				PRigidDynamic->wakeUp();
 			}
 		}
-	}
+		
+	});
 #endif
 
 #if WITH_BOX2D
@@ -2222,23 +2199,29 @@ void FBodyInstance::UpdateInstanceSimulatePhysics()
 	}
 }
 
+bool FBodyInstance::IsNonKinematic() const
+{
+	return bSimulatePhysics;
+}
+
 bool FBodyInstance::IsDynamic() const
 {
+	bool bIsDynamic = false;
 #if WITH_PHYSX
-	if (PxRigidActor const* const PRigidActor = GetPxRigidActor())
+	ExecuteOnPxRigidDynamicReadOnly(this, [&](const PxRigidDynamic* PRigidDynamic)
 	{
-		return PRigidActor->isRigidDynamic() != nullptr;
-	}
+		bIsDynamic = true;
+	});
 #endif // WITH_PHYSX
 
 #if WITH_BOX2D
 	if (BodyInstancePtr != nullptr)
 	{
-		return BodyInstancePtr->GetType() != b2_staticBody;
+		bIsDynamic = BodyInstancePtr->GetType() != b2_staticBody;
 	}
 #endif // WITH_BOX2D
 
-	return false;
+	return bIsDynamic;
 }
 
 void FBodyInstance::SetInstanceSimulatePhysics(bool bSimulate, bool bMaintainPhysicsBlending)
@@ -2319,7 +2302,7 @@ bool FBodyInstance::ShouldInstanceSimulatingPhysics()
 bool FBodyInstance::IsValidBodyInstance() const
 {
 #if WITH_PHYSX
-	if (PxRigidActor* PActor = GetPxRigidActor())
+	if (PxRigidActor* PActor = GetPxRigidActor_AssumesLocked())
 	{
 		return true;
 	}
@@ -2335,34 +2318,45 @@ bool FBodyInstance::IsValidBodyInstance() const
 	return false;
 }
 
-FTransform FBodyInstance::GetUnrealWorldTransform() const
+template <bool NeedsLock>
+FTransform GetUnrealWorldTransformImp(const FBodyInstance* BodyInstance)
 {
+	FTransform WorldTM = FTransform::Identity;
 #if WITH_PHYSX
-	PxRigidActor* PActor = GetPxRigidActor();
-	if (PActor != NULL)
-	{		
-		SCOPED_SCENE_READ_LOCK(PActor->getScene());
-
+	auto Func = [&](const PxRigidActor* PActor)
+	{
 		PxTransform PTM = PActor->getGlobalPose();
+		WorldTM = P2UTransform(PTM);
+	};
 
-		return P2UTransform(PTM);
-	}
+	ExecuteOnPxRigidActorReadOnly<decltype(Func), NeedsLock>(BodyInstance, Func);
 #endif // WITH_PHYSX
 
 #if WITH_BOX2D
-	if (BodyInstancePtr != NULL)
+	if (BodyInstance->BodyInstancePtr != NULL)
 	{
-		const b2Vec2 Pos2D = BodyInstancePtr->GetPosition();
-		const float RotationInRadians = BodyInstancePtr->GetAngle();
-		
+		const b2Vec2 Pos2D = BodyInstance->BodyInstancePtr->GetPosition();
+		const float RotationInRadians = BodyInstance->BodyInstancePtr->GetAngle();
+
 		const FVector Translation3D(FPhysicsIntegration2D::ConvertBoxVectorToUnreal(Pos2D));
 		const FRotator Rotation3D(FMath::RadiansToDegrees(RotationInRadians), 0.0f, 0.0f); //@TODO: BOX2D: Should be moved to FPhysicsIntegration2D
 
-		return FTransform(Rotation3D, Translation3D, Scale3D);
+		WorldTM = FTransform(Rotation3D, Translation3D, BodyInstance->Scale3D);
 	}
 #endif
 
-	return FTransform::Identity;
+	return WorldTM;
+}
+
+FTransform FBodyInstance::GetUnrealWorldTransform() const
+{
+	return GetUnrealWorldTransformImp<true>(this);
+}
+
+
+FTransform FBodyInstance::GetUnrealWorldTransform_AssumesLocked() const
+{
+	return GetUnrealWorldTransformImp<false>(this);
 }
 
 void FBodyInstance::SetBodyTransform(const FTransform& NewTransform, bool bTeleport)
@@ -2385,10 +2379,7 @@ void FBodyInstance::SetBodyTransform(const FTransform& NewTransform, bool bTelep
 	}
 
 #if WITH_PHYSX
-	PxRigidActor* RigidActor = GetPxRigidActor();
-	PxRigidDynamic* PRigidDynamic = GetPxRigidDynamic();
-
-	if (RigidActor != nullptr)
+	if (PxRigidActor* RigidActor = GetPxRigidActor_AssumesLocked())
 	{
 		const PxTransform PNewPose = U2PTransform(NewTransform);
 
@@ -2398,16 +2389,17 @@ void FBodyInstance::SetBodyTransform(const FTransform& NewTransform, bool bTelep
 			return;
 		}
 
-		SCENE_LOCK_WRITE(RigidActor->getScene());
+		ExecuteOnPhysicsReadWrite([&]
+		{
 		// SIMULATED & KINEMATIC
-		if (PRigidDynamic)
+			if (PxRigidDynamic* PRigidDynamic = GetPxRigidDynamic_AssumesLocked())
 		{
 			// If kinematic and not teleporting, set kinematic target
-			if (!IsRigidBodyNonKinematic(PRigidDynamic) && !bTeleport)
+				if (!IsRigidBodyNonKinematic_AssumesLocked(PRigidDynamic) && !bTeleport)
 			{
 				const PxScene* PScene = PRigidDynamic->getScene();
 				FPhysScene* PhysScene = FPhysxUserData::Get<FPhysScene>(PScene->userData);
-				PhysScene->SetKinematicTarget(this, NewTransform, true);
+					PhysScene->SetKinematicTarget_AssumesLocked(this, NewTransform, true);
 			}
 			// Otherwise, set global pose
 			else
@@ -2429,8 +2421,7 @@ void FBodyInstance::SetBodyTransform(const FTransform& NewTransform, bool bTelep
 			// In EDITOR, go ahead and move it with no warning, we are editing the level
 			RigidActor->setGlobalPose(PNewPose);
 		}
-
-		SCENE_UNLOCK_WRITE(RigidActor->getScene());
+		});
 	}
 #endif  // WITH_PHYSX
 
@@ -2449,100 +2440,134 @@ void FBodyInstance::SetBodyTransform(const FTransform& NewTransform, bool bTelep
 #endif
 }
 
-FVector FBodyInstance::GetUnrealWorldVelocity() const
+template <bool NeedsLock>
+FVector GetUnrealWorldVelocityImp(const FBodyInstance* BodyInstance)
 {
 	FVector LinVel(EForceInit::ForceInitToZero);
 
 #if WITH_PHYSX
-	if (PxRigidBody* PRigidBody = GetPxRigidBody())
+	auto Func = [&](const PxRigidBody* PRigidBody)
 	{
-		SCOPED_SCENE_READ_LOCK(PRigidBody->getScene());
 		LinVel = P2UVector(PRigidBody->getLinearVelocity());
-	}
+	};
+
+	ExecuteOnPxRigidBodyReadOnly<decltype(Func), NeedsLock>(BodyInstance, Func);
 #endif // WITH_PHYSX
 
 #if WITH_BOX2D
-	if (BodyInstancePtr != nullptr)
+	if (BodyInstance->BodyInstancePtr != nullptr)
 	{
-		LinVel = FPhysicsIntegration2D::ConvertBoxVectorToUnreal(BodyInstancePtr->GetLinearVelocity());
+		LinVel = FPhysicsIntegration2D::ConvertBoxVectorToUnreal(BodyInstance->BodyInstancePtr->GetLinearVelocity());
 	}
 #endif
 
 	return LinVel;
 }
 
-/** Note: returns angular velocity in degrees per second. */
-FVector FBodyInstance::GetUnrealWorldAngularVelocity() const
+FVector FBodyInstance::GetUnrealWorldVelocity() const
+{
+	return GetUnrealWorldVelocityImp<true>(this);
+}
+
+FVector FBodyInstance::GetUnrealWorldVelocity_AssumesLocked() const
+{
+	return GetUnrealWorldVelocityImp<false>(this);
+}
+
+template <bool NeedsLock>
+FVector GetUnrealWorldAngularVelocityImp(const FBodyInstance* BodyInstance)
 {
 	FVector AngVel(EForceInit::ForceInitToZero);
 
 #if WITH_PHYSX
-	PxRigidBody* PRigidBody = GetPxRigidBody();
-	if (PRigidBody != NULL)
+	auto Func = [&](const PxRigidBody* PRigidBody)
 	{
-		SCOPED_SCENE_READ_LOCK(PRigidBody->getScene());
 		AngVel = FVector::RadiansToDegrees(P2UVector(PRigidBody->getAngularVelocity()));
-	}
+	};
+
+	ExecuteOnPxRigidBodyReadOnly<decltype(Func), NeedsLock>(BodyInstance, Func);
 #endif // WITH_PHYSX
 
 #if WITH_BOX2D
-	if (BodyInstancePtr != nullptr)
+	if (BodyInstance->BodyInstancePtr != nullptr)
 	{
-		AngVel = FPhysicsIntegration2D::ConvertBoxAngularVelocityToUnreal(BodyInstancePtr->GetAngularVelocity());
+		AngVel = FPhysicsIntegration2D::ConvertBoxAngularVelocityToUnreal(BodyInstance->BodyInstancePtr->GetAngularVelocity());
 	}
 #endif
 
 	return AngVel;
 }
 
-FVector FBodyInstance::GetUnrealWorldVelocityAtPoint(const FVector& Point) const
+/** Note: returns angular velocity in degrees per second. */
+FVector FBodyInstance::GetUnrealWorldAngularVelocity() const
+{
+	return GetUnrealWorldAngularVelocityImp<true>(this);
+}
+
+/** Note: returns angular velocity in degrees per second. */
+FVector FBodyInstance::GetUnrealWorldAngularVelocity_AssumesLocked() const
+{
+	return GetUnrealWorldAngularVelocityImp<false>(this);
+}
+
+template <bool NeedsLock>
+FVector GetUnrealWorldVelocityAtPointImp(const FBodyInstance* BodyInstance, const FVector& Point)
 {
 	FVector LinVel(EForceInit::ForceInitToZero);
 
 #if WITH_PHYSX
-	if (PxRigidBody* PRigidBody = GetPxRigidBody())
+	auto Func = [&](const PxRigidBody* PRigidBody)
 	{
-		SCOPED_SCENE_READ_LOCK(PRigidBody->getScene());
 		PxVec3 PPoint = U2PVector(Point);
 		LinVel = P2UVector(PxRigidBodyExt::getVelocityAtPos(*PRigidBody, PPoint));
-	}
+	};
+
+	ExecuteOnPxRigidBodyReadOnly<decltype(Func), NeedsLock>(BodyInstance, Func);
 #endif // WITH_PHYSX
 
 #if WITH_BOX2D
-	if (BodyInstancePtr != nullptr)
+	if (BodyInstance->BodyInstancePtr != nullptr)
 	{
 		const b2Vec2 BoxPoint = FPhysicsIntegration2D::ConvertUnrealVectorToBox(Point);
-		LinVel = FPhysicsIntegration2D::ConvertBoxVectorToUnreal(BodyInstancePtr->GetLinearVelocityFromWorldPoint(BoxPoint));
+		LinVel = FPhysicsIntegration2D::ConvertBoxVectorToUnreal(BodyInstance->BodyInstancePtr->GetLinearVelocityFromWorldPoint(BoxPoint));
 	}
 #endif
 
 	return LinVel;
 }
 
+
+FVector FBodyInstance::GetUnrealWorldVelocityAtPoint(const FVector& Point) const
+{
+	return GetUnrealWorldVelocityAtPointImp<true>(this, Point);
+}
+
+
+FVector FBodyInstance::GetUnrealWorldVelocityAtPoint_AssumesLocked(const FVector& Point) const
+{
+	return GetUnrealWorldVelocityAtPointImp<false>(this, Point);
+}
+
 FVector FBodyInstance::GetCOMPosition() const
 {
+	FVector WorldCOM = FVector::ZeroVector;
 #if WITH_PHYSX
-	if (PxRigidBody* PRigidBody = GetPxRigidBody())
+	ExecuteOnPxRigidBodyReadOnly(this, [&](const PxRigidBody* PRigidBody)
 	{
-		SCOPED_SCENE_READ_LOCK(PRigidBody->getScene());
 		PxTransform PLocalCOM = PRigidBody->getCMassLocalPose();
 		PxVec3 PWorldCOM = PRigidBody->getGlobalPose().transform(PLocalCOM.p);
-		return P2UVector(PWorldCOM);
-	}
-	else
-	{
-		return FVector::ZeroVector;
-	}
+		WorldCOM = P2UVector(PWorldCOM);
+	});
 #endif // WITH_PHYSX
 
 #if WITH_BOX2D
 	if (BodyInstancePtr != nullptr)
 	{
-		return FPhysicsIntegration2D::ConvertBoxVectorToUnreal(BodyInstancePtr->GetWorldCenter());
+		WorldCOM = FPhysicsIntegration2D::ConvertBoxVectorToUnreal(BodyInstancePtr->GetWorldCenter());
 	}
 #endif
 
-	return FVector::ZeroVector;
+	return WorldCOM;
 }
 
 float FBodyInstance::GetBodyMass() const
@@ -2550,11 +2575,10 @@ float FBodyInstance::GetBodyMass() const
 	float Retval = 0.f;
 
 #if WITH_PHYSX
-	if (PxRigidBody* PRigidBody = GetPxRigidBody())
+	ExecuteOnPxRigidBodyReadOnly(this, [&] (const PxRigidBody* PRigidBody)
 	{
-		SCOPED_SCENE_READ_LOCK(PRigidBody->getScene());
 		Retval = PRigidBody->getMass();
-	}
+	});
 #endif // WITH_PHYSX
 
 #if WITH_BOX2D
@@ -2573,11 +2597,10 @@ FVector FBodyInstance::GetBodyInertiaTensor() const
 	FVector Retval = FVector::ZeroVector;
 
 #if WITH_PHYSX
-	if (PxRigidBody* PRigidBody = GetPxRigidBody())
+	ExecuteOnPxRigidBodyReadOnly(this, [&] (const PxRigidBody* PRigidBody)
 	{
-		SCOPED_SCENE_READ_LOCK(PRigidBody->getScene());
 		Retval = P2UVector(PRigidBody->getMassSpaceInertiaTensor());
-	}
+	});
 #endif // WITH_PHYSX
 
 #if WITH_BOX2D
@@ -2595,16 +2618,13 @@ FBox FBodyInstance::GetBodyBounds() const
 	FBox Bounds;
 
 #if WITH_PHYSX
-	PxRigidBody* PRigidBody = GetPxRigidBody();
-	if (PRigidBody)
+	ExecuteOnPxRigidBodyReadOnly(this, [&](const PxRigidBody* PRigidBody)
 	{
-		SCOPED_SCENE_READ_LOCK(PRigidBody->getScene());
-		
 		PxBounds3 PBounds = PRigidBody->getWorldBounds();
 
 		Bounds.Min = P2UVector(PBounds.minimum);
 		Bounds.Max = P2UVector(PBounds.maximum);
-	}
+	});
 #endif // WITH_PHYSX
 
 	//@TODO: BOX2D: Implement GetBodyBounds
@@ -2651,9 +2671,41 @@ void FBodyInstance::CopyBodyInstancePropertiesFrom(const FBodyInstance* FromInst
 	*this = *FromInst;
 }
 
+void FBodyInstance::ExecuteOnPhysicsReadOnly(TFunctionRef<void()> Func) const
+{
+#if WITH_PHYSX
+	const int32 SceneIndex = RigidActorSync ? SceneIndexSync : SceneIndexAsync;
+	SCOPED_SCENE_READ_LOCK(GetPhysXSceneFromIndex(SceneIndex));
+	Func();
+#endif
+}
+
+void FBodyInstance::ExecuteOnPhysicsReadWrite(TFunctionRef<void()> Func) const
+{
+#if WITH_PHYSX
+	const int32 SceneIndex = RigidActorSync ? SceneIndexSync : SceneIndexAsync;
+	SCOPED_SCENE_WRITE_LOCK(GetPhysXSceneFromIndex(SceneIndex));
+	Func();
+#endif
+}
+
 
 #if WITH_PHYSX
-PxRigidActor* FBodyInstance::GetPxRigidActor(int32 SceneType) const
+int32 FBodyInstance::GetSceneIndex(int32 SceneType /* = -1 */) const
+{
+	if(SceneType < 0)
+	{
+		return RigidActorSync ? SceneIndexSync : SceneIndexAsync;
+	}
+	else if(SceneType < PST_MAX)
+	{
+		return SceneType == PST_Sync ? SceneIndexSync : SceneIndexAsync;
+	}
+
+	return -1;
+}
+
+PxRigidActor* FBodyInstance::GetPxRigidActor_AssumesLocked(int32 SceneType) const
 {
 	// Negative scene type means to return whichever is not NULL, preferring the sync scene.
 	if( SceneType < 0 )
@@ -2669,40 +2721,20 @@ PxRigidActor* FBodyInstance::GetPxRigidActor(int32 SceneType) const
 	return NULL;
 }
 
-PxRigidDynamic* FBodyInstance::GetPxRigidDynamic() const
+PxRigidDynamic* FBodyInstance::GetPxRigidDynamic_AssumesLocked() const
 {
 	// The logic below works because dynamic actors are non-NULL in only one scene.
 	// If this assumption changes, the logic needs to change.
-	if(RigidActorSync != NULL)
-	{
-		return RigidActorSync->isRigidDynamic();
-	}
-	else if(RigidActorAsync != NULL)
-	{
-		return RigidActorAsync->isRigidDynamic();
-	}
-	else
-	{
-		return NULL;
-	}
+	PxRigidActor* RigidActor = RigidActorSync ? RigidActorSync : RigidActorAsync;
+	return RigidActor ? RigidActor->isRigidDynamic() : nullptr;
 }
 
-PxRigidBody* FBodyInstance::GetPxRigidBody() const
+PxRigidBody* FBodyInstance::GetPxRigidBody_AssumesLocked() const
 {
 	// The logic below works because dynamic actors are non-NULL in only one scene.
 	// If this assumption changes, the logic needs to change.
-	if (RigidActorSync != NULL)
-	{
-		return RigidActorSync->isRigidBody();
-	}
-	else if (RigidActorAsync != NULL)
-	{
-		return RigidActorAsync->isRigidBody();
-	}
-	else
-	{
-		return NULL;
-	}
+	PxRigidActor* RigidActor = RigidActorSync ? RigidActorSync : RigidActorAsync;
+	return RigidActor ? RigidActor->isRigidBody() : nullptr;
 }
 #endif // WITH_PHYSX
 
@@ -2834,7 +2866,7 @@ void FBodyInstance::GetComplexPhysicalMaterials(TArray<UPhysicalMaterial*> &Phys
 
 #if WITH_PHYSX
 /** Util for finding the number of 'collision sim' shapes on this PxRigidActor */
-int32 GetNumSimShapes(PxRigidBody* PRigidBody)
+int32 GetNumSimShapes_AssumesLocked(const PxRigidBody* PRigidBody)
 {
 	TArray<PxShape*> PShapes;
 	PShapes.AddZeroed(PRigidBody->getNbShapes());
@@ -2877,8 +2909,9 @@ void FBodyInstance::UpdateMassProperties()
 	check(PhysMat);
 
 #if WITH_PHYSX
-	PxRigidBody* PRigidBody = GetPxRigidBody();
-	if ((PRigidBody != NULL) && (GetNumSimShapes(PRigidBody) > 0))
+	ExecuteOnPxRigidBodyReadWrite(this, [&] (PxRigidBody* PRigidBody)
+	{
+		if(GetNumSimShapes_AssumesLocked(PRigidBody) > 0)
 	{
 		// physical material - nothing can weigh less than hydrogen (0.09 kg/m^3)
 		float DensityKGPerCubicUU = FMath::Max(KgPerM3ToKgPerCm3(0.09f), gPerCm3ToKgPerCm3(PhysMat->Density));
@@ -2892,7 +2925,7 @@ void FBodyInstance::UpdateMassProperties()
 		{
 			float UsePow = FMath::Clamp<float>(PhysMat->RaiseMassToPower, KINDA_SMALL_NUMBER, 1.f);
 			NewMass = FMath::Pow(OldMass, UsePow);
-			
+
 			// Apply user-defined mass scaling.
 			NewMass *= FMath::Clamp<float>(MassScale, 0.01f, 100.0f);
 
@@ -2903,9 +2936,9 @@ void FBodyInstance::UpdateMassProperties()
 			NewMass = FMath::Max(MassInKg, 0.001f);	//min weight of 1g
 		}
 
-		check (NewMass > 0.f);
+			check(NewMass > 0.f);
 
-		float MassRatio = NewMass/OldMass;
+			float MassRatio = NewMass / OldMass;
 		PxVec3 InertiaTensor = PRigidBody->getMassSpaceInertiaTensor();
 
 		PRigidBody->setMassSpaceInertiaTensor(InertiaTensor * MassRatio);
@@ -2920,6 +2953,7 @@ void FBodyInstance::UpdateMassProperties()
 			PRigidBody->setCMassLocalPose(PCOMTransform);
 		}
 	}
+	});
 #endif
 
 #if WITH_BOX2D
@@ -2959,12 +2993,11 @@ void FBodyInstance::UpdateMassProperties()
 void FBodyInstance::UpdateDampingProperties()
 {
 #if WITH_PHYSX
-	if (PxRigidDynamic* PRigidDynamic = GetPxRigidDynamic())
+	ExecuteOnPxRigidDynamicReadWrite(this, [&](PxRigidDynamic* PRigidDynamic)
 	{
-		SCOPED_SCENE_WRITE_LOCK(PRigidDynamic->getScene());
 		PRigidDynamic->setLinearDamping(LinearDamping);
 		PRigidDynamic->setAngularDamping(AngularDamping);
-	}
+	});
 #endif
 
 #if WITH_BOX2D
@@ -2979,11 +3012,10 @@ void FBodyInstance::UpdateDampingProperties()
 bool FBodyInstance::IsInstanceAwake() const
 {
 #if WITH_PHYSX
-	if (PxRigidDynamic* PRigidDynamic = GetPxRigidDynamic())
+	ExecuteOnPxRigidDynamicReadOnly(this, [&](const PxRigidDynamic* PRigidDynamic)
 	{
-		SCOPED_SCENE_READ_LOCK(PRigidDynamic->getScene());
 		return !PRigidDynamic->isSleeping();
-	}
+	});
 #endif
 
 #if WITH_BOX2D
@@ -2999,12 +3031,13 @@ bool FBodyInstance::IsInstanceAwake() const
 void FBodyInstance::WakeInstance()
 {
 #if WITH_PHYSX
-	PxRigidDynamic* PRigidDynamic = GetPxRigidDynamic();
-	if (IsRigidBodyNonKinematic(PRigidDynamic))
+	ExecuteOnPxRigidDynamicReadWrite(this, [&](PxRigidDynamic* PRigidDynamic)
+	{
+		if (IsRigidBodyNonKinematic_AssumesLocked(PRigidDynamic))
 	{
-		SCOPED_SCENE_WRITE_LOCK(PRigidDynamic->getScene());
 		PRigidDynamic->wakeUp();
 	}
+	});
 #endif
 
 #if WITH_BOX2D
@@ -3018,11 +3051,13 @@ void FBodyInstance::WakeInstance()
 void FBodyInstance::PutInstanceToSleep()
 {
 #if WITH_PHYSX
-	PxRigidDynamic* PRigidDynamic = GetPxRigidDynamic();
-	if (IsRigidBodyNonKinematic(PRigidDynamic))
+	ExecuteOnPxRigidDynamicReadWrite(this, [&](PxRigidDynamic* PRigidDynamic)
+	{
+		if (IsRigidBodyNonKinematic_AssumesLocked(PRigidDynamic))
 	{
 		PRigidDynamic->putToSleep();
 	}
+	});
 #endif //WITH_PHYSX
 
 #if WITH_BOX2D
@@ -3036,7 +3071,7 @@ void FBodyInstance::PutInstanceToSleep()
 void FBodyInstance::SetLinearVelocity(const FVector& NewVel, bool bAddToCurrent)
 {
 #if WITH_PHYSX
-	if (PxRigidBody* PRigidBody = GetPxRigidBody())
+	ExecuteOnPxRigidBodyReadWrite(this, [&](PxRigidBody* PRigidBody)
 	{
 		PxVec3 PNewVel = U2PVector(NewVel);
 
@@ -3047,7 +3082,7 @@ void FBodyInstance::SetLinearVelocity(const FVector& NewVel, bool bAddToCurrent)
 		}
 
 		PRigidBody->setLinearVelocity(PNewVel);
-	}
+	});
 #endif // WITH_PHYSX
 
 #if WITH_BOX2D
@@ -3070,7 +3105,7 @@ void FBodyInstance::SetLinearVelocity(const FVector& NewVel, bool bAddToCurrent)
 void FBodyInstance::SetAngularVelocity(const FVector& NewAngVel, bool bAddToCurrent)
 {
 #if WITH_PHYSX
-	if (PxRigidBody* PRigidBody = GetPxRigidBody())
+	ExecuteOnPxRigidBodyReadWrite(this, [&](PxRigidBody* PRigidBody)
 	{
 		PxVec3 PNewAngVel = U2PVector( FVector::DegreesToRadians(NewAngVel) );
 
@@ -3081,7 +3116,7 @@ void FBodyInstance::SetAngularVelocity(const FVector& NewAngVel, bool bAddToCurr
 		}
 
 		PRigidBody->setAngularVelocity(PNewAngVel);
-	}
+	});
 #endif //WITH_PHYSX
 
 #if WITH_BOX2D
@@ -3103,8 +3138,7 @@ void FBodyInstance::SetAngularVelocity(const FVector& NewAngVel, bool bAddToCurr
 void FBodyInstance::SetMaxAngularVelocity(float NewMaxAngVel, bool bAddToCurrent)
 {
 #if WITH_PHYSX
-	PxRigidDynamic* PRigidDynamic = GetPxRigidDynamic();
-	if (PRigidDynamic)
+	const bool bIsDynamic = ExecuteOnPxRigidDynamicReadWrite(this, [&](PxRigidDynamic* PRigidDynamic)
 	{
 		float RadNewMaxAngVel = FMath::DegreesToRadians(NewMaxAngVel);
 		
@@ -3121,8 +3155,9 @@ void FBodyInstance::SetMaxAngularVelocity(float NewMaxAngVel, bool bAddToCurrent
 		PRigidDynamic->setMaxAngularVelocity(RadNewMaxAngVel);
 
 		MaxAngularVelocity = NewMaxAngVel;
-	}
-	else
+	});
+
+	if(!bIsDynamic)
 	{
 		MaxAngularVelocity = NewMaxAngVel;	//doesn't really matter since we are not dynamic, but makes sense that we update this anyway
 	}
@@ -3137,11 +3172,10 @@ void FBodyInstance::SetMaxDepenetrationVelocity(float MaxVelocity)
 #if WITH_PHYSX
 	//0 from both instance and project settings means 
 	float UseMaxVelocity = MaxVelocity == 0.f ? PX_MAX_F32 : MaxVelocity;
-
-	if (PxRigidBody* PRigidBody = GetPxRigidBody())
+	ExecuteOnPxRigidBodyReadWrite(this, [&](PxRigidBody* PRigidBody)
 	{
 		PRigidBody->setMaxDepenetrationVelocity(UseMaxVelocity);
-	}
+	});
 #endif
 
 }
@@ -3151,13 +3185,16 @@ void FBodyInstance::AddCustomPhysics(FCalculateCustomPhysics& CalculateCustomPhy
 {
 	
 #if WITH_PHYSX
-	PxRigidBody* PRigidBody = GetPxRigidBody();
-	if (IsRigidBodyNonKinematic(PRigidBody))
+	ExecuteOnPxRigidBodyReadOnly(this, [&](const PxRigidBody* PRigidBody)
+	{
+		if (IsRigidBodyNonKinematic_AssumesLocked(PRigidBody))
 	{
 		const PxScene* PScene = PRigidBody->getScene();
 		FPhysScene* PhysScene = FPhysxUserData::Get<FPhysScene>(PScene->userData);
-		PhysScene->AddCustomPhysics(this, CalculateCustomPhysics);
+			PhysScene->AddCustomPhysics_AssumesLocked(this, CalculateCustomPhysics);
 	}
+	});
+	
 #endif // WITH_PHYSX
 
 #if WITH_BOX2D
@@ -3172,14 +3209,16 @@ void FBodyInstance::AddCustomPhysics(FCalculateCustomPhysics& CalculateCustomPhy
 void FBodyInstance::AddForce(const FVector& Force, bool bAllowSubstepping, bool bAccelChange)
 {
 #if WITH_PHYSX
-	PxRigidBody* PRigidBody = GetPxRigidBody();
-	if (IsRigidBodyNonKinematic(PRigidBody))
+	ExecuteOnPxRigidBodyReadWrite(this, [&](PxRigidBody* PRigidBody)
+	{
+		if (IsRigidBodyNonKinematic_AssumesLocked(PRigidBody))
 	{
 		const PxScene* PScene = PRigidBody->getScene();
 		FPhysScene* PhysScene = FPhysxUserData::Get<FPhysScene>(PScene->userData);
-		PhysScene->AddForce(this, Force, bAllowSubstepping, bAccelChange);
-		
+			PhysScene->AddForce_AssumesLocked(this, Force, bAllowSubstepping, bAccelChange);
+
 	}
+	});
 #endif // WITH_PHYSX
 
 #if WITH_BOX2D
@@ -3197,13 +3236,16 @@ void FBodyInstance::AddForce(const FVector& Force, bool bAllowSubstepping, bool
 void FBodyInstance::AddForceAtPosition(const FVector& Force, const FVector& Position, bool bAllowSubstepping)
 {
 #if WITH_PHYSX
-	PxRigidBody* PRigidBody = GetPxRigidBody();
-	if (IsRigidBodyNonKinematic(PRigidBody))
+	ExecuteOnPxRigidBodyReadWrite(this, [&](PxRigidBody* PRigidBody)
+	{
+		if (IsRigidBodyNonKinematic_AssumesLocked(PRigidBody))
 	{
 		const PxScene* PScene = PRigidBody->getScene();
 		FPhysScene* PhysScene = FPhysxUserData::Get<FPhysScene>(PScene->userData);
-		PhysScene->AddForceAtPosition(this, Force, Position, bAllowSubstepping);
+			PhysScene->AddForceAtPosition_AssumesLocked(this, Force, Position, bAllowSubstepping);
 	}
+	});
+	
 #endif // WITH_PHYSX
 
 #if WITH_BOX2D
@@ -3222,13 +3264,16 @@ void FBodyInstance::AddForceAtPosition(const FVector& Force, const FVector& Posi
 void FBodyInstance::AddTorque(const FVector& Torque, bool bAllowSubstepping, bool bAccelChange)
 {
 #if WITH_PHYSX
-	PxRigidBody* PRigidBody = GetPxRigidBody();
-	if (IsRigidBodyNonKinematic(PRigidBody))
+	ExecuteOnPxRigidBodyReadWrite(this, [&](PxRigidBody* PRigidBody)
+	{
+		if (IsRigidBodyNonKinematic_AssumesLocked(PRigidBody))
 	{
 		const PxScene* PScene = PRigidBody->getScene();
 		FPhysScene* PhysScene = FPhysxUserData::Get<FPhysScene>(PScene->userData);
-		PhysScene->AddTorque(this, Torque, bAllowSubstepping, bAccelChange);
+			PhysScene->AddTorque_AssumesLocked(this, Torque, bAllowSubstepping, bAccelChange);
 	}
+	});
+	
 #endif // WITH_PHYSX
 
 #if WITH_BOX2D
@@ -3246,12 +3291,15 @@ void FBodyInstance::AddTorque(const FVector& Torque, bool bAllowSubstepping, boo
 void FBodyInstance::AddAngularImpulse(const FVector& AngularImpulse, bool bVelChange)
 {
 #if WITH_PHYSX
-	PxRigidBody* PRigidBody = GetPxRigidBody();
-	if (IsRigidBodyNonKinematic(PRigidBody))
+	ExecuteOnPxRigidBodyReadWrite(this, [&](PxRigidBody* PRigidBody)
+	{
+		if (IsRigidBodyNonKinematic_AssumesLocked(PRigidBody))
 	{
 		PxForceMode::Enum Mode = bVelChange ? PxForceMode::eVELOCITY_CHANGE : PxForceMode::eIMPULSE;
 		PRigidBody->addTorque(U2PVector(AngularImpulse), Mode, true);
 	}
+	});
+	
 #endif // WITH_PHYSX
 
 #if WITH_BOX2D
@@ -3269,12 +3317,15 @@ void FBodyInstance::AddAngularImpulse(const FVector& AngularImpulse, bool bVelCh
 void FBodyInstance::AddImpulse(const FVector& Impulse, bool bVelChange)
 {
 #if WITH_PHYSX
-	PxRigidBody* PRigidBody = GetPxRigidBody();
-	if (IsRigidBodyNonKinematic(PRigidBody))
+	ExecuteOnPxRigidBodyReadWrite(this, [&](PxRigidBody* PRigidBody)
+	{
+		if (IsRigidBodyNonKinematic_AssumesLocked(PRigidBody))
 	{
 		PxForceMode::Enum Mode = bVelChange ? PxForceMode::eVELOCITY_CHANGE : PxForceMode::eIMPULSE;
 		PRigidBody->addForce(U2PVector(Impulse), Mode, true);
 	}
+	});
+	
 #endif // WITH_PHYSX
 
 #if WITH_BOX2D
@@ -3292,12 +3343,15 @@ void FBodyInstance::AddImpulse(const FVector& Impulse, bool bVelChange)
 void FBodyInstance::AddImpulseAtPosition(const FVector& Impulse, const FVector& Position)
 {
 #if WITH_PHYSX
-	PxRigidBody* PRigidBody = GetPxRigidBody();
-	if (IsRigidBodyNonKinematic(PRigidBody))
+	ExecuteOnPxRigidBodyReadWrite(this, [&](PxRigidBody* PRigidBody)
+	{
+		if (IsRigidBodyNonKinematic_AssumesLocked(PRigidBody))
 	{
 		PxForceMode::Enum Mode = PxForceMode::eIMPULSE; // does not support eVELOCITY_CHANGE
 		PxRigidBodyExt::addForceAtPos(*PRigidBody, U2PVector(Impulse), U2PVector(Position), Mode, true);
 	}
+	});
+	
 #endif // WITH_PHYSX
 
 #if WITH_BOX2D
@@ -3326,10 +3380,10 @@ void FBodyInstance::SetEnableGravity(bool bInGravityEnabled)
 		bEnableGravity = bInGravityEnabled;
 
 #if WITH_PHYSX
-		if (PxRigidBody* PActor = GetPxRigidBody())
+		ExecuteOnPxRigidBodyReadWrite(this, [&](PxRigidBody* PRigidBody)
 		{
-			PActor->setActorFlag( PxActorFlag::eDISABLE_GRAVITY, !bEnableGravity );
-		}
+			PRigidBody->setActorFlag( PxActorFlag::eDISABLE_GRAVITY, !bEnableGravity );
+		});
 #endif // WITH_PHYSX
 
 #if WITH_BOX2D
@@ -3392,11 +3446,14 @@ void AddRadialImpulseToBox2DBodyInstance(class b2Body* BodyInstancePtr, const FV
 void FBodyInstance::AddRadialImpulseToBody(const FVector& Origin, float Radius, float Strength, uint8 Falloff, bool bVelChange)
 {
 #if WITH_PHYSX
-	PxRigidBody* PRigidBody = GetPxRigidBody();
-	if (IsRigidBodyNonKinematic(PRigidBody))
+	ExecuteOnPxRigidBodyReadWrite(this, [&](PxRigidBody* PRigidBody)
+	{
+		if (IsRigidBodyNonKinematic_AssumesLocked(PRigidBody))
 	{
-		AddRadialImpulseToPxRigidBody(*PRigidBody, Origin, Radius, Strength, Falloff, bVelChange);
+			AddRadialImpulseToPxRigidBody_AssumesLocked(*PRigidBody, Origin, Radius, Strength, Falloff, bVelChange);
 	}
+	});
+	
 #endif // WITH_PHYSX
 
 #if WITH_BOX2D
@@ -3410,14 +3467,16 @@ void FBodyInstance::AddRadialImpulseToBody(const FVector& Origin, float Radius,
 void FBodyInstance::AddRadialForceToBody(const FVector& Origin, float Radius, float Strength, uint8 Falloff, bool bAccelChange, bool bAllowSubstepping)
 {
 #if WITH_PHYSX
-	PxRigidBody* PRigidBody = GetPxRigidBody();
-	if (IsRigidBodyNonKinematic(PRigidBody))
+	ExecuteOnPxRigidBodyReadWrite(this, [&](PxRigidBody* PRigidBody)
+	{
+		if (IsRigidBodyNonKinematic_AssumesLocked(PRigidBody))
 	{
 		const PxScene* PScene = PRigidBody->getScene();
 		FPhysScene* PhysScene = FPhysxUserData::Get<FPhysScene>(PScene->userData);
-		PhysScene->AddRadialForceToBody(this, Origin, Radius, Strength, Falloff, bAccelChange, bAllowSubstepping);
+			PhysScene->AddRadialForceToBody_AssumesLocked(this, Origin, Radius, Strength, Falloff, bAccelChange, bAllowSubstepping);
 
 	}
+	});
 #endif // WITH_PHYSX
 
 #if WITH_BOX2D
@@ -3469,7 +3528,9 @@ bool FBodyInstance::LineTrace(struct FHitResult& OutHit, const FVector& Start, c
 	if (DeltaMag > KINDA_SMALL_NUMBER)
 	{
 #if WITH_PHYSX
-		const PxRigidActor* RigidBody = WeldParent ? WeldParent->GetPxRigidActor() : GetPxRigidActor();
+		ExecuteOnPhysicsReadOnly([&]
+		{
+			const PxRigidActor* RigidBody = WeldParent ? WeldParent->GetPxRigidActor_AssumesLocked() : GetPxRigidActor_AssumesLocked();
 		if ((RigidBody != NULL) && (RigidBody->getNbShapes() != 0))
 		{
 			// Create filter data used to filter collisions, should always return eTOUCH for LineTraceComponent		
@@ -3535,6 +3596,8 @@ bool FBodyInstance::LineTrace(struct FHitResult& OutHit, const FVector& Start, c
 				bHitSomething = true;
 			}
 		}
+		});
+		
 #endif //WITH_PHYSX
 
 #if WITH_BOX2D
@@ -3561,12 +3624,17 @@ bool FBodyInstance::Sweep(struct FHitResult& OutHit, const FVector& Start, const
 
 		bool bSweepHit = false;
 #if WITH_PHYSX
-		const PxRigidActor* RigidBody = WeldParent ? WeldParent->GetPxRigidActor() : GetPxRigidActor();
+		ExecuteOnPhysicsReadOnly([&]
+		{
+			const PxRigidActor* RigidBody = WeldParent ? WeldParent->GetPxRigidActor_AssumesLocked() : GetPxRigidActor_AssumesLocked();
+
 		if ((RigidBody != NULL) && (RigidBody->getNbShapes() != 0) && (OwnerComponent != NULL))
 		{
 			FPhysXShapeAdaptor ShapeAdaptor(FQuat::Identity, CollisionShape);
 			bSweepHit = InternalSweepPhysX(OutHit, Start, End, CollisionShape, bTraceComplex, RigidBody, &ShapeAdaptor.GetGeometry());
 		}
+		});
+		
 #endif //WITH_PHYSX
 
 #if WITH_BOX2D
@@ -3648,32 +3716,37 @@ bool FBodyInstance::InternalSweepPhysX(struct FHitResult& OutHit, const FVector&
 float FBodyInstance::GetDistanceToBody(const FVector& Point, FVector& OutPointOnBody) const
 {
 	OutPointOnBody = Point;
+	float ReturnDistance = -1.f;
+	float MinDistanceSqr = BIG_NUMBER;
+	bool bFoundValidBody = false;
+	bool bEarlyOut = true;
 
 #if WITH_PHYSX
-	const PxRigidActor* RigidBody = GetPxRigidActor();
-	if (RigidBody == NULL || RigidBody->getNbShapes()==0 || OwnerComponent==NULL)
+	ExecuteOnPxRigidActorReadOnly(this, [&](const PxRigidActor* RigidActor)
 	{
-		return -1.f;
+		if (RigidActor->getNbShapes() == 0 || OwnerComponent == NULL)
+		{
+			return;
 	}
 
+		bEarlyOut = false;
+
 	// Get all the shapes from the actor
 	TArray<PxShape*, TInlineAllocator<16>> PShapes;
-	PShapes.AddZeroed(RigidBody->getNbShapes());
-	int32 NumShapes = RigidBody->getShapes(PShapes.GetData(), PShapes.Num());
+		PShapes.AddZeroed(RigidActor->getNbShapes());
+		int32 NumShapes = RigidActor->getShapes(PShapes.GetData(), PShapes.Num());
 
 	const PxVec3 PPoint = U2PVector(Point);
-	float MinDistanceSqr=BIG_NUMBER;
-	bool bFoundValidBody = false;
 
 	// Iterate over each shape
-	for(int32 ShapeIdx=0; ShapeIdx<PShapes.Num(); ShapeIdx++)
+		for (int32 ShapeIdx = 0; ShapeIdx < PShapes.Num(); ShapeIdx++)
 	{
 		PxShape* PShape = PShapes[ShapeIdx];
 		check(PShape);
 		PxGeometry& PGeom = PShape->getGeometry().any();
-		PxTransform PGlobalPose = PxShapeExt::getGlobalPose(*PShape, *RigidBody);
+			PxTransform PGlobalPose = PxShapeExt::getGlobalPose(*PShape, *RigidActor);
 		PxGeometryType::Enum GeomType = PShape->getGeometryType();
-		
+
 		if (GeomType == PxGeometryType::eTRIANGLEMESH)
 		{
 			// Type unsupported for this function, but some other shapes will probably work. 
@@ -3684,39 +3757,81 @@ float FBodyInstance::GetDistanceToBody(const FVector& Point, FVector& OutPointOn
 		PxVec3 PClosestPoint;
 		float SqrDistance = PxGeometryQuery::pointDistance(PPoint, PGeom, PGlobalPose, &PClosestPoint);
 		// distance has valid data and smaller than mindistance
-		if ( SqrDistance > 0.f && MinDistanceSqr > SqrDistance )
+			if (SqrDistance > 0.f && MinDistanceSqr > SqrDistance)
 		{
 			MinDistanceSqr = SqrDistance;
 			OutPointOnBody = P2UVector(PClosestPoint);
 		}
-		else if ( SqrDistance == 0.f )
+			else if (SqrDistance == 0.f)
 		{
 			MinDistanceSqr = 0.f;
 			break;
 		}
-	}
+		}	
+	});
+#endif //WITH_PHYSX
 
-	if (!bFoundValidBody)
+	if (!bFoundValidBody && !bEarlyOut)
 	{
 		UE_LOG(LogPhysics, Warning, TEXT("GetDistanceToBody: Component (%s) has no simple collision and cannot be queried for closest point."), OwnerComponent.Get() ? *OwnerComponent->GetPathName() : TEXT("NONE"));
 	}
 
 	return (bFoundValidBody ? FMath::Sqrt(MinDistanceSqr) : -1.f);
-#else
-	// returns -1 if shape is not found
-	return -1.f;
-#endif //WITH_PHYSX
 
 	//@TODO: BOX2D: Implement DistanceToBody
 }
 
+bool FBodyInstance::OverlapTestForBodies(const FVector& Pos, const FQuat& Rot, const TArray<FBodyInstance*>& Bodies) const
+{
+	bool bHaveOverlap = false;
+#if WITH_PHYSX
+
+	ExecuteOnPxRigidBodyReadOnly(this, [&] (const PxRigidBody* TargetRigidBody)
+	{
+		// calculate the test global pose of the rigid body
+		PxTransform PTestGlobalPose = U2PTransform(FTransform(Rot, Pos));
+
+		// Get all the shapes from the actor
+		TArray<PxShape*> PTargetShapes;
+		PTargetShapes.AddZeroed(TargetRigidBody->getNbShapes());
+		int32 NumTargetShapes = TargetRigidBody->getShapes(PTargetShapes.GetData(), PTargetShapes.Num());
+
+		for (int32 TargetShapeIdx = 0; TargetShapeIdx < PTargetShapes.Num(); ++TargetShapeIdx)
+		{
+			const PxShape * PTargetShape = PTargetShapes[TargetShapeIdx];
+
+			// Calc shape global pose
+			PxTransform PShapeGlobalPose = PTestGlobalPose.transform(PTargetShape->getLocalPose());
+
+			GET_GEOMETRY_FROM_SHAPE(PGeom, PTargetShape);
+
+			if (PGeom)
+			{
+				for (const FBodyInstance* BodyInstance : Bodies)
+				{
+					bHaveOverlap = BodyInstance->OverlapPhysX_AssumesLocked(*PGeom, PShapeGlobalPose);
+					if (bHaveOverlap)
+					{
+						return;
+					}
+				}
+			}
+		}
+	});
+#endif //WITH_PHYSX
+	return bHaveOverlap;
+}
+
 bool FBodyInstance::OverlapTest(const FVector& Position, const FQuat& Rotation, const struct FCollisionShape& CollisionShape, FMTDResult* OutMTD) const
 {
 	bool bHasOverlap = false;
 
 #if WITH_PHYSX
+	ExecuteOnPxRigidActorReadOnly(this, [&] (const PxRigidActor* )
+	{
 	FPhysXShapeAdaptor ShapeAdaptor(Rotation, CollisionShape);
-	bHasOverlap = OverlapPhysX(ShapeAdaptor.GetGeometry(), ShapeAdaptor.GetGeomPose(Position), OutMTD);
+		bHasOverlap = OverlapPhysX_AssumesLocked(ShapeAdaptor.GetGeometry(), ShapeAdaptor.GetGeomPose(Position), OutMTD);
+	});
 #endif
 
 #if WITH_BOX2D
@@ -3780,15 +3895,14 @@ bool FBodyInstance::OverlapMulti(TArray<struct FOverlapResult>& InOutOverlaps, c
 	}
 
 #if WITH_PHYSX
-	PxRigidActor* PRigidActor = WeldParent ? WeldParent->GetPxRigidActor() : GetPxRigidActor();
+	PxRigidActor* PRigidActor = WeldParent ? WeldParent->GetPxRigidActor_AssumesLocked() : GetPxRigidActor_AssumesLocked();
 	if (PRigidActor)
 	{
 		const PxTransform PBodyInstanceSpaceToTestSpace = U2PTransform(BodyInstanceSpaceToTestSpace);
 
-		// Get all the shapes from the actor
+		ExecuteOnPhysicsReadOnly([&]
 		{
-			SCOPED_SCENE_READ_LOCK(PRigidActor->getScene());
-
+		// Get all the shapes from the actor
 			TArray<PxShape*, TInlineAllocator<16>> PShapes;
 			PShapes.AddZeroed(PRigidActor->getNbShapes());
 			int32 NumShapes = PRigidActor->getShapes(PShapes.GetData(), PShapes.Num());
@@ -3821,7 +3935,7 @@ bool FBodyInstance::OverlapMulti(TArray<struct FOverlapResult>& InOutOverlaps, c
 					InOutOverlaps.Append(TempOverlaps);
 				}
 			}
-		}
+		});
 	}
 #endif
 
@@ -3837,9 +3951,9 @@ bool FBodyInstance::OverlapMulti(TArray<struct FOverlapResult>& InOutOverlaps, c
 
 
 #if WITH_PHYSX
-bool FBodyInstance::OverlapPhysX(const PxGeometry& PGeom, const PxTransform& ShapePose, FMTDResult* OutMTD) const
+bool FBodyInstance::OverlapPhysX_AssumesLocked(const PxGeometry& PGeom, const PxTransform& ShapePose, FMTDResult* OutMTD) const
 {
-	const PxRigidActor* RigidBody = WeldParent ? WeldParent->GetPxRigidActor() : GetPxRigidActor();
+	const PxRigidActor* RigidBody = WeldParent ? WeldParent->GetPxRigidActor_AssumesLocked() : GetPxRigidActor_AssumesLocked();
 
 	if (RigidBody==NULL || RigidBody->getNbShapes()==0)
 	{
@@ -4010,7 +4124,7 @@ void FBodyInstance::SetUseAsyncScene(bool bNewUseAsyncScene)
 	bUseAsyncScene = bNewUseAsyncScene;
 }
 
-void FBodyInstance::ApplyMaterialToShape(PxShape* PShape, PxMaterial* PSimpleMat, TArray<UPhysicalMaterial*>& ComplexPhysMats)
+void FBodyInstance::ApplyMaterialToShape_AssumesLocked(PxShape* PShape, PxMaterial* PSimpleMat, TArray<UPhysicalMaterial*>& ComplexPhysMats)
 {
 	// If a triangle mesh, need to get array of materials...
 	if(PShape->getGeometryType() == PxGeometryType::eTRIANGLEMESH)
@@ -4042,16 +4156,16 @@ void FBodyInstance::ApplyMaterialToShape(PxShape* PShape, PxMaterial* PSimpleMat
 	}
 }
 
-void FBodyInstance::ApplyMaterialToInstanceShapes(PxMaterial* PSimpleMat, TArray<UPhysicalMaterial*>& ComplexPhysMats)
+void FBodyInstance::ApplyMaterialToInstanceShapes_AssumesLocked(PxMaterial* PSimpleMat, TArray<UPhysicalMaterial*>& ComplexPhysMats)
 {
-	int32 NumSyncShapes = 0;
-	TArray<PxShape*> AllShapes = GetAllShapes(NumSyncShapes);
+	TArray<PxShape*> AllShapes;
+	GetAllShapes_AssumesLocked(AllShapes);
 
 	for(int32 ShapeIdx = 0; ShapeIdx < AllShapes.Num(); ShapeIdx++)
 	{
 		PxShape* PShape = AllShapes[ShapeIdx];
 
-		ApplyMaterialToShape(PShape, PSimpleMat, ComplexPhysMats);
+		ApplyMaterialToShape_AssumesLocked(PShape, PSimpleMat, ComplexPhysMats);
 	}
 }
 
@@ -4087,9 +4201,9 @@ bool FBodyInstance::ValidateTransform(const FTransform &Transform, FString& Debu
 }
 
 #if WITH_PHYSX
-void FBodyInstance::InitDynamicProperties()
+void FBodyInstance::InitDynamicProperties_AssumesLocked()
 {
-	if(PxRigidDynamic* RigidActor = GetPxRigidDynamic())
+	if(PxRigidDynamic* RigidActor = GetPxRigidDynamic_AssumesLocked())
 	{
 		UpdateMassProperties();
 		UpdateDampingProperties();
@@ -4115,7 +4229,7 @@ void FBodyInstance::InitDynamicProperties()
 		if(RigidActor->getScene())
 		{
 			CreateDOFLock();
-			if(IsRigidBodyNonKinematic(RigidActor))
+			if(IsRigidBodyNonKinematic_AssumesLocked(RigidActor))
 			{
 				if(bStartAwake || bWokenExternally)
 				{
@@ -4130,7 +4244,7 @@ void FBodyInstance::InitDynamicProperties()
 	}
 }
 
-void FBodyInstance::GetFilterData(FShapeData& ShapeData, bool bForceSimpleAsComplex)//PxFilterData& SimData, PxFilterData& QueryDataSimple, PxFilterData& QueryDataComplex, TEnumAsByte<ECollisionEnabled::Type>& UseCollisionEnabled, bool bForceSimpleAsComplex)
+void FBodyInstance::GetFilterData_AssumesLocked(FShapeData& ShapeData, bool bForceSimpleAsComplex)//PxFilterData& SimData, PxFilterData& QueryDataSimple, PxFilterData& QueryDataComplex, TEnumAsByte<ECollisionEnabled::Type>& UseCollisionEnabled, bool bForceSimpleAsComplex)
 {
 	// Do nothing if no physics actor
 	if(!IsValidBodyInstance())
@@ -4219,7 +4333,7 @@ void FBodyInstance::GetFilterData(FShapeData& ShapeData, bool bForceSimpleAsComp
 	const bool bUseComplexAsSimple = !bForceSimpleAsComplex && (BodySetup.Get()->CollisionTraceFlag == CTF_UseComplexAsSimple);
 	const bool bUseSimpleAsComplex = bForceSimpleAsComplex || (BodySetup.Get()->CollisionTraceFlag == CTF_UseSimpleAsComplex);
 
-	if(PxRigidActor* PActor = GetPxRigidActor())
+	if(GetPxRigidActor_AssumesLocked())
 	{
 		if(ShapeData.CollisionEnabled != ECollisionEnabled::NoCollision)
 		{
@@ -4308,7 +4422,7 @@ void FBodyInstance::InitStaticBodies(TArray<FBodyInstance*>& Bodies, TArray<FTra
 		bool bInitFail = false;
 
 		// Don't process if we've already got a body
-		if(Instance->GetPxRigidActor())
+		if(Instance->GetPxRigidActor_AssumesLocked())
 		{
 			Instance->OwnerComponent = nullptr;
 			Instance->BodySetup = nullptr;
@@ -4340,8 +4454,8 @@ void FBodyInstance::InitStaticBodies(TArray<FBodyInstance*>& Bodies, TArray<FTra
 		PxFilterData QuerySimpleFilter;
 		PxFilterData QueryComplexFilter;
 		TEnumAsByte<ECollisionEnabled::Type> UseCollisionEnabled;
-		Instance->GetFilterData(ShapeData);
-		Instance->GetShapeFlags(ShapeData, Instance->CollisionEnabled, BodySetup->CollisionTraceFlag == CTF_UseComplexAsSimple);
+		Instance->GetFilterData_AssumesLocked(ShapeData);
+		Instance->GetShapeFlags_AssumesLocked(ShapeData, Instance->CollisionEnabled, BodySetup->CollisionTraceFlag == CTF_UseComplexAsSimple);
 
 		Instance->GetComplexPhysicalMaterials(ComplexMats);
 		SimpleMat = Instance->GetSimplePhysicalMaterial();
@@ -4349,12 +4463,12 @@ void FBodyInstance::InitStaticBodies(TArray<FBodyInstance*>& Bodies, TArray<FTra
 
 		if(PNewActorSync)
 		{
-			BodySetup->AddShapesToRigidActor(Instance, PNewActorSync, PST_Sync, Instance->Scale3D, PhysxSimpleMat, ComplexMats, ShapeData, FTransform::Identity);
+			BodySetup->AddShapesToRigidActor_AssumesLocked(Instance, PNewActorSync, PST_Sync, Instance->Scale3D, PhysxSimpleMat, ComplexMats, ShapeData, FTransform::Identity);
 		}
 		
 		if(PNewActorAsync)
 		{
-			BodySetup->AddShapesToRigidActor(Instance, PNewActorAsync, PST_Async, Instance->Scale3D, PhysxSimpleMat, ComplexMats, ShapeData, FTransform::Identity);
+			BodySetup->AddShapesToRigidActor_AssumesLocked(Instance, PNewActorAsync, PST_Async, Instance->Scale3D, PhysxSimpleMat, ComplexMats, ShapeData, FTransform::Identity);
 		}
 	}
 
@@ -4365,7 +4479,7 @@ void FBodyInstance::InitStaticBodies(TArray<FBodyInstance*>& Bodies, TArray<FTra
 	}
 }
 
-void FBodyInstance::SetShapeFlags(TEnumAsByte<ECollisionEnabled::Type> UseCollisionEnabled, PxShape* PShape, EPhysicsSceneType SceneType, const bool bUseComplexAsSimple)
+void FBodyInstance::SetShapeFlags_AssumesLocked(TEnumAsByte<ECollisionEnabled::Type> UseCollisionEnabled, PxShape* PShape, EPhysicsSceneType SceneType, const bool bUseComplexAsSimple)
 {
 	// If query collision is enabled..
 	bool bUpdateMassProperties = false;
@@ -4426,7 +4540,7 @@ void FBodyInstance::SetShapeFlags(TEnumAsByte<ECollisionEnabled::Type> UseCollis
 			}
 
 			// enable swept bounds for CCD for this shape
-			PxRigidBody* PBody = GetPxRigidActor()->is<PxRigidBody>();
+			PxRigidBody* PBody = GetPxRigidActor_AssumesLocked()->is<PxRigidBody>();
 			if(bSimCollision && !bPhysicsStatic && bUseCCD && PBody)
 			{
 				PBody->setRigidBodyFlag(PxRigidBodyFlag::eENABLE_CCD, true);
@@ -4451,7 +4565,7 @@ void FBodyInstance::SetShapeFlags(TEnumAsByte<ECollisionEnabled::Type> UseCollis
 	}
 }
 
-void FBodyInstance::GetShapeFlags(FShapeData& ShapeData, TEnumAsByte<ECollisionEnabled::Type> UseCollisionEnabled, const bool bUseComplexAsSimple /*= false*/)
+void FBodyInstance::GetShapeFlags_AssumesLocked(FShapeData& ShapeData, TEnumAsByte<ECollisionEnabled::Type> UseCollisionEnabled, const bool bUseComplexAsSimple /*= false*/)
 {
 	// If query collision is enabled..
 	ShapeData.SyncShapeFlags = PxShapeFlags(0);
@@ -4500,7 +4614,7 @@ void FBodyInstance::GetShapeFlags(FShapeData& ShapeData, TEnumAsByte<ECollisionE
 		}
 
 		// enable swept bounds for CCD for this shape
-		PxRigidBody* PBody = GetPxRigidActor()->is<PxRigidBody>();
+		PxRigidBody* PBody = GetPxRigidActor_AssumesLocked()->is<PxRigidBody>();
 		if(bSimCollision && !bPhysicsStatic && bUseCCD && PBody)
 		{
 			ShapeData.SyncBodyFlags |= PxRigidBodyFlag::eENABLE_CCD;
diff --git a/Engine/Source/Runtime/Engine/Private/PhysicsEngine/BodySetup.cpp b/Engine/Source/Runtime/Engine/Private/PhysicsEngine/BodySetup.cpp
--- a/Engine/Source/Runtime/Engine/Private/PhysicsEngine/BodySetup.cpp
+++ b/Engine/Source/Runtime/Engine/Private/PhysicsEngine/BodySetup.cpp
@@ -313,7 +313,7 @@ PxMaterial* GetDefaultPhysMaterial()
 	return GEngine->DefaultPhysMaterial->GetPhysXMaterial();
 }
 
-void UBodySetup::AddSpheresToRigidActor(PxRigidActor* PDestActor, const FTransform& RelativeTM, float MinScale, float MinScaleAbs, TArray<PxShape*>* NewShapes) const
+void UBodySetup::AddSpheresToRigidActor_AssumesLocked(PxRigidActor* PDestActor, const FTransform& RelativeTM, float MinScale, float MinScaleAbs, TArray<PxShape*>* NewShapes) const
 {
 	float ContactOffsetFactor, MaxContactOffset;
 	GetContactOffsetParams(ContactOffsetFactor, MaxContactOffset);
@@ -333,15 +333,19 @@ void UBodySetup::AddSpheresToRigidActor(PxRigidActor* PDestActor, const FTransfo
 			PLocalPose.p *= MinScale;
 
 			ensure(PLocalPose.isValid());
-			PxShape* NewShape = PDestActor->createShape(PSphereGeom, *PDefaultMat, PLocalPose);
-
-			if (NewShapes)
 			{
-				NewShapes->Add(NewShape);
+				PxShape* NewShape = PDestActor->createShape(PSphereGeom, *PDefaultMat, PLocalPose);
+				if (NewShapes)
+				{
+					NewShapes->Add(NewShape);
+				}
+
+				const float ContactOffset = FMath::Min(MaxContactOffset, ContactOffsetFactor * PSphereGeom.radius);
+				NewShape->setContactOffset(ContactOffset);
 			}
+			
 
-			const float ContactOffset = FMath::Min(MaxContactOffset, ContactOffsetFactor * PSphereGeom.radius);
-			NewShape->setContactOffset(ContactOffset);
+			
 		}
 		else
 		{
@@ -350,7 +354,7 @@ void UBodySetup::AddSpheresToRigidActor(PxRigidActor* PDestActor, const FTransfo
 	}
 }
 
-void UBodySetup::AddBoxesToRigidActor(PxRigidActor* PDestActor, const FTransform& RelativeTM, const FVector& Scale3D, const FVector& Scale3DAbs, TArray<PxShape*>* NewShapes) const
+void UBodySetup::AddBoxesToRigidActor_AssumesLocked(PxRigidActor* PDestActor, const FTransform& RelativeTM, const FVector& Scale3D, const FVector& Scale3DAbs, TArray<PxShape*>* NewShapes) const
 {
 	float ContactOffsetFactor, MaxContactOffset;
 	GetContactOffsetParams(ContactOffsetFactor, MaxContactOffset);
@@ -374,15 +378,20 @@ void UBodySetup::AddBoxesToRigidActor(PxRigidActor* PDestActor, const FTransform
 			PLocalPose.p.z *= Scale3D.Z;
 
 			ensure(PLocalPose.isValid());
-			PxShape* NewShape = PDestActor->createShape(PBoxGeom, *PDefaultMat, PLocalPose);
-
-			if (NewShapes)
 			{
-				NewShapes->Add(NewShape);
+				PxShape* NewShape = PDestActor->createShape(PBoxGeom, *PDefaultMat, PLocalPose);
+
+				if (NewShapes)
+				{
+					NewShapes->Add(NewShape);
+				}
+
+				const float ContactOffset = FMath::Min(MaxContactOffset, ContactOffsetFactor * PBoxGeom.halfExtents.minElement());
+				NewShape->setContactOffset(ContactOffset);
 			}
+			
 
-			const float ContactOffset = FMath::Min(MaxContactOffset, ContactOffsetFactor * PBoxGeom.halfExtents.minElement());
-			NewShape->setContactOffset(ContactOffset);
+			
 		}
 		else
 		{
@@ -391,7 +400,7 @@ void UBodySetup::AddBoxesToRigidActor(PxRigidActor* PDestActor, const FTransform
 	}
 }
 
-void UBodySetup::AddSphylsToRigidActor(PxRigidActor* PDestActor, const FTransform& RelativeTM, const FVector& Scale3D, const FVector& Scale3DAbs, TArray<PxShape*>* NewShapes) const
+void UBodySetup::AddSphylsToRigidActor_AssumesLocked(PxRigidActor* PDestActor, const FTransform& RelativeTM, const FVector& Scale3D, const FVector& Scale3DAbs, TArray<PxShape*>* NewShapes) const
 {
 	float ContactOffsetFactor, MaxContactOffset;
 	GetContactOffsetParams(ContactOffsetFactor, MaxContactOffset);
@@ -426,15 +435,20 @@ void UBodySetup::AddSphylsToRigidActor(PxRigidActor* PDestActor, const FTransfor
 			PLocalPose.p.z *= Scale3D.Z;
 
 			ensure(PLocalPose.isValid());
-			PxShape* NewShape = PDestActor->createShape(PCapsuleGeom, *PDefaultMat, PLocalPose);
-
-			if (NewShapes)
 			{
-				NewShapes->Add(NewShape);
+				PxShape* NewShape = PDestActor->createShape(PCapsuleGeom, *PDefaultMat, PLocalPose);
+
+				if (NewShapes)
+				{
+					NewShapes->Add(NewShape);
+				}
+
+				const float ContactOffset = FMath::Min(MaxContactOffset, ContactOffsetFactor * PCapsuleGeom.radius);
+				NewShape->setContactOffset(ContactOffset);
 			}
+			
 
-			const float ContactOffset = FMath::Min(MaxContactOffset, ContactOffsetFactor * PCapsuleGeom.radius);
-			NewShape->setContactOffset(ContactOffset);
+			
 		}
 		else
 		{
@@ -443,7 +457,7 @@ void UBodySetup::AddSphylsToRigidActor(PxRigidActor* PDestActor, const FTransfor
 	}
 }
 
-void UBodySetup::AddConvexElemsToRigidActor(PxRigidActor* PDestActor, const FTransform& RelativeTM, const FVector& Scale3D, const FVector& Scale3DAbs, TArray<PxShape*>* NewShapes) const
+void UBodySetup::AddConvexElemsToRigidActor_AssumesLocked(PxRigidActor* PDestActor, const FTransform& RelativeTM, const FVector& Scale3D, const FVector& Scale3DAbs, TArray<PxShape*>* NewShapes) const
 {
 	float ContactOffsetFactor, MaxContactOffset;
 	GetContactOffsetParams(ContactOffsetFactor, MaxContactOffset);
@@ -480,15 +494,20 @@ void UBodySetup::AddConvexElemsToRigidActor(PxRigidActor* PDestActor, const FTra
 					PxVec3 PBoundsExtents = PConvexGeom.convexMesh->getLocalBounds().getExtents();
 
 					ensure(PLocalPose.isValid());
-					PxShape* NewShape = PDestActor->createShape(PConvexGeom, *PDefaultMat, PLocalPose);
-
-					if (NewShapes)
 					{
-						NewShapes->Add(NewShape);
+						PxShape* NewShape = PDestActor->createShape(PConvexGeom, *PDefaultMat, PLocalPose);
+
+						if (NewShapes)
+						{
+							NewShapes->Add(NewShape);
+						}
+
+						const float ContactOffset = FMath::Min(MaxContactOffset, ContactOffsetFactor * PBoundsExtents.minElement());
+						NewShape->setContactOffset(ContactOffset);
 					}
+					
 
-					const float ContactOffset = FMath::Min(MaxContactOffset, ContactOffsetFactor * PBoundsExtents.minElement());
-					NewShape->setContactOffset(ContactOffset);
+					
 				}
 				else
 				{
@@ -507,7 +526,7 @@ void UBodySetup::AddConvexElemsToRigidActor(PxRigidActor* PDestActor, const FTra
 	}
 }
 
-void UBodySetup::AddTriMeshToRigidActor(PxRigidActor* PDestActor, const FVector& Scale3D, const FVector& Scale3DAbs, TArray<PxShape*>* NewShapes) const
+void UBodySetup::AddTriMeshToRigidActor_AssumesLocked(PxRigidActor* PDestActor, const FVector& Scale3D, const FVector& Scale3DAbs, TArray<PxShape*>* NewShapes) const
 {
 	float ContactOffsetFactor, MaxContactOffset;
 	GetContactOffsetParams(ContactOffsetFactor, MaxContactOffset);
@@ -543,21 +562,25 @@ void UBodySetup::AddTriMeshToRigidActor(PxRigidActor* PDestActor, const FVector&
 				ensure(PLocalPose.isValid());
 
 				// Create without 'sim shape' flag, problematic if it's kinematic, and it gets set later anyway.
-				PxShape* NewShape = PDestActor->createShape(PTriMeshGeom, *PDefaultMat, PxShapeFlag::eSCENE_QUERY_SHAPE | PxShapeFlag::eVISUALIZATION);
-				if (NewShape)
 				{
-					NewShape->setLocalPose(PLocalPose);
-					NewShape->setContactOffset(MaxContactOffset);
+					PxShape* NewShape = PDestActor->createShape(PTriMeshGeom, *PDefaultMat, PxShapeFlag::eSCENE_QUERY_SHAPE | PxShapeFlag::eVISUALIZATION);
 
-					if(NewShapes)
+					if (NewShape)
 					{
-						NewShapes->Add(NewShape);
+						NewShape->setLocalPose(PLocalPose);
+						NewShape->setContactOffset(MaxContactOffset);
+
+						if (NewShapes)
+						{
+							NewShapes->Add(NewShape);
+						}
+					}
+					else
+					{
+						UE_LOG(LogPhysics, Log, TEXT("Can't create new mesh shape in AddShapesToRigidActor"));
 					}
 				}
-				else
-				{
-					UE_LOG(LogPhysics, Log, TEXT("Can't create new mesh shape in AddShapesToRigidActor"));
-				}
+				
 			}
 			else
 			{
@@ -567,7 +590,7 @@ void UBodySetup::AddTriMeshToRigidActor(PxRigidActor* PDestActor, const FVector&
 	}
 }
 
-void UBodySetup::AddShapesToRigidActor(FBodyInstance* OwningInstance, physx::PxRigidActor* PDestActor, EPhysicsSceneType SceneType, FVector& Scale3D, physx::PxMaterial* SimpleMaterial, TArray<UPhysicalMaterial*>& ComplexMaterials, FShapeData& ShapeData, const FTransform& RelativeTM, TArray<physx::PxShape*>* NewShapes)
+void UBodySetup::AddShapesToRigidActor_AssumesLocked(FBodyInstance* OwningInstance, physx::PxRigidActor* PDestActor, EPhysicsSceneType SceneType, FVector& Scale3D, physx::PxMaterial* SimpleMaterial, TArray<UPhysicalMaterial*>& ComplexMaterials, FShapeData& ShapeData, const FTransform& RelativeTM, TArray<physx::PxShape*>* NewShapes)
 {
 #if WITH_RUNTIME_PHYSICS_COOKING || WITH_EDITOR
 	// in editor, there are a lot of things relying on body setup to create physics meshes
@@ -598,81 +621,84 @@ void UBodySetup::AddShapesToRigidActor(FBodyInstance* OwningInstance, physx::PxR
 	// for simple queries as well
 	if (CollisionTraceFlag != ECollisionTraceFlag::CTF_UseComplexAsSimple)
 	{
-		AddSpheresToRigidActor(PDestActor, RelativeTM, MinScale, MinScaleAbs, &Shapes);
-		AddBoxesToRigidActor(PDestActor, RelativeTM, Scale3D, Scale3DAbs, &Shapes);
-		AddSphylsToRigidActor(PDestActor, RelativeTM, Scale3D, Scale3DAbs, &Shapes);
-		AddConvexElemsToRigidActor(PDestActor, RelativeTM, Scale3D, Scale3DAbs, &Shapes);
+		AddSpheresToRigidActor_AssumesLocked(PDestActor, RelativeTM, MinScale, MinScaleAbs, &Shapes);
+		AddBoxesToRigidActor_AssumesLocked(PDestActor, RelativeTM, Scale3D, Scale3DAbs, &Shapes);
+		AddSphylsToRigidActor_AssumesLocked(PDestActor, RelativeTM, Scale3D, Scale3DAbs, &Shapes);
+		AddConvexElemsToRigidActor_AssumesLocked(PDestActor, RelativeTM, Scale3D, Scale3DAbs, &Shapes);
 	}
 
 	// Create tri-mesh shape, when we are not using simple collision shapes for 
 	// complex queries as well
 	if( CollisionTraceFlag != ECollisionTraceFlag::CTF_UseSimpleAsComplex )
 	{
-		AddTriMeshToRigidActor(PDestActor, Scale3D, Scale3DAbs, &Shapes);
+		AddTriMeshToRigidActor_AssumesLocked(PDestActor, Scale3D, Scale3DAbs, &Shapes);
 	}
 
-	for(PxShape* Shape : Shapes)
 	{
-		FShapeFilterData& Filters = ShapeData.FilterData;
-		// Apply provided materials
-		// If a triangle mesh, need to get array of materials...
-		if(Shape->getGeometryType() == PxGeometryType::eTRIANGLEMESH)
+		for (PxShape* Shape : Shapes)
 		{
-			TArray<PxMaterial*> PComplexMats;
-			PComplexMats.AddUninitialized(ComplexMaterials.Num());
-			for(int MatIdx = 0; MatIdx < ComplexMaterials.Num(); MatIdx++)
+			FShapeFilterData& Filters = ShapeData.FilterData;
+			// Apply provided materials
+			// If a triangle mesh, need to get array of materials...
+			if (Shape->getGeometryType() == PxGeometryType::eTRIANGLEMESH)
 			{
-				check(ComplexMaterials[MatIdx] != NULL);
-				PComplexMats[MatIdx] = ComplexMaterials[MatIdx]->GetPhysXMaterial();
-				check(PComplexMats[MatIdx] != NULL);
-			}
+				TArray<PxMaterial*> PComplexMats;
+				PComplexMats.AddUninitialized(ComplexMaterials.Num());
+				for (int MatIdx = 0; MatIdx < ComplexMaterials.Num(); MatIdx++)
+				{
+					check(ComplexMaterials[MatIdx] != NULL);
+					PComplexMats[MatIdx] = ComplexMaterials[MatIdx]->GetPhysXMaterial();
+					check(PComplexMats[MatIdx] != NULL);
+				}
 
-			if(PComplexMats.Num())
-			{
-				Shape->setMaterials(PComplexMats.GetData(), PComplexMats.Num());
-			}
-			else
-			{
-				UE_LOG(LogPhysics, Warning, TEXT("FBodyInstance::UpdatePhysicalMaterials : PComplexMats is empty - falling back on simple physical material."));
-				Shape->setMaterials(&SimpleMaterial, 1);
-			}
+				if (PComplexMats.Num())
+				{
+					Shape->setMaterials(PComplexMats.GetData(), PComplexMats.Num());
+				}
+				else
+				{
+					UE_LOG(LogPhysics, Warning, TEXT("FBodyInstance::UpdatePhysicalMaterials : PComplexMats is empty - falling back on simple physical material."));
+					Shape->setMaterials(&SimpleMaterial, 1);
+				}
 
-			Shape->setQueryFilterData(Filters.QueryComplexFilter);
-			
-			if(SceneType == PST_Sync)
-			{
-				Shape->setFlags(ShapeData.SyncShapeFlags | ShapeData.ComplexShapeFlags);
-			}
-			else
-			{
-				Shape->setFlags(ShapeData.AsyncShapeFlags | ShapeData.ComplexShapeFlags);
-			}
-		}
-		// Simple shape, 
-		else
-		{
-			Shape->setMaterials(&SimpleMaterial, 1);
-			Shape->setQueryFilterData(Filters.QuerySimpleFilter);
+				Shape->setQueryFilterData(Filters.QueryComplexFilter);
 
-			if(SceneType == PST_Sync)
-			{
-				Shape->setFlags(ShapeData.SyncShapeFlags | ShapeData.SimpleShapeFlags);
+				if (SceneType == PST_Sync)
+				{
+					Shape->setFlags(ShapeData.SyncShapeFlags | ShapeData.ComplexShapeFlags);
+				}
+				else
+				{
+					Shape->setFlags(ShapeData.AsyncShapeFlags | ShapeData.ComplexShapeFlags);
+				}
 			}
+			// Simple shape, 
 			else
 			{
-				Shape->setFlags(ShapeData.AsyncShapeFlags | ShapeData.SimpleShapeFlags);
+				Shape->setMaterials(&SimpleMaterial, 1);
+				Shape->setQueryFilterData(Filters.QuerySimpleFilter);
+
+				if (SceneType == PST_Sync)
+				{
+					Shape->setFlags(ShapeData.SyncShapeFlags | ShapeData.SimpleShapeFlags);
+				}
+				else
+				{
+					Shape->setFlags(ShapeData.AsyncShapeFlags | ShapeData.SimpleShapeFlags);
+				}
 			}
-		}
-		Shape->setSimulationFilterData(Filters.SimFilter);
+			Shape->setSimulationFilterData(Filters.SimFilter);
 
-		if(OwningInstance)
-		{
-			if(PxRigidBody* RigidBody = OwningInstance->GetPxRigidActor()->is<PxRigidBody>())
+			if (OwningInstance)
 			{
-				RigidBody->setRigidBodyFlags(ShapeData.SyncBodyFlags);
+				if (PxRigidBody* RigidBody = OwningInstance->GetPxRigidBody_AssumesLocked())
+				{
+					RigidBody->setRigidBodyFlags(ShapeData.SyncBodyFlags);
+				}
 			}
 		}
 	}
+	
 
 	if(NewShapes)
 	{
diff --git a/Engine/Source/Runtime/Engine/Private/PhysicsEngine/ConstraintInstance.cpp b/Engine/Source/Runtime/Engine/Private/PhysicsEngine/ConstraintInstance.cpp
--- a/Engine/Source/Runtime/Engine/Private/PhysicsEngine/ConstraintInstance.cpp
+++ b/Engine/Source/Runtime/Engine/Private/PhysicsEngine/ConstraintInstance.cpp
@@ -37,7 +37,7 @@ FVector RevolutionsToRads(const FVector Revolutions)
 
 #if WITH_PHYSX
 /** Util for setting soft limit params */
-void SetSoftLimitParams(PxJointLimitParameters* PLimit, bool bSoft, float Spring, float Damping)
+void SetSoftLimitParams_AssumesLocked(PxJointLimitParameters* PLimit, bool bSoft, float Spring, float Damping)
 {
 	if(bSoft)
 	{
@@ -76,7 +76,7 @@ PxD6Motion::Enum U2PAngularMotion(EAngularConstraintMotion InMotion)
 
 /** Util for setting linear movement for an axis */
 template <PxD6Axis::Enum PAxis>
-void SetLinearMovement(PxD6Joint* PD6Joint, ELinearConstraintMotion Motion, bool bLockLimitSize)
+void SetLinearMovement_AssumesLocked(PxD6Joint* PD6Joint, ELinearConstraintMotion Motion, bool bLockLimitSize)
 {
 	if(Motion == LCM_Locked || (Motion == LCM_Limited && bLockLimitSize))
 	{
@@ -87,22 +87,54 @@ void SetLinearMovement(PxD6Joint* PD6Joint, ELinearConstraintMotion Motion, bool
 	}
 }
 
-physx::PxD6Joint* FConstraintInstance::GetUnbrokenJoint() const
+physx::PxD6Joint* FConstraintInstance::GetUnbrokenJoint_AssumesLocked() const
 {
 	return (ConstraintData && !(ConstraintData->getConstraintFlags()&PxConstraintFlag::eBROKEN)) ? ConstraintData : nullptr;
 }
+
+bool FConstraintInstance::ExecuteOnUnbrokenJointReadOnly(TFunctionRef<void(const physx::PxD6Joint*)> Func) const
+{
+	if(ConstraintData)
+	{
+		SCOPED_SCENE_READ_LOCK(ConstraintData->getScene());
+
+		if(!(ConstraintData->getConstraintFlags()&PxConstraintFlag::eBROKEN))
+		{
+			Func(ConstraintData);
+			return true;
+		}
+	}
+
+	return false;
+}
+
+bool FConstraintInstance::ExecuteOnUnbrokenJointReadWrite(TFunctionRef<void(physx::PxD6Joint*)> Func) const
+{
+	if (ConstraintData)
+	{
+		SCOPED_SCENE_WRITE_LOCK(ConstraintData->getScene());
+
+		if (!(ConstraintData->getConstraintFlags()&PxConstraintFlag::eBROKEN))
+		{
+			Func(ConstraintData);
+			return true;
+		}
+	}
+
+	return false;
+}
 #endif //WITH_PHYSX
 
 void FConstraintInstance::UpdateLinearLimit()
 {
 #if WITH_PHYSX
-	if (PxD6Joint* Joint = GetUnbrokenJoint())
+	ExecuteOnUnbrokenJointReadWrite([&] (PxD6Joint* Joint)
 	{
 		bool bLockLimitSize = (LinearLimitSize < RB_MinSizeToLockDOF);
 
-		SetLinearMovement<PxD6Axis::eX>(Joint, LinearXMotion, bLockLimitSize);
-		SetLinearMovement<PxD6Axis::eY>(Joint, LinearYMotion, bLockLimitSize);
-		SetLinearMovement<PxD6Axis::eZ>(Joint, LinearZMotion, bLockLimitSize);
+		SetLinearMovement_AssumesLocked<PxD6Axis::eX>(Joint, LinearXMotion, bLockLimitSize);
+		SetLinearMovement_AssumesLocked<PxD6Axis::eY>(Joint, LinearYMotion, bLockLimitSize);
+		SetLinearMovement_AssumesLocked<PxD6Axis::eZ>(Joint, LinearZMotion, bLockLimitSize);
 
 		// If any DOF is locked/limited, set up the joint limit
 		if (LinearXMotion != LCM_Free || LinearYMotion != LCM_Free || LinearZMotion != LCM_Free)
@@ -110,10 +142,10 @@ void FConstraintInstance::UpdateLinearLimit()
 			// If limit drops below RB_MinSizeToLockDOF, just pass RB_MinSizeToLockDOF to physics - that axis will be locked anyway, and PhysX dislikes 0 here
 			float LinearLimit = FMath::Max<float>(LinearLimitSize, RB_MinSizeToLockDOF);
 			PxJointLinearLimit PLinearLimit(GPhysXSDK->getTolerancesScale(), LinearLimit, 0.05f * GPhysXSDK->getTolerancesScale().length);
-			SetSoftLimitParams(&PLinearLimit, bLinearLimitSoft, LinearLimitStiffness*AverageMass, LinearLimitDamping*AverageMass);
+			SetSoftLimitParams_AssumesLocked(&PLinearLimit, bLinearLimitSoft, LinearLimitStiffness*AverageMass, LinearLimitDamping*AverageMass);
 			Joint->setLinearLimit(PLinearLimit);
 		}
-	}
+	});
 #endif
 }
 
@@ -122,7 +154,7 @@ void FConstraintInstance::UpdateLinearLimit()
 void FConstraintInstance::UpdateAngularLimit()
 {
 #if WITH_PHYSX
-	if (PxD6Joint* Joint = GetUnbrokenJoint())
+	ExecuteOnUnbrokenJointReadWrite([&] (PxD6Joint* Joint)
 	{
 		/////////////// TWIST LIMIT
 		PxD6Motion::Enum TwistMotion = PxD6Motion::eFREE;
@@ -132,7 +164,7 @@ void FConstraintInstance::UpdateAngularLimit()
 			// If angle drops below RB_MinAngleToLockDOF, just pass RB_MinAngleToLockDOF to physics - that axis will be locked anyway, and PhysX dislikes 0 here
 			float TwistLimitRad = FMath::DegreesToRadians(TwistLimitAngle);
 			PxJointAngularLimitPair PTwistLimitPair(-TwistLimitRad, TwistLimitRad, FMath::DegreesToRadians(1.f));
-			SetSoftLimitParams(&PTwistLimitPair, bTwistLimitSoft, TwistLimitStiffness*AverageMass, TwistLimitDamping*AverageMass);
+			SetSoftLimitParams_AssumesLocked(&PTwistLimitPair, bTwistLimitSoft, TwistLimitStiffness*AverageMass, TwistLimitDamping*AverageMass);
 			Joint->setTwistLimit(PTwistLimitPair);
 		}
 		else if (AngularTwistMotion == ACM_Locked)
@@ -151,13 +183,13 @@ void FConstraintInstance::UpdateAngularLimit()
 			float Limit1Rad = FMath::DegreesToRadians(FMath::ClampAngle(Swing1LimitAngle, KINDA_SMALL_NUMBER, 179.9999f));
 			float Limit2Rad = FMath::DegreesToRadians(FMath::ClampAngle(Swing2LimitAngle, KINDA_SMALL_NUMBER, 179.9999f));
 			PxJointLimitCone PSwingLimitCone(Limit2Rad, Limit1Rad, FMath::DegreesToRadians(1.f));
-			SetSoftLimitParams(&PSwingLimitCone, bSwingLimitSoft, SwingLimitStiffness*AverageMass, SwingLimitDamping*AverageMass);
+			SetSoftLimitParams_AssumesLocked(&PSwingLimitCone, bSwingLimitSoft, SwingLimitStiffness*AverageMass, SwingLimitDamping*AverageMass);
 			Joint->setSwingLimit(PSwingLimitCone);
 		}
 
 		Joint->setMotion(PxD6Axis::eSWING2, Swing1Motion);
 		Joint->setMotion(PxD6Axis::eSWING1, Swing2Motion);
-	}
+	});
 #endif
 	
 }
@@ -175,13 +207,13 @@ void FConstraintInstance::UpdateBreakable()
 void FConstraintInstance::UpdateDriveTarget()
 {
 #if WITH_PHYSX
-	if (PxD6Joint* PJoint = GetUnbrokenJoint())
+	ExecuteOnUnbrokenJointReadWrite([&] (PxD6Joint* Joint)
 	{
 		FQuat OrientationTargetQuat(AngularOrientationTarget);
 
-		PJoint->setDrivePosition(PxTransform(U2PVector(LinearPositionTarget), U2PQuat(OrientationTargetQuat)));
-		PJoint->setDriveVelocity(U2PVector(LinearVelocityTarget), U2PVector(RevolutionsToRads(AngularVelocityTarget)));
-	}
+		Joint->setDrivePosition(PxTransform(U2PVector(LinearPositionTarget), U2PQuat(OrientationTargetQuat)));
+		Joint->setDriveVelocity(U2PVector(LinearVelocityTarget), U2PVector(RevolutionsToRads(AngularVelocityTarget)));
+	});
 #endif
 }
 
@@ -246,7 +278,7 @@ void FConstraintInstance::SetDisableCollision(bool InDisableCollision)
 {
 	bDisableCollision = InDisableCollision;
 #if WITH_PHYSX
-	if (PxD6Joint* Joint = GetUnbrokenJoint())
+	ExecuteOnUnbrokenJointReadWrite([&] (PxD6Joint* Joint)
 	{
 		PxConstraintFlags Flags = Joint->getConstraintFlags();
 		if (bDisableCollision)
@@ -259,12 +291,12 @@ void FConstraintInstance::SetDisableCollision(bool InDisableCollision)
 		}
 
 		Joint->setConstraintFlags(Flags);
-	}
+	});
 #endif
 }
 
 #if WITH_PHYSX
-float ComputeAverageMass(const PxRigidActor* PActor1, const PxRigidActor* PActor2)
+float ComputeAverageMass_AssumesLocked(const PxRigidActor* PActor1, const PxRigidActor* PActor2)
 {
 	float AverageMass = 0;
 	{
@@ -290,16 +322,41 @@ float ComputeAverageMass(const PxRigidActor* PActor1, const PxRigidActor* PActor
 	return AverageMass;
 }
 
+/** Finds the common scene and appropriate actors for the passed in body instances. Makes sure to do this without requiring a scene lock*/
+PxScene* GetPScene_LockFree(const FBodyInstance* Body1, const FBodyInstance* Body2)
+{
+	const int32 SceneIndex1 = Body1 ? Body1->GetSceneIndex() : -1;
+	const int32 SceneIndex2 = Body2 ? Body2->GetSceneIndex() : -1;
+	PxScene* PScene = nullptr;
+
+	//now we check if the two actors are valid
+	if(SceneIndex1 == -1 && SceneIndex2 == -1)
+	{
+		UE_LOG(LogPhysics, Log, TEXT("Attempting to create a joint between two null actors.  No joint created."));
+	}
+	else if(SceneIndex1 >= 0 && SceneIndex2 >= 0 && SceneIndex1 != SceneIndex2)
+	{
+		UE_LOG(LogPhysics, Log, TEXT("Attempting to create a joint between two actors in different scenes.  No joint created."));
+	}
+	else
+	{
+		PScene = GetPhysXSceneFromIndex(SceneIndex1);
+	}
+
+	return PScene;
+}
+
 /*various logical checks to find the correct physx actor. Returns true if found valid actors that can be constrained*/
-bool GetPActors(const FBodyInstance* Body1, const FBodyInstance* Body2, PxRigidActor** PActor1Out, PxRigidActor** PActor2Out)
+bool GetPActors_AssumesLocked(const FBodyInstance* Body1, const FBodyInstance* Body2, PxRigidActor** PActor1Out, PxRigidActor** PActor2Out)
 {
-	PxRigidActor* PActor1 = Body1 ? Body1->GetPxRigidActor() : NULL;
-	PxRigidActor* PActor2 = Body2 ? Body2->GetPxRigidActor() : NULL;
+	PxRigidActor* PActor1 = Body1 ? Body1->GetPxRigidActor_AssumesLocked() : NULL;
+	PxRigidActor* PActor2 = Body2 ? Body2->GetPxRigidActor_AssumesLocked() : NULL;
 
 	// Do not create joint unless you have two actors
 	// Do not create joint unless one of the actors is dynamic
 	if ((!PActor1 || !PActor1->isRigidBody()) && (!PActor2 || !PActor2->isRigidBody()))
 	{
+		UE_LOG(LogPhysics, Log, TEXT("Attempting to create a joint between actors that are static.  No joint created."));
 		return false;
 	}
 
@@ -309,13 +366,13 @@ bool GetPActors(const FBodyInstance* Body1, const FBodyInstance* Body2, PxRigidA
 		if (PActor1->isRigidStatic() && PActor2->isRigidBody())
 		{
 			const uint32 SceneType = Body2->RigidActorSync != NULL ? PST_Sync : PST_Async;
-			PActor1 = Body1->GetPxRigidActor(SceneType);
+			PActor1 = Body1->GetPxRigidActor_AssumesLocked(SceneType);
 		}
 		else
 		if (PActor2->isRigidStatic() && PActor1->isRigidBody())
 		{
 			const uint32 SceneType = Body1->RigidActorSync != NULL ? PST_Sync : PST_Async;
-			PActor2 = Body2->GetPxRigidActor(SceneType);
+			PActor2 = Body2->GetPxRigidActor_AssumesLocked(SceneType);
 		}
 	}
 
@@ -324,26 +381,7 @@ bool GetPActors(const FBodyInstance* Body1, const FBodyInstance* Body2, PxRigidA
 	return true;
 }
 
-bool GetPScene(const PxRigidActor* PActor1, const PxRigidActor* PActor2, PxScene** PSceneOut)
-{
-	// make sure actors are in same scene
-	PxScene* PScene1 = PActor1 ? PActor1->getScene() : nullptr;
-	PxScene* PScene2 = PActor2 ? PActor2->getScene() : nullptr;
-
-	// make sure actors are in same scene
-	if (PScene1 && PScene2 && PScene1 != PScene2)
-	{
-		UE_LOG(LogPhysics, Log, TEXT("Attempting to create a joint between actors in two different scenes.  No joint created."));
-		return false;
-	}
-
-	*PSceneOut = PScene1 ? PScene1 : PScene2;
-	check(PSceneOut);
-	
-	return true;
-}
-
-bool FConstraintInstance::CreatePxJoint(physx::PxRigidActor* PActor1, physx::PxRigidActor* PActor2, physx::PxScene* PScene, const float Scale)
+bool FConstraintInstance::CreatePxJoint_AssumesLocked(physx::PxRigidActor* PActor1, physx::PxRigidActor* PActor2, physx::PxScene* PScene, const float Scale)
 {
 	ConstraintData = nullptr;
 
@@ -390,7 +428,7 @@ bool FConstraintInstance::CreatePxJoint(physx::PxRigidActor* PActor1, physx::PxR
 	return true;
 }
 
-void FConstraintInstance::UpdateConstraintFlags()
+void FConstraintInstance::UpdateConstraintFlags_AssumesLocked()
 {
 	PxConstraintFlags Flags = PxConstraintFlags();
 
@@ -415,12 +453,12 @@ void FConstraintInstance::UpdateConstraintFlags()
 }
 
 
-void FConstraintInstance::UpdateAverageMass(const PxRigidActor* PActor1, const PxRigidActor* PActor2)
+void FConstraintInstance::UpdateAverageMass_AssumesLocked(const PxRigidActor* PActor1, const PxRigidActor* PActor2)
 {
-	AverageMass = ComputeAverageMass(PActor1, PActor2);
+	AverageMass = ComputeAverageMass_AssumesLocked(PActor1, PActor2);
 }
 
-void EnsureSleepingActorsStaySleeping(PxRigidActor* PActor1, PxRigidActor* PActor2)
+void EnsureSleepingActorsStaySleeping_AssumesLocked(PxRigidActor* PActor1, PxRigidActor* PActor2)
 {
 	// record if actors are asleep before creating joint, so we can sleep them afterwards if so (creating joint wakes them)
 	const bool bActor1Asleep = (PActor1 == nullptr || !PActor1->isRigidDynamic() || PActor1->isRigidDynamic()->isSleeping());
@@ -429,12 +467,12 @@ void EnsureSleepingActorsStaySleeping(PxRigidActor* PActor1, PxRigidActor* PActo
 	// creation of joints wakes up rigid bodies, so we put them to sleep again if both were initially asleep
 	if (bActor1Asleep && bActor2Asleep)
 	{
-		if (PActor1 && IsRigidBodyNonKinematic(PActor1->isRigidDynamic()))
+		if (PActor1 && IsRigidBodyNonKinematic_AssumesLocked(PActor1->isRigidDynamic()))
 		{
 			PActor1->isRigidDynamic()->putToSleep();
 		}
 
-		if (PActor2 && IsRigidBodyNonKinematic(PActor2->isRigidDynamic()))
+		if (PActor2 && IsRigidBodyNonKinematic_AssumesLocked(PActor2->isRigidDynamic()))
 		{
 			PActor2->isRigidDynamic()->putToSleep();
 		}
@@ -459,21 +497,23 @@ void FConstraintInstance::InitConstraint(USceneComponent* Owner, FBodyInstance*
 		TermConstraint();
 	}
 
-	PxRigidActor* PActor1;
-	PxRigidActor* PActor2;
-	PxScene* PScene;
+	PxRigidActor* PActor1 = nullptr;
+	PxRigidActor* PActor2 = nullptr;
+	PxScene* PScene = GetPScene_LockFree(Body1, Body2);
+	SCOPED_SCENE_WRITE_LOCK(PScene);
 
-	const bool bValidConstraintSetup = GetPActors(Body1, Body2, &PActor1, &PActor2) && GetPScene(PActor1, PActor2, &PScene) && CreatePxJoint(PActor1, PActor2, PScene, Scale);
+	const bool bValidConstraintSetup = PScene && GetPActors_AssumesLocked(Body1, Body2, &PActor1, &PActor2) && CreatePxJoint_AssumesLocked(PActor1, PActor2, PScene, Scale);
 	if (!bValidConstraintSetup)
 	{
 		return;
 	}
 
+
 	// update mass
-	UpdateAverageMass(PActor1, PActor2);
+	UpdateAverageMass_AssumesLocked(PActor1, PActor2);
 	
 	//flags and projection settings
-	UpdateConstraintFlags();
+	UpdateConstraintFlags_AssumesLocked();
 	
 	//limits
 	UpdateAngularLimit();
@@ -487,7 +527,7 @@ void FConstraintInstance::InitConstraint(USceneComponent* Owner, FBodyInstance*
 	SetAngularDriveParams(AngularDriveSpring, AngularDriveDamping, AngularDriveForceLimit);
 	UpdateDriveTarget();
 	
-	EnsureSleepingActorsStaySleeping(PActor1, PActor2);
+	EnsureSleepingActorsStaySleeping_AssumesLocked(PActor1, PActor2);
 #endif // WITH_PHYSX
 
 }
@@ -503,6 +543,7 @@ void FConstraintInstance::TermConstraint()
 	// use correct scene
 	if(PxScene* PScene = GetPhysXSceneFromIndex(SceneIndex))
 	{
+		SCOPED_SCENE_WRITE_LOCK(PScene);
 		ConstraintData->release();
 	}
 
@@ -602,11 +643,11 @@ void FConstraintInstance::SetRefFrame(EConstraintFrame::Type Frame, const FTrans
 	}
 
 #if WITH_PHYSX
-	if (PxD6Joint* Joint = GetUnbrokenJoint())
+	ExecuteOnUnbrokenJointReadWrite([&] (PxD6Joint* Joint)
 	{
 		PxTransform PxRefFrame = U2PTransform(RefFrame);
 		Joint->setLocalPose(PxFrame, PxRefFrame);
-	}
+	});
 #endif
 
 }
@@ -627,12 +668,12 @@ void FConstraintInstance::SetRefPosition(EConstraintFrame::Type Frame, const FVe
 	}
 
 #if WITH_PHYSX
-	if (PxD6Joint* Joint = GetUnbrokenJoint())
+	ExecuteOnUnbrokenJointReadWrite([&] (PxD6Joint* Joint)
 	{
 		PxTransform PxRefFrame = ConstraintData->getLocalPose(PxFrame);
 		PxRefFrame.p = U2PVector(RefPosition);
 		Joint->setLocalPose(PxFrame, PxRefFrame);
-	}
+	});
 #endif
 }
 
@@ -657,12 +698,12 @@ void FConstraintInstance::SetRefOrientation(EConstraintFrame::Type Frame, const
 	}
 	
 #if WITH_PHYSX
-	if (PxD6Joint* Joint = GetUnbrokenJoint())
+	ExecuteOnUnbrokenJointReadWrite([&] (PxD6Joint* Joint)
 	{
 		FTransform URefTransform = FTransform(PriAxis1, SecAxis, PriAxis ^ SecAxis, RefPos);
 		PxTransform PxRefFrame = U2PTransform(URefTransform);
 		Joint->setLocalPose(PxFrame, PxRefFrame);
-		}
+	});
 #endif
 }
 
@@ -704,8 +745,10 @@ FVector FConstraintInstance::GetConstraintLocation()
 
 void FConstraintInstance::GetConstraintForce(FVector& OutLinearForce, FVector& OutAngularForce)
 {
+	OutLinearForce = FVector::ZeroVector;
+	OutAngularForce = FVector::ZeroVector;
 #if WITH_PHYSX
-	if (PxD6Joint* Joint = GetUnbrokenJoint())
+	ExecuteOnUnbrokenJointReadOnly([&] (const PxD6Joint* Joint)
 	{
 		PxVec3 PxOutLinearForce;
 		PxVec3 PxOutAngularForce;
@@ -713,16 +756,7 @@ void FConstraintInstance::GetConstraintForce(FVector& OutLinearForce, FVector& O
 
 		OutLinearForce = P2UVector(PxOutLinearForce);
 		OutAngularForce = P2UVector(PxOutAngularForce);
-	}
-	else
-	{
-		OutLinearForce = FVector::ZeroVector;
-		OutAngularForce = FVector::ZeroVector;
-	}
-
-#else
-	OutLinearForce = FVector::ZeroVector;
-	OutAngularForce = FVector::ZeroVector;
+	});
 #endif
 }
 
@@ -732,7 +766,7 @@ void FConstraintInstance::GetConstraintForce(FVector& OutLinearForce, FVector& O
 void FConstraintInstance::SetLinearPositionDrive(bool bEnableXDrive, bool bEnableYDrive, bool bEnableZDrive)
 {
 #if WITH_PHYSX
-	if (PxD6Joint* Joint = GetUnbrokenJoint())
+	ExecuteOnUnbrokenJointReadWrite([&] (PxD6Joint* Joint)
 	{
 		// Get the current drives
 		PxD6JointDrive CurrentDriveX = Joint->getDrive(PxD6Drive::eX);
@@ -746,7 +780,7 @@ void FConstraintInstance::SetLinearPositionDrive(bool bEnableXDrive, bool bEnabl
 		Joint->setDrive(PxD6Drive::eX, CurrentDriveX);
 		Joint->setDrive(PxD6Drive::eY, CurrentDriveY);
 		Joint->setDrive(PxD6Drive::eZ, CurrentDriveZ);
-	}
+	});
 #endif
 
 	bLinearXPositionDrive = bEnableXDrive;
@@ -786,7 +820,7 @@ void FConstraintInstance::SetLinearVelocityDrive(bool bEnableXDrive, bool bEnabl
 void FConstraintInstance::SetAngularPositionDrive(bool bEnableSwingDrive, bool bEnableTwistDrive)
 {
 #if WITH_PHYSX
-	if (PxD6Joint* Joint = GetUnbrokenJoint())
+	ExecuteOnUnbrokenJointReadWrite([&] (PxD6Joint* Joint)
 	{
 		// Get the current drives
 		PxD6JointDrive CurrentDriveSwing = Joint->getDrive(PxD6Drive::eSWING);
@@ -803,7 +837,7 @@ void FConstraintInstance::SetAngularPositionDrive(bool bEnableSwingDrive, bool b
 		Joint->setDrive(PxD6Drive::eSLERP, CurrentDriveSlerp);
 		
 		bAngularOrientationDrive = bEnableSwingDrive || bEnableTwistDrive;
-	}
+	});
 #endif
 }
 
@@ -881,7 +915,7 @@ void FConstraintInstance::SetLinearDriveParams(float InSpring, float InDamping,
 #if WITH_PHYSX
 	if (bLinearPositionDrive || bLinearVelocityDrive)
 	{
-		if (PxD6Joint* Joint = GetUnbrokenJoint())
+		ExecuteOnUnbrokenJointReadWrite([&] (PxD6Joint* Joint)
 		{
 			// X-Axis linear drive
 			const float DriveSpringX = bLinearPositionDrive && bLinearXPositionDrive ? InSpring : 0.0f;
@@ -898,7 +932,7 @@ void FConstraintInstance::SetLinearDriveParams(float InSpring, float InDamping,
 			const float DriveSpringZ = bLinearPositionDrive && bLinearZPositionDrive ? InSpring : 0.0f;
 			const float DriveDampingZ = (bLinearVelocityDrive && LinearVelocityTarget.Z != 0.f) ? InDamping : 0.0f;
 			Joint->setDrive(PxD6Drive::eZ, PxD6JointDrive(DriveSpringZ, DriveDampingZ, LinearForceLimit, bIsAccelerationDrive));
-		}
+		});
 	}
 #endif
 
@@ -998,7 +1032,7 @@ void FConstraintInstance::SetAngularDriveParams(float InSpring, float InDamping,
 #if WITH_PHYSX
 	if (bAngularOrientationDrive || bAngularVelocityDrive)
 	{
-		if (PxD6Joint* Joint = GetUnbrokenJoint())
+		ExecuteOnUnbrokenJointReadWrite([&] (PxD6Joint* Joint)
 		{
 			const float AngularForceLimit = InForceLimit > 0.0f ? InForceLimit : PX_MAX_F32;
 			const float DriveSpring = bAngularOrientationDrive ? InSpring : 0.0f;
@@ -1016,7 +1050,7 @@ void FConstraintInstance::SetAngularDriveParams(float InSpring, float InDamping,
 				Joint->setDrive(PxD6Drive::eSWING, PxD6JointDrive(DriveSpring, DriveDamping, AngularForceLimit, bIsAccelerationDrive));
 				Joint->setDrive(PxD6Drive::eSLERP, PxD6JointDrive());
 			}
-		}
+		});
 	}
 #endif
 
@@ -1029,7 +1063,7 @@ void FConstraintInstance::SetAngularDriveParams(float InSpring, float InDamping,
 void FConstraintInstance::SetAngularDOFLimitScale(float InSwing1LimitScale, float InSwing2LimitScale, float InTwistLimitScale)
 {
 #if WITH_PHYSX
-	if (PxD6Joint* Joint = GetUnbrokenJoint())
+	ExecuteOnUnbrokenJointReadWrite([&] (PxD6Joint* Joint)
 	{
 		if ( AngularSwing1Motion == ACM_Limited || AngularSwing2Motion == ACM_Limited )
 		{
@@ -1075,7 +1109,7 @@ void FConstraintInstance::SetAngularDOFLimitScale(float InSwing1LimitScale, floa
 			Joint->setMotion(PxD6Axis::eTWIST, PxD6Motion::eLOCKED);
 		}
 		
-	}
+	});
 #endif
 }
 
@@ -1083,11 +1117,11 @@ void FConstraintInstance::SetAngularDOFLimitScale(float InSwing1LimitScale, floa
 void FConstraintInstance::SetLinearLimitSize(float NewLimitSize)
 {
 #if WITH_PHYSX
-	if (PxD6Joint* Joint = GetUnbrokenJoint())
+	ExecuteOnUnbrokenJointReadWrite([&] (PxD6Joint* Joint)
 	{
 		PxReal LimitContractDistance =  1.f * (PI/180.0f);
 		Joint->setLinearLimit(PxJointLinearLimit(GPhysXSDK->getTolerancesScale(), NewLimitSize, LimitContractDistance * GPhysXSDK->getTolerancesScale().length)); // LOC_MOD33 need to scale the contactDistance if not using its default value
-	}
+	});
 #endif
 }
 
diff --git a/Engine/Source/Runtime/Engine/Private/PhysicsEngine/PhysAnim.cpp b/Engine/Source/Runtime/Engine/Private/PhysicsEngine/PhysAnim.cpp
--- a/Engine/Source/Runtime/Engine/Private/PhysicsEngine/PhysAnim.cpp
+++ b/Engine/Source/Runtime/Engine/Private/PhysicsEngine/PhysAnim.cpp
@@ -356,10 +356,10 @@ void USkeletalMeshComponent::UpdateKinematicBonesToAnim(const TArray<FTransform>
 						ensure(!BoneTransform.ContainsNaN());
 
 						// If kinematic and not teleporting, set kinematic target
-						PxRigidDynamic* PRigidDynamic = BodyInst->GetPxRigidDynamic();
-						if (!IsRigidBodyNonKinematic(PRigidDynamic) && !bTeleport)
+						PxRigidDynamic* PRigidDynamic = BodyInst->GetPxRigidDynamic_AssumesLocked();
+						if (!IsRigidBodyNonKinematic_AssumesLocked(PRigidDynamic) && !bTeleport)
 						{
-							PhysScene->SetKinematicTarget(BodyInst, BoneTransform, true);
+							PhysScene->SetKinematicTarget_AssumesLocked(BodyInst, BoneTransform, true);
 						}
 						// Otherwise, set global pose
 						else
diff --git a/Engine/Source/Runtime/Engine/Private/PhysicsEngine/PhysScene.cpp b/Engine/Source/Runtime/Engine/Private/PhysicsEngine/PhysScene.cpp
--- a/Engine/Source/Runtime/Engine/Private/PhysicsEngine/PhysScene.cpp
+++ b/Engine/Source/Runtime/Engine/Private/PhysicsEngine/PhysScene.cpp
@@ -27,11 +27,11 @@
 static int32 PhysXSceneCount = 1;
 static const int PhysXSlowRebuildRate = 10;
 
-FORCEINLINE EPhysicsSceneType SceneType(const FBodyInstance* BodyInstance)
+FORCEINLINE EPhysicsSceneType SceneType_AssumesLocked(const FBodyInstance* BodyInstance)
 {
 #if WITH_PHYSX
 	//This is a helper function for dynamic actors - static actors are in both scenes
-	check(BodyInstance->GetPxRigidBody());
+	check(BodyInstance->GetPxRigidBody_AssumesLocked());
 	return UPhysicsSettings::Get()->bEnableAsyncScene && BodyInstance->UseAsyncScene() ? PST_Async : PST_Sync;
 #endif
 
@@ -164,23 +164,21 @@ bool UseSyncTime(uint32 SceneType)
 
 }
 
-bool FPhysScene::GetKinematicTarget(const FBodyInstance* BodyInstance, FTransform& OutTM) const
+bool FPhysScene::GetKinematicTarget_AssumesLocked(const FBodyInstance* BodyInstance, FTransform& OutTM) const
 {
 #if WITH_PHYSX
-	if (PxRigidDynamic * PRigidDynamic = BodyInstance->GetPxRigidDynamic())
+	if (PxRigidDynamic * PRigidDynamic = BodyInstance->GetPxRigidDynamic_AssumesLocked())
 	{
 #if WITH_SUBSTEPPING
-		uint32 BodySceneType = SceneType(BodyInstance);
+		uint32 BodySceneType = SceneType_AssumesLocked(BodyInstance);
 		if (IsSubstepping(BodySceneType))
 		{
 			FPhysSubstepTask * PhysSubStepper = PhysSubSteppers[BodySceneType];
-			return PhysSubStepper->GetKinematicTarget(BodyInstance, OutTM);
+			return PhysSubStepper->GetKinematicTarget_AssumesLocked(BodyInstance, OutTM);
 		}
 		else
 #endif
 		{
-
-			SCOPED_SCENE_READ_LOCK(PRigidDynamic->getScene());
 			PxTransform POutTM;
 			bool validTM = PRigidDynamic->getKinematicTarget(POutTM);
 			if (validTM)
@@ -195,42 +193,40 @@ bool FPhysScene::GetKinematicTarget(const FBodyInstance* BodyInstance, FTransfor
 	return false;
 }
 
-void FPhysScene::SetKinematicTarget(FBodyInstance* BodyInstance, const FTransform& TargetTransform, bool bAllowSubstepping)
+void FPhysScene::SetKinematicTarget_AssumesLocked(FBodyInstance* BodyInstance, const FTransform& TargetTransform, bool bAllowSubstepping)
 {
 	TargetTransform.DiagnosticCheckNaN_All();
 
 #if WITH_PHYSX
-	if (PxRigidDynamic * PRigidDynamic = BodyInstance->GetPxRigidDynamic())
+	if (PxRigidDynamic * PRigidDynamic = BodyInstance->GetPxRigidDynamic_AssumesLocked())
 	{
 #if WITH_SUBSTEPPING
-		uint32 BodySceneType = SceneType(BodyInstance);
+		uint32 BodySceneType = SceneType_AssumesLocked(BodyInstance);
 		if (bAllowSubstepping && IsSubstepping(BodySceneType))
 		{
 			FPhysSubstepTask * PhysSubStepper = PhysSubSteppers[BodySceneType];
-			PhysSubStepper->SetKinematicTarget(BodyInstance, TargetTransform);
+			PhysSubStepper->SetKinematicTarget_AssumesLocked(BodyInstance, TargetTransform);
 		}
 		else
 #endif
 		{
 			const PxTransform PNewPose = U2PTransform(TargetTransform);
 			check(PNewPose.isValid());
-
-			SCOPED_SCENE_WRITE_LOCK(PRigidDynamic->getScene());
 			PRigidDynamic->setKinematicTarget(PNewPose);
 		}
 	}
 #endif
 }
 
-void FPhysScene::AddCustomPhysics(FBodyInstance* BodyInstance, FCalculateCustomPhysics& CalculateCustomPhysics)
+void FPhysScene::AddCustomPhysics_AssumesLocked(FBodyInstance* BodyInstance, FCalculateCustomPhysics& CalculateCustomPhysics)
 {
 #if WITH_PHYSX
 #if WITH_SUBSTEPPING
-	uint32 BodySceneType = SceneType(BodyInstance);
+	uint32 BodySceneType = SceneType_AssumesLocked(BodyInstance);
 	if (IsSubstepping(BodySceneType))
 	{
-		FPhysSubstepTask * PhysSubStepper = PhysSubSteppers[SceneType(BodyInstance)];
-		PhysSubStepper->AddCustomPhysics(BodyInstance, CalculateCustomPhysics);
+		FPhysSubstepTask * PhysSubStepper = PhysSubSteppers[SceneType_AssumesLocked(BodyInstance)];
+		PhysSubStepper->AddCustomPhysics_AssumesLocked(BodyInstance, CalculateCustomPhysics);
 	}
 	else
 #endif
@@ -241,92 +237,88 @@ void FPhysScene::AddCustomPhysics(FBodyInstance* BodyInstance, FCalculateCustomP
 #endif
 }
 
-void FPhysScene::AddForce(FBodyInstance* BodyInstance, const FVector& Force, bool bAllowSubstepping, bool bAccelChange)
+void FPhysScene::AddForce_AssumesLocked(FBodyInstance* BodyInstance, const FVector& Force, bool bAllowSubstepping, bool bAccelChange)
 {
 #if WITH_PHYSX
 
-	if (PxRigidBody * PRigidBody = BodyInstance->GetPxRigidBody())
+	if (PxRigidBody * PRigidBody = BodyInstance->GetPxRigidBody_AssumesLocked())
 	{
 #if WITH_SUBSTEPPING
-		uint32 BodySceneType = SceneType(BodyInstance);
+		uint32 BodySceneType = SceneType_AssumesLocked(BodyInstance);
 		if (bAllowSubstepping && IsSubstepping(BodySceneType))
 		{
 			FPhysSubstepTask * PhysSubStepper = PhysSubSteppers[BodySceneType];
-			PhysSubStepper->AddForce(BodyInstance, Force, bAccelChange);
+			PhysSubStepper->AddForce_AssumesLocked(BodyInstance, Force, bAccelChange);
 		}
 		else
 #endif
 		{
-			SCOPED_SCENE_WRITE_LOCK(PRigidBody->getScene());
 			PRigidBody->addForce(U2PVector(Force), bAccelChange ? PxForceMode::eACCELERATION : PxForceMode::eFORCE, true);
 		}
 	}
 #endif
 }
 
-void FPhysScene::AddForceAtPosition(FBodyInstance* BodyInstance, const FVector& Force, const FVector& Position, bool bAllowSubstepping)
+void FPhysScene::AddForceAtPosition_AssumesLocked(FBodyInstance* BodyInstance, const FVector& Force, const FVector& Position, bool bAllowSubstepping)
 {
 #if WITH_PHYSX
 
-	if (PxRigidBody * PRigidBody = BodyInstance->GetPxRigidBody())
+	if (PxRigidBody * PRigidBody = BodyInstance->GetPxRigidBody_AssumesLocked())
 	{
 #if WITH_SUBSTEPPING
-		uint32 BodySceneType = SceneType(BodyInstance);
+		uint32 BodySceneType = SceneType_AssumesLocked(BodyInstance);
 		if (bAllowSubstepping && IsSubstepping(BodySceneType))
 		{
 			FPhysSubstepTask * PhysSubStepper = PhysSubSteppers[BodySceneType];
-			PhysSubStepper->AddForceAtPosition(BodyInstance, Force, Position);
+			PhysSubStepper->AddForceAtPosition_AssumesLocked(BodyInstance, Force, Position);
 		}
 		else
 #endif
 		{
-			SCOPED_SCENE_WRITE_LOCK(PRigidBody->getScene());
 			PxRigidBodyExt::addForceAtPos(*PRigidBody, U2PVector(Force), U2PVector(Position), PxForceMode::eFORCE, true);
 		}
 	}
 #endif
 }
 
-void FPhysScene::AddRadialForceToBody(FBodyInstance* BodyInstance, const FVector& Origin, const float Radius, const float Strength, const uint8 Falloff, bool bAccelChange, bool bAllowSubstepping)
+void FPhysScene::AddRadialForceToBody_AssumesLocked(FBodyInstance* BodyInstance, const FVector& Origin, const float Radius, const float Strength, const uint8 Falloff, bool bAccelChange, bool bAllowSubstepping)
 {
 #if WITH_PHYSX
 
-	if (PxRigidBody * PRigidBody = BodyInstance->GetPxRigidBody())
+	if (PxRigidBody * PRigidBody = BodyInstance->GetPxRigidBody_AssumesLocked())
 	{
 #if WITH_SUBSTEPPING
-		uint32 BodySceneType = SceneType(BodyInstance);
+		uint32 BodySceneType = SceneType_AssumesLocked(BodyInstance);
 		if (bAllowSubstepping && IsSubstepping(BodySceneType))
 		{
 			FPhysSubstepTask * PhysSubStepper = PhysSubSteppers[BodySceneType];
-			PhysSubStepper->AddRadialForceToBody(BodyInstance, Origin, Radius, Strength, Falloff, bAccelChange);
+			PhysSubStepper->AddRadialForceToBody_AssumesLocked(BodyInstance, Origin, Radius, Strength, Falloff, bAccelChange);
 		}
 		else
 #endif
 		{
-			SCOPED_SCENE_WRITE_LOCK(PRigidBody->getScene());
-			AddRadialForceToPxRigidBody(*PRigidBody, Origin, Radius, Strength, Falloff, bAccelChange);
+			AddRadialForceToPxRigidBody_AssumesLocked(*PRigidBody, Origin, Radius, Strength, Falloff, bAccelChange);
 		}
 	}
 #endif
 }
 
-void FPhysScene::AddTorque(FBodyInstance* BodyInstance, const FVector& Torque, bool bAllowSubstepping, bool bAccelChange)
+void FPhysScene::AddTorque_AssumesLocked(FBodyInstance* BodyInstance, const FVector& Torque, bool bAllowSubstepping, bool bAccelChange)
 {
 #if WITH_PHYSX
 
-	if (PxRigidBody * PRigidBody = BodyInstance->GetPxRigidBody())
+	if (PxRigidBody * PRigidBody = BodyInstance->GetPxRigidBody_AssumesLocked())
 	{
 #if WITH_SUBSTEPPING
-		uint32 BodySceneType = SceneType(BodyInstance);
+		uint32 BodySceneType = SceneType_AssumesLocked(BodyInstance);
 		if (bAllowSubstepping && IsSubstepping(BodySceneType))
 		{
 			FPhysSubstepTask * PhysSubStepper = PhysSubSteppers[BodySceneType];
-			PhysSubStepper->AddTorque(BodyInstance, Torque, bAccelChange);
+			PhysSubStepper->AddTorque_AssumesLocked(BodyInstance, Torque, bAccelChange);
 		}
 		else
 #endif
 		{
-			SCOPED_SCENE_WRITE_LOCK(PRigidBody->getScene());
 			PRigidBody->addTorque(U2PVector(Torque), bAccelChange ? PxForceMode::eACCELERATION : PxForceMode::eFORCE, true);
 		}
 	}
@@ -343,13 +335,13 @@ void FPhysScene::RemoveActiveBody(FBodyInstance* BodyInstance, uint32 SceneType)
 	}
 }
 #endif
-void FPhysScene::TermBody(FBodyInstance* BodyInstance)
+void FPhysScene::TermBody_AssumesLocked(FBodyInstance* BodyInstance)
 {
-	if (PxRigidBody* PRigidBody = BodyInstance->GetPxRigidBody())
+	if (PxRigidBody* PRigidBody = BodyInstance->GetPxRigidBody_AssumesLocked())
 	{
 #if WITH_SUBSTEPPING
-		FPhysSubstepTask* PhysSubStepper = PhysSubSteppers[SceneType(BodyInstance)];
-		PhysSubStepper->RemoveBodyInstance(BodyInstance);
+		FPhysSubstepTask* PhysSubStepper = PhysSubSteppers[SceneType_AssumesLocked(BodyInstance)];
+		PhysSubStepper->RemoveBodyInstance_AssumesLocked(BodyInstance);
 #endif
 	}
 
@@ -463,6 +455,7 @@ void GatherPhysXStats(PxScene* PScene, uint32 SceneType)
 	{
 		if (PScene)
 		{
+			SCOPED_SCENE_READ_LOCK(PScene);
 			PxSimulationStatistics SimStats;
 			PScene->getSimulationStatistics(SimStats);
 
@@ -752,7 +745,7 @@ void FPhysScene::UpdateActiveTransforms(uint32 SceneType)
 	}
 	PxScene* PScene = GetPhysXScene(SceneType);
 	check(PScene);
-	SCOPED_SCENE_WRITE_LOCK(PScene);
+	SCOPED_SCENE_READ_LOCK(PScene);
 
 	PxU32 NumTransforms = 0;
 	const PxActiveTransform* PActiveTransforms = PScene->getActiveTransforms(NumTransforms);
@@ -783,7 +776,7 @@ void FPhysScene::UpdateActiveTransforms(uint32 SceneType)
 
 DEFINE_STAT(STAT_SyncComponentsToBodies);
 
-void FPhysScene::SyncComponentsToBodies(uint32 SceneType)
+void FPhysScene::SyncComponentsToBodies_AssumesLocked(uint32 SceneType)
 {
 	SCOPE_CYCLE_COUNTER(STAT_TotalPhysicsTime);
 	SCOPE_CYCLE_COUNTER(STAT_SyncComponentsToBodies);
@@ -797,7 +790,7 @@ void FPhysScene::SyncComponentsToBodies(uint32 SceneType)
 		AActor* Owner = BodyInstance->OwnerComponent->GetOwner();
 
 		// See if the transform is actually different, and if so, move the component to match physics
-		const FTransform NewTransform = BodyInstance->GetUnrealWorldTransform();
+		const FTransform NewTransform = BodyInstance->GetUnrealWorldTransform_AssumesLocked();
 		if (!NewTransform.EqualsNoScale(BodyInstance->OwnerComponent->ComponentToWorld))
 		{
 			const FVector MoveBy = NewTransform.GetLocation() - BodyInstance->OwnerComponent->ComponentToWorld.GetLocation();
@@ -817,13 +810,12 @@ void FPhysScene::SyncComponentsToBodies(uint32 SceneType)
 #if WITH_APEX
 	if (ActiveDestructibleActors[SceneType].Num())
 	{
-		SCOPED_SCENE_READ_LOCK(GetPhysXScene(SceneType));
 		UDestructibleComponent::UpdateDestructibleChunkTM(ActiveDestructibleActors[SceneType]);
 	}
 #endif
 }
 
-void FPhysScene::DispatchPhysNotifications()
+void FPhysScene::DispatchPhysNotifications_AssumesLocked()
 {
 	SCOPE_CYCLE_COUNTER(STAT_PhysicsEventTime);
 
@@ -832,7 +824,7 @@ void FPhysScene::DispatchPhysNotifications()
 		// Let the game-specific PhysicsCollisionHandler process any physics collisions that took place
 		if (OwningWorld != NULL && OwningWorld->PhysicsCollisionHandler != NULL)
 		{
-			OwningWorld->PhysicsCollisionHandler->HandlePhysicsCollisions(PendingCollisionNotifies);
+			OwningWorld->PhysicsCollisionHandler->HandlePhysicsCollisions_AssumesLocked(PendingCollisionNotifies);
 		}
 
 		// Fire any collision notifies in the queue.
@@ -1029,17 +1021,27 @@ void FPhysScene::StartCloth()
 
 void FPhysScene::EndFrame(ULineBatchComponent* InLineBatcher)
 {
+	check(IsInGameThread());
+
 	PhysicsSceneCompletion = NULL;
 
+	/**
+	* At this point physics simulation has finished. We obtain both scene locks so that the various read/write operations needed can be done quickly.
+	* This means that anyone attempting to write on other threads will be blocked. This is OK because acessing any of these game objects from another thread is probably a bad idea!
+	*/
+
+	SCOPED_SCENE_WRITE_LOCK(GetPhysXScene(PST_Sync));
+	SCOPED_SCENE_WRITE_LOCK(bAsyncSceneEnabled ? GetPhysXScene(PST_Async) : nullptr);
+
 	if (bAsyncSceneEnabled)
 	{
-		SyncComponentsToBodies(PST_Async);
+		SyncComponentsToBodies_AssumesLocked(PST_Async);
 	}
 
-	SyncComponentsToBodies(PST_Sync);
+	SyncComponentsToBodies_AssumesLocked(PST_Sync);
 
 	// Perform any collision notification events
-	DispatchPhysNotifications();
+	DispatchPhysNotifications_AssumesLocked();
 
 #if !(UE_BUILD_SHIPPING || UE_BUILD_TEST)
 	// Handle debug rendering
@@ -1243,10 +1245,11 @@ void FPhysScene::InitPhysScene(uint32 SceneType)
 
 	// If we're frame lagging the async scene (truly running it async) then use the scene lock
 #if USE_SCENE_LOCK
-	if (SceneType == PST_Async)
+	if(UPhysicsSettings::Get()->bWarnMissingLocks)
 	{
 		PSceneDesc.flags |= PxSceneFlag::eREQUIRE_RW_LOCK;
 	}
+	
 #endif
 
 	// We want to use 'active transforms'
@@ -1414,10 +1417,10 @@ void FPhysScene::FDeferredSceneData::FlushDeferredActors()
 			++Idx;
 			Instance->CurrentSceneState = BodyInstanceSceneState::Added;
 
-			if (Instance->IsDynamic())
+			if(Instance->GetPxRigidDynamic_AssumesLocked())
 			{
 				// Extra setup necessary for dynamic objects.
-				Instance->InitDynamicProperties();
+				Instance->InitDynamicProperties_AssumesLocked();
 			}
 		}
 
diff --git a/Engine/Source/Runtime/Engine/Private/PhysicsEngine/PhysSubstepTasks.cpp b/Engine/Source/Runtime/Engine/Private/PhysicsEngine/PhysSubstepTasks.cpp
--- a/Engine/Source/Runtime/Engine/Private/PhysicsEngine/PhysSubstepTasks.cpp
+++ b/Engine/Source/Runtime/Engine/Private/PhysicsEngine/PhysSubstepTasks.cpp
@@ -18,7 +18,6 @@ FPhysSubstepTask::FPhysSubstepTask(PxApexScene * GivenScene) :
 	NumSubsteps(0),
 	SubTime(0.f),
 	DeltaSeconds(0.f),
-	External(0),
 	FullSimulationTask(0),
 	Alpha(0.f),
 	StepScale(0.f),
@@ -36,28 +35,19 @@ void FPhysSubstepTask::SwapBuffers()
 	External = !External;
 }
 
-void FPhysSubstepTask::RemoveBodyInstance(FBodyInstance* BodyInstance)
+void FPhysSubstepTask::RemoveBodyInstance_AssumesLocked(FBodyInstance* BodyInstance)
 {
 	
 	PhysTargetBuffers[External].Remove(BodyInstance);
-	
-	//This function gets called when we terminate a body instance. (FBodyInstance::TermBody)
-	//To do that FBodyInstance needs to be holding on to the scene write lock.
-	//This is the same lock that we use when iterating over the internal buffer - so it's safe to modify this buffer from this thread.
-	//If this assumption changes you need to take care of accessing the buffer from both threads
-	PhysTargetBuffers[!External].Remove(BodyInstance);
 }
 
-void FPhysSubstepTask::SetKinematicTarget(FBodyInstance* Body, const FTransform& TM)
+void FPhysSubstepTask::SetKinematicTarget_AssumesLocked(FBodyInstance* Body, const FTransform& TM)
 {
 #if WITH_PHYSX
-	check(Body);
 	TM.DiagnosticCheckNaN_All();
 
-	PxRigidDynamic * PRigidDynamic = Body->GetPxRigidDynamic();
-	SCOPED_SCENE_READ_LOCK(PRigidDynamic->getScene());
 	//We only interpolate kinematic actors
-	if (!IsRigidBodyNonKinematic(PRigidDynamic))
+	if (!Body->IsNonKinematic())
 	{
 		FKinematicTarget KinmaticTarget(Body, TM);
 		FPhysTarget & TargetState = PhysTargetBuffers[External].FindOrAdd(Body);
@@ -67,11 +57,9 @@ void FPhysSubstepTask::SetKinematicTarget(FBodyInstance* Body, const FTransform&
 #endif
 }
 
-bool FPhysSubstepTask::GetKinematicTarget(const FBodyInstance* Body, FTransform& OutTM) const
+bool FPhysSubstepTask::GetKinematicTarget_AssumesLocked(const FBodyInstance* Body, FTransform& OutTM) const
 {
 #if WITH_PHYSX
-	const PxRigidDynamic * PRigidDynamic = Body->GetPxRigidDynamic();
-	SCOPED_SCENE_READ_LOCK(PRigidDynamic->getScene());
 	if (const FPhysTarget* TargetState = PhysTargetBuffers[External].Find(Body))
 	{
 		if (TargetState->bKinematicTarget)
@@ -85,14 +73,11 @@ bool FPhysSubstepTask::GetKinematicTarget(const FBodyInstance* Body, FTransform&
 	return false;
 }
 
-void FPhysSubstepTask::AddCustomPhysics(FBodyInstance* Body, const FCalculateCustomPhysics& CalculateCustomPhysics)
+void FPhysSubstepTask::AddCustomPhysics_AssumesLocked(FBodyInstance* Body, const FCalculateCustomPhysics& CalculateCustomPhysics)
 {
 #if WITH_PHYSX
-	check(Body);
-	PxRigidBody* PRigidBody = Body->GetPxRigidBody();
-	SCOPED_SCENE_READ_LOCK(PRigidBody->getScene());
 	//Limit custom physics to non kinematic actors
-	if (IsRigidBodyNonKinematic(PRigidBody))
+	if (Body->IsNonKinematic())
 	{
 		FCustomTarget CustomTarget(CalculateCustomPhysics);
 
@@ -102,14 +87,11 @@ void FPhysSubstepTask::AddCustomPhysics(FBodyInstance* Body, const FCalculateCus
 #endif
 }
 
-void FPhysSubstepTask::AddForce(FBodyInstance* Body, const FVector& Force, bool bAccelChange)
+void FPhysSubstepTask::AddForce_AssumesLocked(FBodyInstance* Body, const FVector& Force, bool bAccelChange)
 {
 #if WITH_PHYSX
-	check(Body);
-	PxRigidBody* PRigidBody = Body->GetPxRigidBody();
-	SCOPED_SCENE_READ_LOCK(PRigidBody->getScene());
 	//We should only apply forces on non kinematic actors
-	if (IsRigidBodyNonKinematic(PRigidBody))
+	if (Body->IsNonKinematic())
 	{
 
 		FForceTarget ForceTarget;
@@ -123,15 +105,10 @@ void FPhysSubstepTask::AddForce(FBodyInstance* Body, const FVector& Force, bool
 #endif
 }
 
-void FPhysSubstepTask::AddForceAtPosition(FBodyInstance* Body, const FVector& Force, const FVector& Position)
+void FPhysSubstepTask::AddForceAtPosition_AssumesLocked(FBodyInstance* Body, const FVector& Force, const FVector& Position)
 {
 #if WITH_PHYSX
-	check(Body);
-
-	PxRigidBody* PRigidBody = Body->GetPxRigidBody();
-	SCOPED_SCENE_READ_LOCK(PRigidBody->getScene());
-	//We should only apply forces on non kinematic actors
-	if (IsRigidBodyNonKinematic(PRigidBody))
+	if (Body->IsNonKinematic())
 	{
 		FForceTarget ForceTarget;
 		ForceTarget.bPosition = true;
@@ -143,15 +120,11 @@ void FPhysSubstepTask::AddForceAtPosition(FBodyInstance* Body, const FVector& Fo
 	}
 #endif
 }
-void FPhysSubstepTask::AddTorque(FBodyInstance* Body, const FVector& Torque, bool bAccelChange)
+void FPhysSubstepTask::AddTorque_AssumesLocked(FBodyInstance* Body, const FVector& Torque, bool bAccelChange)
 {
 #if WITH_PHYSX
-	check(Body);
-
-	PxRigidBody* PRigidBody = Body->GetPxRigidBody();
-	SCOPED_SCENE_READ_LOCK(PRigidBody->getScene());
 	//We should only apply torque on non kinematic actors
-	if (IsRigidBodyNonKinematic(PRigidBody))
+	if (Body->IsNonKinematic())
 	{
 		FTorqueTarget TorqueTarget;
 		TorqueTarget.Torque = Torque;
@@ -163,15 +136,11 @@ void FPhysSubstepTask::AddTorque(FBodyInstance* Body, const FVector& Torque, boo
 #endif
 }
 
-void FPhysSubstepTask::AddRadialForceToBody(FBodyInstance* Body, const FVector& Origin, const float Radius, const float Strength, const uint8 Falloff, const bool bAccelChange)
+void FPhysSubstepTask::AddRadialForceToBody_AssumesLocked(FBodyInstance* Body, const FVector& Origin, const float Radius, const float Strength, const uint8 Falloff, const bool bAccelChange)
 {
 #if WITH_PHYSX
-	check(Body);
-
-	PxRigidBody* PRigidBody = Body->GetPxRigidBody();
-	SCOPED_SCENE_READ_LOCK(PRigidBody->getScene());
 	//We should only apply torque on non kinematic actors
-	if (IsRigidBodyNonKinematic(PRigidBody))
+	if (Body->IsNonKinematic())
 	{
 		FRadialForceTarget RadialForceTarget;
 		RadialForceTarget.Origin = Origin;
@@ -200,11 +169,11 @@ void FPhysSubstepTask::ApplyCustomPhysics(const FPhysTarget& PhysTarget, FBodyIn
 }
 
 /** Applies forces - Assumes caller has obtained writer lock */
-void FPhysSubstepTask::ApplyForces(const FPhysTarget& PhysTarget, FBodyInstance* BodyInstance)
+void FPhysSubstepTask::ApplyForces_AssumesLocked(const FPhysTarget& PhysTarget, FBodyInstance* BodyInstance)
 {
 #if WITH_PHYSX
 	/** Apply Forces */
-	PxRigidBody* PRigidBody = BodyInstance->GetPxRigidBody();
+	PxRigidBody* PRigidBody = BodyInstance->GetPxRigidBody_AssumesLocked();
 
 	for (int32 i = 0; i < PhysTarget.Forces.Num(); ++i)
 	{
@@ -223,11 +192,11 @@ void FPhysSubstepTask::ApplyForces(const FPhysTarget& PhysTarget, FBodyInstance*
 }
 
 /** Applies torques - Assumes caller has obtained writer lock */
-void FPhysSubstepTask::ApplyTorques(const FPhysTarget& PhysTarget, FBodyInstance* BodyInstance)
+void FPhysSubstepTask::ApplyTorques_AssumesLocked(const FPhysTarget& PhysTarget, FBodyInstance* BodyInstance)
 {
 #if WITH_PHYSX
 	/** Apply Torques */
-	PxRigidBody* PRigidBody = BodyInstance->GetPxRigidBody();
+	PxRigidBody* PRigidBody = BodyInstance->GetPxRigidBody_AssumesLocked();
 
 	for (int32 i = 0; i < PhysTarget.Torques.Num(); ++i)
 	{
@@ -238,33 +207,33 @@ void FPhysSubstepTask::ApplyTorques(const FPhysTarget& PhysTarget, FBodyInstance
 }
 
 /** Applies radial forces - Assumes caller has obtained writer lock */
-void FPhysSubstepTask::ApplyRadialForces(const FPhysTarget& PhysTarget, FBodyInstance* BodyInstance)
+void FPhysSubstepTask::ApplyRadialForces_AssumesLocked(const FPhysTarget& PhysTarget, FBodyInstance* BodyInstance)
 {
 #if WITH_PHYSX
 	/** Apply Torques */
-	PxRigidBody* PRigidBody = BodyInstance->GetPxRigidBody();
+	PxRigidBody* PRigidBody = BodyInstance->GetPxRigidBody_AssumesLocked();
 
 	for (int32 i = 0; i < PhysTarget.RadialForces.Num(); ++i)
 	{
 		const FRadialForceTarget& RadialForceTArget= PhysTarget.RadialForces[i];
-		AddRadialForceToPxRigidBody(*PRigidBody, RadialForceTArget.Origin, RadialForceTArget.Radius, RadialForceTArget.Strength, RadialForceTArget.Falloff, RadialForceTArget.bAccelChange);
+		AddRadialForceToPxRigidBody_AssumesLocked(*PRigidBody, RadialForceTArget.Origin, RadialForceTArget.Radius, RadialForceTArget.Strength, RadialForceTArget.Falloff, RadialForceTArget.bAccelChange);
 	}
 #endif
 }
 
 
 /** Interpolates kinematic actor transform - Assumes caller has obtained writer lock */
-void FPhysSubstepTask::InterpolateKinematicActor(const FPhysTarget& PhysTarget, FBodyInstance* BodyInstance, float InAlpha)
+void FPhysSubstepTask::InterpolateKinematicActor_AssumesLocked(const FPhysTarget& PhysTarget, FBodyInstance* BodyInstance, float InAlpha)
 {
 #if WITH_PHYSX
-	PxRigidDynamic * PRigidDynamic = BodyInstance->GetPxRigidDynamic();
+	PxRigidDynamic * PRigidDynamic = BodyInstance->GetPxRigidDynamic_AssumesLocked();
 	InAlpha = FMath::Clamp(InAlpha, 0.f, 1.f);
 
 	/** Interpolate kinematic actors */
 	if (PhysTarget.bKinematicTarget)
 	{
 		//It's possible that the actor is no longer kinematic and is now simulating. In that case do nothing
-		if (!IsRigidBodyNonKinematic(PRigidDynamic))
+		if (!BodyInstance->IsNonKinematic())
 		{
 			const FKinematicTarget& KinematicTarget = PhysTarget.KinematicTarget;
 			const FTransform& TargetTM = KinematicTarget.TargetTM;
@@ -291,16 +260,16 @@ void FPhysSubstepTask::SubstepInterpolation(float InAlpha, float DeltaTime)
 	PxScene * PScene = PAScene;
 #endif
 
-	PhysTargetMap & Targets = PhysTargetBuffers[!External];
+	PhysTargetMap& Targets = PhysTargetBuffers[!External];
 
 	/** Note: We lock the entire scene before iterating. The assumption is that removing an FBodyInstance from the map will also be wrapped by this lock */
-	SCENE_LOCK_WRITE(PScene);
+	SCOPED_SCENE_WRITE_LOCK(PScene);
 
 	for (PhysTargetMap::TIterator Itr = Targets.CreateIterator(); Itr; ++Itr)
 	{
 		FPhysTarget & PhysTarget = Itr.Value();
 		FBodyInstance * BodyInstance = Itr.Key();
-		PxRigidBody* PRigidBody = BodyInstance->GetPxRigidBody();
+		PxRigidBody* PRigidBody = BodyInstance->GetPxRigidBody_AssumesLocked();
 
 		if (PRigidBody == NULL)
 		{
@@ -311,10 +280,10 @@ void FPhysSubstepTask::SubstepInterpolation(float InAlpha, float DeltaTime)
 		check(PRigidBody->getScene() == PScene);
 
 		ApplyCustomPhysics(PhysTarget, BodyInstance, DeltaTime);
-		ApplyForces(PhysTarget, BodyInstance);
-		ApplyTorques(PhysTarget, BodyInstance);
-		ApplyRadialForces(PhysTarget, BodyInstance);
-		InterpolateKinematicActor(PhysTarget, BodyInstance, InAlpha);
+		ApplyForces_AssumesLocked(PhysTarget, BodyInstance);
+		ApplyTorques_AssumesLocked(PhysTarget, BodyInstance);
+		ApplyRadialForces_AssumesLocked(PhysTarget, BodyInstance);
+		InterpolateKinematicActor_AssumesLocked(PhysTarget, BodyInstance, InAlpha);
 	}
 
 	/** Final substep */
@@ -322,8 +291,6 @@ void FPhysSubstepTask::SubstepInterpolation(float InAlpha, float DeltaTime)
 	{
 		Targets.Empty(Targets.Num());
 	}
-
-	SCENE_UNLOCK_WRITE(PScene);
 #endif
 }
 
diff --git a/Engine/Source/Runtime/Engine/Private/PhysicsEngine/PhysSubstepTasks.h b/Engine/Source/Runtime/Engine/Private/PhysicsEngine/PhysSubstepTasks.h
--- a/Engine/Source/Runtime/Engine/Private/PhysicsEngine/PhysSubstepTasks.h
+++ b/Engine/Source/Runtime/Engine/Private/PhysicsEngine/PhysSubstepTasks.h
@@ -41,7 +41,7 @@ struct FKinematicTarget
 	FKinematicTarget(FBodyInstance* Body, const FTransform& TM) :
 		BodyInstance(Body),
 		TargetTM(TM),
-		OriginalTM(Body->GetUnrealWorldTransform())
+		OriginalTM(Body->GetUnrealWorldTransform_AssumesLocked())
 	{
 	}
 
@@ -125,16 +125,16 @@ public:
 	FPhysSubstepTask(PxApexScene * GivenScene);
 #endif
 
-	void SetKinematicTarget(FBodyInstance* Body, const FTransform& TM);
-	bool GetKinematicTarget(const FBodyInstance* Body, FTransform& OutTM) const;
-	void AddCustomPhysics(FBodyInstance* Body, const FCalculateCustomPhysics& CalculateCustomPhysics);
-	void AddForce(FBodyInstance* Body, const FVector& Force, bool bAccelChange);
-	void AddForceAtPosition(FBodyInstance* Body, const FVector& Force, const FVector& Position);
-	void AddTorque(FBodyInstance* Body, const FVector& Torque, bool bAccelChange);
-	void AddRadialForceToBody(FBodyInstance* Body, const FVector& Origin, const float Radius, const float Strength, const uint8 Falloff, const bool bAccelChange);
+	void SetKinematicTarget_AssumesLocked(FBodyInstance* Body, const FTransform& TM);
+	bool GetKinematicTarget_AssumesLocked(const FBodyInstance* Body, FTransform& OutTM) const;
+	void AddCustomPhysics_AssumesLocked(FBodyInstance* Body, const FCalculateCustomPhysics& CalculateCustomPhysics);
+	void AddForce_AssumesLocked(FBodyInstance* Body, const FVector& Force, bool bAccelChange);
+	void AddForceAtPosition_AssumesLocked(FBodyInstance* Body, const FVector& Force, const FVector& Position);
+	void AddTorque_AssumesLocked(FBodyInstance* Body, const FVector& Torque, bool bAccelChange);
+	void AddRadialForceToBody_AssumesLocked(FBodyInstance* Body, const FVector& Origin, const float Radius, const float Strength, const uint8 Falloff, const bool bAccelChange);
 
 	/** Removes a BodyInstance from doing substep work - should only be called when the FBodyInstance is getting destroyed */
-	void RemoveBodyInstance(FBodyInstance* Body);
+	void RemoveBodyInstance_AssumesLocked(FBodyInstance* Body);
 
 	
 	void SwapBuffers();
@@ -152,17 +152,17 @@ private:
 	/** Applies interpolation and forces on all needed actors*/
 	void SubstepInterpolation(float Scale, float DeltaTime);
 	void ApplyCustomPhysics(const FPhysTarget& PhysTarget, FBodyInstance* BodyInstance, float DeltaTime);
-	void ApplyForces(const FPhysTarget& PhysTarget, FBodyInstance* BodyInstance);
-	void ApplyTorques(const FPhysTarget& PhysTarget, FBodyInstance* BodyInstance);
-	void ApplyRadialForces(const FPhysTarget& PhysTarget, FBodyInstance* BodyInstance);
-	void InterpolateKinematicActor(const FPhysTarget& PhysTarget, FBodyInstance* BodyInstance, float Alpha);
+	void ApplyForces_AssumesLocked(const FPhysTarget& PhysTarget, FBodyInstance* BodyInstance);
+	void ApplyTorques_AssumesLocked(const FPhysTarget& PhysTarget, FBodyInstance* BodyInstance);
+	void ApplyRadialForces_AssumesLocked(const FPhysTarget& PhysTarget, FBodyInstance* BodyInstance);
+	void InterpolateKinematicActor_AssumesLocked(const FPhysTarget& PhysTarget, FBodyInstance* BodyInstance, float Alpha);
 
 	typedef TMap<FBodyInstance*, FPhysTarget> PhysTargetMap;
 	PhysTargetMap PhysTargetBuffers[2];	//need to double buffer between physics thread and game thread
 	uint32 NumSubsteps;
 	float SubTime;
 	float DeltaSeconds;
-	bool External;
+	FThreadSafeBool External;
 	class  PhysXCompletionTask * FullSimulationTask;
 	float Alpha;
 	float StepScale;
diff --git a/Engine/Source/Runtime/Engine/Private/PhysicsEngine/PhysXSupport.cpp b/Engine/Source/Runtime/Engine/Private/PhysicsEngine/PhysXSupport.cpp
--- a/Engine/Source/Runtime/Engine/Private/PhysicsEngine/PhysXSupport.cpp
+++ b/Engine/Source/Runtime/Engine/Private/PhysicsEngine/PhysXSupport.cpp
@@ -213,7 +213,7 @@ PxScene* GetPhysXSceneFromIndex(int32 InSceneIndex)
 #endif	// #if WITH_APEX
 
 
-void AddRadialImpulseToPxRigidBody(PxRigidBody& PRigidBody, const FVector& Origin, float Radius, float Strength, uint8 Falloff, bool bVelChange)
+void AddRadialImpulseToPxRigidBody_AssumesLocked(PxRigidBody& PRigidBody, const FVector& Origin, float Radius, float Strength, uint8 Falloff, bool bVelChange)
 {
 #if WITH_PHYSX
 	if (!(PRigidBody.getRigidBodyFlags() & PxRigidBodyFlag::eKINEMATIC))
@@ -249,7 +249,7 @@ void AddRadialImpulseToPxRigidBody(PxRigidBody& PRigidBody, const FVector& Origi
 #endif // WITH_PHYSX
 }
 
-void AddRadialForceToPxRigidBody(PxRigidBody& PRigidBody, const FVector& Origin, float Radius, float Strength, uint8 Falloff, bool bAccelChange)
+void AddRadialForceToPxRigidBody_AssumesLocked(PxRigidBody& PRigidBody, const FVector& Origin, float Radius, float Strength, uint8 Falloff, bool bAccelChange)
 {
 #if WITH_PHYSX
 	if (!(PRigidBody.getRigidBodyFlags() & PxRigidBodyFlag::eKINEMATIC))
@@ -284,9 +284,9 @@ void AddRadialForceToPxRigidBody(PxRigidBody& PRigidBody, const FVector& Origin,
 #endif // WITH_PHYSX
 }
 
-bool IsRigidBodyNonKinematic(PxRigidBody* PRigidBody)
+bool IsRigidBodyNonKinematic_AssumesLocked(const PxRigidBody* PRigidBody)
 {
-	if (PRigidBody != NULL)
+	if (PRigidBody)
 	{
 		return !(PRigidBody->getRigidBodyFlags() & PxRigidBodyFlag::eKINEMATIC);
 	}
diff --git a/Engine/Source/Runtime/Engine/Private/PhysicsEngine/PhysXSupport.h b/Engine/Source/Runtime/Engine/Private/PhysicsEngine/PhysXSupport.h
--- a/Engine/Source/Runtime/Engine/Private/PhysicsEngine/PhysXSupport.h
+++ b/Engine/Source/Runtime/Engine/Private/PhysicsEngine/PhysXSupport.h
@@ -33,7 +33,7 @@ public:
 	FPhysXSceneReadLock(PxScene* PInScene)
 		: PScene(PInScene)
 	{
-		SCOPE_CYCLE_COUNTER(STAT_PhysSceneLock);
+		SCOPE_CYCLE_COUNTER(STAT_PhysSceneReadLock);
 		if(PScene)
 		{
 			PScene->lockRead();
@@ -42,7 +42,6 @@ public:
 
 	~FPhysXSceneReadLock()
 	{
-		SCOPE_CYCLE_COUNTER(STAT_PhysSceneLock);
 		if(PScene)
 		{
 			PScene->unlockRead();
@@ -60,7 +59,7 @@ public:
 	FPhysXSceneWriteLock(PxScene* PInScene)
 		: PScene(PInScene)
 	{
-		SCOPE_CYCLE_COUNTER(STAT_PhysSceneLock);
+		SCOPE_CYCLE_COUNTER(STAT_PhysSceneWriteLock);
 		if(PScene)
 		{
 			PScene->lockWrite();
@@ -69,7 +68,6 @@ public:
 
 	~FPhysXSceneWriteLock()
 	{
-		SCOPE_CYCLE_COUNTER(STAT_PhysSceneLock);
 		if(PScene)
 		{
 			PScene->unlockWrite();
@@ -85,10 +83,10 @@ private:
 #define SCOPED_SCENE_READ_LOCK( _scene ) FPhysXSceneReadLock SCOPED_SCENE_TOKENPASTE(_rlock,__LINE__)(_scene)
 #define SCOPED_SCENE_WRITE_LOCK( _scene ) FPhysXSceneWriteLock SCOPED_SCENE_TOKENPASTE(_wlock,__LINE__)(_scene)
 
-#define SCENE_LOCK_READ( _scene )		{ SCOPE_CYCLE_COUNTER(STAT_PhysSceneLock); if((_scene) != NULL) { (_scene)->lockRead(); } }
-#define SCENE_UNLOCK_READ( _scene )		{ SCOPE_CYCLE_COUNTER(STAT_PhysSceneLock); if((_scene) != NULL) { (_scene)->unlockRead(); } }
-#define SCENE_LOCK_WRITE( _scene )		{ SCOPE_CYCLE_COUNTER(STAT_PhysSceneLock); if((_scene) != NULL) { (_scene)->lockWrite(); } }
-#define SCENE_UNLOCK_WRITE( _scene )	{ SCOPE_CYCLE_COUNTER(STAT_PhysSceneLock); if((_scene) != NULL) { (_scene)->unlockWrite(); } }
+#define SCENE_LOCK_READ( _scene )		{ SCOPE_CYCLE_COUNTER(STAT_PhysSceneReadLock); if((_scene) != NULL) { (_scene)->lockRead(); } }
+#define SCENE_UNLOCK_READ( _scene )		{ if((_scene) != NULL) { (_scene)->unlockRead(); } }
+#define SCENE_LOCK_WRITE( _scene )		{ SCOPE_CYCLE_COUNTER(STAT_PhysSceneWriteLock); if((_scene) != NULL) { (_scene)->lockWrite(); } }
+#define SCENE_UNLOCK_WRITE( _scene )	{ if((_scene) != NULL) { (_scene)->unlockWrite(); } }
 #else
 #define SCOPED_SCENE_READ_LOCK_INDEXED( _scene, _index )
 #define SCOPED_SCENE_READ_LOCK( _scene )
@@ -100,6 +98,172 @@ private:
 #define SCENE_UNLOCK_WRITE( _scene )
 #endif
 
+/** Get a pointer to the PxScene from an SceneIndex (will be NULL if scene already shut down) */
+PxScene* GetPhysXSceneFromIndex(int32 InSceneIndex);
+
+#if WITH_APEX
+/** Get a pointer to the NxApexScene from an SceneIndex (will be NULL if scene already shut down) */
+NxApexScene* GetApexSceneFromIndex(int32 InSceneIndex);
+#endif
+
+/** Obtains the appropriate PhysX scene lock for READING and executes the passed in lambda.
+ *  If SceneType < 0, the Sync actor is used, otherwise the async.
+ *  Note: The lambda is only executed if the physx actor requested is non-null.
+ *  returns true if the requested actor is non-null
+ */
+template <typename LambdaType, bool NeedsLock = true>
+bool ExecuteOnPxRigidActorReadOnly(const FBodyInstance* BI, const LambdaType& Func, int32 SceneType = -1)
+{
+	if (const PxRigidActor* PRigidActor = BI->GetPxRigidActor_AssumesLocked(SceneType))
+	{
+		const int32 SceneIndex = (PRigidActor == BI->RigidActorSync ? BI->SceneIndexSync : BI->SceneIndexAsync);
+		PxScene* PScene = GetPhysXSceneFromIndex(SceneIndex);
+		if (NeedsLock)
+		{
+			SCENE_LOCK_READ(PScene);
+		}
+
+		Func(PRigidActor);
+
+		if (NeedsLock)
+		{
+			SCENE_UNLOCK_READ(PScene);
+		}
+
+		return true;
+	}
+
+	return false;
+}
+
+/** Obtains the appropriate PhysX scene lock for READING and executes the passed in lambda.
+ *  Note: The lambda is only executed if the physx actor is a non-null RigidBody
+ *  returns true if found a non-null RigidBody
+ */
+template <typename LambdaType, bool NeedsLock = true>
+bool ExecuteOnPxRigidBodyReadOnly(const FBodyInstance* BI, const LambdaType& Func)
+{
+	bool bSuccess = false;
+	if (const physx::PxRigidActor* RigidActor = BI->GetPxRigidActor_AssumesLocked())
+	{
+		const int32 SceneIndex = (RigidActor == BI->RigidActorSync ? BI->SceneIndexSync : BI->SceneIndexAsync);
+		PxScene* PScene = GetPhysXSceneFromIndex(SceneIndex);
+		if (NeedsLock)
+		{
+			SCENE_LOCK_READ(PScene);
+		}
+
+		if (const physx::PxRigidBody* PRigidBody = RigidActor->isRigidBody())
+		{
+			Func(PRigidBody);
+			bSuccess = true;
+		}
+
+		if (NeedsLock)
+		{
+			SCENE_UNLOCK_READ(PScene);
+		}
+	}
+
+	return bSuccess;
+}
+
+/** Obtains the appropriate PhysX scene lock for WRITING and executes the passed in lambda.
+ *  Note: The lambda is only executed if the physx actor is a non-null RigidBody
+ *  returns true if found a non-null RigidBody.
+ */
+template <typename LambdaType, bool NeedsLock = true>
+bool ExecuteOnPxRigidBodyReadWrite(const FBodyInstance* BI, const LambdaType& Func)
+{
+	bool bSuccess = false;
+	if (physx::PxRigidActor* RigidActor = BI->GetPxRigidActor_AssumesLocked())
+	{
+		const int32 SceneIndex = (RigidActor == BI->RigidActorSync ? BI->SceneIndexSync : BI->SceneIndexAsync);
+		PxScene* PScene = GetPhysXSceneFromIndex(SceneIndex);
+		if(NeedsLock)
+		{
+			SCENE_LOCK_WRITE(PScene);
+		}
+			
+		if (physx::PxRigidBody* PRigidBody = RigidActor->isRigidBody())
+		{
+			Func(PRigidBody);
+			bSuccess = true;
+		}
+
+		if(NeedsLock)
+		{
+			SCENE_UNLOCK_WRITE(PScene);
+		}
+	}
+
+	return bSuccess;
+}
+
+/** Obtains the appropriate PhysX scene lock for READING and executes the passed in lambda.
+ *  Note: The lambda is only executed if the physx actor is a non-null RigidDynamic
+ *  returns true if found a non-null RigidDynamic.
+ */
+template <typename LambdaType, bool NeedsLock = true>
+bool ExecuteOnPxRigidDynamicReadOnly(const FBodyInstance* BI, const LambdaType& Func)
+{
+	bool bSuccess = false;
+	if (physx::PxRigidActor* RigidActor = BI->GetPxRigidActor_AssumesLocked())
+	{
+		const int32 SceneIndex = (RigidActor == BI->RigidActorSync ? BI->SceneIndexSync : BI->SceneIndexAsync);
+		PxScene* PScene = GetPhysXSceneFromIndex(SceneIndex);
+		if (NeedsLock)
+		{
+			SCENE_LOCK_READ(PScene);
+		}
+
+		if (physx::PxRigidDynamic* PRigidDynamic = RigidActor->isRigidDynamic())
+		{
+			Func(PRigidDynamic);
+			bSuccess = true;
+		}
+
+		if (NeedsLock)
+		{
+			SCENE_UNLOCK_READ(PScene);
+		}
+	}
+
+	return bSuccess;
+}
+
+/** Obtains the appropriate PhysX scene lock for WRITING and executes the passed in lambda.
+ *  Note: The lambda is only executed if the physx actor is a non-null RigidDynamic
+ *  returns true if found a non-null RigidDynamic.
+ */
+template <typename LambdaType, bool NeedsLock = true>
+bool ExecuteOnPxRigidDynamicReadWrite(const FBodyInstance* BI, const LambdaType& Func)
+{
+	bool bSuccess = false;
+	if (physx::PxRigidActor* RigidActor = BI->GetPxRigidActor_AssumesLocked())
+	{
+		const int32 SceneIndex = (RigidActor == BI->RigidActorSync ? BI->SceneIndexSync : BI->SceneIndexAsync);
+		PxScene* PScene = GetPhysXSceneFromIndex(SceneIndex);
+		if (NeedsLock)
+		{
+			SCENE_LOCK_WRITE(PScene);
+		}
+
+		if (physx::PxRigidDynamic* PRigidDynamic = RigidActor->isRigidDynamic())
+		{
+			Func(PRigidDynamic);
+			bSuccess = true;
+		}
+
+		if (NeedsLock)
+		{
+			SCENE_UNLOCK_WRITE(PScene);
+		}
+	}
+
+	return bSuccess;
+}
+
 //////// BASIC TYPE CONVERSION
 
 /** Convert Unreal FMatrix to PhysX PxTransform */
@@ -153,26 +317,38 @@ const bool bGlobalCCD = true;
 
 /////// UTILS
 
-/** Get a pointer to the PxScene from an SceneIndex (will be NULL if scene already shut down) */
-PxScene* GetPhysXSceneFromIndex(int32 InSceneIndex);
-
-#if WITH_APEX
-/** Get a pointer to the NxApexScene from an SceneIndex (will be NULL if scene already shut down) */
-NxApexScene* GetApexSceneFromIndex(int32 InSceneIndex);
-#endif
-
 
 /** Perform any deferred cleanup of resources (GPhysXPendingKillConvex etc) */
 void DeferredPhysResourceCleanup();
 
 /** Calculates correct impulse at the body's center of mass and adds the impulse to the body. */
-ENGINE_API void AddRadialImpulseToPxRigidBody(PxRigidBody& PRigidBody, const FVector& Origin, float Radius, float Strength, uint8 Falloff, bool bVelChange);
+ENGINE_API void AddRadialImpulseToPxRigidBody_AssumesLocked(PxRigidBody& PRigidBody, const FVector& Origin, float Radius, float Strength, uint8 Falloff, bool bVelChange);
+
+
+/** Calculates correct impulse at the body's center of mass and adds the impulse to the body. */
+DEPRECATED(4.8, "Please call AddRadialImpulseToPxRigidBody_AssumesLocked and make sure you obtain the appropriate PhysX scene locks")
+inline void AddRadialImpulseToPxRigidBody(PxRigidBody& PRigidBody, const FVector& Origin, float Radius, float Strength, uint8 Falloff, bool bVelChange)
+{
+	AddRadialImpulseToPxRigidBody_AssumesLocked(PRigidBody, Origin, Radius, Strength, Falloff, bVelChange);
+}
+
+ENGINE_API void AddRadialForceToPxRigidBody_AssumesLocked(PxRigidBody& PRigidBody, const FVector& Origin, float Radius, float Strength, uint8 Falloff, bool bAccelChange);
 
 /** Calculates correct force at the body's center of mass and adds force to the body. */
-ENGINE_API void AddRadialForceToPxRigidBody(PxRigidBody& PRigidBody, const FVector& Origin, float Radius, float Strength, uint8 Falloff, bool bAccelChange);
+DEPRECATED(4.8, "Please call AddRadialImpulseToPxRigidBody_AssumesLocked and make sure you obtain the appropriate PhysX scene locks")
+inline void AddRadialForceToPxRigidBody(PxRigidBody& PRigidBody, const FVector& Origin, float Radius, float Strength, uint8 Falloff, bool bAccelChange)
+{
+	AddRadialForceToPxRigidBody_AssumesLocked(PRigidBody, Origin, Radius, Strength, Falloff, bAccelChange);
+}
+
+bool IsRigidBodyNonKinematic_AssumesLocked(const PxRigidBody* PRigidBody);
 
 /** Util to see if a PxRigidActor is non-kinematic */
-bool IsRigidBodyNonKinematic(PxRigidBody* PRigidBody);
+DEPRECATED(4.8, "Please call IsRigidBodyNonKinematic_AssumesLocked and make sure you obtain the appropriate PhysX scene locks")
+inline bool IsRigidBodyNonKinematic(PxRigidBody* PRigidBody)
+{
+	return IsRigidBodyNonKinematic_AssumesLocked(PRigidBody);
+}
 
 
 /////// GLOBAL POINTERS
@@ -669,4 +845,4 @@ struct FShapeData
 	PxShapeFlags ComplexShapeFlags;
 	PxRigidBodyFlags SyncBodyFlags;
 	PxRigidBodyFlags AsyncBodyFlags;
-};
+};
\ No newline at end of file
diff --git a/Engine/Source/Runtime/Engine/Private/PhysicsEngine/PhysicsCollisionHandler.cpp b/Engine/Source/Runtime/Engine/Private/PhysicsEngine/PhysicsCollisionHandler.cpp
--- a/Engine/Source/Runtime/Engine/Private/PhysicsEngine/PhysicsCollisionHandler.cpp
+++ b/Engine/Source/Runtime/Engine/Private/PhysicsEngine/PhysicsCollisionHandler.cpp
@@ -11,7 +11,7 @@ UPhysicsCollisionHandler::UPhysicsCollisionHandler(const FObjectInitializer& Obj
 	ImpactReFireDelay = 0.1f;
 }
 
-void UPhysicsCollisionHandler::DefaultHandleCollision(const FRigidBodyCollisionInfo& MyInfo, const FRigidBodyCollisionInfo& OtherInfo, const FCollisionImpactData& RigidCollisionData)
+void UPhysicsCollisionHandler::DefaultHandleCollision_AssumesLocked(const FRigidBodyCollisionInfo& MyInfo, const FRigidBodyCollisionInfo& OtherInfo, const FCollisionImpactData& RigidCollisionData)
 {
 	FRigidBodyContactInfo ContactInfo = RigidCollisionData.ContactInfos[0];
 
@@ -21,11 +21,11 @@ void UPhysicsCollisionHandler::DefaultHandleCollision(const FRigidBodyCollisionI
 	if(BodyInst0 && BodyInst1)
 	{
 		// Find relative velocity.
-		FVector Velocity0 = BodyInst0->GetUnrealWorldVelocityAtPoint(ContactInfo.ContactPosition);
-		FVector AngularVel0 = BodyInst0->GetUnrealWorldAngularVelocity();
+		FVector Velocity0 = BodyInst0->GetUnrealWorldVelocityAtPoint_AssumesLocked(ContactInfo.ContactPosition);
+		FVector AngularVel0 = BodyInst0->GetUnrealWorldAngularVelocity_AssumesLocked();
 
-		FVector Velocity1 = BodyInst1->GetUnrealWorldVelocityAtPoint(ContactInfo.ContactPosition);
-		FVector AngularVel1 = BodyInst1->GetUnrealWorldAngularVelocity();
+		FVector Velocity1 = BodyInst1->GetUnrealWorldVelocityAtPoint_AssumesLocked(ContactInfo.ContactPosition);
+		FVector AngularVel1 = BodyInst1->GetUnrealWorldAngularVelocity_AssumesLocked();
 
 		const FVector RelVel = Velocity1 - Velocity0;
 
@@ -52,7 +52,7 @@ void UPhysicsCollisionHandler::DefaultHandleCollision(const FRigidBodyCollisionI
 	}
 }
 
-void UPhysicsCollisionHandler::HandlePhysicsCollisions(const TArray<FCollisionNotifyInfo>& PendingCollisionNotifies)
+void UPhysicsCollisionHandler::HandlePhysicsCollisions_AssumesLocked(const TArray<FCollisionNotifyInfo>& PendingCollisionNotifies)
 {
 	// Fire any collision notifies in the queue.
 	for(int32 i=0; i<PendingCollisionNotifies.Num(); i++)
@@ -69,7 +69,7 @@ void UPhysicsCollisionHandler::HandlePhysicsCollisions(const TArray<FCollisionNo
 		if( NotifyInfo.IsValidForNotify() &&
 			NotifyInfo.RigidCollisionData.ContactInfos.Num() > 0 )
 		{
-			DefaultHandleCollision(NotifyInfo.Info0, NotifyInfo.Info1, NotifyInfo.RigidCollisionData);
+			DefaultHandleCollision_AssumesLocked(NotifyInfo.Info0, NotifyInfo.Info1, NotifyInfo.RigidCollisionData);
 		}
 	}
 }
\ No newline at end of file
diff --git a/Engine/Source/Runtime/Engine/Private/PhysicsEngine/PhysicsHandleComponent.cpp b/Engine/Source/Runtime/Engine/Private/PhysicsEngine/PhysicsHandleComponent.cpp
--- a/Engine/Source/Runtime/Engine/Private/PhysicsEngine/PhysicsHandleComponent.cpp
+++ b/Engine/Source/Runtime/Engine/Private/PhysicsEngine/PhysicsHandleComponent.cpp
@@ -34,6 +34,7 @@ void UPhysicsHandleComponent::OnUnregister()
 
 		// use correct scene
 		PxScene* PScene = GetPhysXSceneFromIndex( SceneIndex );
+		SCOPED_SCENE_WRITE_LOCK(PScene);
 		if(PScene)
 		{
 			// destroy joint
@@ -72,76 +73,74 @@ void UPhysicsHandleComponent::GrabComponent(UPrimitiveComponent* InComponent, FN
 		return;
 	}
 
-	PxRigidDynamic* Actor = BodyInstance->GetPxRigidDynamic();
-	if (!Actor)
-		return;
+	ExecuteOnPxRigidDynamicReadWrite(BodyInstance, [&](PxRigidDynamic* Actor)
+	{
+		PxScene* Scene = Actor->getScene();
+		
+		// Get transform of actor we are grabbing
+		PxVec3 KinLocation = U2PVector(Location);
+		PxTransform GrabbedActorPose = Actor->getGlobalPose();
+		PxTransform KinPose(KinLocation, GrabbedActorPose.q);
 
-	// Get the scene the PxRigidDynamic we want to grab is in.
-	PxScene* Scene = Actor->getScene();
-	check(Scene);
+		// set target and current, so we don't need another "Tick" call to have it right
+		TargetTransform = CurrentTransform = P2UTransform(KinPose);
 
-	// Get transform of actor we are grabbing
-	PxVec3 KinLocation = U2PVector(Location);
-	PxTransform GrabbedActorPose = Actor->getGlobalPose();
-	PxTransform KinPose(KinLocation, GrabbedActorPose.q);
+		// If we don't already have a handle - make one now.
+		if (!HandleData)
+		{
+			// Create kinematic actor we are going to create joint with. This will be moved around with calls to SetLocation/SetRotation.
+			PxRigidDynamic* KinActor = Scene->getPhysics().createRigidDynamic(KinPose);
+			KinActor->setRigidDynamicFlag(PxRigidDynamicFlag::eKINEMATIC, true);
+			KinActor->setMass(1.0f);
+			KinActor->setMassSpaceInertiaTensor(PxVec3(1.0f, 1.0f, 1.0f));
 
-	// set target and current, so we don't need another "Tick" call to have it right
-	TargetTransform = CurrentTransform = P2UTransform(KinPose);
+			// No bodyinstance
+			KinActor->userData = NULL;
 
-	// If we don't already have a handle - make one now.
-	if (!HandleData)
-	{
-		SCOPED_SCENE_WRITE_LOCK(Scene);
-		// Create kinematic actor we are going to create joint with. This will be moved around with calls to SetLocation/SetRotation.
-		PxRigidDynamic* KinActor = Scene->getPhysics().createRigidDynamic(KinPose);
-		KinActor->setRigidDynamicFlag(PxRigidDynamicFlag::eKINEMATIC, true);
-		KinActor->setMass(1.0f);
-		KinActor->setMassSpaceInertiaTensor(PxVec3(1.0f, 1.0f, 1.0f));
+			// Add to Scene
+			Scene->addActor(*KinActor);
 
-		// No bodyinstance
-		KinActor->userData = NULL;
+			// Save reference to the kinematic actor.
+			KinActorData = KinActor;
 
-		// Add to Scene
-		Scene->addActor(*KinActor);
+			// Create the joint
+			PxVec3 LocalHandlePos = GrabbedActorPose.transformInv(KinLocation);
+			PxD6Joint* NewJoint = PxD6JointCreate(Scene->getPhysics(), KinActor, PxTransform::createIdentity(), Actor, PxTransform(LocalHandlePos));
 
-		// Save reference to the kinematic actor.
-		KinActorData = KinActor;
-		
-		// Create the joint
-		PxVec3 LocalHandlePos = GrabbedActorPose.transformInv(KinLocation);
-		PxD6Joint* NewJoint = PxD6JointCreate(Scene->getPhysics(), KinActor, PxTransform::createIdentity(), Actor, PxTransform(LocalHandlePos));
+			if (!NewJoint)
+			{
+				HandleData = 0;
+			}
+			else
+			{
+				// No constraint instance
+				NewJoint->userData = NULL;
+				HandleData = NewJoint;
+
+				// Remember the scene index that the handle joint/actor are in.
+				FPhysScene* RBScene = FPhysxUserData::Get<FPhysScene>(Scene->userData);
+				const uint32 SceneType = InComponent->BodyInstance.UseAsyncScene() ? PST_Async : PST_Sync;
+				SceneIndex = RBScene->PhysXSceneIndex[SceneType];
+
+				// Setting up the joint
+				NewJoint->setMotion(PxD6Axis::eX, PxD6Motion::eFREE);
+				NewJoint->setMotion(PxD6Axis::eY, PxD6Motion::eFREE);
+				NewJoint->setMotion(PxD6Axis::eZ, PxD6Motion::eFREE);
+				NewJoint->setDrivePosition(PxTransform(PxVec3(0, 0, 0)));
+
+				NewJoint->setMotion(PxD6Axis::eTWIST, PxD6Motion::eFREE);
+				NewJoint->setMotion(PxD6Axis::eSWING1, PxD6Motion::eFREE);
+				NewJoint->setMotion(PxD6Axis::eSWING2, PxD6Motion::eFREE);
+
+				bRotationConstrained = bConstrainRotation;
+
+				UpdateDriveSettings();
+			}
 
-		if(!NewJoint)
-		{
-			HandleData = 0;
-		}
-		else
-		{
-			// No constraint instance
-			NewJoint->userData = NULL;
-			HandleData = NewJoint;
-
-			// Remember the scene index that the handle joint/actor are in.
-			FPhysScene* RBScene = FPhysxUserData::Get<FPhysScene>(Scene->userData);
-			const uint32 SceneType = InComponent->BodyInstance.UseAsyncScene() ? PST_Async : PST_Sync;
-			SceneIndex = RBScene->PhysXSceneIndex[SceneType];
-
-			// Setting up the joint
-			NewJoint->setMotion(PxD6Axis::eX, PxD6Motion::eFREE);
-			NewJoint->setMotion(PxD6Axis::eY, PxD6Motion::eFREE);
-			NewJoint->setMotion(PxD6Axis::eZ, PxD6Motion::eFREE);
-			NewJoint->setDrivePosition(PxTransform(PxVec3(0,0,0)));
-
-			NewJoint->setMotion(PxD6Axis::eTWIST, PxD6Motion::eFREE);
-			NewJoint->setMotion(PxD6Axis::eSWING1, PxD6Motion::eFREE);
-			NewJoint->setMotion(PxD6Axis::eSWING2, PxD6Motion::eFREE);
-			
-			bRotationConstrained = bConstrainRotation;
-			
-			UpdateDriveSettings();
 		}
+	});
+	
 	
-	}
 #endif // WITH_PHYSX
 
 
@@ -233,6 +232,8 @@ void UPhysicsHandleComponent::UpdateHandleTransform(const FTransform& NewTransfo
 	bool bChangedRotation = true;
 
 	PxRigidDynamic* KinActor = KinActorData;
+	PxScene* PScene = GetPhysXSceneFromIndex(SceneIndex);
+	SCOPED_SCENE_WRITE_LOCK(PScene);
 
 	// Check if the new location is worthy of change
 	PxVec3 PNewLocation = U2PVector(NewTransform.GetTranslation());
diff --git a/Engine/Source/Runtime/Engine/Private/PhysicsEngine/PhysicsSettings.cpp b/Engine/Source/Runtime/Engine/Private/PhysicsEngine/PhysicsSettings.cpp
--- a/Engine/Source/Runtime/Engine/Private/PhysicsEngine/PhysicsSettings.cpp
+++ b/Engine/Source/Runtime/Engine/Private/PhysicsEngine/PhysicsSettings.cpp
@@ -13,6 +13,7 @@ UPhysicsSettings::UPhysicsSettings(const FObjectInitializer& ObjectInitializer)
 	, RagdollAggregateThreshold(4)
 	, TriangleMeshTriangleMinAreaThreshold(5.0f)
 	, bEnableAsyncScene(false)
+	, bWarnMissingLocks(true)
 	, bEnable2DPhysics(false)
 	, LockedAxis_DEPRECATED(static_cast<ESettingsLockedAxis::Type>(-1))
 	, BounceThresholdVelocity(200.f)
diff --git a/Engine/Source/Runtime/Engine/Private/SkeletalMeshComponentPhysics.cpp b/Engine/Source/Runtime/Engine/Private/SkeletalMeshComponentPhysics.cpp
--- a/Engine/Source/Runtime/Engine/Private/SkeletalMeshComponentPhysics.cpp
+++ b/Engine/Source/Runtime/Engine/Private/SkeletalMeshComponentPhysics.cpp
@@ -825,16 +825,16 @@ void USkeletalMeshComponent::InitArticulated(FPhysScene* PhysScene)
 	SetRootBodyIndex(RootBodyIndex);
 
 #if WITH_PHYSX
-	// add Aggregate into the scene
-	if(Aggregate && Aggregate->getNbActors() > 0 && PhysScene)
-	{
+
 		// Get the scene type from the SkeletalMeshComponent's BodyInstance
 		const uint32 SceneType = (bHasBodiesInAsyncScene && PhysScene->HasAsyncScene()) ? PST_Async : PST_Sync;
-		{
 			PxScene* PScene = PhysScene->GetPhysXScene(SceneType);
 			SCOPED_SCENE_WRITE_LOCK(PScene);
+
+	// add Aggregate into the scene
+	if(Aggregate && Aggregate->getNbActors() > 0 && PhysScene)
+	{
 			PScene->addAggregate(*Aggregate);
-		}
 		
 		// If we've used an aggregate, InitBody would not be able to set awake status as we *must* have a scene
 		// to do that, so we reconcile this here.
@@ -849,7 +849,7 @@ void USkeletalMeshComponent::InitArticulated(FPhysScene* PhysScene)
 			// Set to sleep if necessary
 			if(bShouldSleep)
 			{
-				Body->GetPxRigidDynamic()->putToSleep();
+				Body->GetPxRigidDynamic_AssumesLocked()->putToSleep();
 			}
 		}
 	}
@@ -899,6 +899,12 @@ void USkeletalMeshComponent::TermArticulated()
 	{
 		PhysScene->DeferredRemoveCollisionDisableTable(SkelMeshCompID);
 	}
+
+	// Get the scene type from the SkeletalMeshComponent's BodyInstance
+	const uint32 SceneType = BodyInstance.UseAsyncScene() ? PST_Async : PST_Sync;
+	PxScene* PScene = PhysScene->GetPhysXScene(SceneType);
+	SCOPED_SCENE_WRITE_LOCK(PScene);
+
 #endif	//#if WITH_PHYSX
 
 	// We shut down the physics for each body and constraint here. 
@@ -1843,73 +1849,18 @@ bool USkeletalMeshComponent::SweepComponent( FHitResult& OutHit, const FVector S
 
 bool USkeletalMeshComponent::ComponentOverlapComponent(class UPrimitiveComponent* PrimComp,const FVector Pos,const FRotator Rot,const struct FCollisionQueryParams& Params)
 {
-	//@TODO: BOX2D: USkeletalMeshComponent::ComponentOverlapComponent is not supported.  Try to rephrase this in terms of agnostic operations on a set of FBodyInstances, reducing the amount of PhysX-specific code
-#if WITH_PHYSX
-	// will have to do per component - default single body or physicsinstance 
-	const PxRigidActor* TargetRigidBody = (PrimComp)? PrimComp->BodyInstance.GetPxRigidActor():NULL;
-	if (TargetRigidBody==NULL || TargetRigidBody->getNbShapes()==0)
-	{
-		return false;
-	}
-
-	// if target is skeletalmeshcomponent and do not support singlebody physics
-	USkeletalMeshComponent * OtherComp = Cast<USkeletalMeshComponent>(PrimComp);
-	if (OtherComp)
+	//we do not support skeletal mesh vs skeletal mesh overlap test
+	if (PrimComp->IsA<USkeletalMeshComponent>())
 	{
 		UE_LOG(LogCollision, Log, TEXT("ComponentOverlapComponent : (%s) Does not support skeletalmesh with Physics Asset"), *PrimComp->GetPathName());
 		return false;
 	}
 
-	// calculate the test global pose of the actor
-	PxTransform PTestGlobalPose = U2PTransform(FTransform(Rot, Pos));
-
-	// Get all the shapes from the actor
-	TArray<PxShape*> PTargetShapes;
-	PTargetShapes.AddZeroed(TargetRigidBody->getNbShapes());
-	int32 NumTargetShapes = TargetRigidBody->getShapes(PTargetShapes.GetData(), PTargetShapes.Num());
-
-	bool bHaveOverlap = false;
-
-	for (int32 TargetShapeIdx=0; TargetShapeIdx<PTargetShapes.Num(); ++TargetShapeIdx)
+	if(FBodyInstance* BI = PrimComp->GetBodyInstance())
 	{
-		const PxShape * PTargetShape = PTargetShapes[TargetShapeIdx];
-		check (PTargetShape);
-
-		// Calc shape global pose
-		PxTransform PShapeGlobalPose = PTestGlobalPose.transform(PTargetShape->getLocalPose());
-
-		GET_GEOMETRY_FROM_SHAPE(PGeom, PTargetShape);
-
-		if(PGeom != NULL)
-		{
-			for (int32 BodyIdx=0; BodyIdx < Bodies.Num(); ++BodyIdx)
-			{
-				bHaveOverlap = Bodies[BodyIdx]->OverlapPhysX(*PGeom, PShapeGlobalPose);
-#if !(UE_BUILD_SHIPPING || UE_BUILD_TEST)
-				if((GetWorld()->DebugDrawTraceTag != NAME_None) && (GetWorld()->DebugDrawTraceTag == Params.TraceTag))
-				{
-					TArray<FOverlapResult> Overlaps;
-					if (bHaveOverlap)
-					{
-						FOverlapResult Result;
-						Result.Component = PrimComp;
-						Result.Actor = PrimComp->GetOwner();
-						Result.bBlockingHit = true;
-						Overlaps.Add(Result);
+		return BI->OverlapTestForBodies(Pos, Rot.Quaternion(), Bodies);
 					}
 
-					DrawGeomOverlaps(GetWorld(), *PGeom, PShapeGlobalPose, Overlaps, DebugLineLifetime);
-				}
-#endif //!(UE_BUILD_SHIPPING || UE_BUILD_TEST)
-				if (bHaveOverlap)
-				{
-					break;
-				}
-			}
-		}
-		return bHaveOverlap;
-	}
-#endif //WITH_PHYSX
 	return false;
 }
 
@@ -2503,24 +2454,26 @@ bool USkeletalMeshComponent::GetClothCollisionDataFromStaticMesh(UPrimitiveCompo
 		return false;
 	}
 
-	int32 NumSyncShapes = 0;
-
-	TArray<PxShape*> AllShapes = PrimComp->BodyInstance.GetAllShapes(NumSyncShapes);
+	bool bSuccess = false;
+	PrimComp->BodyInstance.ExecuteOnPhysicsReadOnly([&]
+	{
+		TArray<PxShape*> AllShapes;
+		const int32 NumSyncShapes = PrimComp->BodyInstance.GetAllShapes_AssumesLocked(AllShapes);
 
-	if(NumSyncShapes == 0 || NumSyncShapes > 3) //skipping complicated object because of collision limitation
+		if (NumSyncShapes == 0 || NumSyncShapes > 3) //skipping complicated object because of collision limitation
 	{
-		return false;
+			return;
 	}
 
 	FVector Center = PrimComp->Bounds.Origin;
 	FTransform Transform = PrimComp->ComponentToWorld;
 	FMatrix TransMat = Transform.ToMatrixWithScale();
 
-	for(int32 ShapeIdx=0; ShapeIdx < NumSyncShapes; ShapeIdx++)
+		for (int32 ShapeIdx = 0; ShapeIdx < NumSyncShapes; ShapeIdx++)
 	{
 		PxGeometryType::Enum GeomType = AllShapes[ShapeIdx]->getGeometryType();
 
-		switch(GeomType)
+			switch (GeomType)
 		{
 		case PxGeometryType::eSPHERE:
 			{
@@ -2570,27 +2523,27 @@ bool USkeletalMeshComponent::GetClothCollisionDataFromStaticMesh(UPrimitiveCompo
 
 				ClothPrimData.ConvexPlanes.Empty(6); //box has 6 planes
 
-				FPlane UPlane1(1,0,0,Center.X + BoxGeom.halfExtents.x);
+				FPlane UPlane1(1, 0, 0, Center.X + BoxGeom.halfExtents.x);
 				UPlane1 = UPlane1.TransformBy(TransMat);
 				ClothPrimData.ConvexPlanes.Add(UPlane1);
 
-				FPlane UPlane2(-1,0,0,Center.X - BoxGeom.halfExtents.x);
+				FPlane UPlane2(-1, 0, 0, Center.X - BoxGeom.halfExtents.x);
 				UPlane2 = UPlane2.TransformBy(TransMat);
 				ClothPrimData.ConvexPlanes.Add(UPlane2);
 
-				FPlane UPlane3(0,1,0,Center.Y + BoxGeom.halfExtents.y);
+				FPlane UPlane3(0, 1, 0, Center.Y + BoxGeom.halfExtents.y);
 				UPlane3 = UPlane3.TransformBy(TransMat);
 				ClothPrimData.ConvexPlanes.Add(UPlane3);
 
-				FPlane UPlane4(0,-1,0,Center.Y - BoxGeom.halfExtents.y);
+				FPlane UPlane4(0, -1, 0, Center.Y - BoxGeom.halfExtents.y);
 				UPlane4 = UPlane4.TransformBy(TransMat);
 				ClothPrimData.ConvexPlanes.Add(UPlane4);
 
-				FPlane UPlane5(0,0,1,Center.Z + BoxGeom.halfExtents.z);
+				FPlane UPlane5(0, 0, 1, Center.Z + BoxGeom.halfExtents.z);
 				UPlane5 = UPlane5.TransformBy(TransMat);
 				ClothPrimData.ConvexPlanes.Add(UPlane5);
 
-				FPlane UPlane6(0,0,-1,Center.Z - BoxGeom.halfExtents.z);
+				FPlane UPlane6(0, 0, -1, Center.Z - BoxGeom.halfExtents.z);
 				UPlane6 = UPlane6.TransformBy(TransMat);
 				ClothPrimData.ConvexPlanes.Add(UPlane6);
 
@@ -2604,7 +2557,7 @@ bool USkeletalMeshComponent::GetClothCollisionDataFromStaticMesh(UPrimitiveCompo
 
 				AllShapes[ShapeIdx]->getConvexMeshGeometry(ConvexGeom);
 
-				if(ConvexGeom.convexMesh)
+				if (ConvexGeom.convexMesh)
 				{
 					FClothCollisionPrimitive ClothPrimData;
 					ClothPrimData.Origin = Center;
@@ -2614,11 +2567,11 @@ bool USkeletalMeshComponent::GetClothCollisionDataFromStaticMesh(UPrimitiveCompo
 
 					ClothPrimData.ConvexPlanes.Empty(NumPoly);
 
-					for(uint32 Poly=0; Poly < NumPoly; Poly++)
+					for (uint32 Poly = 0; Poly < NumPoly; Poly++)
 					{
 						PxHullPolygon HullData;
-						ConvexGeom.convexMesh->getPolygonData(Poly, HullData);						
-						physx::PxPlane PPlane(HullData.mPlane[0],HullData.mPlane[1],HullData.mPlane[2],HullData.mPlane[3]);
+						ConvexGeom.convexMesh->getPolygonData(Poly, HullData);
+						physx::PxPlane PPlane(HullData.mPlane[0], HullData.mPlane[1], HullData.mPlane[2], HullData.mPlane[3]);
 						FPlane UPlane = P2UPlane(PPlane);
 						UPlane = UPlane.TransformBy(TransMat);
 						ClothPrimData.ConvexPlanes.Add(UPlane);
@@ -2631,7 +2584,10 @@ bool USkeletalMeshComponent::GetClothCollisionDataFromStaticMesh(UPrimitiveCompo
 			break;
 		}
 	}
-	return true;
+		bSuccess = true;
+	});
+
+	return bSuccess;
 }
 
 void USkeletalMeshComponent::FindClothCollisions(TArray<FApexClothCollisionVolumeData>& OutCollisions)
diff --git a/Engine/Source/Runtime/Engine/Private/Vehicles/PhysXVehicleManager.cpp b/Engine/Source/Runtime/Engine/Private/Vehicles/PhysXVehicleManager.cpp
--- a/Engine/Source/Runtime/Engine/Private/Vehicles/PhysXVehicleManager.cpp
+++ b/Engine/Source/Runtime/Engine/Private/Vehicles/PhysXVehicleManager.cpp
@@ -268,7 +268,11 @@ void FPhysXVehicleManager::Update( float DeltaTime )
 	}
 
 	// Suspension raycasts
-	PxVehicleSuspensionRaycasts( WheelRaycastBatchQuery, PVehicles.Num(), PVehicles.GetData(), WheelQueryResults.Num(), WheelQueryResults.GetData() );
+	{
+		SCOPED_SCENE_READ_LOCK(Scene);
+		PxVehicleSuspensionRaycasts( WheelRaycastBatchQuery, PVehicles.Num(), PVehicles.GetData(), WheelQueryResults.Num(), WheelQueryResults.GetData() );
+	}
+	
 
 	// Tick vehicles
 	for ( int32 i = Vehicles.Num() - 1; i >= 0; --i )
@@ -304,10 +308,11 @@ void FPhysXVehicleManager::PreTick( float DeltaTime )
 
 void FPhysXVehicleManager::UpdateVehicles( float DeltaTime )
 {
-	PxVehicleUpdates( DeltaTime, GetSceneGravity(), *SurfaceTirePairs, PVehicles.Num(), PVehicles.GetData(), PVehiclesWheelsStates.GetData());
+	SCOPED_SCENE_WRITE_LOCK(Scene);
+	PxVehicleUpdates( DeltaTime, GetSceneGravity_AssumesLocked(), *SurfaceTirePairs, PVehicles.Num(), PVehicles.GetData(), PVehiclesWheelsStates.GetData());
 }
 
-PxVec3 FPhysXVehicleManager::GetSceneGravity()
+PxVec3 FPhysXVehicleManager::GetSceneGravity_AssumesLocked()
 {
 	return Scene->getGravity();
 }
@@ -356,6 +361,7 @@ void FPhysXVehicleManager::SetRecordTelemetry( TWeakObjectPtr<UWheeledVehicleMov
 void FPhysXVehicleManager::SetupTelemetryData()
 {
 	// set up telemetry for 4 wheels
+	SCOPED_SCENE_WRITE_LOCK(Scene);
 	if(TelemetryData4W == NULL)
 	{
 		float Empty[] = { 0.0f, 0.0f, 0.0f, 0.0f, 0.0f, 0.0f };
@@ -370,26 +376,25 @@ void FPhysXVehicleManager::UpdateVehiclesWithTelemetry( float DeltaTime )
 	check(TelemetryVehicle);
 	check(PVehicles.Find(TelemetryVehicle) == 0);
 
-	PxVehicleTelemetryData* TelemetryData = GetTelemetryData();
-
-	if ( TelemetryData )
+	SCOPED_SCENE_WRITE_LOCK(Scene);
+	if ( PxVehicleTelemetryData* TelemetryData = GetTelemetryData_AssumesLocked() )
 	{
-		PxVehicleUpdateSingleVehicleAndStoreTelemetryData( DeltaTime, GetSceneGravity(), *SurfaceTirePairs, TelemetryVehicle, PVehiclesWheelsStates.GetData(), *TelemetryData );
+		PxVehicleUpdateSingleVehicleAndStoreTelemetryData( DeltaTime, GetSceneGravity_AssumesLocked(), *SurfaceTirePairs, TelemetryVehicle, PVehiclesWheelsStates.GetData(), *TelemetryData );
 
 		if ( PVehicles.Num() > 1 )
 		{
-			PxVehicleUpdates( DeltaTime, GetSceneGravity(), *SurfaceTirePairs, PVehicles.Num() - 1, &PVehicles[1], &PVehiclesWheelsStates[1] );
+			PxVehicleUpdates( DeltaTime, GetSceneGravity_AssumesLocked(), *SurfaceTirePairs, PVehicles.Num() - 1, &PVehicles[1], &PVehiclesWheelsStates[1] );
 		}
 	}
 	else
 	{
 		UE_LOG( LogPhysics, Warning, TEXT("Cannot record telemetry for vehicle, it does not have 4 wheels") );
 
-		PxVehicleUpdates( DeltaTime, GetSceneGravity(), *SurfaceTirePairs, PVehicles.Num(), PVehicles.GetData(), PVehiclesWheelsStates.GetData() );
+		PxVehicleUpdates( DeltaTime, GetSceneGravity_AssumesLocked(), *SurfaceTirePairs, PVehicles.Num(), PVehicles.GetData(), PVehiclesWheelsStates.GetData() );
 	}
 }
 
-PxVehicleTelemetryData* FPhysXVehicleManager::GetTelemetryData()
+PxVehicleTelemetryData* FPhysXVehicleManager::GetTelemetryData_AssumesLocked()
 {
 	if ( TelemetryVehicle )
 	{
@@ -404,7 +409,7 @@ PxVehicleTelemetryData* FPhysXVehicleManager::GetTelemetryData()
 
 #endif //PX_DEBUG_VEHICLE_ON
 
-PxWheelQueryResult* FPhysXVehicleManager::GetWheelsStates(TWeakObjectPtr<class UWheeledVehicleMovementComponent> Vehicle)
+PxWheelQueryResult* FPhysXVehicleManager::GetWheelsStates_AssumesLocked(TWeakObjectPtr<class UWheeledVehicleMovementComponent> Vehicle)
 {
 	int32 Index = Vehicles.Find(Vehicle);
 
diff --git a/Engine/Source/Runtime/Engine/Private/Vehicles/PhysXVehicleManager.h b/Engine/Source/Runtime/Engine/Private/Vehicles/PhysXVehicleManager.h
--- a/Engine/Source/Runtime/Engine/Private/Vehicles/PhysXVehicleManager.h
+++ b/Engine/Source/Runtime/Engine/Private/Vehicles/PhysXVehicleManager.h
@@ -47,12 +47,30 @@ public:
 	/**
 	 * Get the updated telemetry data
 	 */
-	PxVehicleTelemetryData* GetTelemetryData();
+	DEPRECATED(4.8, "Please call GetTelemetryData_AssumesLocked and make sure you obtain the appropriate PhysX scene locks")
+	PxVehicleTelemetryData* GetTelemetryData()
+	{
+		return GetTelemetryData_AssumesLocked();
+	}
+
+	/**
+	 * Get the updated telemetry data
+	 */
+	PxVehicleTelemetryData* GetTelemetryData_AssumesLocked();
 	
 	/**
 	 * Get a vehicle's wheels states, such as isInAir, suspJounce, contactPoints, etc
 	 */
-	PxWheelQueryResult* GetWheelsStates(TWeakObjectPtr<class UWheeledVehicleMovementComponent> Vehicle);
+	DEPRECATED(4.8, "Please call GetWheelsStates_AssumesLocked and make sure you obtain the appropriate PhysX scene locks")
+	PxWheelQueryResult* GetWheelsStates(TWeakObjectPtr<class UWheeledVehicleMovementComponent> Vehicle)
+	{
+		return GetWheelsStates_AssumesLocked(Vehicle);
+	}
+
+	/**
+	 * Get a vehicle's wheels states, such as isInAir, suspJounce, contactPoints, etc
+	 */
+	PxWheelQueryResult* GetWheelsStates_AssumesLocked(TWeakObjectPtr<class UWheeledVehicleMovementComponent> Vehicle);
 
 	/**
 	 * Update vehicle data before the scene simulates
@@ -63,6 +81,8 @@ public:
 	 * Update vehicle tuning and other state such as input */
 	void PreTick( float DeltaTime );
 
+	PxScene* GetScene() const { return Scene; }
+
 private:
 
 	// True if the tire friction table needs to be updated
@@ -110,7 +130,7 @@ private:
 	/**
 	 * Get the gravity for our phys scene
 	 */
-	PxVec3 GetSceneGravity();
+	PxVec3 GetSceneGravity_AssumesLocked();
 
 #if PX_DEBUG_VEHICLE_ON
 
diff --git a/Engine/Source/Runtime/Engine/Private/Vehicles/VehicleWheel.cpp b/Engine/Source/Runtime/Engine/Private/Vehicles/VehicleWheel.cpp
--- a/Engine/Source/Runtime/Engine/Private/Vehicles/VehicleWheel.cpp
+++ b/Engine/Source/Runtime/Engine/Private/Vehicles/VehicleWheel.cpp
@@ -44,7 +44,9 @@ UVehicleWheel::UVehicleWheel(const FObjectInitializer& ObjectInitializer)
 float UVehicleWheel::GetSteerAngle()
 {
 #if WITH_VEHICLE
-	return FMath::RadiansToDegrees( VehicleSim->GetWorld()->GetPhysicsScene()->GetVehicleManager()->GetWheelsStates(VehicleSim)[WheelIndex].steerAngle );
+	FPhysXVehicleManager* VehicleManager = VehicleSim->GetWorld()->GetPhysicsScene()->GetVehicleManager();
+	SCOPED_SCENE_READ_LOCK(VehicleManager->GetScene());
+	return FMath::RadiansToDegrees( VehicleManager->GetWheelsStates_AssumesLocked(VehicleSim)[WheelIndex].steerAngle );
 #else
 	return 0.0f;
 #endif // WITH_PHYSX
@@ -53,6 +55,9 @@ float UVehicleWheel::GetSteerAngle()
 float UVehicleWheel::GetRotationAngle()
 {
 #if WITH_VEHICLE
+	FPhysXVehicleManager* VehicleManager = VehicleSim->GetWorld()->GetPhysicsScene()->GetVehicleManager();
+	SCOPED_SCENE_READ_LOCK(VehicleManager->GetScene());
+
 	float RotationAngle = -1.0f * FMath::RadiansToDegrees( VehicleSim->PVehicle->mWheelsDynData.getWheelRotationAngle( WheelIndex ) );
 	check(!FMath::IsNaN(RotationAngle));
 	return RotationAngle;
@@ -64,7 +69,10 @@ float UVehicleWheel::GetRotationAngle()
 float UVehicleWheel::GetSuspensionOffset()
 {
 #if WITH_VEHICLE
-	return VehicleSim->GetWorld()->GetPhysicsScene()->GetVehicleManager()->GetWheelsStates(VehicleSim)[WheelIndex].suspJounce;
+	FPhysXVehicleManager* VehicleManager = VehicleSim->GetWorld()->GetPhysicsScene()->GetVehicleManager();
+	SCOPED_SCENE_READ_LOCK(VehicleManager->GetScene());
+
+	return VehicleManager->GetWheelsStates_AssumesLocked(VehicleSim)[WheelIndex].suspJounce;
 #else
 	return 0.0f;
 #endif // WITH_PHYSX
@@ -82,6 +90,9 @@ void UVehicleWheel::Init( UWheeledVehicleMovementComponent* InVehicleSim, int32
 	WheelShape = NULL;
 
 #if WITH_VEHICLE
+	FPhysXVehicleManager* VehicleManager = VehicleSim->GetWorld()->GetPhysicsScene()->GetVehicleManager();
+	SCOPED_SCENE_READ_LOCK(VehicleManager->GetScene());
+
 	const int32 WheelShapeIdx = VehicleSim->PVehicle->mWheelsSimData.getWheelShapeMapping( WheelIndex );
 	check(WheelShapeIdx >= 0);
 
@@ -115,6 +126,9 @@ FVector UVehicleWheel::GetPhysicsLocation()
 #if WITH_VEHICLE
 	if ( WheelShape )
 	{
+		FPhysXVehicleManager* VehicleManager = VehicleSim->GetWorld()->GetPhysicsScene()->GetVehicleManager();
+		SCOPED_SCENE_READ_LOCK(VehicleManager->GetScene());
+
 		PxVec3 PLocation = VehicleSim->PVehicle->getRigidDynamicActor()->getGlobalPose().transform( WheelShape->getLocalPose() ).p;
 		return P2UVector( PLocation );
 	}
@@ -142,7 +156,10 @@ UPhysicalMaterial* UVehicleWheel::GetContactSurfaceMaterial()
 	UPhysicalMaterial* PhysMaterial = NULL;
 
 #if WITH_VEHICLE
-	const PxMaterial* ContactSurface = VehicleSim->GetWorld()->GetPhysicsScene()->GetVehicleManager()->GetWheelsStates(VehicleSim)[WheelIndex].tireSurfaceMaterial;
+	FPhysXVehicleManager* VehicleManager = VehicleSim->GetWorld()->GetPhysicsScene()->GetVehicleManager();
+	SCOPED_SCENE_READ_LOCK(VehicleManager->GetScene());
+
+	const PxMaterial* ContactSurface = VehicleManager->GetWheelsStates_AssumesLocked(VehicleSim)[WheelIndex].tireSurfaceMaterial;
 	if (ContactSurface)
 	{
 		PhysMaterial = FPhysxUserData::Get<UPhysicalMaterial>(ContactSurface->userData);
diff --git a/Engine/Source/Runtime/Engine/Private/Vehicles/WheeledVehicleMovementComponent.cpp b/Engine/Source/Runtime/Engine/Private/Vehicles/WheeledVehicleMovementComponent.cpp
--- a/Engine/Source/Runtime/Engine/Private/Vehicles/WheeledVehicleMovementComponent.cpp
+++ b/Engine/Source/Runtime/Engine/Private/Vehicles/WheeledVehicleMovementComponent.cpp
@@ -181,7 +181,7 @@ void UWheeledVehicleMovementComponent::CreateVehicle()
 			check(UpdatedComponent);
 			if (ensure(UpdatedPrimitive != nullptr))
 			{
-				check(UpdatedPrimitive->GetBodyInstance()->GetPxRigidDynamic());
+				check(UpdatedPrimitive->GetBodyInstance()->IsDynamic());
 				
 				SetupVehicle();
 
@@ -201,68 +201,69 @@ void UWheeledVehicleMovementComponent::SetupVehicleShapes()
 		return;
 	}
 
-	PxRigidDynamic* PVehicleActor = UpdatedPrimitive->GetBodyInstance()->GetPxRigidDynamic();
-
 	static PxMaterial* WheelMaterial = GPhysXSDK->createMaterial(0.0f, 0.0f, 0.0f);
 
-	// Add wheel shapes to actor
-	for ( int32 WheelIdx = 0; WheelIdx < WheelSetups.Num(); ++WheelIdx )
+	ExecuteOnPxRigidDynamicReadWrite(UpdatedPrimitive->GetBodyInstance(), [&](PxRigidDynamic* PVehicleActor)
 	{
-		FWheelSetup& WheelSetup = WheelSetups[WheelIdx];
-		UVehicleWheel* Wheel = WheelSetup.WheelClass.GetDefaultObject();
-		check(Wheel);
+		// Add wheel shapes to actor
+		for (int32 WheelIdx = 0; WheelIdx < WheelSetups.Num(); ++WheelIdx)
+		{
+			FWheelSetup& WheelSetup = WheelSetups[WheelIdx];
+			UVehicleWheel* Wheel = WheelSetup.WheelClass.GetDefaultObject();
+			check(Wheel);
 
-		const FVector WheelOffset = GetWheelRestingPosition( WheelSetup );
-		const PxTransform PLocalPose = PxTransform( U2PVector( WheelOffset ) );
-		PxShape* PWheelShape = NULL;
+			const FVector WheelOffset = GetWheelRestingPosition(WheelSetup);
+			const PxTransform PLocalPose = PxTransform(U2PVector(WheelOffset));
+			PxShape* PWheelShape = NULL;
 
-		// Prepare shape
-		if (Wheel->bDontCreateShape)
-		{
-			continue;
-		}
-		else if (Wheel->CollisionMesh && Wheel->CollisionMesh->BodySetup)
-		{
-			FBoxSphereBounds MeshBounds = Wheel->CollisionMesh->GetBounds();
-			FVector MeshScaleV = FVector(1,1,1);
-			if (Wheel->bAutoAdjustCollisionSize)
+			// Prepare shape
+			if (Wheel->bDontCreateShape)
 			{
-				MeshScaleV.X = Wheel->ShapeRadius / MeshBounds.BoxExtent.X;
-				MeshScaleV.Y = Wheel->ShapeWidth / MeshBounds.BoxExtent.Y;
-				MeshScaleV.Z = Wheel->ShapeRadius / MeshBounds.BoxExtent.Z;
+				continue;
 			}
-			PxMeshScale MeshScale( U2PVector(UpdatedComponent->RelativeScale3D * MeshScaleV), PxQuat::createIdentity() );
-			if (Wheel->CollisionMesh->BodySetup->TriMesh)
+			else if (Wheel->CollisionMesh && Wheel->CollisionMesh->BodySetup)
 			{
-				PxTriangleMesh* TriMesh = Wheel->CollisionMesh->BodySetup->TriMesh;
+				FBoxSphereBounds MeshBounds = Wheel->CollisionMesh->GetBounds();
+				FVector MeshScaleV = FVector(1, 1, 1);
+				if (Wheel->bAutoAdjustCollisionSize)
+				{
+					MeshScaleV.X = Wheel->ShapeRadius / MeshBounds.BoxExtent.X;
+					MeshScaleV.Y = Wheel->ShapeWidth / MeshBounds.BoxExtent.Y;
+					MeshScaleV.Z = Wheel->ShapeRadius / MeshBounds.BoxExtent.Z;
+				}
+				PxMeshScale MeshScale(U2PVector(UpdatedComponent->RelativeScale3D * MeshScaleV), PxQuat::createIdentity());
+				if (Wheel->CollisionMesh->BodySetup->TriMesh)
+				{
+					PxTriangleMesh* TriMesh = Wheel->CollisionMesh->BodySetup->TriMesh;
 
-				// No eSIMULATION_SHAPE flag for wheels
-				PWheelShape = PVehicleActor->createShape( PxTriangleMeshGeometry(TriMesh, MeshScale), *WheelMaterial, PxShapeFlag::eSCENE_QUERY_SHAPE | PxShapeFlag::eVISUALIZATION );
-				PWheelShape->setLocalPose(PLocalPose);			
+					// No eSIMULATION_SHAPE flag for wheels
+					PWheelShape = PVehicleActor->createShape(PxTriangleMeshGeometry(TriMesh, MeshScale), *WheelMaterial, PxShapeFlag::eSCENE_QUERY_SHAPE | PxShapeFlag::eVISUALIZATION);
+					PWheelShape->setLocalPose(PLocalPose);
+				}
+				else if (Wheel->CollisionMesh->BodySetup->AggGeom.ConvexElems.Num() == 1)
+				{
+					PxConvexMesh* ConvexMesh = Wheel->CollisionMesh->BodySetup->AggGeom.ConvexElems[0].ConvexMesh;
+					PWheelShape = PVehicleActor->createShape(PxConvexMeshGeometry(ConvexMesh, MeshScale), *WheelMaterial, PLocalPose);
+				}
+
+				check(PWheelShape);
 			}
-			else if (Wheel->CollisionMesh->BodySetup->AggGeom.ConvexElems.Num() == 1)
+			else
 			{
-				PxConvexMesh* ConvexMesh = Wheel->CollisionMesh->BodySetup->AggGeom.ConvexElems[0].ConvexMesh;
-				PWheelShape = PVehicleActor->createShape( PxConvexMeshGeometry(ConvexMesh, MeshScale), *WheelMaterial, PLocalPose );
+				PWheelShape = PVehicleActor->createShape(PxSphereGeometry(Wheel->ShapeRadius), *WheelMaterial, PLocalPose);
 			}
 
-			check(PWheelShape);
-		}
-		else
-		{
-			PWheelShape = PVehicleActor->createShape( PxSphereGeometry(Wheel->ShapeRadius), *WheelMaterial, PLocalPose );
-		}
-
-		// Init filter data
-		FCollisionResponseContainer CollisionResponse;
-		CollisionResponse.SetAllChannels( ECR_Ignore );
+			// Init filter data
+			FCollisionResponseContainer CollisionResponse;
+			CollisionResponse.SetAllChannels(ECR_Ignore);
 
-		PxFilterData PWheelQueryFilterData, PDummySimData;
-		CreateShapeFilterData(ECC_Vehicle, UpdatedComponent->GetUniqueID(), CollisionResponse, 0, 0, PWheelQueryFilterData, PDummySimData, false, false, false );
+			PxFilterData PWheelQueryFilterData, PDummySimData;
+			CreateShapeFilterData(ECC_Vehicle, UpdatedComponent->GetUniqueID(), CollisionResponse, 0, 0, PWheelQueryFilterData, PDummySimData, false, false, false);
 
-		//// Give suspension raycasts the same group ID as the chassis so that they don't hit each other
-		PWheelShape->setQueryFilterData( PWheelQueryFilterData );
-	}
+			//// Give suspension raycasts the same group ID as the chassis so that they don't hit each other
+			PWheelShape->setQueryFilterData(PWheelQueryFilterData);
+		}
+	});
 }
 
 void UWheeledVehicleMovementComponent::SetupVehicleMass()
@@ -272,23 +273,23 @@ void UWheeledVehicleMovementComponent::SetupVehicleMass()
 		return;
 	}
 
-	const FBodyInstance& BodyInst = *UpdatedPrimitive->GetBodyInstance();
-	PxRigidDynamic* PVehicleActor = BodyInst.GetPxRigidDynamic();
-
-	// Override mass
-	const float MassRatio = Mass > 0.0f ? Mass / PVehicleActor->getMass() : 1.0f;
+	ExecuteOnPxRigidDynamicReadWrite(UpdatedPrimitive->GetBodyInstance(), [&](PxRigidDynamic* PVehicleActor)
+	{
+		// Override mass
+		const float MassRatio = Mass > 0.0f ? Mass / PVehicleActor->getMass() : 1.0f;
 
-	PxVec3 PInertiaTensor = PVehicleActor->getMassSpaceInertiaTensor();
+		PxVec3 PInertiaTensor = PVehicleActor->getMassSpaceInertiaTensor();
 
-	PInertiaTensor.x *= InertiaTensorScale.X * MassRatio;
-	PInertiaTensor.y *= InertiaTensorScale.Y * MassRatio;
-	PInertiaTensor.z *= InertiaTensorScale.Z * MassRatio;
+		PInertiaTensor.x *= InertiaTensorScale.X * MassRatio;
+		PInertiaTensor.y *= InertiaTensorScale.Y * MassRatio;
+		PInertiaTensor.z *= InertiaTensorScale.Z * MassRatio;
 
-	PVehicleActor->setMassSpaceInertiaTensor( PInertiaTensor );
-	PVehicleActor->setMass( Mass );
+		PVehicleActor->setMassSpaceInertiaTensor(PInertiaTensor);
+		PVehicleActor->setMass(Mass);
 
-	const PxVec3 PCOMOffset = U2PVector(GetLocalCOM());
-	PVehicleActor->setCMassLocalPose(PxTransform(PCOMOffset, PxQuat::createIdentity()));
+		const PxVec3 PCOMOffset = U2PVector(GetLocalCOM());
+		PVehicleActor->setCMassLocalPose(PxTransform(PCOMOffset, PxQuat::createIdentity()));
+	});
 }
 
 void UWheeledVehicleMovementComponent::SetupWheels( PxVehicleWheelsSimData* PWheelsSimData )
@@ -298,106 +299,104 @@ void UWheeledVehicleMovementComponent::SetupWheels( PxVehicleWheelsSimData* PWhe
 		return;
 	}
 
-	const FBodyInstance& VehicleBodyInstance = *UpdatedPrimitive->GetBodyInstance();
-	PxRigidDynamic* PVehicleActor = VehicleBodyInstance.GetPxRigidDynamic();
-
-	const PxReal LengthScale = 100.f; // Convert default from m to cm
+	ExecuteOnPxRigidDynamicReadWrite(UpdatedPrimitive->GetBodyInstance(), [&](PxRigidDynamic* PVehicleActor)
+	{
+		const PxReal LengthScale = 100.f; // Convert default from m to cm
 
-	// Control substepping
-	PWheelsSimData->setSubStepCount( 5.f * LengthScale, 3, 1 );
-	PWheelsSimData->setMinLongSlipDenominator(4.f * LengthScale); 
+		// Control substepping
+		PWheelsSimData->setSubStepCount(5.f * LengthScale, 3, 1);
+		PWheelsSimData->setMinLongSlipDenominator(4.f * LengthScale);
 
-	// Prealloc data for the sprung masses
-	PxVec3 WheelOffsets[32];
-	float SprungMasses[32];
-	check(WheelSetups.Num() <= 32);
+		// Prealloc data for the sprung masses
+		PxVec3 WheelOffsets[32];
+		float SprungMasses[32];
+		check(WheelSetups.Num() <= 32);
 
-	// Calculate wheel offsets first, necessary for sprung masses
-	for ( int32 WheelIdx = 0; WheelIdx < WheelSetups.Num(); ++WheelIdx )
-	{
-		WheelOffsets[WheelIdx] = U2PVector( GetWheelRestingPosition( WheelSetups[WheelIdx] ) );
-	}
+		// Calculate wheel offsets first, necessary for sprung masses
+		for (int32 WheelIdx = 0; WheelIdx < WheelSetups.Num(); ++WheelIdx)
+		{
+			WheelOffsets[WheelIdx] = U2PVector(GetWheelRestingPosition(WheelSetups[WheelIdx]));
+		}
 
-	// Now that we have all the wheel offsets, calculate the sprung masses
-	PxVec3 PLocalCOM = U2PVector(GetLocalCOM());
-	PxVehicleComputeSprungMasses(WheelSetups.Num(), WheelOffsets, PLocalCOM, PVehicleActor->getMass(), 2, SprungMasses);
+		// Now that we have all the wheel offsets, calculate the sprung masses
+		PxVec3 PLocalCOM = U2PVector(GetLocalCOM());
+		PxVehicleComputeSprungMasses(WheelSetups.Num(), WheelOffsets, PLocalCOM, PVehicleActor->getMass(), 2, SprungMasses);
 
-	for ( int32 WheelIdx = 0; WheelIdx < WheelSetups.Num(); ++WheelIdx )
-	{
-		UVehicleWheel* Wheel = WheelSetups[WheelIdx].WheelClass.GetDefaultObject();
+		for (int32 WheelIdx = 0; WheelIdx < WheelSetups.Num(); ++WheelIdx)
+		{
+			UVehicleWheel* Wheel = WheelSetups[WheelIdx].WheelClass.GetDefaultObject();
+
+			// init wheel data
+			PxVehicleWheelData PWheelData;
+			PWheelData.mRadius = Wheel->ShapeRadius;
+			PWheelData.mWidth = Wheel->ShapeWidth;
+			PWheelData.mMaxSteer = FMath::DegreesToRadians(Wheel->SteerAngle);
+			PWheelData.mMaxBrakeTorque = M2ToCm2(Wheel->MaxBrakeTorque);
+			PWheelData.mMaxHandBrakeTorque = Wheel->bAffectedByHandbrake ? M2ToCm2(Wheel->MaxHandBrakeTorque) : 0.0f;
+
+			PWheelData.mDampingRate = M2ToCm2(Wheel->DampingRate);
+			PWheelData.mMass = Wheel->Mass;
+			PWheelData.mMOI = 0.5f * PWheelData.mMass * FMath::Square(PWheelData.mRadius);
+
+			// init tire data
+			PxVehicleTireData PTireData;
+			PTireData.mType = Wheel->TireType ? Wheel->TireType->GetTireTypeID() : GEngine->DefaultTireType->GetTireTypeID();
+			//PTireData.mCamberStiffnessPerUnitGravity = 0.0f;
+			PTireData.mLatStiffX = Wheel->LatStiffMaxLoad;
+			PTireData.mLatStiffY = Wheel->LatStiffValue;
+			PTireData.mLongitudinalStiffnessPerUnitGravity = Wheel->LongStiffValue;
+
+			// init suspension data
+			PxVehicleSuspensionData PSuspensionData;
+			PSuspensionData.mSprungMass = SprungMasses[WheelIdx];
+			PSuspensionData.mMaxCompression = Wheel->SuspensionMaxRaise;
+			PSuspensionData.mMaxDroop = Wheel->SuspensionMaxDrop;
+			PSuspensionData.mSpringStrength = FMath::Square(Wheel->SuspensionNaturalFrequency) * PSuspensionData.mSprungMass;
+			PSuspensionData.mSpringDamperRate = Wheel->SuspensionDampingRatio * 2.0f * FMath::Sqrt(PSuspensionData.mSpringStrength * PSuspensionData.mSprungMass);
+
+			// init offsets
+			const PxVec3 PWheelOffset = WheelOffsets[WheelIdx];
+
+			PxVec3 PSuspTravelDirection = PxVec3(0.0f, 0.0f, -1.0f);
+			PxVec3 PWheelCentreCMOffset = PWheelOffset - PLocalCOM;
+			PxVec3 PSuspForceAppCMOffset = PxVec3(PWheelCentreCMOffset.x, PWheelCentreCMOffset.y, Wheel->SuspensionForceOffset);
+			PxVec3 PTireForceAppCMOffset = PSuspForceAppCMOffset;
+
+			// finalize sim data
+			PWheelsSimData->setWheelData(WheelIdx, PWheelData);
+			PWheelsSimData->setTireData(WheelIdx, PTireData);
+			PWheelsSimData->setSuspensionData(WheelIdx, PSuspensionData);
+			PWheelsSimData->setSuspTravelDirection(WheelIdx, PSuspTravelDirection);
+			PWheelsSimData->setWheelCentreOffset(WheelIdx, PWheelCentreCMOffset);
+			PWheelsSimData->setSuspForceAppPointOffset(WheelIdx, PSuspForceAppCMOffset);
+			PWheelsSimData->setTireForceAppPointOffset(WheelIdx, PTireForceAppCMOffset);
+		}
 
-		// init wheel data
-		PxVehicleWheelData PWheelData;
-		PWheelData.mRadius = Wheel->ShapeRadius;
-		PWheelData.mWidth = Wheel->ShapeWidth;	
-		PWheelData.mMaxSteer = FMath::DegreesToRadians(Wheel->SteerAngle);
-		PWheelData.mMaxBrakeTorque = M2ToCm2(Wheel->MaxBrakeTorque);
-		PWheelData.mMaxHandBrakeTorque = Wheel->bAffectedByHandbrake ? M2ToCm2(Wheel->MaxHandBrakeTorque) : 0.0f;
+		const int32 NumShapes = PVehicleActor->getNbShapes();
+		const int32 NumChassisShapes = NumShapes - WheelSetups.Num();
+		check(NumChassisShapes >= 1);
 
-		PWheelData.mDampingRate = M2ToCm2(Wheel->DampingRate);
-		PWheelData.mMass = Wheel->Mass;
-		PWheelData.mMOI = 0.5f * PWheelData.mMass * FMath::Square(PWheelData.mRadius);
-		
-		// init tire data
-		PxVehicleTireData PTireData; 
-		PTireData.mType = Wheel->TireType ? Wheel->TireType->GetTireTypeID() : GEngine->DefaultTireType->GetTireTypeID();
-		//PTireData.mCamberStiffnessPerUnitGravity = 0.0f;
-		PTireData.mLatStiffX = Wheel->LatStiffMaxLoad;
-		PTireData.mLatStiffY = Wheel->LatStiffValue;
-		PTireData.mLongitudinalStiffnessPerUnitGravity = Wheel->LongStiffValue;
-
-		// init suspension data
-		PxVehicleSuspensionData PSuspensionData;
-		PSuspensionData.mSprungMass = SprungMasses[WheelIdx];
-		PSuspensionData.mMaxCompression = Wheel->SuspensionMaxRaise;
-		PSuspensionData.mMaxDroop = Wheel->SuspensionMaxDrop;
-		PSuspensionData.mSpringStrength = FMath::Square( Wheel->SuspensionNaturalFrequency ) * PSuspensionData.mSprungMass;
-		PSuspensionData.mSpringDamperRate = Wheel->SuspensionDampingRatio * 2.0f * FMath::Sqrt( PSuspensionData.mSpringStrength * PSuspensionData.mSprungMass );
-
-		// init offsets
-		const PxVec3 PWheelOffset = WheelOffsets[WheelIdx];
-
-		PxVec3 PSuspTravelDirection = PxVec3( 0.0f, 0.0f, -1.0f );
-		PxVec3 PWheelCentreCMOffset = PWheelOffset - PLocalCOM;
-		PxVec3 PSuspForceAppCMOffset = PxVec3(PWheelCentreCMOffset.x, PWheelCentreCMOffset.y, Wheel->SuspensionForceOffset);
-		PxVec3 PTireForceAppCMOffset = PSuspForceAppCMOffset;
-
-		// finalize sim data
-		PWheelsSimData->setWheelData( WheelIdx, PWheelData );
-		PWheelsSimData->setTireData( WheelIdx, PTireData );
-		PWheelsSimData->setSuspensionData( WheelIdx, PSuspensionData );
-		PWheelsSimData->setSuspTravelDirection( WheelIdx, PSuspTravelDirection );
-		PWheelsSimData->setWheelCentreOffset( WheelIdx, PWheelCentreCMOffset );
-		PWheelsSimData->setSuspForceAppPointOffset( WheelIdx, PSuspForceAppCMOffset );
-		PWheelsSimData->setTireForceAppPointOffset( WheelIdx, PTireForceAppCMOffset );
-	}
-
-	const int32 NumShapes = PVehicleActor->getNbShapes();
-	const int32 NumChassisShapes = NumShapes - WheelSetups.Num();
-	check( NumChassisShapes >= 1 );
-
-	TArray<PxShape*> Shapes;
-	Shapes.AddZeroed(NumShapes);
-
-	PVehicleActor->getShapes( Shapes.GetData(), NumShapes );
+		TArray<PxShape*> Shapes;
+		Shapes.AddZeroed(NumShapes);
 
-	for ( int32 WheelIdx = 0; WheelIdx < WheelSetups.Num(); ++WheelIdx )
-	{
-		const int32 WheelShapeIndex = NumChassisShapes + WheelIdx;
+		PVehicleActor->getShapes(Shapes.GetData(), NumShapes);
 
-		PWheelsSimData->setWheelShapeMapping( WheelIdx, WheelShapeIndex );
-		PWheelsSimData->setSceneQueryFilterData( WheelIdx, Shapes[WheelShapeIndex]->getQueryFilterData());
-	}
+		for (int32 WheelIdx = 0; WheelIdx < WheelSetups.Num(); ++WheelIdx)
+		{
+			const int32 WheelShapeIndex = NumChassisShapes + WheelIdx;
 
-	// tire load filtering
-	PxVehicleTireLoadFilterData PTireLoadFilter;
-	PTireLoadFilter.mMinNormalisedLoad = MinNormalizedTireLoad;
-	PTireLoadFilter.mMinFilteredNormalisedLoad = MinNormalizedTireLoadFiltered;
-	PTireLoadFilter.mMaxNormalisedLoad = MaxNormalizedTireLoad;
-	PTireLoadFilter.mMaxFilteredNormalisedLoad = MaxNormalizedTireLoadFiltered;
-	PWheelsSimData->setTireLoadFilterData(PTireLoadFilter);
+			PWheelsSimData->setWheelShapeMapping(WheelIdx, WheelShapeIndex);
+			PWheelsSimData->setSceneQueryFilterData(WheelIdx, Shapes[WheelShapeIndex]->getQueryFilterData());
+		}
 
-	
+		// tire load filtering
+		PxVehicleTireLoadFilterData PTireLoadFilter;
+		PTireLoadFilter.mMinNormalisedLoad = MinNormalizedTireLoad;
+		PTireLoadFilter.mMinFilteredNormalisedLoad = MinNormalizedTireLoadFiltered;
+		PTireLoadFilter.mMaxNormalisedLoad = MaxNormalizedTireLoad;
+		PTireLoadFilter.mMaxFilteredNormalisedLoad = MaxNormalizedTireLoadFiltered;
+		PWheelsSimData->setTireLoadFilterData(PTireLoadFilter);
+	});
 }
  ////////////////////////////////////////////////////////////////////////////
 //Default tire force shader function.
@@ -549,19 +548,20 @@ FVector UWheeledVehicleMovementComponent::GetWheelRestingPosition( const FWheelS
 
 FVector UWheeledVehicleMovementComponent::GetLocalCOM() const
 {
+	FVector LocalCOM = FVector::ZeroVector;
 	if (UpdatedPrimitive)
 	{
 		if (const FBodyInstance* BodyInst = UpdatedPrimitive->GetBodyInstance())
 		{
-			if (PxRigidDynamic* PVehicleActor = BodyInst->GetPxRigidDynamic())
+			ExecuteOnPxRigidDynamicReadOnly(BodyInst, [&](const PxRigidDynamic* PVehicleActor)
 			{
 				PxTransform PCOMTransform = PVehicleActor->getCMassLocalPose();
-				return P2UVector(PCOMTransform.p);
-			}
+				LocalCOM = P2UVector(PCOMTransform.p);
+			});
 		}
 	}
 
-	return FVector::ZeroVector;
+	return LocalCOM;
 }
 
 USkinnedMeshComponent* UWheeledVehicleMovementComponent::GetMesh()
@@ -621,12 +621,13 @@ void UWheeledVehicleMovementComponent::CreatePhysicsState()
 
 			if ( PVehicle )
 			{
-				PhysScene->GetVehicleManager()->AddVehicle( this );
+				FPhysXVehicleManager* VehicleManager = PhysScene->GetVehicleManager();
+				VehicleManager->AddVehicle( this );
 
 				CreateWheels();
 
 				//LogVehicleSettings( PVehicle );
-
+				SCOPED_SCENE_WRITE_LOCK(VehicleManager->GetScene());
 				PVehicle->getRigidDynamicActor()->wakeUp();
 			}
 		}
@@ -1040,14 +1041,18 @@ void UWheeledVehicleMovementComponent::SetUseAutoGears(bool bUseAuto)
 
 float UWheeledVehicleMovementComponent::GetForwardSpeed() const
 {
+	float ForwardSpeed = 0.f;
 #if WITH_VEHICLE
 	if ( PVehicle )
 	{
-		return PVehicle->computeForwardSpeed();
+		UpdatedPrimitive->GetBodyInstance()->ExecuteOnPhysicsReadOnly([&]
+		{
+			ForwardSpeed = PVehicle->computeForwardSpeed();
+		});
 	}
 #endif
 
-	return 0.0f;
+	return ForwardSpeed;
 }
 
 float UWheeledVehicleMovementComponent::GetEngineRotationSpeed() const
@@ -1205,7 +1210,9 @@ bool UWheeledVehicleMovementComponent::CheckSlipThreshold(float AbsLongSlipThres
 	}
 
 	FPhysXVehicleManager* MyVehicleManager = World->GetPhysicsScene()->GetVehicleManager();
-	PxWheelQueryResult * WheelsStates = MyVehicleManager->GetWheelsStates(this);
+	SCOPED_SCENE_READ_LOCK(MyVehicleManager->GetScene());
+
+	PxWheelQueryResult * WheelsStates = MyVehicleManager->GetWheelsStates_AssumesLocked(this);
 	check(WheelsStates);
 
 	PxReal MaxLongSlip = 0.f;
@@ -1239,7 +1246,9 @@ float UWheeledVehicleMovementComponent::GetMaxSpringForce() const
 	}
 
 	FPhysXVehicleManager* MyVehicleManager = World->GetPhysicsScene()->GetVehicleManager();
-	PxWheelQueryResult * WheelsStates = MyVehicleManager->GetWheelsStates(this);
+	SCOPED_SCENE_READ_LOCK(MyVehicleManager->GetScene());
+
+	PxWheelQueryResult * WheelsStates = MyVehicleManager->GetWheelsStates_AssumesLocked(this);
 	check(WheelsStates);
 
 	PxReal MaxSpringCompression = 0.f;
@@ -1280,7 +1289,8 @@ void UWheeledVehicleMovementComponent::DrawDebug( UCanvas* Canvas, float& YL, fl
 
 	}
 
-	PxWheelQueryResult* WheelsStates = MyVehicleManager->GetWheelsStates(this);
+	SCOPED_SCENE_READ_LOCK(MyVehicleManager->GetScene());
+	PxWheelQueryResult* WheelsStates = MyVehicleManager->GetWheelsStates_AssumesLocked(this);
 	check(WheelsStates);
 
 	// draw wheel data
@@ -1321,7 +1331,7 @@ void UWheeledVehicleMovementComponent::DrawDebug( UCanvas* Canvas, float& YL, fl
 	}
 
 	// draw wheel graphs
-	PxVehicleTelemetryData* TelemetryData = MyVehicleManager->GetTelemetryData();
+	PxVehicleTelemetryData* TelemetryData = MyVehicleManager->GetTelemetryData_AssumesLocked();
 
 	if ( TelemetryData )
 	{
@@ -1436,9 +1446,10 @@ void UWheeledVehicleMovementComponent::DrawDebugLines()
 	DrawDebugLine(World, P2UVector(T.p), P2UVector(T.p + T.rotate(PxVec3(0, ChassisSize, 0))), FColor::Green);
 	DrawDebugLine(World, P2UVector(T.p), P2UVector(T.p + T.rotate(PxVec3(0, 0, ChassisSize))), FColor::Blue);
 
-	PxVehicleTelemetryData* TelemetryData = MyVehicleManager->GetTelemetryData();
+	SCOPED_SCENE_READ_LOCK(MyVehicleManager->GetScene());
+	PxVehicleTelemetryData* TelemetryData = MyVehicleManager->GetTelemetryData_AssumesLocked();
 	
-	PxWheelQueryResult* WheelsStates = MyVehicleManager->GetWheelsStates(this);
+	PxWheelQueryResult* WheelsStates = MyVehicleManager->GetWheelsStates_AssumesLocked(this);
 	check(WheelsStates);
 
 	for ( uint32 w = 0; w < PNumWheels; ++w )
diff --git a/Engine/Source/Runtime/Engine/Private/Vehicles/WheeledVehicleMovementComponent4W.cpp b/Engine/Source/Runtime/Engine/Private/Vehicles/WheeledVehicleMovementComponent4W.cpp
--- a/Engine/Source/Runtime/Engine/Private/Vehicles/WheeledVehicleMovementComponent4W.cpp
+++ b/Engine/Source/Runtime/Engine/Private/Vehicles/WheeledVehicleMovementComponent4W.cpp
@@ -283,11 +283,15 @@ void UWheeledVehicleMovementComponent4W::SetupVehicle()
 	PxVehicleDrive4W* PVehicleDrive4W = PxVehicleDrive4W::allocate(4);
 	check(PVehicleDrive4W);
 
-	PVehicleDrive4W->setup( GPhysXSDK, UpdatedPrimitive->GetBodyInstance()->GetPxRigidDynamic(), *PWheelsSimData, DriveData, 0);
-	PVehicleDrive4W->setToRestState();
+	ExecuteOnPxRigidDynamicReadWrite(UpdatedPrimitive->GetBodyInstance(), [&] (PxRigidDynamic* PRigidDynamic)
+	{
+		PVehicleDrive4W->setup( GPhysXSDK, PRigidDynamic, *PWheelsSimData, DriveData, 0);
+		PVehicleDrive4W->setToRestState();
 
-	// cleanup
-	PWheelsSimData->free();
+		// cleanup
+		PWheelsSimData->free();
+	});
+	
 	PWheelsSimData = NULL;
 
 	// cache values
@@ -302,35 +306,38 @@ void UWheeledVehicleMovementComponent4W::UpdateSimulation(float DeltaTime)
 	if (PVehicleDrive == NULL)
 		return;
 
-	PxVehicleDrive4WRawInputData RawInputData;
-	RawInputData.setAnalogAccel(ThrottleInput);
-	RawInputData.setAnalogSteer(SteeringInput);
-	RawInputData.setAnalogBrake(BrakeInput);
-	RawInputData.setAnalogHandbrake(HandbrakeInput);
-	
-	if (!PVehicleDrive->mDriveDynData.getUseAutoGears())
+	UpdatedPrimitive->GetBodyInstance()->ExecuteOnPhysicsReadWrite([&]
 	{
-		RawInputData.setGearUp(bRawGearUpInput);
-		RawInputData.setGearDown(bRawGearDownInput);
-	}
+		PxVehicleDrive4WRawInputData RawInputData;
+		RawInputData.setAnalogAccel(ThrottleInput);
+		RawInputData.setAnalogSteer(SteeringInput);
+		RawInputData.setAnalogBrake(BrakeInput);
+		RawInputData.setAnalogHandbrake(HandbrakeInput);
 
-	// Convert from our curve to PxFixedSizeLookupTable
-	PxFixedSizeLookupTable<8> SpeedSteerLookup;
-	TArray<FRichCurveKey> SteerKeys = SteeringCurve.GetRichCurve()->GetCopyOfKeys();
-	const int32 MaxSteeringSamples = FMath::Min(8, SteerKeys.Num());
-	for (int32 KeyIdx = 0; KeyIdx < MaxSteeringSamples; KeyIdx++)
-	{
-		FRichCurveKey& Key = SteerKeys[KeyIdx];
-		SpeedSteerLookup.addPair(KmHToCmS(Key.Time), FMath::Clamp(Key.Value, 0.f, 1.f));
-	}
+		if (!PVehicleDrive->mDriveDynData.getUseAutoGears())
+		{
+			RawInputData.setGearUp(bRawGearUpInput);
+			RawInputData.setGearDown(bRawGearDownInput);
+		}
+
+		// Convert from our curve to PxFixedSizeLookupTable
+		PxFixedSizeLookupTable<8> SpeedSteerLookup;
+		TArray<FRichCurveKey> SteerKeys = SteeringCurve.GetRichCurve()->GetCopyOfKeys();
+		const int32 MaxSteeringSamples = FMath::Min(8, SteerKeys.Num());
+		for (int32 KeyIdx = 0; KeyIdx < MaxSteeringSamples; KeyIdx++)
+		{
+			FRichCurveKey& Key = SteerKeys[KeyIdx];
+			SpeedSteerLookup.addPair(KmHToCmS(Key.Time), FMath::Clamp(Key.Value, 0.f, 1.f));
+		}
 
-	PxVehiclePadSmoothingData SmoothData = {
-		{ ThrottleInputRate.RiseRate, BrakeInputRate.RiseRate, HandbrakeInputRate.RiseRate, SteeringInputRate.RiseRate, SteeringInputRate.RiseRate },
-		{ ThrottleInputRate.FallRate, BrakeInputRate.FallRate, HandbrakeInputRate.FallRate, SteeringInputRate.FallRate, SteeringInputRate.FallRate }
-	};
+		PxVehiclePadSmoothingData SmoothData = {
+			{ ThrottleInputRate.RiseRate, BrakeInputRate.RiseRate, HandbrakeInputRate.RiseRate, SteeringInputRate.RiseRate, SteeringInputRate.RiseRate },
+			{ ThrottleInputRate.FallRate, BrakeInputRate.FallRate, HandbrakeInputRate.FallRate, SteeringInputRate.FallRate, SteeringInputRate.FallRate }
+		};
 
-	PxVehicleDrive4W* PVehicleDrive4W = (PxVehicleDrive4W*)PVehicleDrive;
-	PxVehicleDrive4WSmoothAnalogRawInputsAndSetAnalogInputs(SmoothData, SpeedSteerLookup, RawInputData, DeltaTime, false, *PVehicleDrive4W);
+		PxVehicleDrive4W* PVehicleDrive4W = (PxVehicleDrive4W*)PVehicleDrive;
+		PxVehicleDrive4WSmoothAnalogRawInputsAndSetAnalogInputs(SmoothData, SpeedSteerLookup, RawInputData, DeltaTime, false, *PVehicleDrive4W);
+	});
 }
 
 #endif // WITH_VEHICLE
diff --git a/Engine/Source/Runtime/Engine/Public/PhysicsPublic.h b/Engine/Source/Runtime/Engine/Public/PhysicsPublic.h
--- a/Engine/Source/Runtime/Engine/Public/PhysicsPublic.h
+++ b/Engine/Source/Runtime/Engine/Public/PhysicsPublic.h
@@ -432,10 +432,10 @@ public:
 	void ProcessPhysScene(uint32 SceneType);
 
 	/** Sync components in the scene to physics bodies that changed */
-	void SyncComponentsToBodies(uint32 SceneType);
+	void SyncComponentsToBodies_AssumesLocked(uint32 SceneType);
 
 	/** Call after WaitPhysScene on the synchronous scene to make deferred OnRigidBodyCollision calls.  */
-	void DispatchPhysNotifications();
+	void DispatchPhysNotifications_AssumesLocked();
 
 	/** Add any debug lines from the physics scene of the given type to the supplied line batcher. */
 	ENGINE_API void AddDebugLines(uint32 SceneType, class ULineBatchComponent* LineBatcherToUse);
@@ -462,29 +462,85 @@ public:
 	ENGINE_API bool HasAsyncScene() const { return bAsyncSceneEnabled; }
 
 	/** Lets the scene update anything related to this BodyInstance as it's now being terminated */
-	void TermBody(FBodyInstance* BodyInstance);
+	DEPRECATED(4.8, "Please call AddCustomPhysics_AssumesLocked and make sure you obtain the appropriate PhysX scene locks")
+	void TermBody(FBodyInstance* BodyInstance)
+	{
+		TermBody_AssumesLocked(BodyInstance);
+	}
+
+	/** Lets the scene update anything related to this BodyInstance as it's now being terminated */
+	void TermBody_AssumesLocked(FBodyInstance* BodyInstance);
+
+	/** Add a custom callback for next step that will be called on every substep */
+	DEPRECATED(4.8, "Please call AddCustomPhysics_AssumesLocked and make sure you obtain the appropriate PhysX scene locks")
+	void AddCustomPhysics(FBodyInstance* BodyInstance, FCalculateCustomPhysics& CalculateCustomPhysics)
+	{
+		AddCustomPhysics_AssumesLocked(BodyInstance, CalculateCustomPhysics);
+	}
 
 	/** Add a custom callback for next step that will be called on every substep */
-	void AddCustomPhysics(FBodyInstance* BodyInstance, FCalculateCustomPhysics& CalculateCustomPhysics);
+	void AddCustomPhysics_AssumesLocked(FBodyInstance* BodyInstance, FCalculateCustomPhysics& CalculateCustomPhysics);
 
 	/** Adds a force to a body - We need to go through scene to support substepping */
-	void AddForce(FBodyInstance* BodyInstance, const FVector& Force, bool bAllowSubstepping, bool bAccelChange);
+	DEPRECATED(4.8, "Please call AddForce_AssumesLocked and make sure you obtain the appropriate PhysX scene locks")
+	void AddForce(FBodyInstance* BodyInstance, const FVector& Force, bool bAllowSubstepping, bool bAccelChange)
+	{
+		AddForce_AssumesLocked(BodyInstance, Force, bAllowSubstepping, bAccelChange);
+	}
+
+	void AddForce_AssumesLocked(FBodyInstance* BodyInstance, const FVector& Force, bool bAllowSubstepping, bool bAccelChange);
+
+	/** Adds a force to a body at a specific position - We need to go through scene to support substepping */
+	DEPRECATED(4.8, "Please call AddForceAtPosition_AssumesLocked and make sure you obtain the appropriate PhysX scene locks")
+	void AddForceAtPosition(FBodyInstance* BodyInstance, const FVector& Force, const FVector& Position, bool bAllowSubstepping)
+	{
+		AddForceAtPosition_AssumesLocked(BodyInstance, Force, Position, bAllowSubstepping);
+	}
 
 	/** Adds a force to a body at a specific position - We need to go through scene to support substepping */
-	void AddForceAtPosition(FBodyInstance* BodyInstance, const FVector& Force, const FVector& Position, bool bAllowSubstepping);
+	void AddForceAtPosition_AssumesLocked(FBodyInstance* BodyInstance, const FVector& Force, const FVector& Position, bool bAllowSubstepping);
+
+	/** Adds a radial force to a body - We need to go through scene to support substepping */
+	DEPRECATED(4.8, "Please call AddRadialForceToBody_AssumesLocked and make sure you obtain the appropriate PhysX scene locks")
+	void AddRadialForceToBody(FBodyInstance* BodyInstance, const FVector& Origin, const float Radius, const float Strength, const uint8 Falloff, bool bAccelChange, bool bAllowSubstepping)
+	{
+		AddRadialForceToBody_AssumesLocked(BodyInstance, Origin, Radius, Strength, Falloff, bAccelChange, bAllowSubstepping);
+	}
 
 	/** Adds a radial force to a body - We need to go through scene to support substepping */
-	void AddRadialForceToBody(FBodyInstance* BodyInstance, const FVector& Origin, const float Radius, const float Strength, const uint8 Falloff, bool bAccelChange, bool bAllowSubstepping);
+	void AddRadialForceToBody_AssumesLocked(FBodyInstance* BodyInstance, const FVector& Origin, const float Radius, const float Strength, const uint8 Falloff, bool bAccelChange, bool bAllowSubstepping);
+
+	/** Adds torque to a body - We need to go through scene to support substepping */
+	DEPRECATED(4.8, "Please call AddTorque_AssumesLocked and make sure you obtain the appropriate PhysX scene locks")
+	void AddTorque(FBodyInstance* BodyInstance, const FVector& Torque, bool bAllowSubstepping, bool bAccelChange)
+	{
+		AddTorque_AssumesLocked(BodyInstance, Torque, bAllowSubstepping, bAccelChange);
+	}
 
 	/** Adds torque to a body - We need to go through scene to support substepping */
-	void AddTorque(FBodyInstance* BodyInstance, const FVector& Torque, bool bAllowSubstepping, bool bAccelChange);
+	void AddTorque_AssumesLocked(FBodyInstance* BodyInstance, const FVector& Torque, bool bAllowSubstepping, bool bAccelChange);
 
 	/** Sets a Kinematic actor's target position - We need to do this here to support substepping*/
-	void SetKinematicTarget(FBodyInstance* BodyInstance, const FTransform& TargetTM, bool bAllowSubstepping);
+	DEPRECATED(4.8, "Please call SetKinematicTarget_AssumesLocked and make sure you obtain the appropriate PhysX scene locks")
+	void SetKinematicTarget(FBodyInstance* BodyInstance, const FTransform& TargetTM, bool bAllowSubstepping)
+	{
+		SetKinematicTarget_AssumesLocked(BodyInstance, TargetTM, bAllowSubstepping);
+	}
+	
+	/** Sets a Kinematic actor's target position - We need to do this here to support substepping*/
+	void SetKinematicTarget_AssumesLocked(FBodyInstance* BodyInstance, const FTransform& TargetTM, bool bAllowSubstepping);
+
+	/** Gets a Kinematic actor's target position - We need to do this here to support substepping
+	  * Returns true if kinematic target has been set. If false the OutTM is invalid */
+	DEPRECATED(4.8, "Please call GetKinematicTarget_AssumesLocked and make sure you obtain the appropriate PhysX scene locks")
+	bool GetKinematicTarget(const FBodyInstance* BodyInstance, FTransform& OutTM) const
+	{
+		return GetKinematicTarget_AssumesLocked(BodyInstance, OutTM);
+	}
 
 	/** Gets a Kinematic actor's target position - We need to do this here to support substepping
 	  * Returns true if kinematic target has been set. If false the OutTM is invalid */
-	bool GetKinematicTarget(const FBodyInstance* BodyInstance, FTransform& OutTM) const;
+	bool GetKinematicTarget_AssumesLocked(const FBodyInstance* BodyInstance, FTransform& OutTM) const;
 
 	/** Gets the collision disable table */
 	const TMap<uint32, TMap<struct FRigidBodyIndexPair, bool> *> & GetCollisionDisableTableLookup()
diff --git a/Engine/Source/Runtime/Landscape/Private/LandscapeCollision.cpp b/Engine/Source/Runtime/Landscape/Private/LandscapeCollision.cpp
--- a/Engine/Source/Runtime/Landscape/Private/LandscapeCollision.cpp
+++ b/Engine/Source/Runtime/Landscape/Private/LandscapeCollision.cpp
@@ -228,7 +228,9 @@ void ULandscapeHeightfieldCollisionComponent::CreatePhysicsState()
 				}
 
 				// Add to scenes
-				PhysScene->GetPhysXScene(PST_Sync)->addActor(*HeightFieldActorSync);
+				PxScene* SyncScene = PhysScene->GetPhysXScene(PST_Sync);
+				SCOPED_SCENE_WRITE_LOCK(SyncScene);
+				SyncScene->addActor(*HeightFieldActorSync);
 
 				if (PhysScene->HasAsyncScene())
 				{
