Refactored CameraShake code for extensibility, simplicity, and to add a Blueprint API.
- Renamed several functions for clarity.
- Shake instance is now part of the CameraShake object.
- FMinimalViewInfo now exposed to BPs.

diff --git a/Engine/Source/Runtime/Engine/Classes/Camera/CameraModifier_CameraShake.h b/Engine/Source/Runtime/Engine/Classes/Camera/CameraModifier_CameraShake.h
--- a/Engine/Source/Runtime/Engine/Classes/Camera/CameraModifier_CameraShake.h
+++ b/Engine/Source/Runtime/Engine/Classes/Camera/CameraModifier_CameraShake.h
@@ -9,89 +9,6 @@
 #include "Camera/CameraModifier.h"
 #include "CameraModifier_CameraShake.generated.h"
 
-/** 
- * A CameraShakeInstance is an active instance of a shake asset. It contains
- * the necessary state information to smoothly update this shake over time.
- */
-USTRUCT()
-struct ENGINE_API FCameraShakeInstance
-{
-	GENERATED_USTRUCT_BODY()
-
-	/** The source camera shake data. */
-	UPROPERTY()
-	const class UCameraShake* SourceShake;
-
-	/** Optional name of shake. Used to identify shakes when single instances are desired. */
-	UPROPERTY()
-	FName SourceShakeName;
-
-	/** Time remaining for oscillation shakes. Less than 0.f means shake infinitely. */
-	UPROPERTY()
-	float OscillatorTimeRemaining;
-
-	/** True if this shake is currently blending in. */
-	UPROPERTY()
-	uint32 bBlendingIn:1;
-
-	/** How long this instance has been blending in. */
-	UPROPERTY()
-	float CurrentBlendInTime;
-
-	/** True if this shake is currently blending out. */
-	UPROPERTY()
-	uint32 bBlendingOut:1;
-
-	/** How long this instance has been blending out. */
-	UPROPERTY()
-	float CurrentBlendOutTime;
-
-	/** Current location sinusoidal offset. */
-	UPROPERTY()
-	FVector LocSinOffset;
-
-	/** Current rotational sinusoidal offset. */
-	UPROPERTY()
-	FVector RotSinOffset;
-
-	/** Current FOV sinusoidal offset. */
-	UPROPERTY()
-	float FOVSinOffset;
-
-	/** Overall intensity scale for this shake instance. */
-	UPROPERTY()
-	float Scale;
-
-	/** The playing instance of the CameraAnim-based shake, if any. */
-	UPROPERTY()
-	class UCameraAnimInst* AnimInst;
-
-	/** What space to play the shake in before applying to the camera.  Affects both Anim and Oscillation shakes. */
-	UPROPERTY()
-	TEnumAsByte<ECameraAnimPlaySpace::Type> PlaySpace;
-
-	/** Matrix defining the playspace, used when PlaySpace == CAPS_UserDefined */
-	UPROPERTY()
-	FMatrix UserPlaySpaceMatrix;
-
-	FCameraShakeInstance()
-		: SourceShake(NULL)
-		, OscillatorTimeRemaining(0)
-		, bBlendingIn(false)
-		, CurrentBlendInTime(0)
-		, bBlendingOut(false)
-		, CurrentBlendOutTime(0)
-		, LocSinOffset(ForceInit)
-		, RotSinOffset(ForceInit)
-		, FOVSinOffset(0)
-		, Scale(0)
-		, AnimInst(NULL)
-		, PlaySpace(0)
-		, UserPlaySpaceMatrix(ForceInit)
-	{}
-
-};
-
 //=============================================================================
 /**
  * A UCameraModifier_CameraShake is a camera modifier that can apply a UCameraShake to 
@@ -103,39 +20,10 @@ class ENGINE_API UCameraModifier_CameraShake : public UCameraModifier
 	GENERATED_UCLASS_BODY()
 
 public:
-	/** List of active CameraShakes */
+	/** List of active CameraShake instances */
 	UPROPERTY()
-	TArray<struct FCameraShakeInstance> ActiveShakes;
-
-protected:
-	/** Scaling factor applied to all camera shakes in when in splitscreen mode. Normally used to dampen, since shakes feel more intense in a smaller viewport. */
-	UPROPERTY(EditAnywhere, Category=CameraModifier_CameraShake)
-	float SplitScreenShakeScale;
+	TArray<class UCameraShake*> ActiveShakes;
 
-	/** @return Returns the current overall scaling factor for this shake. */
-	virtual float GetShakeScale(FCameraShakeInstance const& ShakeInst) const;
-	
-	/** 
-	 * Restarts a playing camerashake from the beginning. 
-	 * @param ShakeInst - The shake instance to restart.
-	 * @param Scale - The base scale to restart at.
-	 */
-	virtual void ReinitShake(FCameraShakeInstance& ShakeInst, float Scale);
-
-	/** @return Returns an initial offset for the given oscillator parameter. */
-	float InitializeOffset( const struct FFOscillator& Param ) const;
-	
-	/** 
-	 * Internal. Creates a new camera shake instance.
-	 * @param NewShake - The class of camera shake to instantiate.
-	 * @param Scale - The scalar intensity to play the shake.
-	 * @param PlaySpace - Which coordinate system to play the shake in.
-	 * @param UserPlaySpaceRot - Coordinate system to play shake when PlaySpace == CAPS_UserDefined.
-	 * @return Returns the new camera shake instance.
-	 */
-	virtual FCameraShakeInstance InitializeShake(TSubclassOf<class UCameraShake> NewShake, float Scale, ECameraAnimPlaySpace::Type PlaySpace, FRotator UserPlaySpaceRot = FRotator::ZeroRotator);
-	
-public:
 	/** 
 	 * Adds a new active screen shake to be applied. 
 	 * @param NewShake - The class of camera shake to instantiate.
@@ -143,26 +31,26 @@ public:
 	 * @param PlaySpace - Which coordinate system to play the shake in.
 	 * @param UserPlaySpaceRot - Coordinate system to play shake when PlaySpace == CAPS_UserDefined.
 	 */
-	virtual void AddCameraShake(TSubclassOf<class UCameraShake> NewShake, float Scale, ECameraAnimPlaySpace::Type PlaySpace=ECameraAnimPlaySpace::CameraLocal, FRotator UserPlaySpaceRot = FRotator::ZeroRotator);
+	virtual class UCameraShake* AddCameraShake(TSubclassOf<class UCameraShake> NewShake, float Scale, ECameraAnimPlaySpace::Type PlaySpace=ECameraAnimPlaySpace::CameraLocal, FRotator UserPlaySpaceRot = FRotator::ZeroRotator);
 	
 	/**
 	 * Stops and removes the camera shake of the given class from the camera.
 	 * @param Shake - the camera shake class to remove.
 	 */
-	virtual void RemoveCameraShake(TSubclassOf<class UCameraShake> Shake);
+	virtual void RemoveCameraShake(UCameraShake* ShakeInst);
+
+	/** Stops and removes all camera shakes of the given class from the camera. */
+	virtual void RemoveAllCameraShakesOfClass(TSubclassOf<class UCameraShake> ShakeClass);
 
 	/** Stops and removes all camera shakes from the camera. */
 	virtual void RemoveAllCameraShakes();
 	
-	/** 
-	 * Called per-tick to update a CameraShake.
-	 * @param DeltaTime - Simulation timeslice, in seconds.
-	 * @param ShakeInst - The shake instance to update.
-	 * @param InOutPOV - the POV data to update with the results of the shake.
-	 */
-	virtual void UpdateCameraShake(float DeltaTime, FCameraShakeInstance& ShakeInst, struct FMinimalViewInfo& InOutPOV);
-	
 	// Begin UCameraModifer Interface
 	virtual bool ModifyCamera(APlayerCameraManager* Camera, float DeltaTime, struct FMinimalViewInfo& InOutPOV) override;
 	// End UCameraModifer Interface
+
+protected:
+	/** Scaling factor applied to all camera shakes in when in splitscreen mode. Normally used to reduce shaking, since shakes feel more intense in a smaller viewport. */
+	UPROPERTY(EditAnywhere, Category = CameraModifier_CameraShake)
+	float SplitScreenShakeScale;
 };
diff --git a/Engine/Source/Runtime/Engine/Classes/Camera/CameraShake.h b/Engine/Source/Runtime/Engine/Classes/Camera/CameraShake.h
--- a/Engine/Source/Runtime/Engine/Classes/Camera/CameraShake.h
+++ b/Engine/Source/Runtime/Engine/Classes/Camera/CameraShake.h
@@ -42,6 +42,12 @@ struct FFOscillator
 		, Frequency(0)
 		, InitialOffset(0)
 	{}
+
+	/** Advances the oscillation time and returns the current value. */
+	static float UpdateOffset(FFOscillator const& Osc, float& CurrentOffset, float DeltaTime);
+
+	/** Returns the initial value of the oscillator. */
+	static float GetInitialOffset(FFOscillator const& Osc);
 };
 
 /** Defines FRotator oscillation. */
@@ -106,7 +112,7 @@ class ENGINE_API UCameraShake : public UObject
 	GENERATED_UCLASS_BODY()
 
 	/** 
-	 *  If true to only allow a single instance of this shake to play at any given time. 
+	 *  If true to only allow a single instance of this shake class to play at any given time.
 	 *  Subsequent attempts to play this shake will simply restart the timer.
 	 */
 	UPROPERTY(EditAnywhere, Category=CameraShake)
@@ -121,7 +127,7 @@ class ENGINE_API UCameraShake : public UObject
 	float OscillationBlendInTime;
 
 	/** Duration of the blend-out, where the oscillation scales from 1 to 0. */
-	UPROPERTY(EditAnywhere, Category = Oscillation, meta = (ClampMin = "0.0"))
+	UPROPERTY(EditAnywhere, Category=Oscillation, meta = (ClampMin = "0.0"))
 	float OscillationBlendOutTime;
 
 	/** Rotational oscillation */
@@ -170,6 +176,79 @@ class ENGINE_API UCameraShake : public UObject
 	/** When bRandomAnimSegment is true, this defines how long the anim should play. */
 	UPROPERTY(EditAnywhere, Category=AnimShake, meta=(ClampMin = "0.0", editcondition = "bRandomAnimSegment"))
 	float RandomAnimSegmentDuration;
+
+
+protected:
+
+	// INSTANCE DATA
+	
+	UPROPERTY(transient, BlueprintReadOnly, Category=CameraShake)
+	class APlayerCameraManager* CameraOwner;
+
+	/** Time remaining for oscillation shakes. Less than 0.f means shake infinitely. */
+	float OscillatorTimeRemaining;
+
+	/** True if this shake is currently blending in. */
+	uint32 bBlendingIn:1;
+
+	/** How long this instance has been blending in. */
+	float CurrentBlendInTime;
+
+	/** True if this shake is currently blending out. */
+	uint32 bBlendingOut:1;
+
+	/** How long this instance has been blending out. */
+	float CurrentBlendOutTime;
+
+	/** Current location sinusoidal offset. */
+	FVector LocSinOffset;
+
+	/** Current rotational sinusoidal offset. */
+	FVector RotSinOffset;
+
+	/** Current FOV sinusoidal offset. */
+	float FOVSinOffset;
+
+	/** Overall intensity scale for this shake instance. */
+	float ShakeScale;
+
+	/** The playing instance of the CameraAnim-based shake, if any. */
+	UPROPERTY()
+	class UCameraAnimInst* AnimInst;
+
+	/** What space to play the shake in before applying to the camera.  Affects both Anim and Oscillation shakes. */
+	ECameraAnimPlaySpace::Type PlaySpace;
+
+	/** Matrix defining the playspace, used when PlaySpace == CAPS_UserDefined */
+	FMatrix UserPlaySpaceMatrix;
+
+public:
+
+	// Blueprint API
+
+	/** Called every tick to let the shake modify the point of view */
+	UFUNCTION(BlueprintImplementableEvent, Category=CameraShake)
+	void BlueprintUpdateCameraShake(float DeltaTime, float Alpha, const FMinimalViewInfo& POV, FMinimalViewInfo& ModifiedPOV);
+
+	/** Called when the shake starts playing */
+	UFUNCTION(BlueprintImplementableEvent, Category = CameraShake)
+	void ReceivePlayShake(float Scale);
+
+	/** Called to allow a shake to decide when it's finished playing. */
+	UFUNCTION(BlueprintNativeEvent, Category = CameraShake)
+	bool ReceiveIsFinished() const;
+
+	/** Called when the shake is explicitly stopped. */
+	UFUNCTION(BlueprintImplementableEvent, Category = CameraShake)
+	void ReceiveStopShake();
+
+	// Native API
+
+	virtual void UpdateAndApplyCameraShake(float DeltaTime, float Alpha, FMinimalViewInfo& InOutPOV);
+	virtual void PlayShake(class APlayerCameraManager* Camera, float Scale, ECameraAnimPlaySpace::Type PlaySpace, FRotator UserPlaySpaceRot = FRotator::ZeroRotator);
+	virtual bool IsFinished() const;
+	virtual void StopShake();
+
 };
 
 
diff --git a/Engine/Source/Runtime/Engine/Classes/Camera/CameraTypes.h b/Engine/Source/Runtime/Engine/Classes/Camera/CameraTypes.h
--- a/Engine/Source/Runtime/Engine/Classes/Camera/CameraTypes.h
+++ b/Engine/Source/Runtime/Engine/Classes/Camera/CameraTypes.h
@@ -31,7 +31,7 @@ namespace ECameraAnimPlaySpace
  	};
 }
 
-USTRUCT()
+USTRUCT(BlueprintType)
 struct FMinimalViewInfo
 {
 	GENERATED_USTRUCT_BODY()
@@ -77,11 +77,11 @@ struct FMinimalViewInfo
 	TEnumAsByte<ECameraProjectionMode::Type> ProjectionMode;
 
 	/** Indicates if PostProcessSettings should be applied. */
-	UPROPERTY()
+	UPROPERTY(BlueprintReadWrite, Category = Camera)
 	float PostProcessBlendWeight;
 
 	/** Post-process settings to use if PostProcessBlendWeight is non-zero. */
-	UPROPERTY()
+	UPROPERTY(BlueprintReadWrite, Category = Camera)
 	struct FPostProcessSettings PostProcessSettings;
 
 	FMinimalViewInfo()
diff --git a/Engine/Source/Runtime/Engine/Classes/Camera/PlayerCameraManager.h b/Engine/Source/Runtime/Engine/Classes/Camera/PlayerCameraManager.h
--- a/Engine/Source/Runtime/Engine/Classes/Camera/PlayerCameraManager.h
+++ b/Engine/Source/Runtime/Engine/Classes/Camera/PlayerCameraManager.h
@@ -637,17 +637,24 @@ public:
 	/** 
 	 * Plays a camera shake on this camera.
 	 * @param Shake - The class of camera shake to play.
-	 * @param Scale - Scalar defining how "intense" to play the anim. 1.0 is normal or as authored.
-	 * @param PlaySpace - Which coordinate system to play the shake in.
+	 * @param Scale - Scalar defining how "intense" to play the shake. 1.0 is normal (as authored).
+	 * @param PlaySpace - Which coordinate system to play the shake in (affects oscillations and camera anims)
 	 * @param UserPlaySpaceRot - Coordinate system to play shake when PlaySpace == CAPS_UserDefined.
 	 */
-	virtual void PlayCameraShake(TSubclassOf<class UCameraShake> Shake, float Scale, enum ECameraAnimPlaySpace::Type PlaySpace=ECameraAnimPlaySpace::CameraLocal, FRotator UserPlaySpaceRot = FRotator::ZeroRotator);
+	UFUNCTION(BlueprintCallable, Category = "Camera Shakes")
+	virtual class UCameraShake* PlayCameraShake(TSubclassOf<class UCameraShake> ShakeClass, float Scale=1.f, enum ECameraAnimPlaySpace::Type PlaySpace = ECameraAnimPlaySpace::CameraLocal, FRotator UserPlaySpaceRot = FRotator::ZeroRotator);
 	
+	/** Immediately stops the given shake instance and invalidates it. */
+	UFUNCTION(BlueprintCallable, Category = "Camera Shakes")
+	virtual void StopCameraShake(class UCameraShake* ShakeInstance);
+
 	/** Stops playing CameraShake of the given class. */
-	virtual void StopCameraShake(TSubclassOf<class UCameraShake> Shake);
+	UFUNCTION(BlueprintCallable, Category = "Camera Shakes")
+	virtual void StopAllInstancesOfCameraShake(TSubclassOf<class UCameraShake> Shake);
 
-	/** Stops playing all camera shakes. */
-	virtual void ClearAllCameraShakes();
+	/** Stops all active camera shakes on this camera. */
+	UFUNCTION(BlueprintCallable, Category = "Camera Shakes")
+	virtual void StopAllCameraShakes();
 
 	//
 	//  CameraAnim fades.
@@ -662,20 +669,20 @@ public:
 	 * @param bShouldFadeAudio - True to fade audio volume along with the alpha of the solid color.
 	 * @param bHoldWhenFinished - True for fade to hold at the ToAlpha until explicitly stopped (e.g. with StopCameraFade)
 	 */
-	UFUNCTION(BlueprintCallable, Category = "Camera")
+	UFUNCTION(BlueprintCallable, Category = "Camera Fades")
 	virtual void StartCameraFade(float FromAlpha, float ToAlpha, float Duration, FLinearColor Color, bool bShouldFadeAudio = false, bool bHoldWhenFinished = false);
 
 	/** 
 	 * Stops camera fading.
 	 */
-	UFUNCTION(BlueprintCallable, Category = "Camera")
+	UFUNCTION(BlueprintCallable, Category = "Camera Fades")
 	virtual void StopCameraFade();
 
 	/** 
 	 * Turns on camera fading at the given opacity. Does not auto-animate, allowing user to animate themselves.
 	 * Call StopCameraFade to turn fading back off.
 	*/
-	UFUNCTION(BlueprintCallable, Category = "Camera")
+	UFUNCTION(BlueprintCallable, Category = "Camera Fades")
 	virtual void SetManualCameraFade(float InFadeAmount, FLinearColor Color, bool bInFadeAudio);
 
 
diff --git a/Engine/Source/Runtime/Engine/Private/Camera/CameraModifier_CameraShake.cpp b/Engine/Source/Runtime/Engine/Private/Camera/CameraModifier_CameraShake.cpp
--- a/Engine/Source/Runtime/Engine/Private/Camera/CameraModifier_CameraShake.cpp
+++ b/Engine/Source/Runtime/Engine/Private/Camera/CameraModifier_CameraShake.cpp
@@ -1,7 +1,6 @@
 // Copyright 1998-2015 Epic Games, Inc. All Rights Reserved.
 
 #include "EnginePrivate.h"
-#include "Camera/CameraAnimInst.h"
 #include "Camera/CameraModifier_CameraShake.h"
 
 //////////////////////////////////////////////////////////////////////////
@@ -13,162 +12,6 @@ UCameraModifier_CameraShake::UCameraModifier_CameraShake(const FObjectInitialize
 	SplitScreenShakeScale = 0.5f;
 }
 
-float UCameraModifier_CameraShake::GetShakeScale(FCameraShakeInstance const& ShakeInst) const
-{
-	return ShakeInst.Scale * Alpha;
-}
-
-static inline float UpdateFOscillator(FFOscillator const& Osc, float& CurrentOffset, float DeltaTime)
-{
-	if (Osc.Amplitude != 0.f)
-	{
-		CurrentOffset += DeltaTime * Osc.Frequency;
-		return Osc.Amplitude * FMath::Sin(CurrentOffset);
-	}
-	return 0.f;
-}
-
-void UCameraModifier_CameraShake::UpdateCameraShake(float DeltaTime, FCameraShakeInstance& Shake, FMinimalViewInfo& InOutPOV)
-{
-	if (!Shake.SourceShake)
-	{
-		return;
-	}
-
-	// this is the base scale for the whole shake, anim and oscillation alike
-	float const BaseShakeScale = FMath::Max<float>(GetShakeScale(Shake), 0.0f);
-
-	// update anims with any desired scaling
-	if (Shake.AnimInst)
-	{
-		Shake.AnimInst->TransientScaleModifier *= BaseShakeScale;
-	}
-
-	// update oscillation times
- 	if (Shake.OscillatorTimeRemaining > 0.f)
- 	{
- 		Shake.OscillatorTimeRemaining -= DeltaTime;
- 		Shake.OscillatorTimeRemaining = FMath::Max(0.f, Shake.OscillatorTimeRemaining);
- 	}
-	if (Shake.bBlendingIn)
-	{
-		Shake.CurrentBlendInTime += DeltaTime;
-	}
-	if (Shake.bBlendingOut)
-	{
-		Shake.CurrentBlendOutTime += DeltaTime;
-	}
-
-	// see if we've crossed any important time thresholds and deal appropriately
-	bool bOscillationFinished = false;
-
-	if (Shake.OscillatorTimeRemaining == 0.f)
-	{
-		// finished!
-		bOscillationFinished = true;
-	}
-	else if (Shake.OscillatorTimeRemaining < 0.0f)
-	{
-		// indefinite shaking
-	}
-	else if (Shake.OscillatorTimeRemaining < Shake.SourceShake->OscillationBlendOutTime)
-	{
-		// start blending out
-		Shake.bBlendingOut = true;
-		Shake.CurrentBlendOutTime = Shake.SourceShake->OscillationBlendOutTime - Shake.OscillatorTimeRemaining;
-	}
-
-	if (Shake.bBlendingIn)
-	{
-		if (Shake.CurrentBlendInTime > Shake.SourceShake->OscillationBlendInTime)
-		{
-			// done blending in!
-			Shake.bBlendingIn = false;
-		}
-	}
-	if (Shake.bBlendingOut)
-	{
-		if (Shake.CurrentBlendOutTime > Shake.SourceShake->OscillationBlendOutTime)
-		{
-			// done!!
-			Shake.CurrentBlendOutTime = Shake.SourceShake->OscillationBlendOutTime;
-			bOscillationFinished = true;
-		}
-	}
-
-	// calculate blend weight. calculating separately and taking the minimum handles overlapping blends nicely.
-	float const BlendInWeight = (Shake.bBlendingIn) ? (Shake.CurrentBlendInTime / Shake.SourceShake->OscillationBlendInTime) : 1.f;
-	float const BlendOutWeight = (Shake.bBlendingOut) ? (1.f - Shake.CurrentBlendOutTime / Shake.SourceShake->OscillationBlendOutTime) : 1.f;
-	float const CurrentBlendWeight = FMath::Min(BlendInWeight, BlendOutWeight);
-
-
-	// Do not update oscillation further if finished
-	if (bOscillationFinished)
-	{
-		return;
-	}
-
-	// this is the oscillation scale, which includes oscillation fading
-	float const OscillationScale = GetShakeScale(Shake) * CurrentBlendWeight;
-
-	if (OscillationScale > 0.f)
-	{
-		// View location offset, compute sin wave value for each component
-		FVector	LocOffset = FVector(0);
-		LocOffset.X = UpdateFOscillator(Shake.SourceShake->LocOscillation.X, Shake.LocSinOffset.X, DeltaTime);
-		LocOffset.Y = UpdateFOscillator(Shake.SourceShake->LocOscillation.Y, Shake.LocSinOffset.Y, DeltaTime);
-		LocOffset.Z = UpdateFOscillator(Shake.SourceShake->LocOscillation.Z, Shake.LocSinOffset.Z, DeltaTime);
-		LocOffset *= OscillationScale;
-
-		// View rotation offset, compute sin wave value for each component
-		FRotator RotOffset;
-		RotOffset.Pitch = UpdateFOscillator(Shake.SourceShake->RotOscillation.Pitch, Shake.RotSinOffset.X, DeltaTime) * OscillationScale;
-		RotOffset.Yaw = UpdateFOscillator(Shake.SourceShake->RotOscillation.Yaw, Shake.RotSinOffset.Y, DeltaTime) * OscillationScale;
-		RotOffset.Roll = UpdateFOscillator(Shake.SourceShake->RotOscillation.Roll, Shake.RotSinOffset.Z, DeltaTime) * OscillationScale;
-
-		if (Shake.PlaySpace == ECameraAnimPlaySpace::CameraLocal)
-		{
-			// the else case will handle this as well, but this is the faster, cleaner, most common code path
-
-			// apply loc offset relative to camera orientation
-			FRotationMatrix CamRotMatrix(InOutPOV.Rotation);
-			InOutPOV.Location += CamRotMatrix.TransformVector(LocOffset);
-
-			// apply rot offset relative to camera orientation
-			FRotationMatrix const AnimRotMat( RotOffset );
-			InOutPOV.Rotation = (AnimRotMat * FRotationMatrix(InOutPOV.Rotation)).Rotator();
-		}
-		else
-		{
-			// find desired space
-			FMatrix const PlaySpaceToWorld = (Shake.PlaySpace == ECameraAnimPlaySpace::UserDefined) ? Shake.UserPlaySpaceMatrix : FMatrix::Identity;
-
-			// apply loc offset relative to desired space
-			InOutPOV.Location += PlaySpaceToWorld.TransformVector( LocOffset );
-
-			// apply rot offset relative to desired space
-
-			// find transform from camera to the "play space"
-			FRotationMatrix const CamToWorld(InOutPOV.Rotation);
-			FMatrix const CameraToPlaySpace = CamToWorld * PlaySpaceToWorld.Inverse();			// CameraToWorld * WorldToPlaySpace
-
-			// find transform from anim (applied in playspace) back to camera
-			FRotationMatrix const AnimToPlaySpace(RotOffset);
-			FMatrix const AnimToCamera = AnimToPlaySpace * CameraToPlaySpace.Inverse();			// AnimToPlaySpace * PlaySpaceToCamera
-
-			// RCS = rotated camera space, meaning camera space after it's been animated
-			// this is what we're looking for, the diff between rotated cam space and regular cam space.
-			// apply the transform back to camera space from the post-animated transform to get the RCS
-			FMatrix const RCSToCamera = CameraToPlaySpace * AnimToCamera;
-
-			// now apply to real camera
-			InOutPOV.Rotation = (RCSToCamera * CamToWorld).Rotator();
-		}
-
-		// Compute FOV change
-		InOutPOV.FOV += OscillationScale * UpdateFOscillator(Shake.SourceShake->FOVOscillation, Shake.FOVSinOffset, DeltaTime);
-	}
-}
 
 bool UCameraModifier_CameraShake::ModifyCamera(class APlayerCameraManager* Camera, float DeltaTime, FMinimalViewInfo& InOutPOV)
 {
@@ -184,23 +27,21 @@ bool UCameraModifier_CameraShake::ModifyCamera(class APlayerCameraManager* Camer
 		return false;
 	}
 
-	// Update Screen Shakes array
+	// Update and apply active shakes
 	if( ActiveShakes.Num() > 0 )
 	{
-		for(int32 i=0; i<ActiveShakes.Num(); i++)
+		for (UCameraShake* ShakeInst : ActiveShakes)
 		{
-			UpdateCameraShake(DeltaTime, ActiveShakes[i], InOutPOV);
+			ShakeInst->UpdateAndApplyCameraShake(DeltaTime, Alpha, InOutPOV);
 		}
 
 		// Delete any obsolete shakes
-		for(int32 i=ActiveShakes.Num()-1; i>=0; i--)
+		for (int32 i=ActiveShakes.Num()-1; i>=0; i--)
 		{
-			FCameraShakeInstance& ShakeInst = ActiveShakes[i];
-			if( !ShakeInst.SourceShake
-				|| ( (ShakeInst.OscillatorTimeRemaining == 0.f) && 
-					 ((ShakeInst.AnimInst == NULL) || (ShakeInst.AnimInst->bFinished == true)) ) )
+			UCameraShake* ShakeInst = ActiveShakes[i];
+			if ((ShakeInst == nullptr) || ShakeInst->IsFinished())
 			{
-				ActiveShakes.RemoveAt(i,1);
+				ActiveShakes.RemoveAt(i, 1);
 			}
 		}
 	}
@@ -213,172 +54,82 @@ bool UCameraModifier_CameraShake::ModifyCamera(class APlayerCameraManager* Camer
 	return false;
 }
 
-float UCameraModifier_CameraShake::InitializeOffset( const FFOscillator& Param ) const
+UCameraShake* UCameraModifier_CameraShake::AddCameraShake(TSubclassOf<class UCameraShake> ShakeClass, float Scale, ECameraAnimPlaySpace::Type PlaySpace, FRotator UserPlaySpaceRot)
 {
-	switch( Param.InitialOffset )
+	if (ShakeClass != nullptr)
 	{
-		case EOO_OffsetRandom	: return FMath::FRand() * 2.f * PI;	break;
-		case EOO_OffsetZero		: return 0.f;					break;
-		default					: return 0.f;
-	}
-}
-
-void UCameraModifier_CameraShake::ReinitShake(FCameraShakeInstance& ShakeInst, float NewScale)
-{
-	if (GEngine->IsSplitScreen(CameraOwner->GetWorld()))
-	{
-		NewScale *= SplitScreenShakeScale;
-	}
-	ShakeInst.Scale = NewScale;
-
-	UCameraShake const* const SourceShake = ShakeInst.SourceShake;
-
-	if (SourceShake->OscillationDuration != 0.f)
-	{
-		ShakeInst.OscillatorTimeRemaining = SourceShake->OscillationDuration;
-
-		if (ShakeInst.bBlendingOut)
-		{
-			ShakeInst.bBlendingOut = false;
-			ShakeInst.CurrentBlendOutTime = 0.f;
-
-			// stop any blendout and reverse it to a blendin
-			ShakeInst.bBlendingIn = true;
-			ShakeInst.CurrentBlendInTime = ShakeInst.SourceShake->OscillationBlendInTime * (1.f - ShakeInst.CurrentBlendOutTime / ShakeInst.SourceShake->OscillationBlendOutTime);
-		}
-	}
-
-	if (SourceShake->Anim != NULL)
-	{
-		bool bLoop = false, bRandomStart = false;
-		float Duration = 0.f;
-		if (SourceShake->bRandomAnimSegment)
+		// adjust for splitscreen
+		if (GEngine->IsSplitScreen(CameraOwner->GetWorld()))
 		{
-			bLoop = true;
-			bRandomStart = true;
-			Duration = SourceShake->RandomAnimSegmentDuration;
+			Scale *= SplitScreenShakeScale;
 		}
 
-		float AnimScale = NewScale * SourceShake->AnimScale;
-
-		if (ShakeInst.AnimInst != nullptr)
-		{
-			ShakeInst.AnimInst->Update(SourceShake->AnimPlayRate, AnimScale, SourceShake->AnimBlendInTime, SourceShake->AnimBlendOutTime, Duration);
-		}
-		else
+		UCameraShake const* const ShakeCDO = GetDefault<UCameraShake>(ShakeClass);
+		if (ShakeCDO && ShakeCDO->bSingleInstance)
 		{
-			ShakeInst.AnimInst = CameraOwner->PlayCameraAnim(SourceShake->Anim, SourceShake->AnimPlayRate, AnimScale, SourceShake->AnimBlendInTime, SourceShake->AnimBlendOutTime, bLoop, bRandomStart, Duration);
+			// look for existing instance of same class
+			for (UCameraShake* ShakeInst : ActiveShakes)
+			{
+				if (ShakeInst && (ShakeClass == ShakeInst->GetClass()))
+				{
+					// just restart the existing shake
+					ShakeInst->PlayShake(CameraOwner, Scale, PlaySpace, UserPlaySpaceRot);
+					return ShakeInst;
+				}
+			}
 		}
-	}
-}
-
-FCameraShakeInstance UCameraModifier_CameraShake::InitializeShake(TSubclassOf<class UCameraShake> NewShake, float Scale, ECameraAnimPlaySpace::Type PlaySpace, FRotator UserPlaySpaceRot)
-{
-	FCameraShakeInstance Inst;
-	Inst.SourceShakeName = NewShake->GetFName();
-
-	// Create a camera shake object of class NewShake
-	UCameraShake const* const SourceShake = CastChecked<UCameraShake>(NewShake->GetDefaultObject());
-	Inst.SourceShake = SourceShake;
-
-	Inst.Scale = Scale;
-	if (GEngine->IsSplitScreen(CameraOwner->GetWorld()))
-	{
-		Scale *= SplitScreenShakeScale;
-	}
-
-	// init oscillations
-	if ( SourceShake->OscillationDuration != 0.f )
-	{
-		Inst.RotSinOffset.X		= InitializeOffset( SourceShake->RotOscillation.Pitch );
-		Inst.RotSinOffset.Y		= InitializeOffset( SourceShake->RotOscillation.Yaw );
-		Inst.RotSinOffset.Z		= InitializeOffset( SourceShake->RotOscillation.Roll );
-
-		Inst.LocSinOffset.X		= InitializeOffset( SourceShake->LocOscillation.X );
-		Inst.LocSinOffset.Y		= InitializeOffset( SourceShake->LocOscillation.Y );
-		Inst.LocSinOffset.Z		= InitializeOffset( SourceShake->LocOscillation.Z );
 
-		Inst.FOVSinOffset		= InitializeOffset( SourceShake->FOVOscillation );
-
-		Inst.OscillatorTimeRemaining = SourceShake->OscillationDuration;
-
-		if (SourceShake->OscillationBlendInTime > 0.f)
+		UCameraShake* const NewInst = NewObject<UCameraShake>(this, ShakeClass);
+		if (NewInst)
 		{
-			Inst.bBlendingIn = true;
-			Inst.CurrentBlendInTime = 0.f;
-		}
-	}
+			// Initialize new shake and add it to the list of active shakes
+			NewInst->PlayShake(CameraOwner, Scale, PlaySpace, UserPlaySpaceRot);
 
-	// init anims
-	if (SourceShake->Anim != NULL)
-	{
-		bool bLoop = false;
-		bool bRandomStart = false;
-		float Duration = 0.f;
-		if (SourceShake->bRandomAnimSegment)
-		{
-			bLoop = true;
-			bRandomStart = true;
-			Duration = SourceShake->RandomAnimSegmentDuration;
-		}
+			// look for nulls in the array to replace first -- keeps the array compact
+			bool bReplacedNull = false;
+			for (int32 Idx = 0; Idx < ActiveShakes.Num(); ++Idx)
+			{
+				if (ActiveShakes[Idx] == nullptr)
+				{
+					ActiveShakes[Idx] = NewInst;
+					bReplacedNull = true;
+				}
+			}
 
-		float AnimScale = Scale * SourceShake->AnimScale;
-		
-		if (AnimScale > 0.f)
-		{
-			Inst.AnimInst = CameraOwner->PlayCameraAnim(SourceShake->Anim, SourceShake->AnimPlayRate, AnimScale, SourceShake->AnimBlendInTime, SourceShake->AnimBlendOutTime, bLoop, bRandomStart, Duration, PlaySpace, UserPlaySpaceRot);
+			// no holes, extend the array
+			if (bReplacedNull == false)
+			{
+				ActiveShakes.Emplace(NewInst);
+			}
 		}
-	}
 
-	Inst.PlaySpace = PlaySpace;
-	if (Inst.PlaySpace == ECameraAnimPlaySpace::UserDefined)
-	{
-		Inst.UserPlaySpaceMatrix = FRotationMatrix(UserPlaySpaceRot);
+		return NewInst;
 	}
 
-	return Inst;
+	return nullptr;
 }
 
-
-void UCameraModifier_CameraShake::AddCameraShake( TSubclassOf<class UCameraShake> NewShake, float Scale, ECameraAnimPlaySpace::Type PlaySpace, FRotator UserPlaySpaceRot )
+void UCameraModifier_CameraShake::RemoveCameraShake(UCameraShake* ShakeInst)
 {
-	if (NewShake != NULL)
+	for (int32 i = 0; i < ActiveShakes.Num(); ++i)
 	{
-		UCameraShake const* const ShakeCDO = GetDefault<UCameraShake>(NewShake);
-
-		if (ShakeCDO && ShakeCDO->bSingleInstance)
+		if (ActiveShakes[i] == ShakeInst)
 		{
-			for (FCameraShakeInstance& ShakeInst : ActiveShakes)
-			{
-				if (ShakeInst.SourceShakeName == NewShake->GetFName())
-				{
-					ReinitShake(ShakeInst, Scale);
-					return;
-				}
-			}
+			ShakeInst->StopShake();
+			ActiveShakes.RemoveAt(i, 1);
+			break;
 		}
-
-		int32 NumShakes = ActiveShakes.Num();
-
-		// Initialize new shake and add it to the list of active shakes
-		ActiveShakes.Add(InitializeShake(NewShake, Scale, PlaySpace, UserPlaySpaceRot));
 	}
 }
 
-void UCameraModifier_CameraShake::RemoveCameraShake(TSubclassOf<class UCameraShake> Shake)
+void UCameraModifier_CameraShake::RemoveAllCameraShakesOfClass(TSubclassOf<class UCameraShake> ShakeClass)
 {
-	for (int32 i = 0; i < ActiveShakes.Num(); ++i)
+	for (int32 i = ActiveShakes.Num()-1; i >= 0; --i)
 	{
-		if (ActiveShakes[i].SourceShakeName == Shake->GetFName())
+		if ( ActiveShakes[i] && (ActiveShakes[i]->GetClass()->IsChildOf(ShakeClass)) )
 		{
-			UCameraAnimInst* AnimInst = ActiveShakes[i].AnimInst;
-			if ( (AnimInst != NULL) && !AnimInst->bFinished )
-			{
-				CameraOwner->StopCameraAnimInst(AnimInst, true);
-			}
-
-			ActiveShakes.RemoveAt(i,1);
-			return;
+			ActiveShakes[i]->StopShake();
+			ActiveShakes.RemoveAt(i, 1);
 		}
 	}
 }
@@ -386,13 +137,9 @@ void UCameraModifier_CameraShake::RemoveCameraShake(TSubclassOf<class UCameraSha
 void UCameraModifier_CameraShake::RemoveAllCameraShakes()
 {
 	// clean up any active camera shake anims
-	for (int32 Idx=0; Idx<ActiveShakes.Num(); ++Idx)
+	for (UCameraShake* Inst : ActiveShakes)
 	{
-		UCameraAnimInst* AnimInst = ActiveShakes[Idx].AnimInst;
-		if ( (AnimInst != NULL) && !AnimInst->bFinished )
-		{
-			CameraOwner->StopCameraAnimInst(AnimInst, true);
-		}
+		Inst->StopShake();
 	}
 
 	// clear ActiveShakes array
diff --git a/Engine/Source/Runtime/Engine/Private/Camera/CameraShake.cpp b/Engine/Source/Runtime/Engine/Private/Camera/CameraShake.cpp
--- a/Engine/Source/Runtime/Engine/Private/Camera/CameraShake.cpp
+++ b/Engine/Source/Runtime/Engine/Private/Camera/CameraShake.cpp
@@ -1,6 +1,29 @@
 // Copyright 1998-2015 Epic Games, Inc. All Rights Reserved.
 
 #include "EnginePrivate.h"
+#include "Camera/PlayerCameraManager.h"
+
+//////////////////////////////////////////////////////////////////////////
+// FFOscillator
+
+// static
+float FFOscillator::UpdateOffset(FFOscillator const& Osc, float& CurrentOffset, float DeltaTime)
+{
+	if (Osc.Amplitude != 0.f)
+	{
+		CurrentOffset += DeltaTime * Osc.Frequency;
+		return Osc.Amplitude * FMath::Sin(CurrentOffset);
+	}
+	return 0.f;
+}
+
+// static
+float FFOscillator::GetInitialOffset(FFOscillator const& Osc)
+{
+	return (Osc.InitialOffset == EOO_OffsetRandom)
+		? FMath::FRand() * (2.f * PI)
+		: 0.f;
+}
 
 //////////////////////////////////////////////////////////////////////////
 // UCameraShake
@@ -15,3 +38,258 @@ UCameraShake::UCameraShake(const FObjectInitializer& ObjectInitializer)
 	OscillationBlendInTime = 0.1f;
 	OscillationBlendOutTime = 0.2f;
 }
+
+void UCameraShake::StopShake()
+{
+	// stop cam anim if playing
+	if ((AnimInst != nullptr) && !AnimInst->bFinished)
+	{
+		CameraOwner->StopCameraAnimInst(AnimInst, true);
+		AnimInst = nullptr;
+	}
+
+	ReceiveStopShake();
+}
+
+void UCameraShake::PlayShake(APlayerCameraManager* Camera, float Scale, ECameraAnimPlaySpace::Type PlaySpace, FRotator UserPlaySpaceRot)
+{
+	ShakeScale = Scale;
+	CameraOwner = Camera;
+
+	// init oscillations
+	if (OscillationDuration != 0.f)
+	{
+
+		if (OscillatorTimeRemaining > 0.f)
+		{
+			// this shake was already playing
+			OscillatorTimeRemaining = OscillationDuration;
+
+			if (bBlendingOut)
+			{
+				bBlendingOut = false;
+				CurrentBlendOutTime = 0.f;
+
+				// stop any blendout and reverse it to a blendin
+				if (OscillationBlendInTime > 0.f)
+				{
+					bBlendingIn = true;
+					CurrentBlendInTime = OscillationBlendInTime * (1.f - CurrentBlendOutTime / OscillationBlendOutTime);
+				}
+				else
+				{
+					bBlendingIn = false;
+					CurrentBlendInTime = 0.f;
+				}
+			}
+		}
+		else
+		{
+			RotSinOffset.X = FFOscillator::GetInitialOffset(RotOscillation.Pitch);
+			RotSinOffset.Y = FFOscillator::GetInitialOffset(RotOscillation.Yaw);
+			RotSinOffset.Z = FFOscillator::GetInitialOffset(RotOscillation.Roll);
+
+			LocSinOffset.X = FFOscillator::GetInitialOffset(LocOscillation.X);
+			LocSinOffset.Y = FFOscillator::GetInitialOffset(LocOscillation.Y);
+			LocSinOffset.Z = FFOscillator::GetInitialOffset(LocOscillation.Z);
+
+			FOVSinOffset = FFOscillator::GetInitialOffset(FOVOscillation);
+
+			OscillatorTimeRemaining = OscillationDuration;
+
+			if (OscillationBlendInTime > 0.f)
+			{
+				bBlendingIn = true;
+				CurrentBlendInTime = 0.f;
+			}
+		}
+	}
+
+	// init cameraanim shakes
+	if (Anim != nullptr)
+	{
+		if (AnimInst)
+		{
+			float const Duration = bRandomAnimSegment ? RandomAnimSegmentDuration : 0.f;
+			float const FinalAnimScale = Scale * AnimScale;
+			AnimInst->Update(AnimPlayRate, FinalAnimScale, AnimBlendInTime, AnimBlendOutTime, Duration);
+		}
+		else
+		{
+			bool bLoop = false;
+			bool bRandomStart = false;
+			float Duration = 0.f;
+			if (bRandomAnimSegment)
+			{
+				bLoop = true;
+				bRandomStart = true;
+				Duration = RandomAnimSegmentDuration;
+			}
+
+			float const FinalAnimScale = Scale * AnimScale;
+			if (FinalAnimScale > 0.f)
+			{
+				AnimInst = CameraOwner->PlayCameraAnim(Anim, AnimPlayRate, FinalAnimScale, AnimBlendInTime, AnimBlendOutTime, bLoop, bRandomStart, Duration, PlaySpace, UserPlaySpaceRot);
+			}
+		}
+	}
+
+	PlaySpace = PlaySpace;
+	if (PlaySpace == ECameraAnimPlaySpace::UserDefined)
+	{
+		UserPlaySpaceMatrix = FRotationMatrix(UserPlaySpaceRot);
+	}
+
+	ReceivePlayShake(Scale);
+}
+
+void UCameraShake::UpdateAndApplyCameraShake(float DeltaTime, float Alpha, FMinimalViewInfo& InOutPOV)
+{
+	// this is the base scale for the whole shake, anim and oscillation alike
+	float const BaseShakeScale = FMath::Max<float>(Alpha * ShakeScale, 0.0f);
+
+	// update anims with any desired scaling
+	if (AnimInst)
+	{
+		AnimInst->TransientScaleModifier *= BaseShakeScale;
+	}
+
+	// update oscillation times
+	if (OscillatorTimeRemaining > 0.f)
+	{
+		OscillatorTimeRemaining -= DeltaTime;
+		OscillatorTimeRemaining = FMath::Max(0.f, OscillatorTimeRemaining);
+	}
+	if (bBlendingIn)
+	{
+		CurrentBlendInTime += DeltaTime;
+	}
+	if (bBlendingOut)
+	{
+		CurrentBlendOutTime += DeltaTime;
+	}
+
+	// see if we've crossed any important time thresholds and deal appropriately
+	bool bOscillationFinished = false;
+
+	if (OscillatorTimeRemaining == 0.f)
+	{
+		// finished!
+		bOscillationFinished = true;
+	}
+	else if (OscillatorTimeRemaining < 0.0f)
+	{
+		// indefinite shaking
+	}
+	else if (OscillatorTimeRemaining < OscillationBlendOutTime)
+	{
+		// start blending out
+		bBlendingOut = true;
+		CurrentBlendOutTime = OscillationBlendOutTime - OscillatorTimeRemaining;
+	}
+
+	if (bBlendingIn)
+	{
+		if (CurrentBlendInTime > OscillationBlendInTime)
+		{
+			// done blending in!
+			bBlendingIn = false;
+		}
+	}
+	if (bBlendingOut)
+	{
+		if (CurrentBlendOutTime > OscillationBlendOutTime)
+		{
+			// done!!
+			CurrentBlendOutTime = OscillationBlendOutTime;
+			bOscillationFinished = true;
+		}
+	}
+
+	// Do not update oscillation further if finished
+	if (bOscillationFinished == false)
+	{
+		// calculate blend weight. calculating separately and taking the minimum handles overlapping blends nicely.
+		float const BlendInWeight = (bBlendingIn) ? (CurrentBlendInTime / OscillationBlendInTime) : 1.f;
+		float const BlendOutWeight = (bBlendingOut) ? (1.f - CurrentBlendOutTime / OscillationBlendOutTime) : 1.f;
+		float const CurrentBlendWeight = FMath::Min(BlendInWeight, BlendOutWeight);
+
+		// this is the oscillation scale, which includes oscillation fading
+		float const OscillationScale = BaseShakeScale * CurrentBlendWeight;
+
+		if (OscillationScale > 0.f)
+		{
+			// View location offset, compute sin wave value for each component
+			FVector	LocOffset = FVector(0);
+			LocOffset.X = FFOscillator::UpdateOffset(LocOscillation.X, LocSinOffset.X, DeltaTime);
+			LocOffset.Y = FFOscillator::UpdateOffset(LocOscillation.Y, LocSinOffset.Y, DeltaTime);
+			LocOffset.Z = FFOscillator::UpdateOffset(LocOscillation.Z, LocSinOffset.Z, DeltaTime);
+			LocOffset *= OscillationScale;
+
+			// View rotation offset, compute sin wave value for each component
+			FRotator RotOffset;
+			RotOffset.Pitch = FFOscillator::UpdateOffset(RotOscillation.Pitch, RotSinOffset.X, DeltaTime) * OscillationScale;
+			RotOffset.Yaw = FFOscillator::UpdateOffset(RotOscillation.Yaw, RotSinOffset.Y, DeltaTime) * OscillationScale;
+			RotOffset.Roll = FFOscillator::UpdateOffset(RotOscillation.Roll, RotSinOffset.Z, DeltaTime) * OscillationScale;
+
+			if (PlaySpace == ECameraAnimPlaySpace::CameraLocal)
+			{
+				// the else case will handle this as well, but this is the faster, cleaner, most common code path
+
+				// apply loc offset relative to camera orientation
+				FRotationMatrix CamRotMatrix(InOutPOV.Rotation);
+				InOutPOV.Location += CamRotMatrix.TransformVector(LocOffset);
+
+				// apply rot offset relative to camera orientation
+				FRotationMatrix const AnimRotMat(RotOffset);
+				InOutPOV.Rotation = (AnimRotMat * FRotationMatrix(InOutPOV.Rotation)).Rotator();
+			}
+			else
+			{
+				// find desired space
+				FMatrix const PlaySpaceToWorld = (PlaySpace == ECameraAnimPlaySpace::UserDefined) ? UserPlaySpaceMatrix : FMatrix::Identity;
+
+				// apply loc offset relative to desired space
+				InOutPOV.Location += PlaySpaceToWorld.TransformVector(LocOffset);
+
+				// apply rot offset relative to desired space
+
+				// find transform from camera to the "play space"
+				FRotationMatrix const CamToWorld(InOutPOV.Rotation);
+				FMatrix const CameraToPlaySpace = CamToWorld * PlaySpaceToWorld.Inverse();			// CameraToWorld * WorldToPlaySpace
+
+				// find transform from anim (applied in playspace) back to camera
+				FRotationMatrix const AnimToPlaySpace(RotOffset);
+				FMatrix const AnimToCamera = AnimToPlaySpace * CameraToPlaySpace.Inverse();			// AnimToPlaySpace * PlaySpaceToCamera
+
+				// RCS = rotated camera space, meaning camera space after it's been animated
+				// this is what we're looking for, the diff between rotated cam space and regular cam space.
+				// apply the transform back to camera space from the post-animated transform to get the RCS
+				FMatrix const RCSToCamera = CameraToPlaySpace * AnimToCamera;
+
+				// now apply to real camera
+				InOutPOV.Rotation = (RCSToCamera * CamToWorld).Rotator();
+			}
+
+			// Compute FOV change
+			InOutPOV.FOV += OscillationScale * FFOscillator::UpdateOffset(FOVOscillation, FOVSinOffset, DeltaTime);
+		}
+	}
+
+	BlueprintUpdateCameraShake(DeltaTime, Alpha, InOutPOV, InOutPOV);
+}
+
+bool UCameraShake::IsFinished() const
+{
+	return ((OscillatorTimeRemaining <= 0.f) &&						// oscillator is finished
+		((AnimInst == nullptr) || AnimInst->bFinished) &&		// anim is finished
+		ReceiveIsFinished()										// BP thinks it's finished
+		);
+}
+
+bool UCameraShake::ReceiveIsFinished_Implementation() const
+{
+	return true;
+}
+
+
diff --git a/Engine/Source/Runtime/Engine/Private/PlayerCameraManager.cpp b/Engine/Source/Runtime/Engine/Private/PlayerCameraManager.cpp
--- a/Engine/Source/Runtime/Engine/Private/PlayerCameraManager.cpp
+++ b/Engine/Source/Runtime/Engine/Private/PlayerCameraManager.cpp
@@ -1088,28 +1088,46 @@ void APlayerCameraManager::ClearCameraLensEffects()
  *  Camera Shakes
  *  ------------------------------------------------------------ */
 
-void APlayerCameraManager::PlayCameraShake(TSubclassOf<UCameraShake> Shake, float Scale, ECameraAnimPlaySpace::Type PlaySpace, FRotator UserPlaySpaceRot)
+UCameraShake* APlayerCameraManager::PlayCameraShake(TSubclassOf<UCameraShake> ShakeClass, float Scale, ECameraAnimPlaySpace::Type PlaySpace, FRotator UserPlaySpaceRot)
 {
-	if ((Shake != nullptr) && (Scale > 0.0f))
+	if (ShakeClass && CameraShakeCamMod && (Scale > 0.0f) )
 	{
-		CameraShakeCamMod->AddCameraShake(Shake, Scale, PlaySpace, UserPlaySpaceRot);
+		return CameraShakeCamMod->AddCameraShake(ShakeClass, Scale, PlaySpace, UserPlaySpaceRot);
 	}
+
+	return nullptr;
 }
 
-void APlayerCameraManager::StopCameraShake(TSubclassOf<class UCameraShake> Shake)
+
+void APlayerCameraManager::StopCameraShake(UCameraShake* ShakeInst)
 {
-	if (Shake != NULL)
+	if (ShakeInst && CameraShakeCamMod)
 	{
-		CameraShakeCamMod->RemoveCameraShake(Shake);
+		CameraShakeCamMod->RemoveCameraShake(ShakeInst);
 	}
 }
 
+void APlayerCameraManager::StopAllInstancesOfCameraShake(TSubclassOf<class UCameraShake> ShakeClass)
+{
+	if (ShakeClass && CameraShakeCamMod)
+	{
+		CameraShakeCamMod->RemoveAllCameraShakesOfClass(ShakeClass);
+	}
+}
 
-float APlayerCameraManager::CalcRadialShakeScale(APlayerCameraManager* Cam, FVector Epicenter, float InnerRadius, float OuterRadius, float Falloff)
+void APlayerCameraManager::StopAllCameraShakes()
+{
+	if (CameraShakeCamMod)
+	{
+		CameraShakeCamMod->RemoveAllCameraShakes();
+	}
+}
+
+float APlayerCameraManager::CalcRadialShakeScale(APlayerCameraManager* Camera, FVector Epicenter, float InnerRadius, float OuterRadius, float Falloff)
 {
 	// using camera location so stuff like spectator cameras get shakes applied sensibly as well
 	// need to ensure server has reasonably accurate camera position
-	FVector POVLoc = Cam->GetActorLocation();
+	FVector POVLoc = Camera->GetActorLocation();
 
 	if (InnerRadius < OuterRadius)
 	{
@@ -1149,11 +1167,6 @@ void APlayerCameraManager::PlayWorldCameraShake(UWorld* InWorld, TSubclassOf<cla
 	}
 }
 
-void APlayerCameraManager::ClearAllCameraShakes()
-{
-	CameraShakeCamMod->RemoveAllCameraShakes();
-//	StopAllCameraAnims(true);
-}
 
 /** ------------------------------------------------------------
  *  Camera fades
diff --git a/Engine/Source/Runtime/Engine/Private/PlayerController.cpp b/Engine/Source/Runtime/Engine/Private/PlayerController.cpp
--- a/Engine/Source/Runtime/Engine/Private/PlayerController.cpp
+++ b/Engine/Source/Runtime/Engine/Private/PlayerController.cpp
@@ -3588,7 +3588,7 @@ void APlayerController::ClientStopCameraShake_Implementation( TSubclassOf<class
 {
 	if (PlayerCameraManager != NULL)
 	{
-		PlayerCameraManager->StopCameraShake(Shake);
+		PlayerCameraManager->StopAllInstancesOfCameraShake(Shake);
 	}
 }
 
