#ue4 - Perf: try to avoid using dynamic allocation in USceneComponent::SetVisibility() and USceneComponent::SetHiddenInGame().

diff --git a/Engine/Source/Runtime/Engine/Private/Components/SceneComponent.cpp b/Engine/Source/Runtime/Engine/Private/Components/SceneComponent.cpp
--- a/Engine/Source/Runtime/Engine/Private/Components/SceneComponent.cpp
+++ b/Engine/Source/Runtime/Engine/Private/Components/SceneComponent.cpp
@@ -1783,10 +1783,7 @@ void USceneComponent::SetVisibility(bool bNewVisibility, bool bPropagateToChildr
 	{
 		// fully traverse down the attachment tree
 		// we do it entirely inline here instead of recursing in case a primitivecomponent is a child of a non-primitivecomponent
-		TArray<USceneComponent*> ComponentStack;
-
-		// presize to minimize reallocs
-		ComponentStack.Reserve(FMath::Max(32, AttachChildren.Num()));
+		TInlineComponentArray<USceneComponent*, NumInlinedActorComponents> ComponentStack;
 
 		// prime the pump
 		ComponentStack.Append(AttachChildren);
@@ -1816,10 +1813,7 @@ void USceneComponent::SetHiddenInGame(bool NewHiddenGame, bool bPropagateToChild
 	{
 		// fully traverse down the attachment tree
 		// we do it entirely inline here instead of recursing in case a primitivecomponent is a child of a non-primitivecomponent
-		TArray<USceneComponent*> ComponentStack;
-
-		// presize to minimize reallocs
-		ComponentStack.Reserve(FMath::Max(32, AttachChildren.Num()));
+		TInlineComponentArray<USceneComponent*, NumInlinedActorComponents> ComponentStack;
 
 		// prime the pump
 		ComponentStack.Append(AttachChildren);
