#ue4 - Simplify definition of GetActorLocation, GetActorRotation, GetActorScale, GetActorQuat. Don't create temps that get overwritten.

diff --git a/Engine/Source/Runtime/Engine/Classes/GameFramework/Actor.h b/Engine/Source/Runtime/Engine/Classes/GameFramework/Actor.h
--- a/Engine/Source/Runtime/Engine/Classes/GameFramework/Actor.h
+++ b/Engine/Source/Runtime/Engine/Classes/GameFramework/Actor.h
@@ -1239,12 +1239,7 @@ public:
 	template<class T>
 	static FORCEINLINE FVector GetActorLocation(const T* RootComponent)
 	{
-		FVector Result(0.f);
-		if( RootComponent != NULL )
-		{
-			Result = RootComponent->GetComponentLocation();
-		}
-		return Result;
+		return (RootComponent != nullptr) ? RootComponent->GetComponentLocation() : FVector(0.f,0.f,0.f);
 	}
 
 	/**
@@ -1254,12 +1249,7 @@ public:
 	template<class T>
 	static FORCEINLINE FRotator GetActorRotation(T* RootComponent)
 	{
-		FRotator Result(0,0,0);
-		if( RootComponent != NULL )
-		{
-			Result = RootComponent->GetComponentRotation();
-		}
-		return Result;
+		return (RootComponent != nullptr) ? RootComponent->GetComponentRotation() : FRotator(0.f,0.f,0.f);
 	}
 
 	/**
@@ -1269,12 +1259,7 @@ public:
 	template<class T>
 	static FORCEINLINE FVector GetActorScale(T* RootComponent)
 	{
-		FVector Result(1,1,1);
-		if( RootComponent != NULL )
-		{
-			Result = RootComponent->GetComponentScale();
-		}
-		return Result;
+		return (RootComponent != nullptr) ? RootComponent->GetComponentScale() : FVector(1.f,1.f,1.f);
 	}
 
 	/**
@@ -1284,12 +1269,7 @@ public:
 	template<class T>
 	static FORCEINLINE FQuat GetActorQuat(T* RootComponent)
 	{
-		FQuat Result(ForceInit);
-		if( RootComponent != NULL )
-		{
-			Result = RootComponent->GetComponentQuat();
-		}
-		return Result;
+		return (RootComponent != nullptr) ? RootComponent->GetComponentQuat() : FQuat(ForceInit);
 	}
 
 	/** Returns this actor's root component. */
