Fixed Clang compile error

diff --git a/Engine/Source/Runtime/Engine/Private/Camera/CameraShake.cpp b/Engine/Source/Runtime/Engine/Private/Camera/CameraShake.cpp
--- a/Engine/Source/Runtime/Engine/Private/Camera/CameraShake.cpp
+++ b/Engine/Source/Runtime/Engine/Private/Camera/CameraShake.cpp
@@ -134,7 +134,6 @@ void UCameraShake::PlayShake(APlayerCameraManager* Camera, float Scale, ECameraA
 		}
 	}
 
-	PlaySpace = PlaySpace;
 	if (PlaySpace == ECameraAnimPlaySpace::UserDefined)
 	{
 		UserPlaySpaceMatrix = FRotationMatrix(UserPlaySpaceRot);
