Stationary lights with many unbuilt objects now fall back to CSM on a frame by frame basis
* This fixes the bug where many unbuilt interactions during registration would get corrected by blueprint construction scripts but the stationary light would continue to think it was built

diff --git a/Engine/Source/Runtime/Engine/Private/Components/DirectionalLightComponent.cpp b/Engine/Source/Runtime/Engine/Private/Components/DirectionalLightComponent.cpp
--- a/Engine/Source/Runtime/Engine/Private/Components/DirectionalLightComponent.cpp
+++ b/Engine/Source/Runtime/Engine/Private/Components/DirectionalLightComponent.cpp
@@ -181,18 +181,6 @@ public:
 
 	// FLightSceneInfo interface.
 
-	/** Called when precomputed lighting has been determined to be invalid */
-	virtual void InvalidatePrecomputedLighting(bool bIsEditor) 
-	{
-		// If the light is unbuilt and is not already using CSM, setup CSM settings to preview unbuilt lighting
-		if (HasStaticShadowing())
-		{
-			CascadeDistributionExponent = 4;
-			DynamicShadowCascades = FMath::Max(0, CVarUnbuiltNumWholeSceneDynamicShadowCascades.GetValueOnAnyThread());
-			WholeSceneDynamicShadowRadius = CVarUnbuiltWholeSceneDynamicShadowRadius.GetValueOnAnyThread();
-		}
-	}
-
 	virtual bool ShouldCreatePerObjectShadowsForDynamicObjects() const 
 	{
 		return FLightSceneProxy::ShouldCreatePerObjectShadowsForDynamicObjects()
@@ -204,9 +192,9 @@ public:
 	}
 
 	/** Returns the number of view dependent shadows this light will create, not counting distance field shadow cascades. */
-	virtual uint32 GetNumViewDependentWholeSceneShadows(const FSceneView& View) const
+	virtual uint32 GetNumViewDependentWholeSceneShadows(const FSceneView& View, bool bPrecomputedLightingIsValid) const
 	{ 
-		uint32 TotalCascades = GetNumShadowMappedCascades(View.MaxShadowCascades) + FarShadowCascadeCount;
+		uint32 TotalCascades = GetNumShadowMappedCascades(View.MaxShadowCascades, bPrecomputedLightingIsValid) + FarShadowCascadeCount;
 
 		return TotalCascades;
 	}
@@ -216,15 +204,15 @@ public:
 	 * @param InCascadeIndex ShadowSplitIndex or INDEX_NONE for the the DistanceFieldCascade
 	 * @return True if the whole-scene projected shadow should be used.
 	 */
-	virtual bool GetViewDependentWholeSceneProjectedShadowInitializer(const FSceneView& View, int32 InCascadeIndex, FWholeSceneProjectedShadowInitializer& OutInitializer) const
+	virtual bool GetViewDependentWholeSceneProjectedShadowInitializer(const FSceneView& View, int32 InCascadeIndex, bool bPrecomputedLightingIsValid, FWholeSceneProjectedShadowInitializer& OutInitializer) const
 	{
 		const FMatrix& WorldToLight = GetWorldToLight();
 
 		const bool bRayTracedCascade = (InCascadeIndex == INDEX_NONE);
 
-		FSphere Bounds = FDirectionalLightSceneProxy::GetShadowSplitBounds(View, InCascadeIndex, &OutInitializer.CascadeSettings);
+		FSphere Bounds = FDirectionalLightSceneProxy::GetShadowSplitBounds(View, InCascadeIndex, bPrecomputedLightingIsValid, &OutInitializer.CascadeSettings);
 
-		uint32 NumNearCascades = GetNumShadowMappedCascades(View.MaxShadowCascades);
+		uint32 NumNearCascades = GetNumShadowMappedCascades(View.MaxShadowCascades, bPrecomputedLightingIsValid);
 
 		// Last cascade is the ray traced shadow, if enabled
 		OutInitializer.CascadeSettings.ShadowSplitIndex = bRayTracedCascade ? NumNearCascades : InCascadeIndex;
@@ -289,11 +277,11 @@ public:
 		return true;
 	}
 
-	virtual FVector2D GetDirectionalLightDistanceFadeParameters(ERHIFeatureLevel::Type InFeatureLevel) const
+	virtual FVector2D GetDirectionalLightDistanceFadeParameters(ERHIFeatureLevel::Type InFeatureLevel, bool bPrecomputedLightingIsValid) const
 	{
-		float FarDistance = GetCSMMaxDistance();
+		float FarDistance = GetCSMMaxDistance(bPrecomputedLightingIsValid);
 		{
-			if (ShouldCreateRayTracedCascade(InFeatureLevel))
+			if (ShouldCreateRayTracedCascade(InFeatureLevel, bPrecomputedLightingIsValid))
 			{
 				FarDistance = GetDistanceFieldShadowDistance();
 			}
@@ -321,30 +309,36 @@ public:
 		return true;
 	}
 
-	virtual bool ShouldCreateRayTracedCascade(ERHIFeatureLevel::Type InFeatureLevel) const
+	virtual bool ShouldCreateRayTracedCascade(ERHIFeatureLevel::Type InFeatureLevel, bool bPrecomputedLightingIsValid) const
 	{
 		//@todo - handle View.MaxShadowCascades properly
-		const uint32 NumCascades = GetNumShadowMappedCascades(10);
+		const uint32 NumCascades = GetNumShadowMappedCascades(10, bPrecomputedLightingIsValid);
 		const float RaytracedShadowDistance = GetDistanceFieldShadowDistance();
-		const bool bCreateWithCSM = NumCascades > 0 && RaytracedShadowDistance > GetCSMMaxDistance();
+		const bool bCreateWithCSM = NumCascades > 0 && RaytracedShadowDistance > GetCSMMaxDistance(bPrecomputedLightingIsValid);
 		const bool bCreateWithoutCSM = NumCascades == 0 && RaytracedShadowDistance > 0;
 		return DoesPlatformSupportDistanceFieldShadowing(GShaderPlatformForFeatureLevel[InFeatureLevel]) && (bCreateWithCSM || bCreateWithoutCSM);
 	}
 
 private:
 
-	uint32 GetNumShadowMappedCascades(uint32 MaxShadowCascades) const
+	float GetEffectiveWholeSceneDynamicShadowRadius(bool bPrecomputedLightingIsValid) const
 	{
-		uint32 NumCascades = WholeSceneDynamicShadowRadius > 0.0f ? DynamicShadowCascades : 0;
+		return bPrecomputedLightingIsValid ? WholeSceneDynamicShadowRadius : CVarUnbuiltWholeSceneDynamicShadowRadius.GetValueOnAnyThread();
+	}
+
+	uint32 GetNumShadowMappedCascades(uint32 MaxShadowCascades, bool bPrecomputedLightingIsValid) const
+	{
+		const int32 EffectiveNumDynamicShadowCascades = bPrecomputedLightingIsValid ? DynamicShadowCascades : FMath::Max(0, CVarUnbuiltNumWholeSceneDynamicShadowCascades.GetValueOnAnyThread());
+		const int32 NumCascades = GetEffectiveWholeSceneDynamicShadowRadius(bPrecomputedLightingIsValid) > 0.0f ? EffectiveNumDynamicShadowCascades : 0;
 		return FMath::Min<int32>(NumCascades, MaxShadowCascades);
 	}
 
-	float GetCSMMaxDistance() const
+	float GetCSMMaxDistance(bool bPrecomputedLightingIsValid) const
 	{
 		static const auto CVar = IConsoleManager::Get().FindTConsoleVariableDataFloat(TEXT("r.Shadow.DistanceScale"));
 		 
 		float Scale = FMath::Clamp(CVar->GetValueOnRenderThread(), 0.1f, 2.0f);
-		float Distance = WholeSceneDynamicShadowRadius * Scale;
+		float Distance = GetEffectiveWholeSceneDynamicShadowRadius(bPrecomputedLightingIsValid) * Scale;
 		return Distance;
 	}
 
@@ -481,14 +475,20 @@ private:
 		return Ret / TotalScale;
 	}
 	
+	float GetEffectiveCascadeDistributionExponent(bool bPrecomputedLightingIsValid) const
+	{
+		return bPrecomputedLightingIsValid ? CascadeDistributionExponent : 4;
+	}
+
 	// @param SplitIndex 0: near, 1:end of 1st cascade, ...
-	inline float GetSplitDistance(const FSceneView& View, uint32 SplitIndex, bool bDistanceFieldShadows) const
+	inline float GetSplitDistance(const FSceneView& View, uint32 SplitIndex, bool bPrecomputedLightingIsValid, bool bDistanceFieldShadows) const
 	{
 		// near cascade means non far and non distance field cascade
-		uint32 NumNearCascades = GetNumShadowMappedCascades(View.MaxShadowCascades);
+		uint32 NumNearCascades = GetNumShadowMappedCascades(View.MaxShadowCascades, bPrecomputedLightingIsValid);
 
-		float CascadeDistanceWithoutFar = GetCSMMaxDistance();
+		float CascadeDistanceWithoutFar = GetCSMMaxDistance(bPrecomputedLightingIsValid);
 		float ShadowNear = View.NearClippingDistance;
+		const float EffectiveCascadeDistributionExponent = GetEffectiveCascadeDistributionExponent(bPrecomputedLightingIsValid);
 
 		// non near cascades are split differently for distance field shadow
 		if(SplitIndex > NumNearCascades)
@@ -502,12 +502,12 @@ private:
 			else
 			{
 				// the far cascades start at the after the near cascades
-				return CascadeDistanceWithoutFar + ComputeAccumulatedScale(CascadeDistributionExponent, SplitIndex - NumNearCascades, FarShadowCascadeCount) * (FarShadowDistance - CascadeDistanceWithoutFar);
+				return CascadeDistanceWithoutFar + ComputeAccumulatedScale(EffectiveCascadeDistributionExponent, SplitIndex - NumNearCascades, FarShadowCascadeCount) * (FarShadowDistance - CascadeDistanceWithoutFar);
 			}
 		}
 		else
 		{
-			return ShadowNear + ComputeAccumulatedScale(CascadeDistributionExponent, SplitIndex, NumNearCascades) * (CascadeDistanceWithoutFar - ShadowNear);
+			return ShadowNear + ComputeAccumulatedScale(EffectiveCascadeDistributionExponent, SplitIndex, NumNearCascades) * (CascadeDistanceWithoutFar - ShadowNear);
 		}
 	}
 
@@ -584,14 +584,14 @@ private:
 	}
 
 	// @param InShadowSplitIndex cascade index or InShadowSplitIndex == INDEX_NONE for the distance field cascade
-	virtual FSphere GetShadowSplitBounds(const FSceneView& View, int32 InCascadeIndex, FShadowCascadeSettings* OutCascadeSettings) const
+	virtual FSphere GetShadowSplitBounds(const FSceneView& View, int32 InCascadeIndex, bool bPrecomputedLightingIsValid, FShadowCascadeSettings* OutCascadeSettings) const
 	{
-		uint32 NumNearCascades = GetNumShadowMappedCascades(View.MaxShadowCascades);
+		uint32 NumNearCascades = GetNumShadowMappedCascades(View.MaxShadowCascades, bPrecomputedLightingIsValid);
 
-		const bool bHasRayTracedCascade = ShouldCreateRayTracedCascade(View.GetFeatureLevel());
+		const bool bHasRayTracedCascade = ShouldCreateRayTracedCascade(View.GetFeatureLevel(), bPrecomputedLightingIsValid);
 
 		// this checks for WholeSceneDynamicShadowRadius and DynamicShadowCascades
-		uint32 NumNearAndFarCascades = GetNumViewDependentWholeSceneShadows(View);
+		uint32 NumNearAndFarCascades = GetNumViewDependentWholeSceneShadows(View, bPrecomputedLightingIsValid);
 
 		uint32 NumTotalCascades = FMath::Max(NumNearAndFarCascades, NumNearCascades + (bHasRayTracedCascade ? 1 : 0));
 
@@ -600,8 +600,8 @@ private:
 
 		// Determine start and end distances to the current cascade's split planes
 		// Presence of the ray traced cascade does not change depth ranges for the shadow-mapped cascades
-		float SplitNear = GetSplitDistance(View, ShadowSplitIndex, bIsRayTracedCascade);
-		float SplitFar = GetSplitDistance(View, ShadowSplitIndex + 1, bIsRayTracedCascade);
+		float SplitNear = GetSplitDistance(View, ShadowSplitIndex, bPrecomputedLightingIsValid, bIsRayTracedCascade);
+		float SplitFar = GetSplitDistance(View, ShadowSplitIndex + 1, bPrecomputedLightingIsValid, bIsRayTracedCascade);
 		float FadePlane = SplitFar;
 
 		float LocalCascadeTransitionFraction = CascadeTransitionFraction * GetShadowTransitionScale();
@@ -624,8 +624,8 @@ private:
 			{
 				// only used to fade the translucency lighting volume
 
-				float BeforeSplitNear = GetSplitDistance(View, ShadowSplitIndex - 1, bIsRayTracedCascade);
-				float BeforeSplitFar = GetSplitDistance(View, ShadowSplitIndex, bIsRayTracedCascade);
+				float BeforeSplitNear = GetSplitDistance(View, ShadowSplitIndex - 1, bPrecomputedLightingIsValid, bIsRayTracedCascade);
+				float BeforeSplitFar = GetSplitDistance(View, ShadowSplitIndex, bPrecomputedLightingIsValid, bIsRayTracedCascade);
 
 				OutCascadeSettings->SplitNearFadeRegion = (BeforeSplitFar - BeforeSplitNear) * LocalCascadeTransitionFraction;
 			}
diff --git a/Engine/Source/Runtime/Engine/Public/SceneManagement.h b/Engine/Source/Runtime/Engine/Public/SceneManagement.h
--- a/Engine/Source/Runtime/Engine/Public/SceneManagement.h
+++ b/Engine/Source/Runtime/Engine/Public/SceneManagement.h
@@ -670,7 +670,7 @@ public:
 	/** Accesses parameters needed for rendering the light. */
 	virtual void GetParameters(FVector4& LightPositionAndInvRadius, FVector4& LightColorAndFalloffExponent, FVector& NormalizedLightDirection, FVector2D& SpotAngles, float& LightSourceRadius, float& LightSourceLength, float& LightMinRoughness) const {}
 
-	virtual FVector2D GetDirectionalLightDistanceFadeParameters(ERHIFeatureLevel::Type InFeatureLevel) const
+	virtual FVector2D GetDirectionalLightDistanceFadeParameters(ERHIFeatureLevel::Type InFeatureLevel, bool bPrecomputedLightingIsValid) const
 	{
 		return FVector2D(0, 0);
 	}
@@ -696,14 +696,11 @@ public:
 		return false;
 	}
 
-	/** Called when precomputed lighting has been determined to be invalid */
-	virtual void InvalidatePrecomputedLighting(bool bIsEditor) {}
-
 	/** Whether this light should create per object shadows for dynamic objects. */
 	virtual bool ShouldCreatePerObjectShadowsForDynamicObjects() const;
 
 	/** Returns the number of view dependent shadows this light will create, not counting distance field shadow cascades. */
-	virtual uint32 GetNumViewDependentWholeSceneShadows(const FSceneView& View) const { return 0; }
+	virtual uint32 GetNumViewDependentWholeSceneShadows(const FSceneView& View, bool bPrecomputedLightingIsValid) const { return 0; }
 
 	/**
 	 * Sets up a projected shadow initializer that's dependent on the current view for shadows from the entire scene.
@@ -712,7 +709,8 @@ public:
 	 */
 	virtual bool GetViewDependentWholeSceneProjectedShadowInitializer(
 		const class FSceneView& View, 
-		int32 InCascadeIndex,
+		int32 InCascadeIndex, 
+		bool bPrecomputedLightingIsValid,
 		class FWholeSceneProjectedShadowInitializer& OutInitializer) const
 	{
 		return false;
@@ -743,7 +741,7 @@ public:
 
 	// @param InCascadeIndex cascade index or INDEX_NONE for the distance field cascade
 	// @param OutCascadeSettings can be 0
-	virtual FSphere GetShadowSplitBounds(const class FSceneView& View, int32 InCascadeIndex, FShadowCascadeSettings* OutCascadeSettings) const { return FSphere(FVector::ZeroVector, 0); }
+	virtual FSphere GetShadowSplitBounds(const class FSceneView& View, int32 InCascadeIndex, bool bPrecomputedLightingIsValid, FShadowCascadeSettings* OutCascadeSettings) const { return FSphere(FVector::ZeroVector, 0); }
 	virtual FSphere GetShadowSplitBoundsDepthRange(const FSceneView& View, FVector ViewOrigin, float SplitNear, float SplitFar, FShadowCascadeSettings* OutCascadeSettings) const { return FSphere(FVector::ZeroVector, 0); }
 
 	virtual bool GetScissorRect(FIntRect& ScissorRect, const FSceneView& View) const
@@ -756,7 +754,7 @@ public:
 	{
 	}
 
-	virtual bool ShouldCreateRayTracedCascade(ERHIFeatureLevel::Type Type) const { return false; }
+	virtual bool ShouldCreateRayTracedCascade(ERHIFeatureLevel::Type Type, bool bPrecomputedLightingIsValid) const { return false; }
 
 	// Accessors.
 	float GetUserShadowBias() const { return ShadowBias; }
diff --git a/Engine/Source/Runtime/Renderer/Private/LightRendering.cpp b/Engine/Source/Runtime/Renderer/Private/LightRendering.cpp
--- a/Engine/Source/Runtime/Renderer/Private/LightRendering.cpp
+++ b/Engine/Source/Runtime/Renderer/Private/LightRendering.cpp
@@ -529,7 +529,7 @@ void FDeferredShadingSceneRenderer::RenderLights(FRHICommandListImmediate& RHICm
 				bUsedLightAttenuation |= bLightFunctionRendered;
 
 				if (ViewFamily.EngineShowFlags.PreviewShadowsIndicator
-					&& !LightSceneInfo.bPrecomputedLightingIsValid 
+					&& !LightSceneInfo.IsPrecomputedLightingValid() 
 					&& LightSceneInfo.Proxy->HasStaticShadowing())
 				{
 					RenderPreviewShadowsIndicator(RHICmdList, &LightSceneInfo, bUsedLightAttenuation);
diff --git a/Engine/Source/Runtime/Renderer/Private/LightRendering.h b/Engine/Source/Runtime/Renderer/Private/LightRendering.h
--- a/Engine/Source/Runtime/Renderer/Private/LightRendering.h
+++ b/Engine/Source/Runtime/Renderer/Private/LightRendering.h
@@ -55,7 +55,7 @@ void SetDeferredLightParameters(
 	DeferredLightUniformsValue.LightColor = LightColorAndFalloffExponent;
 	DeferredLightUniformsValue.LightFalloffExponent = LightColorAndFalloffExponent.W;
 
-	const FVector2D FadeParams = LightSceneInfo->Proxy->GetDirectionalLightDistanceFadeParameters(View.GetFeatureLevel());
+	const FVector2D FadeParams = LightSceneInfo->Proxy->GetDirectionalLightDistanceFadeParameters(View.GetFeatureLevel(), LightSceneInfo->IsPrecomputedLightingValid());
 
 	// use MAD for efficiency in the shader
 	DeferredLightUniformsValue.DistanceFadeMAD = FVector2D(FadeParams.Y, -FadeParams.X * FadeParams.Y);
diff --git a/Engine/Source/Runtime/Renderer/Private/LightSceneInfo.cpp b/Engine/Source/Runtime/Renderer/Private/LightSceneInfo.cpp
--- a/Engine/Source/Runtime/Renderer/Private/LightSceneInfo.cpp
+++ b/Engine/Source/Runtime/Renderer/Private/LightSceneInfo.cpp
@@ -8,6 +8,14 @@
 #include "ScenePrivate.h"
 #include "DistanceFieldLightingShared.h"
 
+int32 GWholeSceneShadowUnbuiltInteractionThreshold = 500;
+static FAutoConsoleVariableRef CVarWholeSceneShadowUnbuiltInteractionThreshold(
+	TEXT("r.Shadow.WholeSceneShadowUnbuiltInteractionThreshold"),
+	GWholeSceneShadowUnbuiltInteractionThreshold,
+	TEXT("How many unbuilt light-primitive interactions there can be for a light before the light switches to whole scene shadows"),
+	ECVF_RenderThreadSafe
+	);
+
 void FLightSceneInfoCompact::Init(FLightSceneInfo* InLightSceneInfo)
 {
 	LightSceneInfo = InLightSceneInfo;
@@ -44,11 +52,6 @@ FLightSceneInfo::FLightSceneInfo(FLightSceneProxy* InProxy, bool InbVisible)
 	// Only visible lights can be added in game
 	check(bVisible || GIsEditor);
 
-	if (!bPrecomputedLightingIsValid)
-	{
-		Proxy->InvalidatePrecomputedLighting(InProxy->GetLightComponent()->GetScene()->IsEditorScene());
-	}
-
 	BeginInitResource(this);
 
 	for(uint32 LightTypeIndex = 0;LightTypeIndex < LightType_MAX;++LightTypeIndex)
@@ -176,6 +179,11 @@ bool FLightSceneInfo::ShouldRenderLight(const FViewInfo& View) const
 		&& (!View.bStaticSceneOnly || Proxy->HasStaticShadowing());
 }
 
+bool FLightSceneInfo::IsPrecomputedLightingValid() const
+{
+	return (bPrecomputedLightingIsValid && NumUnbuiltInteractions < GWholeSceneShadowUnbuiltInteractionThreshold) || !Proxy->HasStaticShadowing();
+}
+
 void FLightSceneInfo::ReleaseRHI()
 {
 	for(uint32 LightTypeIndex = 0;LightTypeIndex < LightType_MAX;++LightTypeIndex)
diff --git a/Engine/Source/Runtime/Renderer/Private/LightSceneInfo.h b/Engine/Source/Runtime/Renderer/Private/LightSceneInfo.h
--- a/Engine/Source/Runtime/Renderer/Private/LightSceneInfo.h
+++ b/Engine/Source/Runtime/Renderer/Private/LightSceneInfo.h
@@ -121,9 +121,12 @@ public:
 	/** Tile intersection buffer for distance field shadowing, stored on the light to avoid reallocating each frame. */
 	mutable TScopedPointer<class FLightTileIntersectionResources> TileIntersectionResources;
 
+protected:
 	/** True if the light is built. */
 	uint32 bPrecomputedLightingIsValid : 1;
 
+public:
+
 	/** 
 	 * True if the light is visible.  
 	 * False if the light is invisible but still needed for previewing, which can only happen in the editor.
@@ -192,7 +195,7 @@ public:
 	{
 		return !Proxy->GetColor().IsAlmostBlack()
 			// Only render lights with dynamic lighting or unbuilt static lights
-			&& (!Proxy->HasStaticLighting() || !bPrecomputedLightingIsValid);
+			&& (!Proxy->HasStaticLighting() || !IsPrecomputedLightingValid());
 	}
 
 	/** Encapsulates all View-Independent reasons to render ViewIndependentWholeSceneShadows for this light */
@@ -205,12 +208,14 @@ public:
 		const bool bCreateShadowToPreviewStaticLight =
 			Proxy->HasStaticShadowing()
 			&& bCastDynamicShadow
-			&& !bPrecomputedLightingIsValid;
+			&& !IsPrecomputedLightingValid();
 
 		bool bShouldRenderShadow = bShouldRenderLight && bCastDynamicShadow && (!Proxy->HasStaticLighting() || bCreateShadowToPreviewStaticLight);
 		return bShouldRenderShadow;
 	}
 
+	bool IsPrecomputedLightingValid() const;
+
 	/** Hash function. */
 	friend uint32 GetTypeHash(const FLightSceneInfo* LightSceneInfo)
 	{
diff --git a/Engine/Source/Runtime/Renderer/Private/SceneCore.cpp b/Engine/Source/Runtime/Renderer/Private/SceneCore.cpp
--- a/Engine/Source/Runtime/Renderer/Private/SceneCore.cpp
+++ b/Engine/Source/Runtime/Renderer/Private/SceneCore.cpp
@@ -8,12 +8,6 @@
 #include "ScenePrivate.h"
 #include "AllocatorFixedSizeFreeList.h"
 
-static TAutoConsoleVariable<int32> CVarWholeSceneShadowUnbuiltInteractionThreshold(
-	TEXT("r.Shadow.WholeSceneShadowUnbuiltInteractionThreshold"),
-	500,
-	TEXT("How many unbuilt light-primitive interactions there can be for a light before the light switches to whole scene shadows"),
-	ECVF_RenderThreadSafe);
-
 /**
  * Fixed Size pool allocator for FLightPrimitiveInteractions
  */
@@ -125,17 +119,6 @@ void FLightPrimitiveInteraction::Create(FLightSceneInfo* LightSceneInfo,FPrimiti
 		{
 			// Create the light interaction.
 			FLightPrimitiveInteraction* Interaction = new FLightPrimitiveInteraction(LightSceneInfo, PrimitiveSceneInfo, bDynamic, bLightMapped, bShadowMapped, bTranslucentObjectShadow, bInsetObjectShadow);
-
-#if !(UE_BUILD_SHIPPING || UE_BUILD_TEST)
-			// Treat the light as completely unbuilt if it has more unbuilt interactions than the threshold.
-			// This will result in the light using whole scene shadows instead of many per-object shadows, 
-			// Which prevents poor performance when many per-object shadows are created for previewing unbuilt lighting.
-			if (LightSceneInfo->NumUnbuiltInteractions >= CVarWholeSceneShadowUnbuiltInteractionThreshold.GetValueOnRenderThread() && LightSceneInfo->bPrecomputedLightingIsValid)
-			{
-				LightSceneInfo->bPrecomputedLightingIsValid = false;
-				LightSceneInfo->Proxy->InvalidatePrecomputedLighting(true);
-			}
-#endif
 		}
 	}
 }
diff --git a/Engine/Source/Runtime/Renderer/Private/SceneRendering.cpp b/Engine/Source/Runtime/Renderer/Private/SceneRendering.cpp
--- a/Engine/Source/Runtime/Renderer/Private/SceneRendering.cpp
+++ b/Engine/Source/Runtime/Renderer/Private/SceneRendering.cpp
@@ -718,7 +718,7 @@ void FViewInfo::CreateForwardLightDataUniformBuffer(FForwardLightData &OutForwar
 
 			if(!LightSceneInfoCompact.Color.IsAlmostBlack()
 				// Only render lights with dynamic lighting or unbuilt static lights
-				&& (!LightSceneInfo->Proxy->HasStaticLighting() || !LightSceneInfo->bPrecomputedLightingIsValid))
+				&& (!LightSceneInfo->Proxy->HasStaticLighting() || !LightSceneInfo->IsPrecomputedLightingValid()))
 			{
 				const ELightComponentType LightType = (const ELightComponentType)LightSceneInfoCompact.LightType;
 
diff --git a/Engine/Source/Runtime/Renderer/Private/ShadowSetup.cpp b/Engine/Source/Runtime/Renderer/Private/ShadowSetup.cpp
--- a/Engine/Source/Runtime/Renderer/Private/ShadowSetup.cpp
+++ b/Engine/Source/Runtime/Renderer/Private/ShadowSetup.cpp
@@ -595,7 +595,12 @@ void FProjectedShadowInfo::SetupWholeSceneProjection(
 		if (CascadeSettings.ShadowSplitIndex >= 0 && bDirectionalLight)
 		{
 			checkSlow(InDependentView);
-			ShadowBounds = InLightSceneInfo->Proxy->GetShadowSplitBounds(*InDependentView, CascadeSettings.bRayTracedDistanceField ? INDEX_NONE : CascadeSettings.ShadowSplitIndex, 0);
+
+			ShadowBounds = InLightSceneInfo->Proxy->GetShadowSplitBounds(
+				*InDependentView, 
+				CascadeSettings.bRayTracedDistanceField ? INDEX_NONE : CascadeSettings.ShadowSplitIndex, 
+				InLightSceneInfo->IsPrecomputedLightingValid(), 
+				0);
 		}
 		else
 		{
@@ -1144,7 +1149,7 @@ bool FSceneRenderer::ShouldCreateObjectShadowForStationaryLight(const FLightScen
 {
 	const bool bCreateObjectShadowForStationaryLight = 
 		LightSceneInfo->bCreatePerObjectShadowsForDynamicObjects 
-		&& LightSceneInfo->bPrecomputedLightingIsValid
+		&& LightSceneInfo->IsPrecomputedLightingValid()
 		&& LightSceneInfo->Proxy->GetShadowMapChannel() != INDEX_NONE
 		// Create a per-object shadow if the object does not want static lighting and needs to integrate with the static shadowing of a stationary light
 		// Or if the object wants static lighting but does not have a built shadowmap (Eg has been moved in the editor)
@@ -1463,7 +1468,7 @@ void FDeferredShadingSceneRenderer::CreatePerObjectProjectedShadow(
 			for (int32 i = 0; i < ViewDependentWholeSceneShadows.Num(); i++)
 			{
 				const FProjectedShadowInfo* WholeSceneShadow = ViewDependentWholeSceneShadows[i];
-				const FVector2D DistanceFadeValues = WholeSceneShadow->GetLightSceneInfo().Proxy->GetDirectionalLightDistanceFadeParameters(Scene->GetFeatureLevel());
+				const FVector2D DistanceFadeValues = WholeSceneShadow->GetLightSceneInfo().Proxy->GetDirectionalLightDistanceFadeParameters(Scene->GetFeatureLevel(), WholeSceneShadow->GetLightSceneInfo().IsPrecomputedLightingValid());
 				const float DistanceFromShadowCenterSquared = (WholeSceneShadow->ShadowBounds.Center - Bounds.Origin).SizeSquared();
 				//@todo - if view dependent whole scene shadows are ever supported in splitscreen, 
 				// We can only disable the preshadow at this point if it is inside a whole scene shadow for all views
@@ -1950,7 +1955,7 @@ inline void FSceneRenderer::GatherShadowsForPrimitiveInner(
 
 				// Include all primitives for movable lights, but only statically shadowed primitives from a light with static shadowing,
 				// Since lights with static shadowing still create per-object shadows for primitives without static shadowing.
-				if( (!LightProxy->HasStaticLighting() || !ProjectedShadowInfo->GetLightSceneInfo().bPrecomputedLightingIsValid)
+				if( (!LightProxy->HasStaticLighting() || !ProjectedShadowInfo->GetLightSceneInfo().IsPrecomputedLightingValid())
 					// Check if this primitive is in the shadow's cylinder
 					&& PrimitiveDistanceFromCylinderAxisSq < FMath::Square(ProjectedShadowInfo->ShadowBounds.W + PrimitiveBounds.SphereRadius)
 					// Check if the primitive is closer than the cylinder cap toward the light
@@ -2129,9 +2134,9 @@ void FSceneRenderer::AddViewDependentWholeSceneShadowsForView(
 		// If rendering in stereo mode we render shadow depths only for the left eye, but project for both eyes!
 		if (View.StereoPass != eSSP_RIGHT_EYE)
 		{
-			const bool bExtraDistanceFieldCascade = LightSceneInfo.Proxy->ShouldCreateRayTracedCascade(View.GetFeatureLevel());
+			const bool bExtraDistanceFieldCascade = LightSceneInfo.Proxy->ShouldCreateRayTracedCascade(View.GetFeatureLevel(), LightSceneInfo.IsPrecomputedLightingValid());
 
-			const int32 ProjectionCount = LightSceneInfo.Proxy->GetNumViewDependentWholeSceneShadows(View) + (bExtraDistanceFieldCascade?1:0);
+			const int32 ProjectionCount = LightSceneInfo.Proxy->GetNumViewDependentWholeSceneShadows(View, LightSceneInfo.IsPrecomputedLightingValid()) + (bExtraDistanceFieldCascade?1:0);
 
 			checkSlow(INDEX_NONE == -1);
 
@@ -2148,7 +2153,7 @@ void FSceneRenderer::AddViewDependentWholeSceneShadowsForView(
 					LocalIndex = INDEX_NONE;
 				}
 
-				if (LightSceneInfo.Proxy->GetViewDependentWholeSceneProjectedShadowInitializer(View, LocalIndex, ProjectedShadowInitializer))
+				if (LightSceneInfo.Proxy->GetViewDependentWholeSceneProjectedShadowInitializer(View, LocalIndex, LightSceneInfo.IsPrecomputedLightingValid(), ProjectedShadowInitializer))
 				{
 					const FIntPoint ShadowBufferResolution = GSceneRenderTargets.GetShadowDepthTextureResolution();
 					// Create the projected shadow info.
@@ -2371,14 +2376,14 @@ void FDeferredShadingSceneRenderer::InitDynamicShadows(FRHICommandListImmediate&
 					const bool bCreateShadowToPreviewStaticLight = 
 						LightSceneInfo->Proxy->HasStaticShadowing()
 						&& LightSceneInfoCompact.bCastStaticShadow 
-						&& !LightSceneInfo->bPrecomputedLightingIsValid;
+						&& !LightSceneInfo->IsPrecomputedLightingValid();
 
 					// Create a whole scene shadow for lights that want static shadowing but didn't get assigned to a valid shadowmap channel due to overlap
 					const bool bCreateShadowForOverflowStaticShadowing =
 						LightSceneInfo->Proxy->HasStaticShadowing()
 						&& !LightSceneInfo->Proxy->HasStaticLighting()
 						&& LightSceneInfoCompact.bCastStaticShadow 
-						&& LightSceneInfo->bPrecomputedLightingIsValid
+						&& LightSceneInfo->IsPrecomputedLightingValid()
 						&& LightSceneInfo->Proxy->GetShadowMapChannel() == INDEX_NONE;
 
 					if (bCreateShadowForMovableLight || bCreateShadowToPreviewStaticLight || bCreateShadowForOverflowStaticShadowing)
diff --git a/Engine/Source/Runtime/Renderer/Private/TranslucentLighting.cpp b/Engine/Source/Runtime/Renderer/Private/TranslucentLighting.cpp
--- a/Engine/Source/Runtime/Renderer/Private/TranslucentLighting.cpp
+++ b/Engine/Source/Runtime/Renderer/Private/TranslucentLighting.cpp
@@ -863,7 +863,7 @@ public:
 			FShadowCascadeSettings ShadowCascadeSettings;
 
 			// todo: optimize, not all computed data is needed
-			InnerBounds = LightSceneInfo->Proxy->GetShadowSplitBounds(View, (uint32)InnerSplitIndex, &ShadowCascadeSettings);
+			InnerBounds = LightSceneInfo->Proxy->GetShadowSplitBounds(View, (uint32)InnerSplitIndex, LightSceneInfo->IsPrecomputedLightingValid(), &ShadowCascadeSettings);
 
 			// near cascade plane
 			{
@@ -871,7 +871,7 @@ public:
 				Planes[0] = FVector4((FVector)(ShadowCascadeSettings.NearFrustumPlane),  -ShadowCascadeSettings.NearFrustumPlane.W);
 			}
 
-			uint32 CascadeCount = LightSceneInfo->Proxy->GetNumViewDependentWholeSceneShadows(View);
+			uint32 CascadeCount = LightSceneInfo->Proxy->GetNumViewDependentWholeSceneShadows(View, LightSceneInfo->IsPrecomputedLightingValid());
 
 			// far cascade plane
 			if(InnerSplitIndex != CascadeCount - 1)
@@ -880,7 +880,7 @@ public:
 				Planes[1] = FVector4((FVector)(ShadowCascadeSettings.FarFrustumPlane), -ShadowCascadeSettings.FarFrustumPlane.W);
 			}
 
-			const FVector2D FadeParams = LightSceneInfo->Proxy->GetDirectionalLightDistanceFadeParameters(View.GetFeatureLevel());
+			const FVector2D FadeParams = LightSceneInfo->Proxy->GetDirectionalLightDistanceFadeParameters(View.GetFeatureLevel(), LightSceneInfo->IsPrecomputedLightingValid());
 
 			// setup constants for the MAD in shader
 			ShadowInjectParamValue.Z = FadeParams.Y;
@@ -929,7 +929,7 @@ public:
 		}
 
 		const FStaticShadowDepthMap* StaticShadowDepthMap = LightSceneInfo->Proxy->GetStaticShadowDepthMap();
-		const uint32 bStaticallyShadowedValue = LightSceneInfo->bPrecomputedLightingIsValid && StaticShadowDepthMap && StaticShadowDepthMap->TextureRHI ? 1 : 0;
+		const uint32 bStaticallyShadowedValue = LightSceneInfo->IsPrecomputedLightingValid() && StaticShadowDepthMap && StaticShadowDepthMap->TextureRHI ? 1 : 0;
 		FTextureRHIParamRef StaticShadowDepthMapTexture = bStaticallyShadowedValue ? StaticShadowDepthMap->TextureRHI : GWhiteTexture->TextureRHI;
 		const FMatrix WorldToStaticShadow = bStaticallyShadowedValue ? StaticShadowDepthMap->WorldToLight : FMatrix::Identity;
 
