Fix build for vs2012. Removed default lock template argument by using free functions wrapped around static functions

diff --git a/Engine/Source/Runtime/Engine/Private/PhysicsEngine/BodyInstance.cpp b/Engine/Source/Runtime/Engine/Private/PhysicsEngine/BodyInstance.cpp
--- a/Engine/Source/Runtime/Engine/Private/PhysicsEngine/BodyInstance.cpp
+++ b/Engine/Source/Runtime/Engine/Private/PhysicsEngine/BodyInstance.cpp
@@ -2323,13 +2323,11 @@ FTransform GetUnrealWorldTransformImp(const FBodyInstance* BodyInstance)
 {
 	FTransform WorldTM = FTransform::Identity;
 #if WITH_PHYSX
-	auto Func = [&](const PxRigidActor* PActor)
+	FPhysXSupport<NeedsLock>::ExecuteOnPxRigidActorReadOnly(BodyInstance, [&](const PxRigidActor* PActor)
 	{
 		PxTransform PTM = PActor->getGlobalPose();
 		WorldTM = P2UTransform(PTM);
-	};
-
-	ExecuteOnPxRigidActorReadOnly<decltype(Func), NeedsLock>(BodyInstance, Func);
+	});
 #endif // WITH_PHYSX
 
 #if WITH_BOX2D
@@ -2446,12 +2444,10 @@ FVector GetUnrealWorldVelocityImp(const FBodyInstance* BodyInstance)
 	FVector LinVel(EForceInit::ForceInitToZero);
 
 #if WITH_PHYSX
-	auto Func = [&](const PxRigidBody* PRigidBody)
+	FPhysXSupport<NeedsLock>::ExecuteOnPxRigidBodyReadOnly(BodyInstance, [&](const PxRigidBody* PRigidBody)
 	{
 		LinVel = P2UVector(PRigidBody->getLinearVelocity());
-	};
-
-	ExecuteOnPxRigidBodyReadOnly<decltype(Func), NeedsLock>(BodyInstance, Func);
+	});
 #endif // WITH_PHYSX
 
 #if WITH_BOX2D
@@ -2480,12 +2476,10 @@ FVector GetUnrealWorldAngularVelocityImp(const FBodyInstance* BodyInstance)
 	FVector AngVel(EForceInit::ForceInitToZero);
 
 #if WITH_PHYSX
-	auto Func = [&](const PxRigidBody* PRigidBody)
+	FPhysXSupport<NeedsLock>::ExecuteOnPxRigidBodyReadOnly(BodyInstance, [&](const PxRigidBody* PRigidBody)
 	{
 		AngVel = FVector::RadiansToDegrees(P2UVector(PRigidBody->getAngularVelocity()));
-	};
-
-	ExecuteOnPxRigidBodyReadOnly<decltype(Func), NeedsLock>(BodyInstance, Func);
+	});
 #endif // WITH_PHYSX
 
 #if WITH_BOX2D
@@ -2516,13 +2510,11 @@ FVector GetUnrealWorldVelocityAtPointImp(const FBodyInstance* BodyInstance, cons
 	FVector LinVel(EForceInit::ForceInitToZero);
 
 #if WITH_PHYSX
-	auto Func = [&](const PxRigidBody* PRigidBody)
+	FPhysXSupport<NeedsLock>::ExecuteOnPxRigidBodyReadOnly(BodyInstance, [&](const PxRigidBody* PRigidBody)
 	{
 		PxVec3 PPoint = U2PVector(Point);
 		LinVel = P2UVector(PxRigidBodyExt::getVelocityAtPos(*PRigidBody, PPoint));
-	};
-
-	ExecuteOnPxRigidBodyReadOnly<decltype(Func), NeedsLock>(BodyInstance, Func);
+	});
 #endif // WITH_PHYSX
 
 #if WITH_BOX2D
diff --git a/Engine/Source/Runtime/Engine/Private/PhysicsEngine/PhysXSupport.h b/Engine/Source/Runtime/Engine/Private/PhysicsEngine/PhysXSupport.h
--- a/Engine/Source/Runtime/Engine/Private/PhysicsEngine/PhysXSupport.h
+++ b/Engine/Source/Runtime/Engine/Private/PhysicsEngine/PhysXSupport.h
@@ -106,163 +106,174 @@ PxScene* GetPhysXSceneFromIndex(int32 InSceneIndex);
 NxApexScene* GetApexSceneFromIndex(int32 InSceneIndex);
 #endif
 
-/** Obtains the appropriate PhysX scene lock for READING and executes the passed in lambda.
- *  If SceneType < 0, the Sync actor is used, otherwise the async.
- *  Note: The lambda is only executed if the physx actor requested is non-null.
- *  returns true if the requested actor is non-null
- */
-template <typename LambdaType, bool NeedsLock = true>
-bool ExecuteOnPxRigidActorReadOnly(const FBodyInstance* BI, const LambdaType& Func, int32 SceneType = -1)
+template <bool NeedsLock>
+struct FPhysXSupport
 {
-	if (const PxRigidActor* PRigidActor = BI->GetPxRigidActor_AssumesLocked(SceneType))
-	{
-		const int32 SceneIndex = (PRigidActor == BI->RigidActorSync ? BI->SceneIndexSync : BI->SceneIndexAsync);
-		PxScene* PScene = GetPhysXSceneFromIndex(SceneIndex);
-		if (NeedsLock)
+	/** Obtains the appropriate PhysX scene lock for READING and executes the passed in lambda.
+	 *  If SceneType < 0, the Sync actor is used, otherwise the async.
+	 *  Note: The lambda is only executed if the physx actor requested is non-null.
+	 *  returns true if the requested actor is non-null
+	 */
+	template <typename LambdaType>
+	static bool ExecuteOnPxRigidActorReadOnly(const FBodyInstance* BI, const LambdaType& Func, int32 SceneType = -1)
+	{
+		if (const PxRigidActor* PRigidActor = BI->GetPxRigidActor_AssumesLocked(SceneType))
 		{
-			SCENE_LOCK_READ(PScene);
-		}
+			const int32 SceneIndex = (PRigidActor == BI->RigidActorSync ? BI->SceneIndexSync : BI->SceneIndexAsync);
+			PxScene* PScene = GetPhysXSceneFromIndex(SceneIndex);
+			if (NeedsLock)
+			{
+				SCENE_LOCK_READ(PScene);
+			}
 
-		Func(PRigidActor);
+			Func(PRigidActor);
 
-		if (NeedsLock)
-		{
-			SCENE_UNLOCK_READ(PScene);
+			if (NeedsLock)
+			{
+				SCENE_UNLOCK_READ(PScene);
+			}
+
+			return true;
 		}
 
-		return true;
+		return false;
 	}
 
-	return false;
-}
-
-/** Obtains the appropriate PhysX scene lock for READING and executes the passed in lambda.
- *  Note: The lambda is only executed if the physx actor is a non-null RigidBody
- *  returns true if found a non-null RigidBody
- */
-template <typename LambdaType, bool NeedsLock = true>
-bool ExecuteOnPxRigidBodyReadOnly(const FBodyInstance* BI, const LambdaType& Func)
-{
-	bool bSuccess = false;
-	if (const physx::PxRigidActor* RigidActor = BI->GetPxRigidActor_AssumesLocked())
+	/** Obtains the appropriate PhysX scene lock for READING and executes the passed in lambda.
+	 *  Note: The lambda is only executed if the physx actor is a non-null RigidBody
+	 *  returns true if found a non-null RigidBody
+	 */
+	template <typename LambdaType>
+	static bool ExecuteOnPxRigidBodyReadOnly(const FBodyInstance* BI, const LambdaType& Func)
 	{
-		const int32 SceneIndex = (RigidActor == BI->RigidActorSync ? BI->SceneIndexSync : BI->SceneIndexAsync);
-		PxScene* PScene = GetPhysXSceneFromIndex(SceneIndex);
-		if (NeedsLock)
+		bool bSuccess = false;
+		if (const physx::PxRigidActor* RigidActor = BI->GetPxRigidActor_AssumesLocked())
 		{
-			SCENE_LOCK_READ(PScene);
-		}
+			const int32 SceneIndex = (RigidActor == BI->RigidActorSync ? BI->SceneIndexSync : BI->SceneIndexAsync);
+			PxScene* PScene = GetPhysXSceneFromIndex(SceneIndex);
+			if (NeedsLock)
+			{
+				SCENE_LOCK_READ(PScene);
+			}
 
-		if (const physx::PxRigidBody* PRigidBody = RigidActor->isRigidBody())
-		{
-			Func(PRigidBody);
-			bSuccess = true;
-		}
+			if (const physx::PxRigidBody* PRigidBody = RigidActor->isRigidBody())
+			{
+				Func(PRigidBody);
+				bSuccess = true;
+			}
 
-		if (NeedsLock)
-		{
-			SCENE_UNLOCK_READ(PScene);
+			if (NeedsLock)
+			{
+				SCENE_UNLOCK_READ(PScene);
+			}
 		}
-	}
 
-	return bSuccess;
-}
+		return bSuccess;
+	}
 
-/** Obtains the appropriate PhysX scene lock for WRITING and executes the passed in lambda.
- *  Note: The lambda is only executed if the physx actor is a non-null RigidBody
- *  returns true if found a non-null RigidBody.
- */
-template <typename LambdaType, bool NeedsLock = true>
-bool ExecuteOnPxRigidBodyReadWrite(const FBodyInstance* BI, const LambdaType& Func)
-{
-	bool bSuccess = false;
-	if (physx::PxRigidActor* RigidActor = BI->GetPxRigidActor_AssumesLocked())
+	/** Obtains the appropriate PhysX scene lock for WRITING and executes the passed in lambda.
+	 *  Note: The lambda is only executed if the physx actor is a non-null RigidBody
+	 *  returns true if found a non-null RigidBody.
+	 */
+	template <typename LambdaType, bool NeedsLock = true>
+	static bool ExecuteOnPxRigidBodyReadWrite(const FBodyInstance* BI, const LambdaType& Func)
 	{
-		const int32 SceneIndex = (RigidActor == BI->RigidActorSync ? BI->SceneIndexSync : BI->SceneIndexAsync);
-		PxScene* PScene = GetPhysXSceneFromIndex(SceneIndex);
-		if(NeedsLock)
+		bool bSuccess = false;
+		if (physx::PxRigidActor* RigidActor = BI->GetPxRigidActor_AssumesLocked())
 		{
-			SCENE_LOCK_WRITE(PScene);
-		}
+			const int32 SceneIndex = (RigidActor == BI->RigidActorSync ? BI->SceneIndexSync : BI->SceneIndexAsync);
+			PxScene* PScene = GetPhysXSceneFromIndex(SceneIndex);
+			if(NeedsLock)
+			{
+				SCENE_LOCK_WRITE(PScene);
+			}
 			
-		if (physx::PxRigidBody* PRigidBody = RigidActor->isRigidBody())
-		{
-			Func(PRigidBody);
-			bSuccess = true;
-		}
+			if (physx::PxRigidBody* PRigidBody = RigidActor->isRigidBody())
+			{
+				Func(PRigidBody);
+				bSuccess = true;
+			}
 
-		if(NeedsLock)
-		{
-			SCENE_UNLOCK_WRITE(PScene);
+			if(NeedsLock)
+			{
+				SCENE_UNLOCK_WRITE(PScene);
+			}
 		}
-	}
 
-	return bSuccess;
-}
+		return bSuccess;
+	}
 
-/** Obtains the appropriate PhysX scene lock for READING and executes the passed in lambda.
- *  Note: The lambda is only executed if the physx actor is a non-null RigidDynamic
- *  returns true if found a non-null RigidDynamic.
- */
-template <typename LambdaType, bool NeedsLock = true>
-bool ExecuteOnPxRigidDynamicReadOnly(const FBodyInstance* BI, const LambdaType& Func)
-{
-	bool bSuccess = false;
-	if (physx::PxRigidActor* RigidActor = BI->GetPxRigidActor_AssumesLocked())
+	/** Obtains the appropriate PhysX scene lock for READING and executes the passed in lambda.
+	 *  Note: The lambda is only executed if the physx actor is a non-null RigidDynamic
+	 *  returns true if found a non-null RigidDynamic.
+	 */
+	template <typename LambdaType, bool NeedsLock = true>
+	static bool ExecuteOnPxRigidDynamicReadOnly(const FBodyInstance* BI, const LambdaType& Func)
 	{
-		const int32 SceneIndex = (RigidActor == BI->RigidActorSync ? BI->SceneIndexSync : BI->SceneIndexAsync);
-		PxScene* PScene = GetPhysXSceneFromIndex(SceneIndex);
-		if (NeedsLock)
+		bool bSuccess = false;
+		if (physx::PxRigidActor* RigidActor = BI->GetPxRigidActor_AssumesLocked())
 		{
-			SCENE_LOCK_READ(PScene);
-		}
+			const int32 SceneIndex = (RigidActor == BI->RigidActorSync ? BI->SceneIndexSync : BI->SceneIndexAsync);
+			PxScene* PScene = GetPhysXSceneFromIndex(SceneIndex);
+			if (NeedsLock)
+			{
+				SCENE_LOCK_READ(PScene);
+			}
 
-		if (physx::PxRigidDynamic* PRigidDynamic = RigidActor->isRigidDynamic())
-		{
-			Func(PRigidDynamic);
-			bSuccess = true;
-		}
+			if (physx::PxRigidDynamic* PRigidDynamic = RigidActor->isRigidDynamic())
+			{
+				Func(PRigidDynamic);
+				bSuccess = true;
+			}
 
-		if (NeedsLock)
-		{
-			SCENE_UNLOCK_READ(PScene);
+			if (NeedsLock)
+			{
+				SCENE_UNLOCK_READ(PScene);
+			}
 		}
-	}
 
-	return bSuccess;
-}
+		return bSuccess;
+	}
 
-/** Obtains the appropriate PhysX scene lock for WRITING and executes the passed in lambda.
- *  Note: The lambda is only executed if the physx actor is a non-null RigidDynamic
- *  returns true if found a non-null RigidDynamic.
- */
-template <typename LambdaType, bool NeedsLock = true>
-bool ExecuteOnPxRigidDynamicReadWrite(const FBodyInstance* BI, const LambdaType& Func)
-{
-	bool bSuccess = false;
-	if (physx::PxRigidActor* RigidActor = BI->GetPxRigidActor_AssumesLocked())
+	/** Obtains the appropriate PhysX scene lock for WRITING and executes the passed in lambda.
+	 *  Note: The lambda is only executed if the physx actor is a non-null RigidDynamic
+	 *  returns true if found a non-null RigidDynamic.
+	 */
+	template <typename LambdaType, bool NeedsLock = true>
+	static bool ExecuteOnPxRigidDynamicReadWrite(const FBodyInstance* BI, const LambdaType& Func)
 	{
-		const int32 SceneIndex = (RigidActor == BI->RigidActorSync ? BI->SceneIndexSync : BI->SceneIndexAsync);
-		PxScene* PScene = GetPhysXSceneFromIndex(SceneIndex);
-		if (NeedsLock)
+		bool bSuccess = false;
+		if (physx::PxRigidActor* RigidActor = BI->GetPxRigidActor_AssumesLocked())
 		{
-			SCENE_LOCK_WRITE(PScene);
-		}
+			const int32 SceneIndex = (RigidActor == BI->RigidActorSync ? BI->SceneIndexSync : BI->SceneIndexAsync);
+			PxScene* PScene = GetPhysXSceneFromIndex(SceneIndex);
+			if (NeedsLock)
+			{
+				SCENE_LOCK_WRITE(PScene);
+			}
 
-		if (physx::PxRigidDynamic* PRigidDynamic = RigidActor->isRigidDynamic())
-		{
-			Func(PRigidDynamic);
-			bSuccess = true;
-		}
+			if (physx::PxRigidDynamic* PRigidDynamic = RigidActor->isRigidDynamic())
+			{
+				Func(PRigidDynamic);
+				bSuccess = true;
+			}
 
-		if (NeedsLock)
-		{
-			SCENE_UNLOCK_WRITE(PScene);
+			if (NeedsLock)
+			{
+				SCENE_UNLOCK_WRITE(PScene);
+			}
 		}
+
+		return bSuccess;
 	}
+};
 
-	return bSuccess;
-}
+// Utility functions for obtaining locks and executing lambda. This indirection is needed for vs2012 but should be inlined
+template <typename LambdaType> bool ExecuteOnPxRigidActorReadOnly(const FBodyInstance* BI, const LambdaType& Func, int32 SceneType = -1){ return FPhysXSupport<true>::ExecuteOnPxRigidActorReadOnly(BI, Func, SceneType); }
+template <typename LambdaType> bool ExecuteOnPxRigidBodyReadOnly(const FBodyInstance* BI, const LambdaType& Func) { return FPhysXSupport<true>::ExecuteOnPxRigidBodyReadOnly(BI, Func); }
+template <typename LambdaType> bool ExecuteOnPxRigidBodyReadWrite(const FBodyInstance* BI, const LambdaType& Func){ return FPhysXSupport<true>::ExecuteOnPxRigidBodyReadWrite(BI, Func); }
+template <typename LambdaType> bool ExecuteOnPxRigidDynamicReadOnly(const FBodyInstance* BI, const LambdaType& Func){ return FPhysXSupport<true>::ExecuteOnPxRigidDynamicReadOnly(BI, Func); }
+template <typename LambdaType> bool ExecuteOnPxRigidDynamicReadWrite(const FBodyInstance* BI, const LambdaType& Func){ return FPhysXSupport<true>::ExecuteOnPxRigidDynamicReadWrite(BI, Func); }
 
 //////// BASIC TYPE CONVERSION
 
