Fix mac build. Forgot to remove some of the NeedsLocked template

diff --git a/Engine/Source/Runtime/Engine/Private/PhysicsEngine/PhysXSupport.h b/Engine/Source/Runtime/Engine/Private/PhysicsEngine/PhysXSupport.h
--- a/Engine/Source/Runtime/Engine/Private/PhysicsEngine/PhysXSupport.h
+++ b/Engine/Source/Runtime/Engine/Private/PhysicsEngine/PhysXSupport.h
@@ -175,7 +175,7 @@ struct FPhysXSupport
 	 *  Note: The lambda is only executed if the physx actor is a non-null RigidBody
 	 *  returns true if found a non-null RigidBody.
 	 */
-	template <typename LambdaType, bool NeedsLock = true>
+	template <typename LambdaType>
 	static bool ExecuteOnPxRigidBodyReadWrite(const FBodyInstance* BI, const LambdaType& Func)
 	{
 		bool bSuccess = false;
@@ -207,7 +207,7 @@ struct FPhysXSupport
 	 *  Note: The lambda is only executed if the physx actor is a non-null RigidDynamic
 	 *  returns true if found a non-null RigidDynamic.
 	 */
-	template <typename LambdaType, bool NeedsLock = true>
+	template <typename LambdaType>
 	static bool ExecuteOnPxRigidDynamicReadOnly(const FBodyInstance* BI, const LambdaType& Func)
 	{
 		bool bSuccess = false;
@@ -239,7 +239,7 @@ struct FPhysXSupport
 	 *  Note: The lambda is only executed if the physx actor is a non-null RigidDynamic
 	 *  returns true if found a non-null RigidDynamic.
 	 */
-	template <typename LambdaType, bool NeedsLock = true>
+	template <typename LambdaType>
 	static bool ExecuteOnPxRigidDynamicReadWrite(const FBodyInstance* BI, const LambdaType& Func)
 	{
 		bool bSuccess = false;
