Fix for nonunity compilation :/

diff --git a/Engine/Source/Runtime/Engine/Private/Camera/CameraShake.cpp b/Engine/Source/Runtime/Engine/Private/Camera/CameraShake.cpp
--- a/Engine/Source/Runtime/Engine/Private/Camera/CameraShake.cpp
+++ b/Engine/Source/Runtime/Engine/Private/Camera/CameraShake.cpp
@@ -2,6 +2,7 @@
 
 #include "EnginePrivate.h"
 #include "Camera/PlayerCameraManager.h"
+#include "Camera/CameraAnimInst.h"
 
 //////////////////////////////////////////////////////////////////////////
 // FFOscillator
