Added new 'interpto' component

diff --git a/Engine/Source/Runtime/Engine/Classes/Components/InterpToMovementComponent.h b/Engine/Source/Runtime/Engine/Classes/Components/InterpToMovementComponent.h
new file mode 100644
--- /dev/null
+++ b/Engine/Source/Runtime/Engine/Classes/Components/InterpToMovementComponent.h
@@ -0,0 +1,236 @@
+// Copyright 1998-2015 Epic Games, Inc. All Rights Reserved.
+
+#pragma once
+#include "GameFramework/MovementComponent.h"
+#include "Components/SplineComponent.h"
+#include "InterpToMovementComponent.generated.h"
+
+
+/**
+ * Controls the movement behaviour
+ */
+UENUM(BlueprintType)
+enum class EInterpToBehaviourType : uint8
+{
+	/** Move to destination and stop */
+	OneShot,
+	/** Move to destination, return and stop */
+	OneShot_Reverse,
+	/** Reset back to start when reaching the end */
+	Loop_Reset,
+	/** Repeatedly lerp from start to end and back */
+	PingPong
+};
+
+USTRUCT()
+struct FInterpControlPoint
+{
+	GENERATED_USTRUCT_BODY()
+public:
+	// The position of the point 
+	UPROPERTY(EditAnywhere, BlueprintReadWrite, Category=Location, meta=(MakeEditWidget=true))
+	FVector PositionControlPoint;
+	
+	// Is the location relative to the root components initial location
+ 	UPROPERTY(EditAnywhere, BlueprintReadWrite, Category=Location)
+ 	bool bPositionIsRelative;
+
+	FInterpControlPoint()
+	{
+		bPositionIsRelative = true;
+	}
+	FInterpControlPoint(FVector InPosition)
+	{
+		PositionControlPoint = InPosition;
+		bPositionIsRelative = true;
+	}
+	
+	float StartTime;
+	float DistanceToNext;
+	float Percentage;
+};
+
+/**
+ * Move the root component between a series of points over a given time * 
+ *
+ * @see UMovementComponent
+ */
+UCLASS(ClassGroup=Movement, meta=(BlueprintSpawnableComponent),HideCategories=Velocity)
+class ENGINE_API UInterpToMovementComponent : public UMovementComponent 
+{
+	GENERATED_UCLASS_BODY()
+
+	DECLARE_DYNAMIC_MULTICAST_DELEGATE_TwoParams( FOnInterpToReverseDelegate, const FHitResult&, ImpactResult, float, Time  );
+	DECLARE_DYNAMIC_MULTICAST_DELEGATE_TwoParams( FOnInterpToStopDelegate, const FHitResult&, ImpactResult, float, Time );
+	DECLARE_DYNAMIC_MULTICAST_DELEGATE_TwoParams( FOnInterpToWaitBeginDelegate , const FHitResult&, ImpactResult, float, Time );
+	DECLARE_DYNAMIC_MULTICAST_DELEGATE_TwoParams( FOnInterpToWaitEndDelegate , const FHitResult&, ImpactResult, float, Time );
+	DECLARE_DYNAMIC_MULTICAST_DELEGATE_TwoParams( FOnInterpToResetDelegate , const FHitResult&, ImpactResult, float, Time );
+
+	/* How to take to move from the first point to the last (or vice versa) */
+	UPROPERTY(EditAnywhere, BlueprintReadWrite, Category=Control,meta=(UIMin=0.1f, ClampMin=0.1f))
+	float Duration;
+	
+	/** If true, will pause movement on impact, otherwise will flip movement direction. */
+	UPROPERTY(EditAnywhere, BlueprintReadWrite, Category=Behaviour)
+	uint32 bPauseOnImpact:1;
+
+	/* Movement behaviour of the component */
+	UPROPERTY(EditAnywhere, BlueprintReadWrite, Category=Behaviour)
+	EInterpToBehaviourType BehaviourType;
+
+	/**
+	 * If true, forces sub-stepping to break up movement into discrete smaller steps to improve accuracy of the trajectory.
+	 * Objects that move in a straight line typically do *not* need to set this, as movement always uses continuous collision detection (sweeps) so collision is not missed.
+	 * Sub-stepping is automatically enabled when under the effects of gravity or when homing towards a target.
+	 * @see MaxSimulationTimeStep, MaxSimulationIterations
+	 */
+	UPROPERTY(EditAnywhere, BlueprintReadWrite, Category=Simulation)
+	uint32 bForceSubStepping:1;
+	
+	/** Called when InterpTo impacts something and reverse is enabled. */
+	UPROPERTY(BlueprintAssignable)
+	FOnInterpToReverseDelegate OnInterpToReverse;
+
+	/** Called when InterpTo has come to a stop. */
+	UPROPERTY(BlueprintAssignable)
+	FOnInterpToStopDelegate OnInterpToStop;
+	
+	/** Called when InterpTo has come to a stop but will resume when possible. */
+	UPROPERTY(BlueprintAssignable)
+	FOnInterpToWaitBeginDelegate OnWaitBeginDelegate;
+
+	/** Called when InterpTo has resumed following a stop. */
+	UPROPERTY(BlueprintAssignable)
+	FOnInterpToWaitEndDelegate OnWaitEndDelegate;
+
+	/** Called when InterpTo reached the end and reset back to start . */
+	UPROPERTY(BlueprintAssignable)
+	FOnInterpToResetDelegate OnResetDelegate; 
+
+	//Begin UActorComponent Interface
+	virtual void TickComponent(float DeltaTime, enum ELevelTick TickType, FActorComponentTickFunction *ThisTickFunction) override;
+	virtual void BeginPlay() override;
+	//End UActorComponent Interface
+
+	//Begin UMovementComponent Interface
+	virtual void InitializeComponent() override;
+	//End UMovementComponent Interface
+
+	/**
+	 * This will check to see if the InterpTo is still in the world.  It will check things like
+	 * the KillZ, outside world bounds, etc. and handle the situation.
+	 */
+	virtual bool CheckStillInWorld();
+
+	/** Clears the reference to UpdatedComponent, fires stop event, and stops ticking. */
+	UFUNCTION(BlueprintCallable, Category="Game|Components|InterpToMovement")
+	void StopSimulating(const FHitResult& HitResult);
+
+	bool HasStoppedSimulation() { return UpdatedComponent == NULL; }
+
+	/**
+	 * Compute remaining time step given remaining time and current iterations.
+	 * The last iteration (limited by MaxSimulationIterations) always returns the remaining time, which may violate MaxSimulationTimeStep.
+	 *
+	 * @param RemainingTime		Remaining time in the tick.
+	 * @param Iterations		Current iteration of the tick (starting at 1).
+	 * @return The remaining time step to use for the next sub-step of iteration.
+	 * @see MaxSimulationTimeStep, MaxSimulationIterations
+	 * @see ShouldUseSubStepping()
+	 */
+	float GetSimulationTimeStep(float RemainingTime, int32 Iterations) const;
+
+	/**
+	 * Determine whether or not to use substepping in the InterpTo motion update.
+	 * If true, GetSimulationTimeStep() will be used to time-slice the update. If false, all remaining time will be used during the tick.
+	 * @return Whether or not to use substepping in the InterpTo motion update.
+	 * @see GetSimulationTimeStep()
+	 */
+	virtual bool ShouldUseSubStepping() const;
+
+	/**
+	 * Max time delta for each discrete simulation step.
+	 * Lowering this value can address issues with fast-moving objects or complex collision scenarios, at the cost of performance.
+	 *
+	 * WARNING: if (MaxSimulationTimeStep * MaxSimulationIterations) is too low for the min framerate, the last simulation step may exceed MaxSimulationTimeStep to complete the simulation.
+	 * @see MaxSimulationIterations, bForceSubStepping
+	 */
+	UPROPERTY(EditAnywhere, BlueprintReadWrite, meta=(ClampMin="0.0166", ClampMax="0.50", UIMin="0.0166", UIMax="0.50"), Category=Simulation)
+	float MaxSimulationTimeStep;
+
+	/**
+	 * Max number of iterations used for each discrete simulation step.
+	 * Increasing this value can address issues with fast-moving objects or complex collision scenarios, at the cost of performance.
+	 *
+	 * WARNING: if (MaxSimulationTimeStep * MaxSimulationIterations) is too low for the min framerate, the last simulation step may exceed MaxSimulationTimeStep to complete the simulation.
+	 * @see MaxSimulationTimeStep, bForceSubStepping
+	 */
+	UPROPERTY(EditAnywhere, BlueprintReadWrite, meta=(ClampMin="1", ClampMax="25", UIMin="1", UIMax="25"), Category=Simulation)
+	int32 MaxSimulationIterations;
+	
+	/* List of control points to visit. */
+	UPROPERTY(EditAnywhere, BlueprintReadWrite, Category=Control)
+	TArray<FInterpControlPoint> ControlPoints;
+		
+	/* Add a control point that represents a position. */
+	virtual void AddControlPointPosition(FVector Pos);
+	
+#if WITH_EDITOR
+	// Begin UObject interface.
+	virtual void PostEditChangeProperty(FPropertyChangedEvent& PropertyChangedEvent) override;
+	// End UObject interface.
+#endif // WITH_EDITOR
+
+	// The position of the point 
+	UPROPERTY(EditAnywhere, BlueprintReadWrite, Category = Location, meta = (MakeEditWidget = true))
+		FVector PositionControlPointTest;
+
+protected:
+
+	/** @return true if the simulation should stop. */
+	bool HandleHitWall(const FHitResult& Hit, float TimeTick, const FVector& MoveDelta);
+
+	/** Deal with an impact. Change direction, stop etc depending on the current behaviour setting. */
+	virtual void HandleImpact(const FHitResult& Hit, float TimeSlice=0.f, const FVector& MoveDelta = FVector::ZeroVector) override;
+
+	/* Reverse direction we are moving */
+	float ReverseDirection(const FHitResult& Hit, float Time, bool InBroadcastEvent);
+
+	/** Compute the distance for the given time. */
+	virtual FVector ComputeMoveDelta(float Time) const;
+
+	/* Current position on spline */
+	float CurrentTime;
+
+	/* User to */
+	float TimeMultiplier;
+
+	/* Current direction along spline timeline */
+	float CurrentDirection;
+
+	/* Have we hit something and are waiting for it to move to carry on moving */
+	bool bIsWaiting;
+
+	/* Have we stopped (because we hit something, or reached the end of the cycle */
+	bool bStopped;
+
+	/* Are there any actor control points in use*/
+	bool bContainsActorControlPoints;
+
+	/** Minimum delta time considered when ticking. Delta times below this are not considered. This is a very small non-zero positive value to avoid potential divide-by-zero in simulation code. */
+	static const float MIN_TICK_TIME;
+	
+	/* Update the control points. Adjusts the positions of there are any actor control points as well as updating the spline type */
+	virtual void UpdateControlPoints(bool InForceUpdate);	
+
+	/* Calculate the new current time */
+	float CalculateNewTime( float TimeNow, float Delta, FHitResult& HitResult, bool InBroadcastEvent, bool& OutStopped );
+
+private:
+	float TotalDistance;
+
+	FVector StartLocation;
+};
+
+
+
diff --git a/Engine/Source/Runtime/Engine/Private/Components/InterpToMovementComponent.cpp b/Engine/Source/Runtime/Engine/Private/Components/InterpToMovementComponent.cpp
new file mode 100644
--- /dev/null
+++ b/Engine/Source/Runtime/Engine/Private/Components/InterpToMovementComponent.cpp
@@ -0,0 +1,476 @@
+// Copyright 1998-2015 Epic Games, Inc. All Rights Reserved.
+
+#include "EnginePrivate.h"
+#include "Components/InterpToMovementComponent.h"
+
+DEFINE_LOG_CATEGORY_STATIC(LogInterpToMovementComponent, Log, All);
+
+const float UInterpToMovementComponent::MIN_TICK_TIME = 0.0002f;
+
+UInterpToMovementComponent::UInterpToMovementComponent(const FObjectInitializer& ObjectInitializer)
+	: Super(ObjectInitializer)
+{
+	bUpdateOnlyIfRendered = false;
+	bForceSubStepping = false;
+
+	bWantsInitializeComponent = true;
+
+	MaxSimulationTimeStep = 0.05f;
+	MaxSimulationIterations = 8;
+
+	bIsWaiting = false;
+	TimeMultiplier = 1.0f;	
+	Duration = 1.0f;
+	CurrentDirection = 1.0f;
+	CurrentTime = 0.0f;
+	bStopped = false;
+}
+
+void UInterpToMovementComponent::InitializeComponent()
+{
+	Super::InitializeComponent();
+}
+
+
+void UInterpToMovementComponent::TickComponent(float DeltaTime, enum ELevelTick TickType, FActorComponentTickFunction *ThisTickFunction)
+{
+	QUICK_SCOPE_CYCLE_COUNTER(STAT_InterpToMovementComponent_TickComponent);
+	Super::TickComponent(DeltaTime, TickType, ThisTickFunction);
+
+	// skip if don't want component updated when not rendered or updated component can't move
+	if (ShouldSkipUpdate(DeltaTime))
+	{
+		return;
+	}
+	AActor* ActorOwner = UpdatedComponent->GetOwner();
+	if (!ActorOwner || !CheckStillInWorld())
+	{
+		return;
+	}
+
+	if (UpdatedComponent->IsSimulatingPhysics())
+	{
+		return;
+	}
+	if((bStopped == true ) || ( ActorOwner->IsPendingKill() ) )
+	{
+		return;
+	}
+	if( ControlPoints.Num()== 0 ) 
+	{
+		return;
+	}
+
+	// This will update any control points coordinates that are linked to actors.
+	UpdateControlPoints(false);
+
+	float RemainingTime = DeltaTime;
+	int32 NumBounces = 0;
+	int32 Iterations = 0;
+	FHitResult Hit(1.f);
+
+	FVector WaitPos = FVector::ZeroVector;
+	if (bIsWaiting == true)
+	{
+		WaitPos = UpdatedComponent->GetComponentLocation();
+	}
+	while (RemainingTime >= MIN_TICK_TIME && (Iterations < MaxSimulationIterations) && !ActorOwner->IsPendingKill() && UpdatedComponent)
+	{
+		Iterations++;
+
+		const float TimeTick = ShouldUseSubStepping() ? GetSimulationTimeStep(RemainingTime, Iterations) : RemainingTime;
+		RemainingTime -= TimeTick;
+
+		// Calculate the current time with this tick iteration
+		float Time = FMath::Clamp(CurrentTime + ((DeltaTime*TimeMultiplier)*CurrentDirection),0.0f,1.0f);		
+		FVector MoveDelta = ComputeMoveDelta(Time);
+		
+		// Update the rotation on the spline if required
+		FRotator CurrentRotation = UpdatedComponent->GetComponentRotation();				
+		
+		// Move the component
+ 		if ((bPauseOnImpact == false ) && (BehaviourType != EInterpToBehaviourType::OneShot))
+ 		{
+ 			// If we can bounce, we are allowed to move out of penetrations, so use SafeMoveUpdatedComponent which does that automatically.
+ 			SafeMoveUpdatedComponent(MoveDelta, CurrentRotation, true, Hit);
+ 		}
+ 		else
+		{
+			// If we can't bounce, then we shouldn't adjust if initially penetrating, because that should be a blocking hit that causes a hit event and stop simulation.
+			TGuardValue<EMoveComponentFlags> ScopedFlagRestore(MoveComponentFlags, MoveComponentFlags | MOVECOMP_NeverIgnoreBlockingOverlaps);
+			MoveUpdatedComponent(MoveDelta, CurrentRotation, true, &Hit);
+		}
+		//DrawDebugPoint(GetWorld(), UpdatedComponent->GetComponentLocation(), 16, FColor::White,true,5.0f);
+		// If we hit a trigger that destroyed us, abort.
+		if (ActorOwner->IsPendingKill() || !UpdatedComponent)
+		{
+			return;
+		}
+
+		// Handle hit result after movement
+		if (!Hit.bBlockingHit)
+		{
+			// If we were 'waiting' were not any more - broadcast we are off again
+			if( bIsWaiting == true )
+			{
+				OnWaitEndDelegate.Broadcast(Hit, Time);
+				bIsWaiting = false;
+			}
+			else
+			{				
+				CalculateNewTime(CurrentTime, TimeTick, Hit, true, bStopped);
+				if (bStopped == true)
+				{
+					return;
+				}
+			}
+		}
+		else
+		{
+			if (HandleHitWall(Hit, TimeTick, MoveDelta))
+			{
+				break;
+			}
+
+			NumBounces++;
+			float SubTickTimeRemaining = TimeTick * (1.f - Hit.Time);
+			
+			// A few initial bounces should add more time and iterations to complete most of the simulation.
+			if (NumBounces <= 2 && SubTickTimeRemaining >= MIN_TICK_TIME)
+			{
+				RemainingTime += SubTickTimeRemaining;
+				Iterations--;
+			}
+		}
+	}
+	if( bIsWaiting == false )
+	{		
+		FHitResult DummyHit;
+		CurrentTime = CalculateNewTime(CurrentTime, DeltaTime, DummyHit, false, bStopped);		
+	}
+	UpdateComponentVelocity();
+}
+
+float UInterpToMovementComponent::CalculateNewTime( float TimeNow, float Delta, FHitResult& HitResult, bool InBroadcastEvent, bool& OutStopped )
+{
+	float NewTime = TimeNow;
+	OutStopped = false;
+	if (bIsWaiting == false)
+	{
+		NewTime += ((Delta*TimeMultiplier)*CurrentDirection);
+		if (NewTime >= 1.0f)
+		{
+			if (BehaviourType == EInterpToBehaviourType::OneShot)
+			{
+				NewTime = 1.0f;
+				OutStopped = true;
+				if(InBroadcastEvent == true )
+				{
+					OnInterpToStop.Broadcast(HitResult, NewTime);
+				}
+			}
+			else if(BehaviourType == EInterpToBehaviourType::Loop_Reset)
+			{
+				NewTime = 0.0f;
+				if (InBroadcastEvent == true)
+				{
+					OnResetDelegate.Broadcast(HitResult, NewTime);
+				}
+			}
+			else
+			{
+				FHitResult DummyHit;
+				NewTime = ReverseDirection(DummyHit,NewTime, InBroadcastEvent);
+			}	
+		}
+		else if (NewTime < 0.0f)
+		{
+			if (BehaviourType == EInterpToBehaviourType::OneShot_Reverse)
+			{
+				NewTime = 0.0f;
+				bStopped = true;
+				if (InBroadcastEvent == true)
+				{
+					OnInterpToStop.Broadcast(HitResult, NewTime);
+				}				
+			}
+			else if (BehaviourType == EInterpToBehaviourType::PingPong)
+			{
+
+				FHitResult DummyHit;
+				NewTime = ReverseDirection(DummyHit, NewTime, InBroadcastEvent);
+			}			
+		}
+	}
+	return NewTime;
+}
+
+FVector UInterpToMovementComponent::ComputeMoveDelta(float InTime) const
+{	
+	FVector MoveDelta = FVector::ZeroVector;
+	FVector NewPosition = FVector::ZeroVector;
+	//Find current control point
+	float Time = 0.0f;
+	int32 CurrentControlPoint = INDEX_NONE;
+	// Always use the end point if we are at the end 
+	if (InTime >= 1.0f)
+	{
+		CurrentControlPoint = ControlPoints.Num() - 1;
+	}
+	else
+	{
+		for (int32 iSpline = 0; iSpline < ControlPoints.Num(); iSpline++)
+		{
+			float NextTime = Time + ControlPoints[iSpline].Percentage;
+			if (InTime < NextTime)
+			{
+				CurrentControlPoint = iSpline;
+				break;
+			}
+			Time = NextTime;
+		}
+	}
+	// If we found a valid control point get the position between it and the next
+	if (CurrentControlPoint != INDEX_NONE)
+	{
+		FRotator CurrentRotation = UpdatedComponent->GetComponentRotation();
+		float Base = InTime - ControlPoints[CurrentControlPoint].StartTime;
+		float ThisAlpha = Base / ControlPoints[CurrentControlPoint].Percentage;
+		FVector BeginControlPoint = ControlPoints[CurrentControlPoint].PositionControlPoint;
+		BeginControlPoint = CurrentRotation.RotateVector(BeginControlPoint) + ( ControlPoints[CurrentControlPoint].bPositionIsRelative == true ? StartLocation : FVector::ZeroVector);
+
+		int32 NextControlPoint = FMath::Clamp(CurrentControlPoint + 1, 0, ControlPoints.Num() - 1);
+		FVector EndControlPoint = ControlPoints[NextControlPoint].PositionControlPoint;
+		EndControlPoint = CurrentRotation.RotateVector(EndControlPoint) + ( ControlPoints[NextControlPoint].bPositionIsRelative == true ? StartLocation : FVector::ZeroVector);
+
+		NewPosition = FMath::Lerp(BeginControlPoint, EndControlPoint, ThisAlpha);
+	}
+
+	FVector CurrentPosition = UpdatedComponent->GetComponentLocation();
+	MoveDelta = NewPosition - CurrentPosition;
+	return MoveDelta;
+}
+
+void UInterpToMovementComponent::StopSimulating(const FHitResult& HitResult)
+{
+	SetUpdatedComponent(nullptr);
+	Velocity = FVector::ZeroVector;
+	OnInterpToStop.Broadcast(HitResult, CurrentTime);
+}
+
+bool UInterpToMovementComponent::HandleHitWall(const FHitResult& Hit, float Time, const FVector& MoveDelta)
+{
+	AActor* ActorOwner = UpdatedComponent ? UpdatedComponent->GetOwner() : NULL;
+	if (!CheckStillInWorld() || !ActorOwner || ActorOwner->IsPendingKill())
+	{
+		return true;
+	}
+	HandleImpact(Hit, Time, MoveDelta);
+
+	if (ActorOwner->IsPendingKill() || !UpdatedComponent)
+	{
+		return true;
+	}
+
+	return false;
+}
+
+void UInterpToMovementComponent::HandleImpact(const FHitResult& Hit, float Time, const FVector& MoveDelta)
+{
+	if( bPauseOnImpact == false )
+	{
+		switch(BehaviourType )
+		{
+		case EInterpToBehaviourType::OneShot:
+		case EInterpToBehaviourType::OneShot_Reverse:
+			OnInterpToStop.Broadcast(Hit, Time);
+			// If one shot we are done. If One shot reverse and we are already in reverse we are done too.
+			if ((BehaviourType == EInterpToBehaviourType::OneShot) || (CurrentDirection == -1.0f))
+			{
+				bStopped = true;
+				StopSimulating(Hit);
+				return;
+			}
+			break;
+		case EInterpToBehaviourType::Loop_Reset:
+			{
+				CurrentTime = 0.0f;
+				OnResetDelegate.Broadcast(Hit, CurrentTime);
+			}
+			break;
+		default:
+			ReverseDirection(Hit, Time, true);
+			break;
+		}		
+	}
+	else
+	{
+		if( bIsWaiting == false )
+		{
+			OnWaitBeginDelegate.Broadcast(Hit, Time);
+			bIsWaiting = true;
+		}
+	}
+}
+
+bool UInterpToMovementComponent::CheckStillInWorld()
+{
+	if (!UpdatedComponent)
+	{
+		return false;
+	}
+	// check the variations of KillZ
+	AWorldSettings* WorldSettings = GetWorld()->GetWorldSettings(true);
+	if (!WorldSettings->bEnableWorldBoundsChecks)
+	{
+		return true;
+	}
+	AActor* ActorOwner = UpdatedComponent->GetOwner();
+	if (!ActorOwner)
+	{
+		return false;
+	}
+	if (ActorOwner->GetActorLocation().Z < WorldSettings->KillZ)
+	{
+		UDamageType const* DmgType = WorldSettings->KillZDamageType ? WorldSettings->KillZDamageType->GetDefaultObject<UDamageType>() : GetDefault<UDamageType>();
+		ActorOwner->FellOutOfWorld(*DmgType);
+		return false;
+	}
+	// Check if box has poked outside the world
+	else if (UpdatedComponent && UpdatedComponent->IsRegistered())
+	{
+		const FBox&	Box = UpdatedComponent->Bounds.GetBox();
+		if (Box.Min.X < -HALF_WORLD_MAX || Box.Max.X > HALF_WORLD_MAX ||
+			Box.Min.Y < -HALF_WORLD_MAX || Box.Max.Y > HALF_WORLD_MAX ||
+			Box.Min.Z < -HALF_WORLD_MAX || Box.Max.Z > HALF_WORLD_MAX)
+		{
+			UE_LOG(LogInterpToMovementComponent, Warning, TEXT("%s is outside the world bounds!"), *ActorOwner->GetName());
+			ActorOwner->OutsideWorldBounds();
+			// not safe to use physics or collision at this point
+			ActorOwner->SetActorEnableCollision(false);
+			FHitResult Hit(1.f);
+			StopSimulating(Hit);
+			return false;
+		}
+	}
+	return true;
+}
+
+bool UInterpToMovementComponent::ShouldUseSubStepping() const
+{
+	return bForceSubStepping;
+}
+
+float UInterpToMovementComponent::GetSimulationTimeStep(float RemainingTime, int32 Iterations) const
+{
+	if (RemainingTime > MaxSimulationTimeStep)
+	{
+		if (Iterations < MaxSimulationIterations)
+		{
+			// Subdivide moves to be no longer than MaxSimulationTimeStep seconds
+			RemainingTime = FMath::Min(MaxSimulationTimeStep, RemainingTime * 0.5f);
+		}
+		else
+		{
+			// If this is the last iteration, just use all the remaining time. This is usually better than cutting things short, as the simulation won't move far enough otherwise.
+			// Print a throttled warning.
+#if !(UE_BUILD_SHIPPING || UE_BUILD_TEST)
+			static uint32 s_WarningCount = 0;
+			if ((s_WarningCount++ < 100) || (GFrameCounter & 15) == 0)
+			{
+				UE_LOG(LogInterpToMovementComponent, Warning, TEXT("GetSimulationTimeStep() - Max iterations %d hit while remaining time %.6f > MaxSimulationTimeStep (%.3f) for '%s'"), MaxSimulationIterations, RemainingTime, MaxSimulationTimeStep, *GetPathNameSafe(UpdatedComponent));
+			}
+#endif
+		}
+	}
+
+	// no less than MIN_TICK_TIME (to avoid potential divide-by-zero during simulation).
+	return FMath::Max(MIN_TICK_TIME, RemainingTime);
+}
+
+
+void UInterpToMovementComponent::BeginPlay()
+{
+	StartLocation = UpdatedComponent->GetComponentLocation();
+	TimeMultiplier = 1.0f / Duration;
+	if( ControlPoints.Num() != 0 )
+	{
+		UpdateControlPoints(true);
+		// Update the component location to match first control point
+		FVector MoveDelta = ComputeMoveDelta(0.0f);
+		FRotator CurrentRotation = UpdatedComponent->GetComponentRotation();
+		FHitResult Hit(1.f);
+		UpdatedComponent->MoveComponent(MoveDelta, CurrentRotation, false, &Hit);
+	}	
+}
+
+void UInterpToMovementComponent::UpdateControlPoints(bool InForceUpdate)
+{
+	if (UpdatedComponent != nullptr)
+	{
+		if (InForceUpdate == true) 
+		{
+			FVector BasePosition = UpdatedComponent->GetComponentLocation();
+			TotalDistance = 0.0f;
+			FVector CurrentPos = ControlPoints[0].PositionControlPoint;
+			if (ControlPoints[0].bPositionIsRelative == true)
+			{
+				CurrentPos += BasePosition;
+			}
+			// Calculate the distances from point to point
+			for (int32 ControlPoint = 0; ControlPoint < ControlPoints.Num(); ControlPoint++)
+			{
+				if (ControlPoint + 1 < ControlPoints.Num())
+				{
+					FVector NextPosition = ControlPoints[ControlPoint + 1].PositionControlPoint;
+					if( ControlPoints[ControlPoint+1].bPositionIsRelative == true )
+					{
+						NextPosition += BasePosition;
+					}
+					ControlPoints[ControlPoint].DistanceToNext = (NextPosition - CurrentPos).Size();
+
+					TotalDistance += ControlPoints[ControlPoint].DistanceToNext;
+					CurrentPos = NextPosition;
+				}
+			}
+			float Percent = 0.0f;
+			// Use the distance to determine what % of time to spend going from each point
+			for (int32 ControlPoint = 0; ControlPoint < ControlPoints.Num(); ControlPoint++)
+			{
+				ControlPoints[ControlPoint].StartTime = Percent;
+				ControlPoints[ControlPoint].Percentage = ControlPoints[ControlPoint].DistanceToNext / TotalDistance;
+				Percent += ControlPoints[ControlPoint].Percentage;
+			}
+		}
+	}
+}
+
+float UInterpToMovementComponent::ReverseDirection(const FHitResult& Hit, float Time, bool InBroadcastEvent)
+{
+	float NewTime = Time;
+	// Invert the direction we are moving 
+	if (InBroadcastEvent == true)
+	{
+		OnInterpToReverse.Broadcast(Hit, NewTime);
+	}
+	// flip dir
+	CurrentDirection = -CurrentDirection;
+	// remove extra time from current time
+	float Remainder = NewTime - 1.0f;
+	NewTime = 1.0f - Remainder;
+	return NewTime;
+}
+
+
+void UInterpToMovementComponent::AddControlPointPosition(FVector Pos)
+{
+	UE_LOG(LogInterpToMovementComponent, Warning, TEXT("Pos:%s"),*Pos.ToString());
+	ControlPoints.Add( FInterpControlPoint(Pos));
+}
+
+void UInterpToMovementComponent::PostEditChangeProperty(FPropertyChangedEvent& PropertyChangedEvent)
+{
+	Super::PostEditChangeProperty(PropertyChangedEvent);
+	UpdateControlPoints(true);
+}
+
