Fixed compile error in InterpToComponent

diff --git a/Engine/Source/Runtime/Engine/Private/Components/InterpToMovementComponent.cpp b/Engine/Source/Runtime/Engine/Private/Components/InterpToMovementComponent.cpp
--- a/Engine/Source/Runtime/Engine/Private/Components/InterpToMovementComponent.cpp
+++ b/Engine/Source/Runtime/Engine/Private/Components/InterpToMovementComponent.cpp
@@ -468,9 +468,10 @@ void UInterpToMovementComponent::AddControlPointPosition(FVector Pos)
 	ControlPoints.Add( FInterpControlPoint(Pos));
 }
 
+#if WITH_EDITOR
 void UInterpToMovementComponent::PostEditChangeProperty(FPropertyChangedEvent& PropertyChangedEvent)
 {
 	Super::PostEditChangeProperty(PropertyChangedEvent);
 	UpdateControlPoints(true);
 }
-
+#endif // WITH_EDITOR
