turned of editor gizmo when rendering high res screenshots

diff --git a/Engine/Source/Runtime/Engine/Private/ShowFlags.cpp b/Engine/Source/Runtime/Engine/Private/ShowFlags.cpp
--- a/Engine/Source/Runtime/Engine/Private/ShowFlags.cpp
+++ b/Engine/Source/Runtime/Engine/Private/ShowFlags.cpp
@@ -312,10 +312,13 @@ void EngineShowFlagOverride(EShowFlagInitMode ShowFlagInitMode, EViewModeIndex V
 	}
 
 	{
-		// Disable temporal AA when taking a high resolution screenshot
+		// when taking a high resolution screenshot
 		if (GIsHighResScreenshot)
 		{
+			// disabled as it requires multiple frames, AA can be done by downsampling, more control and better masking
 			EngineShowFlags.TemporalAA = 0;
+			// to not have editor gizmos on high res screenshots
+			EngineShowFlags.ModeWidgets = 0;
 		}
 	}
 
