When seeking out a SCSNode by name need to look at both the variable name and the component template name due to the addition of the GUID

diff --git a/Engine/Source/Runtime/Engine/Private/SimpleConstructionScript.cpp b/Engine/Source/Runtime/Engine/Private/SimpleConstructionScript.cpp
--- a/Engine/Source/Runtime/Engine/Private/SimpleConstructionScript.cpp
+++ b/Engine/Source/Runtime/Engine/Private/SimpleConstructionScript.cpp
@@ -591,7 +591,7 @@ USCS_Node* USimpleConstructionScript::FindSCSNode(const FName InName) const
 
 	for( USCS_Node* SCSNode : AllNodes )
 	{
-		if (SCSNode->GetVariableName() == InName)
+		if (SCSNode->GetVariableName() == InName || (SCSNode->ComponentTemplate && SCSNode->ComponentTemplate->GetFName() == InName))
 		{
 			ReturnSCSNode = SCSNode;
 			break;
