turned off editor selection when rendering high res screenshots

diff --git a/Engine/Source/Runtime/Engine/Private/ShowFlags.cpp b/Engine/Source/Runtime/Engine/Private/ShowFlags.cpp
--- a/Engine/Source/Runtime/Engine/Private/ShowFlags.cpp
+++ b/Engine/Source/Runtime/Engine/Private/ShowFlags.cpp
@@ -317,8 +317,10 @@ void EngineShowFlagOverride(EShowFlagInitMode ShowFlagInitMode, EViewModeIndex V
 		{
 			// disabled as it requires multiple frames, AA can be done by downsampling, more control and better masking
 			EngineShowFlags.TemporalAA = 0;
-			// to not have editor gizmos on high res screenshots
+			// no editor gizmos / selection
 			EngineShowFlags.ModeWidgets = 0;
+			EngineShowFlags.Selection = 0;
+			EngineShowFlags.SelectionOutline = 0;
 		}
 	}
 
