Refactored CameraModifiers to make them more extensible, including via Blueprints.
Also added some test content.

diff --git a/Engine/Source/Runtime/Engine/Classes/Camera/CameraModifier.h b/Engine/Source/Runtime/Engine/Classes/Camera/CameraModifier.h
--- a/Engine/Source/Runtime/Engine/Classes/Camera/CameraModifier.h
+++ b/Engine/Source/Runtime/Engine/Classes/Camera/CameraModifier.h
@@ -10,102 +10,110 @@
  * being computed by the APlayerCameraManager (@see ModifyCamera). A CameraModifier
  * can be stateful, and is associated uniquely with a specific APlayerCameraManager.
  */
-UCLASS()
+UCLASS(BlueprintType, Blueprintable)
 class ENGINE_API UCameraModifier : public UObject
 {
 	GENERATED_UCLASS_BODY()
 
 protected:
 	/** If true, do not apply this modifier to the camera. */
-	UPROPERTY()
 	uint32 bDisabled:1;
 
-public:
 	/** If true, this modifier will disable itself when finished interpolating out. */
-	UPROPERTY()
 	uint32 bPendingDisable:1;
 
 	/** Camera this object is associated with. */
-	UPROPERTY()
+	UPROPERTY(transient, BlueprintReadOnly, Category = CameraModifier)
 	class APlayerCameraManager* CameraOwner;
 
-protected:
+public:
 	/** Priority value that determines the order in which modifiers are applied. 0 = highest priority, 255 = lowest. */
-	UPROPERTY()
+	UPROPERTY(EditDefaultsOnly, BlueprintReadOnly, Category = CameraModifier)
 	uint8 Priority;
 
 	/** If true, no other modifiers of same priority allowed. */
-	UPROPERTY()
+	UPROPERTY(EditDefaultsOnly, BlueprintReadOnly, Category = CameraModifier)
 	uint32 bExclusive:1;
 
+protected:
 	/** When blending in, alpha proceeds from 0 to 1 over this time */
-	UPROPERTY()
+	UPROPERTY(EditDefaultsOnly, BlueprintReadOnly, Category = CameraModifier)
 	float AlphaInTime;
 
 	/** When blending out, alpha proceeds from 1 to 0 over this time */
-	UPROPERTY()
+	UPROPERTY(EditDefaultsOnly, BlueprintReadOnly, Category = CameraModifier)
 	float AlphaOutTime;
 
 	/** Current blend alpha. */
-	UPROPERTY(transient)
+	UPROPERTY(transient, BlueprintReadOnly, Category = CameraModifier)
 	float Alpha;
 
-	/** Desired alpha we are interpolating towards. */
-	UPROPERTY(transient)
-	float TargetAlpha;
-
 public:
 	/** If true, enables certain debug visualization features. */
 	UPROPERTY(EditAnywhere, Category=Debug)
-	uint32 bDebug:1;
+	uint32 bDebug :1;
 
 protected:
 	/** @return Returns the ideal blend alpha for this modifier. Interpolation will seek this value. */
-	virtual float GetTargetAlpha(class APlayerCameraManager* Camera);
+	virtual float GetTargetAlpha();
 
 public:
 	/** 
 	 * Allows any custom initialization. Called immediately after creation.
 	 * @param Camera - The camera this modifier should be associated with.
 	 */
-	virtual void Init( APlayerCameraManager* Camera );
+	virtual void AddedToCamera( APlayerCameraManager* Camera );
 	
 	/**
-	 * Directly modifies variables in the camera actor
-	 * @param	Camera		reference to camera actor we are modifying
+	 * Directly modifies variables in the owning camera
 	 * @param	DeltaTime	Change in time since last update
-	 * @param	InOutPOV		current Point of View, to be updated.
-	 * @return	bool		true if should STOP looping the chain, false otherwise
+	 * @param	InOutPOV	Current Point of View, to be updated.
+	 * @return	bool		True if should STOP looping the chain, false otherwise
 	 */
-	virtual bool ModifyCamera(APlayerCameraManager* Camera, float DeltaTime, struct FMinimalViewInfo& InOutPOV);
-	
-	/** @return Returns true if modifier is disabled, false otherwise. */
-	virtual bool IsDisabled() const;
-	
-	/**
-	 * Camera modifier evaluates itself vs the given camera's modifier list
-	 * and decides whether to add itself or not. Handles adding by priority and avoiding 
-	 * adding the same modifier twice.
-	 *
-	 * @param	Camera - reference to camera actor we want add this modifier to
-	 * @return	Returns true if modifier added to camera's modifier list, false otherwise.
+	virtual bool ModifyCamera(float DeltaTime, struct FMinimalViewInfo& InOutPOV);
+
+	/** 
+	 * Called per tick that the modifier is active to allow Blueprinted modifiers to modify the camera's transform. 
+	 * Scaling by Alpha happens after this in code, so no need to deal with that in the blueprint.
+	 * @param	DeltaTime	Change in time since last update
+	 * @param	ViewLocation		The current camera location.
+	 * @param	ViewRotation		The current camera rotation.
+	 * @param	FOV					The current camera fov.
+	 * @param	NewViewLocation		(out) The modified camera location.
+	 * @param	NewViewRotation		(out) The modified camera rotation.
+	 * @param	NewFOV				(out) The modified camera FOV.
 	 */
-	virtual bool AddCameraModifier( APlayerCameraManager* Camera );
-	
+	UFUNCTION(BlueprintImplementableEvent, BlueprintCosmetic)
+	void BlueprintModifyCamera(float DeltaTime, FVector ViewLocation, FRotator ViewRotation, float FOV, FVector& NewViewLocation, FRotator& NewViewRotation, float& NewFOV);
+
 	/**
-	 * Camera modifier removes itself from given camera's modifier list
-	 * @param	Camera	- reference to camera actor we want to remove this modifier from
-	 * @return	Returns true if modifier removed successfully, false otherwise.
+	 * Called per tick that the modifier is active to allow Blueprinted modifiers to modify the camera's postprocess effects.
+	 * Scaling by Alpha happens after this in code, so no need to deal with that in the blueprint.
+	 * @param	DeltaTime				Change in time since last update
+	 * @param	PostProcessBlendWeight	(out) Blend weight applied to the entire postprocess structure.
+	 * @param	PostProcessSettings		(out) Post process structure defining what settings and values to override.
 	 */
-	virtual bool RemoveCameraModifier( APlayerCameraManager* Camera );
+	UFUNCTION(BlueprintImplementableEvent, BlueprintCosmetic)
+	void BlueprintModifyPostProcess(float DeltaTime, float& PostProcessBlendWeight, FPostProcessSettings& PostProcessSettings);
+
+	/** @return Returns true if modifier is disabled, false otherwise. */
+	UFUNCTION(BlueprintCallable, Category = CameraModifier)
+	virtual bool IsDisabled() const;
 	
+	/** @return Returns the actor the camera is currently viewing. */
+	UFUNCTION(BlueprintCallable, Category = CameraModifier)
+	virtual AActor* GetViewTarget() const;
+
+
 	/** 
 	 *  Disables this modifier.
 	 *  @param  bImmediate  - true to disable with no blend out, false (default) to allow blend out
 	 */
+	UFUNCTION(BlueprintCallable, Category=CameraModifier)
 	virtual void DisableModifier(bool bImmediate = false);
 
 	/** Enables this modifier. */
+	UFUNCTION(BlueprintCallable, Category = CameraModifier)
 	virtual void EnableModifier();
 
 	/** Toggled disabled/enabled state of this modifier. */
@@ -129,7 +137,7 @@ public:
 	 * @param	Camera		- Camera that is being updated
 	 * @param	DeltaTime	- Amount of time since last update
 	 */
-	virtual void UpdateAlpha( APlayerCameraManager* Camera, float DeltaTime );
+	virtual void UpdateAlpha(float DeltaTime);
 
 	/** @return Returns the appropriate world context for this object. */
 	UWorld* GetWorld() const;
diff --git a/Engine/Source/Runtime/Engine/Classes/Camera/CameraModifier_CameraShake.h b/Engine/Source/Runtime/Engine/Classes/Camera/CameraModifier_CameraShake.h
--- a/Engine/Source/Runtime/Engine/Classes/Camera/CameraModifier_CameraShake.h
+++ b/Engine/Source/Runtime/Engine/Classes/Camera/CameraModifier_CameraShake.h
@@ -46,7 +46,7 @@ public:
 	virtual void RemoveAllCameraShakes();
 	
 	// Begin UCameraModifer Interface
-	virtual bool ModifyCamera(APlayerCameraManager* Camera, float DeltaTime, struct FMinimalViewInfo& InOutPOV) override;
+	virtual bool ModifyCamera(float DeltaTime, struct FMinimalViewInfo& InOutPOV) override;
 	// End UCameraModifer Interface
 
 protected:
diff --git a/Engine/Source/Runtime/Engine/Classes/Camera/PlayerCameraManager.h b/Engine/Source/Runtime/Engine/Classes/Camera/PlayerCameraManager.h
--- a/Engine/Source/Runtime/Engine/Classes/Camera/PlayerCameraManager.h
+++ b/Engine/Source/Runtime/Engine/Classes/Camera/PlayerCameraManager.h
@@ -256,10 +256,16 @@ public:
 	UPROPERTY()
 	struct FViewTargetTransitionParams BlendParams;
 
-	/** List of active camera modifiers that have a chance to update the final camera POV */
-	UPROPERTY()
+protected:
+	/** List of active camera modifier instances that have a chance to update the final camera POV */
+	UPROPERTY(transient)
 	TArray<class UCameraModifier*> ModifierList;
 
+public:
+	/** List of modifiers to create by default for this camera */
+	UPROPERTY(EditDefaultsOnly, BlueprintReadOnly, Category = CameraModifier)
+	TArray< TSubclassOf<class UCameraModifier> > DefaultModifiers;
+
 	/** Distance to place free camera from view target (used in certain CameraStyles) */
 	UPROPERTY(EditAnywhere, BlueprintReadWrite, Category=Debug)
 	float FreeCamDistance;
@@ -286,19 +292,13 @@ protected:
 	UPROPERTY(transient)
 	TArray<class AEmitterCameraLensEffectBase*> CameraLensEffects;
 
-public:
 	/////////////////////
 	// Camera Modifiers
 	/////////////////////
-	
-	/** Camera modifier for cone-driven screen shakes */
-	UPROPERTY(EditAnywhere, BlueprintReadWrite, Instanced, transient, Category = PlayerCameraManager)
-	class UCameraModifier_CameraShake* CameraShakeCamMod;
 
-protected:
-	/** Class to use when instantiating screenshake modifier object. */
-	UPROPERTY(EditAnywhere, BlueprintReadWrite, Category=PlayerCameraManager)
-	TSubclassOf<class UCameraModifier_CameraShake> CameraShakeCamModClass;
+	/** Cached ref to modifier for code-driven screen shakes */
+	UPROPERTY(transient)
+	class UCameraModifier_CameraShake* CachedCameraShakeMod;
 
 	/** By default camera modifiers are not applied to stock debug cameras (e.g. CameraStyle == "FreeCam"). Setting to true will always apply modifiers. */
 	UPROPERTY()
@@ -319,8 +319,10 @@ protected:
 	/** Internal list of weights for active post process effects. Parallel array to PostProcessBlendCache. */
 	TArray<float> PostProcessBlendCacheWeights;
 
+public:
 	/** Adds a postprocess effect at the given weight. */
 	void AddCachedPPBlend(struct FPostProcessSettings& PPSettings, float BlendWeight);
+protected:
 	/** Removes all postprocess effects. */
 	void ClearCachedPPBlends();
 
@@ -483,12 +485,32 @@ public:
 	virtual void UpdateCamera(float DeltaTime);
 
 	/** 
-	 * Internal. Creates and initializes a new camera modifier of the specified class. 
+	 * Creates and initializes a new camera modifier of the specified class. 
 	 * @param ModifierClass - The class of camera modifier to create.
 	 * @return Returns the newly created camera modifier.
 	 */
-	virtual class UCameraModifier* CreateCameraModifier(TSubclassOf<class UCameraModifier> ModifierClass);
-	
+	UFUNCTION(BlueprintCallable, Category = "Game|Player")
+	virtual class UCameraModifier* AddNewCameraModifier(TSubclassOf<class UCameraModifier> ModifierClass);
+
+	/** 
+	 * Returns camera modifier for this camera of the given class, if it exists. 
+	 * Exact class match only. If there are multiple modifiers of the same class, the first one is returned.
+	 */
+	UFUNCTION(BlueprintCallable, Category = "Game|Player")
+	virtual class UCameraModifier* FindCameraModifierByClass(TSubclassOf<class UCameraModifier> ModifierClass);
+
+	/** 
+	 * Removes the given camera modifier from this camera (if it's on the camera in the first place) and discards it. 
+	 * @return True if successfully removed, false otherwise.
+	 */
+	UFUNCTION(BlueprintCallable, Category = "Game|Player")
+	virtual bool RemoveCameraModifier(class UCameraModifier* ModifierToRemove);
+
+protected:
+	/** Internal. Places the given modifier in the ModifierList at the appropriate priority. */
+	virtual bool AddCameraModifierToList(class UCameraModifier* NewModifier);
+
+public:	
 	/**
 	 * Applies the current set of camera modifiers to the given camera POV.
 	 * @param	DeltaTime	Time in seconds since last update
diff --git a/Engine/Source/Runtime/Engine/Private/Camera/CameraModifier.cpp b/Engine/Source/Runtime/Engine/Private/Camera/CameraModifier.cpp
--- a/Engine/Source/Runtime/Engine/Private/Camera/CameraModifier.cpp
+++ b/Engine/Source/Runtime/Engine/Private/Camera/CameraModifier.cpp
@@ -19,97 +19,84 @@ UCameraModifier::UCameraModifier(const FObjectInitializer& ObjectInitializer)
 	Priority = 127;
 }
 
-bool UCameraModifier::ModifyCamera(class APlayerCameraManager* Camera, float DeltaTime, FMinimalViewInfo& InOutPOV)
+bool UCameraModifier::ModifyCamera(float DeltaTime, FMinimalViewInfo& InOutPOV)
 {
+	// Update the alpha
+	UpdateAlpha(DeltaTime);
+
+	// let BP do what it wants
+	BlueprintModifyCamera(DeltaTime, InOutPOV.Location, InOutPOV.Rotation, InOutPOV.FOV, InOutPOV.Location, InOutPOV.Rotation, InOutPOV.FOV);
+
+	if (CameraOwner)
+	{
+		// note: pushing these through the cached PP blend system in the camera to get
+		// proper layered blending, rather than letting subsequent mods stomp over each other in the 
+		// InOutPOV struct.
+		float PPBlendWeight = 0.f;
+		FPostProcessSettings PPSettings;
+		BlueprintModifyPostProcess(DeltaTime, PPBlendWeight, PPSettings);
+		if (PPBlendWeight > 0.f)
+		{
+			CameraOwner->AddCachedPPBlend(PPSettings, PPBlendWeight);
+		}
+	}
+
 	// If pending disable and fully alpha'd out, truly disable this modifier
-	if (bPendingDisable && (Alpha <= 0.0f))
+	if (bPendingDisable && (Alpha <= 0.f))
 	{
 		DisableModifier(true);
 	}
 
+	// allow subsequent modifiers to update
 	return false;
 }
 
-float UCameraModifier::GetTargetAlpha(class APlayerCameraManager* Camera)
+float UCameraModifier::GetTargetAlpha()
 {
-	if( bPendingDisable )
-	{
-		return 0.0f;
-	}
-	return 1.0f;		
+	return bPendingDisable ? 0.0f : 1.f;
 }
 
-void UCameraModifier::UpdateAlpha(class APlayerCameraManager* Camera, float DeltaTime)
+void UCameraModifier::UpdateAlpha(float DeltaTime)
 {
-	float Time;
-
-	TargetAlpha = GetTargetAlpha( Camera );
-
-	// Alpha out
-	if( TargetAlpha == 0.f )
-	{
-		Time = AlphaOutTime;
-	}
-	else
-	{
-		// Otherwise, alpha in
-		Time = AlphaInTime;
-	}
+	float const TargetAlpha = GetTargetAlpha();
+	float const BlendTime = (TargetAlpha == 0.f) ? AlphaOutTime : AlphaInTime;
 
 	// interpolate!
-	if( Time <= 0.f )
+	if (BlendTime <= 0.f)
 	{
 		Alpha = TargetAlpha;
 	}
-	else if( Alpha > TargetAlpha )
+	else if (BlendTime > TargetAlpha)
 	{
-		Alpha = FMath::Max<float>( Alpha - DeltaTime/Time, TargetAlpha );
+		Alpha = FMath::Max<float>(Alpha - DeltaTime / BlendTime, TargetAlpha);
 	}
 	else
 	{
-		Alpha = FMath::Min<float>( Alpha + DeltaTime/Time, TargetAlpha );
+		Alpha = FMath::Min<float>(Alpha + DeltaTime / BlendTime, TargetAlpha);
 	}
 }
 
-
 bool UCameraModifier::IsDisabled() const
 {
 	return bDisabled;
 }
 
-void UCameraModifier::Init( APlayerCameraManager* Camera ) 
+AActor* UCameraModifier::GetViewTarget() const
 {
-	AddCameraModifier( Camera );
+	return CameraOwner ? CameraOwner->GetViewTarget() : nullptr;
 }
 
-UWorld* UCameraModifier::GetWorld() const
+
+void UCameraModifier::AddedToCamera( APlayerCameraManager* Camera ) 
 {
-	UWorld* World = NULL;
-	if( CameraOwner )
-	{
-		World = CameraOwner->GetWorld();
-	}
-	return World;
+	CameraOwner = Camera;
 }
 
-bool UCameraModifier::RemoveCameraModifier( APlayerCameraManager* Camera )
+UWorld* UCameraModifier::GetWorld() const
 {
-	// Loop through each modifier in camera
-	for( int32 ModifierIdx = 0; ModifierIdx < Camera->ModifierList.Num(); ModifierIdx++ )
-	{
-		// If we found ourselves, remove ourselves from the list and return
-		if( Camera->ModifierList[ModifierIdx] == this )
-		{
-			Camera->ModifierList.RemoveAt(ModifierIdx, 1);
-			return true;
-		}
-	}
-
-	// Didn't find ourselves in the list
-	return false;
+	return CameraOwner ? CameraOwner->GetWorld() : nullptr;
 }
 
-
 void UCameraModifier::DisableModifier(bool bImmediate)
 {
 	if (bImmediate)
@@ -147,79 +134,3 @@ bool UCameraModifier::ProcessViewRotation( AActor* ViewTarget, float DeltaTime,
 	return false;
 }
 
-bool UCameraModifier::AddCameraModifier( APlayerCameraManager* Camera )
-{
-	// Make sure we don't already have this modifier in the list
-	for( int32 ModifierIdx = 0; ModifierIdx < Camera->ModifierList.Num(); ModifierIdx++ )
-	{
-		if ( Camera->ModifierList[ModifierIdx] == this )
-		{
-			return false;
-		}
-	}
-
-	// Make sure we don't already have a modifier of this type
-	for( int32 ModifierIdx = 0; ModifierIdx < Camera->ModifierList.Num(); ModifierIdx++ )
-	{
-		if ( Camera->ModifierList[ModifierIdx]->GetClass() == GetClass() )
-		{
-			UE_LOG(LogCamera, Log, TEXT("AddCameraModifier found existing Modifier in list, replacing with new one %s"), *GetName());
-
-			// hack replace old by new (delete??)
-			Camera->ModifierList[ModifierIdx] = this;
-
-			// Save camera
-			CameraOwner = Camera;
-
-			return true;
-		}
-	}
-
-	// Look through current modifier list and find slot for this priority
-	int32 BestIdx = 0;
-
-	for( int32 ModifierIdx = 0; ModifierIdx < Camera->ModifierList.Num(); ModifierIdx++ )
-	{
-		UCameraModifier* Modifier = Camera->ModifierList[ModifierIdx];
-		if( Modifier == NULL ) {
-			continue;
-		}
-
-		// If priority of current index has passed or equaled ours - we have the insert location
-		if( Priority <= Modifier->Priority )
-		{
-			// Disallow addition of exclusive modifier if priority is already occupied
-			if( bExclusive && Priority == Modifier->Priority )
-			{
-				return false;
-			}
-
-			break;
-		}
-
-		// Update best index
-		BestIdx++;
-	}
-
-	// Insert self into best index
-	Camera->ModifierList.InsertUninitialized( BestIdx, 1 );
-	Camera->ModifierList[BestIdx] = this;
-
-	// Save camera
-	CameraOwner = Camera;
-
-#if !(UE_BUILD_SHIPPING || UE_BUILD_TEST)
-	//debug
-	if( bDebug )
-	{
-		UE_LOG(LogCamera, Log,  TEXT("AddModifier %i %s"), BestIdx, *GetName() );
-		for( int32 ModifierIdx = 0; ModifierIdx < Camera->ModifierList.Num(); ModifierIdx++ )
-		{
-			UE_LOG(LogCamera, Log,  TEXT("%s Idx %i Pri %i"), *Camera->ModifierList[ModifierIdx]->GetName(), ModifierIdx, Camera->ModifierList[ModifierIdx]->Priority);
-		}
-		UE_LOG(LogCamera, Log,  TEXT("****************") );
-	}
-#endif
-
-	return true;
-}
diff --git a/Engine/Source/Runtime/Engine/Private/Camera/CameraModifier_CameraShake.cpp b/Engine/Source/Runtime/Engine/Private/Camera/CameraModifier_CameraShake.cpp
--- a/Engine/Source/Runtime/Engine/Private/Camera/CameraModifier_CameraShake.cpp
+++ b/Engine/Source/Runtime/Engine/Private/Camera/CameraModifier_CameraShake.cpp
@@ -13,13 +13,10 @@ UCameraModifier_CameraShake::UCameraModifier_CameraShake(const FObjectInitialize
 }
 
 
-bool UCameraModifier_CameraShake::ModifyCamera(class APlayerCameraManager* Camera, float DeltaTime, FMinimalViewInfo& InOutPOV)
+bool UCameraModifier_CameraShake::ModifyCamera(float DeltaTime, FMinimalViewInfo& InOutPOV)
 {
-	// Update the alpha
-	UpdateAlpha(Camera, DeltaTime);
-
 	// Call super where modifier may be disabled
-	Super::ModifyCamera(Camera, DeltaTime, InOutPOV);
+	Super::ModifyCamera(DeltaTime, InOutPOV);
 
 	// If no alpha, exit early
 	if( Alpha <= 0.f )
diff --git a/Engine/Source/Runtime/Engine/Private/PlayerCameraManager.cpp b/Engine/Source/Runtime/Engine/Private/PlayerCameraManager.cpp
--- a/Engine/Source/Runtime/Engine/Private/PlayerCameraManager.cpp
+++ b/Engine/Source/Runtime/Engine/Private/PlayerCameraManager.cpp
@@ -37,7 +37,6 @@ APlayerCameraManager::APlayerCameraManager(const FObjectInitializer& ObjectIniti
 	ViewYawMax = 359.999f;
 	ViewRollMin = -89.99f;
 	ViewRollMax = 89.99f;
-	CameraShakeCamModClass = UCameraModifier_CameraShake::StaticClass();
 	bUseClientSideCameraUpdates = true;
 	CameraStyle = NAME_Default;
 	bCanBeDamaged = false;
@@ -47,6 +46,9 @@ APlayerCameraManager::APlayerCameraManager(const FObjectInitializer& ObjectIniti
 	// create dummy transform component
 	TransformComponent = CreateDefaultSubobject<USceneComponent>(TEXT("TransformComponent0"));
 	RootComponent = TransformComponent;
+
+	// support camerashakes by default
+	DefaultModifiers.Add(UCameraModifier_CameraShake::StaticClass());
 }
 
 APlayerController* APlayerCameraManager::GetOwningPlayerController() const
@@ -199,6 +201,8 @@ bool APlayerCameraManager::ShouldTickIfViewportsOnly() const
 
 void APlayerCameraManager::ApplyCameraModifiers(float DeltaTime, FMinimalViewInfo& InOutPOV)
 {
+	ClearCachedPPBlends();
+
 	// Loop through each camera modifier
 	for (int32 ModifierIdx = 0; ModifierIdx < ModifierList.Num(); ++ModifierIdx)
 	{
@@ -208,15 +212,13 @@ void APlayerCameraManager::ApplyCameraModifiers(float DeltaTime, FMinimalViewInf
 			// If ModifyCamera returns true, exit loop
 			// Allows high priority things to dictate if they are
 			// the last modifier to be applied
-			if (ModifierList[ModifierIdx]->ModifyCamera(this, DeltaTime, InOutPOV))
+			if (ModifierList[ModifierIdx]->ModifyCamera(DeltaTime, InOutPOV))
 			{
 				break;
 			}
 		}
 	}
 
-	ClearCachedPPBlends();
-
 	// Now apply CameraAnims
 	// these essentially behave as the highest-pri modifier.
 	for (int32 Idx = 0; Idx < ActiveAnims.Num(); ++Idx)
@@ -627,24 +629,122 @@ void APlayerCameraManager::StopAudioFade()
 	}
 }
 
-UCameraModifier* APlayerCameraManager::CreateCameraModifier(TSubclassOf<UCameraModifier> ModifierClass)
+UCameraModifier* APlayerCameraManager::AddNewCameraModifier(TSubclassOf<UCameraModifier> ModifierClass)
+{
+	UCameraModifier* const NewMod = NewObject<UCameraModifier>(this, ModifierClass);
+	if (NewMod)
+	{
+		if (AddCameraModifierToList(NewMod) == true)
+		{
+			return NewMod;
+		}
+	}
+	
+	return nullptr;
+}
+
+UCameraModifier* APlayerCameraManager::FindCameraModifierByClass(TSubclassOf<UCameraModifier> ModifierClass)
+{
+	for (UCameraModifier* Mod : ModifierList)
+	{
+		if (Mod->GetClass() == ModifierClass)
+		{
+			return Mod;
+		}
+	}
+
+	return nullptr;
+}
+
+
+bool APlayerCameraManager::AddCameraModifierToList(UCameraModifier* NewModifier)
+{
+	if (NewModifier)
+	{
+		// Look through current modifier list and find slot for this priority
+		int32 BestIdx = 0;
+		for (int32 ModifierIdx = 0; ModifierIdx < ModifierList.Num(); ModifierIdx++)
+		{
+			UCameraModifier* const M = ModifierList[ModifierIdx];
+			if (M)
+			{
+				if (M == NewModifier)
+				{
+					// already in list, just bail
+					return false;
+				}
+
+				// If priority of current index has passed or equaled ours - we have the insert location
+				if (NewModifier->Priority <= M->Priority)
+				{
+					// Disallow addition of exclusive modifier if priority is already occupied
+					if (NewModifier->bExclusive && NewModifier->Priority == M->Priority)
+					{
+						return false;
+					}
+
+					break;
+				}
+			}
+
+			// Update best index
+			BestIdx = ModifierIdx;
+		}
+
+		// Insert self into best index
+		ModifierList.InsertUninitialized(BestIdx, 1);
+		ModifierList[BestIdx] = NewModifier;
+
+		// Save camera
+		NewModifier->AddedToCamera(this);
+		return true;
+	}
+
+	return false;
+}
+
+bool APlayerCameraManager::RemoveCameraModifier(UCameraModifier* ModifierToRemove)
 {
-	UCameraModifier* NewMod = NewObject<UCameraModifier>(this, ModifierClass);
-	NewMod->Init(this);
-	return NewMod;
+	if (ModifierToRemove)
+	{
+		// Loop through each modifier in camera
+		for (int32 ModifierIdx = 0; ModifierIdx < ModifierList.Num(); ModifierIdx++)
+		{
+			// If we found ourselves, remove ourselves from the list and return
+			if (ModifierList[ModifierIdx] == ModifierToRemove)
+			{
+				ModifierList.RemoveAt(ModifierIdx, 1);
+				return true;
+			}
+		}
+	}
+
+	// Didn't find it in the list, nothing removed
+	return false;
 }
 
+
 void APlayerCameraManager::PostInitializeComponents()
 {
 	Super::PostInitializeComponents();
 
- 	// Setup camera modifiers
- 	if( (CameraShakeCamMod == NULL) && (CameraShakeCamModClass != NULL) )
- 	{
- 		CameraShakeCamMod = Cast<UCameraModifier_CameraShake>(CreateCameraModifier(CameraShakeCamModClass));
- 	}
+ 	// Setup default camera modifiers
+	if (DefaultModifiers.Num() > 0)
+	{
+		for (auto ModifierClass : DefaultModifiers)
+		{
+			UCameraModifier* const NewMod = AddNewCameraModifier(ModifierClass);
+		
+			// cache ref to camera shake if this is it
+			UCameraModifier_CameraShake* const ShakeMod = Cast<UCameraModifier_CameraShake>(NewMod);
+			if (ShakeMod)
+			{
+				CachedCameraShakeMod = ShakeMod;
+			}
+		}
+	}
 
-	// create CameraAnimInsts in pool
+ 	// create CameraAnimInsts in pool
 	for (int32 Idx=0; Idx<MAX_ACTIVE_CAMERA_ANIMS; ++Idx)
 	{
 		AnimInstPool[Idx] = NewObject<UCameraAnimInst>(this);
@@ -673,8 +773,6 @@ void APlayerCameraManager::Destroyed()
 	Super::Destroyed();
 }
 
-
-
 void APlayerCameraManager::InitializeFor(APlayerController* PC)
 {
 	CameraCache.POV.FOV = DefaultFOV;
@@ -1090,9 +1188,9 @@ void APlayerCameraManager::ClearCameraLensEffects()
 
 UCameraShake* APlayerCameraManager::PlayCameraShake(TSubclassOf<UCameraShake> ShakeClass, float Scale, ECameraAnimPlaySpace::Type PlaySpace, FRotator UserPlaySpaceRot)
 {
-	if (ShakeClass && CameraShakeCamMod && (Scale > 0.0f) )
+	if (ShakeClass && CachedCameraShakeMod && (Scale > 0.0f) )
 	{
-		return CameraShakeCamMod->AddCameraShake(ShakeClass, Scale, PlaySpace, UserPlaySpaceRot);
+		return CachedCameraShakeMod->AddCameraShake(ShakeClass, Scale, PlaySpace, UserPlaySpaceRot);
 	}
 
 	return nullptr;
@@ -1101,25 +1199,25 @@ UCameraShake* APlayerCameraManager::PlayCameraShake(TSubclassOf<UCameraShake> Sh
 
 void APlayerCameraManager::StopCameraShake(UCameraShake* ShakeInst)
 {
-	if (ShakeInst && CameraShakeCamMod)
+	if (ShakeInst && CachedCameraShakeMod)
 	{
-		CameraShakeCamMod->RemoveCameraShake(ShakeInst);
+		CachedCameraShakeMod->RemoveCameraShake(ShakeInst);
 	}
 }
 
 void APlayerCameraManager::StopAllInstancesOfCameraShake(TSubclassOf<class UCameraShake> ShakeClass)
 {
-	if (ShakeClass && CameraShakeCamMod)
+	if (ShakeClass && CachedCameraShakeMod)
 	{
-		CameraShakeCamMod->RemoveAllCameraShakesOfClass(ShakeClass);
+		CachedCameraShakeMod->RemoveAllCameraShakesOfClass(ShakeClass);
 	}
 }
 
 void APlayerCameraManager::StopAllCameraShakes()
 {
-	if (CameraShakeCamMod)
+	if (CachedCameraShakeMod)
 	{
-		CameraShakeCamMod->RemoveAllCameraShakes();
+		CachedCameraShakeMod->RemoveAllCameraShakes();
 	}
 }
 
