Change FBodyInstance.UseAsyncScene to take in the owning PhysScene. This is important in cases where garbage collection may have marked the OwnerComponent as unreachable.

In normal contexts the user is always getting the physscene to do some work on it anyway so this API change seems correct.

diff --git a/Engine/Source/Runtime/Engine/Classes/PhysicsEngine/BodyInstance.h b/Engine/Source/Runtime/Engine/Classes/PhysicsEngine/BodyInstance.h
--- a/Engine/Source/Runtime/Engine/Classes/PhysicsEngine/BodyInstance.h
+++ b/Engine/Source/Runtime/Engine/Classes/PhysicsEngine/BodyInstance.h
@@ -638,8 +638,11 @@ public:
 	void SetWalkableSlopeOverride(const FWalkableSlopeOverride& NewOverride);
 
 	/** Returns whether this body wants (and can) use the async scene. */
+	DEPRECATED(4.8, "Please call UseAsyncScene and pass in the PhysScene of the owning component (GetWorld()->GetPhysicsScene()).")
 	bool UseAsyncScene() const;
 
+	bool UseAsyncScene(const FPhysScene* PhysScene) const;
+
 	/** Indicates whether this body should use the async scene. Must be called before body is init'd, will assert otherwise. Will have no affect if there is no async scene. */
 	void SetUseAsyncScene(bool bNewUseAsyncScene);
 
diff --git a/Engine/Source/Runtime/Engine/Private/Components/DestructibleComponent.cpp b/Engine/Source/Runtime/Engine/Private/Components/DestructibleComponent.cpp
--- a/Engine/Source/Runtime/Engine/Private/Components/DestructibleComponent.cpp
+++ b/Engine/Source/Runtime/Engine/Private/Components/DestructibleComponent.cpp
@@ -343,7 +343,7 @@ void UDestructibleComponent::CreatePhysicsState()
 	
 
 	// Destructibles are always dynamic or kinematic, and therefore only go into one of the scenes
-	const uint32 SceneType = BodyInstance.UseAsyncScene() ? PST_Async : PST_Sync;
+	const uint32 SceneType = BodyInstance.UseAsyncScene(PhysScene) ? PST_Async : PST_Sync;
 	NxApexScene* ApexScene = PhysScene->GetApexScene(SceneType);
 	check(ApexScene);
 
@@ -401,7 +401,7 @@ void UDestructibleComponent::DestroyPhysicsState()
 		{
 			if (FPhysScene * PhysScene = World->GetPhysicsScene())
 			{
-				const uint32 SceneType = BodyInstance.UseAsyncScene() ? PST_Async : PST_Sync;
+				const uint32 SceneType = BodyInstance.UseAsyncScene(PhysScene) ? PST_Async : PST_Sync;
 				PxScene * PScene = PhysScene->GetPhysXScene(SceneType);
 				SCOPED_SCENE_WRITE_LOCK(PScene);
 				ApexDestructibleActor->release();
@@ -1260,8 +1260,11 @@ void UDestructibleComponent::SetupFakeBodyInstance( physx::PxRigidActor* NewRigi
 		PrevState->InstanceIndex = BodyInstance.InstanceBodyIndex;
 	}
 
-	BodyInstance.RigidActorSync = BodyInstance.UseAsyncScene() ? NULL : NewRigidActor;
-	BodyInstance.RigidActorAsync = BodyInstance.UseAsyncScene() ? NewRigidActor : NULL;
+	const UWorld* World = GetWorld();
+	const FPhysScene* PhysScene = World ? World->GetPhysicsScene() : nullptr;
+
+	BodyInstance.RigidActorSync = BodyInstance.UseAsyncScene(PhysScene) ? NULL : NewRigidActor;
+	BodyInstance.RigidActorAsync = BodyInstance.UseAsyncScene(PhysScene) ? NewRigidActor : NULL;
 	BodyInstance.BodyAggregate = NULL;
 	BodyInstance.InstanceBodyIndex = InstanceIdx;
 }
diff --git a/Engine/Source/Runtime/Engine/Private/InstancedStaticMesh.cpp b/Engine/Source/Runtime/Engine/Private/InstancedStaticMesh.cpp
--- a/Engine/Source/Runtime/Engine/Private/InstancedStaticMesh.cpp
+++ b/Engine/Source/Runtime/Engine/Private/InstancedStaticMesh.cpp
@@ -1018,7 +1018,7 @@ void UInstancedStaticMeshComponent::CreatePhysicsState()
 	const int32 NumAggregates = (Mobility == EComponentMobility::Movable) ? FMath::DivideAndRoundUp<int32>(NumBodies, AggregateMaxSize) : 0;
 
 	// Get the scene type from the main BodyInstance
-	const uint32 SceneType = BodyInstance.UseAsyncScene() ? PST_Async : PST_Sync;
+	const uint32 SceneType = BodyInstance.UseAsyncScene(PhysScene) ? PST_Async : PST_Sync;
 
 	for (int32 i = 0; i < NumAggregates; i++)
 	{
@@ -1481,11 +1481,11 @@ void UInstancedStaticMeshComponent::SetupNewInstanceData(FInstancedStaticMeshIns
 			if (AggregateIndex >= Aggregates.Num())
 			{
 				// Get the scene type from the main BodyInstance
-				const uint32 SceneType = BodyInstance.UseAsyncScene() ? PST_Async : PST_Sync;
-
 				FPhysScene* PhysScene = GetWorld()->GetPhysicsScene();
 				check(PhysScene);
 
+				const uint32 SceneType = BodyInstance.UseAsyncScene(PhysScene) ? PST_Async : PST_Sync;
+
 				auto* Aggregate = GPhysXSDK->createAggregate(AggregateMaxSize, false);
 				const int32 AddedIndex = Aggregates.Add(Aggregate);
 				check(AddedIndex == AggregateIndex);
diff --git a/Engine/Source/Runtime/Engine/Private/PhysicsEngine/BodyInstance.cpp b/Engine/Source/Runtime/Engine/Private/PhysicsEngine/BodyInstance.cpp
--- a/Engine/Source/Runtime/Engine/Private/PhysicsEngine/BodyInstance.cpp
+++ b/Engine/Source/Runtime/Engine/Private/PhysicsEngine/BodyInstance.cpp
@@ -1275,7 +1275,7 @@ void FBodyInstance::InitBodies(TArray<FBodyInstance*>& Bodies, TArray<FTransform
 			PNewDynamic = GPhysXSDK->createRigidDynamic(PTransform);
 
 			// Put the dynamic actor in one scene or the other
-			if(!Instance->UseAsyncScene())
+			if(!Instance->UseAsyncScene(InRBScene))
 			{
 				PNewActorSync = PNewDynamic;
 				PSceneForNewDynamic = PSceneSync;
@@ -1365,7 +1365,7 @@ void FBodyInstance::InitBodies(TArray<FBodyInstance*>& Bodies, TArray<FTransform
 		{
 			if(PNewDynamic)
 			{
-				InRBScene->DeferAddActor(Instance, PNewDynamic, Instance->UseAsyncScene() ? PST_Async : PST_Sync);
+				InRBScene->DeferAddActor(Instance, PNewDynamic, Instance->UseAsyncScene(InRBScene) ? PST_Async : PST_Sync);
 			}
 			else
 			{
@@ -4095,12 +4095,19 @@ void FBodyInstance::FixupData(class UObject* Loader)
 
 bool FBodyInstance::UseAsyncScene() const
 {
-	bool bHasAsyncScene = true;
 	UPrimitiveComponent* OwnerComponentInst = OwnerComponent.Get();
-	if (OwnerComponentInst)
+	UWorld* World = OwnerComponentInst ? OwnerComponentInst->GetWorld() : nullptr;
+	return UseAsyncScene(World ? World->GetPhysicsScene() : nullptr);
+}
+
+bool FBodyInstance::UseAsyncScene(const FPhysScene* PhysScene) const
+{
+	bool bHasAsyncScene = true;
+	if (PhysScene)
 	{
-		bHasAsyncScene = OwnerComponentInst->GetWorld()->GetPhysicsScene()->HasAsyncScene();
+		bHasAsyncScene = PhysScene->HasAsyncScene();
 	}
+
 	return bUseAsyncScene && bHasAsyncScene;
 }
 
diff --git a/Engine/Source/Runtime/Engine/Private/PhysicsEngine/PhysScene.cpp b/Engine/Source/Runtime/Engine/Private/PhysicsEngine/PhysScene.cpp
--- a/Engine/Source/Runtime/Engine/Private/PhysicsEngine/PhysScene.cpp
+++ b/Engine/Source/Runtime/Engine/Private/PhysicsEngine/PhysScene.cpp
@@ -27,12 +27,12 @@
 static int32 PhysXSceneCount = 1;
 static const int PhysXSlowRebuildRate = 10;
 
-FORCEINLINE EPhysicsSceneType SceneType_AssumesLocked(const FBodyInstance* BodyInstance)
+EPhysicsSceneType FPhysScene::SceneType_AssumesLocked(const FBodyInstance* BodyInstance) const
 {
 #if WITH_PHYSX
 	//This is a helper function for dynamic actors - static actors are in both scenes
 	check(BodyInstance->GetPxRigidBody_AssumesLocked());
-	return UPhysicsSettings::Get()->bEnableAsyncScene && BodyInstance->UseAsyncScene() ? PST_Async : PST_Sync;
+	return UPhysicsSettings::Get()->bEnableAsyncScene && BodyInstance->UseAsyncScene(this) ? PST_Async : PST_Sync;
 #endif
 
 	return PST_Sync;
diff --git a/Engine/Source/Runtime/Engine/Private/PhysicsEngine/PhysicsHandleComponent.cpp b/Engine/Source/Runtime/Engine/Private/PhysicsEngine/PhysicsHandleComponent.cpp
--- a/Engine/Source/Runtime/Engine/Private/PhysicsEngine/PhysicsHandleComponent.cpp
+++ b/Engine/Source/Runtime/Engine/Private/PhysicsEngine/PhysicsHandleComponent.cpp
@@ -119,7 +119,7 @@ void UPhysicsHandleComponent::GrabComponent(UPrimitiveComponent* InComponent, FN
 
 				// Remember the scene index that the handle joint/actor are in.
 				FPhysScene* RBScene = FPhysxUserData::Get<FPhysScene>(Scene->userData);
-				const uint32 SceneType = InComponent->BodyInstance.UseAsyncScene() ? PST_Async : PST_Sync;
+				const uint32 SceneType = InComponent->BodyInstance.UseAsyncScene(RBScene) ? PST_Async : PST_Sync;
 				SceneIndex = RBScene->PhysXSceneIndex[SceneType];
 
 				// Setting up the joint
diff --git a/Engine/Source/Runtime/Engine/Private/SkeletalMeshComponentPhysics.cpp b/Engine/Source/Runtime/Engine/Private/SkeletalMeshComponentPhysics.cpp
--- a/Engine/Source/Runtime/Engine/Private/SkeletalMeshComponentPhysics.cpp
+++ b/Engine/Source/Runtime/Engine/Private/SkeletalMeshComponentPhysics.cpp
@@ -810,7 +810,7 @@ void USkeletalMeshComponent::InitArticulated(FPhysScene* PhysScene)
 #endif //WITH_PHYSX
 
 			// Remember if we have bodies in sync/async scene, so we know which scene(s) to lock when moving bodies
-			if(BodyInst->UseAsyncScene())
+			if(BodyInst->UseAsyncScene(PhysScene))
 			{
 				bHasBodiesInAsyncScene = true;
 			}
@@ -901,7 +901,7 @@ void USkeletalMeshComponent::TermArticulated()
 	}
 
 	// Get the scene type from the SkeletalMeshComponent's BodyInstance
-	const uint32 SceneType = BodyInstance.UseAsyncScene() ? PST_Async : PST_Sync;
+	const uint32 SceneType = BodyInstance.UseAsyncScene(PhysScene) ? PST_Async : PST_Sync;
 	PxScene* PScene = PhysScene->GetPhysXScene(SceneType);
 	SCOPED_SCENE_WRITE_LOCK(PScene);
 
diff --git a/Engine/Source/Runtime/Engine/Public/PhysicsPublic.h b/Engine/Source/Runtime/Engine/Public/PhysicsPublic.h
--- a/Engine/Source/Runtime/Engine/Public/PhysicsPublic.h
+++ b/Engine/Source/Runtime/Engine/Public/PhysicsPublic.h
@@ -574,6 +574,9 @@ private:
 	/** Called when all subscenes of a given scene are complete, calls  ProcessPhysScene*/
 	void SceneCompletionTask(ENamedThreads::Type CurrentThread, const FGraphEventRef& MyCompletionGraphEvent, EPhysicsSceneType SceneType);
 
+	/** Helper function for determining which scene a dyanmic body is in*/
+	EPhysicsSceneType SceneType_AssumesLocked(const FBodyInstance* BodyInstance) const;
+
 #if WITH_SUBSTEPPING
 	/** Task created from TickPhysScene so we can substep without blocking */
 	bool SubstepSimulation(uint32 SceneType, FGraphEventRef& InOutCompletionEvent);
