Engine: Fixed an issue that prevented the BlendMode from being changed on FCanvasNGonItem

diff --git a/Engine/Source/Runtime/Engine/Private/UserInterface/CanvasItem.cpp b/Engine/Source/Runtime/Engine/Private/UserInterface/CanvasItem.cpp
--- a/Engine/Source/Runtime/Engine/Private/UserInterface/CanvasItem.cpp
+++ b/Engine/Source/Runtime/Engine/Private/UserInterface/CanvasItem.cpp
@@ -1362,6 +1362,7 @@ void FCanvasTriangleItem::SetColor( const FLinearColor& InColor )
 
 void FCanvasNGonItem::Draw( class FCanvas* InCanvas )
 {
+	TriListItem->BlendMode = BlendMode;
 	TriListItem->Draw( InCanvas );
 }
 
