fixed DeferredDecal material output pin greyed state

diff --git a/Engine/Source/Runtime/Engine/Private/Materials/Material.cpp b/Engine/Source/Runtime/Engine/Private/Materials/Material.cpp
--- a/Engine/Source/Runtime/Engine/Private/Materials/Material.cpp
+++ b/Engine/Source/Runtime/Engine/Private/Materials/Material.cpp
@@ -3742,11 +3742,15 @@ bool UMaterial::IsPropertyActive(EMaterialProperty InProperty) const
 	}
 	else if(MaterialDomain == MD_DeferredDecal)
 	{
-		// todo: this should also check for <= MP_CustomizedUVs7 but then MaterialAttruibutes would not return true, should it? Wht we don't chekc for the checkbox in the material
-		if(InProperty >= MP_CustomizedUVs0)
+		if (InProperty >= MP_CustomizedUVs0 && InProperty <= MP_CustomizedUVs7)
 		{
 			return true;
 		}
+		else if (InProperty == MP_MaterialAttributes)
+		{
+			// todo: MaterialAttruibutes would not return true, should it? Why we don't check for the checkbox in the material
+			return true;
+		}
 
 		switch(DecalBlendMode)
 		{
