UE-6274 - DemoNetDriver isn't exposed via ENGINE_API

diff --git a/Engine/Source/Runtime/Engine/Classes/Engine/DemoNetDriver.h b/Engine/Source/Runtime/Engine/Classes/Engine/DemoNetDriver.h
--- a/Engine/Source/Runtime/Engine/Classes/Engine/DemoNetDriver.h
+++ b/Engine/Source/Runtime/Engine/Classes/Engine/DemoNetDriver.h
@@ -8,8 +8,7 @@
  * Simulated network driver for recording and playing back game sessions.
  */
 UCLASS(transient, config=Engine)
-class UDemoNetDriver
-	: public UNetDriver
+class ENGINE_API UDemoNetDriver : public UNetDriver
 {
 	GENERATED_UCLASS_BODY()
 
@@ -90,8 +89,8 @@ public:
 	virtual void TickFlush( float DeltaSeconds ) override;
 	virtual void ProcessRemoteFunction( class AActor* Actor, class UFunction* Function, void* Parameters, struct FOutParmRec* OutParms, struct FFrame* Stack, class UObject* SubObject = nullptr ) override;
 	virtual bool IsAvailable() const override { return true; }
-	ENGINE_API void SkipTime(const float InTimeToSkip);
-	ENGINE_API void GotoTimeInSeconds( const float TimeInSeconds );
+	void SkipTime(const float InTimeToSkip);
+	void GotoTimeInSeconds( const float TimeInSeconds );
 	bool InitConnectInternal( FString& Error );
 	virtual bool ShouldClientDestroyTearOffActors() const override;
 	virtual bool ShouldSkipRepNotifies() const override;
