UE-9986 - UNetDriver shadow variables World

diff --git a/Engine/Source/Runtime/Engine/Private/NetworkDriver.cpp b/Engine/Source/Runtime/Engine/Private/NetworkDriver.cpp
--- a/Engine/Source/Runtime/Engine/Private/NetworkDriver.cpp
+++ b/Engine/Source/Runtime/Engine/Private/NetworkDriver.cpp
@@ -660,8 +660,7 @@ bool UNetDriver::IsLevelInitializedForActor(const AActor* InActor, const UNetCon
 {
 	check(InActor);
     check(InConnection);
-	UWorld* World = InActor->GetWorld();
-	check(World);
+	check(World == InActor->GetWorld());
 
 	// we can't create channels while the client is in the wrong world
 	const bool bCorrectWorld = (InConnection->ClientWorldPackageName == World->GetOutermost()->GetFName() && InConnection->ClientHasInitializedLevelFor(InActor));
@@ -1842,8 +1841,9 @@ int32 UNetDriver::ServerReplicateActors(float DeltaSeconds)
 	bool bNetRelevantActorCount = false;
 	int32 NetRelevantActorCount = 0;
 
+	check( World );
+
 	FMemMark Mark(FMemStack::Get());
-	UWorld* World = NULL;
 	// initialize connections
 	bool bFoundReadyConnection = false; // whether we have at least one connection ready for property replication
 	for( int32 ConnIdx = 0; ConnIdx < ClientConnections.Num(); ConnIdx++ )
@@ -1859,16 +1859,7 @@ int32 UNetDriver::ServerReplicateActors(float DeltaSeconds)
 		AActor* OwningActor = Connection->OwningActor;
 		if (OwningActor != NULL && Connection->State == USOCK_Open && (Connection->Driver->Time - Connection->LastReceiveTime < 1.5f))
 		{
-			if( !World )
-			{
-				World = OwningActor->GetWorld();
-			}
-			else
-			{
-				check( World == OwningActor->GetWorld() );
-			}
-
-			check( World );
+			check( World == OwningActor->GetWorld() );
 			if( !bNetRelevantActorCount )
 			{
 				NetRelevantActorCount = World->GetNetRelevantActorCount() + 2;
