Stopped sound waves from clearing CompressedFormatData early.
Fix for UE-11486.

diff --git a/Engine/Source/Runtime/Engine/Private/AudioDerivedData.cpp b/Engine/Source/Runtime/Engine/Private/AudioDerivedData.cpp
--- a/Engine/Source/Runtime/Engine/Private/AudioDerivedData.cpp
+++ b/Engine/Source/Runtime/Engine/Private/AudioDerivedData.cpp
@@ -1104,8 +1104,10 @@ void USoundWave::WillNeverCacheCookedPlatformDataAgain()
 	// this is called after we have finished caching the platform data but before we have saved the data
 	// so need to keep the cached platform data around 
 	Super::WillNeverCacheCookedPlatformDataAgain();
-	RawData.RemoveBulkData();
-	CompressedFormatData.FlushData();
+
+	// TODO: We can clear these arrays if we never need to cook again. 
+	// RawData.RemoveBulkData();
+	// CompressedFormatData.FlushData();
 }
 #endif
 
