Fixed: In-Editor Seamless Travel Crash On Round-Trip
UE-13111

diff --git a/Engine/Source/Runtime/Engine/Private/World.cpp b/Engine/Source/Runtime/Engine/Private/World.cpp
--- a/Engine/Source/Runtime/Engine/Private/World.cpp
+++ b/Engine/Source/Runtime/Engine/Private/World.cpp
@@ -4292,23 +4292,6 @@ void FSeamlessTravelHandler::SeamlessTravelLoadCallback(const FName& PackageName
 			}
 		}
 
-#if WITH_EDITOR
-		FWorldContext &WorldContext = GEngine->GetWorldContextFromWorldChecked(CurrentWorld);
-		if (WorldContext.WorldType == EWorldType::PIE)
-		{
-			// If we are a PIE world and the world we just "loaded" is already initialized, then we're probably travelling to the editor world and we
-			// need to create a PIE world by duplication instead
-			if ( World && World->bIsWorldInitialized)
-			{
-				FString PIEMapName;
-				World = GEngine->CreatePIEWorldByDuplication(WorldContext, World, PIEMapName);
-				World->ClearFlags(RF_Standalone);
-				// CreatePIEWorldByDuplication clears GIsPlayInEditorWorld so set it again
-				GIsPlayInEditorWorld = true;
-			}
-		}
-#endif
-
 		SetHandlerLoadedData(LevelPackage, World);
 	}
 }
@@ -4514,11 +4497,31 @@ void FSeamlessTravelHandler::StartLoadingDestination()
 		// Set the world type in the static map, so that UWorld::PostLoad can set the world type
 		const FName URLMapFName = FName(*PendingTravelURL.Map);
 		UWorld::WorldTypePreLoadMap.FindOrAdd(URLMapFName) = CurrentWorld->WorldType;
-
-		LoadPackageAsync(PendingTravelURL.Map, 
+		// In PIE we might want to mangle MapPackageName when traveling to a map loaded in the editor
+		FString URLMapPackageName = PendingTravelURL.Map;
+		FString URLMapPackageToLoadFrom = PendingTravelURL.Map;
+		uint32 PackageFlags = 0;
+		int32 PIEInstanceID = INDEX_NONE;
+		
+#if WITH_EDITOR
+		if (GIsEditor)
+		{
+			UPackage* EditorLevelPackage = (UPackage*)StaticFindObjectFast(UPackage::StaticClass(), NULL, URLMapFName, 0, 0, RF_PendingKill);
+			if (EditorLevelPackage)
+			{
+				FWorldContext &WorldContext = GEngine->GetWorldContextFromHandleChecked(WorldContextHandle);
+				PackageFlags |= PKG_PlayInEditor;
+				PIEInstanceID = WorldContext.PIEInstance;
+				URLMapPackageName = UWorld::ConvertToPIEPackageName(URLMapPackageName, PIEInstanceID);
+			}
+		}
+#endif
+		LoadPackageAsync(URLMapPackageName, 
 			FLoadPackageAsyncDelegate::CreateRaw(this, &FSeamlessTravelHandler::SeamlessTravelLoadCallback), 
-			PendingTravelGuid.IsValid() ? &PendingTravelGuid : NULL
-			);
+			PendingTravelGuid.IsValid() ? &PendingTravelGuid : NULL,
+			NAME_None,
+			*URLMapPackageToLoadFrom
+			).SetPackageData(PackageFlags, PIEInstanceID);
 	}
 	else
 	{
