[UE-11897] Reparenting a component class between scene/non-scene base types will no longer corrupt the Blueprint component tree.

change summary:
- added utility methods to USimpleConstructionScript for finding both the current scene root component/node and promotable child nodes
- modified USimpleConstructionScript::RemoveNodeAndPromoteChildren() to call the new utility method for finding the promotable child node candidate
- modified USimpleConstructionScript::ValidateSceneRootNode() to call the new utility method for finding the current native or non-native scene root component
- modified USimpleConstructionScript::PostLoad() to detect and handle all cases of reparenting from USceneComponent to UActorComponent types and vice-versa

diff --git a/Engine/Source/Runtime/Engine/Classes/Engine/SimpleConstructionScript.h b/Engine/Source/Runtime/Engine/Classes/Engine/SimpleConstructionScript.h
--- a/Engine/Source/Runtime/Engine/Classes/Engine/SimpleConstructionScript.h
+++ b/Engine/Source/Runtime/Engine/Classes/Engine/SimpleConstructionScript.h
@@ -90,6 +90,12 @@ private:
 	TArray<USCS_Node*> ActorComponentNodes_DEPRECATED;
 
 #if WITH_EDITOR
+	/** Helper function to find the current scene root component template and/or owning SCS node */
+	USceneComponent* GetSceneRootComponentTemplate(USCS_Node** OutSCSNode = nullptr) const;
+
+	/** Helper function to find an appropriate child node index that can be promoted to the parent's level */
+	int32 FindPromotableChildNodeIndex(USCS_Node* InParentNode) const;
+
 	/** Helper function for generating list of currently used names */
 	void GenerateListOfExistingNames(TArray<FName>& CurrentNames) const;
 
diff --git a/Engine/Source/Runtime/Engine/Private/SimpleConstructionScript.cpp b/Engine/Source/Runtime/Engine/Private/SimpleConstructionScript.cpp
--- a/Engine/Source/Runtime/Engine/Private/SimpleConstructionScript.cpp
+++ b/Engine/Source/Runtime/Engine/Private/SimpleConstructionScript.cpp
@@ -148,9 +148,6 @@ void USimpleConstructionScript::PreloadChain()
 void USimpleConstructionScript::PostLoad()
 {
 	Super::PostLoad();
-	
-	int32 NodeIndex;
-	TArray<USCS_Node*> Nodes = GetAllNodes();
 
 #if WITH_EDITOR
 	// Get the Blueprint that owns the SCS
@@ -162,17 +159,106 @@ void USimpleConstructionScript::PostLoad()
 		return;
 	}
 
-	// Fix up any uninitialized category names
+	int32 NodeIndex;
+	TArray<USCS_Node*> Nodes = GetAllNodes();
+	USCS_Node* SceneRootNode = nullptr;
+	USceneComponent* SceneRootComponentTemplate = GetSceneRootComponentTemplate(&SceneRootNode);
+
 	for (NodeIndex=0; NodeIndex < Nodes.Num(); ++NodeIndex)
 	{
 		USCS_Node* Node = Nodes[NodeIndex];
+
+		// Fix up any uninitialized category names
 		if(Node->CategoryName == NAME_None)
 		{
 			Node->CategoryName = TEXT("Default");
 		}
-	}
 
+		// Fix up components that may have switched from scene to non-scene type and vice-versa
+		if(Node->ComponentTemplate != nullptr)
+		{
+			// Check to see if switched from non-scene to scene component type
+			if(Node->ComponentTemplate->IsA<USceneComponent>())
+			{
+				if(RootNodes.Contains(Node)
+					&& SceneRootComponentTemplate != nullptr
+					&& Node->ComponentTemplate != SceneRootComponentTemplate
+					&& Node->ParentComponentOrVariableName == NAME_None)
+				{
+					if(SceneRootNode != nullptr)
+					{
+						if(RootNodes.Contains(SceneRootNode))
+						{
+							// Reparent to this BP's root node if it's still in the root set
+							RootNodes.Remove(Node);
+							SceneRootNode->ChildNodes.Add(Node);
+						}
+						else
+						{
+							// Parent to an inherited parent BP's node if not already attached
+							Node->SetParent(SceneRootNode);
+						}
+					}
+					else
+					{
+						// Parent to the native component template if not already attached
+						Node->SetParent(SceneRootComponentTemplate);
+					}
+				}
+			}
+			else
+			{
+				// Otherwise, check to see if switched from scene to non-scene component type
+				int32 RootNodeIndex = INDEX_NONE;
+				if(!RootNodes.Find(Node, RootNodeIndex))
+				{
+					// Move the node into the root set if it's currently in the scene hierarchy
+					USCS_Node* ParentNode = FindParentNode(Node);
+					if(ParentNode != nullptr)
+					{
+						ParentNode->ChildNodes.Remove(Node);
+					}
+
+					RootNodes.Add(Node);
+				}
+				else
+				{
+					// Otherwise, if it's a root node, promote one of its children (if any) to take its place
+					int32 PromoteIndex = FindPromotableChildNodeIndex(Node);
+					if(PromoteIndex != INDEX_NONE)
+					{
+						// Remove it as a child node
+						USCS_Node* ChildToPromote = Node->ChildNodes[PromoteIndex];
+						Node->ChildNodes.RemoveAt(PromoteIndex);
+
+						// Insert it as a root node just before its prior parent node; this way if it switches back to a scene type it won't supplant the new root we've just created
+						RootNodes.Insert(ChildToPromote, RootNodeIndex);
+
+						// Append previous root node's children to the new root
+						ChildToPromote->ChildNodes.Append(Node->ChildNodes);
+
+						// Clear all child nodes from the old root (because it's now a non-scene type and no longer supports attached components)
+						Node->ChildNodes.Empty();
+
+						// Copy any previous external attachment info from the previous root node
+						ChildToPromote->bIsParentComponentNative = Node->bIsParentComponentNative;
+						ChildToPromote->ParentComponentOrVariableName = Node->ParentComponentOrVariableName;
+						ChildToPromote->ParentComponentOwnerClassName = Node->ParentComponentOwnerClassName;
+					}
+
+					// Clear info for any previous external attachment if set
+					if(Node->ParentComponentOrVariableName != NAME_None)
+					{
+						Node->bIsParentComponentNative = false;
+						Node->ParentComponentOrVariableName = NAME_None;
+						Node->ParentComponentOwnerClassName = NAME_None;
+					}
+				}
+			}
+		}
+	}
 #endif // WITH_EDITOR
+
 	// Fix up native/inherited parent attachments, in case anything has changed
 	FixupRootNodeParentReferences();
 
@@ -495,36 +581,45 @@ void USimpleConstructionScript::RemoveNode(USCS_Node* Node)
 	}
 }
 
-void USimpleConstructionScript::RemoveNodeAndPromoteChildren(USCS_Node* Node)
+int32 USimpleConstructionScript::FindPromotableChildNodeIndex(USCS_Node* InParentNode) const
 {
-	Node->Modify();
+	int32 PromoteIndex = INDEX_NONE;
 
-	if (RootNodes.Contains(Node))
+	if (InParentNode->ChildNodes.Num() > 0)
 	{
-		USCS_Node* ChildToPromote = NULL;
+		PromoteIndex = 0;
+		USCS_Node* Child = InParentNode->ChildNodes[PromoteIndex];
 
-		// Pick the first child to promote as a new 'root'
-		if (Node->ChildNodes.Num() > 0)
+		// if this is an editor-only component, then it can't have any game-component children (better make sure that's the case)
+		if (Child->ComponentTemplate != NULL && Child->ComponentTemplate->IsEditorOnly())
 		{
-			int32 PromoteIndex = 0;
-			ChildToPromote = Node->ChildNodes[PromoteIndex];
-
-			// if this is an editor-only component, then it can't have any game-component children (better make sure that's the case)
-			if (ChildToPromote->ComponentTemplate != NULL && ChildToPromote->ComponentTemplate->IsEditorOnly())
+			for (int32 ChildIndex = 1; ChildIndex < InParentNode->ChildNodes.Num(); ++ChildIndex)
 			{
-				for (int32 ChildIndex = 1; ChildIndex < Node->ChildNodes.Num(); ++ChildIndex)
+				Child = InParentNode->ChildNodes[ChildIndex];
+				// we found a game-component sibling, better make it the child to promote
+				if (Child->ComponentTemplate != NULL && !Child->ComponentTemplate->IsEditorOnly())
 				{
-					USCS_Node* Child = Node->ChildNodes[ChildIndex];
-					// we found a game-component sibling, better make it the child to promote
-					if (Child->ComponentTemplate != NULL && !Child->ComponentTemplate->IsEditorOnly())
-					{
-						ChildToPromote = Child;
-						PromoteIndex = ChildIndex;
-						break;
-					}
+					PromoteIndex = ChildIndex;
+					break;
 				}
 			}
+		}
+	}
+
+	return PromoteIndex;
+}
 
+void USimpleConstructionScript::RemoveNodeAndPromoteChildren(USCS_Node* Node)
+{
+	Node->Modify();
+
+	if (RootNodes.Contains(Node))
+	{
+		USCS_Node* ChildToPromote = nullptr;
+		int32 PromoteIndex = FindPromotableChildNodeIndex(Node);
+		if(PromoteIndex != INDEX_NONE)
+		{
+			ChildToPromote = Node->ChildNodes[PromoteIndex];
 			Node->ChildNodes.RemoveAt(PromoteIndex);
 		}
 
@@ -616,103 +711,131 @@ USCS_Node* USimpleConstructionScript::FindSCSNodeByGuid(const FGuid Guid) const
 	return ReturnSCSNode;
 }
 
-void USimpleConstructionScript::ValidateSceneRootNodes()
-{
 #if WITH_EDITOR
+USceneComponent* USimpleConstructionScript::GetSceneRootComponentTemplate(USCS_Node** OutSCSNode) const
+{
 	UBlueprint* Blueprint = GetBlueprint();
 
-	if(DefaultSceneRootNode == nullptr)
+	UClass* GeneratedClass = GetOwnerClass();
+
+	if(OutSCSNode)
 	{
-		// If applicable, create a default scene component node
-		if(Blueprint != nullptr
-			&& FBlueprintEditorUtils::IsActorBased(Blueprint)
-			&& Blueprint->BlueprintType != BPTYPE_MacroLibrary)
-		{
-			DefaultSceneRootNode = CreateNode(USceneComponent::StaticClass(), USceneComponent::GetDefaultSceneRootVariableName());
-			CastChecked<USceneComponent>(DefaultSceneRootNode->ComponentTemplate)->bVisualizeComponent = true;
-		}
+		*OutSCSNode = nullptr;
 	}
 
-	if(DefaultSceneRootNode != nullptr)
+	// Get the Blueprint class default object
+	AActor* CDO = nullptr;
+	if(GeneratedClass != nullptr)
 	{
-		UClass* GeneratedClass = GetOwnerClass();
+		CDO = Cast<AActor>(GeneratedClass->GetDefaultObject(false));
+	}
 
-		// Get the Blueprint class default object
-		AActor* CDO = nullptr;
-		if(GeneratedClass != nullptr)
+	// If the generated class does not yet have a CDO, defer to the parent class
+	if(CDO == nullptr && Blueprint->ParentClass != nullptr)
+	{
+		CDO = Cast<AActor>(Blueprint->ParentClass->GetDefaultObject(false));
+	}
+
+	// Check to see if we already have a native root component template
+	USceneComponent* RootComponentTemplate = nullptr;
+	if(CDO != nullptr)
+	{
+		// If the root component property is not set, the first available scene component will be used as the root. This matches what's done in the SCS editor.
+		RootComponentTemplate = CDO->GetRootComponent();
+		if(!RootComponentTemplate)
 		{
-			CDO = Cast<AActor>(GeneratedClass->GetDefaultObject(false));
+			TInlineComponentArray<USceneComponent*> SceneComponents;
+			CDO->GetComponents(SceneComponents);
+			if(SceneComponents.Num() > 0)
+			{
+				RootComponentTemplate = SceneComponents[0];
+			}
 		}
+	}
 
-		// If the generated class does not yet have a CDO, defer to the parent class
-		if(CDO == nullptr && Blueprint->ParentClass != nullptr)
+	// Don't add the default scene root if we already have a native scene root component
+	if(!RootComponentTemplate)
+	{
+		// Get the Blueprint hierarchy
+		TArray<UBlueprint*> BPStack;
+		if(Blueprint->GeneratedClass != nullptr)
 		{
-			CDO = Cast<AActor>(Blueprint->ParentClass->GetDefaultObject(false));
+			UBlueprint::GetBlueprintHierarchyFromClass(Blueprint->GeneratedClass, BPStack);
 		}
-
-		// Check to see if we already have a native root component
-		bool bHasSceneComponentRootNodes = false;
-		if(CDO != nullptr)
+		else if(Blueprint->ParentClass != nullptr)
 		{
-			// If the root component property is not set, the first available scene component will be used as the root. This matches what's done in the SCS editor.
-			bHasSceneComponentRootNodes = CDO->GetRootComponent() != nullptr;
-			if(!bHasSceneComponentRootNodes)
-			{
-				TInlineComponentArray<USceneComponent*> SceneComponents;
-				CDO->GetComponents(SceneComponents);
-				bHasSceneComponentRootNodes = SceneComponents.Num() > 0;
-			}
+			UBlueprint::GetBlueprintHierarchyFromClass(Blueprint->ParentClass, BPStack);
 		}
 
-		// Don't add the default scene root if we already have a native scene root component
-		if(!bHasSceneComponentRootNodes)
+		// Note: Normally if the Blueprint has a parent, we can assume that the parent already has a scene root component set,
+		// ...but we'll run through the hierarchy just in case there are legacy BPs out there that might not adhere to this assumption.
+		TArray<const USimpleConstructionScript*> SCSStack;
+		SCSStack.Add(this);
+
+		for(int32 StackIndex = 0; StackIndex < BPStack.Num(); ++StackIndex)
 		{
-			// Get the Blueprint hierarchy
-			TArray<UBlueprint*> BPStack;
-			if(Blueprint->GeneratedClass != nullptr)
+			if(BPStack[StackIndex] && BPStack[StackIndex]->SimpleConstructionScript)
 			{
-				UBlueprint::GetBlueprintHierarchyFromClass(Blueprint->GeneratedClass, BPStack);
+				SCSStack.AddUnique(BPStack[StackIndex]->SimpleConstructionScript);
 			}
-			else if(Blueprint->ParentClass != nullptr)
-			{
-				UBlueprint::GetBlueprintHierarchyFromClass(Blueprint->ParentClass, BPStack);
-			}
-
-			// Note: Normally if the Blueprint has a parent, we can assume that the parent already has a scene root component set,
-			// ...but we'll run through the hierarchy just in case there are legacy BPs out there that might not adhere to this assumption.
-			TArray<USimpleConstructionScript*> SCSStack;
-			SCSStack.Add(this);
+		}
 
-			for(int32 StackIndex = 0; StackIndex < BPStack.Num(); ++StackIndex)
+		for(int32 StackIndex = 0; StackIndex < SCSStack.Num() && !RootComponentTemplate; ++StackIndex)
+		{
+			// Check for any scene component nodes in the root set that are not the default scene root
+			const TArray<USCS_Node*>& SCSRootNodes = SCSStack[StackIndex]->GetRootNodes();
+			for(int32 RootNodeIndex = 0; RootNodeIndex < SCSRootNodes.Num() && RootComponentTemplate == nullptr; ++RootNodeIndex)
 			{
-				if(BPStack[StackIndex] && BPStack[StackIndex]->SimpleConstructionScript)
+				USCS_Node* RootNode = SCSRootNodes[RootNodeIndex];
+				if(RootNode != nullptr
+					&& RootNode != DefaultSceneRootNode
+					&& RootNode->ComponentTemplate != nullptr
+					&& RootNode->ComponentTemplate->IsA<USceneComponent>())
 				{
-					SCSStack.AddUnique(BPStack[StackIndex]->SimpleConstructionScript);
+					if(OutSCSNode)
+					{
+						*OutSCSNode = RootNode;
+					}
+					
+					RootComponentTemplate = Cast<USceneComponent>(RootNode->ComponentTemplate);
 				}
 			}
+		}
+	}
 
-			for(int32 StackIndex = 0; StackIndex < SCSStack.Num() && !bHasSceneComponentRootNodes; ++StackIndex)
-			{
-				// Check for any scene component nodes in the root set that are not the default scene root
-				const TArray<USCS_Node*>& SCSRootNodes = SCSStack[StackIndex]->GetRootNodes();
-				for(int32 RootNodeIndex = 0; RootNodeIndex < SCSRootNodes.Num() && !bHasSceneComponentRootNodes; ++RootNodeIndex)
-				{
-					USCS_Node* RootNode = SCSRootNodes[RootNodeIndex];
-					bHasSceneComponentRootNodes = RootNode != nullptr
-						&& RootNode != DefaultSceneRootNode
-						&& RootNode->ComponentTemplate != nullptr
-						&& RootNode->ComponentTemplate->IsA<USceneComponent>();
-				}
-			}
+	return RootComponentTemplate;
+}
+#endif
+
+void USimpleConstructionScript::ValidateSceneRootNodes()
+{
+#if WITH_EDITOR
+	UBlueprint* Blueprint = GetBlueprint();
+
+	if(DefaultSceneRootNode == nullptr)
+	{
+		// If applicable, create a default scene component node
+		if(Blueprint != nullptr
+			&& FBlueprintEditorUtils::IsActorBased(Blueprint)
+			&& Blueprint->BlueprintType != BPTYPE_MacroLibrary)
+		{
+			DefaultSceneRootNode = CreateNode(USceneComponent::StaticClass(), USceneComponent::GetDefaultSceneRootVariableName());
+			CastChecked<USceneComponent>(DefaultSceneRootNode->ComponentTemplate)->bVisualizeComponent = true;
 		}
+	}
+
+	if(DefaultSceneRootNode != nullptr)
+	{
+		// Get the current root component template
+		const USceneComponent* RootComponentTemplate = GetSceneRootComponentTemplate();
 
 		// Add the default scene root back in if there are no other scene component nodes that can be used as root; otherwise, remove it
-		if(!bHasSceneComponentRootNodes
+		if(RootComponentTemplate == nullptr
 			&& !RootNodes.Contains(DefaultSceneRootNode))
 		{
 			RootNodes.Add(DefaultSceneRootNode);
 		}
-		else if(bHasSceneComponentRootNodes
+		else if(RootComponentTemplate != nullptr
 			&& RootNodes.Contains(DefaultSceneRootNode))
 		{
 			RootNodes.Remove(DefaultSceneRootNode);
