Moved buffer visualization setup code around so it can uses showflags AFTER overrides have been applied, allowing them to work correctly in non-editor modes when debug view modes are allowed.

diff --git a/Engine/Source/Runtime/Engine/Private/GameViewportClient.cpp b/Engine/Source/Runtime/Engine/Private/GameViewportClient.cpp
--- a/Engine/Source/Runtime/Engine/Private/GameViewportClient.cpp
+++ b/Engine/Source/Runtime/Engine/Private/GameViewportClient.cpp
@@ -804,8 +804,29 @@ void UGameViewportClient::Draw(FViewport* InViewport, FCanvas* SceneCanvas)
 
 	bool bUIDisableWorldRendering = false;
 	FGameViewDrawer GameViewDrawer;
+	
+	// create the view family for rendering the world scene to the viewport's render target
+	FSceneViewFamilyContext ViewFamily(FSceneViewFamily::ConstructionValues( 	
+		InViewport,
+		GetWorld()->Scene,
+		EngineShowFlags)
+		.SetRealtimeUpdate(true));
+
+	// Allow HMD to modify the view later, just before rendering
+	if (GEngine->HMDDevice.IsValid() && GEngine->IsStereoscopic3D(InViewport))
+	{
+		ISceneViewExtension* HmdViewExt = GEngine->HMDDevice->GetViewExtension();
+		if (HmdViewExt)
+		{
+			ViewFamily.ViewExtensions.Add(HmdViewExt);
+			HmdViewExt->ModifyShowFlags(ViewFamily.EngineShowFlags);
+		}
+	}
 
-	if (EngineShowFlags.VisualizeBuffer)
+	ESplitScreenType::Type SplitScreenConfig = GetCurrentSplitscreenConfiguration();
+	EngineShowFlagOverride(ESFIM_Game, (EViewModeIndex)ViewModeIndex, ViewFamily.EngineShowFlags, NAME_None, SplitScreenConfig != ESplitScreenType::None);
+
+	if (ViewFamily.EngineShowFlags.VisualizeBuffer && AllowDebugViewmodes())
 	{
 		// Process the buffer visualization console command
 		FName NewBufferVisualizationMode = NAME_None;
@@ -842,28 +863,6 @@ void UGameViewportClient::Draw(FViewport* InViewport, FCanvas* SceneCanvas)
 		}
 	}
 
-	// create the view family for rendering the world scene to the viewport's render target
-	FSceneViewFamilyContext ViewFamily(FSceneViewFamily::ConstructionValues( 	
-		InViewport,
-		GetWorld()->Scene,
-		EngineShowFlags)
-		.SetRealtimeUpdate(true));
-
-	// Allow HMD to modify the view later, just before rendering
-	if (GEngine->HMDDevice.IsValid() && GEngine->IsStereoscopic3D(InViewport))
-	{
-		ISceneViewExtension* HmdViewExt = GEngine->HMDDevice->GetViewExtension();
-		if (HmdViewExt)
-		{
-			ViewFamily.ViewExtensions.Add(HmdViewExt);
-			HmdViewExt->ModifyShowFlags(ViewFamily.EngineShowFlags);
-		}
-	}
-
-
-	ESplitScreenType::Type SplitScreenConfig = GetCurrentSplitscreenConfiguration();
-	EngineShowFlagOverride(ESFIM_Game, (EViewModeIndex)ViewModeIndex, ViewFamily.EngineShowFlags, NAME_None, SplitScreenConfig != ESplitScreenType::None);
-
 	TMap<ULocalPlayer*,FSceneView*> PlayerViewMap;
 
 	FAudioDevice* AudioDevice = GetWorld()->GetAudioDevice();
