Fixed in-editor changing of UNavigationSystem::bAllowClientSideNavigation_Experimental not propagating to existing instances of navigation data #UE4

diff --git a/Engine/Source/Runtime/Engine/Private/AI/Navigation/NavigationData.cpp b/Engine/Source/Runtime/Engine/Private/AI/Navigation/NavigationData.cpp
--- a/Engine/Source/Runtime/Engine/Private/AI/Navigation/NavigationData.cpp
+++ b/Engine/Source/Runtime/Engine/Private/AI/Navigation/NavigationData.cpp
@@ -109,7 +109,7 @@ ANavigationData::ANavigationData(const FObjectInitializer& ObjectInitializer)
 	, NavDataUniqueID(GetNextUniqueID())
 {
 	PrimaryActorTick.bCanEverTick = true;
-	bNetLoadOnClient = (HasAnyFlags(RF_ClassDefaultObject) == false) && (*GEngine->NavigationSystemClass != nullptr) && (GEngine->NavigationSystemClass->GetDefaultObject<UNavigationSystem>()->ShouldLoadNavigationOnClient());
+	bNetLoadOnClient = false;
 	bCanBeDamaged = false;
 	DefaultQueryFilter = MakeShareable(new FNavigationQueryFilter());
 	ObservedPathsTickInterval = 0.5;
@@ -137,9 +137,10 @@ void ANavigationData::PostInitProperties()
 			RuntimeGeneration = bRebuildAtRuntime_DEPRECATED ? ERuntimeGenerationType::Dynamic : ERuntimeGenerationType::Static;
 		}
 	}
-		
-	if (HasAnyFlags(RF_ClassDefaultObject) == false)
+	else
 	{
+		bNetLoadOnClient = (*GEngine->NavigationSystemClass != nullptr) && (GEngine->NavigationSystemClass->GetDefaultObject<UNavigationSystem>()->ShouldLoadNavigationOnClient(this));
+
 		UWorld* WorldOuter = GetWorld();
 		
 		if (WorldOuter != NULL && WorldOuter->GetNavigationSystem() != NULL)
@@ -182,6 +183,8 @@ void ANavigationData::PostLoad()
 	InstantiateAndRegisterRenderingComponent();
 
 	CachedWorld = GetWorld();
+
+	bNetLoadOnClient = (*GEngine->NavigationSystemClass != nullptr) && (GEngine->NavigationSystemClass->GetDefaultObject<UNavigationSystem>()->ShouldLoadNavigationOnClient(this));
 }
 
 void ANavigationData::TickActor(float DeltaTime, enum ELevelTick TickType, FActorTickFunction& ThisTickFunction)
