Compilation fix

diff --git a/Engine/Source/Runtime/Engine/Private/SimpleConstructionScript.cpp b/Engine/Source/Runtime/Engine/Private/SimpleConstructionScript.cpp
--- a/Engine/Source/Runtime/Engine/Private/SimpleConstructionScript.cpp
+++ b/Engine/Source/Runtime/Engine/Private/SimpleConstructionScript.cpp
@@ -284,7 +284,7 @@ void USimpleConstructionScript::PostLoad()
 				{
 					// If no native root component exists, find the first non-native, non-parented SCS node with a
 					// scene component template. This will be designated as the root component at construction time.
-					for(NodeIndex = 0; NodeIndex < RootNodes.Num(); ++NodeIndex)
+					for( int32 NodeIndex = 0; NodeIndex < RootNodes.Num(); ++NodeIndex)
 					{
 						USCS_Node* Node = RootNodes[NodeIndex];
 						if(Node->ParentComponentOrVariableName == NAME_None)
@@ -581,6 +581,8 @@ void USimpleConstructionScript::RemoveNode(USCS_Node* Node)
 	}
 }
 
+#if	WITH_EDITOR
+
 int32 USimpleConstructionScript::FindPromotableChildNodeIndex(USCS_Node* InParentNode) const
 {
 	int32 PromoteIndex = INDEX_NONE;
@@ -609,6 +611,8 @@ int32 USimpleConstructionScript::FindPromotableChildNodeIndex(USCS_Node* InParen
 	return PromoteIndex;
 }
 
+#endif // WITH_EDITOR
+
 void USimpleConstructionScript::RemoveNodeAndPromoteChildren(USCS_Node* Node)
 {
 	Node->Modify();
@@ -616,12 +620,14 @@ void USimpleConstructionScript::RemoveNodeAndPromoteChildren(USCS_Node* Node)
 	if (RootNodes.Contains(Node))
 	{
 		USCS_Node* ChildToPromote = nullptr;
+#if	WITH_EDITOR
 		int32 PromoteIndex = FindPromotableChildNodeIndex(Node);
 		if(PromoteIndex != INDEX_NONE)
 		{
 			ChildToPromote = Node->ChildNodes[PromoteIndex];
 			Node->ChildNodes.RemoveAt(PromoteIndex);
 		}
+#endif // WITH_EDITOR
 
 		Modify();
 
