Compile fix for editor-only 'override' function in AnimSequence

diff --git a/Engine/Source/Runtime/Engine/Classes/Animation/AnimSequence.h b/Engine/Source/Runtime/Engine/Classes/Animation/AnimSequence.h
--- a/Engine/Source/Runtime/Engine/Classes/Animation/AnimSequence.h
+++ b/Engine/Source/Runtime/Engine/Classes/Animation/AnimSequence.h
@@ -800,12 +800,12 @@ public:
 	 * @return					true if the operation was successful.
 	 */
 	ENGINE_API bool InsertFramesToRawAnimData( int32 StartFrame, int32 EndFrame, int32 CopyFrame);
-#endif
 
 	/** 
 	 * Add validation check to see if it's being ready to play or not
 	 */
-	virtual bool IsValidToPlay() const override;
+	ENGINE_API virtual bool IsValidToPlay() const override;
+#endif
 
 private:
 	/** 
diff --git a/Engine/Source/Runtime/Engine/Private/Animation/AnimSequence.cpp b/Engine/Source/Runtime/Engine/Private/Animation/AnimSequence.cpp
--- a/Engine/Source/Runtime/Engine/Private/Animation/AnimSequence.cpp
+++ b/Engine/Source/Runtime/Engine/Private/Animation/AnimSequence.cpp
@@ -300,11 +300,13 @@ int32 UAnimSequence::GetNumberOfTracks() const
 	return TrackToSkeletonMapTable.Num();
 }
 
+#if WITH_EDITOR
 bool UAnimSequence::IsValidToPlay() const
 {
 	// make sure sequence length is valid and raw animation data exists, and compressed
 	return ( SequenceLength > 0.f && RawAnimationData.Num() > 0 && CompressedTrackOffsets.Num() > 0 );
 }
+#endif
 
 void UAnimSequence::PreSave()
 {
