Moved notifying NavigationSystem about undo/redo events from LandscapeProxy to Actor for consistency and to have all actors benefit from it. #UE4

diff --git a/Engine/Source/Runtime/Engine/Private/ActorEditor.cpp b/Engine/Source/Runtime/Engine/Private/ActorEditor.cpp
--- a/Engine/Source/Runtime/Engine/Private/ActorEditor.cpp
+++ b/Engine/Source/Runtime/Engine/Private/ActorEditor.cpp
@@ -336,6 +336,9 @@ void AActor::PreEditUndo()
 		}
 	}
 
+	// let navigation system know to not care about this actor anymore
+	UNavigationSystem::ClearNavOctreeAll(this);
+
 	Super::PreEditUndo();
 }
 
@@ -351,6 +354,12 @@ void AActor::PostEditUndo()
 	if (!IsPendingKill())
 	{
 		ResetOwnedComponents();
+		// notify navigation system
+		UNavigationSystem::UpdateNavOctreeAll(this);
+	}
+	else
+	{
+		UNavigationSystem::ClearNavOctreeAll(this);
 	}
 
 	Super::PostEditUndo();
@@ -370,6 +379,12 @@ void AActor::PostEditUndo(TSharedPtr<ITransactionObjectAnnotation> TransactionAn
 	if (!IsPendingKill())
 	{
 		ResetOwnedComponents();
+		// notify navigation system
+		UNavigationSystem::UpdateNavOctreeAll(this);
+	}
+	else
+	{
+		UNavigationSystem::ClearNavOctreeAll(this);
 	}
 
 	Super::PostEditUndo(TransactionAnnotation);
diff --git a/Engine/Source/Runtime/Landscape/Private/Landscape.cpp b/Engine/Source/Runtime/Landscape/Private/Landscape.cpp
--- a/Engine/Source/Runtime/Landscape/Private/Landscape.cpp
+++ b/Engine/Source/Runtime/Landscape/Private/Landscape.cpp
@@ -993,8 +993,6 @@ void ALandscapeProxy::PreEditUndo()
 	if (GIsEditor && GetWorld() && !GetWorld()->IsPlayInEditor())
 	{
 		GEngine->DeferredCommands.AddUnique(TEXT("UpdateLandscapeEditorData"));
-
-		UNavigationSystem::ClearNavOctreeAll(this);
 	}
 }
 
@@ -1005,8 +1003,6 @@ void ALandscapeProxy::PostEditUndo()
 	if (GIsEditor && GetWorld() && !GetWorld()->IsPlayInEditor())
 	{
 		GEngine->DeferredCommands.AddUnique(TEXT("UpdateLandscapeEditorData"));
-
-		UNavigationSystem::UpdateNavOctreeAll(this);
 	}
 }
 
