UE-13076 "StartfpsChart" not generating data in cooked build of GDC demo (misc fixes, after StopFPSChart the location of the stats will be opened via file explorer)

diff --git a/Engine/Source/Runtime/Engine/Private/ChartCreation.cpp b/Engine/Source/Runtime/Engine/Private/ChartCreation.cpp
--- a/Engine/Source/Runtime/Engine/Private/ChartCreation.cpp
+++ b/Engine/Source/Runtime/Engine/Private/ChartCreation.cpp
@@ -84,19 +84,27 @@ static FString GetFPSChartType()
  * @param InMapName	the name of the map
  * @param FileExtension	the filename extension to append
  **/
-FString CreateFileNameForChart( const FString& ChartType, const FString& InMapName, const FString& FileExtension )
+static FString CreateFileNameForChart( const FString& ChartType, const FString& InMapName, const FString& FileExtension )
 {
 	FString Retval;
 	// Create FPS chart filename.
 	FString Platform;
 	// determine which platform we are
-	Platform = TEXT( "PC" );
+	Platform = FPlatformProperties::PlatformName();
 
 	Retval = InMapName + TEXT("-") + ChartType + TEXT("-") + Platform + FileExtension;
 
 	return Retval;
 }
 
+/** This will create the folder name for the output directory for FPS charts. */
+static FString CreateOutputDirectory()
+{
+	// Create folder for FPS chart data.
+	const FString OutputDir = FPaths::ProfilingDir() / TEXT( "FPSChartStats" ) / GCaptureStartTime;
+	IFileManager::Get().MakeDirectory( *OutputDir, true );
+	return OutputDir;
+}
 
 /**
 * Ticks the FPS chart.
@@ -600,12 +608,11 @@ void UEngine::DumpFrameTimesToStatsLog( float TotalTime, float DeltaTime, int32
 {
 #if ALLOW_DEBUG_FILES
 	// Create folder for FPS chart data.
-	FString OutputDir = FPaths::ProfilingDir() + TEXT("FPSChartStats/") + GCaptureStartTime + TEXT("/");
-	IFileManager::Get().MakeDirectory( *OutputDir, true );
+	const FString OutputDir = CreateOutputDirectory();
 
 	// Create archive for log data.
 	const FString ChartType = GetFPSChartType();
-	const FString ChartName = (OutputDir + CreateFileNameForChart( ChartType, InMapName, TEXT( ".csv" ) ) );
+	const FString ChartName = OutputDir / CreateFileNameForChart( ChartType, InMapName, TEXT( ".csv" ) );
 	FArchive* OutputFile = IFileManager::Get().CreateDebugFileWriter( *ChartName );
 
 	if( OutputFile )
@@ -646,13 +653,11 @@ void UEngine::DumpFrameTimesToStatsLog( float TotalTime, float DeltaTime, int32
 void UEngine::DumpFPSChartToStatsLog( float TotalTime, float DeltaTime, int32 NumFrames, const FString& InMapName )
 {
 #if ALLOW_DEBUG_FILES
-	// Create folder for FPS chart data.
-	FString OutputDir = FPaths::ProfilingDir() + TEXT("FPSChartStats/") + GCaptureStartTime + TEXT("/");
-	IFileManager::Get().MakeDirectory( *OutputDir );
+	const FString OutputDir = CreateOutputDirectory();
 	
 	// Create archive for log data.
 	const FString ChartType = GetFPSChartType();
-	const FString ChartName = (OutputDir + CreateFileNameForChart( ChartType, InMapName, TEXT( ".log" ) ) );
+	const FString ChartName = OutputDir / CreateFileNameForChart( ChartType, InMapName, TEXT( ".log" ) );
 	FArchive* OutputFile = IFileManager::Get().CreateDebugFileWriter( *ChartName, FILEWRITE_Append );
 
 	if( OutputFile )
@@ -791,7 +796,13 @@ void UEngine::DumpFPSChartToStatsLog( float TotalTime, float DeltaTime, int32 Nu
 		// Flush, close and delete.
 		delete OutputFile;
 
-		UE_LOG( LogProfilingDebugging, Warning, TEXT( "FPS Chart (logfile) saved to %s" ), *IFileManager::Get().ConvertToAbsolutePathForExternalAppForRead( *ChartName ) );
+		const FString AbsolutePath = IFileManager::Get().ConvertToAbsolutePathForExternalAppForRead( *ChartName );
+
+		UE_LOG( LogProfilingDebugging, Warning, TEXT( "FPS Chart (logfile) saved to %s" ), *AbsolutePath );
+
+#if	PLATFORM_DESKTOP
+		FPlatformProcess::ExploreFolder( *AbsolutePath );
+#endif // PLATFORM_DESKTOP
 	}
 #endif // ALLOW_DEBUG_FILES
 }
@@ -938,14 +949,12 @@ void UEngine::DumpFPSChartToHTML( float TotalTime, float DeltaTime, int32 NumFra
 		FPSChartRow = FPSChartRow.Replace( TEXT("TOKEN_BOUND_GPU_TIME"),		*FString::Printf(TEXT("%4.2f"), ((GTotalFramesBoundTime_GPU)/DeltaTime)*100.0f ), ESearchCase::CaseSensitive );
 
 
-		// Create folder for FPS chart data.
-		const FString OutputDir = FPaths::ProfilingDir() + TEXT("FPSChartStats/") + GCaptureStartTime + TEXT("/");
-		IFileManager::Get().MakeDirectory( *OutputDir );
+		const FString OutputDir = CreateOutputDirectory();
 
 		// Create FPS chart filename.
 		const FString ChartType = GetFPSChartType();
 
-		const FString& FPSChartFilename = OutputDir + CreateFileNameForChart( ChartType, *InMapName, TEXT( ".html" ) );
+		const FString& FPSChartFilename = OutputDir / CreateFileNameForChart( ChartType, *InMapName, TEXT( ".html" ) );
 		FString FPSChart;
 
 		// See whether file already exists and load it into string if it does.
