Fix for crash in visual logger when you run the editor with an invalid .uproject. Fixes UE-13056.

diff --git a/Engine/Source/Runtime/Engine/Private/VisualLogger/VisualLogger.cpp b/Engine/Source/Runtime/Engine/Private/VisualLogger/VisualLogger.cpp
--- a/Engine/Source/Runtime/Engine/Private/VisualLogger/VisualLogger.cpp
+++ b/Engine/Source/Runtime/Engine/Private/VisualLogger/VisualLogger.cpp
@@ -151,7 +151,7 @@ void FVisualLogger::SetIsRecordingToFile(bool InIsRecording)
 		SetIsRecording(true);
 	}
 
-	UWorld* World = GEngine->GetWorld();
+	UWorld* World = GEngine ? GEngine->GetWorld() : NULL;
 
 	const FString BaseFileName = LogFileNameGetter.IsBound() ? LogFileNameGetter.Execute() : TEXT("VisualLog");
 	const FString MapName = World ? World->GetMapName() : TEXT("");
