Modified FRecastTileGenerator::DoWork to not crash on invalid parent generator, and instead just quit with a failure #UE4

Should fix UE-13131

diff --git a/Engine/Source/Runtime/Engine/Private/AI/Navigation/RecastNavMeshGenerator.cpp b/Engine/Source/Runtime/Engine/Private/AI/Navigation/RecastNavMeshGenerator.cpp
--- a/Engine/Source/Runtime/Engine/Private/AI/Navigation/RecastNavMeshGenerator.cpp
+++ b/Engine/Source/Runtime/Engine/Private/AI/Navigation/RecastNavMeshGenerator.cpp
@@ -1642,15 +1642,18 @@ bool FRecastTileGenerator::HasDataToBuild() const
 void FRecastTileGenerator::DoWork()
 {
 	TSharedPtr<FNavDataGenerator, ESPMode::ThreadSafe> ParentGenerator = ParentGeneratorWeakPtr.Pin();
-	FRecastNavMeshGenerator* RecastParentGenerator = static_cast<FRecastNavMeshGenerator*>(ParentGenerator.Get());
-	check(ParentGenerator.Get());
 
-	if (InclusionBounds.Num())
+	if (ParentGenerator.IsValid())
 	{
-		DoAsyncGeometryGathering();
-	}
+		FRecastNavMeshGenerator* RecastParentGenerator = static_cast<FRecastNavMeshGenerator*>(ParentGenerator.Get());
+
+		if (InclusionBounds.Num())
+		{
+			DoAsyncGeometryGathering();
+		}
 
-	bSucceeded = GenerateTile();
+		bSucceeded = GenerateTile();
+	}
 
 	DumpAsyncData();
 }
