sort AutoComplte entries in game console as it is sorted in editor UI

diff --git a/Engine/Source/Runtime/Engine/Private/UserInterface/Console.cpp b/Engine/Source/Runtime/Engine/Private/UserInterface/Console.cpp
--- a/Engine/Source/Runtime/Engine/Private/UserInterface/Console.cpp
+++ b/Engine/Source/Runtime/Engine/Private/UserInterface/Console.cpp
@@ -1178,6 +1178,8 @@ void UConsole::PostRender_InputLine(UCanvas* Canvas, FIntPoint UserInputLinePos)
 		ConsoleTile.SetColor( ConsoleDefs::AutocompleteBackgroundColor );
 		ConsoleTile.Texture = DefaultTexture_White->Resource;
 
+		AutoComplete.Sort();
+
 		// wasteful memory allocations but when typing in a console command this is fine
 		TArray<FString> AutoCompleteElements;
 		// to avoid memory many allocations
diff --git a/Engine/Source/Runtime/EngineSettings/Classes/ConsoleSettings.h b/Engine/Source/Runtime/EngineSettings/Classes/ConsoleSettings.h
--- a/Engine/Source/Runtime/EngineSettings/Classes/ConsoleSettings.h
+++ b/Engine/Source/Runtime/EngineSettings/Classes/ConsoleSettings.h
@@ -18,6 +18,12 @@ struct FAutoCompleteCommand
 
 	UPROPERTY(config, EditAnywhere, Category=Command)
 	FString Desc;
+
+	bool operator<(const FAutoCompleteCommand& rhs) const
+	{
+		// sort them in opposite order for game console UI rendering (bottom up)
+		return Command >= rhs.Command;
+	}
 };
 
 
