Improved the manner in which the quad tree stores and retrieves instances.
Fixes an out of memory error for procedural foliage simulation in some cases.

diff --git a/Engine/Source/Runtime/Engine/Public/GenericQuadTree.h b/Engine/Source/Runtime/Engine/Public/GenericQuadTree.h
--- a/Engine/Source/Runtime/Engine/Public/GenericQuadTree.h
+++ b/Engine/Source/Runtime/Engine/Public/GenericQuadTree.h
@@ -1,6 +1,8 @@
 // Copyright 1998-2015 Epic Games, Inc. All Rights Reserved.
 
 #pragma once
+DEFINE_LOG_CATEGORY_STATIC(LogQuadTree, Log, Warning);
+
 template <typename ElementType, int32 NodeCapacity = 4>
 class TQuadTree
 {
@@ -15,16 +17,13 @@ public:
 	TQuadTree(const FBox2D& InBox);
 
 	/** Inserts an object of type ElementType with an associated 2D box of size Box (log n)*/
-	void Insert(const ElementType& Instance, const FBox2D& Box);
+	void Insert(const ElementType& Element, const FBox2D& Box);
 
-	/** Given a 2D box, returns an array of elements within the box. This may have duplicates based on how many quads the element overlaps*/
+	/** Given a 2D box, returns an array of elements within the box. There will not be any duplicates in the list. */
 	void GetElements(const FBox2D& Box, TArray<ElementType>& ElementsOut) const;
 
-	/** Given a 2D box, returns an array of elements within the box. The array will have no duplicates*/
-	void GetElementsUnique(const FBox2D& Box, TArray<ElementType>& ElementsOut) const;
-
 	/** Removes an object of type ElementType with an associated 2D box of size Box (log n). Does not cleanup tree*/
-	void Remove(const ElementType& Instance, const FBox2D& Box);
+	bool Remove(const ElementType& Instance, const FBox2D& Box);
 
 	/** Does a deep copy of the tree by going through and re-creating the internal data. Cheaper than re-insertion as it should be linear instead of nlogn */
 	void Duplicate(TreeType& OutDuplicate) const;
@@ -47,16 +46,6 @@ private:
 		BottomRight = 3
 	};
 
-	/** Given a 2D box, returns an array of leaf "trees" to be used for insertion,removal, etc... */
-	void GetLeaves(const FBox2D& Box, TArray<TreeType*>& Leaves);
-	void GetLeaves(const FBox2D& Box, TArray<const TreeType*>& Leaves) const;
-
-	/** Given a 2D box, return the subtrees that are touched*/
-	int32 GetQuads(const FBox2D& Box, TreeType* Quads[4]) const;
-
-	/** Split the tree into 4 sub-trees */
-	void Split();
-
 	/** Node used to hold the element and its corresponding 2D box*/
 	struct FNode
 	{
@@ -66,8 +55,8 @@ private:
 		FNode() {};
 
 		FNode(const ElementType& InElement, const FBox2D& InBox)
-		: Box(InBox)
-		, Element(InElement)
+			: Box(InBox)
+			, Element(InElement)
 		{}
 
 		friend FArchive& operator<<(FArchive& Ar, typename TQuadTree<ElementType, NodeCapacity>::FNode& Node)
@@ -76,19 +65,43 @@ private:
 		}
 	};
 
+	/** Given a 2D box, return the subtrees that are touched. Returns 0 for leaves. */
+	int32 GetQuads(const FBox2D& Box, TreeType* Quads[4]) const;
+
+	/** Split the tree into 4 sub-trees */
+	void Split();
+
+	/** Given a list of nodes, return which ones actually intersect the box */
+	void GetIntersectingElements(const FBox2D& Box, TArray/*TSet*/<ElementType>& ElementsOut) const;
+
+	/** Given a list of nodes, remove the node that contains the given element */
+	bool RemoveNodeForElement(const ElementType& Element);
+
+	/** Internal recursive implementation of @see GetElements */
+	void GetElementsRecursive(const FBox2D& Box, TArray<ElementType>& OutElementSet) const;
+
+	/** Internal recursive implementation of @see Insert */
+	void InsertElementRecursive(const ElementType& Element, const FBox2D& Box);
+
 private:
 	
-	/** Array containing the actual data. The node is used because we want to store the AABB and the element */
+	/** 
+	 * Contains the actual elements this tree is responsible for. Nodes are used to keep track of each element's AABB as well.
+	 * For a non-internal leaf, this is the list of nodes that are fully contained within this tree.
+	 * For an internal tree, this contains the nodes that overlap multiple subtrees.
+	 */
 	TArray<FNode> Nodes;
+
+	/** The sub-trees of this tree */
 	TreeType* SubTrees[4];
 
-	/** AABB of the tree*/
+	/** AABB of the tree */
 	FBox2D TreeBox;
 
 	/** Center position of the tree */
 	FVector2D Position;
 
-	/** Whether we are a leaf or an internal sub-tree */
+	/** Whether this is a leaf or an internal sub-tree */
 	bool bInternal;
 };
 
@@ -148,6 +161,7 @@ TQuadTree<ElementType, NodeCapacity>::~TQuadTree()
 	for (TreeType* SubTree : SubTrees)
 	{
 		delete SubTree;
+		SubTree = nullptr;
 	}
 }
 
@@ -183,179 +197,196 @@ void TQuadTree<ElementType, NodeCapacity>::Split()
 	SubTrees[TopRight] = new TreeType(FBox2D(C, TR));
 	SubTrees[BottomLeft] = new TreeType(FBox2D(BL, C));
 	SubTrees[BottomRight] = new TreeType(FBox2D(BM, MR));
+	
+	//mark as no longer a leaf
+	bInternal = true;
 
-	//take existing nodes and split them into new sub-trees
+	// Place existing nodes and place them into the new subtrees that contain them
+	// If a node overlaps multiple subtrees, we retain the reference to it here in this quad
+	TArray<FNode> OverlappingNodes;
 	for (const FNode& Node : Nodes)
 	{
-		TreeType* Leaves[4];
-		const int32 NumLeaves = GetQuads(Node.Box, Leaves);
-		for (int32 LeafIdx = 0; LeafIdx < NumLeaves; ++LeafIdx)
+		TreeType* Quads[4];
+		const int32 NumQuads = GetQuads(Node.Box, Quads);
+		check(NumQuads > 0);
+
+		if (NumQuads == 1)
+		{
+			Quads[0]->Nodes.Add(Node);
+		}
+		else
 		{
-			Leaves[LeafIdx]->Nodes.Add(FNode(Node.Element, Node.Box));
+			OverlappingNodes.Add(Node);
 		}
 	}
 
-	Nodes.Empty();
-	bInternal = true;	//mark as no longer a leaf
+	// Hang onto the nodes that don't fit cleanly into a single subtree
+	Nodes = OverlappingNodes;
 }
 
 template <typename ElementType, int32 NodeCapacity>
 void TQuadTree<ElementType, NodeCapacity>::Insert(const ElementType& Element, const FBox2D& Box)
 {
-	TArray<TreeType*> Leaves;
-	GetLeaves(Box, Leaves);
-	
-	for (TreeType* Leaf : Leaves)
+	if (!Box.Intersect(TreeBox))
 	{
-		check(Leaf->bInternal == false);	//make sure we only got leaves back
-		TArray<FNode>& LeafNodes = Leaf->Nodes;
-
-		if (LeafNodes.Num() < NodeCapacity || Leaf->TreeBox.GetSize().Size() <= 1.f)	//it's possible that all elements in the leaf are bigger than the leaf, and we can get into an endless spiral of splitting
-		{
-			LeafNodes.Add(FNode(Element, Box));
-		}
-		else
-		{
-			//no room so split and then try to add again
-			Leaf->Split();
-			Leaf->Insert(Element, Box);
-		}
+		// Elements shouldn't be added outside the bounds of the top-level quad
+		UE_LOG(LogQuadTree, Warning, TEXT("Adding element (%s) that is outside the bounds of the quadtree root (%s). Consider resizing."), *Box.ToString(), *TreeBox.ToString());
 	}
+
+	InsertElementRecursive(Element, Box);
 }
 
 template <typename ElementType, int32 NodeCapacity>
-void TQuadTree<ElementType, NodeCapacity>::Remove(const ElementType& Element, const FBox2D& Box)
+void TQuadTree<ElementType, NodeCapacity>::InsertElementRecursive(const ElementType& Element, const FBox2D& Box)
 {
-	TArray<TQuadTree*> Leaves;
-	GetLeaves(Box, Leaves);
-	for (TreeType* Leaf : Leaves)
+	TreeType* Quads[4];
+	const int32 NumQuads = GetQuads(Box, Quads);
+	if (NumQuads == 0)
 	{
-		int32 ElementIdx = INDEX_NONE;
-		for (int32 NodeIdx = 0, NumNodes = Leaf->Nodes.Num(); NodeIdx < NumNodes; ++NodeIdx)
+		// This should only happen for leaves
+		check(!bInternal);
+
+		// It's possible that all elements in the leaf are bigger than the leaf or that more elements than NodeCapacity exist outside the top level quad
+		// In either case, we can get into an endless spiral of splitting
+		static const float MinimumQuadSize = 100.f;
+		const bool bCanSplitTree = TreeBox.GetSize().Size() > MinimumQuadSize;
+		if (!bCanSplitTree || Nodes.Num() < NodeCapacity)
 		{
-			if (Leaf->Nodes[NodeIdx].Element == Element)
+			Nodes.Add(FNode(Element, Box));
+
+			if (!bCanSplitTree)
 			{
-				ElementIdx = NodeIdx;
-				break;
+				UE_LOG(LogQuadTree, Warning, TEXT("Minimum size reached for quadtree at %s. Filling beyond capacity %d to %d"), *Position.ToString(), NodeCapacity, Nodes.Num());
 			}
 		}
-		
-		if (ElementIdx != INDEX_NONE)
+		else
 		{
-			Leaf->Nodes.RemoveAtSwap(ElementIdx);
+			// This quad is at capacity, so split and try again
+			Split();
+			InsertElementRecursive(Element, Box);
 		}
 	}
-}
+	else if (NumQuads == 1)
+	{
+		check(bInternal);
 
-template <typename ElementType, int32 NodeCapacity>
-void TQuadTree<ElementType, NodeCapacity>::GetElements(const FBox2D& Box, TArray<ElementType>& ElementsOut) const
-{
-	TArray<const TQuadTree*> Leaves;
-	GetLeaves(Box, Leaves);
-	for (const TreeType* Leaf : Leaves)
+		// Fully contained in a single subtree, so insert it there
+		Quads[0]->InsertElementRecursive(Element, Box);
+	}
+	else
 	{
-		ElementsOut.Reserve(ElementsOut.Num() + Leaf->Nodes.Num());
-		for (const FNode& Node : Leaf->Nodes)
-		{
-			if (Box.Intersect(Node.Box))
-			{
-				ElementsOut.Add(Node.Element);
-			}
-		};
+		// Overlaps multiple subtrees, store here
+		check(bInternal);
+		Nodes.Add(FNode(Element, Box));
 	}
 }
 
 template <typename ElementType, int32 NodeCapacity>
-void TQuadTree<ElementType, NodeCapacity>::GetElementsUnique(const FBox2D& Box, TArray<ElementType>& ElementsOut) const
+bool TQuadTree<ElementType, NodeCapacity>::RemoveNodeForElement(const ElementType& Element)
 {
-	TSet<ElementType> Unique;
-	TArray<const TQuadTree*> Leaves;
-	GetLeaves(Box, Leaves);
-	for (const TreeType* Leaf : Leaves)
+	int32 ElementIdx = INDEX_NONE;
+	for (int32 NodeIdx = 0, NumNodes = Nodes.Num(); NodeIdx < NumNodes; ++NodeIdx)
 	{
-		for (const FNode& Node : Leaf->Nodes)
+		if (Nodes[NodeIdx].Element == Element)
 		{
-			if (Box.Intersect(Node.Box))
-			{
-				Unique.Add(Node.Element);
-			}
-		};
+			ElementIdx = NodeIdx;
+			break;
+		}
 	}
 
-	for (ElementType Elem : Unique)
+	if (ElementIdx != INDEX_NONE)
 	{
-		ElementsOut.Add(Elem);
+		Nodes.RemoveAtSwap(ElementIdx);
+		return true;
 	}
-}
 
+	return false;
+}
 
 template <typename ElementType, int32 NodeCapacity>
-int32 TQuadTree<ElementType, NodeCapacity>::GetQuads(const FBox2D& Box, TreeType* Quads[4]) const
+bool TQuadTree<ElementType, NodeCapacity>::Remove(const ElementType& Element, const FBox2D& Box)
 {
-	bool bNegX = Box.Min.X < Position.X;
-	bool bNegY = Box.Min.Y < Position.Y;
+	bool bElementRemoved = false;
 
-	bool bPosX = Box.Max.X > Position.X;
-	bool bPosY = Box.Max.Y > Position.Y;
+	TreeType* Quads[4];
+	const int32 NumQuads = GetQuads(Box, Quads);
 
-	int32 QuadCount = 0;
-	if (bNegX && bNegY)
-	{
-		Quads[QuadCount++] = SubTrees[BottomLeft];
-	}
+	// Remove from nodes referenced by this quad
+	bElementRemoved = RemoveNodeForElement(Element);
 
-	if (bPosX && bNegY)
+	// Try to remove from subtrees if necessary
+	for (int32 QuadIndex = 0; QuadIndex < NumQuads && !bElementRemoved; QuadIndex++)
 	{
-		Quads[QuadCount++] = SubTrees[BottomRight];
+		bElementRemoved = Quads[QuadIndex]->Remove(Element, Box);
 	}
 
-	if (bNegX && bPosY)
-	{
-		Quads[QuadCount++] = SubTrees[TopLeft];
-	}
+	return bElementRemoved;
+}
+
+template <typename ElementType, int32 NodeCapacity>
+void TQuadTree<ElementType, NodeCapacity>::GetElements(const FBox2D& Box, TArray<ElementType>& ElementsOut) const
+{
+	TreeType* Quads[4];
+	const int32 NumQuads = GetQuads(Box, Quads);
+
+	// Always include any nodes contained in this quad
+	GetIntersectingElements(Box, ElementsOut);
 
-	if (bPosX && bPosY)
+	// As well as all relevant subtrees
+	for (int32 QuadIndex = 0; QuadIndex < NumQuads; QuadIndex++)
 	{
-		Quads[QuadCount++] = SubTrees[TopRight];
+		Quads[QuadIndex]->GetElements(Box, ElementsOut);
 	}
-
-	return QuadCount;
 }
 
 template <typename ElementType, int32 NodeCapacity>
-void TQuadTree<ElementType, NodeCapacity>::GetLeaves(const FBox2D& Box, TArray<TreeType*>& LeavesOut)
+void TQuadTree<ElementType, NodeCapacity>::GetIntersectingElements(const FBox2D& Box, TArray/*TSet*/<ElementType>& ElementsOut) const
 {
-	if (bInternal)
+	ElementsOut.Reserve(ElementsOut.Num() + Nodes.Num());
+	for (const FNode& Node : Nodes)
 	{
-		TreeType* UseNodes[4];
-		const int32 NumQuads = GetQuads(Box, UseNodes);
-		for (int32 NodeIdx = 0; NodeIdx < NumQuads; ++NodeIdx)
+		if (Box.Intersect(Node.Box))
 		{
-			UseNodes[NodeIdx]->GetLeaves(Box, LeavesOut);
+			check(!ElementsOut.Contains(Node.Element));
+			ElementsOut.Add(Node.Element);
 		}
-	}
-	else
-	{
-		LeavesOut.Add(this);
-	}
+	};
 }
 
 template <typename ElementType, int32 NodeCapacity>
-void TQuadTree<ElementType, NodeCapacity>::GetLeaves(const FBox2D& Box, TArray<const TreeType*>& LeavesOut) const
+int32 TQuadTree<ElementType, NodeCapacity>::GetQuads(const FBox2D& Box, TreeType* Quads[4]) const
 {
+	int32 QuadCount = 0;
 	if (bInternal)
 	{
-		TreeType* UseNodes[4];
-		const int32 NumQuads = GetQuads(Box, UseNodes);
-		for (int32 NodeIdx = 0; NodeIdx < NumQuads; ++NodeIdx)
+		bool bNegX = Box.Min.X < Position.X;
+		bool bNegY = Box.Min.Y < Position.Y;
+
+		bool bPosX = Box.Max.X > Position.X;
+		bool bPosY = Box.Max.Y > Position.Y;
+
+		if (bNegX && bNegY)
 		{
-			UseNodes[NodeIdx]->GetLeaves(Box, LeavesOut);
+			Quads[QuadCount++] = SubTrees[BottomLeft];
+		}
+
+		if (bPosX && bNegY)
+		{
+			Quads[QuadCount++] = SubTrees[BottomRight];
+		}
+
+		if (bNegX && bPosY)
+		{
+			Quads[QuadCount++] = SubTrees[TopLeft];
+		}
+
+		if (bPosX && bPosY)
+		{
+			Quads[QuadCount++] = SubTrees[TopRight];
 		}
 	}
-	else
-	{
-		LeavesOut.Add(this);
-	}
+
+	return QuadCount;
 }
 
 template <typename ElementType, int32 NodeCapacity>
diff --git a/Engine/Source/Runtime/Foliage/Private/ProceduralFoliageBroadphase.cpp b/Engine/Source/Runtime/Foliage/Private/ProceduralFoliageBroadphase.cpp
--- a/Engine/Source/Runtime/Foliage/Private/ProceduralFoliageBroadphase.cpp
+++ b/Engine/Source/Runtime/Foliage/Private/ProceduralFoliageBroadphase.cpp
@@ -5,7 +5,7 @@
 #include "ProceduralFoliageInstance.h"
 
 FProceduralFoliageBroadphase::FProceduralFoliageBroadphase(float TileSize)
-: QuadTree(FBox2D(FVector2D(0, 0), FVector2D(TileSize, TileSize)))
+	: QuadTree(FBox2D(FVector2D(-TileSize * 0.5f, -TileSize * 0.5f), FVector2D(TileSize * 1.5, TileSize * 1.5)))
 {
 }
 
@@ -48,14 +48,14 @@ bool FProceduralFoliageBroadphase::GetOverlaps(FProceduralFoliageInstance* Insta
 
 	TArray<FProceduralFoliageInstance*> PossibleOverlaps;
 	const FBox2D AABB = GetMaxAABB(Instance);
-	QuadTree.GetElementsUnique(AABB, PossibleOverlaps);
+	QuadTree.GetElements(AABB, PossibleOverlaps);
 	Overlaps.Reserve(Overlaps.Num() + PossibleOverlaps.Num());
 	
 	for (FProceduralFoliageInstance* Overlap : PossibleOverlaps)
 	{
 		if (Overlap != Instance)
 		{
-			//We must determine if this is an overlap of shade or and overlap of collision. If both the collision overlap wins
+			//We must determine if this is an overlap of shade or an overlap of collision. If both the collision overlap wins
 			bool bCollisionOverlap = CircleOverlap(Instance->Location, ACollisionRadius, Overlap->Location, Overlap->GetCollisionRadius());
 			bool bShadeOverlap     = CircleOverlap(Instance->Location, AShadeRadius, Overlap->Location, Overlap->GetShadeRadius());
 
@@ -73,25 +73,11 @@ bool FProceduralFoliageBroadphase::GetOverlaps(FProceduralFoliageInstance* Insta
 void FProceduralFoliageBroadphase::Remove(FProceduralFoliageInstance* Instance)
 {
 	const FBox2D AABB = GetMaxAABB(Instance);
-	QuadTree.Remove(Instance, AABB);
+	const bool bRemoved = QuadTree.Remove(Instance, AABB);
+	check(bRemoved);
 }
 
 void FProceduralFoliageBroadphase::GetInstancesInBox(const FBox2D& Box, TArray<FProceduralFoliageInstance*>& Instances) const
 {
-	TArray<FProceduralFoliageInstance*> AABBInstances;
-	QuadTree.GetElementsUnique(Box, AABBInstances);
-
-	Instances.Reserve(AABBInstances.Num() + Instances.Num());
-
-	for (FProceduralFoliageInstance* Inst : AABBInstances)
-	{
-		const FVector2D Location2D(Inst->Location.X, Inst->Location.Y);
-		const FVector2D ClosestPt = Box.GetClosestPointTo(Location2D);
-		const float Radius2 = (ClosestPt - Location2D).SizeSquared();
-		const float InstRadius = Inst->GetMaxRadius();
-		if (Radius2 < InstRadius*InstRadius)
-		{
-			Instances.Add(Inst);
-		}
-	}
+	QuadTree.GetElements(Box, Instances);
 }
\ No newline at end of file
