Inheritable component records can be found via template name as well as component name

diff --git a/Engine/Source/Runtime/Engine/Private/InheritableComponentHandler.cpp b/Engine/Source/Runtime/Engine/Private/InheritableComponentHandler.cpp
--- a/Engine/Source/Runtime/Engine/Private/InheritableComponentHandler.cpp
+++ b/Engine/Source/Runtime/Engine/Private/InheritableComponentHandler.cpp
@@ -277,7 +277,7 @@ FComponentKey UInheritableComponentHandler::FindKey(const FName VariableName) co
 {
 	for (const FComponentOverrideRecord& Record : Records)
 	{
-		if (Record.ComponentKey.VariableName == VariableName)
+		if (Record.ComponentKey.VariableName == VariableName || (Record.ComponentTemplate && Record.ComponentTemplate->GetFName() == VariableName))
 		{
 			return Record.ComponentKey;
 		}
