[UE-11897] Slightly revised WITH_EDITOR compile fix.

diff --git a/Engine/Source/Runtime/Engine/Private/SimpleConstructionScript.cpp b/Engine/Source/Runtime/Engine/Private/SimpleConstructionScript.cpp
--- a/Engine/Source/Runtime/Engine/Private/SimpleConstructionScript.cpp
+++ b/Engine/Source/Runtime/Engine/Private/SimpleConstructionScript.cpp
@@ -149,6 +149,9 @@ void USimpleConstructionScript::PostLoad()
 {
 	Super::PostLoad();
 
+	int32 NodeIndex;
+	TArray<USCS_Node*> Nodes = GetAllNodes();
+
 #if WITH_EDITOR
 	// Get the Blueprint that owns the SCS
 	UBlueprint* Blueprint = GetBlueprint();
@@ -159,8 +162,6 @@ void USimpleConstructionScript::PostLoad()
 		return;
 	}
 
-	int32 NodeIndex;
-	TArray<USCS_Node*> Nodes = GetAllNodes();
 	USCS_Node* SceneRootNode = nullptr;
 	USceneComponent* SceneRootComponentTemplate = GetSceneRootComponentTemplate(&SceneRootNode);
 
@@ -284,7 +285,7 @@ void USimpleConstructionScript::PostLoad()
 				{
 					// If no native root component exists, find the first non-native, non-parented SCS node with a
 					// scene component template. This will be designated as the root component at construction time.
-					for( int32 NodeIndex = 0; NodeIndex < RootNodes.Num(); ++NodeIndex)
+					for(NodeIndex = 0; NodeIndex < RootNodes.Num(); ++NodeIndex)
 					{
 						USCS_Node* Node = RootNodes[NodeIndex];
 						if(Node->ParentComponentOrVariableName == NAME_None)
@@ -581,8 +582,6 @@ void USimpleConstructionScript::RemoveNode(USCS_Node* Node)
 	}
 }
 
-#if	WITH_EDITOR
-
 int32 USimpleConstructionScript::FindPromotableChildNodeIndex(USCS_Node* InParentNode) const
 {
 	int32 PromoteIndex = INDEX_NONE;
@@ -611,8 +610,6 @@ int32 USimpleConstructionScript::FindPromotableChildNodeIndex(USCS_Node* InParen
 	return PromoteIndex;
 }
 
-#endif // WITH_EDITOR
-
 void USimpleConstructionScript::RemoveNodeAndPromoteChildren(USCS_Node* Node)
 {
 	Node->Modify();
@@ -620,14 +617,12 @@ void USimpleConstructionScript::RemoveNodeAndPromoteChildren(USCS_Node* Node)
 	if (RootNodes.Contains(Node))
 	{
 		USCS_Node* ChildToPromote = nullptr;
-#if	WITH_EDITOR
 		int32 PromoteIndex = FindPromotableChildNodeIndex(Node);
 		if(PromoteIndex != INDEX_NONE)
 		{
 			ChildToPromote = Node->ChildNodes[PromoteIndex];
 			Node->ChildNodes.RemoveAt(PromoteIndex);
 		}
-#endif // WITH_EDITOR
 
 		Modify();
 
