Update cached geometry in the scene viewports Tick.

diff --git a/Engine/Source/Runtime/Engine/Private/Slate/SceneViewport.cpp b/Engine/Source/Runtime/Engine/Private/Slate/SceneViewport.cpp
--- a/Engine/Source/Runtime/Engine/Private/Slate/SceneViewport.cpp
+++ b/Engine/Source/Runtime/Engine/Private/Slate/SceneViewport.cpp
@@ -1217,6 +1217,7 @@ void FSceneViewport::EndRenderFrame(FRHICommandListImmediate& RHICmdList, bool b
 
 void FSceneViewport::Tick( const FGeometry& AllottedGeometry, double InCurrentTime, float DeltaTime )
 {
+	UpdateCachedGeometry(AllottedGeometry);
 	ProcessInput( DeltaTime );
 }
 
