Hash network version with date and time for non promoted builds

diff --git a/Engine/Source/Runtime/Engine/Private/UnrealNetwork.cpp b/Engine/Source/Runtime/Engine/Private/UnrealNetwork.cpp
--- a/Engine/Source/Runtime/Engine/Private/UnrealNetwork.cpp
+++ b/Engine/Source/Runtime/Engine/Private/UnrealNetwork.cpp
@@ -36,9 +36,20 @@ uint32 FNetworkVersion::GetLocalNetworkVersion()
 	const uint32 VersionHash = FCrc::StrCrc32( *ProjectVersion, FCrc::StrCrc32( *ProjectName, GEngineNetVersion ) );
 
 	// Hash with internal protocol version
-	const uint32 LocalNetworkVersion = FCrc::MemCrc32( &InternalProtocolVersion, sizeof( InternalProtocolVersion ), VersionHash );
+	uint32 LocalNetworkVersion = FCrc::MemCrc32( &InternalProtocolVersion, sizeof( InternalProtocolVersion ), VersionHash );
 
-	UE_LOG( LogNet, Log, TEXT( "GetLocalNetworkVersion: GEngineNetVersion: %i, ProjectName: %s, ProjectVersion: %s, InternalProtocolVersion: %i, LocalNetworkVersion: %i" ), GEngineNetVersion, *ProjectName, *ProjectVersion, InternalProtocolVersion, LocalNetworkVersion );
+	if ( !GEngineVersion.IsPromotedBuild() )
+	{
+		// Further hash with date time if this is a non promoted build
+		const FString DateTimeString = FString::Printf( TEXT( "%s_%s" ), ANSI_TO_TCHAR( __DATE__ ), ANSI_TO_TCHAR( __TIME__ ) );
+		const uint32 LocalNetworkVersionNonPromoted = FCrc::StrCrc32( *DateTimeString, LocalNetworkVersion );
+
+		UE_LOG( LogNet, Log, TEXT( "GetLocalNetworkVersion: NON-PROMOTED: Date: %s, ProjectName: %s, ProjectVersion: %s, InternalProtocolVersion: %i, LocalNetworkVersionNonPromoted: %u" ), *DateTimeString, *ProjectName, *ProjectVersion, InternalProtocolVersion, LocalNetworkVersionNonPromoted );
+
+		return LocalNetworkVersionNonPromoted;
+	}
+
+	UE_LOG( LogNet, Log, TEXT( "GetLocalNetworkVersion: GEngineNetVersion: %i, ProjectName: %s, ProjectVersion: %s, InternalProtocolVersion: %i, LocalNetworkVersion: %u" ), GEngineNetVersion, *ProjectName, *ProjectVersion, InternalProtocolVersion, LocalNetworkVersion );
 
 	return LocalNetworkVersion;
 }
