Fixed: LevelBounds actor does not update bounds when ReatimeViewports is off
UE-11278

diff --git a/Engine/Source/Runtime/Engine/Classes/Engine/Level.h b/Engine/Source/Runtime/Engine/Classes/Engine/Level.h
--- a/Engine/Source/Runtime/Engine/Classes/Engine/Level.h
+++ b/Engine/Source/Runtime/Engine/Classes/Engine/Level.h
@@ -479,6 +479,10 @@ public:
 	FLevelBoundsActorUpdatedEvent& LevelBoundsActorUpdated() { return LevelBoundsActorUpdatedEvent; }
 	/**	Broadcasts that Level bounds actor has been updated */ 
 	void BroadcastLevelBoundsActorUpdated() { LevelBoundsActorUpdatedEvent.Broadcast(); }
+
+	/** Marks level bounds as dirty so they will be recalculated  */
+	ENGINE_API void MarkLevelBoundsDirty();
+
 private:
 	FLevelBoundsActorUpdatedEvent LevelBoundsActorUpdatedEvent; 
 
@@ -740,10 +744,6 @@ public:
 
 	/** Push any pending auto receive input actor's input components on to the player controller's input stack */
 	void PushPendingAutoReceiveInput(APlayerController* PC);
-
-	/** Increments number of steaming objects referring this level */
-	ENGINE_API void IncStreamingLevelRefs();
-	
 	
 	// Begin IInterface_AssetUserData Interface
 	virtual void AddAssetUserData(UAssetUserData* InUserData) override;
diff --git a/Engine/Source/Runtime/Engine/Classes/Engine/LevelBounds.h b/Engine/Source/Runtime/Engine/Classes/Engine/LevelBounds.h
--- a/Engine/Source/Runtime/Engine/Classes/Engine/LevelBounds.h
+++ b/Engine/Source/Runtime/Engine/Classes/Engine/LevelBounds.h
@@ -1,16 +1,30 @@
 // Copyright 1998-2015 Epic Games, Inc. All Rights Reserved.
 #pragma once
-#include "TimerManager.h"
+#include "Tickable.h"
 #include "LevelBounds.generated.h"
 
+
+
+/**
+ *  Ticks only in the editor, regardless of viewport 'Realtime' option
+ */
+class FEditorTickableLevelBounds 
+#if WITH_EDITOR
+	: FTickableGameObject
+#endif
+{
+};
+
 /**
  *
  * Defines level bounds
  * Updates bounding box automatically based on actors transformation changes or holds fixed user defined bounding box
- * 
+ * Uses only actors where AActor::IsLevelBoundsRelevant() == true
  */
 UCLASS(hidecategories=(Advanced, Collision, Display, Rendering, Physics, Input), showcategories=("Input|MouseInput", "Input|TouchInput"), MinimalAPI)
-class ALevelBounds : public AActor
+class ALevelBounds 
+	: public AActor
+	, public FEditorTickableLevelBounds 
 {
 	GENERATED_UCLASS_BODY()
 		
@@ -29,7 +43,7 @@ class ALevelBounds : public AActor
 
 	/** @return Bounding box which includes all relevant actors bounding boxes belonging to specified level */
 	ENGINE_API static FBox CalculateLevelBounds(ULevel* InLevel);
-		
+
 #if WITH_EDITOR
 	virtual void PostEditUndo() override;
 	virtual void PostEditMove(bool bFinished) override;
@@ -37,8 +51,8 @@ class ALevelBounds : public AActor
 	virtual void PostRegisterAllComponents() override;
 	virtual void PostUnregisterAllComponents() override;
 	
-	/** Tells actors that level bounds was changed so actor bounds should be recalculated  */
-	void OnLevelBoundsDirtied();
+	/** Marks level bounds as dirty so they will be recalculated on next tick */
+	void MarkLevelBoundsDirty();
 
 	/** @return True if there were no actors contributing to bounds and we are currently using the default bounds */
 	ENGINE_API bool IsUsingDefaultBounds() const;
@@ -47,10 +61,12 @@ class ALevelBounds : public AActor
 	ENGINE_API void UpdateLevelBoundsImmediately();
 	
 private:
+	/** FTickableGameObject interface */
+	virtual void Tick(float DeltaTime) override;
+	virtual TStatId GetStatId() const override;
+	virtual bool IsTickable() const override;
+	virtual bool IsTickableInEditor() const override;
 	
-	/** Handle for efficient management of OnTimerTick timer */
-	FTimerHandle TimerHandle_OnTimerTick;
-
 	/** Updates this actor bounding box by summing all level actors bounding boxes  */
 	void UpdateLevelBounds();
 
@@ -63,18 +79,12 @@ private:
 	/** Called whenever any actor added or removed  */
 	void OnLevelActorAddedRemoved(AActor* InActor);
 
-	/** Timer tick for auto updating level bounds  */
-	void OnTimerTick();
-
 	/** Subscribes for actors transformation events */
 	void SubscribeToUpdateEvents();
 	
 	/** Unsubscribes from actors transformation events */
 	void UnsubscribeFromUpdateEvents();
 	
-	/** Whether this actor is currently subscribed to transformation events */
-	bool bSubscribedToEvents;
-
 	/** Whether currently level bounds is dirty and needs to be updated  */
 	bool bLevelBoundsDirty;
 
diff --git a/Engine/Source/Runtime/Engine/Private/ActorEditor.cpp b/Engine/Source/Runtime/Engine/Private/ActorEditor.cpp
--- a/Engine/Source/Runtime/Engine/Private/ActorEditor.cpp
+++ b/Engine/Source/Runtime/Engine/Private/ActorEditor.cpp
@@ -345,9 +345,9 @@ void AActor::PreEditUndo()
 void AActor::PostEditUndo()
 {
 	// Notify LevelBounds actor that level bounding box might be changed
-	if (!IsTemplate() && GetLevel()->LevelBoundsActor.IsValid())
+	if (!IsTemplate())
 	{
-		GetLevel()->LevelBoundsActor.Get()->OnLevelBoundsDirtied();
+		GetLevel()->MarkLevelBoundsDirty();
 	}
 
 	// Restore OwnedComponents array
@@ -370,9 +370,9 @@ void AActor::PostEditUndo(TSharedPtr<ITransactionObjectAnnotation> TransactionAn
 	CurrentTransactionAnnotation = StaticCastSharedPtr<FActorTransactionAnnotation>(TransactionAnnotation);
 
 	// Notify LevelBounds actor that level bounding box might be changed
-	if (!IsTemplate() && GetLevel()->LevelBoundsActor.IsValid())
+	if (!IsTemplate())
 	{
-		GetLevel()->LevelBoundsActor.Get()->OnLevelBoundsDirtied();
+		GetLevel()->MarkLevelBoundsDirty();
 	}
 
 	// Restore OwnedComponents array
diff --git a/Engine/Source/Runtime/Engine/Private/Level.cpp b/Engine/Source/Runtime/Engine/Private/Level.cpp
--- a/Engine/Source/Runtime/Engine/Private/Level.cpp
+++ b/Engine/Source/Runtime/Engine/Private/Level.cpp
@@ -1101,13 +1101,19 @@ void ULevel::PostEditUndo()
 		}
 	}
 
+	MarkLevelBoundsDirty();
+}
+#endif // WITH_EDITOR
+
+void ULevel::MarkLevelBoundsDirty()
+{
+#if WITH_EDITOR
 	if (LevelBoundsActor.IsValid())
 	{
-		LevelBoundsActor.Get()->OnLevelBoundsDirtied();
+		LevelBoundsActor->MarkLevelBoundsDirty();
 	}
+#endif// WITH_EDITOR
 }
-#endif // WITH_EDITOR
-
 
 void ULevel::InvalidateModelGeometry()
 {
diff --git a/Engine/Source/Runtime/Engine/Private/LevelBounds.cpp b/Engine/Source/Runtime/Engine/Private/LevelBounds.cpp
--- a/Engine/Source/Runtime/Engine/Private/LevelBounds.cpp
+++ b/Engine/Source/Runtime/Engine/Private/LevelBounds.cpp
@@ -12,7 +12,7 @@ ALevelBounds::ALevelBounds(const FObjectInitializer& ObjectInitializer)
 {
 	UBoxComponent* BoxComponent = CreateDefaultSubobject<UBoxComponent>(TEXT("BoxComponent0"));
 	RootComponent = BoxComponent;
-	RootComponent->Mobility = EComponentMobility::Static;
+	RootComponent->Mobility = EComponentMobility::Movable;
 	RootComponent->RelativeScale3D = DefaultLevelSize;
 
 	bAutoUpdateBounds = true;
@@ -27,7 +27,6 @@ ALevelBounds::ALevelBounds(const FObjectInitializer& ObjectInitializer)
 	
 #if WITH_EDITOR
 	bLevelBoundsDirty = true;
-	bSubscribedToEvents = false;
 	bUsingDefaultBounds = false;
 #endif
 }
@@ -75,26 +74,22 @@ FBox ALevelBounds::CalculateLevelBounds(ULevel* InLevel)
 void ALevelBounds::PostEditUndo()
 {
 	Super::PostEditUndo();
-	bLevelBoundsDirty = true;
+	
+	MarkLevelBoundsDirty();
 }
 
 void ALevelBounds::PostEditMove(bool bFinished)
 {
 	Super::PostEditMove(bFinished);
-	bLevelBoundsDirty = true;
+	
+	MarkLevelBoundsDirty();
 }
 
 void ALevelBounds::PostEditChangeProperty(FPropertyChangedEvent& PropertyChangedEvent)
 {
 	Super::PostEditChangeProperty( PropertyChangedEvent );
-
-	if (GIsEditor && !GetWorld()->IsPlayInEditor())
-	{
-		if (bAutoUpdateBounds)
-		{
-			UpdateLevelBounds();
-		}
-	}
+	
+	MarkLevelBoundsDirty();
 }
 
 void ALevelBounds::PostRegisterAllComponents()
@@ -111,6 +106,36 @@ void ALevelBounds::PostUnregisterAllComponents()
 	Super::PostUnregisterAllComponents();
 }
 
+void ALevelBounds::Tick(float DeltaTime)
+{
+	if (bLevelBoundsDirty)
+	{
+		UpdateLevelBounds();
+		bLevelBoundsDirty = false;
+	}
+}
+
+TStatId ALevelBounds::GetStatId() const
+{
+	RETURN_QUICK_DECLARE_CYCLE_STAT(ALevelBounds, STATGROUP_Tickables);
+}
+
+bool ALevelBounds::IsTickable() const
+{
+	if (GIsEditor && bAutoUpdateBounds && !IsTemplate())
+	{
+		UWorld* World = GetWorld();
+		return (World && World->WorldType == EWorldType::Editor);
+	}
+
+	return false;
+}
+
+bool ALevelBounds::IsTickableInEditor() const
+{
+	return IsTickable();
+}
+
 void ALevelBounds::UpdateLevelBounds()
 {
 	FBox LevelBounds = CalculateLevelBounds(GetLevel());
@@ -131,7 +156,7 @@ void ALevelBounds::UpdateLevelBounds()
 	BroadcastLevelBoundsUpdated();
 }
 
-void ALevelBounds::OnLevelBoundsDirtied()
+void ALevelBounds::MarkLevelBoundsDirty()
 {
 	bLevelBoundsDirty = true;
 }
@@ -158,7 +183,7 @@ void ALevelBounds::OnLevelActorMoved(AActor* InActor)
 		}
 		else
 		{
-			OnLevelBoundsDirtied();
+			MarkLevelBoundsDirty();
 		}
 	}
 }
@@ -167,16 +192,7 @@ void ALevelBounds::OnLevelActorAddedRemoved(AActor* InActor)
 {
 	if (InActor->GetOuter() == GetOuter())
 	{
-		OnLevelBoundsDirtied();
-	}
-}
-
-void ALevelBounds::OnTimerTick()
-{
-	if (bLevelBoundsDirty && bAutoUpdateBounds)
-	{
-		UpdateLevelBounds();
-		bLevelBoundsDirty = false;
+		MarkLevelBoundsDirty();
 	}
 }
 
@@ -192,31 +208,22 @@ void ALevelBounds::BroadcastLevelBoundsUpdated()
 
 void ALevelBounds::SubscribeToUpdateEvents()
 {
-	if (bSubscribedToEvents == false && GIsEditor && !GetWorld()->IsPlayInEditor())
+	// Subscribe only in editor worlds
+	if (!GetWorld()->IsGameWorld())
 	{
-		GetWorldTimerManager().SetTimer(TimerHandle_OnTimerTick, this, &ALevelBounds::OnTimerTick, 1.f, true);
+		UnsubscribeFromUpdateEvents();
+
 		OnLevelActorMovedDelegateHandle   = GEngine->OnActorMoved       ().AddUObject(this, &ALevelBounds::OnLevelActorMoved);
 		OnLevelActorDeletedDelegateHandle = GEngine->OnLevelActorDeleted().AddUObject(this, &ALevelBounds::OnLevelActorAddedRemoved);
 		OnLevelActorAddedDelegateHandle   = GEngine->OnLevelActorAdded  ().AddUObject(this, &ALevelBounds::OnLevelActorAddedRemoved);
-		
-		bSubscribedToEvents = true;
 	}
 }
 
 void ALevelBounds::UnsubscribeFromUpdateEvents()
 {
-	if (bSubscribedToEvents == true)
-	{
-		if (GetWorld())
-		{
-			GetWorldTimerManager().ClearTimer(TimerHandle_OnTimerTick);
-		}
-		GEngine->OnActorMoved       ().Remove(OnLevelActorMovedDelegateHandle);
-		GEngine->OnLevelActorDeleted().Remove(OnLevelActorDeletedDelegateHandle);
-		GEngine->OnLevelActorAdded  ().Remove(OnLevelActorAddedDelegateHandle);
-
-		bSubscribedToEvents = false;
-	}
+	GEngine->OnActorMoved       ().Remove(OnLevelActorMovedDelegateHandle);
+	GEngine->OnLevelActorDeleted().Remove(OnLevelActorDeletedDelegateHandle);
+	GEngine->OnLevelActorAdded  ().Remove(OnLevelActorAddedDelegateHandle);
 }
 
 
