Fixed Matinee toggle tracks not working if bActivateSystemEachUpdate was enabled (legacy feature)
- Fixed tracks performing the last track keyframe's event every single frame
- This bug was introduced in CL 219965, and discovered due to syntax fixes in CL 2497708

diff --git a/Engine/Source/Runtime/Engine/Private/Interpolation.cpp b/Engine/Source/Runtime/Engine/Private/Interpolation.cpp
--- a/Engine/Source/Runtime/Engine/Private/Interpolation.cpp
+++ b/Engine/Source/Runtime/Engine/Private/Interpolation.cpp
@@ -5058,24 +5058,26 @@ void UInterpTrackToggle::UpdateTrack(float NewPosition, UInterpTrackInst* TrInst
 			for (int32 KeyIndex = ToggleTrack.Num() - 1; KeyIndex >= 0; KeyIndex--)
 			{
 				FToggleTrackKey& ToggleKey = ToggleTrack[KeyIndex];
-				// We have found the key to the left of the position
-				if (ToggleKey.ToggleAction == ETTA_On)
+				if( ToggleKey.Time < NewPosition )
 				{
-					EmitterActor->GetParticleSystemComponent()->ActivateSystem(bActivateWithJustAttachedFlag);
-				}
-				else
-				if (ToggleKey.ToggleAction == ETTA_Trigger)
-				{
-					if (ToggleKey.Time >= ToggleInst->LastUpdatePosition)
+					// We have found the key to the left of the position
+					if( ToggleKey.ToggleAction == ETTA_On )
 					{
-						EmitterActor->GetParticleSystemComponent()->SetActive(true, bActivateWithJustAttachedFlag);
+						EmitterActor->GetParticleSystemComponent()->ActivateSystem( bActivateWithJustAttachedFlag );
 					}
+					else if( ToggleKey.ToggleAction == ETTA_Trigger )
+					{
+						if( ToggleKey.Time >= ToggleInst->LastUpdatePosition )
+						{
+							EmitterActor->GetParticleSystemComponent()->SetActive( true, bActivateWithJustAttachedFlag );
+						}
+					}
+					else
+					{
+						EmitterActor->GetParticleSystemComponent()->DeactivateSystem();
+					}
+					break;
 				}
-				else
-				{
-					EmitterActor->GetParticleSystemComponent()->DeactivateSystem();
-				}
-				break;
 			}
 		}
 	}
