Fix locked axis being saved as (Invalid) which causes the cooker to crash.

diff --git a/Engine/Source/Runtime/Engine/Classes/PhysicsEngine/PhysicsSettings.h b/Engine/Source/Runtime/Engine/Classes/PhysicsEngine/PhysicsSettings.h
--- a/Engine/Source/Runtime/Engine/Classes/PhysicsEngine/PhysicsSettings.h
+++ b/Engine/Source/Runtime/Engine/Classes/PhysicsEngine/PhysicsSettings.h
@@ -76,7 +76,9 @@ namespace ESettingsLockedAxis
 		/*Lock movement along the y-axis*/
 		Y,
 		/*Lock movement along the z-axis*/
-		Z
+		Z,
+		/* Used for backwards compatibility. Indicates that we've updated into the new struct*/
+		Invalid
 	};
 }
 
diff --git a/Engine/Source/Runtime/Engine/Private/PhysicsEngine/PhysicsSettings.cpp b/Engine/Source/Runtime/Engine/Private/PhysicsEngine/PhysicsSettings.cpp
--- a/Engine/Source/Runtime/Engine/Private/PhysicsEngine/PhysicsSettings.cpp
+++ b/Engine/Source/Runtime/Engine/Private/PhysicsEngine/PhysicsSettings.cpp
@@ -15,7 +15,7 @@ UPhysicsSettings::UPhysicsSettings(const FObjectInitializer& ObjectInitializer)
 	, bEnableAsyncScene(false)
 	, bWarnMissingLocks(true)
 	, bEnable2DPhysics(false)
-	, LockedAxis_DEPRECATED(static_cast<ESettingsLockedAxis::Type>(-1))
+	, LockedAxis_DEPRECATED(ESettingsLockedAxis::Invalid)
 	, BounceThresholdVelocity(200.f)
 	, bSimulateSkeletalMeshOnDedicatedServer(true)
 	, MaxPhysicsDeltaTime(1.f / 30.f)
@@ -36,7 +36,12 @@ void UPhysicsSettings::PostInitProperties()
 	LoadSurfaceType();
 #endif
 
-	if (LockedAxis_DEPRECATED != static_cast<ESettingsLockedAxis::Type>(-1))
+	if (LockedAxis_DEPRECATED == static_cast<ESettingsLockedAxis::Type>(-1))
+	{
+		LockedAxis_DEPRECATED == ESettingsLockedAxis::Invalid;
+	}
+
+	if (LockedAxis_DEPRECATED != ESettingsLockedAxis::Invalid)
 	{
 		if (LockedAxis_DEPRECATED == ESettingsLockedAxis::None)
 		{
@@ -55,7 +60,7 @@ void UPhysicsSettings::PostInitProperties()
 			DefaultDegreesOfFreedom = ESettingsDOF::XYPlane;
 		}
 
-		LockedAxis_DEPRECATED = static_cast<ESettingsLockedAxis::Type>(-1);
+		LockedAxis_DEPRECATED = ESettingsLockedAxis::Invalid;
 	}
 }
 
