Merging 2499328

//depot/UE4-UT/...

to //depot/UE4/...

diff --git a/Engine/Source/Runtime/Engine/Classes/Engine/World.h b/Engine/Source/Runtime/Engine/Classes/Engine/World.h
--- a/Engine/Source/Runtime/Engine/Classes/Engine/World.h
+++ b/Engine/Source/Runtime/Engine/Classes/Engine/World.h
@@ -3192,6 +3192,9 @@ public:
 	// Delegate type for level change events
 	DECLARE_MULTICAST_DELEGATE_TwoParams(FOnLevelChanged, ULevel*, UWorld*);
 
+	// delegate for generating world asset registry tags so project/game scope can add additional tags for filtering levels in their UI, etc
+	DECLARE_MULTICAST_DELEGATE_TwoParams(FWorldGetAssetTags, const UWorld*, TArray<UObject::FAssetRegistryTag>&);
+
 	// Callback for world creation
 	static FWorldEvent OnPostWorldCreation;
 	
@@ -3223,6 +3226,9 @@ public:
 	// removed from the world without a RemoveFromWorld call for each)
 	static FOnLevelChanged			LevelRemovedFromWorld;
 
+	// called by UWorld::GetAssetRegistryTags()
+	static FWorldGetAssetTags GetAssetTags;
+
 private:
 	FWorldDelegates() {}
 };
diff --git a/Engine/Source/Runtime/Engine/Private/World.cpp b/Engine/Source/Runtime/Engine/Private/World.cpp
--- a/Engine/Source/Runtime/Engine/Private/World.cpp
+++ b/Engine/Source/Runtime/Engine/Private/World.cpp
@@ -103,6 +103,7 @@ FWorldDelegates::FWorldCleanupEvent FWorldDelegates::OnWorldCleanup;
 FWorldDelegates::FWorldEvent FWorldDelegates::OnPreWorldFinishDestroy;
 FWorldDelegates::FOnLevelChanged FWorldDelegates::LevelAddedToWorld;
 FWorldDelegates::FOnLevelChanged FWorldDelegates::LevelRemovedFromWorld;
+FWorldDelegates::FWorldGetAssetTags FWorldDelegates::GetAssetTags;
 
 UWorld::UWorld( const FObjectInitializer& ObjectInitializer )
 :	UObject(ObjectInitializer)
@@ -5539,6 +5540,7 @@ void UWorld::GetAssetRegistryTags(TArray<FAssetRegistryTag>& OutTags) const
 			OutTags.Add(FAssetRegistryTag("FiB", FString(), FAssetRegistryTag::TT_Hidden));
 		}
 	}
+	FWorldDelegates::GetAssetTags.Broadcast(this, OutTags);
 }
 #endif
 
