SpawnActor will now log a warning and fail if the input location or rotation contains NaNs.

diff --git a/Engine/Source/Runtime/Engine/Private/LevelActor.cpp b/Engine/Source/Runtime/Engine/Private/LevelActor.cpp
--- a/Engine/Source/Runtime/Engine/Private/LevelActor.cpp
+++ b/Engine/Source/Runtime/Engine/Private/LevelActor.cpp
@@ -272,6 +272,17 @@ AActor* UWorld::SpawnActor( UClass* Class, FVector const* Location, FRotator con
 		UE_LOG(LogSpawn, Warning, TEXT("SpawnActor failed because we are in the process of tearing down the world"));
 		return NULL;
 	}
+	else if (Location && Location->ContainsNaN())
+	{
+		UE_LOG(LogSpawn, Warning, TEXT("SpawnActor failed because the given location (%s) is invalid"), *Location->ToString());
+		return NULL;
+	}
+	else if (Rotation && Rotation->ContainsNaN())
+	{
+		UE_LOG(LogSpawn, Warning, TEXT("SpawnActor failed because the given rotation (%s) is invalid"), *Rotation->ToString());
+		return NULL;
+		
+	}
 
 	AActor* Template = SpawnParameters.Template;
 	// Use class's default actor as a template.
