Merging 2500036

//depot/UE4-UT/...

to //depot/UE4/...

diff --git a/Engine/Source/Runtime/Engine/Private/PhysicsEngine/PhysAnim.cpp b/Engine/Source/Runtime/Engine/Private/PhysicsEngine/PhysAnim.cpp
--- a/Engine/Source/Runtime/Engine/Private/PhysicsEngine/PhysAnim.cpp
+++ b/Engine/Source/Runtime/Engine/Private/PhysicsEngine/PhysAnim.cpp
@@ -67,11 +67,16 @@ void USkeletalMeshComponent::BlendPhysicsBones( TArray<FBoneIndexType>& InRequir
 	UPhysicsAsset * const PhysicsAsset = GetPhysicsAsset();
 	check( PhysicsAsset );
 
+	if (GetNumSpaceBases() == 0)
+	{
+		return;
+	}
+
 	// Make sure scratch space is big enough.
 	TArray<FAssetWorldBoneTM> WorldBoneTMs;
 	WorldBoneTMs.Reset();
 	WorldBoneTMs.AddZeroed(GetNumSpaceBases());
-
+	
 	FTransform LocalToWorldTM = ComponentToWorld;
 	LocalToWorldTM.RemoveScaling();
 
