Fix physics settings typo (== instead of =)

diff --git a/Engine/Source/Runtime/Engine/Private/PhysicsEngine/PhysicsSettings.cpp b/Engine/Source/Runtime/Engine/Private/PhysicsEngine/PhysicsSettings.cpp
--- a/Engine/Source/Runtime/Engine/Private/PhysicsEngine/PhysicsSettings.cpp
+++ b/Engine/Source/Runtime/Engine/Private/PhysicsEngine/PhysicsSettings.cpp
@@ -38,7 +38,7 @@ void UPhysicsSettings::PostInitProperties()
 
 	if (LockedAxis_DEPRECATED == static_cast<ESettingsLockedAxis::Type>(-1))
 	{
-		LockedAxis_DEPRECATED == ESettingsLockedAxis::Invalid;
+		LockedAxis_DEPRECATED = ESettingsLockedAxis::Invalid;
 	}
 
 	if (LockedAxis_DEPRECATED != ESettingsLockedAxis::Invalid)
