#UE4 - reverted a PVS-Studio change because it was incorrect

diff --git a/Engine/Source/Runtime/Engine/Private/PendingNetGame.cpp b/Engine/Source/Runtime/Engine/Private/PendingNetGame.cpp
--- a/Engine/Source/Runtime/Engine/Private/PendingNetGame.cpp
+++ b/Engine/Source/Runtime/Engine/Private/PendingNetGame.cpp
@@ -245,19 +245,28 @@ void UPendingNetGame::Tick( float DeltaTime )
 	check(NetDriver);
 	check(NetDriver->ServerConnection);
 
+	// Handle timed out or failed connection.
+	//-V522   the check statement above guarantees safety here
+	if (NetDriver->ServerConnection->State == USOCK_Closed && ConnectionError == TEXT(""))
+	{
+		ConnectionError = NSLOCTEXT("Engine", "ConnectionFailed", "Your connection to the host has been lost.").ToString();
+		return;
+	}
+
+	/**
+	 *   Update the network driver
+	 *   ****may NULL itself via CancelPending if a disconnect/error occurs****
+	 */
+
+	//-V522 the check statement above guarantees safety here, the null check below is required, do not remove
+	NetDriver->TickDispatch(DeltaTime);
 	if (NetDriver)
 	{
-		// Handle timed out or failed connection.
-		if (NetDriver->ServerConnection->State == USOCK_Closed && ConnectionError == TEXT(""))
+		NetDriver->TickFlush(DeltaTime);
+		if (NetDriver)
 		{
-			ConnectionError = NSLOCTEXT("Engine", "ConnectionFailed", "Your connection to the host has been lost.").ToString();
-			return;
+			NetDriver->PostTickFlush();
 		}
-
-		// Update network driver (may NULL itself via CancelPending if a disconnect/error occurs)
-		NetDriver->TickDispatch(DeltaTime);
-		NetDriver->TickFlush(DeltaTime);
-		NetDriver->PostTickFlush();
 	}
 }
 
