Tweaking PVS-Studio suppression

diff --git a/Engine/Source/Runtime/Engine/Private/PendingNetGame.cpp b/Engine/Source/Runtime/Engine/Private/PendingNetGame.cpp
--- a/Engine/Source/Runtime/Engine/Private/PendingNetGame.cpp
+++ b/Engine/Source/Runtime/Engine/Private/PendingNetGame.cpp
@@ -245,8 +245,12 @@ void UPendingNetGame::Tick( float DeltaTime )
 	check(NetDriver);
 	check(NetDriver->ServerConnection);
 
+	// The following line disables checks for nullptr access on NetDriver. We have checked() it's validity above,
+	// but the TickDispatch call below may invalidate the ptr, thus we must null check after calling TickDispatch.
+	// PVS-Studio notes that we have used the pointer before null checking (it currently does not understand check)
+	//-V:NetDriver<<:522
+
 	// Handle timed out or failed connection.
-	//-V522   the check statement above guarantees safety here
 	if (NetDriver->ServerConnection->State == USOCK_Closed && ConnectionError == TEXT(""))
 	{
 		ConnectionError = NSLOCTEXT("Engine", "ConnectionFailed", "Your connection to the host has been lost.").ToString();
@@ -257,8 +261,6 @@ void UPendingNetGame::Tick( float DeltaTime )
 	 *   Update the network driver
 	 *   ****may NULL itself via CancelPending if a disconnect/error occurs****
 	 */
-
-	//-V522 the check statement above guarantees safety here, the null check below is required, do not remove
 	NetDriver->TickDispatch(DeltaTime);
 	if (NetDriver)
 	{
