refactor InitBodies to make it easier for serialization change.

diff --git a/Engine/Source/Runtime/Engine/Classes/PhysicsEngine/BodyInstance.h b/Engine/Source/Runtime/Engine/Classes/PhysicsEngine/BodyInstance.h
--- a/Engine/Source/Runtime/Engine/Classes/PhysicsEngine/BodyInstance.h
+++ b/Engine/Source/Runtime/Engine/Classes/PhysicsEngine/BodyInstance.h
@@ -426,7 +426,7 @@ public:
 	 *	@param DebugName Name of the instance for logging
 	 *	@param Setup Body setup for this instance
 	 */
-	static bool ValidateTransform(const FTransform &Transform, FString& DebugName, UBodySetup* Setup);
+	static bool ValidateTransform(const FTransform &Transform, const FString& DebugName, const UBodySetup* Setup);
 
 	/** Standalone path to batch initialise large amounts of static bodies, which will be deferred till the next scene update for fast scene addition.
 	 *	@param Bodies
@@ -995,6 +995,7 @@ private:
 
 	friend class UCollisionProfile;
 	friend class FBodyInstanceCustomization;
+	friend struct FInitBodiesHelper;
 
 #if WITH_BOX2D
 
diff --git a/Engine/Source/Runtime/Engine/Private/PhysicsEngine/BodyInstance.cpp b/Engine/Source/Runtime/Engine/Private/PhysicsEngine/BodyInstance.cpp
--- a/Engine/Source/Runtime/Engine/Private/PhysicsEngine/BodyInstance.cpp
+++ b/Engine/Source/Runtime/Engine/Private/PhysicsEngine/BodyInstance.cpp
@@ -930,58 +930,20 @@ void FBodyInstance::InitBody(class UBodySetup* Setup, const FTransform& Transfor
 	InitBodies(Bodies, Transforms, Setup, PrimComp, InRBScene, InAggregate, bDefer);
 }
 
-void FBodyInstance::InitBodies(TArray<FBodyInstance*>& Bodies, TArray<FTransform>& Transforms, class UBodySetup* BodySetup, class UPrimitiveComponent* PrimitiveComp, class FPhysScene* InRBScene, PhysXAggregateType InAggregate /*= NULL*/, bool bDefer /*= false*/ )
+TSharedPtr<TArray<ANSICHAR>> GetDebugDebugName(const UPrimitiveComponent* PrimitiveComp, const UBodySetup* BodySetup, FString& DebugName)
 {
-	SCOPE_CYCLE_COUNTER(STAT_InitBodies);
-
-	check(BodySetup);
-	check(InRBScene);
-	check(Bodies.Num() > 0);
-
-	int32 NumBodies = Bodies.Num();
-	AActor* OwningActor = PrimitiveComp ? PrimitiveComp->GetOwner() : nullptr;
-
-	const bool bStatic = PrimitiveComp == nullptr || PrimitiveComp->Mobility != EComponentMobility::Movable;
-
-#if WITH_PHYSX
-	PxScene* PSceneForNewDynamic = nullptr;	// TBD
-	PxScene* PSceneSync = InRBScene->GetPhysXScene(PST_Sync);
-	PxScene* PSceneAsync = InRBScene->HasAsyncScene() ? InRBScene->GetPhysXScene(PST_Async) : nullptr;
-
-	TArray<PxActor*> PhysXSyncActors;
-	TArray<PxActor*> PhysXAsyncActors;
-	TArray<PxActor*> PhysXDynamicActors;
-	PhysXSyncActors.Reserve(NumBodies);
-
-	// No reason to defer if there's an existing aggregate.
-	const bool bCanDefer = bDefer && !InAggregate;
-
-	if(PSceneAsync)
-	{
-		// Only allocate this if we can have async bodies
-		PhysXAsyncActors.Reserve(NumBodies);
-	}
-
-	// Ensure we have the AggGeom inside the body setup so we can calculate the number of shapes
-	BodySetup->CreatePhysicsMeshes();
-	int32 NumShapes = BodySetup->AggGeom.GetElementCount() * NumBodies;
-	TArray<PxShape*> PhysxShapes;
-	int32 ShapesWritten = 0;
-	PhysxShapes.AddUninitialized(NumShapes);
-#endif
-
 	// Setup names
 	// Make the debug name for this geometry...
-	FString DebugName(TEXT(""));
+	DebugName = (TEXT(""));
 	TSharedPtr<TArray<ANSICHAR>> PhysXName;
 
 #if (WITH_EDITORONLY_DATA || UE_BUILD_DEBUG || LOOKING_FOR_PERF_ISSUES) && !(UE_BUILD_SHIPPING || UE_BUILD_TEST) && !NO_LOGGING
-	if(PrimitiveComp)
+	if (PrimitiveComp)
 	{
 		DebugName += FString::Printf(TEXT("Component: %s "), *PrimitiveComp->GetReadableName());
 	}
 
-	if(BodySetup->BoneName != NAME_None)
+	if (BodySetup->BoneName != NAME_None)
 	{
 		DebugName += FString::Printf(TEXT("Bone: %s "), *BodySetup->BoneName.ToString());
 	}
@@ -990,23 +952,24 @@ void FBodyInstance::InitBodies(TArray<FBodyInstance*>& Bodies, TArray<FTransform
 	PhysXName = MakeShareable(new TArray<ANSICHAR>(StringToArray<ANSICHAR>(*DebugName, DebugName.Len() + 1)));
 #endif
 
-	// Setup for skeletal components
-	float InstanceBlendWeight = -1.0f;
-	bool bInstanceSimulatePhysics = false;
+	return PhysXName;
+}
 
-	USkeletalMeshComponent* SkelMeshComp = Cast<USkeletalMeshComponent>(PrimitiveComp);
-	if(SkelMeshComp)
+const USkeletalMeshComponent* GetSkeletalMeshComponentAndProperties(const UPrimitiveComponent* PrimitiveComp, const UBodySetup* BodySetup, float& InstanceBlendWeight, bool& bInstanceSimulatePhysics, bool& bComponentAwake)
+{
+	const USkeletalMeshComponent* SkelMeshComp = Cast<USkeletalMeshComponent>(PrimitiveComp);
+	if (SkelMeshComp)
 	{
-		if((BodySetup->PhysicsType == PhysType_Simulated) || (BodySetup->PhysicsType == PhysType_Default))
+		if ((BodySetup->PhysicsType == PhysType_Simulated) || (BodySetup->PhysicsType == PhysType_Default))
 		{
 			bool bEnableSim = (SkelMeshComp && IsRunningDedicatedServer()) ? SkelMeshComp->bEnablePhysicsOnDedicatedServer : true;
 			bEnableSim &= ((BodySetup->PhysicsType == PhysType_Simulated) || (SkelMeshComp->BodyInstance.bSimulatePhysics));	//if unfixed enable. If default look at parent
 
-			if(bEnableSim)
+			if (bEnableSim)
 			{
 				// set simulate to true if using physics
 				bInstanceSimulatePhysics = true;
-				if(BodySetup->PhysicsType == PhysType_Simulated)
+				if (BodySetup->PhysicsType == PhysType_Simulated)
 				{
 					InstanceBlendWeight = 1.f;
 				}
@@ -1014,49 +977,449 @@ void FBodyInstance::InitBodies(TArray<FBodyInstance*>& Bodies, TArray<FTransform
 			else
 			{
 				bInstanceSimulatePhysics = false;
-				if(BodySetup->PhysicsType == PhysType_Simulated)
+				if (BodySetup->PhysicsType == PhysType_Simulated)
 				{
 					InstanceBlendWeight = 0.f;
 				}
 			}
 		}
-	}
 
-	bool bComponentAwake = false;
-	if(SkelMeshComp)
-	{
 		bComponentAwake = SkelMeshComp->BodyInstance.bStartAwake;
 	}
 
+	return SkelMeshComp;
+}
+
+FVector GetInitialLinearVelocity(const AActor* OwningActor, bool& bComponentAwake)
+{
 	FVector InitialLinVel(EForceInit::ForceInitToZero);
-	if(OwningActor)
+	if (OwningActor)
 	{
 		InitialLinVel = OwningActor->GetVelocity();
 
-		if(InitialLinVel.Size() > KINDA_SMALL_NUMBER)
+		if (InitialLinVel.Size() > KINDA_SMALL_NUMBER)
 		{
 			bComponentAwake = true;
 		}
 	}
-	PxVec3 PxInitialLinVel = U2PVector(InitialLinVel);
+
+	return InitialLinVel;
+}
+
+struct FInitBodiesHelper
+{
+	FInitBodiesHelper(const TArray<FBodyInstance*>& InBodies, const TArray<FTransform>& InTransforms, class UBodySetup* InBodySetup, class UPrimitiveComponent* InPrimitiveComp, class FPhysScene* InInRBScene, FBodyInstance::PhysXAggregateType InInAggregate = NULL, bool InDefer = false)
+	: Bodies(InBodies)
+	, Transforms(InTransforms)
+	, BodySetup(InBodySetup)
+	, PrimitiveComp(InPrimitiveComp)
+	, PhysScene(InInRBScene)
+	, InAggregate(InInAggregate)
+	, bDefer(InDefer)
+	, DebugName(TEXT(""))
+	, InstanceBlendWeight(-1.f)
+	, bInstanceSimulatePhysics(false)
+	, bComponentAwake(false)
+	, SkelMeshComp(nullptr)
+	, InitialLinVel(EForceInit::ForceInitToZero)
+	{
+		//Compute all the needed constants
+
+		PhysXName = GetDebugDebugName(PrimitiveComp, BodySetup, DebugName);
+
+		bStatic = PrimitiveComp == nullptr || PrimitiveComp->Mobility != EComponentMobility::Movable;
+		SkelMeshComp = GetSkeletalMeshComponentAndProperties(PrimitiveComp, BodySetup, InstanceBlendWeight, bInstanceSimulatePhysics, bComponentAwake);
+		
+		const AActor* OwningActor = PrimitiveComp ? PrimitiveComp->GetOwner() : nullptr;
+		InitialLinVel = GetInitialLinearVelocity(OwningActor, bComponentAwake);
+
+#if WITH_PHYSX
+		PSyncScene = PhysScene->GetPhysXScene(PST_Sync);
+		PAsyncScene = PhysScene->HasAsyncScene() ? PhysScene->GetPhysXScene(PST_Async) : nullptr;
+#endif
+	
+	}
+
+	//The arguments passed into InitBodies
+	const TArray<FBodyInstance*>& Bodies;
+	const TArray<FTransform>& Transforms;
+	class UBodySetup* BodySetup;
+	class UPrimitiveComponent* PrimitiveComp;
+	class FPhysScene* PhysScene;
+	FBodyInstance::PhysXAggregateType InAggregate;
+	const bool bDefer;
+	
+	FString DebugName;
+	TSharedPtr<TArray<ANSICHAR>> PhysXName;
+
+	//The constants shared between PhysX and Box2D
+	bool bStatic;
+	float InstanceBlendWeight;
+	bool bInstanceSimulatePhysics;
+	bool bComponentAwake;
+
+	const USkeletalMeshComponent* SkelMeshComp;
+	FVector InitialLinVel;
+
+#if WITH_PHYSX
+
+	PxScene* PSyncScene;
+	PxScene* PAsyncScene;
+
+	void InitBodies_PhysX() const;
+	bool CreateShapesAndActors_PhysX(TArray<PxActor*>& PSyncActors, TArray<PxActor*>& PAsyncActors, TArray<PxActor*>& PDynamicActors, const bool bCanDefer, bool& bDynamicsUseAsyncScene) const;
+	physx::PxRigidActor* CreateActor_PhysX_AssumesLocked(FBodyInstance* Instance, const PxTransform& PTransform) const;
+	bool CreateShapes_PhysX_AssumesLocked(FBodyInstance* Instance, physx::PxRigidActor* PNewDynamic, TArray<PxShape*>& PShapes, int32& ShapesWritten) const;
+	void AddActorsToScene_PhysX_AssumesLocked(TArray<PxActor*>& PSyncActors, TArray<PxActor*>& PAsyncActors, TArray<PxActor*>& PDynamicActors) const;
+#endif
 
 #if WITH_BOX2D
+	void InitBodies_Box2D() const;
+#endif
+};
+
+void FBodyInstance::InitBodies(TArray<FBodyInstance*>& Bodies, TArray<FTransform>& Transforms, class UBodySetup* BodySetup, class UPrimitiveComponent* PrimitiveComp, class FPhysScene* InRBScene, PhysXAggregateType InAggregate /*= NULL*/, bool bDefer /*= false*/ )
+{
+	SCOPE_CYCLE_COUNTER(STAT_InitBodies);
+
+	check(BodySetup);
+	check(InRBScene);
+	check(Bodies.Num() > 0);
+
+	FInitBodiesHelper InitBodiesHelper(Bodies, Transforms, BodySetup, PrimitiveComp, InRBScene, InAggregate, bDefer);
+	
+#if WITH_PHYSX
+	InitBodiesHelper.InitBodies_PhysX();
+#endif
+
+#if WITH_BOX2D
+	InitBodiesHelper.InitBodies_Box2D();
+#endif
+}
+
+#if WITH_PHYSX
+
+physx::PxRigidActor* FInitBodiesHelper::CreateActor_PhysX_AssumesLocked(FBodyInstance* Instance, const PxTransform& PTransform) const
+{
+	if (bStatic)
+	{
+		Instance->RigidActorSync = GPhysXSDK->createRigidStatic(PTransform);
+		if (PAsyncScene)
+		{
+			Instance->RigidActorAsync = GPhysXSDK->createRigidStatic(PTransform);
+		}
+	}
+	else
+	{
+		physx::PxRigidActor* PNewDynamic = GPhysXSDK->createRigidDynamic(PTransform);
+		bool bDynamicsUseAsyncScene = Instance->UseAsyncScene(PhysScene);
+
+		// turn off gravity if desired
+		if (!Instance->bEnableGravity)
+		{
+			PNewDynamic->setActorFlag(PxActorFlag::eDISABLE_GRAVITY, true);
+		}
+
+		if (bDynamicsUseAsyncScene)
+		{
+			Instance->RigidActorAsync = PNewDynamic;
+		}
+		else
+		{
+			Instance->RigidActorSync = PNewDynamic;
+		}
+
+		return PNewDynamic;
+	}
+
+	return nullptr;
+}
+
+bool FInitBodiesHelper::CreateShapes_PhysX_AssumesLocked(FBodyInstance* Instance, physx::PxRigidActor* PNewDynamic, TArray<PxShape*>& PShapes, int32& ShapesWritten) const
+{
+	UPhysicalMaterial* SimplePhysMat = Instance->GetSimplePhysicalMaterial();
+	TArray<UPhysicalMaterial*> ComplexPhysMats = Instance->GetComplexPhysicalMaterials();
+	PxMaterial* PSimpleMat = SimplePhysMat->GetPhysXMaterial();
+
+	FShapeData ShapeData;
+	Instance->GetFilterData_AssumesLocked(ShapeData);
+	Instance->GetShapeFlags_AssumesLocked(ShapeData, Instance->CollisionEnabled, BodySetup->CollisionTraceFlag == CTF_UseComplexAsSimple);
+
+	if (PNewDynamic)
+	{
+		if (!Instance->ShouldInstanceSimulatingPhysics())
+		{
+			ShapeData.SyncBodyFlags |= PxRigidBodyFlag::eKINEMATIC;
+		}
+		ShapeData.SyncBodyFlags |= PxRigidBodyFlag::eUSE_KINEMATIC_TARGET_FOR_SCENE_QUERIES;
+	}
+
+	bool bInitFail = false;
+
+	if (Instance->RigidActorSync)
+	{
+		BodySetup->AddShapesToRigidActor_AssumesLocked(Instance, Instance->RigidActorSync, PST_Sync, Instance->Scale3D, PSimpleMat, ComplexPhysMats, ShapeData);
+		bInitFail |= Instance->RigidActorSync->getNbShapes() == 0;
+		Instance->RigidActorSync->userData = &Instance->PhysxUserData;
+		Instance->RigidActorSync->setName(Instance->CharDebugName.IsValid() ? Instance->CharDebugName->GetData() : nullptr);
+
+		ShapesWritten += Instance->RigidActorSync->getShapes(PShapes.GetData() + ShapesWritten, PShapes.Num() - ShapesWritten);
+
+		check(FPhysxUserData::Get<FBodyInstance>(Instance->RigidActorSync->userData) == Instance && FPhysxUserData::Get<FBodyInstance>(Instance->RigidActorSync->userData)->OwnerComponent != NULL);
+
+	}
+
+	if (Instance->RigidActorAsync)
+	{
+		check(PAsyncScene);
+		BodySetup->AddShapesToRigidActor_AssumesLocked(Instance, Instance->RigidActorAsync, PST_Async, Instance->Scale3D, PSimpleMat, ComplexPhysMats, ShapeData);
+		bInitFail |= Instance->RigidActorAsync->getNbShapes() == 0;
+		Instance->RigidActorAsync->userData = &Instance->PhysxUserData;
+		Instance->RigidActorAsync->setName(Instance->CharDebugName.IsValid() ? Instance->CharDebugName->GetData() : nullptr);
+
+		ShapesWritten += Instance->RigidActorAsync->getShapes(PShapes.GetData() + ShapesWritten, PShapes.Num() - ShapesWritten);
+
+		check(FPhysxUserData::Get<FBodyInstance>(Instance->RigidActorAsync->userData) == Instance && FPhysxUserData::Get<FBodyInstance>(Instance->RigidActorAsync->userData)->OwnerComponent != NULL);
+	}
+
+	return bInitFail;
+}
+
+bool FInitBodiesHelper::CreateShapesAndActors_PhysX(TArray<PxActor*>& PSyncActors,TArray<PxActor*>& PAsyncActors,TArray<PxActor*>& PDynamicActors,const bool bCanDefer, bool& bDynamicsUseAsyncScene) const
+{
+	const int32 NumBodies = Bodies.Num();
+	PSyncActors.Reserve(NumBodies);
+
+	if (PAsyncScene)
+	{
+		// Only allocate this if we can have async bodies
+		PAsyncActors.Reserve(NumBodies);
+	}
+
+	// Ensure we have the AggGeom inside the body setup so we can calculate the number of shapes
+	BodySetup->CreatePhysicsMeshes();
+	const int32 NumShapes = BodySetup->AggGeom.GetElementCount() * NumBodies;
+	TArray<PxShape*> PShapes;
+	int32 ShapesWritten = 0;
+	PShapes.AddUninitialized(NumShapes);
+
+	for (int32 BodyIdx = 0; BodyIdx < NumBodies; ++BodyIdx)
+	{
+		FBodyInstance* Instance = Bodies[BodyIdx];
+		const FTransform& Transform = Transforms[BodyIdx];
+
+		FBodyInstance::ValidateTransform(Transform, DebugName, BodySetup);
+
+		Instance->OwnerComponent = PrimitiveComp;
+		Instance->BodySetup = BodySetup;
+		Instance->Scale3D = Transform.GetScale3D();
+		Instance->CharDebugName = PhysXName;
+
+		// Handle autowelding here to avoid extra work
+		if (Instance->bAutoWeld)
+		{
+			UPrimitiveComponent * ParentPrimComponent = PrimitiveComp ? Cast<UPrimitiveComponent>(PrimitiveComp->AttachParent) : NULL;
+
+			//if we have a parent we will now do the weld and exit any further initialization
+			if (ParentPrimComponent && PrimitiveComp->GetWorld() && PrimitiveComp->GetWorld()->IsGameWorld())
+			{
+				if (PrimitiveComp->WeldToImplementation(ParentPrimComponent, PrimitiveComp->AttachSocketName, false))	//welded new simulated body so initialization is done
+				{
+					return false;
+				}
+			}
+		}
+
+		// Don't process if we've already got a body
+		if (Instance->GetPxRigidActor_AssumesLocked())
+		{
+			Instance->OwnerComponent = nullptr;
+			Instance->BodySetup = nullptr;
+
+			continue;
+		}
+
+		// Set sim parameters for bodies from skeletal mesh components
+		if (SkelMeshComp)
+		{
+			Instance->bSimulatePhysics = bInstanceSimulatePhysics;
+			if (InstanceBlendWeight != -1.0f)
+			{
+				Instance->PhysicsBlendWeight = InstanceBlendWeight;
+			}
+		}
+
+		Instance->PhysxUserData = FPhysxUserData(Instance);
+
+		physx::PxRigidActor* PNewDynamic = CreateActor_PhysX_AssumesLocked(Instance, U2PTransform(Transform));	
+		const bool bInitFail = CreateShapes_PhysX_AssumesLocked(Instance, PNewDynamic, PShapes, ShapesWritten);
+
+		if (bInitFail)
+		{
+			UE_LOG(LogPhysics, Log, TEXT("Init Instance %d of Primitive Component %s failed"), BodyIdx, *PrimitiveComp->GetName());
+			if (Instance->RigidActorSync)
+			{
+				Instance->RigidActorSync->release();
+				Instance->RigidActorSync = nullptr;
+			}
+
+			if (Instance->RigidActorAsync)
+			{
+				Instance->RigidActorAsync->release();
+				Instance->RigidActorAsync = nullptr;
+			}
+
+			Instance->OwnerComponent = nullptr;
+			Instance->BodySetup = nullptr;
+
+			continue;
+		}
+
+		if (bCanDefer)
+		{
+			if (PNewDynamic)
+			{
+				PhysScene->DeferAddActor(Instance, PNewDynamic, Instance->UseAsyncScene(PhysScene) ? PST_Async : PST_Sync);
+			}
+			else
+			{
+				if (Instance->RigidActorSync)
+				{
+					PhysScene->DeferAddActor(Instance, Instance->RigidActorSync, PST_Sync);
+				}
+				if (Instance->RigidActorAsync)
+				{
+					PhysScene->DeferAddActor(Instance, Instance->RigidActorAsync, PST_Async);
+				}
+			}
+		}
+		else
+		{
+			if (Instance->RigidActorSync)
+			{
+				PSyncActors.Add(Instance->RigidActorSync);
+			}
+
+			if (Instance->RigidActorAsync)
+			{
+				PAsyncActors.Add(Instance->RigidActorAsync);
+			}
+
+			if (PNewDynamic)
+			{
+				PDynamicActors.Add(PNewDynamic);
+				bDynamicsUseAsyncScene = Instance->UseAsyncScene(PhysScene);
+			}
+
+			// Set the instance to added as we'll add it to the scene in a moment.
+			// This makes sure if it ends up in one of the deferred queues later it 
+			// functions correctly
+			Instance->CurrentSceneState = BodyInstanceSceneState::Added;
+		}
+
+		Instance->SceneIndexSync = PhysScene->PhysXSceneIndex[PST_Sync];
+		Instance->SceneIndexAsync = PAsyncScene ? PhysScene->PhysXSceneIndex[PST_Async] : 0;
+		Instance->InitialLinearVelocity = InitialLinVel;
+		Instance->bWokenExternally = bComponentAwake;
+	}
+
+	return true;
+}
+
+void FInitBodiesHelper::AddActorsToScene_PhysX_AssumesLocked(TArray<PxActor*>& PSyncActors, TArray<PxActor*>& PAsyncActors, TArray<PxActor*>& PDynamicActors) const
+{
+	SCOPE_CYCLE_COUNTER(STAT_BulkSceneAdd);
+
+
+	// Use the aggregate if it exists, has enough slots and is in the correct scene (or no scene)
+	if (InAggregate && PDynamicActors.Num() > 0 && (InAggregate->getMaxNbActors() - InAggregate->getNbActors()) >= (uint32)PDynamicActors.Num())
+	{
+
+		if (InAggregate->getScene() == nullptr || InAggregate->getScene() == PDynamicActors[0]->getScene())
+		{
+			for (PxActor* Actor : PDynamicActors)
+			{
+				InAggregate->addActor(*Actor);
+			}
+		}
+	}
+	else
+	{
+
+		for (PxActor* Actor : PSyncActors)
+		{
+			PSyncScene->addActor(*Actor);
+		}
+
+		if (PAsyncScene)
+		{
+			for (PxActor* Actor : PAsyncActors)
+			{
+				PAsyncScene->addActor(*Actor);
+			}
+		}
+	}
+
+	// Set up dynamic instance data
+	if (!bStatic)
+	{
+		SCOPE_CYCLE_COUNTER(STAT_InitBodyPostAdd);
+		for (int32 BodyIdx = 0, NumBodies = Bodies.Num(); BodyIdx < NumBodies; ++BodyIdx)
+		{
+			FBodyInstance* Instance = Bodies[BodyIdx];
+			Instance->InitDynamicProperties_AssumesLocked();
+		}
+	}
+}
+
+void FInitBodiesHelper::InitBodies_PhysX() const
+{
+	TArray<PxActor*> PSyncActors;
+	TArray<PxActor*> PAsyncActors;
+	TArray<PxActor*> PDynamicActors;
 	
-	if(UBodySetup2D* BodySetup2D = Cast<UBodySetup2D>(BodySetup))
+	// No reason to defer if there's an existing aggregate.
+	const bool bCanDefer = bDefer && !InAggregate;
+	bool bDynamicsUseAsync = false;
+	if(CreateShapesAndActors_PhysX(PSyncActors, PAsyncActors, PDynamicActors, bCanDefer, bDynamicsUseAsync) == false)
 	{
-		if(b2World* BoxWorld = FPhysicsIntegration2D::FindAssociatedWorld(PrimitiveComp->GetWorld()))
+		return;
+	}
+
+	if (!bCanDefer)
+	{
+		const bool bAddingToSyncScene = (PSyncActors.Num() || (PDynamicActors.Num() && !bDynamicsUseAsync)) && PSyncScene;
+		const bool bAddingToAsyncScene = (PAsyncActors.Num() || (PDynamicActors.Num() && bDynamicsUseAsync)) && PAsyncScene;
+
+		SCOPED_SCENE_WRITE_LOCK(bAddingToSyncScene ? PSyncScene : nullptr);
+		SCOPED_SCENE_WRITE_LOCK(bAddingToAsyncScene ? PAsyncScene : nullptr);
+
+		AddActorsToScene_PhysX_AssumesLocked(PSyncActors, PAsyncActors, PDynamicActors);
+	}
+}
+#endif
+
+#if WITH_BOX2D
+
+void FInitBodiesHelper::InitBodies_Box2D() const
+{
+	const int32 NumBodies = Bodies.Num();
+	bool bLocalComponentAwake = bComponentAwake;
+	if (UBodySetup2D* BodySetup2D = Cast<UBodySetup2D>(BodySetup))
+	{
+		if (b2World* BoxWorld = FPhysicsIntegration2D::FindAssociatedWorld(PrimitiveComp->GetWorld()))
 		{
-			for(int32 BodyIdx = 0 ; BodyIdx < NumBodies ; ++BodyIdx)
+			for (int32 BodyIdx = 0; BodyIdx < NumBodies; ++BodyIdx)
 			{
 				FBodyInstance* Instance = Bodies[BodyIdx];
-				FTransform& Transform = Transforms[BodyIdx];
+				const FTransform& Transform = Transforms[BodyIdx];
 
 				const b2Vec2 Scale2D = FPhysicsIntegration2D::ConvertUnrealVectorToBox(Instance->Scale3D);
 				Instance->OwnerComponent = PrimitiveComp;
 
 				// Create the body definition
 				b2BodyDef BodyDefinition;
-				if(bStatic)
+				if (bStatic)
 				{
 					BodyDefinition.type = b2_staticBody;
 				}
@@ -1065,12 +1428,12 @@ void FBodyInstance::InitBodies(TArray<FBodyInstance*>& Bodies, TArray<FTransform
 					BodyDefinition.type = Instance->ShouldInstanceSimulatingPhysics() ? b2_dynamicBody : b2_kinematicBody;
 				}
 
-				if(SkelMeshComp)
+				if (SkelMeshComp)
 				{
-					bComponentAwake = Instance->bStartAwake && SkelMeshComp->BodyInstance.bStartAwake;
+					bLocalComponentAwake = Instance->bStartAwake && SkelMeshComp->BodyInstance.bStartAwake;
 				}
 
-				if(Instance->bStartAwake || bComponentAwake)
+				if (Instance->bStartAwake || bLocalComponentAwake)
 				{
 					BodyDefinition.awake = true;
 				}
@@ -1079,7 +1442,7 @@ void FBodyInstance::InitBodies(TArray<FBodyInstance*>& Bodies, TArray<FTransform
 					BodyDefinition.awake = false;
 				}
 
-				if(Instance->ShouldInstanceSimulatingPhysics())
+				if (Instance->ShouldInstanceSimulatingPhysics())
 				{
 					BodyDefinition.linearVelocity = FPhysicsIntegration2D::ConvertUnrealVectorToBox(InitialLinVel);
 				}
@@ -1089,7 +1452,7 @@ void FBodyInstance::InitBodies(TArray<FBodyInstance*>& Bodies, TArray<FTransform
 				Instance->BodyInstancePtr->SetUserData(Instance);
 
 				// Circles
-				for(const FCircleElement2D& Circle : BodySetup2D->AggGeom2D.CircleElements)
+				for (const FCircleElement2D& Circle : BodySetup2D->AggGeom2D.CircleElements)
 				{
 					b2CircleShape CircleShape;
 					CircleShape.m_radius = Circle.Radius * Instance->Scale3D.Size() / UnrealUnitsPerMeter;
@@ -1103,7 +1466,7 @@ void FBodyInstance::InitBodies(TArray<FBodyInstance*>& Bodies, TArray<FTransform
 				}
 
 				// Boxes
-				for(const FBoxElement2D& Box : BodySetup2D->AggGeom2D.BoxElements)
+				for (const FBoxElement2D& Box : BodySetup2D->AggGeom2D.BoxElements)
 				{
 					const b2Vec2 HalfBoxSize(Box.Width * 0.5f * Scale2D.x, Box.Height * 0.5f * Scale2D.y);
 					const b2Vec2 BoxCenter(Box.Center.X * Scale2D.x, Box.Center.Y * Scale2D.y);
@@ -1118,15 +1481,15 @@ void FBodyInstance::InitBodies(TArray<FBodyInstance*>& Bodies, TArray<FTransform
 				}
 
 				// Convex hulls
-				for(const FConvexElement2D& Convex : BodySetup2D->AggGeom2D.ConvexElements)
+				for (const FConvexElement2D& Convex : BodySetup2D->AggGeom2D.ConvexElements)
 				{
 					const int32 NumVerts = Convex.VertexData.Num();
 
-					if(NumVerts <= b2_maxPolygonVertices)
+					if (NumVerts <= b2_maxPolygonVertices)
 					{
 						TArray<b2Vec2, TInlineAllocator<b2_maxPolygonVertices>> Verts;
 
-						for(int32 VertexIndex = 0; VertexIndex < Convex.VertexData.Num(); ++VertexIndex)
+						for (int32 VertexIndex = 0; VertexIndex < Convex.VertexData.Num(); ++VertexIndex)
 						{
 							const FVector2D SourceVert = Convex.VertexData[VertexIndex];
 							new (Verts)b2Vec2(SourceVert.X * Scale2D.x, SourceVert.Y * Scale2D.y);
@@ -1147,9 +1510,9 @@ void FBodyInstance::InitBodies(TArray<FBodyInstance*>& Bodies, TArray<FTransform
 				}
 
 				// Make sure it contained at least one shape
-				if(Instance->BodyInstancePtr->GetFixtureList() == nullptr)
+				if (Instance->BodyInstancePtr->GetFixtureList() == nullptr)
 				{
-					if(DebugName.Len())
+					if (DebugName.Len())
 					{
 						UE_LOG(LogPhysics, Log, TEXT("InitBody: failed - no shapes: %s"), *DebugName);
 					}
@@ -1176,7 +1539,7 @@ void FBodyInstance::InitBodies(TArray<FBodyInstance*>& Bodies, TArray<FTransform
 				// Set the filter data on the shapes (call this after setting BodyData because it uses that pointer)
 				Instance->UpdatePhysicsFilterData();
 
-				if(!bStatic)
+				if (!bStatic)
 				{
 					// Compute mass (call this after setting BodyData because it uses that pointer)
 					Instance->UpdateMassProperties();
@@ -1207,270 +1570,10 @@ void FBodyInstance::InitBodies(TArray<FBodyInstance*>& Bodies, TArray<FTransform
 		}
 		return;
 	}
-#endif
-
-#if WITH_PHYSX
-	for(int32 BodyIdx = 0 ; BodyIdx < NumBodies ; ++BodyIdx)
-	{
-		FBodyInstance* Instance = Bodies[BodyIdx];
-		FTransform& Transform = Transforms[BodyIdx];
-
-		ValidateTransform(Transform, DebugName, BodySetup);
-
-		Instance->OwnerComponent = PrimitiveComp;
-		Instance->BodySetup = BodySetup;
-		Instance->Scale3D = Transform.GetScale3D();
-		Instance->CharDebugName = PhysXName;
-
-		// Handle autowelding here to avoid extra work
-		if(Instance->bAutoWeld)
-		{
-			UPrimitiveComponent * ParentPrimComponent = PrimitiveComp ? Cast<UPrimitiveComponent>(PrimitiveComp->AttachParent) : NULL;
-
-			//if we have a parent we will now do the weld and exit any further initialization
-			if(ParentPrimComponent && PrimitiveComp->GetWorld() && PrimitiveComp->GetWorld()->IsGameWorld())
-			{
-				if(PrimitiveComp->WeldToImplementation(ParentPrimComponent, PrimitiveComp->AttachSocketName, false))	//welded new simulated body so initialization is done
-				{
-					return;
-				}
-			}
-		}
-
-		PxRigidDynamic* PNewDynamic = nullptr;
-		PxRigidActor* PNewActorSync = nullptr;
-		PxRigidActor* PNewActorAsync = nullptr;
-		PxTransform PTransform = U2PTransform(Transform);
-		bool bInitFail = false;
-
-		// Don't process if we've already got a body
-		if(Instance->GetPxRigidActor_AssumesLocked())
-		{
-			Instance->OwnerComponent = nullptr;
-			Instance->BodySetup = nullptr;
-
-			continue;
-		}
-
-		// Set sim parameters for bodies from skeletal mesh components
-		if(SkelMeshComp)
-		{
-			Instance->bSimulatePhysics = bInstanceSimulatePhysics;
-			if(InstanceBlendWeight != -1.0f)
-			{
-				Instance->PhysicsBlendWeight = InstanceBlendWeight;
-			}
-		}
-
-		if(bStatic)
-		{
-			PNewActorSync = GPhysXSDK->createRigidStatic(PTransform);
-			if(PSceneAsync)
-			{
-				PNewActorAsync = GPhysXSDK->createRigidStatic(PTransform);
-			}
-		}
-		else
-		{
-			PNewDynamic = GPhysXSDK->createRigidDynamic(PTransform);
-
-			// Put the dynamic actor in one scene or the other
-			if(!Instance->UseAsyncScene(InRBScene))
-			{
-				PNewActorSync = PNewDynamic;
-				PSceneForNewDynamic = PSceneSync;
-			}
-			else
-			{
-				PNewActorAsync = PNewDynamic;
-				PSceneForNewDynamic = PSceneAsync;
-			}
-
-			// turn off gravity if desired
-			if(!Instance->bEnableGravity)
-			{
-				PNewDynamic->setActorFlag(PxActorFlag::eDISABLE_GRAVITY, true);
-			}
-		}
-
-		Instance->RigidActorSync = PNewActorSync;
-		Instance->RigidActorAsync = PNewActorAsync;
-		Instance->PhysxUserData = FPhysxUserData(Instance);
-
-		UPhysicalMaterial* SimplePhysMat = Instance->GetSimplePhysicalMaterial();
-		TArray<UPhysicalMaterial*> ComplexPhysMats = Instance->GetComplexPhysicalMaterials();
-		PxMaterial* PSimpleMat = SimplePhysMat->GetPhysXMaterial();
-
-		FShapeData ShapeData;
-		Instance->GetFilterData_AssumesLocked(ShapeData);
-		Instance->GetShapeFlags_AssumesLocked(ShapeData, Instance->CollisionEnabled, BodySetup->CollisionTraceFlag == CTF_UseComplexAsSimple);
-
-		if(PNewDynamic)
-		{
-			if(!Instance->ShouldInstanceSimulatingPhysics())
-			{
-				ShapeData.SyncBodyFlags |= PxRigidBodyFlag::eKINEMATIC;
-			}
-			ShapeData.SyncBodyFlags |= PxRigidBodyFlag::eUSE_KINEMATIC_TARGET_FOR_SCENE_QUERIES;
-		}
-
-		if(PNewActorSync)
-		{
-			BodySetup->AddShapesToRigidActor_AssumesLocked(Instance, PNewActorSync, PST_Sync, Instance->Scale3D, PSimpleMat, ComplexPhysMats, ShapeData);
-			bInitFail |= PNewActorSync->getNbShapes() == 0;
-			PNewActorSync->userData = &Instance->PhysxUserData;
-			PNewActorSync->setName(Instance->CharDebugName.IsValid() ? Instance->CharDebugName->GetData() : nullptr);
-
-			ShapesWritten += PNewActorSync->getShapes(PhysxShapes.GetData() + ShapesWritten, PhysxShapes.Num() - ShapesWritten);
-
-			check(FPhysxUserData::Get<FBodyInstance>(PNewActorSync->userData) == Instance && FPhysxUserData::Get<FBodyInstance>(PNewActorSync->userData)->OwnerComponent != NULL);
-
-		}
-
-		if(PNewActorAsync)
-		{
-			check(PSceneAsync);
-			BodySetup->AddShapesToRigidActor_AssumesLocked(Instance, PNewActorAsync, PST_Async, Instance->Scale3D, PSimpleMat, ComplexPhysMats, ShapeData);
-			bInitFail |= PNewActorAsync->getNbShapes() == 0;
-			PNewActorAsync->userData = &Instance->PhysxUserData;
-			PNewActorAsync->setName(Instance->CharDebugName.IsValid() ? Instance->CharDebugName->GetData() : nullptr);
-
-			ShapesWritten += PNewActorAsync->getShapes(PhysxShapes.GetData() + ShapesWritten, PhysxShapes.Num() - ShapesWritten);
-
-			check(FPhysxUserData::Get<FBodyInstance>(PNewActorAsync->userData) == Instance && FPhysxUserData::Get<FBodyInstance>(PNewActorAsync->userData)->OwnerComponent != NULL);
-		}
-
-		if(bInitFail)
-		{
-			UE_LOG(LogPhysics, Log, TEXT("Init Instance %d of Primitive Component %s failed"), BodyIdx, *PrimitiveComp->GetName());
-			if(PNewActorSync)
-			{
-				PNewActorSync->release();
-				Instance->RigidActorSync = nullptr;
-			}
-
-			if(PNewActorAsync)
-			{
-				PNewActorAsync->release();
-				Instance->RigidActorAsync = nullptr;
-			}
-
-			Instance->OwnerComponent = nullptr;
-			Instance->BodySetup = nullptr;
-
-			continue;
-		}
-
-		if(bCanDefer)
-		{
-			if(PNewDynamic)
-			{
-				InRBScene->DeferAddActor(Instance, PNewDynamic, Instance->UseAsyncScene(InRBScene) ? PST_Async : PST_Sync);
-			}
-			else
-			{
-				if(PNewActorSync)
-				{
-					InRBScene->DeferAddActor(Instance, PNewActorSync, PST_Sync);
-				}
-				if(PNewActorAsync)
-				{
-					InRBScene->DeferAddActor(Instance, PNewActorAsync, PST_Async);
-				}
-			}
-		}
-		else
-		{
-			if(PNewActorSync)
-			{
-				PhysXSyncActors.Add(PNewActorSync);
-			}
-
-			if(PNewActorAsync)
-			{
-				PhysXAsyncActors.Add(PNewActorAsync);
-			}
-
-			if(PNewDynamic)
-			{
-				PhysXDynamicActors.Add(PNewDynamic);
-			}
-
-			// Set the instance to added as we'll add it to the scene in a moment.
-			// This makes sure if it ends up in one of the deferred queues later it 
-			// functions correctly
-			Instance->CurrentSceneState = BodyInstanceSceneState::Added;
-		}
-		
-		Instance->SceneIndexSync = InRBScene->PhysXSceneIndex[PST_Sync];
-		Instance->SceneIndexAsync = PSceneAsync ? InRBScene->PhysXSceneIndex[PST_Async] : 0;
-		Instance->InitialLinearVelocity = InitialLinVel;
-		Instance->bWokenExternally = bComponentAwake;
-	}
-
-	// Grab a material to use, for now we only support batching of bodies using the same material
-	UPhysicalMaterial* SimpleMat = Bodies[0]->GetSimplePhysicalMaterial();
-	TArray<UPhysicalMaterial*> ComplexMats = Bodies[0]->GetComplexPhysicalMaterials();
-	PxMaterial* PhysxSimpleMat = SimpleMat->GetPhysXMaterial();
-
-	if(!bCanDefer)
-	{
-		SCOPE_CYCLE_COUNTER(STAT_BulkSceneAdd);
-		SCOPED_SCENE_WRITE_LOCK(PSceneSync);
-		const bool bAddingToAsyncScene = PhysXAsyncActors.Num() && PSceneAsync;
-		
-		if(bAddingToAsyncScene)
-		{
-			PSceneAsync->lockWrite();
-		}
-
-
-		// Use the aggregate if it exists, has enough slots and is in the correct scene (or no scene)
-		if(InAggregate && PhysXDynamicActors.Num() > 0 && (InAggregate->getMaxNbActors() - InAggregate->getNbActors()) >= (uint32)PhysXDynamicActors.Num() && (InAggregate->getScene() == PSceneForNewDynamic || InAggregate->getScene() == NULL))
-		{
-			for(PxActor* Actor : PhysXDynamicActors)
-			{
-				InAggregate->addActor(*Actor);
-			}
-		}
-		else
-		{
-			
-			for(PxActor* Actor : PhysXSyncActors)
-			{
-				PSceneSync->addActor(*Actor);
-			}
-
-			if(PSceneAsync)
-			{
-				for (PxActor* Actor : PhysXAsyncActors)
-				{
-					PSceneAsync->addActor(*Actor);
-				}
-			}
-		}
-
-		// Set up dynamic instance data
-		if(!bStatic)
-		{
-			SCOPE_CYCLE_COUNTER(STAT_InitBodyPostAdd);
-			for(int32 BodyIdx = 0 ; BodyIdx < NumBodies ; ++BodyIdx)
-			{
-				FBodyInstance* Instance = Bodies[BodyIdx];
-				Instance->InitDynamicProperties_AssumesLocked();
-			}
-		}
-
-		if(bAddingToAsyncScene)
-		{
-			PSceneAsync->unlockWrite();
-		}
-	}
-
-	
-#endif // WITH_PHYSX
 }
 
+#endif
+
 
 #endif // UE_WITH_PHYSICS
 
@@ -4168,7 +4271,7 @@ void FBodyInstance::ApplyMaterialToInstanceShapes_AssumesLocked(PxMaterial* PSim
 	}
 }
 
-bool FBodyInstance::ValidateTransform(const FTransform &Transform, FString& DebugName, UBodySetup* Setup)
+bool FBodyInstance::ValidateTransform(const FTransform &Transform, const FString& DebugName, const UBodySetup* Setup)
 {
 	if(Transform.GetScale3D().IsNearlyZero())
 	{
@@ -4371,7 +4474,7 @@ void FBodyInstance::InitStaticBodies(TArray<FBodyInstance*>& Bodies, TArray<FTra
 {
 	SCOPE_CYCLE_COUNTER(STAT_StaticInitBodies);
 
-	int32 NumBodies = Bodies.Num();
+	const int32 NumBodies = Bodies.Num();
 	AActor* OwningActor = PrimitiveComp ? PrimitiveComp->GetOwner() : nullptr;
 
 	UPhysicalMaterial* SimpleMat;
