Fix: use the correct platform typedef

diff --git a/Engine/Source/Runtime/Engine/Private/UnrealEngine.cpp b/Engine/Source/Runtime/Engine/Private/UnrealEngine.cpp
--- a/Engine/Source/Runtime/Engine/Private/UnrealEngine.cpp
+++ b/Engine/Source/Runtime/Engine/Private/UnrealEngine.cpp
@@ -8782,7 +8782,7 @@ bool UEngine::LoadMap( FWorldContext& WorldContext, FURL URL, class UPendingNetG
 	MALLOC_PROFILER( FMallocProfiler::SnapshotMemoryLoadMapStart( URL.Map ) );
 	Error = TEXT("");
 
-	FGenericPlatformMisc::PreLoadMap(URL.Map, WorldContext.LastURL.Map, nullptr);
+	FPlatformMisc::PreLoadMap(URL.Map, WorldContext.LastURL.Map, nullptr);
 	
 	// make sure level streaming isn't frozen
 	if (WorldContext.World())
