Removed "virtual" keyword from several engine-level BlueprintImplementableEvents.

diff --git a/Engine/Source/Runtime/Engine/Classes/Engine/CanvasRenderTarget2D.h b/Engine/Source/Runtime/Engine/Classes/Engine/CanvasRenderTarget2D.h
--- a/Engine/Source/Runtime/Engine/Classes/Engine/CanvasRenderTarget2D.h
+++ b/Engine/Source/Runtime/Engine/Classes/Engine/CanvasRenderTarget2D.h
@@ -60,7 +60,7 @@ public:
 	 * @param	Height				Height of the render target.
 	 */
 	UFUNCTION(BlueprintImplementableEvent, Category="Canvas Render Target 2D")
-	virtual void ReceiveUpdate(UCanvas* Canvas, int32 Width, int32 Height);
+	void ReceiveUpdate(UCanvas* Canvas, int32 Width, int32 Height);
 
 	/**
 	 * Gets a specific render target's size from the global map of canvas render targets.
diff --git a/Engine/Source/Runtime/Engine/Classes/Engine/GameInstance.h b/Engine/Source/Runtime/Engine/Classes/Engine/GameInstance.h
--- a/Engine/Source/Runtime/Engine/Classes/Engine/GameInstance.h
+++ b/Engine/Source/Runtime/Engine/Classes/Engine/GameInstance.h
@@ -72,11 +72,11 @@ public:
 
 	/** Opportunity for blueprints to handle network errors. */
 	UFUNCTION(BlueprintImplementableEvent, meta=(DisplayName = "NetworkError"))
-	virtual void HandleNetworkError(ENetworkFailure::Type FailureType, bool bIsServer);
+	void HandleNetworkError(ENetworkFailure::Type FailureType, bool bIsServer);
 
 	/** Opportunity for blueprints to handle travel errors. */
 	UFUNCTION(BlueprintImplementableEvent, meta=(DisplayName = "TravelError"))
-	virtual void HandleTravelError(ETravelFailure::Type FailureType);
+	void HandleTravelError(ETravelFailure::Type FailureType);
 
 	/* Called to initialize the game instance for standalone instances of the game */
 	void InitializeStandalone();
diff --git a/Engine/Source/Runtime/Engine/Classes/Engine/LevelScriptActor.h b/Engine/Source/Runtime/Engine/Classes/Engine/LevelScriptActor.h
--- a/Engine/Source/Runtime/Engine/Classes/Engine/LevelScriptActor.h
+++ b/Engine/Source/Runtime/Engine/Classes/Engine/LevelScriptActor.h
@@ -44,7 +44,7 @@ class ENGINE_API ALevelScriptActor : public AActor
 	// --- Level State Functions ------------------------
 	/** @todo document */
 	UFUNCTION(BlueprintImplementableEvent, BlueprintAuthorityOnly)
-	virtual void LevelReset();
+	void LevelReset();
 
 	/**
 	 * Event called on world origin location changes
@@ -53,7 +53,7 @@ class ENGINE_API ALevelScriptActor : public AActor
 	 * @param	NewOriginLocation	New world origin location
 	 */
 	UFUNCTION(BlueprintImplementableEvent)
-	virtual void WorldOriginLocationChanged(FIntVector OldOriginLocation, FIntVector NewOriginLocation);
+	void WorldOriginLocationChanged(FIntVector OldOriginLocation, FIntVector NewOriginLocation);
 	
 #if WITH_EDITOR
 	// Begin UObject Interface
