Further changes from PVS-Studio, submitted without modification, reviewed with owners where I was unsure of the change

diff --git a/Engine/Plugins/Developer/VisualStudioSourceCodeAccess/Source/VisualStudioSourceCodeAccess/Private/VisualStudioSourceCodeAccessor.cpp b/Engine/Plugins/Developer/VisualStudioSourceCodeAccess/Source/VisualStudioSourceCodeAccess/Private/VisualStudioSourceCodeAccessor.cpp
--- a/Engine/Plugins/Developer/VisualStudioSourceCodeAccess/Source/VisualStudioSourceCodeAccess/Private/VisualStudioSourceCodeAccessor.cpp
+++ b/Engine/Plugins/Developer/VisualStudioSourceCodeAccess/Source/VisualStudioSourceCodeAccess/Private/VisualStudioSourceCodeAccessor.cpp
@@ -300,7 +300,7 @@ bool FVisualStudioSourceCodeAccessor::OpenVisualStudioFilesInternalViaDTE(const
 							if ( SUCCEEDED(DTE->get_ActiveDocument(&Document)) &&
 								 SUCCEEDED(Document->get_Selection(&SelectionDispatch)) &&
 								 SUCCEEDED(SelectionDispatch->QueryInterface(&Selection)) &&
-								 SUCCEEDED(Selection->GotoLine(Request.LineNumber, true)) )
+								 SUCCEEDED(Selection->GotoLine(Request.LineNumber, VARIANT_TRUE)) )
 							{
 								if ( !SUCCEEDED(Selection->MoveToLineAndOffset(Request.LineNumber, Request.ColumnNumber, false)) )
 								{
diff --git a/Engine/Plugins/Runtime/WindowsMoviePlayer/Source/WindowsMoviePlayer/Private/WindowsMovieStreamer.cpp b/Engine/Plugins/Runtime/WindowsMoviePlayer/Source/WindowsMoviePlayer/Private/WindowsMovieStreamer.cpp
--- a/Engine/Plugins/Runtime/WindowsMoviePlayer/Source/WindowsMoviePlayer/Private/WindowsMovieStreamer.cpp
+++ b/Engine/Plugins/Runtime/WindowsMoviePlayer/Source/WindowsMoviePlayer/Private/WindowsMovieStreamer.cpp
@@ -411,7 +411,7 @@ FIntPoint FVideoPlayer::AddStreamToTopology(IMFTopology* Topology, IMFPresentati
 
 			IMFMediaType* InputType = NULL;
 			HResult = MFCreateMediaType(&InputType);
-
+			check(SUCCEEDED(HResult));
 			UINT32 Width = 0, Height = 0;
 			HResult = MFGetAttributeSize(OutputType, MF_MT_FRAME_SIZE, &Width, &Height);
 			check(SUCCEEDED(HResult));
diff --git a/Engine/Source/Developer/MeshSimplifier/Private/QuadricMeshReduction.cpp b/Engine/Source/Developer/MeshSimplifier/Private/QuadricMeshReduction.cpp
--- a/Engine/Source/Developer/MeshSimplifier/Private/QuadricMeshReduction.cpp
+++ b/Engine/Source/Developer/MeshSimplifier/Private/QuadricMeshReduction.cpp
@@ -302,10 +302,12 @@ public:
 
 		uint32* Indexes = new uint32[ NumIndexes ];
 
-		for( uint32 i = 0; i < NumIndexes; i++ )
+#if 0
+		for (uint32 i = 0; i < NumIndexes; i++)
 		{
 			 Indexes[i] = InMesh.WedgeIndices[i];
 		}
+#endif
 
 		const uint32 NumAttributes = ( sizeof( TVertSimp<1> ) - sizeof( FVector ) ) / sizeof(float);
 		TMeshSimplifier< TVertSimp<1>, NumAttributes >* MeshSimp = new TMeshSimplifier< TVertSimp<1>, NumAttributes >( Verts, NumVerts, Indexes, NumIndexes );
diff --git a/Engine/Source/Editor/MaterialEditor/Private/SMaterialPalette.cpp b/Engine/Source/Editor/MaterialEditor/Private/SMaterialPalette.cpp
--- a/Engine/Source/Editor/MaterialEditor/Private/SMaterialPalette.cpp
+++ b/Engine/Source/Editor/MaterialEditor/Private/SMaterialPalette.cpp
@@ -208,32 +208,25 @@ void SMaterialPalette::CategorySelectionChanged(TSharedPtr<FString> NewSelection
 
 void SMaterialPalette::AddAssetFromAssetRegistry(const FAssetData& InAddedAssetData)
 {
-	// Grab the asset class, it will be checked for being a material function.
-	UClass* Asset = FindObject<UClass>(ANY_PACKAGE, *InAddedAssetData.AssetClass.ToString());
-
-	if(Asset->IsChildOf(UMaterialFunction::StaticClass()))
-	{
-		RefreshActionsList(true);
-	}
+	RefreshAssetInRegistry(InAddedAssetData);
 }
 
 void SMaterialPalette::RemoveAssetFromRegistry(const FAssetData& InAddedAssetData)
 {
-	// Grab the asset class, it will be checked for being a material function.
-	UClass* Asset = FindObject<UClass>(ANY_PACKAGE, *InAddedAssetData.AssetClass.ToString());
-
-	if(Asset->IsChildOf(UMaterialFunction::StaticClass()))
-	{
-		RefreshActionsList(true);
-	}
+	RefreshAssetInRegistry(InAddedAssetData);
 }
 
 void SMaterialPalette::RenameAssetFromRegistry(const FAssetData& InAddedAssetData, const FString& InNewName)
+{
+	RefreshAssetInRegistry(InAddedAssetData);
+}
+
+void SMaterialPalette::RefreshAssetInRegistry(const FAssetData& InAddedAssetData)
 {
 	// Grab the asset class, it will be checked for being a material function.
 	UClass* Asset = FindObject<UClass>(ANY_PACKAGE, *InAddedAssetData.AssetClass.ToString());
 
-	if(Asset->IsChildOf(UMaterialFunction::StaticClass()))
+	if (Asset->IsChildOf(UMaterialFunction::StaticClass()))
 	{
 		RefreshActionsList(true);
 	}
diff --git a/Engine/Source/Editor/MaterialEditor/Private/SMaterialPalette.h b/Engine/Source/Editor/MaterialEditor/Private/SMaterialPalette.h
--- a/Engine/Source/Editor/MaterialEditor/Private/SMaterialPalette.h
+++ b/Engine/Source/Editor/MaterialEditor/Private/SMaterialPalette.h
@@ -50,6 +50,8 @@ protected:
 	/** Callback from the Asset Registry when an asset is renamed. */
 	void RenameAssetFromRegistry(const FAssetData& InAddedAssetData, const FString& InNewName);
 
+	void RefreshAssetInRegistry(const FAssetData& InAddedAssetData);
+
 protected:
 	/** Pointer back to the material editor that owns us */
 	TWeakPtr<FMaterialEditor> MaterialEditorPtr;
diff --git a/Engine/Source/Editor/UnrealEd/Private/AutoReimport/FileCache.cpp b/Engine/Source/Editor/UnrealEd/Private/AutoReimport/FileCache.cpp
--- a/Engine/Source/Editor/UnrealEd/Private/AutoReimport/FileCache.cpp
+++ b/Engine/Source/Editor/UnrealEd/Private/AutoReimport/FileCache.cpp
@@ -73,7 +73,7 @@ FMD5Hash ReadFileMD5(const FString& Filename, TArray<uint8>* Buffer = nullptr)
 		if (!Buffer)
 		{
 			LocalScratch.SetNumUninitialized(1024*64);
-			Buffer = &LocalScratch;
+			Buffer = &LocalScratch; //-V506
 		}
 		FMD5 MD5;
 
diff --git a/Engine/Source/Editor/UnrealEd/Private/EditorActor.cpp b/Engine/Source/Editor/UnrealEd/Private/EditorActor.cpp
--- a/Engine/Source/Editor/UnrealEd/Private/EditorActor.cpp
+++ b/Engine/Source/Editor/UnrealEd/Private/EditorActor.cpp
@@ -1869,7 +1869,7 @@ void UUnrealEdEngine::edactSelectMatchingStaticMesh( bool bAllClasses )
 			if ( FStaticMeshActor::GetStaticMeshInfoFromActor( Actor, ActorInfo ) )
 			{
 				bool bSelectActor = false;
-				if ( !bSelectActor && (bAllClasses || ActorInfo.IsStaticMeshActor()) )
+				if ( bAllClasses || ActorInfo.IsStaticMeshActor() )
 				{
 					for ( int32 i = 0 ; i < StaticMeshActors.Num() ; ++i )
 					{
diff --git a/Engine/Source/Editor/UnrealEd/Private/EditorObject.cpp b/Engine/Source/Editor/UnrealEd/Private/EditorObject.cpp
--- a/Engine/Source/Editor/UnrealEd/Private/EditorObject.cpp
+++ b/Engine/Source/Editor/UnrealEd/Private/EditorObject.cpp
@@ -587,7 +587,7 @@ const TCHAR* ImportObjectProperties( FImportObjectParams& InParams )
 			Supplier.ClassName = InParams.ObjectStruct->GetOwnerClass() ? InParams.ObjectStruct->GetOwnerClass()->GetName() : FName(NAME_None).ToString();
 			Supplier.CurrentLine = InParams.LineNumber; 
 
-			ContextSupplier = &Supplier;
+			ContextSupplier = &Supplier; //-V506
 		}
 		else
 		{
diff --git a/Engine/Source/Editor/UnrealEd/Private/Fbx/FbxMainExport.cpp b/Engine/Source/Editor/UnrealEd/Private/Fbx/FbxMainExport.cpp
--- a/Engine/Source/Editor/UnrealEd/Private/Fbx/FbxMainExport.cpp
+++ b/Engine/Source/Editor/UnrealEd/Private/Fbx/FbxMainExport.cpp
@@ -169,11 +169,8 @@ void FFbxExporter::WriteToFile(const TCHAR* Filename)
 	FbxExporter* Exporter = FbxExporter::Create(SdkManager, "");
 
 	// set file format
-	if( FileFormat < 0 || FileFormat >= SdkManager->GetIOPluginRegistry()->GetWriterFormatCount() )
-	{
-		// Write in fall back format if pEmbedMedia is true
-		FileFormat = SdkManager->GetIOPluginRegistry()->GetNativeWriterFormat();
-	}
+	// Write in fall back format if pEmbedMedia is true
+	FileFormat = SdkManager->GetIOPluginRegistry()->GetNativeWriterFormat();
 
 	// Set the export states. By default, the export states are always set to 
 	// true except for the option eEXPORT_TEXTURE_AS_EMBEDDED. The code below 
diff --git a/Engine/Source/Programs/UnrealHeaderTool/Private/CodeGenerator.cpp b/Engine/Source/Programs/UnrealHeaderTool/Private/CodeGenerator.cpp
--- a/Engine/Source/Programs/UnrealHeaderTool/Private/CodeGenerator.cpp
+++ b/Engine/Source/Programs/UnrealHeaderTool/Private/CodeGenerator.cpp
@@ -1546,6 +1546,8 @@ FString GetPreservedAccessSpecifierString(FClass* Class)
 		case EAccessSpecifier::ACCESS_Public:
 			PreservedAccessSpecifier = "public:";
 			break;
+		case EAccessSpecifier::ACCESS_NotAnAccessSpecifier :
+			break;
 		}
 	}
 
@@ -2731,11 +2733,8 @@ void FNativeClassHeaderGenerator::ExportMirrorsForNoexportStructs(const TArray<U
 
 bool FNativeClassHeaderGenerator::WillExportEventParms( UFunction* Function )
 {
-	for( TFieldIterator<UProperty> It(Function); It && (It->PropertyFlags&CPF_Parm); ++It )
-	{
-		return true;
-	}
-	return false;
+  TFieldIterator<UProperty> It(Function);
+  return It && (It->PropertyFlags&CPF_Parm);
 }
 
 void WriteEventFunctionPrologue(FUHTStringBuilder& Output, int32 Indent, const FParmsAndReturnProperties& Parameters, UObject* FunctionOuter, const TCHAR* FunctionName)
diff --git a/Engine/Source/Programs/UnrealHeaderTool/Private/HeaderParser.cpp b/Engine/Source/Programs/UnrealHeaderTool/Private/HeaderParser.cpp
--- a/Engine/Source/Programs/UnrealHeaderTool/Private/HeaderParser.cpp
+++ b/Engine/Source/Programs/UnrealHeaderTool/Private/HeaderParser.cpp
@@ -4922,7 +4922,7 @@ void FHeaderParser::CompileInterfaceDeclaration(FClasses& AllClasses)
 
 	// Set the appropriate interface class flags
 	InterfaceClass->ClassFlags |= CLASS_Interface | CLASS_Abstract;
-	if (InterfaceClass->GetSuperStruct() != NULL)
+	if (InterfaceClass->GetSuperClass() != NULL)
 	{
 		InterfaceClass->ClassCastFlags |= InterfaceClass->GetSuperClass()->ClassCastFlags;
 	}
diff --git a/Engine/Source/Runtime/Engine/Private/GameViewportClient.cpp b/Engine/Source/Runtime/Engine/Private/GameViewportClient.cpp
--- a/Engine/Source/Runtime/Engine/Private/GameViewportClient.cpp
+++ b/Engine/Source/Runtime/Engine/Private/GameViewportClient.cpp
@@ -563,7 +563,7 @@ bool UGameViewportClient::RequiresUncapturedAxisInput() const
 
 EMouseCursor::Type UGameViewportClient::GetCursor(FViewport* InViewport, int32 X, int32 Y)
 {
-	bool bIsPlayingMovie = false;//GetMoviePlayer()->IsMovieCurrentlyPlaying();
+	//bool bIsPlayingMovie = false;//GetMoviePlayer()->IsMovieCurrentlyPlaying();
 
 #if !PLATFORM_WINDOWS
 	bool bIsWithinTitleBar = false;
@@ -592,7 +592,7 @@ EMouseCursor::Type UGameViewportClient::GetCursor(FViewport* InViewport, int32 X
 	{
 		return EMouseCursor::Default;
 	}
-	else if ( (!bIsPlayingMovie) && (InViewport->IsFullscreen() || !bIsWithinTitleBar) )
+	else if ( /*(!bIsPlayingMovie) && */(InViewport->IsFullscreen() || !bIsWithinTitleBar) ) //bIsPlayingMovie has always false value
 	{
 		if (GetWorld() && GetWorld()->GetFirstPlayerController())
 		{
@@ -1002,7 +1002,7 @@ void UGameViewportClient::Draw(FViewport* InViewport, FCanvas* SceneCanvas)
 	GetWorld()->UpdateLevelStreaming();
 
 	// Draw the player views.
-	if (!bDisableWorldRendering && !bUIDisableWorldRendering && PlayerViewMap.Num() > 0)
+	if (!bDisableWorldRendering && !bUIDisableWorldRendering && PlayerViewMap.Num() > 0) //-V560
 	{
 		GetRendererModule().BeginRenderingViewFamily(SceneCanvas,&ViewFamily);
 	}
@@ -1207,7 +1207,10 @@ void UGameViewportClient::Draw(FViewport* InViewport, FCanvas* SceneCanvas)
 					GEngine->HMDDevice->DrawDebug(DebugCanvasObject, eSSP_LEFT_EYE);
 				}
 #endif
-				DebugCanvas->PopTransform();
+				if (DebugCanvas != NULL)
+				{
+					DebugCanvas->PopTransform();
+				}
 
 				GEngine->StereoRenderingDevice->PushViewportCanvas(eSSP_RIGHT_EYE, DebugCanvas, DebugCanvasObject, Viewport);
 				ViewportConsole->PostRender_Console(DebugCanvasObject);
@@ -1217,7 +1220,10 @@ void UGameViewportClient::Draw(FViewport* InViewport, FCanvas* SceneCanvas)
 					GEngine->HMDDevice->DrawDebug(DebugCanvasObject, eSSP_RIGHT_EYE);
 				}
 #endif
-				DebugCanvas->PopTransform();
+				if (DebugCanvas != NULL)
+				{
+					DebugCanvas->PopTransform();
+				}
 
 				// Reset the canvas for rendering to the full viewport.
 				DebugCanvasObject->Reset();
diff --git a/Engine/Source/Runtime/Engine/Private/HierarchicalInstancedStaticMesh.cpp b/Engine/Source/Runtime/Engine/Private/HierarchicalInstancedStaticMesh.cpp
--- a/Engine/Source/Runtime/Engine/Private/HierarchicalInstancedStaticMesh.cpp
+++ b/Engine/Source/Runtime/Engine/Private/HierarchicalInstancedStaticMesh.cpp
@@ -892,8 +892,8 @@ static FORCEINLINE_DEBUGGABLE bool CullNode(const FFoliageCullInstanceParams& Pa
 		checkSlow(Params.ViewFrustumLocal.PermutedPlanes.Num() == 4);
 
 		//@todo, once we have more than one mesh per tree, these should be aligned
-		VectorRegister BoxMin = VectorLoad(&Node.BoundMin);
-		VectorRegister BoxMax = VectorLoad(&Node.BoundMax);
+		VectorRegister BoxMin = VectorLoadFloat3(&Node.BoundMin);
+		VectorRegister BoxMax = VectorLoadFloat3(&Node.BoundMax);
 
 		VectorRegister BoxDiff = VectorSubtract(BoxMax,BoxMin);
 		VectorRegister BoxSum = VectorAdd(BoxMax,BoxMin);
@@ -1085,7 +1085,6 @@ void FHierarchicalStaticMeshSceneProxy::FillDynamicMeshElements(FMeshElementColl
 
 					MeshElement.VertexFactory = &InstancedRenderData.VertexFactories[LODIndex];
 					FMeshBatchElement& BatchElement0 = MeshElement.Elements[0];
-					BatchElement0.UserData = (void*)&UserData_AllInstances;
 
 					BatchElement0.UserData = ElementParams.PassUserData[SelectionGroupIndex];
 					BatchElement0.MaxScreenSize = 1.0;
diff --git a/Engine/Source/Runtime/Engine/Private/Materials/HLSLMaterialTranslator.h b/Engine/Source/Runtime/Engine/Private/Materials/HLSLMaterialTranslator.h
--- a/Engine/Source/Runtime/Engine/Private/Materials/HLSLMaterialTranslator.h
+++ b/Engine/Source/Runtime/Engine/Private/Materials/HLSLMaterialTranslator.h
@@ -511,7 +511,7 @@ public:
 							MaterialProperty = MP_MAX; // Indicates we're not compiling any material property.
 							ShaderFrequency = SF_Pixel;
 							TArray<FShaderCodeChunk> CustomExpressionChunks;
-							CurrentScopeChunks = &CustomExpressionChunks;
+							CurrentScopeChunks = &CustomExpressionChunks; //-V506
 							CustomOutput->Compile(this, Index, 0);
 						}
 					}
@@ -2064,7 +2064,7 @@ protected:
 			// be exactly the same as the offset one, so there is no point bringing in the extra code.
 			// Also, we can't access the full offset world position in anything other than the pixel shader, because it won't have
 			// been calculated yet
-			bool bNonOffsetWorldPositionAvailable = Material->MaterialMayModifyMeshPosition() && ShaderFrequency == SF_Pixel;
+			bool bNonOffsetWorldPositionAvailable = Material->MaterialMayModifyMeshPosition();
 
 			switch (WorldPositionIncludedOffsets)
 			{
@@ -2362,6 +2362,8 @@ protected:
 				// Normal maps need to be unpacked in the pixel shader.
 				SampleCode = FString::Printf( TEXT("UnpackNormalMap(%s)"), *SampleCode );
 				break;
+			case SAMPLERTYPE_Masks:
+				break;
 		}
 
 		FString TextureName =
diff --git a/Engine/Source/Runtime/OpenGLDrv/Private/OpenGLDevice.cpp b/Engine/Source/Runtime/OpenGLDrv/Private/OpenGLDevice.cpp
--- a/Engine/Source/Runtime/OpenGLDrv/Private/OpenGLDevice.cpp
+++ b/Engine/Source/Runtime/OpenGLDrv/Private/OpenGLDevice.cpp
@@ -420,7 +420,7 @@ PFNWGLSWAPINTERVALEXTPROC wglSwapIntervalEXT_ProcAddress = NULL;
 static inline void SetupTextureFormat( EPixelFormat Format, const FOpenGLTextureFormat& GLFormat)
 {
 	GOpenGLTextureFormats[Format] = GLFormat;
-	GPixelFormats[Format].Supported = (GLFormat.Format != GL_NONE && GLFormat.InternalFormat != GL_NONE);
+	GPixelFormats[Format].Supported = (GLFormat.Format != GL_NONE && (GLFormat.InternalFormat[0] != GL_NONE || GLFormat.InternalFormat[1] != GL_NONE));
 }
 
 
diff --git a/Engine/Source/Runtime/OpenGLDrv/Public/OpenGLResources.h b/Engine/Source/Runtime/OpenGLDrv/Public/OpenGLResources.h
--- a/Engine/Source/Runtime/OpenGLDrv/Public/OpenGLResources.h
+++ b/Engine/Source/Runtime/OpenGLDrv/Public/OpenGLResources.h
@@ -1028,7 +1028,7 @@ public:
 	
 	}
 	uint32 GetSizeX() const { return GetSize(); }
-	uint32 GetSizeY() const { return GetSize(); }
+	uint32 GetSizeY() const { return GetSize(); } //-V524
 	uint32 GetSizeZ() const { return ArraySize > 1 ? ArraySize : 0; }
 	
 	uint32 GetArraySize() const {return ArraySize;}
diff --git a/Engine/Source/Runtime/Windows/D3D11RHI/Public/D3D11Resources.h b/Engine/Source/Runtime/Windows/D3D11RHI/Public/D3D11Resources.h
--- a/Engine/Source/Runtime/Windows/D3D11RHI/Public/D3D11Resources.h
+++ b/Engine/Source/Runtime/Windows/D3D11RHI/Public/D3D11Resources.h
@@ -478,7 +478,7 @@ public:
 	: FRHITextureCube(InSizeX,InNumMips,InFormat,InFlags)
 	{ check(InNumSamples == 1); }
 	uint32 GetSizeX() const { return GetSize(); }
-	uint32 GetSizeY() const { return GetSize(); }
+	uint32 GetSizeY() const { return GetSize(); } //-V524
 	uint32 GetSizeZ() const { return 0; }
 };
 
