Fix BSP not rebuilding after making changes to some properties

diff --git a/Engine/Source/Runtime/Engine/Private/Brush.cpp b/Engine/Source/Runtime/Engine/Private/Brush.cpp
--- a/Engine/Source/Runtime/Engine/Private/Brush.cpp
+++ b/Engine/Source/Runtime/Engine/Private/Brush.cpp
@@ -58,15 +58,11 @@ void ABrush::PostEditChangeProperty(FPropertyChangedEvent& PropertyChangedEvent)
 		Brush->BuildBound();
 	}
 
-	if(IsStaticBrush() && !GIsTransacting)
+	if(IsStaticBrush() && PropertyChangedEvent.ChangeType != EPropertyChangeType::Interactive && GUndo)
 	{
-		if(PropertyChangedEvent.ChangeType != EPropertyChangeType::Interactive)
-		{
-			GEditor->RebuildAlteredBSP();
-		}
+		// BSP can only be rebuilt during a transaction
+		GEditor->RebuildAlteredBSP();
 
-		// Trigger a csg rebuild if we're in the editor.
-		SetNeedRebuild(GetLevel());
 	}
 
 	bool bIsBuilderBrush = FActorEditorUtils::IsABuilderBrush( this );
