Removed "virtual" from AHUD hitbox BlueprintImplementableEvents. Replaced with  native-overrideable set of virtual functions.

diff --git a/Engine/Source/Runtime/Engine/Classes/GameFramework/HUD.h b/Engine/Source/Runtime/Engine/Classes/GameFramework/HUD.h
--- a/Engine/Source/Runtime/Engine/Classes/GameFramework/HUD.h
+++ b/Engine/Source/Runtime/Engine/Classes/GameFramework/HUD.h
@@ -175,19 +175,27 @@ public:
 
 	/** Called when a hit box is clicked on. Provides the name associated with that box. */
 	UFUNCTION(BlueprintImplementableEvent, BlueprintCosmetic)
-	virtual void ReceiveHitBoxClick(const FName BoxName);
+	void ReceiveHitBoxClick(const FName BoxName);
+	/** Native handler, called when a hit box is clicked on. Provides the name associated with that box. */
+	virtual void NotifyHitBoxClick(FName BoxName);
 
 	/** Called when a hit box is unclicked. Provides the name associated with that box. */
 	UFUNCTION(BlueprintImplementableEvent, BlueprintCosmetic)
-	virtual void ReceiveHitBoxRelease(const FName BoxName);
+	void ReceiveHitBoxRelease(const FName BoxName);
+	/** Native handler, called when a hit box is unclicked. Provides the name associated with that box. */
+	virtual void NotifyHitBoxRelease(FName BoxName);
 
 	/** Called when a hit box is moused over. */
 	UFUNCTION(BlueprintImplementableEvent, BlueprintCosmetic)
-	virtual void ReceiveHitBoxBeginCursorOver(const FName BoxName);
+	void ReceiveHitBoxBeginCursorOver(const FName BoxName);
+	/** Native handler, called when a hit box is moused over. */
+	virtual void NotifyHitBoxBeginCursorOver(FName BoxName);
 
 	/** Called when a hit box no longer has the mouse over it. */
 	UFUNCTION(BlueprintImplementableEvent, BlueprintCosmetic)
-	virtual void ReceiveHitBoxEndCursorOver(const FName BoxName);
+	void ReceiveHitBoxEndCursorOver(const FName BoxName);
+	/** Native handler, called when a hit box no longer has the mouse over it. */
+	virtual void NotifyHitBoxEndCursorOver(FName BoxName);
 
 	//=============================================================================
 	// Kismet API for simple HUD drawing.
diff --git a/Engine/Source/Runtime/Engine/Classes/GameFramework/HUDHitBox.h b/Engine/Source/Runtime/Engine/Classes/GameFramework/HUDHitBox.h
--- a/Engine/Source/Runtime/Engine/Classes/GameFramework/HUDHitBox.h
+++ b/Engine/Source/Runtime/Engine/Classes/GameFramework/HUDHitBox.h
@@ -3,11 +3,12 @@
 #pragma once
 
 /**
- * Hitbox use for clickable hud elements.
+ * Hitbox used for clickable HUD elements.
  * These are named regions that receive various events when clicked upon.
  * @see AHud
  * @see AddHitBox(), ReceiveHitBoxClick(), ReceiveHitBoxRelease(), ReceiveHitBoxBeginCursorOver(), ReceiveHitBoxEndCursorOver()
- * @see  GetHitBoxAtCoordinates(), GetHitBoxesAtCoordinates(), GetHitBoxWithName()
+ * @see NotifyHitBoxClick(), NotifyHitBoxRelease(), NotifyHitBoxBeginCursorOver(), NotifyHitBoxEndCursorOver()
+ * @see GetHitBoxAtCoordinates(), GetHitBoxesAtCoordinates(), GetHitBoxWithName()
  */
 class ENGINE_API FHUDHitBox
 {
diff --git a/Engine/Source/Runtime/Engine/Private/HUD.cpp b/Engine/Source/Runtime/Engine/Private/HUD.cpp
--- a/Engine/Source/Runtime/Engine/Private/HUD.cpp
+++ b/Engine/Source/Runtime/Engine/Private/HUD.cpp
@@ -577,6 +577,30 @@ void AHUD::RemoveAllDebugStrings_Implementation()
 	DebugTextList.Empty();
 }
 
+void AHUD::NotifyHitBoxClick(FName BoxName)
+{
+	// dispatch BP event
+	ReceiveHitBoxClick(BoxName);
+}
+
+void AHUD::NotifyHitBoxRelease(FName BoxName)
+{
+	// dispatch BP event
+	ReceiveHitBoxRelease(BoxName);
+}
+
+void AHUD::NotifyHitBoxBeginCursorOver(FName BoxName)
+{
+	// dispatch BP event
+	ReceiveHitBoxBeginCursorOver(BoxName);
+}
+
+void AHUD::NotifyHitBoxEndCursorOver(FName BoxName)
+{
+	// dispatch BP event
+	ReceiveHitBoxEndCursorOver(BoxName);
+}
+
 void AHUD::GetTextSize(const FString& Text, float& OutWidth, float& OutHeight, class UFont* Font, float Scale) const
 {
 	if (IsCanvasValid_WarnIfNot())
@@ -860,7 +884,7 @@ void AHUD::UpdateHitBoxCandidates( TArray<FVector2D> InContactPoints )
 	for (auto It = NotOverHitBoxes.CreateConstIterator(); It; ++It)
 	{
 		const FName HitBoxName = *It;
-		ReceiveHitBoxEndCursorOver(HitBoxName);
+		NotifyHitBoxEndCursorOver(HitBoxName);
 		HitBoxesOver.Remove(HitBoxName);
 	}
 
@@ -868,7 +892,7 @@ void AHUD::UpdateHitBoxCandidates( TArray<FVector2D> InContactPoints )
 	for (int32 OverIndex = 0; OverIndex < NewlyOverHitBoxes.Num(); ++OverIndex)
 	{
 		const FName HitBoxName = NewlyOverHitBoxes[OverIndex];
-		ReceiveHitBoxBeginCursorOver(HitBoxName);
+		NotifyHitBoxBeginCursorOver(HitBoxName);
 		HitBoxesOver.Add(HitBoxName);
 	}
 }
@@ -934,7 +958,7 @@ bool AHUD::UpdateAndDispatchHitBoxClickEvents(FVector2D ClickLocation, const EIn
 			{
 				bHit = true;
 
-				ReceiveHitBoxClick(HitBox.GetName());
+				NotifyHitBoxClick(HitBox.GetName());
 
 				if (HitBox.ConsumesInput())
 				{
@@ -953,7 +977,7 @@ bool AHUD::UpdateAndDispatchHitBoxClickEvents(FVector2D ClickLocation, const EIn
 
 				if (InEventType == IE_Released)
 				{
-					ReceiveHitBoxRelease(HitBoxHit->GetName());
+					NotifyHitBoxRelease(HitBoxHit->GetName());
 				}
 
 				if (HitBoxHit->ConsumesInput() == true)
