Support for amortizing the cost of spawning many actors during replay scrubbing over multiple frames.
While loading a replay checkpoint, most actor channels will queue bunches instead of processing them immediately.
After the checkpoint load is complete, channels will begin processing bunches again, unless we've already spent a lot of time processing queued bunches during the current frame.
This feature can be toggled with the cvar demo.QueueCheckpointChannels.
The queued bunch processing time limit can be changed with the cvar net.ProcessQueuedBunchesMillisecondLimit.
Also added cycle counter stats for UActorChannel::ProcessQueuedBunches, UDemoNetDriver::ReadDemoFrame, and UDemoNetDriver::LoadCheckpoint.

diff --git a/Engine/Source/Runtime/Engine/Classes/Engine/DemoNetDriver.h b/Engine/Source/Runtime/Engine/Classes/Engine/DemoNetDriver.h
--- a/Engine/Source/Runtime/Engine/Classes/Engine/DemoNetDriver.h
+++ b/Engine/Source/Runtime/Engine/Classes/Engine/DemoNetDriver.h
@@ -75,6 +75,12 @@ class ENGINE_API UDemoNetDriver : public UNetDriver
 
 private:
 	bool		bIsFastForwarding;
+	bool		bIsLoadingCheckpoint;
+
+	TArray<FNetworkGUID> NonQueuedGUIDsForScrubbing;
+
+	// If a channel is associated with Actor, adds the channel's GUID to the list of GUIDs excluded from queuing bunches during scrubbing.
+	void		AddNonQueuedActorForScrubbing(AActor* Actor);
 
 public:
 
@@ -94,6 +100,7 @@ public:
 	bool InitConnectInternal( FString& Error );
 	virtual bool ShouldClientDestroyTearOffActors() const override;
 	virtual bool ShouldSkipRepNotifies() const override;
+	virtual bool ShouldQueueBunchesForActorGUID(FNetworkGUID InGUID) const override;
 
 public:
 
diff --git a/Engine/Source/Runtime/Engine/Classes/Engine/NetDriver.h b/Engine/Source/Runtime/Engine/Classes/Engine/NetDriver.h
--- a/Engine/Source/Runtime/Engine/Classes/Engine/NetDriver.h
+++ b/Engine/Source/Runtime/Engine/Classes/Engine/NetDriver.h
@@ -302,6 +302,9 @@ public:
 	FDelegateHandle TickFlushDelegateHandle;
 	FDelegateHandle PostTickFlushDelegateHandle;
 
+	/** Tracks the amount of time spent during the current frame processing queued bunches. */
+	float ProcessQueuedBunchesCurrentFrameMilliseconds;
+
 	/**
 	* Updates the standby cheat information and
 	 * causes the dialog to be shown/hidden as needed
@@ -537,6 +540,9 @@ public:
 	/** Returns whether or not properties that are replicating using this driver should not call RepNotify functions. */
 	virtual bool ShouldSkipRepNotifies() const { return false; }
 
+	/** Returns true if actor channels with InGUID should queue up bunches, even if they wouldn't otherwise be queued. */
+	virtual bool ShouldQueueBunchesForActorGUID(FNetworkGUID InGUID) const { return false; }
+
 protected:
 
 	/** Adds (fully initialized, ready to go) client connection to the ClientConnections list + any other game related setup */
diff --git a/Engine/Source/Runtime/Engine/Private/DataChannel.cpp b/Engine/Source/Runtime/Engine/Private/DataChannel.cpp
--- a/Engine/Source/Runtime/Engine/Private/DataChannel.cpp
+++ b/Engine/Source/Runtime/Engine/Private/DataChannel.cpp
@@ -30,6 +30,11 @@ static TAutoConsoleVariable<int32> CVarNetReliableDebug(
 	TEXT(" 2: Print reliable bunch buffer each net update"),
 	ECVF_Default);
 
+static TAutoConsoleVariable<float> CVarNetProcessQueuedBunchesMillisecondLimit(
+	TEXT("net.ProcessQueuedBunchesMillisecondLimit"),
+	30.0f,
+	TEXT("Time threshold for processing queued bunches. If it takes longer than this in a single frame, wait until the next frame to continue processing queued bunches."));
+
 /*-----------------------------------------------------------------------------
 	UChannel implementation.
 -----------------------------------------------------------------------------*/
@@ -1678,6 +1683,10 @@ void UActorChannel::Tick()
 
 bool UActorChannel::ProcessQueuedBunches()
 {
+	DECLARE_SCOPE_CYCLE_COUNTER(TEXT("ProcessQueuedBunches time"), STAT_ProcessQueuedBunchesTime, STATGROUP_Net);
+
+	const uint32 QueueBunchStartCycles = FPlatformTime::Cycles();
+
 	// Try to resolve any guids that are holding up the network stream on this channel
 	for ( auto It = PendingGuidResolves.CreateIterator(); It; ++It )
 	{
@@ -1697,11 +1706,16 @@ bool UActorChannel::ProcessQueuedBunches()
 		}
 	}
 
+	const bool bHasTimeToProcess = Connection->Driver->ProcessQueuedBunchesCurrentFrameMilliseconds < CVarNetProcessQueuedBunchesMillisecondLimit.GetValueOnGameThread();
+
 	// We can process all of the queued up bunches if ALL of these are true:
 	//	1. We have queued bunches to process
 	//	2. We no longer have any pending guids to load
 	//	3. We aren't still processing bunches on another channel that this actor was previously on
-	if ( QueuedBunches.Num() > 0 && PendingGuidResolves.Num() == 0 && ( ChIndex == -1 || !Connection->KeepProcessingActorChannelBunchesMap.Contains( ActorNetGUID ) ) )
+	//	4. We haven't spent too much time yet this frame processing queued bunches
+	//	5. The driver isn't requesting queuing for this GUID
+	if ( QueuedBunches.Num() > 0 && PendingGuidResolves.Num() == 0 && ( ChIndex == -1 || !Connection->KeepProcessingActorChannelBunchesMap.Contains( ActorNetGUID ) ) &&
+		 bHasTimeToProcess && !Connection->Driver->ShouldQueueBunchesForActorGUID( ActorNetGUID ) )
 	{
 		for ( int32 i = 0; i < QueuedBunches.Num(); i++ )
 		{
@@ -1726,6 +1740,13 @@ bool UActorChannel::ProcessQueuedBunches()
 		}
 	}
 
+	// Update the driver with our time spent
+	const uint32 QueueBunchEndCycles = FPlatformTime::Cycles();
+	const uint32 QueueBunchDeltaCycles = QueueBunchEndCycles - QueueBunchStartCycles;
+	const float QueueBunchDeltaMilliseconds = FPlatformTime::ToMilliseconds(QueueBunchDeltaCycles);
+
+	Connection->Driver->ProcessQueuedBunchesCurrentFrameMilliseconds += QueueBunchDeltaMilliseconds;
+
 	// Return true if we are done processing queued bunches
 	return QueuedBunches.Num() == 0;
 }
@@ -1792,7 +1813,9 @@ void UActorChannel::ReceivedBunch( FInBunch & Bunch )
 		//	1. We have pending guids to resolve
 		//	2. We already have queued up bunches
 		//	3. If this actor was previously on a channel that is now still processing bunches after a close
-		if ( PendingGuidResolves.Num() > 0 || QueuedBunches.Num() > 0 || Connection->KeepProcessingActorChannelBunchesMap.Contains( ActorNetGUID ) )
+		//	4. The driver is requesting queuing for this GUID
+		if ( PendingGuidResolves.Num() > 0 || QueuedBunches.Num() > 0 || Connection->KeepProcessingActorChannelBunchesMap.Contains( ActorNetGUID ) ||
+			 ( Connection->Driver->ShouldQueueBunchesForActorGUID( ActorNetGUID ) ) )
 		{
 			if ( Connection->KeepProcessingActorChannelBunchesMap.Contains( ActorNetGUID ) )
 			{
diff --git a/Engine/Source/Runtime/Engine/Private/DemoNetDriver.cpp b/Engine/Source/Runtime/Engine/Private/DemoNetDriver.cpp
--- a/Engine/Source/Runtime/Engine/Private/DemoNetDriver.cpp
+++ b/Engine/Source/Runtime/Engine/Private/DemoNetDriver.cpp
@@ -31,6 +31,7 @@ static TAutoConsoleVariable<int32> CVarGotoCheckpointIndex( TEXT( "demo.GotoChec
 static TAutoConsoleVariable<float> CVarGotoTimeInSeconds( TEXT( "demo.GotoTimeInSeconds" ), -1, TEXT( "For testing only, jump to a particular time" ) );
 static TAutoConsoleVariable<int32> CVarDemoFastForwardDestroyTearOffActors( TEXT( "demo.FastForwardDestroyTearOffActors" ), 1, TEXT( "If true, the driver will destroy any torn-off actors immediately while fast-forwarding a replay." ) );
 static TAutoConsoleVariable<int32> CVarDemoFastForwardSkipRepNotifies( TEXT( "demo.FastForwardSkipRepNotifies" ), 1, TEXT( "If true, the driver will optimize fast-forwarding by deferring calls to RepNotify functions until the fast-forward is complete. " ) );
+static TAutoConsoleVariable<int32> CVarDemoQueueCheckpointChannels( TEXT( "demo.QueueCheckpointChannels" ), 1, TEXT( "If true, the driver will put all channels created during checkpoint loading into queuing mode, to amortize the cost of spawning new actors across multiple frames." ) );
 
 static const int32 MAX_DEMO_READ_WRITE_BUFFER = 1024 * 2;
 
@@ -58,6 +59,7 @@ bool UDemoNetDriver::InitBase( bool bInitAsClient, FNetworkNotify* InNotify, con
 		bIsFastForwarding		= false;
 		GotoCheckpointSkipExtraTimeInMS = -1;
 		LastGotoTimeInSeconds	= -1.0f;
+		bIsLoadingCheckpoint	= false;
 
 		ResetDemoState();
 
@@ -1029,6 +1031,8 @@ bool UDemoNetDriver::ConditionallyReadDemoFrame()
 
 bool UDemoNetDriver::ReadDemoFrame( FArchive* Archive )
 {
+	DECLARE_SCOPE_CYCLE_COUNTER(TEXT("ReadDemoFrame time"), STAT_ReplayReadDemoFrameTime, STATGROUP_Net);
+
 	while ( true )
 	{
 		uint8 ReadBuffer[ MAX_DEMO_READ_WRITE_BUFFER ];
@@ -1221,6 +1225,10 @@ void UDemoNetDriver::TickDemoPlayback( float DeltaSeconds )
 	{
 		bIsFastForwarding = false;
 
+		// We may have been fast-forwarding immediately after loading a checkpoint
+		// for fine-grained scrubbing. If so, at this point we are no longer loading a checkpoint.
+		bIsLoadingCheckpoint = false;
+
 		// Flush all pending RepNotifies that were built up during the fast-forward.
 		if ( ServerConnection != nullptr)
 		{
@@ -1311,6 +1319,8 @@ void UDemoNetDriver::CheckpointReady( const bool bSuccess, const int64 SkipExtra
 
 void UDemoNetDriver::LoadCheckpoint()
 {
+	DECLARE_SCOPE_CYCLE_COUNTER(TEXT("LoadCheckpoint time"), STAT_ReplayCheckpointLoadTime, STATGROUP_Net);
+
 	if ( GotoCheckpointArchive == NULL )
 	{
 		UE_LOG( LogDemo, Warning, TEXT( "UDemoNetConnection::LoadCheckpoint: GotoCheckpointArchive == NULL." ) );
@@ -1329,6 +1339,8 @@ void UDemoNetDriver::LoadCheckpoint()
 		return;
 	}
 
+	bIsLoadingCheckpoint = true;
+
 	bRestoreSpectatorPosition = true;
 	SpectatorLocation = SpectatorController->GetSpectatorPawn()->GetActorLocation();
 	SpectatorRotation = SpectatorController->GetControlRotation();//GetSpectatorPawn()->GetActorRotation();
@@ -1336,10 +1348,24 @@ void UDemoNetDriver::LoadCheckpoint()
 	FURL ConnectURL;
 	ConnectURL.Map = DemoFilename;
 
+	// Reset the never-queue GUID list, we'll rebuild it
+	NonQueuedGUIDsForScrubbing.Empty();
+
+	// Save off the SpectatorController's GUID so that we know not to queue his bunches
+	AddNonQueuedActorForScrubbing(SpectatorController);
+
 #if 1
 	// Destroy all non startup actors. They will get restored with the checkpoint
 	for ( FActorIterator It( GetWorld() ); It; ++It )
 	{
+		// If there are any existing actors that are bAlwaysRelevant, don't queue their bunches.
+		// Actors that do queue their bunches might not appear immediately after the checkpoint is loaded,
+		// and missing bAlwaysRelevant actors are more likely to cause noticeable artifacts.
+		if ( It->bAlwaysRelevant )
+		{
+			AddNonQueuedActorForScrubbing(*It);
+		}
+
 		if ( !It->IsNetStartupActor() )
 		{
 			GetWorld()->DestroyActor( *It, true );
@@ -1388,6 +1414,7 @@ void UDemoNetDriver::LoadCheckpoint()
 
 		GotoCheckpointArchive			= NULL;
 		GotoCheckpointSkipExtraTimeInMS = -1;
+		bIsLoadingCheckpoint = false;
 		return;
 	}
 
@@ -1433,6 +1460,10 @@ void UDemoNetDriver::LoadCheckpoint()
 		DemoCurrentTime += (float)GotoCheckpointSkipExtraTimeInMS / 1000;
 		bIsFastForwarding = true;
 	}
+	else
+	{
+		bIsLoadingCheckpoint = false;
+	}
 
 	GotoCheckpointSkipExtraTimeInMS = -1;
 	GotoCheckpointArchive			= NULL;
@@ -1445,6 +1476,32 @@ void UDemoNetDriver::LoadCheckpoint()
 	}
 }
 
+bool UDemoNetDriver::ShouldQueueBunchesForActorGUID(FNetworkGUID InGUID) const
+{
+	if ( CVarDemoQueueCheckpointChannels.GetValueOnGameThread() == 0)
+	{
+		return false;
+	}
+
+	// While loading a checkpoint, queue most bunches so that we don't process them all on one frame.
+	if (bIsLoadingCheckpoint)
+	{
+		return !NonQueuedGUIDsForScrubbing.Contains(InGUID);
+	}
+
+	return false;
+}
+
+void UDemoNetDriver::AddNonQueuedActorForScrubbing(AActor* Actor)
+{
+	auto FoundChannel = ServerConnection->ActorChannels.Find(Actor);
+	if (FoundChannel != nullptr && *FoundChannel != nullptr)
+	{
+		FNetworkGUID ActorGUID = (*FoundChannel)->ActorNetGUID;
+		NonQueuedGUIDsForScrubbing.Add(ActorGUID);
+	}
+}
+
 /*-----------------------------------------------------------------------------
 	UDemoNetConnection.
 -----------------------------------------------------------------------------*/
diff --git a/Engine/Source/Runtime/Engine/Private/NetworkDriver.cpp b/Engine/Source/Runtime/Engine/Private/NetworkDriver.cpp
--- a/Engine/Source/Runtime/Engine/Private/NetworkDriver.cpp
+++ b/Engine/Source/Runtime/Engine/Private/NetworkDriver.cpp
@@ -123,6 +123,7 @@ UNetDriver::UNetDriver(const FObjectInitializer& ObjectInitializer)
 ,	StatPeriod(1.f)
 ,	NetTag(0)
 ,	DebugRelevantActors(false)
+,	ProcessQueuedBunchesCurrentFrameMilliseconds(0.0f)
 {
 }
 
@@ -170,6 +171,9 @@ void UNetDriver::TickFlush(float DeltaSeconds)
 #endif // WITH_SERVER_CODE
 	}
 
+	// Reset queued bunch amortization timer
+	ProcessQueuedBunchesCurrentFrameMilliseconds = 0.0f;
+
 #if STATS
 	// Update network stats (only main game net driver for now)
 	if (NetDriverName == NAME_GameNetDriver && 
