Mac CIS fix #UE4

diff --git a/Engine/Source/Runtime/Engine/Private/AI/Navigation/NavFilters/RecastFilter_UseDefaultArea.cpp b/Engine/Source/Runtime/Engine/Private/AI/Navigation/NavFilters/RecastFilter_UseDefaultArea.cpp
--- a/Engine/Source/Runtime/Engine/Private/AI/Navigation/NavFilters/RecastFilter_UseDefaultArea.cpp
+++ b/Engine/Source/Runtime/Engine/Private/AI/Navigation/NavFilters/RecastFilter_UseDefaultArea.cpp
@@ -10,7 +10,7 @@ URecastFilter_UseDefaultArea::URecastFilter_UseDefaultArea(const FObjectInitiali
 void URecastFilter_UseDefaultArea::InitializeFilter(const ANavigationData& NavData, FNavigationQueryFilter& Filter) const
 {
 #if WITH_RECAST
-	Filter.SetFilterImplementation(static_cast<const INavigationQueryFilterInterface*>(ARecastNavMesh::GetNamedFilter(ERecastNamedFilter::FilterOutAreas)));
+	Filter.SetFilterImplementation(dynamic_cast<const INavigationQueryFilterInterface*>(ARecastNavMesh::GetNamedFilter(ERecastNamedFilter::FilterOutAreas)));
 #endif // WITH_RECAST
 
 	Super::InitializeFilter(NavData, Filter);
