Properly set bIsLoadingCheckpoint when loading the very first checkpoint.

diff --git a/Engine/Source/Runtime/Engine/Private/DemoNetDriver.cpp b/Engine/Source/Runtime/Engine/Private/DemoNetDriver.cpp
--- a/Engine/Source/Runtime/Engine/Private/DemoNetDriver.cpp
+++ b/Engine/Source/Runtime/Engine/Private/DemoNetDriver.cpp
@@ -1411,10 +1411,13 @@ void UDemoNetDriver::LoadCheckpoint()
 			DemoCurrentTime += (float)GotoCheckpointSkipExtraTimeInMS / 1000;
 			bIsFastForwarding = true;
 		}
+		else
+		{
+			bIsLoadingCheckpoint = false;
+		}
 
 		GotoCheckpointArchive			= NULL;
 		GotoCheckpointSkipExtraTimeInMS = -1;
-		bIsLoadingCheckpoint = false;
 		return;
 	}
 
