Fix Mac CIS

diff --git a/Engine/Source/Runtime/Engine/Private/AI/Navigation/NavFilters/RecastFilter_UseDefaultArea.cpp b/Engine/Source/Runtime/Engine/Private/AI/Navigation/NavFilters/RecastFilter_UseDefaultArea.cpp
--- a/Engine/Source/Runtime/Engine/Private/AI/Navigation/NavFilters/RecastFilter_UseDefaultArea.cpp
+++ b/Engine/Source/Runtime/Engine/Private/AI/Navigation/NavFilters/RecastFilter_UseDefaultArea.cpp
@@ -2,6 +2,7 @@
 
 #include "EnginePrivate.h"
 #include "AI/Navigation/NavFilters/RecastFilter_UseDefaultArea.h"
+#include "AI/Navigation/PImplRecastNavMesh.h"
 
 URecastFilter_UseDefaultArea::URecastFilter_UseDefaultArea(const FObjectInitializer& ObjectInitializer) : Super(ObjectInitializer)
 {
