Added friendlier display names for AHUD BlueprintImplementableEvents

diff --git a/Engine/Source/Runtime/Engine/Classes/GameFramework/HUD.h b/Engine/Source/Runtime/Engine/Classes/GameFramework/HUD.h
--- a/Engine/Source/Runtime/Engine/Classes/GameFramework/HUD.h
+++ b/Engine/Source/Runtime/Engine/Classes/GameFramework/HUD.h
@@ -174,25 +174,25 @@ public:
 	void ReceiveDrawHUD(int32 SizeX, int32 SizeY);
 
 	/** Called when a hit box is clicked on. Provides the name associated with that box. */
-	UFUNCTION(BlueprintImplementableEvent, BlueprintCosmetic)
+	UFUNCTION(BlueprintImplementableEvent, BlueprintCosmetic, meta=(DisplayName = "HitBoxClicked"))
 	void ReceiveHitBoxClick(const FName BoxName);
 	/** Native handler, called when a hit box is clicked on. Provides the name associated with that box. */
 	virtual void NotifyHitBoxClick(FName BoxName);
 
 	/** Called when a hit box is unclicked. Provides the name associated with that box. */
-	UFUNCTION(BlueprintImplementableEvent, BlueprintCosmetic)
+	UFUNCTION(BlueprintImplementableEvent, BlueprintCosmetic, meta=(DisplayName = "HitBoxReleased"))
 	void ReceiveHitBoxRelease(const FName BoxName);
 	/** Native handler, called when a hit box is unclicked. Provides the name associated with that box. */
 	virtual void NotifyHitBoxRelease(FName BoxName);
 
 	/** Called when a hit box is moused over. */
-	UFUNCTION(BlueprintImplementableEvent, BlueprintCosmetic)
+	UFUNCTION(BlueprintImplementableEvent, BlueprintCosmetic, meta=(DisplayName = "HitBoxBeginCursorOver"))
 	void ReceiveHitBoxBeginCursorOver(const FName BoxName);
 	/** Native handler, called when a hit box is moused over. */
 	virtual void NotifyHitBoxBeginCursorOver(FName BoxName);
 
 	/** Called when a hit box no longer has the mouse over it. */
-	UFUNCTION(BlueprintImplementableEvent, BlueprintCosmetic)
+	UFUNCTION(BlueprintImplementableEvent, BlueprintCosmetic, meta=(DisplayName = "HitBoxEndCursorOver"))
 	void ReceiveHitBoxEndCursorOver(const FName BoxName);
 	/** Native handler, called when a hit box no longer has the mouse over it. */
 	virtual void NotifyHitBoxEndCursorOver(FName BoxName);
