Removed "virtual" from AActor BlueprintImplementableEvents. Replaced with native-overrideable set of virtual functions.

diff --git a/Engine/Source/Runtime/Engine/Classes/GameFramework/Actor.h b/Engine/Source/Runtime/Engine/Classes/GameFramework/Actor.h
--- a/Engine/Source/Runtime/Engine/Classes/GameFramework/Actor.h
+++ b/Engine/Source/Runtime/Engine/Classes/GameFramework/Actor.h
@@ -1027,52 +1027,79 @@ public:
 	UFUNCTION(BlueprintImplementableEvent, meta=(DisplayName = "Tick"))
 	void ReceiveTick(float DeltaSeconds);
 
+	/** 
+	 *	Event when this actor overlaps another actor, for example a player walking into a trigger.
+	 *	For events when objects have a blocking collision, for example a player hitting a wall, see 'Hit' events.
+	 *	@note Components on both this and the other Actor must have bGenerateOverlapEvents set to true to generate overlap events.
+	 */
+	virtual void NotifyActorBeginOverlap(AActor* OtherActor);
 	/** 
 	 *	Event when this actor overlaps another actor, for example a player walking into a trigger.
 	 *	For events when objects have a blocking collision, for example a player hitting a wall, see 'Hit' events.
 	 *	@note Components on both this and the other Actor must have bGenerateOverlapEvents set to true to generate overlap events.
 	 */
 	UFUNCTION(BlueprintImplementableEvent, meta=(DisplayName = "ActorBeginOverlap"), Category="Collision")
-	virtual void ReceiveActorBeginOverlap(AActor* OtherActor);
+	void ReceiveActorBeginOverlap(AActor* OtherActor);
 
+	/** 
+	 *	Event when an actor no longer overlaps another actor, and they have separated. 
+	 *	@note Components on both this and the other Actor must have bGenerateOverlapEvents set to true to generate overlap events.
+	 */
+	virtual void NotifyActorEndOverlap(AActor* OtherActor);
 	/** 
 	 *	Event when an actor no longer overlaps another actor, and they have separated. 
 	 *	@note Components on both this and the other Actor must have bGenerateOverlapEvents set to true to generate overlap events.
 	 */
 	UFUNCTION(BlueprintImplementableEvent, meta=(DisplayName = "ActorEndOverlap"), Category="Collision")
-	virtual void ReceiveActorEndOverlap(AActor* OtherActor);
+	void ReceiveActorEndOverlap(AActor* OtherActor);
 
+	/** Event when this actor has the mouse moved over it with the clickable interface. */
+	virtual void NotifyActorBeginCursorOver();
 	/** Event when this actor has the mouse moved over it with the clickable interface. */
 	UFUNCTION(BlueprintImplementableEvent, meta=(DisplayName = "ActorBeginCursorOver"), Category="Mouse Input")
-	virtual void ReceiveActorBeginCursorOver();
+	void ReceiveActorBeginCursorOver();
 
+	/** Event when this actor has the mouse moved off of it with the clickable interface. */
+	virtual void NotifyActorEndCursorOver();
 	/** Event when this actor has the mouse moved off of it with the clickable interface. */
 	UFUNCTION(BlueprintImplementableEvent, meta=(DisplayName = "ActorEndCursorOver"), Category="Mouse Input")
-	virtual void ReceiveActorEndCursorOver();
+	void ReceiveActorEndCursorOver();
 
+	/** Event when this actor is clicked by the mouse when using the clickable interface. */
+	virtual void NotifyActorOnClicked();
 	/** Event when this actor is clicked by the mouse when using the clickable interface. */
 	UFUNCTION(BlueprintImplementableEvent, meta=(DisplayName = "ActorOnClicked"), Category="Mouse Input")
-	virtual void ReceiveActorOnClicked();
+	void ReceiveActorOnClicked();
 
+	/** Event when this actor is under the mouse when left mouse button is released while using the clickable interface. */
+	virtual void NotifyActorOnReleased();
 	/** Event when this actor is under the mouse when left mouse button is released while using the clickable interface. */
 	UFUNCTION(BlueprintImplementableEvent, meta=(DisplayName = "ActorOnReleased"), Category="Mouse Input")
-	virtual void ReceiveActorOnReleased();
+	void ReceiveActorOnReleased();
 
+	/** Event when this actor is touched when click events are enabled. */
+	virtual void NotifyActorOnInputTouchBegin(const ETouchIndex::Type FingerIndex);
 	/** Event when this actor is touched when click events are enabled. */
 	UFUNCTION(BlueprintImplementableEvent, meta=(DisplayName = "BeginInputTouch"), Category="Touch Input")
-	virtual void ReceiveActorOnInputTouchBegin(const ETouchIndex::Type FingerIndex);
+	void ReceiveActorOnInputTouchBegin(const ETouchIndex::Type FingerIndex);
 
+	/** Event when this actor is under the finger when untouched when click events are enabled. */
+	virtual void NotifyActorOnInputTouchEnd(const ETouchIndex::Type FingerIndex);
 	/** Event when this actor is under the finger when untouched when click events are enabled. */
 	UFUNCTION(BlueprintImplementableEvent, meta=(DisplayName = "EndInputTouch"), Category="Touch Input")
-	virtual void ReceiveActorOnInputTouchEnd(const ETouchIndex::Type FingerIndex);
+	void ReceiveActorOnInputTouchEnd(const ETouchIndex::Type FingerIndex);
 
+	/** Event when this actor has a finger moved over it with the clickable interface. */
+	virtual void NotifyActorOnInputTouchEnter(const ETouchIndex::Type FingerIndex);
 	/** Event when this actor has a finger moved over it with the clickable interface. */
 	UFUNCTION(BlueprintImplementableEvent, meta=(DisplayName = "TouchEnter"), Category="Touch Input")
-	virtual void ReceiveActorOnInputTouchEnter(const ETouchIndex::Type FingerIndex);
+	void ReceiveActorOnInputTouchEnter(const ETouchIndex::Type FingerIndex);
 
+	/** Event when this actor has a finger moved off of it with the clickable interface. */
+	virtual void NotifyActorOnInputTouchLeave(const ETouchIndex::Type FingerIndex);
 	/** Event when this actor has a finger moved off of it with the clickable interface. */
 	UFUNCTION(BlueprintImplementableEvent, meta=(DisplayName = "TouchLeave"), Category="Touch Input")
-	virtual void ReceiveActorOnInputTouchLeave(const ETouchIndex::Type FingerIndex);
+	void ReceiveActorOnInputTouchLeave(const ETouchIndex::Type FingerIndex);
 
 	/** 
 	 * Returns list of actors this actor is overlapping (any component overlapping any component). Does not return itself.
@@ -1086,6 +1113,16 @@ public:
 	UFUNCTION(BlueprintCallable, Category="Collision", meta=(UnsafeDuringActorConstruction="true"))
 	void GetOverlappingComponents(TArray<UPrimitiveComponent*>& OverlappingComponents) const;
 
+	/** 
+	 * Event when this actor bumps into a blocking object, or blocks another actor that bumps into it.
+	 * This could happen due to things like Character movement, using Set Location with 'sweep' enabled, or physics simulation.
+	 * For events when objects overlap (e.g. walking into a trigger) see the 'Overlap' event.
+	 *
+	 * @note For collisions during physics simulation to generate hit events, 'Simulation Generates Hit Events' must be enabled.
+	 * @note When receiving a hit from another object's movement (bSelfMoved is false), the directions of 'Hit.Normal' and 'Hit.ImpactNormal'
+	 * will be adjusted to indicate force from the other object against this object.
+	 */
+	virtual void NotifyHit(class UPrimitiveComponent* MyComp, AActor* Other, class UPrimitiveComponent* OtherComp, bool bSelfMoved, FVector HitLocation, FVector HitNormal, FVector NormalImpulse, const FHitResult& Hit);
 	/** 
 	 * Event when this actor bumps into a blocking object, or blocks another actor that bumps into it.
 	 * This could happen due to things like Character movement, using Set Location with 'sweep' enabled, or physics simulation.
@@ -1096,7 +1133,8 @@ public:
 	 * will be adjusted to indicate force from the other object against this object.
 	 */
 	UFUNCTION(BlueprintImplementableEvent, meta=(DisplayName = "Hit"), Category="Collision")
-	virtual void ReceiveHit(class UPrimitiveComponent* MyComp, AActor* Other, class UPrimitiveComponent* OtherComp, bool bSelfMoved, FVector HitLocation, FVector HitNormal, FVector NormalImpulse, const FHitResult& Hit);
+	void ReceiveHit(class UPrimitiveComponent* MyComp, AActor* Other, class UPrimitiveComponent* OtherComp, bool bSelfMoved, FVector HitLocation, FVector HitNormal, FVector NormalImpulse, const FHitResult& Hit);
+
 
 	/** Set the lifespan of this actor. When it expires the object will be destroyed. If requested lifespan is 0, the timer is cleared and the actor will not be destroyed. */
 	UFUNCTION(BlueprintCallable, Category="Utilities", meta=(Keywords = "delete destroy"))
diff --git a/Engine/Source/Runtime/Engine/Private/Actor.cpp b/Engine/Source/Runtime/Engine/Private/Actor.cpp
--- a/Engine/Source/Runtime/Engine/Private/Actor.cpp
+++ b/Engine/Source/Runtime/Engine/Private/Actor.cpp
@@ -969,7 +969,6 @@ bool AActor::IsOverlappingActor(const AActor* Other) const
 	return false;
 }
 
-
 void AActor::GetOverlappingActors(TArray<AActor*>& OutOverlappingActors, UClass* ClassFilter) const
 {
 	// prepare output
@@ -1039,6 +1038,74 @@ void AActor::GetOverlappingComponents(TArray<UPrimitiveComponent*>& OutOverlappi
 	}
 }
 
+
+void AActor::NotifyActorBeginOverlap(AActor* OtherActor)
+{
+	// call BP handler
+	ReceiveActorBeginOverlap(OtherActor);
+}
+
+void AActor::NotifyActorEndOverlap(AActor* OtherActor)
+{
+	// call BP handler
+	ReceiveActorEndOverlap(OtherActor);
+}
+
+void AActor::NotifyActorBeginCursorOver()
+{
+	// call BP handler
+	ReceiveActorBeginCursorOver();
+}
+
+void AActor::NotifyActorEndCursorOver()
+{
+	// call BP handler
+	ReceiveActorEndCursorOver();
+}
+
+void AActor::NotifyActorOnClicked()
+{
+	// call BP handler
+	ReceiveActorOnClicked();
+}
+
+void AActor::NotifyActorOnReleased()
+{
+	// call BP handler
+	ReceiveActorOnReleased();
+}
+
+void AActor::NotifyActorOnInputTouchBegin(const ETouchIndex::Type FingerIndex)
+{
+	// call BP handler
+	ReceiveActorOnInputTouchBegin(FingerIndex);
+}
+
+void AActor::NotifyActorOnInputTouchEnd(const ETouchIndex::Type FingerIndex)
+{
+	// call BP handler
+	ReceiveActorOnInputTouchEnd(FingerIndex);
+}
+
+void AActor::NotifyActorOnInputTouchEnter(const ETouchIndex::Type FingerIndex)
+{
+	// call BP handler
+	ReceiveActorOnInputTouchEnter(FingerIndex);
+}
+
+void AActor::NotifyActorOnInputTouchLeave(const ETouchIndex::Type FingerIndex)
+{
+	// call BP handler
+	ReceiveActorOnInputTouchLeave(FingerIndex);
+}
+
+
+void AActor::NotifyHit(class UPrimitiveComponent* MyComp, AActor* Other, class UPrimitiveComponent* OtherComp, bool bSelfMoved, FVector HitLocation, FVector HitNormal, FVector NormalImpulse, const FHitResult& Hit)
+{
+	// call BP handler
+	ReceiveHit(MyComp, Other, OtherComp, bSelfMoved, HitLocation, HitNormal, NormalImpulse, Hit);
+}
+
 /** marks all PrimitiveComponents for which their Owner is relevant for visibility as dirty because the Owner of some Actor in the chain has changed
  * @param TheActor the actor to mark components dirty for
  */
@@ -1788,7 +1855,7 @@ void AActor::InternalDispatchBlockingHit(UPrimitiveComponent* MyComp, UPrimitive
 	// Call virtual
 	if(IsActorValidToNotify(this))
 	{
-		ReceiveHit(MyComp, OtherActor, OtherComp, bSelfMoved, Hit.ImpactPoint, Hit.ImpactNormal, FVector(0,0,0), Hit);
+		NotifyHit(MyComp, OtherActor, OtherComp, bSelfMoved, Hit.ImpactPoint, Hit.ImpactNormal, FVector(0,0,0), Hit);
 	}
 
 	// If we are still ok, call delegate on actor
@@ -3207,7 +3274,7 @@ void AActor::DispatchPhysicsCollisionHit(const FRigidBodyCollisionInfo& MyInfo,
 	Result.BoneName = MyInfo.BoneName;
 	Result.bBlockingHit = true;
 
-	ReceiveHit(MyInfo.Component.Get(), OtherInfo.Actor.Get(), OtherInfo.Component.Get(), true, Result.Location, Result.Normal, RigidCollisionData.TotalNormalImpulse, Result);
+	NotifyHit(MyInfo.Component.Get(), OtherInfo.Actor.Get(), OtherInfo.Component.Get(), true, Result.Location, Result.Normal, RigidCollisionData.TotalNormalImpulse, Result);
 
 	// Execute delegates if bound
 
diff --git a/Engine/Source/Runtime/Engine/Private/Components/PrimitiveComponent.cpp b/Engine/Source/Runtime/Engine/Private/Components/PrimitiveComponent.cpp
--- a/Engine/Source/Runtime/Engine/Private/Components/PrimitiveComponent.cpp
+++ b/Engine/Source/Runtime/Engine/Private/Components/PrimitiveComponent.cpp
@@ -1921,12 +1921,12 @@ void UPrimitiveComponent::BeginComponentOverlap(const FOverlapInfo& OtherOverlap
 					// First actor virtuals
 					if (IsActorValidToNotify(MyActor))
 					{
-						MyActor->ReceiveActorBeginOverlap(OtherActor);
+						MyActor->NotifyActorBeginOverlap(OtherActor);
 					}
 
 					if (IsActorValidToNotify(OtherActor))
 					{
-						OtherActor->ReceiveActorBeginOverlap(MyActor);
+						OtherActor->NotifyActorBeginOverlap(MyActor);
 					}
 
 					// Then level-script delegates
@@ -1985,13 +1985,13 @@ void UPrimitiveComponent::EndComponentOverlap(const FOverlapInfo& OtherOverlap,
 	{			
 		if (IsActorValidToNotify(MyActor))
 		{
-			MyActor->ReceiveActorEndOverlap(OtherActor);
+			MyActor->NotifyActorEndOverlap(OtherActor);
 			MyActor->OnActorEndOverlap.Broadcast(OtherActor);
 		}
 
 		if (IsActorValidToNotify(OtherActor))
 		{
-			OtherActor->ReceiveActorEndOverlap(MyActor);
+			OtherActor->NotifyActorEndOverlap(MyActor);
 			OtherActor->OnActorEndOverlap.Broadcast(MyActor);
 		}
 	}
@@ -2348,7 +2348,7 @@ void UPrimitiveComponent::DispatchMouseOverEvents(UPrimitiveComponent* CurrentCo
 				}
 				if (bBroadcastActorBegin && IsActorValidToNotify(CurrentOwner))
 				{
-					CurrentOwner->ReceiveActorEndCursorOver();
+					CurrentOwner->NotifyActorEndCursorOver();
 					if (IsActorValidToNotify(CurrentOwner))
 					{
 						CurrentOwner->OnEndCursorOver.Broadcast();
@@ -2361,7 +2361,7 @@ void UPrimitiveComponent::DispatchMouseOverEvents(UPrimitiveComponent* CurrentCo
 		{
 			if (bBroadcastActorBegin && IsActorValidToNotify(NewOwner))
 			{
-				NewOwner->ReceiveActorBeginCursorOver();
+				NewOwner->NotifyActorBeginCursorOver();
 				if (IsActorValidToNotify(NewOwner))
 				{
 					NewOwner->OnBeginCursorOver.Broadcast();
@@ -2384,7 +2384,7 @@ void UPrimitiveComponent::DispatchMouseOverEvents(UPrimitiveComponent* CurrentCo
 
 		if (IsActorValidToNotify(CurrentOwner))
 		{
-			CurrentOwner->ReceiveActorEndCursorOver();
+			CurrentOwner->NotifyActorEndCursorOver();
 			if (IsActorValidToNotify(CurrentOwner))
 			{
 				CurrentOwner->OnEndCursorOver.Broadcast();
@@ -2419,7 +2419,7 @@ void UPrimitiveComponent::DispatchTouchOverEvents(ETouchIndex::Type FingerIndex,
 				}
 				if (bBroadcastActorBegin && IsActorValidToNotify(CurrentOwner))
 				{
-					CurrentOwner->ReceiveActorOnInputTouchLeave(FingerIndex);
+					CurrentOwner->NotifyActorOnInputTouchLeave(FingerIndex);
 					if (IsActorValidToNotify(CurrentOwner))
 					{
 						CurrentOwner->OnInputTouchLeave.Broadcast(FingerIndex);
@@ -2432,7 +2432,7 @@ void UPrimitiveComponent::DispatchTouchOverEvents(ETouchIndex::Type FingerIndex,
 		{
 			if (bBroadcastActorBegin && IsActorValidToNotify(NewOwner))
 			{
-				NewOwner->ReceiveActorOnInputTouchEnter(FingerIndex);
+				NewOwner->NotifyActorOnInputTouchEnter(FingerIndex);
 				if (IsActorValidToNotify(NewOwner))
 				{
 					NewOwner->OnInputTouchEnter.Broadcast(FingerIndex);
@@ -2455,7 +2455,7 @@ void UPrimitiveComponent::DispatchTouchOverEvents(ETouchIndex::Type FingerIndex,
 
 		if (IsActorValidToNotify(CurrentOwner))
 		{
-			CurrentOwner->ReceiveActorOnInputTouchLeave(FingerIndex);
+			CurrentOwner->NotifyActorOnInputTouchLeave(FingerIndex);
 			if (IsActorValidToNotify(CurrentOwner))
 			{
 				CurrentOwner->OnInputTouchLeave.Broadcast(FingerIndex);
@@ -2468,7 +2468,7 @@ void UPrimitiveComponent::DispatchOnClicked()
 {
 	if (IsActorValidToNotify(GetOwner()))
 	{
-		GetOwner()->ReceiveActorOnClicked();
+		GetOwner()->NotifyActorOnClicked();
 		if (IsActorValidToNotify(GetOwner()))
 		{
 			GetOwner()->OnClicked.Broadcast();
@@ -2485,7 +2485,7 @@ void UPrimitiveComponent::DispatchOnReleased()
 {
 	if (IsActorValidToNotify(GetOwner()))
 	{
-		GetOwner()->ReceiveActorOnReleased();
+		GetOwner()->NotifyActorOnReleased();
 		if (IsActorValidToNotify(GetOwner()))
 		{
 			GetOwner()->OnReleased.Broadcast();
@@ -2502,7 +2502,7 @@ void UPrimitiveComponent::DispatchOnInputTouchBegin(const ETouchIndex::Type Fing
 {
 	if (IsActorValidToNotify(GetOwner()))
 	{
-		GetOwner()->ReceiveActorOnInputTouchBegin(FingerIndex);
+		GetOwner()->NotifyActorOnInputTouchBegin(FingerIndex);
 		if (IsActorValidToNotify(GetOwner()))
 		{
 			GetOwner()->OnInputTouchBegin.Broadcast(FingerIndex);
@@ -2519,7 +2519,7 @@ void UPrimitiveComponent::DispatchOnInputTouchEnd(const ETouchIndex::Type Finger
 {
 	if (IsActorValidToNotify(GetOwner()))
 	{
-		GetOwner()->ReceiveActorOnInputTouchEnd(FingerIndex);
+		GetOwner()->NotifyActorOnInputTouchEnd(FingerIndex);
 		if (IsActorValidToNotify(GetOwner()))
 		{
 			GetOwner()->OnInputTouchEnd.Broadcast(FingerIndex);
