Added invariant time zone accessor to FText. Made minor changes to blueprint methods for date-time formatting to use the new accessor, rather than hardcoded duplicate string literals. This also fixes the use of an incorrect "Unknown" time zone id by changing it to "Etc/Unknown".

diff --git a/Engine/Source/Runtime/Core/Private/Internationalization/Text.cpp b/Engine/Source/Runtime/Core/Private/Internationalization/Text.cpp
--- a/Engine/Source/Runtime/Core/Private/Internationalization/Text.cpp
+++ b/Engine/Source/Runtime/Core/Private/Internationalization/Text.cpp
@@ -523,6 +523,11 @@ DEF_ASPERCENT(float)
 #undef DEF_ASPERCENT
 #undef DEF_ASPERCENT_CAST
 
+FString FText::GetInvariantTimeZone()
+{
+	return TEXT("Etc/Unknown");
+}
+
 bool FText::FindText( const FString& Namespace, const FString& Key, FText& OutText, const FString* const SourceString )
 {
 	TSharedPtr< FString, ESPMode::ThreadSafe > FoundString = FTextLocalizationManager::Get().FindString( Namespace, Key );
diff --git a/Engine/Source/Runtime/Core/Public/Internationalization/Text.h b/Engine/Source/Runtime/Core/Public/Internationalization/Text.h
--- a/Engine/Source/Runtime/Core/Public/Internationalization/Text.h
+++ b/Engine/Source/Runtime/Core/Public/Internationalization/Text.h
@@ -183,6 +183,11 @@ public:
 	static FText AsTime(const FDateTime& DateTime, const EDateTimeStyle::Type TimeStyle = EDateTimeStyle::Default, const FString& TimeZone = TEXT(""), const FCulturePtr& TargetCulture = NULL);
 	static FText AsTimespan(const FTimespan& Timespan, const FCulturePtr& TargetCulture = NULL);
 
+	/**
+	 * Gets the time zone string that represents a non-specific, zero offset, culture invariant time zone.
+	 */
+	static FString GetInvariantTimeZone();
+
 	/**
 	 * Generate an FText that represents the passed number as a memory size in the current culture
 	 */
diff --git a/Engine/Source/Runtime/Engine/Private/KismetTextLibrary.cpp b/Engine/Source/Runtime/Engine/Private/KismetTextLibrary.cpp
--- a/Engine/Source/Runtime/Engine/Private/KismetTextLibrary.cpp
+++ b/Engine/Source/Runtime/Engine/Private/KismetTextLibrary.cpp
@@ -199,19 +199,19 @@ FText UKismetTextLibrary::AsPercent_Float(float Value, TEnumAsByte<ERoundingMode
 
 FText UKismetTextLibrary::AsDate_DateTime(const FDateTime& InDateTime)
 {
-	return FText::AsDate(InDateTime, EDateTimeStyle::Default, TEXT("Unknown"));
+	return FText::AsDate(InDateTime, EDateTimeStyle::Default, FText::GetInvariantTimeZone());
 }
 
 
 FText UKismetTextLibrary::AsDateTime_DateTime(const FDateTime& InDateTime)
 {
-	return FText::AsDateTime(InDateTime, EDateTimeStyle::Default, EDateTimeStyle::Default, TEXT("Unknown"));
+	return FText::AsDateTime(InDateTime, EDateTimeStyle::Default, EDateTimeStyle::Default, FText::GetInvariantTimeZone());
 }
 
 
 FText UKismetTextLibrary::AsTime_DateTime(const FDateTime& InDateTime)
 {
-	return FText::AsTime(InDateTime, EDateTimeStyle::Default, TEXT("Unknown"));
+	return FText::AsTime(InDateTime, EDateTimeStyle::Default, FText::GetInvariantTimeZone());
 }
 
 
