PR #1012: Add vehicle physx substep (Contributed by [REDACTED_PERSON])

diff --git a/Engine/Source/Runtime/Engine/Classes/Vehicles/WheeledVehicleMovementComponent.h b/Engine/Source/Runtime/Engine/Classes/Vehicles/WheeledVehicleMovementComponent.h
--- a/Engine/Source/Runtime/Engine/Classes/Vehicles/WheeledVehicleMovementComponent.h
+++ b/Engine/Source/Runtime/Engine/Classes/Vehicles/WheeledVehicleMovementComponent.h
@@ -225,6 +225,21 @@ class ENGINE_API UWheeledVehicleMovementComponent : public UPawnMovementComponen
 	UPROPERTY(EditAnywhere, Category = VehicleSetup, AdvancedDisplay)
 	float MaxNormalizedTireLoadFiltered;
 
+    /** PhysX sub-steps
+     More sub-steps provides better stability but with greater computational cost.
+     Typically, vehicles require more sub-steps at very low forward speeds.
+     The threshold longitudinal speed has a default value of 5 metres per second. */
+    UPROPERTY(EditAnywhere, Category=VehicleSetup, AdvancedDisplay, meta = (ClampMin = "0.1", UIMin = "1.0", ClampMax = "1000.0", UIMax = "10.0"))
+    float ThresholdLongitudinalSpeed;
+    
+    /** The sub-step count below the threshold longitudinal speed has a default of 3. */
+    UPROPERTY(EditAnywhere, Category=VehicleSetup, AdvancedDisplay, meta = (ClampMin = "1", UIMin = "1", ClampMax = "10", UIMax = "5"))
+    int32 LowForwardSpeedSubStepCount;
+    
+    /** The sub-step count above the threshold longitudinal speed has a default of 1. */
+    UPROPERTY(EditAnywhere, Category=VehicleSetup, AdvancedDisplay, meta = (ClampMin = "1", UIMin = "1", ClampMax = "10", UIMax = "5"))
+    int32 HighForwardSpeedSubStepCount;
+    
 	// Our instanced wheels
 	UPROPERTY(transient, duplicatetransient, BlueprintReadOnly, Category=Vehicle)
 	TArray<class UVehicleWheel*> Wheels;
diff --git a/Engine/Source/Runtime/Engine/Private/Vehicles/WheeledVehicleMovementComponent.cpp b/Engine/Source/Runtime/Engine/Private/Vehicles/WheeledVehicleMovementComponent.cpp
--- a/Engine/Source/Runtime/Engine/Private/Vehicles/WheeledVehicleMovementComponent.cpp
+++ b/Engine/Source/Runtime/Engine/Private/Vehicles/WheeledVehicleMovementComponent.cpp
@@ -121,6 +121,10 @@ UWheeledVehicleMovementComponent::UWheeledVehicleMovementComponent(const FObject
 	AvoidanceConsiderationRadius = 2000.0f;
 	RVOSteeringStep = 0.5f;
 	RVOThrottleStep = 0.25f;
+    
+    ThresholdLongitudinalSpeed = 5.f;
+    LowForwardSpeedSubStepCount = 3;
+    HighForwardSpeedSubStepCount = 1;
 
 #if WITH_VEHICLE
 	// tire load filtering
@@ -304,7 +308,7 @@ void UWheeledVehicleMovementComponent::SetupWheels( PxVehicleWheelsSimData* PWhe
 		const PxReal LengthScale = 100.f; // Convert default from m to cm
 
 		// Control substepping
-		PWheelsSimData->setSubStepCount(5.f * LengthScale, 3, 1);
+		PWheelsSimData->setSubStepCount(ThresholdLongitudinalSpeed * LengthScale, LowForwardSpeedSubStepCount, HighForwardSpeedSubStepCount);
 		PWheelsSimData->setMinLongSlipDenominator(4.f * LengthScale);
 
 		// Prealloc data for the sprung masses
