#UE4 - Changed the following functions to use AActor instead of APlayerController
- IsNetRelevantFor
- GetNetPriority
- GetNetDormancy
- Added GetNetOwner to remove cast to Pawn/PlayerController in ServerReplicateActors
- constify a few functions (GetNetConnection, IsRelevancyOwnerFor)
- Updated FNetViewer to properly handle non player controller "viewers"

diff --git a/Engine/Source/Developer/GameplayDebugger/Classes/GameplayDebuggingReplicator.h b/Engine/Source/Developer/GameplayDebugger/Classes/GameplayDebuggingReplicator.h
--- a/Engine/Source/Developer/GameplayDebugger/Classes/GameplayDebuggingReplicator.h
+++ b/Engine/Source/Developer/GameplayDebugger/Classes/GameplayDebuggingReplicator.h
@@ -105,9 +105,9 @@ class GAMEPLAYDEBUGGER_API AGameplayDebuggingReplicator : public AActor
 	UFUNCTION()
 	virtual void OnRep_AutoActivate();
 
-	virtual class UNetConnection* GetNetConnection() override;
+	virtual class UNetConnection* GetNetConnection() const override;
 
-	virtual bool IsNetRelevantFor(const APlayerController* RealViewer, const AActor* ViewTarget, const FVector& SrcLocation) const override;
+	virtual bool IsNetRelevantFor(const AActor* RealViewer, const AActor* ViewTarget, const FVector& SrcLocation) const override;
 
 	virtual void PostNetInit() override;
 	
diff --git a/Engine/Source/Developer/GameplayDebugger/Private/GameplayDebuggingReplicator.cpp b/Engine/Source/Developer/GameplayDebugger/Private/GameplayDebuggingReplicator.cpp
--- a/Engine/Source/Developer/GameplayDebugger/Private/GameplayDebuggingReplicator.cpp
+++ b/Engine/Source/Developer/GameplayDebugger/Private/GameplayDebuggingReplicator.cpp
@@ -103,7 +103,7 @@ void AGameplayDebuggingReplicator::GetLifetimeReplicatedProps(TArray< FLifetimeP
 #endif
 }
 
-bool AGameplayDebuggingReplicator::IsNetRelevantFor(const APlayerController* RealViewer, const AActor* ViewTarget, const FVector& SrcLocation) const
+bool AGameplayDebuggingReplicator::IsNetRelevantFor(const AActor* RealViewer, const AActor* ViewTarget, const FVector& SrcLocation) const
 {
 	return LocalPlayerOwner == RealViewer;
 }
@@ -325,7 +325,7 @@ UGameplayDebuggingComponent* AGameplayDebuggingReplicator::GetDebugComponent()
 	return DebugComponent;
 }
 
-class UNetConnection* AGameplayDebuggingReplicator::GetNetConnection()
+class UNetConnection* AGameplayDebuggingReplicator::GetNetConnection() const
 {
 #if !(UE_BUILD_SHIPPING || UE_BUILD_TEST)
 	if (LocalPlayerOwner && LocalPlayerOwner->IsPendingKill() == false && IsPendingKill() == false)
diff --git a/Engine/Source/Runtime/Engine/Classes/GameFramework/Actor.h b/Engine/Source/Runtime/Engine/Classes/GameFramework/Actor.h
--- a/Engine/Source/Runtime/Engine/Classes/GameFramework/Actor.h
+++ b/Engine/Source/Runtime/Engine/Classes/GameFramework/Actor.h
@@ -1500,19 +1500,22 @@ public:
 	/**
 	 * @param ViewPos		Position of the viewer
 	 * @param ViewDir		Vector direction of viewer
-	 * @param Viewer		PlayerController owned by the client for whom net priority is being determined
+	 * @param Viewer		"net object" owned by the client for whom net priority is being determined (typically player controller)
 	 * @param ViewTarget	The actor that is currently being viewed/controlled by Viewer, usually a pawn
 	 * @param InChannel		Channel on which this actor is being replicated.
 	 * @param Time			Time since actor was last replicated
-	 * @param bLowBandwidth True if low bandwith of viewer
+	 * @param bLowBandwidth True if low bandwidth of viewer
 	 * @return				Priority of this actor for replication
 	 */
-	virtual float GetNetPriority(const FVector& ViewPos, const FVector& ViewDir, class APlayerController* Viewer, AActor* ViewTarget, UActorChannel* InChannel, float Time, bool bLowBandwidth);
+	virtual float GetNetPriority(const FVector& ViewPos, const FVector& ViewDir, class AActor* Viewer, AActor* ViewTarget, UActorChannel* InChannel, float Time, bool bLowBandwidth);
 
 	DEPRECATED(4.8, "GetNetPriority now takes a ViewTarget, please override that version.")
 	virtual float GetNetPriority(const FVector& ViewPos, const FVector& ViewDir, class APlayerController* Viewer, UActorChannel* InChannel, float Time, bool bLowBandwidth);
 
 	/** Returns true if the actor should be dormant for a specific net connection. Only checked for DORM_DormantPartial */
+	virtual bool GetNetDormancy(const FVector& ViewPos, const FVector& ViewDir, class AActor* Viewer, AActor* ViewTarget, UActorChannel* InChannel, float Time, bool bLowBandwidth);
+
+	DEPRECATED(4.8, "GetNetDormancy changed from PlayerController parameter to Actor parameter, please override that version.")
 	virtual bool GetNetDormancy(const FVector& ViewPos, const FVector& ViewDir, class APlayerController* Viewer, AActor* ViewTarget, UActorChannel* InChannel, float Time, bool bLowBandwidth);
 
 	/** 
@@ -1686,12 +1689,15 @@ public:
 	// Actor relevancy determination
 
 	/** 
-	  * @param RealViewer - is the PlayerController associated with the client for which network relevancy is being checked
-	  * @param VieweTarget - is the Actor being used as the point of view for the PlayerController
+	  * @param RealViewer - is the "controlling net object" associated with the client for which network relevancy is being checked (typically player controller)
+	  * @param ViewTarget - is the Actor being used as the point of view for the RealViewer
 	  * @param SrcLocation - is the viewing location
 	  *
 	  * @return bool - true if this actor is network relevant to the client associated with RealViewer 
 	  */
+	virtual bool IsNetRelevantFor(const AActor* RealViewer, const AActor* ViewTarget, const FVector& SrcLocation) const;
+
+	DEPRECATED(4.8, "GetNetDormancy changed from PlayerController parameter to Actor parameter, please override that version.")
 	virtual bool IsNetRelevantFor(const APlayerController* RealViewer, const AActor* ViewTarget, const FVector& SrcLocation) const;
 
 	/**
@@ -1703,7 +1709,7 @@ public:
 	 *
 	 * @return bool - true if this actor should be considered the owner
 	 */
-	virtual bool IsRelevancyOwnerFor(AActor* ReplicatedActor, AActor* ActorOwner, AActor* ConnectionActor);
+	virtual bool IsRelevancyOwnerFor(const AActor* ReplicatedActor, const AActor* ActorOwner, const AActor* ConnectionActor) const;
 
 	/** Called after the actor is spawned in the world.  Responsible for setting up actor for play. */
 	void PostSpawnInitialize(FVector const& SpawnLocation, FRotator const& SpawnRotation, AActor* InOwner, APawn* InInstigator, bool bRemoteOwned, bool bNoFail, bool bDeferConstruction);
@@ -1739,6 +1745,9 @@ public:
 	/** Dispatches ReceiveHit virtual and OnComponentHit delegate */
 	void DispatchPhysicsCollisionHit(const struct FRigidBodyCollisionInfo& MyInfo, const struct FRigidBodyCollisionInfo& OtherInfo, const FCollisionImpactData& RigidCollisionData);
 	
+	/** @return the actor responsible for replication, if any.  Typically the player controller */
+	virtual const AActor* GetNetOwner() const;
+
 	/** @return the owning UPlayer (if any) of this actor. This will be a local player, a net connection, or NULL. */
 	virtual class UPlayer* GetNetOwningPlayer();
 
@@ -1746,7 +1755,7 @@ public:
 	 * Get the owning connection used for communicating between client/server 
 	 * @return NetConnection to the client or server for this actor
 	 */
-	virtual class UNetConnection* GetNetConnection();
+	virtual class UNetConnection* GetNetConnection() const;
 
 	/**
 	 * Gets the net mode for this actor, indicating whether it is a client or server (including standalone/not networked).
diff --git a/Engine/Source/Runtime/Engine/Classes/GameFramework/Pawn.h b/Engine/Source/Runtime/Engine/Classes/GameFramework/Pawn.h
--- a/Engine/Source/Runtime/Engine/Classes/GameFramework/Pawn.h
+++ b/Engine/Source/Runtime/Engine/Classes/GameFramework/Pawn.h
@@ -198,7 +198,7 @@ public:
 	virtual void Reset() override;
 	virtual FString GetHumanReadableName() const override;
 	virtual bool ShouldTickIfViewportsOnly() const override;
-	virtual bool IsNetRelevantFor(const APlayerController* RealViewer, const AActor* ViewTarget, const FVector& SrcLocation) const override;
+	virtual bool IsNetRelevantFor(const AActor* RealViewer, const AActor* ViewTarget, const FVector& SrcLocation) const override;
 	virtual void PostNetReceiveLocationAndRotation() override;
 	virtual void PostNetReceiveVelocity(const FVector& NewVelocity) override;
 	virtual void DisplayDebug(UCanvas* Canvas, const FDebugDisplayInfo& DebugDisplay, float& YL, float& YPos) override;
@@ -207,8 +207,9 @@ public:
 	virtual void Destroyed() override;
 	virtual void PreInitializeComponents() override;
 	virtual void PostInitializeComponents() override;
+	virtual const AActor* GetNetOwner() const override;
 	virtual UPlayer* GetNetOwningPlayer() override;
-	virtual UNetConnection* GetNetConnection() override;
+	virtual class UNetConnection* GetNetConnection() const override;
 	virtual void PostLoad() override;
 	virtual void PostRegisterAllComponents() override;
 	virtual float TakeDamage(float Damage, struct FDamageEvent const& DamageEvent, AController* EventInstigator, AActor* DamageCauser) override;
diff --git a/Engine/Source/Runtime/Engine/Classes/GameFramework/PlayerController.h b/Engine/Source/Runtime/Engine/Classes/GameFramework/PlayerController.h
--- a/Engine/Source/Runtime/Engine/Classes/GameFramework/PlayerController.h
+++ b/Engine/Source/Runtime/Engine/Classes/GameFramework/PlayerController.h
@@ -1173,7 +1173,7 @@ public:
 	virtual void GetActorEyesViewPoint(FVector& Location, FRotator& Rotation) const override;
 	virtual void CalcCamera(float DeltaTime, struct FMinimalViewInfo& OutResult) override;
 	virtual void TickActor(float DeltaTime, enum ELevelTick TickType, FActorTickFunction& ThisTickFunction) override;
-	virtual bool IsNetRelevantFor(const APlayerController* RealViewer, const AActor* ViewTarget, const FVector& SrcLocation) const override;
+	virtual bool IsNetRelevantFor(const AActor* RealViewer, const AActor* ViewTarget, const FVector& SrcLocation) const override;
 	virtual void FellOutOfWorld(const class UDamageType& dmgType) override;
 	virtual void Reset() override;
 	virtual void Possess(APawn* aPawn) override;
@@ -1184,9 +1184,10 @@ public:
 	virtual void OnActorChannelOpen(class FInBunch& InBunch, class UNetConnection* Connection) override;
 	virtual void OnSerializeNewActor(class FOutBunch& OutBunch) override;
 	virtual void OnNetCleanup(class UNetConnection* Connection) override;
-	virtual float GetNetPriority(const FVector& ViewPos, const FVector& ViewDir, APlayerController* Viewer, AActor* ViewTarget, UActorChannel* InChannel, float Time, bool bLowBandwidth) override;
+	virtual float GetNetPriority(const FVector& ViewPos, const FVector& ViewDir, AActor* Viewer, AActor* ViewTarget, UActorChannel* InChannel, float Time, bool bLowBandwidth) override;
+	virtual const AActor* GetNetOwner() const override;
 	virtual class UPlayer* GetNetOwningPlayer() override;
-	virtual class UNetConnection* GetNetConnection() override;
+	virtual class UNetConnection* GetNetConnection() const override;
 	virtual void DisplayDebug(class UCanvas* Canvas, const FDebugDisplayInfo& DebugDisplay, float& YL, float& YPos) override;
 	virtual void PostInitializeComponents() override;
 	virtual void EnableInput(class APlayerController* PlayerController) override;
diff --git a/Engine/Source/Runtime/Engine/Classes/GameFramework/WorldSettings.h b/Engine/Source/Runtime/Engine/Classes/GameFramework/WorldSettings.h
--- a/Engine/Source/Runtime/Engine/Classes/GameFramework/WorldSettings.h
+++ b/Engine/Source/Runtime/Engine/Classes/GameFramework/WorldSettings.h
@@ -164,11 +164,11 @@ struct ENGINE_API FNetViewer
 {
 	GENERATED_USTRUCT_BODY()
 
-	/** The player controller corresponding to this view */
+	/** The "controlling net object" associated with this view (typically player controller) */
 	UPROPERTY()
-	class APlayerController* InViewer;
+	class AActor* InViewer;
 
-	/** The actor that is being directly viewed, usually a pawn */
+	/** The actor that is being directly viewed, usually a pawn.  Could also be the net actor of consequence */
 	UPROPERTY()
 	class AActor* ViewTarget;
 
@@ -241,7 +241,7 @@ class ENGINE_API AWorldSettings : public AInfo, public IInterface_AssetUserData
 	UPROPERTY(EditAnywhere, BlueprintReadOnly, Category=World, AdvancedDisplay)
 	uint32 bEnableWorldBoundsChecks:1;
 
-	/** if set to false navigation system will not get created (and all navigation functionality won't be accesible)*/
+	/** if set to false navigation system will not get created (and all navigation functionality won't be accessible)*/
 	UPROPERTY(EditAnywhere, BlueprintReadOnly, config, Category=World, AdvancedDisplay)
 	uint32 bEnableNavigationSystem:1;
 
diff --git a/Engine/Source/Runtime/Engine/Classes/Matinee/MatineeActor.h b/Engine/Source/Runtime/Engine/Classes/Matinee/MatineeActor.h
--- a/Engine/Source/Runtime/Engine/Classes/Matinee/MatineeActor.h
+++ b/Engine/Source/Runtime/Engine/Classes/Matinee/MatineeActor.h
@@ -309,7 +309,7 @@ public:
 #endif
 
 	// Begin AActor Interface
-	virtual float GetNetPriority(const FVector& ViewPos, const FVector& ViewDir, APlayerController* Viewer, AActor* ViewTarget, UActorChannel* InChannel, float Time, bool bLowBandwidth) override;
+	virtual float GetNetPriority(const FVector& ViewPos, const FVector& ViewDir, AActor* Viewer, AActor* ViewTarget, UActorChannel* InChannel, float Time, bool bLowBandwidth) override;
 	virtual void Tick(float DeltaTime) override;
 	virtual void PreNetReceive() override;
 	virtual void PostNetReceive() override;
diff --git a/Engine/Source/Runtime/Engine/Private/Actor.cpp b/Engine/Source/Runtime/Engine/Private/Actor.cpp
--- a/Engine/Source/Runtime/Engine/Private/Actor.cpp
+++ b/Engine/Source/Runtime/Engine/Private/Actor.cpp
@@ -622,7 +622,7 @@ bool AActor::Rename( const TCHAR* InName, UObject* NewOuter, ERenameFlags Flags
 	return bSuccess;
 }
 
-UNetConnection* AActor::GetNetConnection()
+UNetConnection* AActor::GetNetConnection() const
 {
 	return Owner ? Owner->GetNetConnection() : NULL;
 }
@@ -1187,6 +1187,13 @@ AActor* AActor::GetOwner() const
 	return Owner; 
 }
 
+const AActor* AActor::GetNetOwner() const
+{
+	// NetOwner is the Actor Owner unless otherwise overridden (see PlayerController/Pawn/Beacon)
+	// Used in ServerReplicateActors
+	return Owner;
+}
+
 bool AActor::HasNetOwner() const
 {
 	if (Owner == NULL)
@@ -1479,7 +1486,7 @@ FVector AActor::GetTargetLocation(AActor* RequestedBy) const
 }
 
 
-bool AActor::IsRelevancyOwnerFor(AActor* ReplicatedActor, AActor* ActorOwner, AActor* ConnectionActor)
+bool AActor::IsRelevancyOwnerFor(const AActor* ReplicatedActor, const AActor* ActorOwner, const AActor* ConnectionActor) const
 {
 	return (ActorOwner == this);
 }
diff --git a/Engine/Source/Runtime/Engine/Private/ActorReplication.cpp b/Engine/Source/Runtime/Engine/Private/ActorReplication.cpp
--- a/Engine/Source/Runtime/Engine/Private/ActorReplication.cpp
+++ b/Engine/Source/Runtime/Engine/Private/ActorReplication.cpp
@@ -23,11 +23,11 @@ float AActor::GetNetPriority(const FVector& ViewPos, const FVector& ViewDir, APl
 	return DEPRECATED_NET_PRIORITY;
 }
 
-float AActor::GetNetPriority(const FVector& ViewPos, const FVector& ViewDir, APlayerController* Viewer, AActor* ViewTarget, UActorChannel* InChannel, float Time, bool bLowBandwidth)
+float AActor::GetNetPriority(const FVector& ViewPos, const FVector& ViewDir, AActor* Viewer, AActor* ViewTarget, UActorChannel* InChannel, float Time, bool bLowBandwidth)
 {
 	PRAGMA_DISABLE_DEPRECATION_WARNINGS
 	// Call the deprecated version so subclasses will still work for a version
-	float DeprecatedValue = AActor::GetNetPriority(ViewPos, ViewDir, Viewer, InChannel, Time, bLowBandwidth);
+	float DeprecatedValue = AActor::GetNetPriority(ViewPos, ViewDir, Cast<APlayerController>(Viewer), InChannel, Time, bLowBandwidth);
 	PRAGMA_ENABLE_DEPRECATION_WARNINGS
 	if (DeprecatedValue != DEPRECATED_NET_PRIORITY)
 	{
@@ -78,12 +78,17 @@ float AActor::GetNetPriority(const FVector& ViewPos, const FVector& ViewDir, APl
 	return NetPriority * Time;
 }
 
-bool AActor::GetNetDormancy(const FVector& ViewPos, const FVector& ViewDir, class APlayerController* Viewer, AActor* ViewTarget, UActorChannel* InChannel, float Time, bool bLowBandwidth)
+bool AActor::GetNetDormancy(const FVector& ViewPos, const FVector& ViewDir, AActor* Viewer, AActor* ViewTarget, UActorChannel* InChannel, float Time, bool bLowBandwidth)
 {
 	// For now, per peer dormancy is not supported
 	return false;
 }
 
+bool AActor::GetNetDormancy(const FVector& ViewPos, const FVector& ViewDir, APlayerController* Viewer, AActor* ViewTarget, UActorChannel* InChannel, float Time, bool bLowBandwidth)
+{
+	return GetNetDormancy(ViewPos, ViewDir, Cast<AActor>(Viewer), ViewTarget, InChannel, Time, bLowBandwidth);
+}
+
 void AActor::PreNetReceive()
 {
 	SavedbHidden = bHidden;
@@ -147,6 +152,11 @@ void AActor::PostNetReceivePhysicState()
 }
 
 bool AActor::IsNetRelevantFor(const APlayerController* RealViewer, const AActor* ViewTarget, const FVector& SrcLocation) const
+{
+	return IsNetRelevantFor(Cast<AActor>(RealViewer), ViewTarget, SrcLocation);
+}
+
+bool AActor::IsNetRelevantFor(const AActor* RealViewer, const AActor* ViewTarget, const FVector& SrcLocation) const
 {
 	if (bAlwaysRelevant || IsOwnedBy(ViewTarget) || IsOwnedBy(RealViewer) || this == ViewTarget || ViewTarget == Instigator)
 	{
diff --git a/Engine/Source/Runtime/Engine/Private/NetworkDriver.cpp b/Engine/Source/Runtime/Engine/Private/NetworkDriver.cpp
--- a/Engine/Source/Runtime/Engine/Private/NetworkDriver.cpp
+++ b/Engine/Source/Runtime/Engine/Private/NetworkDriver.cpp
@@ -1716,17 +1716,22 @@ bool FPacketSimulationSettings::ParseSettings(const TCHAR* Cmd)
 #endif
 
 FNetViewer::FNetViewer(UNetConnection* InConnection, float DeltaSeconds) :
-	InViewer(InConnection->PlayerController),
+	InViewer(InConnection->PlayerController ? InConnection->PlayerController : InConnection->OwningActor),
 	ViewTarget(InConnection->ViewTarget),
 	ViewLocation(ForceInit),
 	ViewDir(ForceInit)
 {
+	check(InConnection->OwningActor);
+	check(!InConnection->PlayerController || (InConnection->PlayerController == InConnection->OwningActor));
+
+	APlayerController* ViewingController = InConnection->PlayerController;
+
 	// Get viewer coordinates.
 	ViewLocation = ViewTarget->GetActorLocation();
-	if (InViewer)
+	if (ViewingController)
 	{
-		FRotator ViewRotation = InViewer->GetControlRotation();
-		InViewer->GetPlayerViewPoint(ViewLocation, ViewRotation);
+		FRotator ViewRotation = ViewingController->GetControlRotation();
+		ViewingController->GetPlayerViewPoint(ViewLocation, ViewRotation);
 		ViewDir = ViewRotation.Vector();
 	}
 
@@ -1871,7 +1876,8 @@ int32 UNetDriver::ServerReplicateActors(float DeltaSeconds)
 			}
 			bFoundReadyConnection = true;
 			
-			Connection->ViewTarget = Connection->PlayerController ? Connection->PlayerController->GetViewTarget() : OwningActor->GetOwner();
+			// the view target is what the player controller is looking at OR the owning actor itself when using beacons
+			Connection->ViewTarget = Connection->PlayerController ? Connection->PlayerController->GetViewTarget() : OwningActor;
 			//@todo - eliminate this mallocs if the connection isn't going to actually be updated this frame (currently needed to verify owner relevancy below)
 			Connection->OwnedConsiderList.Empty(NetRelevantActorCount);
 
@@ -2024,11 +2030,7 @@ int32 UNetDriver::ServerReplicateActors(float DeltaSeconds)
 				}
 				else
 				{
-					AActor* ActorOwner = Actor->GetOwner();
-					if ( !ActorOwner && (Cast<APlayerController>(Actor) || Cast<APawn>(Actor) ) ) 
-					{
-						ActorOwner = Actor;
-					}
+					const AActor* ActorOwner = Actor->GetNetOwner();
 					if ( ActorOwner )
 					{
 						// iterate through each connection (and child connections) looking for an owner for this actor
diff --git a/Engine/Source/Runtime/Engine/Private/Pawn.cpp b/Engine/Source/Runtime/Engine/Private/Pawn.cpp
--- a/Engine/Source/Runtime/Engine/Private/Pawn.cpp
+++ b/Engine/Source/Runtime/Engine/Private/Pawn.cpp
@@ -501,7 +501,7 @@ void APawn::UnPossessed()
 }
 
 
-class UNetConnection* APawn::GetNetConnection()
+class UNetConnection* APawn::GetNetConnection() const
 {
 	// if have a controller, it has the net connection
 	if ( Controller )
@@ -511,6 +511,10 @@ class UNetConnection* APawn::GetNetConnection()
 	return Super::GetNetConnection();
 }
 
+const AActor* APawn::GetNetOwner() const
+{
+	return this;
+}
 
 class UPlayer* APawn::GetNetOwningPlayer()
 {
@@ -1002,7 +1006,7 @@ bool APawn::IsBasedOnActor(const AActor* Other) const
 }
 
 
-bool APawn::IsNetRelevantFor(const APlayerController* RealViewer, const AActor* ViewTarget, const FVector& SrcLocation) const
+bool APawn::IsNetRelevantFor(const AActor* RealViewer, const AActor* ViewTarget, const FVector& SrcLocation) const
 {
 	if (bAlwaysRelevant || RealViewer == Controller || IsOwnedBy(ViewTarget) || IsOwnedBy(RealViewer) || this == ViewTarget || ViewTarget == Instigator
 		|| IsBasedOnActor(ViewTarget) || (ViewTarget && ViewTarget->IsBasedOnActor(this)))
diff --git a/Engine/Source/Runtime/Engine/Private/PlayerController.cpp b/Engine/Source/Runtime/Engine/Private/PlayerController.cpp
--- a/Engine/Source/Runtime/Engine/Private/PlayerController.cpp
+++ b/Engine/Source/Runtime/Engine/Private/PlayerController.cpp
@@ -83,7 +83,7 @@ APlayerController::APlayerController(const FObjectInitializer& ObjectInitializer
 	}
 }
 
-float APlayerController::GetNetPriority(const FVector& ViewPos, const FVector& ViewDir, APlayerController* Viewer, AActor* ViewTarget, UActorChannel* InChannel, float Time, bool bLowBandwidth)
+float APlayerController::GetNetPriority(const FVector& ViewPos, const FVector& ViewDir, AActor* Viewer, AActor* ViewTarget, UActorChannel* InChannel, float Time, bool bLowBandwidth)
 {
 	if ( Viewer == this )
 	{
@@ -92,6 +92,11 @@ float APlayerController::GetNetPriority(const FVector& ViewPos, const FVector& V
 	return NetPriority * Time;
 }
 
+const AActor* APlayerController::GetNetOwner() const
+{
+	return this;
+}
+
 UPlayer* APlayerController::GetNetOwningPlayer() 
 {
 	return Player;
@@ -103,7 +108,7 @@ bool APlayerController::HasNetOwner() const
 	return true;
 }
 
-UNetConnection* APlayerController::GetNetConnection()
+UNetConnection* APlayerController::GetNetConnection() const
 {
 	// A controller without a player has no "owner"
 	return (Player != NULL) ? NetConnection : NULL;
@@ -3873,7 +3878,7 @@ void APlayerController::TickActor( float DeltaSeconds, ELevelTick TickType, FAct
 #endif // !(UE_BUILD_SHIPPING || UE_BUILD_TEST)
 }
 
-bool APlayerController::IsNetRelevantFor(const APlayerController* RealViewer, const AActor* ViewTarget, const FVector& SrcLocation) const
+bool APlayerController::IsNetRelevantFor(const AActor* RealViewer, const AActor* ViewTarget, const FVector& SrcLocation) const
 {
 	return ( this==RealViewer );
 }
diff --git a/Engine/Source/Runtime/Engine/Private/Replication.cpp b/Engine/Source/Runtime/Engine/Private/Replication.cpp
--- a/Engine/Source/Runtime/Engine/Private/Replication.cpp
+++ b/Engine/Source/Runtime/Engine/Private/Replication.cpp
@@ -36,7 +36,7 @@ static bool				SavedbReversePlayback;
 static float			SavedPosition;
 static uint8			SavedReplicationForceIsPlaying;
 
-float AMatineeActor::GetNetPriority(const FVector& ViewPos, const FVector& ViewDir, APlayerController* Viewer, AActor* ViewTarget, UActorChannel* InChannel, float Time, bool bLowBandwidth)
+float AMatineeActor::GetNetPriority(const FVector& ViewPos, const FVector& ViewDir, AActor* Viewer, AActor* ViewTarget, UActorChannel* InChannel, float Time, bool bLowBandwidth)
 {
 	float Result = Super::GetNetPriority(ViewPos, ViewDir, Viewer, ViewTarget, InChannel, Time, bLowBandwidth);
 	// attempt to replicate MatineeActors approximately in the order that they were spawned to reduce ordering issues
diff --git a/Engine/Source/Runtime/GameplayAbilities/Private/Abilities/GameplayAbilityTargetActor.cpp b/Engine/Source/Runtime/GameplayAbilities/Private/Abilities/GameplayAbilityTargetActor.cpp
--- a/Engine/Source/Runtime/GameplayAbilities/Private/Abilities/GameplayAbilityTargetActor.cpp
+++ b/Engine/Source/Runtime/GameplayAbilities/Private/Abilities/GameplayAbilityTargetActor.cpp
@@ -95,7 +95,7 @@ void AGameplayAbilityTargetActor::CancelTargeting()
 	Destroy();
 }
 
-bool AGameplayAbilityTargetActor::IsNetRelevantFor(const APlayerController* RealViewer, const AActor* ViewTarget, const FVector& SrcLocation) const
+bool AGameplayAbilityTargetActor::IsNetRelevantFor(const AActor* RealViewer, const AActor* ViewTarget, const FVector& SrcLocation) const
 {
 	//The player who created the ability doesn't need to be updated about it - there should be local prediction in place.
 	if (RealViewer == MasterPC)
diff --git a/Engine/Source/Runtime/GameplayAbilities/Private/Abilities/GameplayAbilityWorldReticle.cpp b/Engine/Source/Runtime/GameplayAbilities/Private/Abilities/GameplayAbilityWorldReticle.cpp
--- a/Engine/Source/Runtime/GameplayAbilities/Private/Abilities/GameplayAbilityWorldReticle.cpp
+++ b/Engine/Source/Runtime/GameplayAbilities/Private/Abilities/GameplayAbilityWorldReticle.cpp
@@ -38,7 +38,7 @@ void AGameplayAbilityWorldReticle::InitializeReticle(AActor* InTargetingActor, A
 	OnParametersInitialized();
 }
 
-bool AGameplayAbilityWorldReticle::IsNetRelevantFor(const APlayerController* RealViewer, const AActor* ViewTarget, const FVector& SrcLocation) const
+bool AGameplayAbilityWorldReticle::IsNetRelevantFor(const AActor* RealViewer, const AActor* ViewTarget, const FVector& SrcLocation) const
 {
 	//The player who created the ability doesn't need to be updated about it - there should be local prediction in place.
 	if (RealViewer == MasterPC)
diff --git a/Engine/Source/Runtime/GameplayAbilities/Public/Abilities/GameplayAbilityTargetActor.h b/Engine/Source/Runtime/GameplayAbilities/Public/Abilities/GameplayAbilityTargetActor.h
--- a/Engine/Source/Runtime/GameplayAbilities/Public/Abilities/GameplayAbilityTargetActor.h
+++ b/Engine/Source/Runtime/GameplayAbilities/Public/Abilities/GameplayAbilityTargetActor.h
@@ -66,7 +66,7 @@ public:
 	FAbilityTargetData	TargetDataReadyDelegate;
 	FAbilityTargetData	CanceledDelegate;
 
-	virtual bool IsNetRelevantFor(const APlayerController* RealViewer, const AActor* ViewTarget, const FVector& SrcLocation) const override;
+	virtual bool IsNetRelevantFor(const AActor* RealViewer, const AActor* ViewTarget, const FVector& SrcLocation) const override;
 
 	UPROPERTY(BlueprintReadOnly, Category = "Targeting")
 	APlayerController* MasterPC;
diff --git a/Engine/Source/Runtime/GameplayAbilities/Public/Abilities/GameplayAbilityWorldReticle.h b/Engine/Source/Runtime/GameplayAbilities/Public/Abilities/GameplayAbilityWorldReticle.h
--- a/Engine/Source/Runtime/GameplayAbilities/Public/Abilities/GameplayAbilityWorldReticle.h
+++ b/Engine/Source/Runtime/GameplayAbilities/Public/Abilities/GameplayAbilityWorldReticle.h
@@ -35,7 +35,7 @@ public:
 
 	// ------------------------------
 
-	virtual bool IsNetRelevantFor(const APlayerController* RealViewer, const AActor* ViewTarget, const FVector& SrcLocation) const override;
+	virtual bool IsNetRelevantFor(const AActor* RealViewer, const AActor* ViewTarget, const FVector& SrcLocation) const override;
 
 	void InitializeReticle(AActor* InTargetingActor, APlayerController* PlayerController, FWorldReticleParameters InParameters);
 
diff --git a/Engine/Source/Runtime/Online/OnlineSubsystemUtils/Classes/OnlineBeacon.h b/Engine/Source/Runtime/Online/OnlineSubsystemUtils/Classes/OnlineBeacon.h
--- a/Engine/Source/Runtime/Online/OnlineSubsystemUtils/Classes/OnlineBeacon.h
+++ b/Engine/Source/Runtime/Online/OnlineSubsystemUtils/Classes/OnlineBeacon.h
@@ -25,7 +25,7 @@ class ONLINESUBSYSTEMUTILS_API AOnlineBeacon : public AActor, public FNetworkNot
 
 	// Begin AActor Interface
 	virtual void OnActorChannelOpen(class FInBunch& InBunch, class UNetConnection* Connection) override;
-	virtual UNetConnection* GetNetConnection() override { return nullptr; }
+	virtual UNetConnection* GetNetConnection() const override { return nullptr; }
 	// End AActor Interface
 
 	// Begin FNetworkNotify Interface
diff --git a/Engine/Source/Runtime/Online/OnlineSubsystemUtils/Classes/OnlineBeaconClient.h b/Engine/Source/Runtime/Online/OnlineSubsystemUtils/Classes/OnlineBeaconClient.h
--- a/Engine/Source/Runtime/Online/OnlineSubsystemUtils/Classes/OnlineBeaconClient.h
+++ b/Engine/Source/Runtime/Online/OnlineSubsystemUtils/Classes/OnlineBeaconClient.h
@@ -29,7 +29,7 @@ class ONLINESUBSYSTEMUTILS_API AOnlineBeaconClient : public AOnlineBeacon
 
 	// Begin AActor Interface
 	void OnNetCleanup(UNetConnection* Connection) override;
-	virtual UNetConnection* GetNetConnection() override;
+	virtual UNetConnection* GetNetConnection() const override;
 	// End AActor Interface
 
 	// Begin FNetworkNotify Interface
diff --git a/Engine/Source/Runtime/Online/OnlineSubsystemUtils/Private/OnlineBeaconClient.cpp b/Engine/Source/Runtime/Online/OnlineSubsystemUtils/Private/OnlineBeaconClient.cpp
--- a/Engine/Source/Runtime/Online/OnlineSubsystemUtils/Private/OnlineBeaconClient.cpp
+++ b/Engine/Source/Runtime/Online/OnlineSubsystemUtils/Private/OnlineBeaconClient.cpp
@@ -25,7 +25,7 @@ void AOnlineBeaconClient::SetBeaconOwner(AOnlineBeaconHostObject* InBeaconOwner)
 	BeaconOwner = InBeaconOwner;
 }
 
-UNetConnection* AOnlineBeaconClient::GetNetConnection()
+UNetConnection* AOnlineBeaconClient::GetNetConnection() const
 {
 	return BeaconConnection;
 }
