Fix CIS.
Constructor chaining is not supported yet.

diff --git a/Engine/Source/Runtime/Engine/Classes/AI/Navigation/NavigationTypes.h b/Engine/Source/Runtime/Engine/Classes/AI/Navigation/NavigationTypes.h
--- a/Engine/Source/Runtime/Engine/Classes/AI/Navigation/NavigationTypes.h
+++ b/Engine/Source/Runtime/Engine/Classes/AI/Navigation/NavigationTypes.h
@@ -484,9 +484,7 @@ struct ENGINE_API FPathFindingQuery
 	FPathFindingQuery(const FPathFindingQuery& Source);
 
 	DEPRECATED(4.8, "This version of FPathFindingQuery's constructor is deprecated. Please use ANavigationData reference rather than a pointer")
-	FPathFindingQuery(const UObject* InOwner, const ANavigationData* InNavData, const FVector& Start, const FVector& End, TSharedPtr<const FNavigationQueryFilter> SourceQueryFilter = NULL, FNavPathSharedPtr InPathInstanceToFill = NULL)
-	: FPathFindingQuery(InOwner, *InNavData, Start, End, SourceQueryFilter, InPathInstanceToFill)
-	{}
+	FPathFindingQuery(const UObject* InOwner, const ANavigationData* InNavData, const FVector& Start, const FVector& End, TSharedPtr<const FNavigationQueryFilter> SourceQueryFilter = NULL, FNavPathSharedPtr InPathInstanceToFill = NULL);
 
 	FPathFindingQuery(const UObject* InOwner, const ANavigationData& InNavData, const FVector& Start, const FVector& End, TSharedPtr<const FNavigationQueryFilter> SourceQueryFilter = NULL, FNavPathSharedPtr InPathInstanceToFill = NULL);
 
diff --git a/Engine/Source/Runtime/Engine/Private/AI/Navigation/NavigationData.cpp b/Engine/Source/Runtime/Engine/Private/AI/Navigation/NavigationData.cpp
--- a/Engine/Source/Runtime/Engine/Private/AI/Navigation/NavigationData.cpp
+++ b/Engine/Source/Runtime/Engine/Private/AI/Navigation/NavigationData.cpp
@@ -8,6 +8,22 @@
 //----------------------------------------------------------------------//
 // FPathFindingQuery
 //----------------------------------------------------------------------//
+FPathFindingQuery::FPathFindingQuery(const UObject* InOwner, const ANavigationData* InNavData, const FVector& Start, const FVector& End, TSharedPtr<const FNavigationQueryFilter> SourceQueryFilter, FNavPathSharedPtr InPathInstanceToFill)
+: NavData(InNavData)
+, Owner(InOwner)
+, StartLocation(Start)
+, EndLocation(End)
+, QueryFilter(SourceQueryFilter)
+, PathInstanceToFill(InPathInstanceToFill)
+, NavDataFlags(0)
+, bAllowPartialPaths(true)
+{
+	if (SourceQueryFilter.IsValid() == false && NavData.IsValid() == true)
+	{
+		QueryFilter = NavData->GetDefaultQueryFilter();
+	}
+}
+
 FPathFindingQuery::FPathFindingQuery(const UObject* InOwner, const ANavigationData& InNavData, const FVector& Start, const FVector& End, TSharedPtr<const FNavigationQueryFilter> SourceQueryFilter, FNavPathSharedPtr InPathInstanceToFill)
 : NavData(&InNavData)
 , Owner(InOwner)
