Only track UCS modified properties for SCS components

diff --git a/Engine/Source/Runtime/Engine/Private/Components/ActorComponent.cpp b/Engine/Source/Runtime/Engine/Private/Components/ActorComponent.cpp
--- a/Engine/Source/Runtime/Engine/Private/Components/ActorComponent.cpp
+++ b/Engine/Source/Runtime/Engine/Private/Components/ActorComponent.cpp
@@ -179,9 +179,17 @@ void UActorComponent::PostLoad()
 		}
 	}
 
-	if ((GetLinkerUE4Version() < VER_UE4_TRACK_UCS_MODIFIED_PROPERTIES) && !HasAnyFlags(RF_ClassDefaultObject))
+	if (CreationMethod == EComponentCreationMethod::SimpleConstructionScript)
 	{
-		DetermineUCSModifiedProperties();
+		if ((GetLinkerUE4Version() < VER_UE4_TRACK_UCS_MODIFIED_PROPERTIES) && !HasAnyFlags(RF_ClassDefaultObject))
+		{
+			DetermineUCSModifiedProperties();
+		}
+	}
+	else
+	{
+		// For a brief period of time we were inadvertently storing these for all components, need to clear it out
+		UCSModifiedProperties.Empty();
 	}
 }
 
@@ -1454,41 +1462,44 @@ bool UActorComponent::IsEditableWhenInherited() const
 
 void UActorComponent::DetermineUCSModifiedProperties()
 {
-	class FComponentPropertySkipper : public FArchive
-	{
-	public:
-		FComponentPropertySkipper()
-			: FArchive()
-		{
-			ArIsSaving = true;
-		}
+	UCSModifiedProperties.Empty();
 
-		virtual bool ShouldSkipProperty(const UProperty* InProperty) const override
+	if (CreationMethod == EComponentCreationMethod::SimpleConstructionScript)
+	{
+		class FComponentPropertySkipper : public FArchive
 		{
-			return (    InProperty->HasAnyPropertyFlags(CPF_Transient | CPF_ContainsInstancedReference | CPF_InstancedReference)
-					|| !InProperty->HasAnyPropertyFlags(CPF_Edit | CPF_Interp));
-		}
-	} PropertySkipper;
+		public:
+			FComponentPropertySkipper()
+				: FArchive()
+			{
+				ArIsSaving = true;
+			}
 
-	UClass* ComponentClass = GetClass();
-	UObject* ComponentArchetype = GetArchetype();
+			virtual bool ShouldSkipProperty(const UProperty* InProperty) const override
+			{
+				return (    InProperty->HasAnyPropertyFlags(CPF_Transient | CPF_ContainsInstancedReference | CPF_InstancedReference)
+						|| !InProperty->HasAnyPropertyFlags(CPF_Edit | CPF_Interp));
+			}
+		} PropertySkipper;
 
-	UCSModifiedProperties.Empty();
+		UClass* ComponentClass = GetClass();
+		UObject* ComponentArchetype = GetArchetype();
 
-	for (TFieldIterator<UProperty> It(ComponentClass); It; ++It)
-	{
-		UProperty* Property = *It;
-		if( Property->ShouldSerializeValue(PropertySkipper) )
+		for (TFieldIterator<UProperty> It(ComponentClass); It; ++It)
 		{
-			for( int32 Idx=0; Idx<Property->ArrayDim; Idx++ )
+			UProperty* Property = *It;
+			if( Property->ShouldSerializeValue(PropertySkipper) )
 			{
-				uint8* DataPtr      = Property->ContainerPtrToValuePtr           <uint8>((uint8*)this, Idx);
-				uint8* DefaultValue = Property->ContainerPtrToValuePtrForDefaults<uint8>(ComponentClass, (uint8*)ComponentArchetype, Idx);
-				if (!Property->Identical( DataPtr, DefaultValue))
+				for( int32 Idx=0; Idx<Property->ArrayDim; Idx++ )
 				{
-					UCSModifiedProperties.Add(FSimpleMemberReference());
-					FMemberReference::FillSimpleMemberReference<UProperty>(Property, UCSModifiedProperties.Last());
-					break;
+					uint8* DataPtr      = Property->ContainerPtrToValuePtr           <uint8>((uint8*)this, Idx);
+					uint8* DefaultValue = Property->ContainerPtrToValuePtrForDefaults<uint8>(ComponentClass, (uint8*)ComponentArchetype, Idx);
+					if (!Property->Identical( DataPtr, DefaultValue))
+					{
+						UCSModifiedProperties.Add(FSimpleMemberReference());
+						FMemberReference::FillSimpleMemberReference<UProperty>(Property, UCSModifiedProperties.Last());
+						break;
+					}
 				}
 			}
 		}
