Fix CIS.
Constructor chaining is not supported yet.

diff --git a/Engine/Source/Runtime/Engine/Classes/AI/Navigation/NavigationData.h b/Engine/Source/Runtime/Engine/Classes/AI/Navigation/NavigationData.h
--- a/Engine/Source/Runtime/Engine/Classes/AI/Navigation/NavigationData.h
+++ b/Engine/Source/Runtime/Engine/Classes/AI/Navigation/NavigationData.h
@@ -867,9 +867,7 @@ struct FAsyncPathFindingQuery : public FPathFindingQuery
 	{ }
 
 	DEPRECATED(4.8, "This version of FAsyncPathFindingQuery constructor is deprecated. Please use ANavigationData reference rather than a pointer version")
-	FAsyncPathFindingQuery(const UObject* InOwner, const ANavigationData* InNavData, const FVector& Start, const FVector& End, const FNavPathQueryDelegate& Delegate, TSharedPtr<const FNavigationQueryFilter> SourceQueryFilter)
-	: FAsyncPathFindingQuery(InOwner, *InNavData, Start, End, Delegate, SourceQueryFilter)
-	{}
+	FAsyncPathFindingQuery(const UObject* InOwner, const ANavigationData* InNavData, const FVector& Start, const FVector& End, const FNavPathQueryDelegate& Delegate, TSharedPtr<const FNavigationQueryFilter> SourceQueryFilter);
 
 	FAsyncPathFindingQuery(const UObject* InOwner, const ANavigationData& InNavData, const FVector& Start, const FVector& End, const FNavPathQueryDelegate& Delegate, TSharedPtr<const FNavigationQueryFilter> SourceQueryFilter);
 	FAsyncPathFindingQuery(const FPathFindingQuery& Query, const FNavPathQueryDelegate& Delegate, const EPathFindingMode::Type QueryMode);
diff --git a/Engine/Source/Runtime/Engine/Private/AI/Navigation/NavigationData.cpp b/Engine/Source/Runtime/Engine/Private/AI/Navigation/NavigationData.cpp
--- a/Engine/Source/Runtime/Engine/Private/AI/Navigation/NavigationData.cpp
+++ b/Engine/Source/Runtime/Engine/Private/AI/Navigation/NavigationData.cpp
@@ -77,6 +77,14 @@ FPathFindingQuery::FPathFindingQuery(FNavPathSharedRef PathToRecalculate, const
 //----------------------------------------------------------------------//
 uint32 FAsyncPathFindingQuery::LastPathFindingUniqueID = INVALID_NAVQUERYID;
 
+FAsyncPathFindingQuery::FAsyncPathFindingQuery(const UObject* InOwner, const ANavigationData* InNavData, const FVector& Start, const FVector& End, const FNavPathQueryDelegate& Delegate, TSharedPtr<const FNavigationQueryFilter> SourceQueryFilter)
+: FPathFindingQuery(InOwner, InNavData, Start, End, SourceQueryFilter)
+, QueryID(GetUniqueID())
+, OnDoneDelegate(Delegate)
+{
+
+}
+
 FAsyncPathFindingQuery::FAsyncPathFindingQuery(const UObject* InOwner, const ANavigationData& InNavData, const FVector& Start, const FVector& End, const FNavPathQueryDelegate& Delegate, TSharedPtr<const FNavigationQueryFilter> SourceQueryFilter)
 : FPathFindingQuery(InOwner, InNavData, Start, End, SourceQueryFilter)
 , QueryID(GetUniqueID())
