thread safe registration for navigation areas
#ue4 UE-13317

diff --git a/Engine/Source/Runtime/Engine/Classes/AI/Navigation/NavigationSystem.h b/Engine/Source/Runtime/Engine/Classes/AI/Navigation/NavigationSystem.h
--- a/Engine/Source/Runtime/Engine/Classes/AI/Navigation/NavigationSystem.h
+++ b/Engine/Source/Runtime/Engine/Classes/AI/Navigation/NavigationSystem.h
@@ -770,6 +770,7 @@ protected:
 
 	// async queries
 	FCriticalSection NavDataRegistrationSection;
+	static FCriticalSection NavAreaRegistrationSection;
 	
 	uint32 bNavigationBuildingLocked:1;
 	uint32 bInitialBuildingLockActive:1;
diff --git a/Engine/Source/Runtime/Engine/Private/AI/Navigation/NavigationSystem.cpp b/Engine/Source/Runtime/Engine/Private/AI/Navigation/NavigationSystem.cpp
--- a/Engine/Source/Runtime/Engine/Private/AI/Navigation/NavigationSystem.cpp
+++ b/Engine/Source/Runtime/Engine/Private/AI/Navigation/NavigationSystem.cpp
@@ -225,6 +225,7 @@ void FNavigationLockContext::UnlockUpdates()
 bool UNavigationSystem::bNavigationAutoUpdateEnabled = true;
 TArray<const UClass*> UNavigationSystem::NavAreaClasses;
 TArray<UClass*> UNavigationSystem::PendingNavAreaRegistration;
+FCriticalSection UNavigationSystem::NavAreaRegistrationSection;
 TSubclassOf<UNavArea> UNavigationSystem::DefaultWalkableArea = NULL;
 TSubclassOf<UNavArea> UNavigationSystem::DefaultObstacleArea = NULL;
 #if !(UE_BUILD_SHIPPING || UE_BUILD_TEST)
@@ -794,9 +795,12 @@ void UNavigationSystem::AddReferencedObjects(UObject* InThis, FReferenceCollecto
 		}
 	}
 
-	for (int32 PendingAreaIndex = 0; PendingAreaIndex < PendingNavAreaRegistration.Num(); PendingAreaIndex++)
 	{
-		Collector.AddReferencedObject(PendingNavAreaRegistration[PendingAreaIndex], InThis);
+		FScopeLock AccessLock(&NavAreaRegistrationSection);
+		for (int32 PendingAreaIndex = 0; PendingAreaIndex < PendingNavAreaRegistration.Num(); PendingAreaIndex++)
+		{
+			Collector.AddReferencedObject(PendingNavAreaRegistration[PendingAreaIndex], InThis);
+		}
 	}
 	
 	UNavigationSystem* This = CastChecked<UNavigationSystem>(InThis);
@@ -1635,6 +1639,8 @@ void UNavigationSystem::ProcessRegistrationCandidates()
 
 void UNavigationSystem::ProcessNavAreaPendingRegistration()
 {
+	FScopeLock AccessLock(&NavAreaRegistrationSection);
+
 	TArray<UClass*> TempPending = PendingNavAreaRegistration;
 
 	// Clear out old array, will fill in if still pending
@@ -1780,10 +1786,10 @@ void UNavigationSystem::UpdateCustomLink(const INavLinkCustomInterface* CustomLi
 
 void UNavigationSystem::RequestAreaUnregistering(UClass* NavAreaClass)
 {
-	check(IsInGameThread());
-
 	if (NavAreaClasses.Contains(NavAreaClass))
 	{
+		FScopeLock AccessLock(&NavAreaRegistrationSection);
+
 		// remove from known areas
 		NavAreaClasses.RemoveSingleSwap(NavAreaClass);
 		PendingNavAreaRegistration.RemoveSingleSwap(NavAreaClass);
@@ -1802,8 +1808,6 @@ void UNavigationSystem::RequestAreaUnregistering(UClass* NavAreaClass)
 
 void UNavigationSystem::RequestAreaRegistering(UClass* NavAreaClass)
 {
-	check(IsInGameThread());
-
 	// can't be null
 	if (NavAreaClass == NULL)
 	{
@@ -1833,6 +1837,7 @@ void UNavigationSystem::RequestAreaRegistering(UClass* NavAreaClass)
 		}
 	}
 
+	FScopeLock AccessLock(&NavAreaRegistrationSection);
 	PendingNavAreaRegistration.Add(NavAreaClass);
 }
 
@@ -1842,6 +1847,7 @@ void UNavigationSystem::RegisterNavAreaClass(UClass* AreaClass)
 	if (AreaClass->ClassGeneratedBy && AreaClass->ClassGeneratedBy->HasAnyFlags(RF_NeedLoad | RF_NeedPostLoad))
 	{
 		// Class isn't done loading, try again later
+		FScopeLock AccessLock(&NavAreaRegistrationSection);
 		PendingNavAreaRegistration.Add(AreaClass);
 		return;
 	}
