#ue4 - Lambda expression formatting cleanup.

diff --git a/Engine/Source/Runtime/Engine/Private/PhysicsEngine/BodyInstance.cpp b/Engine/Source/Runtime/Engine/Private/PhysicsEngine/BodyInstance.cpp
--- a/Engine/Source/Runtime/Engine/Private/PhysicsEngine/BodyInstance.cpp
+++ b/Engine/Source/Runtime/Engine/Private/PhysicsEngine/BodyInstance.cpp
@@ -2492,36 +2492,36 @@ void FBodyInstance::SetBodyTransform(const FTransform& NewTransform, bool bTelep
 
 		ExecuteOnPhysicsReadWrite([&]
 		{
-		// SIMULATED & KINEMATIC
+			// SIMULATED & KINEMATIC
 			if (PxRigidDynamic* PRigidDynamic = GetPxRigidDynamic_AssumesLocked())
-		{
-			// If kinematic and not teleporting, set kinematic target
-				if (!IsRigidBodyNonKinematic_AssumesLocked(PRigidDynamic) && !bTeleport)
 			{
-				const PxScene* PScene = PRigidDynamic->getScene();
-				FPhysScene* PhysScene = FPhysxUserData::Get<FPhysScene>(PScene->userData);
-					PhysScene->SetKinematicTarget_AssumesLocked(this, NewTransform, true);
+				// If kinematic and not teleporting, set kinematic target
+				if (!IsRigidBodyNonKinematic_AssumesLocked(PRigidDynamic) && !bTeleport)
+				{
+					const PxScene* PScene = PRigidDynamic->getScene();
+					FPhysScene* PhysScene = FPhysxUserData::Get<FPhysScene>(PScene->userData);
+						PhysScene->SetKinematicTarget_AssumesLocked(this, NewTransform, true);
+				}
+				// Otherwise, set global pose
+				else
+				{
+					PRigidDynamic->setGlobalPose(PNewPose);
+				}
 			}
-			// Otherwise, set global pose
+			// STATIC
 			else
 			{
-				PRigidDynamic->setGlobalPose(PNewPose);
-			}
-		}
-		// STATIC
-		else
-		{
-			UPrimitiveComponent* OwnerComponentInst = OwnerComponent.Get();
-			const bool bIsGame = !GIsEditor || (OwnerComponentInst != NULL && OwnerComponentInst->GetWorld()->IsGameWorld());
-			// Do NOT move static actors in-game, give a warning but let it happen
-			if (bIsGame)
-			{
-				const FString ComponentPathName = (OwnerComponentInst != NULL) ? OwnerComponentInst->GetPathName() : TEXT("NONE");
-				UE_LOG(LogPhysics, Warning, TEXT("MoveFixedBody: Trying to move component'%s' with a non-Movable Mobility."), *ComponentPathName);
+				UPrimitiveComponent* OwnerComponentInst = OwnerComponent.Get();
+				const bool bIsGame = !GIsEditor || (OwnerComponentInst != NULL && OwnerComponentInst->GetWorld()->IsGameWorld());
+				// Do NOT move static actors in-game, give a warning but let it happen
+				if (bIsGame)
+				{
+					const FString ComponentPathName = (OwnerComponentInst != NULL) ? OwnerComponentInst->GetPathName() : TEXT("NONE");
+					UE_LOG(LogPhysics, Warning, TEXT("MoveFixedBody: Trying to move component'%s' with a non-Movable Mobility."), *ComponentPathName);
+				}
+				// In EDITOR, go ahead and move it with no warning, we are editing the level
+				RigidActor->setGlobalPose(PNewPose);
 			}
-			// In EDITOR, go ahead and move it with no warning, we are editing the level
-			RigidActor->setGlobalPose(PNewPose);
-		}
 		});
 	}
 #endif  // WITH_PHYSX
