#ue4 - Make it clear that Controller possession of Pawns is supposed to only happen on the authority.

UE-10565

diff --git a/Engine/Source/Runtime/Engine/Classes/GameFramework/Controller.h b/Engine/Source/Runtime/Engine/Classes/GameFramework/Controller.h
--- a/Engine/Source/Runtime/Engine/Classes/GameFramework/Controller.h
+++ b/Engine/Source/Runtime/Engine/Classes/GameFramework/Controller.h
@@ -5,6 +5,7 @@
 #include "Controller.generated.h"
 
 class UPathFollowingComponent;
+class APawn;
 
 UDELEGATE(BlueprintAuthorityOnly)
 DECLARE_DYNAMIC_MULTICAST_DELEGATE_FourParams( FInstigatedAnyDamageSignature, float, Damage, const class UDamageType*, DamageType, class AActor*, DamagedActor, class AActor*, DamageCauser );
@@ -36,13 +37,13 @@ class ENGINE_API AController : public AActor, public INavAgentInterface
 private:
 	/** Pawn currently being controlled by this controller.  Use Pawn.Possess() to take control of a pawn */
 	UPROPERTY(replicatedUsing=OnRep_Pawn)
-	class APawn* Pawn;
+	APawn* Pawn;
 
 	/**
 	 * Used to track when pawn changes during OnRep_Pawn. 
 	 * It's possible to use a OnRep parameter here, but I'm not sure what happens with pointers yet so playing it safe.
 	 */
-	TWeakObjectPtr< class APawn > OldPawn;
+	TWeakObjectPtr< APawn > OldPawn;
 
 	/** Character currently being controlled by this controller.  Value is same as Pawn if the controlled pawn is a character, otherwise NULL */
 	UPROPERTY()
@@ -162,10 +163,10 @@ public:
 
 	/** Return the Pawn that is currently 'controlled' by this PlayerController */
 	UFUNCTION(BlueprintCallable, Category="Pawn", meta=(DisplayName="Get Controlled Pawn"))
-	class APawn* K2_GetPawn() const;
+	APawn* K2_GetPawn() const;
 
 	DEPRECATED(4.4, "Use GetPawn() instead of GetControlledPawn()")
-	class APawn* GetControlledPawn() const { return GetPawn(); }
+	APawn* GetControlledPawn() const { return GetPawn(); }
 
 public:
 
@@ -216,15 +217,17 @@ public:
 	/** Called from Destroyed().  Cleans up PlayerState. */
 	virtual void CleanupPlayerState();
 
-	/** Handles attaching this controller to the specified pawn.
-	 * @param inPawn is the pawn to be possessed
+	/**
+	 * Handles attaching this controller to the specified pawn.
+	 * Only runs on the network authority (where HasAuthority() returns true).
+	 * @param InPawn The Pawn to be possessed.
+	 * @see HasAuthority()
 	 */
-	UFUNCTION(BlueprintCallable, Category="Pawn")
-	virtual void Possess(class APawn* inPawn);
+	UFUNCTION(BlueprintCallable, BlueprintAuthorityOnly, Category="Pawn")
+	virtual void Possess(APawn* InPawn);
 
 	/**
-	 * Called to unpossess our pawn for any reason that is not the pawn being destroyed
-	 * (destruction handled by PawnDestroyed())
+	 * Called to unpossess our pawn for any reason that is not the pawn being destroyed (destruction handled by PawnDestroyed()).
 	 */
 	UFUNCTION(BlueprintCallable, Category="Pawn")
 	virtual void UnPossess();
@@ -233,7 +236,7 @@ public:
 	 * Called to unpossess our pawn because it is going to be destroyed.
 	 * (other unpossession handled by UnPossess())
 	 */
-	virtual void PawnPendingDestroy(class APawn* inPawn);
+	virtual void PawnPendingDestroy(APawn* inPawn);
 
 	DEPRECATED(4.5, "NotifyKilled is unused in engine code and will be removed from AController.")
 	virtual void NotifyKilled(AController* Killer, AController* KilledPlayer, APawn* KilledPawn, const UDamageType* DamageType) {}
diff --git a/Engine/Source/Runtime/Engine/Private/Controller.cpp b/Engine/Source/Runtime/Engine/Private/Controller.cpp
--- a/Engine/Source/Runtime/Engine/Private/Controller.cpp
+++ b/Engine/Source/Runtime/Engine/Private/Controller.cpp
@@ -12,7 +12,7 @@
 #include "GameFramework/Character.h"
 #include "GameFramework/GameMode.h"
 #include "GameFramework/PawnMovementComponent.h"
-#include "GameFramework/CharacterMovementComponent.h"
+#include "MessageLog.h"
 
 // @todo this is here only due to circular dependency to AIModule. To be removed
 #include "Navigation/PathFollowingComponent.h"
@@ -22,6 +22,8 @@
 
 DEFINE_LOG_CATEGORY(LogPath);
 
+#define LOCTEXT_NAMESPACE "Controller"
+
 
 AController::AController(const FObjectInitializer& ObjectInitializer)
 	: Super(ObjectInitializer)
@@ -214,6 +216,15 @@ void AController::PostInitializeComponents()
 
 void AController::Possess(APawn* InPawn)
 {
+	if (!HasAuthority())
+	{
+		FMessageLog("PIE").Warning(FText::Format(
+			LOCTEXT("ControllerPossessAuthorityOnly", "Possess function should only be used by the network authority for {0}"),
+			FText::FromName(GetFName())
+			));
+		return;
+	}
+
 	REDIRECT_OBJECT_TO_VLOG(InPawn, this);
 
 	if (InPawn != NULL)
@@ -592,3 +603,5 @@ void AController::GetLifetimeReplicatedProps( TArray< FLifetimeProperty > & OutL
 
 /** Returns TransformComponent subobject **/
 USceneComponent* AController::GetTransformComponent() const { return TransformComponent; }
+
+#undef LOCTEXT_NAMESPACE
diff --git a/Engine/Source/Runtime/Engine/Private/PlayerController.cpp b/Engine/Source/Runtime/Engine/Private/PlayerController.cpp
--- a/Engine/Source/Runtime/Engine/Private/PlayerController.cpp
+++ b/Engine/Source/Runtime/Engine/Private/PlayerController.cpp
@@ -6,9 +6,7 @@
 #include "Engine/WorldComposition.h"
 #include "Engine/LevelScriptActor.h"
 #include "GameFramework/GameNetworkManager.h"
-#include "GameFramework/Character.h"
 #include "Interfaces/NetworkPredictionInterface.h"
-#include "GameFramework/CharacterMovementComponent.h"
 #include "ConfigCacheIni.h"
 #include "SoundDefinitions.h"
 #include "OnlineSubsystemUtils.h"
@@ -39,9 +37,12 @@
 #include "Engine/GameEngine.h"
 #include "Engine/GameInstance.h"
 #include "VisualLogger/VisualLogger.h"
+#include "MessageLog.h"
 
 DEFINE_LOG_CATEGORY(LogPlayerController);
 
+#define LOCTEXT_NAMESPACE "PlayerController"
+
 const float RetryClientRestartThrottleTime = 0.5f;
 const float RetryServerAcknowledgeThrottleTime = 0.25f;
 const float RetryServerCheckSpectatorThrottleTime = 0.25f;
@@ -747,6 +748,15 @@ void APlayerController::ClientRestart_Implementation(APawn* NewPawn)
 
 void APlayerController::Possess(APawn* PawnToPossess)
 {
+	if (!HasAuthority())
+	{
+		FMessageLog("PIE").Warning(FText::Format(
+			LOCTEXT("PlayerControllerPossessAuthorityOnly", "Possess function should only be used by the network authority for {0}"),
+			FText::FromName(GetFName())
+			));
+		return;
+	}
+
 	if ( PawnToPossess != NULL && 
 		(PlayerState == NULL || !PlayerState->bOnlySpectator) )
 	{
@@ -4542,3 +4552,5 @@ void APlayerController::OnServerStartedVisualLogger_Implementation(bool bIsLoggi
 	ClientMessage(FString::Printf(TEXT("Visual Loggger is %s."), FVisualLogger::Get().IsRecordingOnServer() ? TEXT("now recording") : TEXT("disabled")));
 #endif
 }
+
+#undef LOCTEXT_NAMESPACE
