Fix comment on RemoveAxisMapping

diff --git a/Engine/Source/Runtime/Engine/Classes/GameFramework/PlayerInput.h b/Engine/Source/Runtime/Engine/Classes/GameFramework/PlayerInput.h
--- a/Engine/Source/Runtime/Engine/Classes/GameFramework/PlayerInput.h
+++ b/Engine/Source/Runtime/Engine/Classes/GameFramework/PlayerInput.h
@@ -329,7 +329,7 @@ public:
 	/** Add a player specific axis mapping. */
 	void AddAxisMapping(const FInputAxisKeyMapping& KeyMapping);
 
-	/** Remove a player specific action mapping. */
+	/** Remove a player specific axis mapping. */
 	void RemoveAxisMapping(const FInputAxisKeyMapping& KeyMapping);
 
 	/** Add an engine defined action mapping that cannot be remapped. */
