#ue4 - Perf: Avoid Quat->Rotator conversion in USceneComponent::SetRelativeLocationAndRotation() when possible. Avoid Rotator->Quat conversion in USceneComponent::AddLocalOffset().

- SetRelativeLocationAndRotation the case that there is no change between relative->world rotation. This is the  common case for root components.
- AddLocalOffset() uses Rotator.RotateVector() instead of converting to Quaternion first.

diff --git a/Engine/Source/Runtime/Engine/Private/Components/SceneComponent.cpp b/Engine/Source/Runtime/Engine/Private/Components/SceneComponent.cpp
--- a/Engine/Source/Runtime/Engine/Private/Components/SceneComponent.cpp
+++ b/Engine/Source/Runtime/Engine/Private/Components/SceneComponent.cpp
@@ -479,8 +479,11 @@ void USceneComponent::SetRelativeLocationAndRotation(FVector NewLocation, FRotat
 		const FTransform DesiredRelTransform(NewRotation, NewLocation);
 		const FTransform DesiredWorldTransform = CalcNewComponentToWorld(DesiredRelTransform);
 		const FVector DesiredDelta = DesiredWorldTransform.GetTranslation() - ComponentToWorld.GetTranslation();
-		
-		MoveComponent(DesiredDelta, DesiredWorldTransform.Rotator(), bSweep, OutSweepHitResult);
+
+		// Avoid Quat->Rotator conversion in the case that there is no change in relative->world rotation, common case for root components.
+		const FRotator DesiredRotation = (DesiredRelTransform.GetRotation() == DesiredWorldTransform.GetRotation()) ? NewRotation : DesiredWorldTransform.Rotator();
+
+		MoveComponent(DesiredDelta, DesiredRotation, bSweep, OutSweepHitResult);
 	}
 	else if (OutSweepHitResult)
 	{
@@ -502,8 +505,7 @@ void USceneComponent::AddLocalOffset(FVector DeltaLocation, bool bSweep, FHitRes
 {
 	if (!DeltaLocation.IsNearlyZero())
 	{
-		const FQuat RelativeRotQuat = RelativeRotation.Quaternion();
-		const FVector LocalOffset = RelativeRotQuat.RotateVector(DeltaLocation);
+		const FVector LocalOffset = RelativeRotation.RotateVector(DeltaLocation);
 
 		SetRelativeLocationAndRotation(RelativeLocation + LocalOffset, RelativeRotation, bSweep, OutSweepHitResult);
 	}
