Allow Actor Spawn Parameters to specify a Pawn's instigator

diff --git a/Engine/Source/Runtime/Engine/Private/Pawn.cpp b/Engine/Source/Runtime/Engine/Private/Pawn.cpp
--- a/Engine/Source/Runtime/Engine/Private/Pawn.cpp
+++ b/Engine/Source/Runtime/Engine/Private/Pawn.cpp
@@ -69,7 +69,10 @@ void APawn::PreInitializeComponents()
 {
 	Super::PreInitializeComponents();
 
-	Instigator = this;
+	if (Instigator == nullptr)
+	{
+		Instigator = this;
+	}
 
 	if (AutoPossessPlayer != EAutoReceiveInput::Disabled && GetNetMode() != NM_Client )
 	{
