#ue4 - Fix ComponentSweepMulti() using incorrect transform when shape is not at the local origin.

Jira: UE-7030

diff --git a/Engine/Source/Runtime/Engine/Private/Collision/WorldCollision.cpp b/Engine/Source/Runtime/Engine/Private/Collision/WorldCollision.cpp
--- a/Engine/Source/Runtime/Engine/Private/Collision/WorldCollision.cpp
+++ b/Engine/Source/Runtime/Engine/Private/Collision/WorldCollision.cpp
@@ -513,10 +513,9 @@ bool UWorld::ComponentSweepMulti(TArray<struct FHitResult>& OutHits, class UPrim
 		}
 
 		// calculate the test global pose of the actor
-		PxTransform PGlobalStartPose = U2PTransform(FTransform(Start));
-		PxTransform PGlobalEndPose = U2PTransform(FTransform(End));
-
-		PxQuat PGeomRot = U2PQuat(Rot.Quaternion());
+		const PxQuat PGeomRot = U2PQuat(Rot.Quaternion());
+		const PxTransform PGlobalStartPose = PxTransform(U2PVector(Start), PGeomRot);
+		const PxTransform PGlobalEndPose = PxTransform(U2PVector(End), PGeomRot);
 
 		// Iterate over each shape
 		for(int32 ShapeIdx=0; ShapeIdx<PShapes.Num(); ShapeIdx++)
