After applying per instance blueprint component saved properties, re-register the component to apply the changes

diff --git a/Engine/Source/Runtime/Engine/Private/ComponentInstanceDataCache.cpp b/Engine/Source/Runtime/Engine/Private/ComponentInstanceDataCache.cpp
--- a/Engine/Source/Runtime/Engine/Private/ComponentInstanceDataCache.cpp
+++ b/Engine/Source/Runtime/Engine/Private/ComponentInstanceDataCache.cpp
@@ -122,7 +122,12 @@ void FActorComponentInstanceData::ApplyToComponent(UActorComponent* Component, c
 
 			TSet<const UProperty*> PropertiesToSkip;
 
-		} ComponentPropertyReader(Component, SavedProperties);	
+		} ComponentPropertyReader(Component, SavedProperties);
+
+		if (Component->IsRegistered())
+		{
+			Component->ReregisterComponent();
+		}
 	}
 }
 
