When registering a Child Actor Component if there is no Child Actor associated or it is of the wrong class, (re)create one.
Fixes overriding child actor class on a per-instance basis

diff --git a/Engine/Source/Runtime/Engine/Private/Components/ChildActorComponent.cpp b/Engine/Source/Runtime/Engine/Private/Components/ChildActorComponent.cpp
--- a/Engine/Source/Runtime/Engine/Private/Components/ChildActorComponent.cpp
+++ b/Engine/Source/Runtime/Engine/Private/Components/ChildActorComponent.cpp
@@ -17,12 +17,24 @@ void UChildActorComponent::OnRegister()
 
 	if (ChildActor)
 	{
-		ChildActorName = ChildActor->GetFName();
-		// attach new actor to this component
-		// we can't attach in CreateChildActor since it has intermediate Mobility set up
-		// causing spam with inconsistent mobility set up
-		// so moving Attach to happen in Register
-		ChildActor->AttachRootComponentTo(this, NAME_None, EAttachLocation::SnapToTarget);
+		if (ChildActor->GetClass() != ChildActorClass)
+		{
+			DestroyChildActor();
+			CreateChildActor();
+		}
+		else
+		{
+			ChildActorName = ChildActor->GetFName();
+			// attach new actor to this component
+			// we can't attach in CreateChildActor since it has intermediate Mobility set up
+			// causing spam with inconsistent mobility set up
+			// so moving Attach to happen in Register
+			ChildActor->AttachRootComponentTo(this, NAME_None, EAttachLocation::SnapToTarget);
+		}
+	}
+	else if (ChildActorClass)
+	{
+		CreateChildActor();
 	}
 }
 
