Reducing risk of crash when setting a MemberReference with self scope and providing a nullptr.

Put in an ensure to catch any time a nullptr is passed.

diff --git a/Engine/Source/Runtime/Engine/Private/MemberReference.cpp b/Engine/Source/Runtime/Engine/Private/MemberReference.cpp
--- a/Engine/Source/Runtime/Engine/Private/MemberReference.cpp
+++ b/Engine/Source/Runtime/Engine/Private/MemberReference.cpp
@@ -54,7 +54,10 @@ void FMemberReference::SetGivenSelfScope(const FName InMemberName, const FGuid I
 	MemberGuid = InMemberGuid;
 	MemberParentClass = (InMemberParentClass != nullptr) ? InMemberParentClass->GetAuthoritativeClass() : nullptr;
 	MemberScope.Empty();
-	bSelfContext = (SelfScope->IsChildOf(InMemberParentClass)) || (SelfScope->ClassGeneratedBy == InMemberParentClass->ClassGeneratedBy);
+
+	// SelfScope should always be valid, but if it's not ensure and move on, the node will be treated as if it's not self scoped.
+	ensure(SelfScope);
+	bSelfContext = SelfScope && ((SelfScope->IsChildOf(InMemberParentClass)) || (SelfScope->ClassGeneratedBy == InMemberParentClass->ClassGeneratedBy));
 	bWasDeprecated = false;
 
 	if (bSelfContext)
