UE-13123 SetGamePaused on EventDestroyed crashes the editor on compile

diff --git a/Engine/Source/Runtime/Engine/Private/GameplayStatics.cpp b/Engine/Source/Runtime/Engine/Private/GameplayStatics.cpp
--- a/Engine/Source/Runtime/Engine/Private/GameplayStatics.cpp
+++ b/Engine/Source/Runtime/Engine/Private/GameplayStatics.cpp
@@ -136,17 +136,8 @@ void UGameplayStatics::SetGlobalTimeDilation(UObject* WorldContextObject, float
 bool UGameplayStatics::SetGamePaused(UObject* WorldContextObject, bool bPaused)
 {
 	UWorld* const World = GEngine->GetWorldFromContextObject( WorldContextObject );
-	if(World != nullptr)
-	{
-		APlayerController* const PC = World->GetFirstPlayerController();
-		check(PC); // Gathering some information for TTP #303973
-
-		return PC->SetPause(bPaused);
-	}
-	else
-	{
-		return false;
-	}
+	APlayerController* const PC = World ? World->GetFirstPlayerController() : nullptr;
+	return PC ? PC->SetPause(bPaused) : false;
 }
 
 bool UGameplayStatics::IsGamePaused(UObject* WorldContextObject)
