Possible crash resolution if the world is eliminated by seamless travel handler

diff --git a/Engine/Source/Runtime/Engine/Classes/Engine/Engine.h b/Engine/Source/Runtime/Engine/Classes/Engine/Engine.h
--- a/Engine/Source/Runtime/Engine/Classes/Engine/Engine.h
+++ b/Engine/Source/Runtime/Engine/Classes/Engine/Engine.h
@@ -2416,7 +2416,7 @@ public:
 	/** Browse to a specified URL, relative to the current one. */
 	virtual EBrowseReturnVal::Type Browse( FWorldContext& WorldContext, FURL URL, FString& Error );
 
-	bool TickWorldTravel(FWorldContext& WorldContext, float DeltaSeconds);
+	void TickWorldTravel(FWorldContext& WorldContext, float DeltaSeconds);
 
 	void BrowseToDefaultMap( FWorldContext& WorldContext );
 
diff --git a/Engine/Source/Runtime/Engine/Private/UnrealEngine.cpp b/Engine/Source/Runtime/Engine/Private/UnrealEngine.cpp
--- a/Engine/Source/Runtime/Engine/Private/UnrealEngine.cpp
+++ b/Engine/Source/Runtime/Engine/Private/UnrealEngine.cpp
@@ -8470,7 +8470,7 @@ void UEngine::BrowseToDefaultMap( FWorldContext& Context )
 	}
 }
 
-bool UEngine::TickWorldTravel(FWorldContext& Context, float DeltaSeconds)
+void UEngine::TickWorldTravel(FWorldContext& Context, float DeltaSeconds)
 {
 	// Handle seamless traveling
 	if (Context.SeamlessTravelHandler.IsInTransition())
@@ -8480,6 +8480,14 @@ bool UEngine::TickWorldTravel(FWorldContext& Context, float DeltaSeconds)
 	}
 
 	// Handle server traveling.
+	if (Context.World() == nullptr)
+	{
+		UE_LOG(LogLoad, Error, TEXT("UEngine::TickWorldTravel has no world after ticking seamless travel handler."));
+		BrowseToDefaultMap(Context);
+		BroadcastTravelFailure(Context.World(), ETravelFailure::ServerTravelFailure, TEXT("UEngine::TickWorldTravel has no world after ticking seamless travel handler."));
+		return;
+	}
+	
 	if( !Context.World()->NextURL.IsEmpty() )
 	{
 		Context.World()->NextSwitchCountdown -= DeltaSeconds;
@@ -8513,7 +8521,7 @@ bool UEngine::TickWorldTravel(FWorldContext& Context, float DeltaSeconds)
 				// Let people know that we failed to server travel
 				BroadcastTravelFailure(Context.World(), ETravelFailure::ServerTravelFailure, Error);
 			}
-			return false;
+			return;
 		}
 	}
 
@@ -8540,7 +8548,7 @@ bool UEngine::TickWorldTravel(FWorldContext& Context, float DeltaSeconds)
 			BroadcastTravelFailure(Context.World(), ETravelFailure::ClientTravelFailure, Error);
 		}
 		check(Context.World() != NULL);
-		return false;
+		return;
 	}
 
 	// Update the pending level.
@@ -8595,7 +8603,7 @@ bool UEngine::TickWorldTravel(FWorldContext& Context, float DeltaSeconds)
 		TransitionType = TT_None;
 	}
 
-	return true;
+	return;
 }
 
 /**
