Show new scrubbing position immediately

diff --git a/Engine/Source/Runtime/Engine/Classes/Engine/DemoNetDriver.h b/Engine/Source/Runtime/Engine/Classes/Engine/DemoNetDriver.h
--- a/Engine/Source/Runtime/Engine/Classes/Engine/DemoNetDriver.h
+++ b/Engine/Source/Runtime/Engine/Classes/Engine/DemoNetDriver.h
@@ -27,6 +27,9 @@ class ENGINE_API UDemoNetDriver : public UNetDriver
 	/** Current record/playback position in seconds */
 	float DemoCurrentTime;
 
+	/** Old current record/playback position in seconds (so we can restore on checkpoint failure) */
+	float OldDemoCurrentTime;
+
 	/** Total number of frames in the demo */
 	int32 DemoTotalFrames;
 
diff --git a/Engine/Source/Runtime/Engine/Private/DemoNetDriver.cpp b/Engine/Source/Runtime/Engine/Private/DemoNetDriver.cpp
--- a/Engine/Source/Runtime/Engine/Private/DemoNetDriver.cpp
+++ b/Engine/Source/Runtime/Engine/Private/DemoNetDriver.cpp
@@ -1113,7 +1113,10 @@ void UDemoNetDriver::SkipTime(const float InTimeToSkip)
 
 void UDemoNetDriver::GotoTimeInSeconds( const float TimeInSeconds )
 {
-	LastGotoTimeInSeconds = TimeInSeconds;
+	if ( LastGotoTimeInSeconds < 0 )
+	{
+		LastGotoTimeInSeconds = TimeInSeconds;
+	}
 }
 
 void UDemoNetDriver::TickDemoPlayback( float DeltaSeconds )
@@ -1144,6 +1147,8 @@ void UDemoNetDriver::TickDemoPlayback( float DeltaSeconds )
 		if ( LastGotoTimeInSeconds >= 0.0f )
 		{
 			ReplayStreamer->GotoTimeInMS( LastGotoTimeInSeconds * 1000, FOnCheckpointReadyDelegate::CreateUObject( this, &UDemoNetDriver::CheckpointReady ) );
+			OldDemoCurrentTime = DemoCurrentTime;		// So we can restore on failure
+			DemoCurrentTime = LastGotoTimeInSeconds;
 			LastGotoTimeInSeconds = -1.0f;
 		}
 	}
@@ -1305,6 +1310,7 @@ void UDemoNetDriver::CheckpointReady( const bool bSuccess, const int64 SkipExtra
 	if ( !bSuccess )
 	{
 		UE_LOG( LogDemo, Warning, TEXT( "UDemoNetConnection::CheckpointReady: Failed to go to checkpoint." ) );
+		DemoCurrentTime = OldDemoCurrentTime;
 		return;
 	}
 
