#ue4 - Check for valid shape before calculating expensive transforms for that shape in ComponentSweepMulti().

diff --git a/Engine/Source/Runtime/Engine/Private/Collision/WorldCollision.cpp b/Engine/Source/Runtime/Engine/Private/Collision/WorldCollision.cpp
--- a/Engine/Source/Runtime/Engine/Private/Collision/WorldCollision.cpp
+++ b/Engine/Source/Runtime/Engine/Private/Collision/WorldCollision.cpp
@@ -523,19 +523,19 @@ bool UWorld::ComponentSweepMulti(TArray<struct FHitResult>& OutHits, class UPrim
 			PxShape* PShape = PShapes[ShapeIdx];
 			check(PShape);
 
-			TArray<struct FHitResult> Hits;
-
-			// Calc shape global pose
-			PxTransform PLocalShape = PShape->getLocalPose();
-			PxTransform PShapeGlobalStartPose = PGlobalStartPose.transform(PLocalShape);
-			PxTransform PShapeGlobalEndPose = PGlobalEndPose.transform(PLocalShape);
-			// consider localshape rotation for shape rotation
-			PxQuat PShapeRot = PGeomRot * PLocalShape.q;
-
 			GET_GEOMETRY_FROM_SHAPE(PGeom, PShape);
 
-			if(PGeom != NULL)
+			if (PGeom != NULL)
 			{
+				TArray<struct FHitResult> Hits;
+
+				// Calc shape global pose
+				PxTransform PLocalShape = PShape->getLocalPose();
+				PxTransform PShapeGlobalStartPose = PGlobalStartPose.transform(PLocalShape);
+				PxTransform PShapeGlobalEndPose = PGlobalEndPose.transform(PLocalShape);
+				// consider localshape rotation for shape rotation
+				PxQuat PShapeRot = PGeomRot * PLocalShape.q;
+
 				if (GeomSweepMulti_PhysX(this, *PGeom, PShapeRot, Hits, P2UVector(PShapeGlobalStartPose.p), P2UVector(PShapeGlobalEndPose.p), TraceChannel, Params, FCollisionResponseParams(PrimComp->GetCollisionResponseToChannels())))
 				{
 					bHaveBlockingHit = true;
