When starting live replays, jump to as close as live as possible

diff --git a/Engine/Source/Runtime/Engine/Private/DemoNetDriver.cpp b/Engine/Source/Runtime/Engine/Private/DemoNetDriver.cpp
--- a/Engine/Source/Runtime/Engine/Private/DemoNetDriver.cpp
+++ b/Engine/Source/Runtime/Engine/Private/DemoNetDriver.cpp
@@ -1299,6 +1299,21 @@ void UDemoNetDriver::ReplayStreamingReady( bool bSuccess, bool bRecord )
 	{
 		FString Error;
 		InitConnectInternal( Error );
+
+		const uint32 TotalDemoTimeInMS = ReplayStreamer->GetTotalDemoTime();
+
+		if ( TotalDemoTimeInMS > 0 )
+		{
+			DemoTotalTime = (float)TotalDemoTimeInMS / 1000.0f;
+		}
+
+		if ( ReplayStreamer->IsLive() && ReplayStreamer->GetTotalDemoTime() > 15 * 1000 )
+		{
+			ReplayStreamer->GotoTimeInMS( ReplayStreamer->GetTotalDemoTime() - 10 * 1000, FOnCheckpointReadyDelegate::CreateUObject( this, &UDemoNetDriver::CheckpointReady ) );
+
+			OldDemoCurrentTime = DemoCurrentTime;		// So we can restore on failure
+			DemoCurrentTime = (float)( ReplayStreamer->GetTotalDemoTime() - 10 * 1000 ) / 1000.0f;
+		}
 	}
 }
 
@@ -1328,32 +1343,23 @@ void UDemoNetDriver::LoadCheckpoint()
 		return;
 	}
 
-	if ( SpectatorController == NULL )
-	{
-		UE_LOG( LogDemo, Warning, TEXT( "UDemoNetConnection::LoadCheckpoint: No spectator player controller." ) );
-		return;
-	}
-
-	if ( SpectatorController->GetSpectatorPawn() == NULL )
-	{
-		UE_LOG( LogDemo, Warning, TEXT( "UDemoNetConnection::LoadCheckpoint: No spectator pawn." ) );
-		return;
-	}
-
 	bIsLoadingCheckpoint = true;
 
-	bRestoreSpectatorPosition = true;
-	SpectatorLocation = SpectatorController->GetSpectatorPawn()->GetActorLocation();
-	SpectatorRotation = SpectatorController->GetControlRotation();//GetSpectatorPawn()->GetActorRotation();
-
-	FURL ConnectURL;
-	ConnectURL.Map = DemoFilename;
-
 	// Reset the never-queue GUID list, we'll rebuild it
 	NonQueuedGUIDsForScrubbing.Empty();
 
-	// Save off the SpectatorController's GUID so that we know not to queue his bunches
-	AddNonQueuedActorForScrubbing(SpectatorController);
+	// Save off the current spectator position
+	// Check for NULL, which can be the case if we haven't played any of the demo yet but want to fast forward (joining live game for example)
+	if ( SpectatorController != NULL && SpectatorController->GetSpectatorPawn() != NULL )
+	{
+		bRestoreSpectatorPosition = true;
+
+		SpectatorLocation = SpectatorController->GetSpectatorPawn()->GetActorLocation();
+		SpectatorRotation = SpectatorController->GetControlRotation();//GetSpectatorPawn()->GetActorRotation();
+
+		// Save off the SpectatorController's GUID so that we know not to queue his bunches
+		AddNonQueuedActorForScrubbing(SpectatorController);
+	}
 
 #if 1
 	// Destroy all non startup actors. They will get restored with the checkpoint
@@ -1362,6 +1368,8 @@ void UDemoNetDriver::LoadCheckpoint()
 		// If there are any existing actors that are bAlwaysRelevant, don't queue their bunches.
 		// Actors that do queue their bunches might not appear immediately after the checkpoint is loaded,
 		// and missing bAlwaysRelevant actors are more likely to cause noticeable artifacts.
+		// NOTE - We are adding the actor guid here, under the assumption that the actor will reclaim the same guid when we load the checkpoint
+		// This is normally the case, but could break if actors get destroyed and re-created with different guids during recording
 		if ( It->bAlwaysRelevant )
 		{
 			AddNonQueuedActorForScrubbing(*It);
@@ -1392,6 +1400,9 @@ void UDemoNetDriver::LoadCheckpoint()
 	ServerConnection->Close();
 	ServerConnection->CleanUp();
 
+	FURL ConnectURL;
+	ConnectURL.Map = DemoFilename;
+
 	ServerConnection = NewObject<UNetConnection>(GetTransientPackage(), UDemoNetConnection::StaticClass());
 	ServerConnection->InitConnection( this, USOCK_Pending, ConnectURL, 1000000 );
 
diff --git a/Engine/Source/Runtime/NetworkReplayStreaming/HttpNetworkReplayStreaming/Private/HttpNetworkReplayStreaming.cpp b/Engine/Source/Runtime/NetworkReplayStreaming/HttpNetworkReplayStreaming/Private/HttpNetworkReplayStreaming.cpp
--- a/Engine/Source/Runtime/NetworkReplayStreaming/HttpNetworkReplayStreaming/Private/HttpNetworkReplayStreaming.cpp
+++ b/Engine/Source/Runtime/NetworkReplayStreaming/HttpNetworkReplayStreaming/Private/HttpNetworkReplayStreaming.cpp
@@ -702,7 +702,7 @@ bool FHttpNetworkReplayStreamer::IsDataAvailableForTimeRange( const uint32 Start
 	return ( StartTimeInMS >= StreamTimeRangeStart && EndTimeInMS <= StreamTimeRangeEnd );
 }
 
-bool FHttpNetworkReplayStreamer::IsLive(const FString& StreamName) const 
+bool FHttpNetworkReplayStreamer::IsLive() const 
 {
 	return bStreamIsLive;
 }
@@ -924,8 +924,6 @@ void FHttpNetworkReplayStreamer::HttpDownloadHeaderFinished( FHttpRequestPtr Htt
 		HeaderArchive.Buffer.Append( HttpResponse->GetContent() );
 
 		UE_LOG( LogHttpReplay, Log, TEXT( "FHttpNetworkReplayStreamer::HttpDownloadHeaderFinished. Size: %i" ), HeaderArchive.Buffer.Num() );
-
-		StartStreamingDelegate.ExecuteIfBound( true, false );
 	}
 	else
 	{
@@ -935,11 +933,11 @@ void FHttpNetworkReplayStreamer::HttpDownloadHeaderFinished( FHttpRequestPtr Htt
 		StreamArchive.Buffer.Empty();
 		StartStreamingDelegate.ExecuteIfBound( false, false );
 
+		// Reset delegate
+		StartStreamingDelegate = FOnStreamReadyDelegate();
+
 		SetLastError( ENetworkReplayError::ServiceUnavailable );
 	}
-
-	// Reset delegate
-	StartStreamingDelegate = FOnStreamReadyDelegate();
 }
 
 void FHttpNetworkReplayStreamer::HttpDownloadFinished( FHttpRequestPtr HttpRequest, FHttpResponsePtr HttpResponse, bool bSucceeded )
@@ -1126,15 +1124,32 @@ void FHttpNetworkReplayStreamer::HttpEnumerateCheckpointsFinished( FHttpRequestP
 		if ( !CheckpointList.FromJson( JsonString ) )
 		{
 			UE_LOG( LogHttpReplay, Warning, TEXT( "FHttpNetworkReplayStreamer::HttpEnumerateCheckpointsFinished. FromJson FAILED" ) );
+
+			StartStreamingDelegate.ExecuteIfBound( false, false );
+
+			// Reset delegate
+			StartStreamingDelegate = FOnStreamReadyDelegate();
+
 			SetLastError( ENetworkReplayError::ServiceUnavailable );
+
 			return;
 		}
+
+		StartStreamingDelegate.ExecuteIfBound( true, false );
 	}
 	else
 	{
 		UE_LOG( LogHttpReplay, Warning, TEXT( "FHttpNetworkReplayStreamer::HttpEnumerateCheckpointsFinished. FAILED" ) );
+		StartStreamingDelegate.ExecuteIfBound( false, false );
+
+		// Reset delegate
+		StartStreamingDelegate = FOnStreamReadyDelegate();
+
 		SetLastError( ENetworkReplayError::ServiceUnavailable );
 	}
+
+	// Reset delegate
+	StartStreamingDelegate = FOnStreamReadyDelegate();
 }
 
 bool FHttpNetworkReplayStreamer::ProcessNextHttpRequest()
diff --git a/Engine/Source/Runtime/NetworkReplayStreaming/HttpNetworkReplayStreaming/Private/HttpNetworkReplayStreaming.h b/Engine/Source/Runtime/NetworkReplayStreaming/HttpNetworkReplayStreaming/Private/HttpNetworkReplayStreaming.h
--- a/Engine/Source/Runtime/NetworkReplayStreaming/HttpNetworkReplayStreaming/Private/HttpNetworkReplayStreaming.h
+++ b/Engine/Source/Runtime/NetworkReplayStreaming/HttpNetworkReplayStreaming/Private/HttpNetworkReplayStreaming.h
@@ -110,7 +110,7 @@ public:
 	virtual bool		IsDataAvailable() const override;
 	virtual void		SetHighPriorityTimeRange( const uint32 StartTimeInMS, const uint32 EndTimeInMS ) override;
 	virtual bool		IsDataAvailableForTimeRange( const uint32 StartTimeInMS, const uint32 EndTimeInMS ) override;
-	virtual bool		IsLive( const FString& StreamName ) const override;
+	virtual bool		IsLive() const override;
 	virtual void		DeleteFinishedStream( const FString& StreamName, const FOnDeleteFinishedStreamComplete& Delegate ) const override;
 	virtual void		EnumerateStreams( const FNetworkReplayVersion& ReplayVersion, const FOnEnumerateStreamsComplete& Delegate ) override;
 	virtual ENetworkReplayError::Type GetLastError() const override;
diff --git a/Engine/Source/Runtime/NetworkReplayStreaming/NetworkReplayStreaming/Public/NetworkReplayStreaming.h b/Engine/Source/Runtime/NetworkReplayStreaming/NetworkReplayStreaming/Public/NetworkReplayStreaming.h
--- a/Engine/Source/Runtime/NetworkReplayStreaming/NetworkReplayStreaming/Public/NetworkReplayStreaming.h
+++ b/Engine/Source/Runtime/NetworkReplayStreaming/NetworkReplayStreaming/Public/NetworkReplayStreaming.h
@@ -118,8 +118,8 @@ public:
 	virtual void SetHighPriorityTimeRange( const uint32 StartTimeInMS, const uint32 EndTimeInMS ) = 0;
 	virtual bool IsDataAvailableForTimeRange( const uint32 StartTimeInMS, const uint32 EndTimeInMS ) = 0;
 	
-	/** Returns true if the given StreamName is a game currently in progress */
-	virtual bool IsLive( const FString& StreamName ) const = 0;
+	/** Returns true if the playing stream is currently in progress */
+	virtual bool IsLive() const = 0;
 
 	/**
 	 * Attempts to delete the stream with the specified name. May execute asynchronously.
diff --git a/Engine/Source/Runtime/NetworkReplayStreaming/NullNetworkReplayStreaming/Private/NullNetworkReplayStreaming.cpp b/Engine/Source/Runtime/NetworkReplayStreaming/NullNetworkReplayStreaming/Private/NullNetworkReplayStreaming.cpp
--- a/Engine/Source/Runtime/NetworkReplayStreaming/NullNetworkReplayStreaming/Private/NullNetworkReplayStreaming.cpp
+++ b/Engine/Source/Runtime/NetworkReplayStreaming/NullNetworkReplayStreaming/Private/NullNetworkReplayStreaming.cpp
@@ -138,7 +138,12 @@ FArchive* FNullNetworkReplayStreamer::GetMetadataArchive()
 	return MetadataFileAr.Get();
 }
 
-bool FNullNetworkReplayStreamer::IsLive( const FString& StreamName ) const
+bool FNullNetworkReplayStreamer::IsLive() const
+{
+	return IsNamedStreamLive(CurrentStreamName);
+}
+
+bool FNullNetworkReplayStreamer::IsNamedStreamLive( const FString& StreamName ) const
 {
 	// If the directory for this stream doesn't exist, it can't possibly be live.
 	if (!IFileManager::Get().DirectoryExists(*GetStreamDirectory(StreamName)))
@@ -154,7 +159,7 @@ bool FNullNetworkReplayStreamer::IsLive( const FString& StreamName ) const
 void FNullNetworkReplayStreamer::DeleteFinishedStream( const FString& StreamName, const FOnDeleteFinishedStreamComplete& Delegate ) const
 {
 	// Live streams can't be deleted
-	if (IsLive(StreamName))
+	if (IsNamedStreamLive(StreamName))
 	{
 		UE_LOG(LogNullReplay, Log, TEXT("Can't delete network replay stream %s because it is live!"), *StreamName);
 		Delegate.ExecuteIfBound(false);
@@ -191,7 +196,7 @@ void FNullNetworkReplayStreamer::EnumerateStreams( const FNetworkReplayVersion&
 		{
 			Info.Name = Directory;
 			Info.Timestamp = IFileManager::Get().GetTimeStamp( *FullDemoFilePath );
-			Info.bIsLive = IsLive( Directory );			
+			Info.bIsLive = IsNamedStreamLive( Directory );			
 
 			// Live streams not supported yet
 			if (!Info.bIsLive)
diff --git a/Engine/Source/Runtime/NetworkReplayStreaming/NullNetworkReplayStreaming/Public/NullNetworkReplayStreaming.h b/Engine/Source/Runtime/NetworkReplayStreaming/NullNetworkReplayStreaming/Public/NullNetworkReplayStreaming.h
--- a/Engine/Source/Runtime/NetworkReplayStreaming/NullNetworkReplayStreaming/Public/NullNetworkReplayStreaming.h
+++ b/Engine/Source/Runtime/NetworkReplayStreaming/NullNetworkReplayStreaming/Public/NullNetworkReplayStreaming.h
@@ -30,12 +30,14 @@ public:
 	virtual bool IsDataAvailable() const override { return true; }
 	virtual void SetHighPriorityTimeRange( const uint32 StartTimeInMS, const uint32 EndTimeInMS ) override { }
 	virtual bool IsDataAvailableForTimeRange( const uint32 StartTimeInMS, const uint32 EndTimeInMS ) override { return true; }
-	virtual bool IsLive( const FString& StreamName ) const override;
+	virtual bool IsLive() const override;
 	virtual void DeleteFinishedStream( const FString& StreamName, const FOnDeleteFinishedStreamComplete& Delegate) const override;
 	virtual void EnumerateStreams( const FNetworkReplayVersion& ReplayVersion, const FOnEnumerateStreamsComplete& Delegate ) override;
 	virtual ENetworkReplayError::Type GetLastError() const override { return ENetworkReplayError::None; }
 
 private:
+	bool IsNamedStreamLive( const FString& StreamName ) const;
+
 	/** Handle to the archive that will read/write network packets */
 	TUniquePtr<FArchive> FileAr;
 
