Better logic used to jump to the end of live games, allows the initial load to settle down so that we can get a more fresh time to jump to

diff --git a/Engine/Source/Runtime/Engine/Classes/Engine/DemoNetDriver.h b/Engine/Source/Runtime/Engine/Classes/Engine/DemoNetDriver.h
--- a/Engine/Source/Runtime/Engine/Classes/Engine/DemoNetDriver.h
+++ b/Engine/Source/Runtime/Engine/Classes/Engine/DemoNetDriver.h
@@ -59,7 +59,8 @@ class ENGINE_API UDemoNetDriver : public UNetDriver
 	double		MaxRecordTime;
 	int32		RecordCountSinceFlush;
 	float		TimeToSkip;
-	float		LastGotoTimeInSeconds;
+	float		QueuedGotoTimeInSeconds;
+	uint32		InitialLiveDemoTime;
 
 	bool		bSavingCheckpoint;
 	double		LastCheckpointTime;
diff --git a/Engine/Source/Runtime/Engine/Private/DemoNetDriver.cpp b/Engine/Source/Runtime/Engine/Private/DemoNetDriver.cpp
--- a/Engine/Source/Runtime/Engine/Private/DemoNetDriver.cpp
+++ b/Engine/Source/Runtime/Engine/Private/DemoNetDriver.cpp
@@ -58,8 +58,9 @@ bool UDemoNetDriver::InitBase( bool bInitAsClient, FNetworkNotify* InNotify, con
 		TimeToSkip				= 0.0f;
 		bIsFastForwarding		= false;
 		GotoCheckpointSkipExtraTimeInMS = -1;
-		LastGotoTimeInSeconds	= -1.0f;
+		QueuedGotoTimeInSeconds	= -1.0f;
 		bIsLoadingCheckpoint	= false;
+		InitialLiveDemoTime		= 0;
 
 		ResetDemoState();
 
@@ -1113,9 +1114,9 @@ void UDemoNetDriver::SkipTime(const float InTimeToSkip)
 
 void UDemoNetDriver::GotoTimeInSeconds( const float TimeInSeconds )
 {
-	if ( LastGotoTimeInSeconds < 0 )
+	if ( QueuedGotoTimeInSeconds < 0 )
 	{
-		LastGotoTimeInSeconds = TimeInSeconds;
+		QueuedGotoTimeInSeconds = TimeInSeconds;
 	}
 }
 
@@ -1127,7 +1128,28 @@ void UDemoNetDriver::TickDemoPlayback( float DeltaSeconds )
 		return;
 	}
 
-	if ( GotoCheckpointArchive == NULL )
+	// See if we need to jump into the end (or close to it) of a live game we just joined
+	if ( InitialLiveDemoTime > 0 )
+	{
+		const uint32 TotalDemoTimeInMS = ReplayStreamer->GetTotalDemoTime();
+
+		// Wait until we have a freshly updated time to join, which allows things to settle down and get a more fresh time
+		if ( TotalDemoTimeInMS == InitialLiveDemoTime )
+		{
+			return;
+		}
+
+		InitialLiveDemoTime = 0;
+
+		DemoTotalTime = (float)TotalDemoTimeInMS / 1000.0f;
+
+		ReplayStreamer->GotoTimeInMS( ReplayStreamer->GetTotalDemoTime() - 10 * 1000, FOnCheckpointReadyDelegate::CreateUObject( this, &UDemoNetDriver::CheckpointReady ) );
+
+		OldDemoCurrentTime = DemoCurrentTime;		// So we can restore on failure
+		DemoCurrentTime = (float)( ReplayStreamer->GetTotalDemoTime() - 10 * 1000 ) / 1000.0f;
+	}
+
+	if ( GotoCheckpointArchive == NULL && !ReplayStreamer->IsLoadingCheckpoint() )
 	{
 		// Check checkpoint debug cvars
 		check( GotoCheckpointSkipExtraTimeInMS == -1 );
@@ -1144,12 +1166,13 @@ void UDemoNetDriver::TickDemoPlayback( float DeltaSeconds )
 			CVarGotoTimeInSeconds.AsVariable()->Set( TEXT( "-1" ), ECVF_SetByConsole );
 		}
 
-		if ( LastGotoTimeInSeconds >= 0.0f )
+		// Process any queued scrub time
+		if ( QueuedGotoTimeInSeconds >= 0.0f )
 		{
-			ReplayStreamer->GotoTimeInMS( LastGotoTimeInSeconds * 1000, FOnCheckpointReadyDelegate::CreateUObject( this, &UDemoNetDriver::CheckpointReady ) );
+			ReplayStreamer->GotoTimeInMS( QueuedGotoTimeInSeconds * 1000, FOnCheckpointReadyDelegate::CreateUObject( this, &UDemoNetDriver::CheckpointReady ) );
 			OldDemoCurrentTime = DemoCurrentTime;		// So we can restore on failure
-			DemoCurrentTime = LastGotoTimeInSeconds;
-			LastGotoTimeInSeconds = -1.0f;
+			DemoCurrentTime = QueuedGotoTimeInSeconds;
+			QueuedGotoTimeInSeconds = -1.0f;
 		}
 	}
 
@@ -1169,7 +1192,7 @@ void UDemoNetDriver::TickDemoPlayback( float DeltaSeconds )
 	}
 
 	// See if we have time to skip
-	if ( TimeToSkip > 0.0f )
+	if ( TimeToSkip > 0.0f && !ReplayStreamer->IsLoadingCheckpoint() )
 	{
 		const uint32 TimeInMSToCheck = FMath::Clamp( GetDemoCurrentTimeInMS() + (uint32)( TimeToSkip * 1000 ), (uint32)0, TotalDemoTimeInMS );
 
@@ -1300,19 +1323,9 @@ void UDemoNetDriver::ReplayStreamingReady( bool bSuccess, bool bRecord )
 		FString Error;
 		InitConnectInternal( Error );
 
-		const uint32 TotalDemoTimeInMS = ReplayStreamer->GetTotalDemoTime();
-
-		if ( TotalDemoTimeInMS > 0 )
-		{
-			DemoTotalTime = (float)TotalDemoTimeInMS / 1000.0f;
-		}
-
 		if ( ReplayStreamer->IsLive() && ReplayStreamer->GetTotalDemoTime() > 15 * 1000 )
 		{
-			ReplayStreamer->GotoTimeInMS( ReplayStreamer->GetTotalDemoTime() - 10 * 1000, FOnCheckpointReadyDelegate::CreateUObject( this, &UDemoNetDriver::CheckpointReady ) );
-
-			OldDemoCurrentTime = DemoCurrentTime;		// So we can restore on failure
-			DemoCurrentTime = (float)( ReplayStreamer->GetTotalDemoTime() - 10 * 1000 ) / 1000.0f;
+			InitialLiveDemoTime = ReplayStreamer->GetTotalDemoTime();
 		}
 	}
 }
diff --git a/Engine/Source/Runtime/NetworkReplayStreaming/HttpNetworkReplayStreaming/Private/HttpNetworkReplayStreaming.cpp b/Engine/Source/Runtime/NetworkReplayStreaming/HttpNetworkReplayStreaming/Private/HttpNetworkReplayStreaming.cpp
--- a/Engine/Source/Runtime/NetworkReplayStreaming/HttpNetworkReplayStreaming/Private/HttpNetworkReplayStreaming.cpp
+++ b/Engine/Source/Runtime/NetworkReplayStreaming/HttpNetworkReplayStreaming/Private/HttpNetworkReplayStreaming.cpp
@@ -707,6 +707,11 @@ bool FHttpNetworkReplayStreamer::IsLive() const
 	return bStreamIsLive;
 }
 
+bool FHttpNetworkReplayStreamer::IsLoadingCheckpoint() const
+{
+	return GotoCheckpointDelegate.IsBound();
+}
+
 void FHttpNetworkReplayStreamer::DeleteFinishedStream( const FString& StreamName, const FOnDeleteFinishedStreamComplete& Delegate ) const
 {
 	// Stubbed!
@@ -1002,6 +1007,7 @@ void FHttpNetworkReplayStreamer::HttpDownloadCheckpointFinished( FHttpRequestPtr
 		{
 			UE_LOG( LogHttpReplay, Warning, TEXT( "FHttpNetworkReplayStreamer::HttpDownloadCheckpointFinished. Checkpoint empty." ) );
 			GotoCheckpointDelegate.ExecuteIfBound( false, -1 );
+			GotoCheckpointDelegate = FOnCheckpointReadyDelegate();
 			return;
 		}
 
diff --git a/Engine/Source/Runtime/NetworkReplayStreaming/HttpNetworkReplayStreaming/Private/HttpNetworkReplayStreaming.h b/Engine/Source/Runtime/NetworkReplayStreaming/HttpNetworkReplayStreaming/Private/HttpNetworkReplayStreaming.h
--- a/Engine/Source/Runtime/NetworkReplayStreaming/HttpNetworkReplayStreaming/Private/HttpNetworkReplayStreaming.h
+++ b/Engine/Source/Runtime/NetworkReplayStreaming/HttpNetworkReplayStreaming/Private/HttpNetworkReplayStreaming.h
@@ -110,6 +110,7 @@ public:
 	virtual bool		IsDataAvailable() const override;
 	virtual void		SetHighPriorityTimeRange( const uint32 StartTimeInMS, const uint32 EndTimeInMS ) override;
 	virtual bool		IsDataAvailableForTimeRange( const uint32 StartTimeInMS, const uint32 EndTimeInMS ) override;
+	virtual bool		IsLoadingCheckpoint() const override;
 	virtual bool		IsLive() const override;
 	virtual void		DeleteFinishedStream( const FString& StreamName, const FOnDeleteFinishedStreamComplete& Delegate ) const override;
 	virtual void		EnumerateStreams( const FNetworkReplayVersion& ReplayVersion, const FOnEnumerateStreamsComplete& Delegate ) override;
diff --git a/Engine/Source/Runtime/NetworkReplayStreaming/NetworkReplayStreaming/Public/NetworkReplayStreaming.h b/Engine/Source/Runtime/NetworkReplayStreaming/NetworkReplayStreaming/Public/NetworkReplayStreaming.h
--- a/Engine/Source/Runtime/NetworkReplayStreaming/NetworkReplayStreaming/Public/NetworkReplayStreaming.h
+++ b/Engine/Source/Runtime/NetworkReplayStreaming/NetworkReplayStreaming/Public/NetworkReplayStreaming.h
@@ -117,6 +117,7 @@ public:
 	virtual bool IsDataAvailable() const = 0;
 	virtual void SetHighPriorityTimeRange( const uint32 StartTimeInMS, const uint32 EndTimeInMS ) = 0;
 	virtual bool IsDataAvailableForTimeRange( const uint32 StartTimeInMS, const uint32 EndTimeInMS ) = 0;
+	virtual bool IsLoadingCheckpoint() const = 0;
 	
 	/** Returns true if the playing stream is currently in progress */
 	virtual bool IsLive() const = 0;
diff --git a/Engine/Source/Runtime/NetworkReplayStreaming/NullNetworkReplayStreaming/Public/NullNetworkReplayStreaming.h b/Engine/Source/Runtime/NetworkReplayStreaming/NullNetworkReplayStreaming/Public/NullNetworkReplayStreaming.h
--- a/Engine/Source/Runtime/NetworkReplayStreaming/NullNetworkReplayStreaming/Public/NullNetworkReplayStreaming.h
+++ b/Engine/Source/Runtime/NetworkReplayStreaming/NullNetworkReplayStreaming/Public/NullNetworkReplayStreaming.h
@@ -30,6 +30,7 @@ public:
 	virtual bool IsDataAvailable() const override { return true; }
 	virtual void SetHighPriorityTimeRange( const uint32 StartTimeInMS, const uint32 EndTimeInMS ) override { }
 	virtual bool IsDataAvailableForTimeRange( const uint32 StartTimeInMS, const uint32 EndTimeInMS ) override { return true; }
+	virtual bool IsLoadingCheckpoint() const override { return false; }
 	virtual bool IsLive() const override;
 	virtual void DeleteFinishedStream( const FString& StreamName, const FOnDeleteFinishedStreamComplete& Delegate) const override;
 	virtual void EnumerateStreams( const FNetworkReplayVersion& ReplayVersion, const FOnEnumerateStreamsComplete& Delegate ) override;
