Fix teleporting actors after scrubbing in replays

diff --git a/Engine/Source/Runtime/Engine/Private/DataChannel.cpp b/Engine/Source/Runtime/Engine/Private/DataChannel.cpp
--- a/Engine/Source/Runtime/Engine/Private/DataChannel.cpp
+++ b/Engine/Source/Runtime/Engine/Private/DataChannel.cpp
@@ -1726,6 +1726,12 @@ bool UActorChannel::ProcessQueuedBunches()
 		UE_LOG( LogNet, VeryVerbose, TEXT( "UActorChannel::ProcessQueuedBunches: Flushing queued bunches. ChIndex: %i, Actor: %s, Queued: %i" ), ChIndex, Actor != NULL ? *Actor->GetPathName() : TEXT( "NULL" ), QueuedBunches.Num() );
 
 		QueuedBunches.Empty();
+
+		// Call any onreps that were delayed because we were queuing bunches
+		for (auto& ReplicatorPair : ReplicationMap)
+		{
+			ReplicatorPair.Value->CallRepNotifies();
+		}
 	}
 
 	// Warn when we have queued bunches for a very long time
diff --git a/Engine/Source/Runtime/Engine/Private/DataReplication.cpp b/Engine/Source/Runtime/Engine/Private/DataReplication.cpp
--- a/Engine/Source/Runtime/Engine/Private/DataReplication.cpp
+++ b/Engine/Source/Runtime/Engine/Private/DataReplication.cpp
@@ -1132,6 +1132,11 @@ void FObjectReplicator::CallRepNotifies()
 		return;
 	}
 
+	if ( OwningChannel != NULL && OwningChannel->QueuedBunches.Num() > 0 )
+	{
+		return;
+	}
+
 	RepLayout->CallRepNotifies( RepState, Object );
 
 	if ( RepNotifies.Num() > 0 )
