#ue4 - Clean up constness of local vars in ComponentSweepMulti().

diff --git a/Engine/Source/Runtime/Engine/Private/Collision/WorldCollision.cpp b/Engine/Source/Runtime/Engine/Private/Collision/WorldCollision.cpp
--- a/Engine/Source/Runtime/Engine/Private/Collision/WorldCollision.cpp
+++ b/Engine/Source/Runtime/Engine/Private/Collision/WorldCollision.cpp
@@ -530,11 +530,11 @@ bool UWorld::ComponentSweepMulti(TArray<struct FHitResult>& OutHits, class UPrim
 				TArray<struct FHitResult> Hits;
 
 				// Calc shape global pose
-				PxTransform PLocalShape = PShape->getLocalPose();
-				PxTransform PShapeGlobalStartPose = PGlobalStartPose.transform(PLocalShape);
-				PxTransform PShapeGlobalEndPose = PGlobalEndPose.transform(PLocalShape);
+				const PxTransform PLocalShape = PShape->getLocalPose();
+				const PxTransform PShapeGlobalStartPose = PGlobalStartPose.transform(PLocalShape);
+				const PxTransform PShapeGlobalEndPose = PGlobalEndPose.transform(PLocalShape);
 				// consider localshape rotation for shape rotation
-				PxQuat PShapeRot = PGeomRot * PLocalShape.q;
+				const PxQuat PShapeRot = PGeomRot * PLocalShape.q;
 
 				if (GeomSweepMulti_PhysX(this, *PGeom, PShapeRot, Hits, P2UVector(PShapeGlobalStartPose.p), P2UVector(PShapeGlobalEndPose.p), TraceChannel, Params, FCollisionResponseParams(PrimComp->GetCollisionResponseToChannels())))
 				{
