added help text

diff --git a/Engine/Source/Runtime/Engine/Private/Atmosphere/Atmosphere.cpp b/Engine/Source/Runtime/Engine/Private/Atmosphere/Atmosphere.cpp
--- a/Engine/Source/Runtime/Engine/Private/Atmosphere/Atmosphere.cpp
+++ b/Engine/Source/Runtime/Engine/Private/Atmosphere/Atmosphere.cpp
@@ -213,7 +213,7 @@ static TAutoConsoleVariable<int32> CVarAtmosphereRender(
 	1,
 	TEXT("Defines atmosphere will render or not. Only changed by r.Atmosphere console command.\n")
 	TEXT("Enable/Disable Atmosphere, Load/Unload related data.\n")
-	TEXT(" 0: off\n")
+	TEXT(" 0: off (To save GPU memory)\n")
 	TEXT(" 1: on (default)"));
 
 // On CPU
