removed r.GBuffer project setting (experimental, not finished)

diff --git a/Engine/Source/Runtime/Engine/Classes/Engine/RendererSettings.h b/Engine/Source/Runtime/Engine/Classes/Engine/RendererSettings.h
--- a/Engine/Source/Runtime/Engine/Classes/Engine/RendererSettings.h
+++ b/Engine/Source/Runtime/Engine/Classes/Engine/RendererSettings.h
@@ -143,12 +143,6 @@ class ENGINE_API URendererSettings
 		ToolTip="Whether to allow any static lighting to use normal maps for lighting computations."))
 	uint32 bUseNormalMapsForStaticLighting:1;
 
-	UPROPERTY(config, EditAnywhere, Category=Lighting, meta=(
-		ConsoleVariable="r.GBuffer",
-		ToolTip="1=Use GBuffer, 0=Don't use GBuffer (minimal limited renderer)."))
-		uint32 bGBuffer:1;
-
-
 	UPROPERTY(config, EditAnywhere, Category=Lighting, meta=(
 		ConsoleVariable="r.GenerateMeshDistanceFields",
 		ToolTip="Whether to build distance fields of static meshes, needed for distance field AO, which is used to implement Movable SkyLight shadows, and ray traced distance field shadows on directional lights.  Enabling will increase mesh build times and memory usage.  Changing this setting requires restarting the editor."))
