Made Atmosphere toggleable by Fog, consistently renamed to AtmosphericFog
UE-13368 Atmospheric Fog cannot be toggled to hidden in the Scene Capture 2D's detail settings

diff --git a/Engine/Source/Editor/DetailCustomizations/Private/SceneCaptureDetails.cpp b/Engine/Source/Editor/DetailCustomizations/Private/SceneCaptureDetails.cpp
--- a/Engine/Source/Editor/DetailCustomizations/Private/SceneCaptureDetails.cpp
+++ b/Engine/Source/Editor/DetailCustomizations/Private/SceneCaptureDetails.cpp
@@ -67,7 +67,7 @@ void FSceneCaptureDetails::CustomizeDetails( IDetailLayoutBuilder& DetailLayout
 	TArray<FEngineShowFlags::EShowFlag> ShowFlagsToAllowForCaptures;
 
 	// General
-	ShowFlagsToAllowForCaptures.Add(FEngineShowFlags::EShowFlag::SF_Atmosphere);
+	ShowFlagsToAllowForCaptures.Add(FEngineShowFlags::EShowFlag::SF_AtmosphericFog);
 	ShowFlagsToAllowForCaptures.Add(FEngineShowFlags::EShowFlag::SF_BSP);
 	ShowFlagsToAllowForCaptures.Add(FEngineShowFlags::EShowFlag::SF_Decals);
 	ShowFlagsToAllowForCaptures.Add(FEngineShowFlags::EShowFlag::SF_Fog);
diff --git a/Engine/Source/Runtime/Engine/Private/ShowFlags.cpp b/Engine/Source/Runtime/Engine/Private/ShowFlags.cpp
--- a/Engine/Source/Runtime/Engine/Private/ShowFlags.cpp
+++ b/Engine/Source/Runtime/Engine/Private/ShowFlags.cpp
@@ -398,10 +398,11 @@ void EngineShowFlagOverride(EShowFlagInitMode ShowFlagInitMode, EViewModeIndex V
 			ViewModeIndex == VMI_BrushWireframe ||
 			ViewModeIndex == VMI_Unlit ||
 			ViewModeIndex == VMI_StationaryLightOverlap ||
+			ViewModeIndex == VMI_ShaderComplexity ||
 			ViewModeIndex == VMI_LightmapDensity)
 		{
 			EngineShowFlags.Lighting = 0;
-			EngineShowFlags.Atmosphere = 0;
+			EngineShowFlags.AtmosphericFog = 0;
 		}
 
 		if( ViewModeIndex == VMI_Lit ||
diff --git a/Engine/Source/Runtime/Engine/Public/ShowFlagsValues.inl b/Engine/Source/Runtime/Engine/Public/ShowFlagsValues.inl
--- a/Engine/Source/Runtime/Engine/Public/ShowFlagsValues.inl
+++ b/Engine/Source/Runtime/Engine/Public/ShowFlagsValues.inl
@@ -234,8 +234,8 @@ SHOWFLAG_ALWAYS_ACCESSIBLE(CollisionPawn, SFG_Hidden, LOCTEXT("CollisionPawn", "
 SHOWFLAG_ALWAYS_ACCESSIBLE(LightShafts, SFG_LightingFeatures, LOCTEXT("LightShaftsSF", "Light Shafts"))
 /** Render the PostProcess Material */
 SHOWFLAG_ALWAYS_ACCESSIBLE(PostProcessMaterial, SFG_PostProcess, LOCTEXT("PostProcessMaterialSF", "Post Process Material"))
-/** Render Atmospheric scattering */
-SHOWFLAG_ALWAYS_ACCESSIBLE(Atmosphere, SFG_Normal, LOCTEXT("AtmosphereSF", "Atmosphere"))
+/** Render Atmospheric scattering (Atmospheric Fog) */
+SHOWFLAG_ALWAYS_ACCESSIBLE(AtmosphericFog, SFG_Advanced, LOCTEXT("AtmosphereSF", "Atmospheric Fog"))
 /** Render safe frames bars*/
 SHOWFLAG_ALWAYS_ACCESSIBLE(CameraAspectRatioBars, SFG_Advanced, LOCTEXT("CameraAspectRatioBarsSF", "Camera Aspect Ratio Bars"))
 /** Render safe frames */
diff --git a/Engine/Source/Runtime/Renderer/Private/BasePassRendering.cpp b/Engine/Source/Runtime/Renderer/Private/BasePassRendering.cpp
--- a/Engine/Source/Runtime/Renderer/Private/BasePassRendering.cpp
+++ b/Engine/Source/Runtime/Renderer/Private/BasePassRendering.cpp
@@ -324,7 +324,7 @@ public:
 			Parameters.BlendMode,
 			Parameters.TextureMode,
 			Scene && Scene->SkyLight && !Scene->SkyLight->bHasStaticLighting && Scene->SkyLight->bWantsStaticShadowing && bIsLitMaterial,
-			IsTranslucentBlendMode(Parameters.BlendMode) && (Scene && Scene->HasAtmosphericFog()) && View.Family->EngineShowFlags.Atmosphere,
+			IsTranslucentBlendMode(Parameters.BlendMode) && (Scene && Scene->HasAtmosphericFog()) && View.Family->EngineShowFlags.AtmosphericFog,
 			View.Family->EngineShowFlags.ShaderComplexity,
 			false,
 			Parameters.bEditorCompositeDepthTest
diff --git a/Engine/Source/Runtime/Renderer/Private/TranslucentRendering.cpp b/Engine/Source/Runtime/Renderer/Private/TranslucentRendering.cpp
--- a/Engine/Source/Runtime/Renderer/Private/TranslucentRendering.cpp
+++ b/Engine/Source/Runtime/Renderer/Private/TranslucentRendering.cpp
@@ -283,7 +283,7 @@ public:
 			// Translucent meshes need scene render targets set as textures
 			ESceneRenderTargetsMode::SetTextures,
 			bIsLitMaterial && Scene && Scene->SkyLight && !Scene->SkyLight->bHasStaticLighting,
-			Scene && Scene->HasAtmosphericFog() && View.Family->EngineShowFlags.Atmosphere,
+			Scene && Scene->HasAtmosphericFog() && View.Family->EngineShowFlags.AtmosphericFog && View.Family->EngineShowFlags.Fog,
 			View.Family->EngineShowFlags.ShaderComplexity,
 			Parameters.bAllowFog
 			);
