Use machine id to identify non promoted network builds

diff --git a/Engine/Source/Runtime/Engine/Private/UnrealNetwork.cpp b/Engine/Source/Runtime/Engine/Private/UnrealNetwork.cpp
--- a/Engine/Source/Runtime/Engine/Private/UnrealNetwork.cpp
+++ b/Engine/Source/Runtime/Engine/Private/UnrealNetwork.cpp
@@ -39,16 +39,19 @@ uint32 FNetworkVersion::GetLocalNetworkVersion()
 	// Hash with internal protocol version
 	uint32 LocalNetworkVersion = FCrc::MemCrc32( &InternalProtocolVersion, sizeof( InternalProtocolVersion ), VersionHash );
 
+#if !(UE_BUILD_SHIPPING || UE_BUILD_TEST)
 	if ( !GEngineVersion.IsPromotedBuild() )
 	{
-		// Further hash with date time if this is a non promoted build
-		const FString DateTimeString = FString::Printf( TEXT( "%s_%s" ), ANSI_TO_TCHAR( __DATE__ ), ANSI_TO_TCHAR( __TIME__ ) );
-		const uint32 LocalNetworkVersionNonPromoted = FCrc::StrCrc32( *DateTimeString, LocalNetworkVersion );
+		// Further hash with machine id if this is a non promoted build
+		const FString MachineId = FPlatformMisc::GetMachineId().ToString( EGuidFormats::Digits ).ToLower();
 
-		UE_LOG( LogNet, Log, TEXT( "GetLocalNetworkVersion: NON-PROMOTED: Date: %s, ProjectName: %s, ProjectVersion: %s, InternalProtocolVersion: %i, LocalNetworkVersionNonPromoted: %u" ), *DateTimeString, *ProjectName, *ProjectVersion, InternalProtocolVersion, LocalNetworkVersionNonPromoted );
+		const uint32 LocalNetworkVersionNonPromoted = FCrc::StrCrc32( *MachineId, LocalNetworkVersion );
+
+		UE_LOG( LogNet, Log, TEXT( "GetLocalNetworkVersion: NON-PROMOTED: MachineId: %s, ProjectName: %s, ProjectVersion: %s, InternalProtocolVersion: %i, LocalNetworkVersionNonPromoted: %u" ), *MachineId, *ProjectName, *ProjectVersion, InternalProtocolVersion, LocalNetworkVersionNonPromoted );
 
 		return LocalNetworkVersionNonPromoted;
 	}
+#endif
 
 	UE_LOG( LogNet, Log, TEXT( "GetLocalNetworkVersion: GEngineNetVersion: %i, ProjectName: %s, ProjectVersion: %s, InternalProtocolVersion: %i, LocalNetworkVersion: %u" ), GEngineNetVersion, *ProjectName, *ProjectVersion, InternalProtocolVersion, LocalNetworkVersion );
 
