Fixing the random weight choosing algorithm in SoundNodeRandom

https://jira.ol.epicgames.net/browse/UE-13319

diff --git a/Engine/Source/Runtime/Engine/Private/SoundNodeRandom.cpp b/Engine/Source/Runtime/Engine/Private/SoundNodeRandom.cpp
--- a/Engine/Source/Runtime/Engine/Private/SoundNodeRandom.cpp
+++ b/Engine/Source/Runtime/Engine/Private/SoundNodeRandom.cpp
@@ -110,31 +110,26 @@ void USoundNodeRandom::ParseNodes( FAudioDevice* AudioDevice, const UPTRINT Node
 			}
 		}
 
-		float Weight = FMath::FRand() * WeightSum;
+		float Choice = FMath::FRand() * WeightSum;
+		WeightSum = 0.0f;
 		for( int32 i = 0; i < ChildNodes.Num() && i < Weights.Num(); ++i )
 		{
 #if WITH_EDITOR
 			if(!bIsPIESound || !PIEHiddenNodes.Contains(i) )
 #endif //WITH_EDITOR
 			{
-				if( bRandomizeWithoutReplacement && ( Weights[ i ] >= Weight ) && ( HasBeenUsed[ i ] != true ) )
+				if (bRandomizeWithoutReplacement && HasBeenUsed[i])
 				{
-					HasBeenUsed[ i ] = true;
-					// we played a sound so increment how many sounds we have played
-					++NumRandomUsed;
-
-					NodeIndex = i;
-					break;
+					continue;
 				}
-				else if( ( bRandomizeWithoutReplacement == false ) && ( Weights[ i ] >= Weight ) )
+				WeightSum += Weights[i];
+				if (Choice < WeightSum)
 				{
 					NodeIndex = i;
+					HasBeenUsed[i] = true;
+					++NumRandomUsed;
 					break;
 				}
-				else
-				{
-					Weight -= Weights[ i ];
-				}
 			}
 		}
 
