Use IsGameWorld for determining behaviors when calling Remove/Destroy Actor
Ensure that an actor could not be removed, don't crash

diff --git a/Engine/Source/Runtime/Engine/Private/LevelActor.cpp b/Engine/Source/Runtime/Engine/Private/LevelActor.cpp
--- a/Engine/Source/Runtime/Engine/Private/LevelActor.cpp
+++ b/Engine/Source/Runtime/Engine/Private/LevelActor.cpp
@@ -418,7 +418,7 @@ bool UWorld::DestroyActor( AActor* ThisActor, bool bNetForce, bool bShouldModify
 	}
 
 	// In-game deletion rules.
-	if( AreActorsInitialized() )
+	if( IsGameWorld() )
 	{
 		// Never destroy the world settings actor. This used to be enforced by bNoDelete and is actually needed for 
 		// seamless travel and network games.
diff --git a/Engine/Source/Runtime/Engine/Private/World.cpp b/Engine/Source/Runtime/Engine/Private/World.cpp
--- a/Engine/Source/Runtime/Engine/Private/World.cpp
+++ b/Engine/Source/Runtime/Engine/Private/World.cpp
@@ -1101,7 +1101,7 @@ void UWorld::RemoveActor(AActor* Actor, bool bShouldModifyLevel)
 			ModifyLevel( CheckLevel );
 		}
 		
-		if (!AreActorsInitialized())
+		if (!IsGameWorld())
 		{
 			CheckLevel->Actors.ModifyItem(ActorListIndex);
 		}
@@ -1135,9 +1135,8 @@ void UWorld::RemoveActor(AActor* Actor, bool bShouldModifyLevel)
 				UE_LOG(LogWorld, Log, TEXT("  %s"), (CurrentActor ? *CurrentActor->GetPathName() : TEXT("NONE")));
 			}
 		}
-		UE_LOG(LogWorld, Fatal, TEXT("Could not remove actor %s from world (check level is %s)"), *Actor->GetPathName(), *CheckLevel->GetPathName());
+		ensureMsgf(false, TEXT("Could not remove actor %s from world (check level is %s)"), *Actor->GetPathName(), *CheckLevel->GetPathName());
 	}
-	check(bSuccessfulRemoval);
 }
 
 
