polish UI

diff --git a/Engine/Source/Runtime/Engine/Classes/Materials/Material.h b/Engine/Source/Runtime/Engine/Classes/Materials/Material.h
--- a/Engine/Source/Runtime/Engine/Classes/Materials/Material.h
+++ b/Engine/Source/Runtime/Engine/Classes/Materials/Material.h
@@ -325,7 +325,7 @@ class UMaterial : public UMaterialInterface
 	 * Defines how the material reacts on DBuffer decals (Affects look, performance and texture/sample usage).
 	 * Non DBuffer Decals can be disabled on the primitive (e.g. static mesh)
 	 */
-	UPROPERTY(EditAnywhere, BlueprintReadOnly, Category=Material, meta=(DisplayName = "Decal Response"), AssetRegistrySearchable)
+	UPROPERTY(EditAnywhere, BlueprintReadOnly, Category=Material, meta=(DisplayName = "Decal Response (DBuffer)"), AssetRegistrySearchable)
 	TEnumAsByte<enum EMaterialDecalResponse> MaterialDecalResponse;
 
 private:
