Improving ensure message to add info to UE-11609, UE-11807, and UE-12014 which are all CNR at the moment

diff --git a/Engine/Source/Runtime/Engine/Private/Blueprint.cpp b/Engine/Source/Runtime/Engine/Private/Blueprint.cpp
--- a/Engine/Source/Runtime/Engine/Private/Blueprint.cpp
+++ b/Engine/Source/Runtime/Engine/Private/Blueprint.cpp
@@ -598,7 +598,11 @@ void UBlueprint::SetObjectBeingDebugged(UObject* NewObject)
 	if ((NewObject != NULL) && !GCompilingBlueprint && BlueprintType != BPTYPE_MacroLibrary)
 	{
 		// You can only debug instances of this!
-		if (!ensure(NewObject->IsA(this->GeneratedClass)))
+		if (!ensureMsgf(
+				NewObject->IsA(this->GeneratedClass), 
+				TEXT("Type mismatch: Expected %s, Found %s"), 
+				this->GeneratedClass ? *(this->GeneratedClass->GetName()) : TEXT("NULL"), 
+				NewObject->GetClass() ? *(this->GetClass()->GetName()) : TEXT("NULL")))
 		{
 			NewObject = NULL;
 		}
