PR #1028: Modifications to allow inheriting some Matinee classes in a plugin (Contributed by [REDACTED_PERSON])

diff --git a/Engine/Source/Runtime/Engine/Classes/Matinee/InterpTrack.h b/Engine/Source/Runtime/Engine/Classes/Matinee/InterpTrack.h
--- a/Engine/Source/Runtime/Engine/Classes/Matinee/InterpTrack.h
+++ b/Engine/Source/Runtime/Engine/Classes/Matinee/InterpTrack.h
@@ -280,12 +280,12 @@ public:
 	/**
 	 *	Conditionally calls PreviewUpdateTrack depending on whether or not the track is enabled.
 	 */
-	virtual void ConditionalPreviewUpdateTrack(float NewPosition, class UInterpTrackInst* TrInst);
+	ENGINE_API virtual void ConditionalPreviewUpdateTrack(float NewPosition, class UInterpTrackInst* TrInst);
 
 	/**
 	 *	Conditionally calls UpdateTrack depending on whether or not the track is enabled.
 	 */
-	virtual void ConditionalUpdateTrack(float NewPosition, class UInterpTrackInst* TrInst, bool bJump);
+	ENGINE_API virtual void ConditionalUpdateTrack(float NewPosition, class UInterpTrackInst* TrInst, bool bJump);
 
 
 	/**
@@ -329,7 +329,7 @@ public:
 	virtual bool AllowStaticActors() { return false; }
 
 	/** Draw this track with the specified parameters */
-	virtual void DrawTrack( FCanvas* Canvas, UInterpGroup* Group, const FInterpTrackDrawParams& Params );
+	ENGINE_API virtual void DrawTrack( FCanvas* Canvas, UInterpGroup* Group, const FInterpTrackDrawParams& Params );
 
 	/** @return color to draw each keyframe in Matinee. */
 	virtual FColor GetKeyframeColor(int32 KeyIndex) const;
diff --git a/Engine/Source/Runtime/Engine/Classes/Matinee/InterpTrackFloatBase.h b/Engine/Source/Runtime/Engine/Classes/Matinee/InterpTrackFloatBase.h
--- a/Engine/Source/Runtime/Engine/Classes/Matinee/InterpTrackFloatBase.h
+++ b/Engine/Source/Runtime/Engine/Classes/Matinee/InterpTrackFloatBase.h
@@ -4,8 +4,8 @@
 #include "Matinee/InterpTrack.h"
 #include "InterpTrackFloatBase.generated.h"
 
-UCLASS(abstract, MinimalAPI)
-class UInterpTrackFloatBase : public UInterpTrack
+UCLASS(abstract)
+class ENGINE_API UInterpTrackFloatBase : public UInterpTrack
 {
 	GENERATED_UCLASS_BODY()
 	
diff --git a/Engine/Source/Runtime/Engine/Classes/Matinee/InterpTrackInst.h b/Engine/Source/Runtime/Engine/Classes/Matinee/InterpTrackInst.h
--- a/Engine/Source/Runtime/Engine/Classes/Matinee/InterpTrackInst.h
+++ b/Engine/Source/Runtime/Engine/Classes/Matinee/InterpTrackInst.h
@@ -56,6 +56,6 @@ class UInterpTrackInst : public UObject
 	*
 	* returns the GroupActor's world
 	*/
-	virtual UWorld* GetWorld() const override;
+	ENGINE_API virtual UWorld* GetWorld() const override;
 };
 
diff --git a/Engine/Source/Runtime/Engine/Classes/Matinee/InterpTrackVectorBase.h b/Engine/Source/Runtime/Engine/Classes/Matinee/InterpTrackVectorBase.h
--- a/Engine/Source/Runtime/Engine/Classes/Matinee/InterpTrackVectorBase.h
+++ b/Engine/Source/Runtime/Engine/Classes/Matinee/InterpTrackVectorBase.h
@@ -4,8 +4,8 @@
 #pragma once
 #include "InterpTrackVectorBase.generated.h"
 
-UCLASS(abstract, MinimalAPI)
-class UInterpTrackVectorBase : public UInterpTrack
+UCLASS(abstract)
+class ENGINE_API UInterpTrackVectorBase : public UInterpTrack
 {
 	GENERATED_UCLASS_BODY()
 
