added docs

diff --git a/Engine/Source/Runtime/Engine/Classes/Components/DecalComponent.h b/Engine/Source/Runtime/Engine/Classes/Components/DecalComponent.h
--- a/Engine/Source/Runtime/Engine/Classes/Components/DecalComponent.h
+++ b/Engine/Source/Runtime/Engine/Classes/Components/DecalComponent.h
@@ -9,6 +9,9 @@ class FDeferredDecalProxy;
 
 /** 
  * A material that is rendered onto the surface of a mesh. A kind of 'bumper sticker' for a model.
+ *
+ * @see https://docs.unrealengine.com/latest/INT/Engine/Actors/DecalActor
+ * @see UDecalActor
  */
 UCLASS(hidecategories=(Collision, Object, Physics, SceneComponent, Activation, "Components|Activation", Mobility), ClassGroup=Rendering, meta=(BlueprintSpawnableComponent))
 class ENGINE_API UDecalComponent : public USceneComponent
