added comments

diff --git a/Engine/Source/Runtime/Engine/Classes/Engine/DecalActor.h b/Engine/Source/Runtime/Engine/Classes/Engine/DecalActor.h
--- a/Engine/Source/Runtime/Engine/Classes/Engine/DecalActor.h
+++ b/Engine/Source/Runtime/Engine/Classes/Engine/DecalActor.h
@@ -7,9 +7,11 @@
 
 class UBoxComponent;
 
-
 /**
-* Decal to add texture/material on existing geometry.
+* DecalActor contains a DecalComponent which can be used to render material modifications on top of existing geometry.
+*
+* @see https://docs.unrealengine.com/latest/INT/Engine/Actors/DecalActor
+* @see UDecalComponent
 */
 UCLASS(hideCategories=(Collision, Attachment, Actor, Input, Replication), showCategories=("Input|MouseInput", "Input|TouchInput"), ComponentWrapperClass,MinimalAPI)
 class ADecalActor
