optimized changing Scalability and fixed some unneccessary updates
fixed UE-11025 Changing Effects Quality with Console Command Resets Audio Playback of Packaged Game

diff --git a/Engine/Source/Editor/UnrealEd/Private/EditorServer.cpp b/Engine/Source/Editor/UnrealEd/Private/EditorServer.cpp
--- a/Engine/Source/Editor/UnrealEd/Private/EditorServer.cpp
+++ b/Engine/Source/Editor/UnrealEd/Private/EditorServer.cpp
@@ -6153,21 +6153,6 @@ bool UEditorEngine::HandleSetDetailModeViewCommand( const TCHAR* Str, FOutputDev
 		static IConsoleVariable* DetailModeCVar = IConsoleManager::Get().FindConsoleVariable(TEXT("r.DetailMode"));
 		check (DetailMode);
 		DetailModeCVar->Set(DetailMode);
-
-		// Reregister all primitive components so the view will respect the new detail mode filter
-		for ( FActorIterator It(InWorld) ; It ; ++It )
-		{
-			AActor* Actor = static_cast<AActor*>( *It );
-			checkSlow( Actor->IsA(AActor::StaticClass()) );
-
-			TInlineComponentArray<UPrimitiveComponent*> Components;
-			Actor->GetComponents(Components);
-
-			for(int32 ComponentIndex = 0;ComponentIndex < Components.Num();ComponentIndex++)
-			{
-				Components[ComponentIndex]->MarkRenderStateDirty();
-			}
-		}
 	}
 
 	RedrawLevelEditingViewports( true );
diff --git a/Engine/Source/Runtime/Engine/Private/Materials/Material.cpp b/Engine/Source/Runtime/Engine/Private/Materials/Material.cpp
--- a/Engine/Source/Runtime/Engine/Private/Materials/Material.cpp
+++ b/Engine/Source/Runtime/Engine/Private/Materials/Material.cpp
@@ -1658,7 +1658,12 @@ void UMaterial::CacheResourceShadersForRendering(bool bRegenerateId)
 
 			// Only cache shaders for the quality level that will actually be used to render
 			ResourcesToCache.Reset();
-			ResourcesToCache.Add(MaterialResources[ActiveQualityLevel][FeatureLevel]);
+
+			FMaterialResource* MaterialResource = MaterialResources[ActiveQualityLevel][FeatureLevel];
+
+			check(MaterialResource);
+
+			ResourcesToCache.Add(MaterialResource);
 			CacheShadersForResources(ShaderPlatform, ResourcesToCache, true);
 		}
 
diff --git a/Engine/Source/Runtime/Engine/Private/UnrealEngine.cpp b/Engine/Source/Runtime/Engine/Private/UnrealEngine.cpp
--- a/Engine/Source/Runtime/Engine/Private/UnrealEngine.cpp
+++ b/Engine/Source/Runtime/Engine/Private/UnrealEngine.cpp
@@ -101,6 +101,7 @@
 #include "NotificationManager.h"
 #include "SNotificationList.h"
 #include "Engine/UserInterfaceSettings.h"
+#include "ComponentRecreateRenderStateContext.h"
 
 DEFINE_LOG_CATEGORY(LogEngine);
 
@@ -221,17 +222,20 @@ static FCachedSystemScalabilityCVars GCachedScalabilityCVars;
 
 const FCachedSystemScalabilityCVars& GetCachedScalabilityCVars()
 {
+	check(GCachedScalabilityCVars.bInitialized);
 	return GCachedScalabilityCVars;
 }
 
 FCachedSystemScalabilityCVars::FCachedSystemScalabilityCVars()
-	: DetailMode(-1)
+	: bInitialized(false)
+	, DetailMode(-1)
 	, MaterialQualityLevel(EMaterialQualityLevel::Num)
 	, MaxAnisotropy(-1)
 	, MaxShadowResolution(-1)
 	, ViewDistanceScale(-1)
 	, ViewDistanceScaleSquared(-1)
 	, GaussianDOFNearThreshold(-1)
+	, SimpleDynamicLighting(-1)
 {
 
 }
@@ -240,85 +244,81 @@ void ScalabilityCVarsSinkCallback()
 {
 	IConsoleManager& ConsoleMan = IConsoleManager::Get();
 
-	static const auto DetailMode = ConsoleMan.FindTConsoleVariableDataInt(TEXT("r.DetailMode"));
-	if( GCachedScalabilityCVars.DetailMode != DetailMode->GetValueOnGameThread() )
-	{
-		TArray<UClass*> ExcludeComponents;
-		ExcludeComponents.Add(UAudioComponent::StaticClass());
+	FCachedSystemScalabilityCVars LocalScalabilityCVars = GCachedScalabilityCVars;
 
-		FGlobalComponentReregisterContext PropagateDetailModeChanges(ExcludeComponents);
-		GCachedScalabilityCVars.DetailMode = DetailMode->GetValueOnGameThread();
+	{
+		static const auto DetailMode = ConsoleMan.FindTConsoleVariableDataInt(TEXT("r.DetailMode"));
+		LocalScalabilityCVars.DetailMode = DetailMode->GetValueOnGameThread();
 	}
 
-	static const auto* MaxAnisotropy = ConsoleMan.FindTConsoleVariableDataInt(TEXT("r.MaxAnisotropy"));
-	static const auto* MaxShadowResolution = ConsoleMan.FindTConsoleVariableDataInt(TEXT("r.Shadow.MaxResolution"));
-	static const auto ViewDistanceScale = ConsoleMan.FindTConsoleVariableDataFloat(TEXT("r.ViewDistanceScale"));
-	GCachedScalabilityCVars.MaxAnisotropy = MaxAnisotropy->GetValueOnGameThread();
-	GCachedScalabilityCVars.MaxShadowResolution = MaxShadowResolution->GetValueOnGameThread();
-	GCachedScalabilityCVars.ViewDistanceScale = FMath::Clamp(ViewDistanceScale->GetValueOnGameThread(), 0.0f, 1.0f);
-	GCachedScalabilityCVars.ViewDistanceScaleSquared = FMath::Square(GCachedScalabilityCVars.ViewDistanceScale);
-	GCachedScalabilityCVars.GaussianDOFNearThreshold = CVarDepthOfFieldNearBlurSizeThreshold.GetValueOnGameThread();
+	{
+		static const auto* MaxAnisotropy = ConsoleMan.FindTConsoleVariableDataInt(TEXT("r.MaxAnisotropy"));
+		LocalScalabilityCVars.MaxAnisotropy = MaxAnisotropy->GetValueOnGameThread();
+	}
 
-	// action needed if we change r.MaterialQualityLevel at runtime
 	{
-		static const auto MaterialQualityLevelVar = IConsoleManager::Get().FindTConsoleVariableDataInt(TEXT("r.MaterialQualityLevel"));
+		static const auto* MaxShadowResolution = ConsoleMan.FindTConsoleVariableDataInt(TEXT("r.Shadow.MaxResolution"));
+		LocalScalabilityCVars.MaxShadowResolution = MaxShadowResolution->GetValueOnGameThread();
+	}
 
-		EMaterialQualityLevel::Type NewMaterialQualityLevel = (EMaterialQualityLevel::Type)FMath::Clamp(MaterialQualityLevelVar->GetValueOnGameThread(), 0, 1);
+	{
+		static const auto ViewDistanceScale = ConsoleMan.FindTConsoleVariableDataFloat(TEXT("r.ViewDistanceScale"));
+		LocalScalabilityCVars.ViewDistanceScale = FMath::Clamp(ViewDistanceScale->GetValueOnGameThread(), 0.0f, 1.0f);
+		LocalScalabilityCVars.ViewDistanceScaleSquared = FMath::Square(LocalScalabilityCVars.ViewDistanceScale);
+	}
 
-		// has the state changed ?
-		if(GCachedScalabilityCVars.MaterialQualityLevel != NewMaterialQualityLevel)
-		{
-			// we had a state before?
-			if(GCachedScalabilityCVars.MaterialQualityLevel != EMaterialQualityLevel::Type::Num)
-			{
-				// state has changed, some action is needed
+	{
+		LocalScalabilityCVars.GaussianDOFNearThreshold = CVarDepthOfFieldNearBlurSizeThreshold.GetValueOnGameThread();
+	}
 
-				// Deregister all components
-				FGlobalComponentReregisterContext RecreateComponents;
+	{
+		static const auto MaterialQualityLevelVar = IConsoleManager::Get().FindTConsoleVariableDataInt(TEXT("r.MaterialQualityLevel"));
+		LocalScalabilityCVars.MaterialQualityLevel = (EMaterialQualityLevel::Type)FMath::Clamp(MaterialQualityLevelVar->GetValueOnGameThread(), 0, 1);
+	}
 
-				// after FGlobalComponentReregisterContext to have the renderthread flushed before so it can use the variable on either thread
-				GCachedScalabilityCVars.MaterialQualityLevel = NewMaterialQualityLevel;
+	{
+		static IConsoleVariable* CVar = IConsoleManager::Get().FindConsoleVariable(TEXT("r.SimpleDynamicLighting"));
+		LocalScalabilityCVars.SimpleDynamicLighting = FMath::Clamp(CVar->GetInt(), 0, 1);
+	}
 
-				// For all materials, UMaterial::CacheResourceShadersForRendering
-				UMaterial::AllMaterialsCacheResourceShadersForRendering();
-				UMaterialInstance::AllMaterialsCacheResourceShadersForRendering();
+	LocalScalabilityCVars.bInitialized = true;
 
-				// destructor of RecreateComponents will register the components again
-			}
-			else
-			{
-				GCachedScalabilityCVars.MaterialQualityLevel = NewMaterialQualityLevel;
-			}
-		}
+	if (!GCachedScalabilityCVars.bInitialized)
+	{
+		// optimization: the first time we assume the render thread wasn't started and we don't need to destroy proxies
+		GCachedScalabilityCVars = LocalScalabilityCVars;
 	}
-
-	// action needed if we change r.SimpleDynamicLighting at runtime
+	else
 	{
-		static IConsoleVariable* CVar = IConsoleManager::Get().FindConsoleVariable(TEXT("r.SimpleDynamicLighting"));
+		bool bRecreateRenderstate = false;
+		bool bCacheResourceShaders = false;
 
-		// 0:off, 1:on, -1:unknown
-		static int32 CurrentSDL = -1;
+		if (LocalScalabilityCVars.DetailMode != GCachedScalabilityCVars.DetailMode ||
+			LocalScalabilityCVars.SimpleDynamicLighting != GCachedScalabilityCVars.SimpleDynamicLighting)
+		{
+			bRecreateRenderstate = true;
+		}
 
-		int32 NewSDL = FMath::Clamp(CVar->GetInt(), 0, 1);
+		if (LocalScalabilityCVars.MaterialQualityLevel != GCachedScalabilityCVars.MaterialQualityLevel)
+		{
+			bCacheResourceShaders = true;
+		}
 
-		// has the state changed ?
-		if(CurrentSDL != NewSDL)
+		if (bRecreateRenderstate || bCacheResourceShaders)
 		{
-			// we had a state before?
-			if(CurrentSDL != -1)
-			{
-				CurrentSDL = NewSDL;
+			FlushRenderingCommands();
 
-				// state has changed, some action is needed
+			// after FlushRenderingCommands() to not have render thread pick up the data partially
+			GCachedScalabilityCVars = LocalScalabilityCVars;
 
-				// Deregister all components
-				FGlobalComponentReregisterContext RecreateComponents;
+			// Note: constructor and destructor has side effect
+			FGlobalComponentRecreateRenderStateContext Recreate;
 
-				// destructor of RecreateComponents will register the components again
-			}
-			else
+			if (bCacheResourceShaders)
 			{
-				CurrentSDL = NewSDL;
+				// For all materials, UMaterial::CacheResourceShadersForRendering
+				UMaterial::AllMaterialsCacheResourceShadersForRendering();
+				UMaterialInstance::AllMaterialsCacheResourceShadersForRendering();
 			}
 		}
 	}
diff --git a/Engine/Source/Runtime/Engine/Public/UnrealEngine.h b/Engine/Source/Runtime/Engine/Public/UnrealEngine.h
--- a/Engine/Source/Runtime/Engine/Public/UnrealEngine.h
+++ b/Engine/Source/Runtime/Engine/Public/UnrealEngine.h
@@ -398,9 +398,13 @@ public:
 	virtual void ClearPendingCleanupObjects() = 0;
 };
 
-/** Cache some of the scalability CVars to avoid some virtual function calls and to detect changes and act accordingly if needed.*/
+/**
+ * Cache some of the scalability CVars to avoid some virtual function calls and to detect changes and act accordingly if needed.
+ * read by rendering thread[s], written by main thread, uses FlushRenderingCommands() to avoid conflict
+ */
 struct FCachedSystemScalabilityCVars
 {
+	bool bInitialized;
 	int32 DetailMode;
 	EMaterialQualityLevel::Type MaterialQualityLevel;
 	int32 MaxAnisotropy;
@@ -408,6 +412,8 @@ struct FCachedSystemScalabilityCVars
 	float ViewDistanceScale;
 	float ViewDistanceScaleSquared;
 	float GaussianDOFNearThreshold;
+	// 0:off, 1:on, -1:unknown
+	int32 SimpleDynamicLighting;
 
 	FCachedSystemScalabilityCVars();
 };
