Stop subscribing ALevelBounds CDO to engine events

diff --git a/Engine/Source/Runtime/Engine/Private/LevelBounds.cpp b/Engine/Source/Runtime/Engine/Private/LevelBounds.cpp
--- a/Engine/Source/Runtime/Engine/Private/LevelBounds.cpp
+++ b/Engine/Source/Runtime/Engine/Private/LevelBounds.cpp
@@ -96,13 +96,20 @@ void ALevelBounds::PostRegisterAllComponents()
 {
 	Super::PostRegisterAllComponents();
 	
-	GetLevel()->LevelBoundsActor = this;
-	SubscribeToUpdateEvents();
+	if (!IsTemplate())
+	{
+		GetLevel()->LevelBoundsActor = this;
+		SubscribeToUpdateEvents();
+	}
 }
 
 void ALevelBounds::PostUnregisterAllComponents()
 {
-	UnsubscribeFromUpdateEvents();
+	if (!IsTemplate())
+	{
+		UnsubscribeFromUpdateEvents();
+	}
+	
 	Super::PostUnregisterAllComponents();
 }
 
