World composition automatic origin rebasing bahavior (distance, 3D) can be changed in config files

diff --git a/Engine/Config/BaseEngine.ini b/Engine/Config/BaseEngine.ini
--- a/Engine/Config/BaseEngine.ini
+++ b/Engine/Config/BaseEngine.ini
@@ -2155,9 +2155,6 @@ DefaultASTCQualityBySize=3
 DefaultGravityZ=-980.0
 bEnable2DPhysics=false
 
-[/Script/Engine.WorldComposition]
-TilesStreamingTimeThreshold=1.0
-
 [/Script/WindowsTargetPlatform.WindowsTargetSettings]
 +TargetedRHIs=PCD3D_SM5
 +TargetedRHIs=PCD3D_SM4
diff --git a/Engine/Source/Runtime/Engine/Classes/Engine/WorldComposition.h b/Engine/Source/Runtime/Engine/Classes/Engine/WorldComposition.h
--- a/Engine/Source/Runtime/Engine/Classes/Engine/WorldComposition.h
+++ b/Engine/Source/Runtime/Engine/Classes/Engine/WorldComposition.h
@@ -221,4 +221,11 @@ public:
 	UPROPERTY(config)
 	bool						bLoadAllTilesDuringCinematic;
 
+	// Whether to rebase origin in 3D space, otherwise only on XY plane
+	UPROPERTY(config)
+	bool						bRebaseOriginIn3DSpace;
+
+	// Maximum distance to current view point where we should initiate origin rebasing
+	UPROPERTY(config)
+	float						RebaseOriginDistance;
 };
\ No newline at end of file
diff --git a/Engine/Source/Runtime/Engine/Private/WorldComposition.cpp b/Engine/Source/Runtime/Engine/Private/WorldComposition.cpp
--- a/Engine/Source/Runtime/Engine/Private/WorldComposition.cpp
+++ b/Engine/Source/Runtime/Engine/Private/WorldComposition.cpp
@@ -17,6 +17,10 @@ UWorldComposition::FWorldCompositionChangedEvent UWorldComposition::WorldComposi
 
 UWorldComposition::UWorldComposition(const FObjectInitializer& ObjectInitializer)
 	: Super(ObjectInitializer)
+	, TilesStreamingTimeThreshold(1.0)
+	, bLoadAllTilesDuringCinematic(false)
+	, bRebaseOriginIn3DSpace(false)
+	, RebaseOriginDistance(HALF_WORLD_MAX1*0.5f)
 {
 }
 
@@ -756,11 +760,15 @@ void UWorldComposition::EvaluateWorldOriginLocation(const FVector& ViewLocation)
 	UWorld* OwningWorld = GetWorld();
 	
 	FVector Location = ViewLocation;
-	// Consider only XY plane
-	Location.Z = 0.f;
+
+	if (!bRebaseOriginIn3DSpace)
+	{
+		// Consider only XY plane
+		Location.Z = 0.f;
+	}
 		
 	// Request to shift world in case current view is quite far from current origin
-	if (Location.Size() > HALF_WORLD_MAX1*0.5f)
+	if (Location.Size() > RebaseOriginDistance)
 	{
 		OwningWorld->RequestNewWorldOrigin(FIntVector(Location.X, Location.Y, Location.Z) + OwningWorld->OriginLocation);
 	}
