Added ConfigRestartRequired metadata to all settings properties which require a restart in order for them to take effect when changed.

diff --git a/Engine/Plugins/2D/Paper2D/Source/Paper2D/Classes/PaperRuntimeSettings.h b/Engine/Plugins/2D/Paper2D/Source/Paper2D/Classes/PaperRuntimeSettings.h
--- a/Engine/Plugins/2D/Paper2D/Source/Paper2D/Classes/PaperRuntimeSettings.h
+++ b/Engine/Plugins/2D/Paper2D/Source/Paper2D/Classes/PaperRuntimeSettings.h
@@ -39,6 +39,6 @@ class PAPER2D_API UPaperRuntimeSettings : public UObject
 	bool bEnableSpriteAtlasGroups;
 
 	// Enables experimental *incomplete and unsupported* 2D terrain spline editing. Note: You need to restart the editor when enabling this setting for the change to fully take effect.
-	UPROPERTY(EditAnywhere, config, Category=Experimental)
+	UPROPERTY(EditAnywhere, config, Category=Experimental, meta=(ConfigRestartRequired=true))
 	bool bEnableTerrainSplineEditing;
 };
diff --git a/Engine/Source/Runtime/Engine/Classes/Engine/RendererSettings.h b/Engine/Source/Runtime/Engine/Classes/Engine/RendererSettings.h
--- a/Engine/Source/Runtime/Engine/Classes/Engine/RendererSettings.h
+++ b/Engine/Source/Runtime/Engine/Classes/Engine/RendererSettings.h
@@ -130,12 +130,14 @@ class ENGINE_API URendererSettings
 
 	UPROPERTY(config, EditAnywhere, Category=Textures, meta=(
 		ConsoleVariable="Compat.UseDXT5NormalMaps",DisplayName="Use DXT5 Normal Maps",
-		ToolTip="Whether to use DXT5 for normal maps, otherwise BC5 will be used, which is not supported on all hardware. Changing this setting requires restarting the editor."))
+		ToolTip="Whether to use DXT5 for normal maps, otherwise BC5 will be used, which is not supported on all hardware. Changing this setting requires restarting the editor.",
+		ConfigRestartRequired=true))
 	uint32 bUseDXT5NormalMaps:1;
 
 	UPROPERTY(config, EditAnywhere, Category=Lighting, meta=(
 		ConsoleVariable="r.AllowStaticLighting",
-		ToolTip="Whether to allow any static lighting to be generated and used, like lightmaps and shadowmaps. Games that only use dynamic lighting should set this to 0 to save some static lighting overhead. Changing this setting requires restarting the editor."))
+		ToolTip="Whether to allow any static lighting to be generated and used, like lightmaps and shadowmaps. Games that only use dynamic lighting should set this to 0 to save some static lighting overhead. Changing this setting requires restarting the editor.",
+		ConfigRestartRequired=true))
 	uint32 bAllowStaticLighting:1;
 
 	UPROPERTY(config, EditAnywhere, Category=Lighting, meta=(
@@ -145,7 +147,8 @@ class ENGINE_API URendererSettings
 
 	UPROPERTY(config, EditAnywhere, Category=Lighting, meta=(
 		ConsoleVariable="r.GenerateMeshDistanceFields",
-		ToolTip="Whether to build distance fields of static meshes, needed for distance field AO, which is used to implement Movable SkyLight shadows, and ray traced distance field shadows on directional lights.  Enabling will increase mesh build times and memory usage.  Changing this setting requires restarting the editor."))
+		ToolTip="Whether to build distance fields of static meshes, needed for distance field AO, which is used to implement Movable SkyLight shadows, and ray traced distance field shadows on directional lights.  Enabling will increase mesh build times and memory usage.  Changing this setting requires restarting the editor.",
+		ConfigRestartRequired=true))
 	uint32 bGenerateMeshDistanceFields:1;
 
 	UPROPERTY(config, EditAnywhere, Category = Lighting, meta = (
@@ -240,7 +243,8 @@ class ENGINE_API URendererSettings
 
 	UPROPERTY(config, EditAnywhere, Category=Optimizations, meta=(
 		ConsoleVariable="r.BasePassOutputsVelocity", DisplayName="Accurate velocities from Vertex Deformation",
-		ToolTip="Enables materials with time-based World Position Offset and/or World Displacement to output accurate velocities. This incurs a performance cost. If this is disabled, those materials will not output velocities. Changing this setting requires restarting the editor."))
+		ToolTip="Enables materials with time-based World Position Offset and/or World Displacement to output accurate velocities. This incurs a performance cost. If this is disabled, those materials will not output velocities. Changing this setting requires restarting the editor.",
+		ConfigRestartRequired=true))
 	uint32 bBasePassOutputsVelocity:1;
 
 	UPROPERTY(config, EditAnywhere, Category=Editor, meta=(
