Fixed curve editor so that the Alpha graph can be edited like the other channels.

diff --git a/Engine/Source/Runtime/Engine/Private/CurveLinearColor.cpp b/Engine/Source/Runtime/Engine/Private/CurveLinearColor.cpp
--- a/Engine/Source/Runtime/Engine/Private/CurveLinearColor.cpp
+++ b/Engine/Source/Runtime/Engine/Private/CurveLinearColor.cpp
@@ -88,6 +88,7 @@ bool UCurveLinearColor::IsValidCurve( FRichCurveEditInfo CurveInfo )
 {
 	return CurveInfo.CurveToEdit == &FloatCurves[0] ||
 		CurveInfo.CurveToEdit == &FloatCurves[1] ||
-		CurveInfo.CurveToEdit == &FloatCurves[2];
+		CurveInfo.CurveToEdit == &FloatCurves[2] ||
+		CurveInfo.CurveToEdit == &FloatCurves[3];
 }
 
