Bugfixes to selection box checking methods in StaticMeshComponent.

diff --git a/Engine/Source/Runtime/Engine/Private/Components/StaticMeshComponent.cpp b/Engine/Source/Runtime/Engine/Private/Components/StaticMeshComponent.cpp
--- a/Engine/Source/Runtime/Engine/Private/Components/StaticMeshComponent.cpp
+++ b/Engine/Source/Runtime/Engine/Private/Components/StaticMeshComponent.cpp
@@ -1729,9 +1729,10 @@ bool UStaticMeshComponent::ComponentIsTouchingSelectionBox(const FBox& InSelBBox
 				}
 			}
 
-			// If the selection box has to encompass all of the component and none of the component's verts failed the intersection test, this component
-			// is consider touching
-			return true;
+			// Either:
+			// a) It must encompass the entire component and all points were intersected (return true), or;
+			// b) It needn't encompass the entire component but no points were intersected (return false)
+			return bMustEncompassEntireComponent;
 		}
 	}
 
@@ -1766,7 +1767,10 @@ bool UStaticMeshComponent::ComponentIsTouchingSelectionFrustum(const FConvexVolu
 				}
 			}
 
-			return true;
+			// Either:
+			// a) It must encompass the entire component and all points were intersected (return true), or;
+			// b) It needn't encompass the entire component but no points were intersected (return false)
+			return bMustEncompassEntireComponent;
 		}
 	}
 
