Fixed crash when Randomize Without Replacement is set to false on a Random Sound Node.

diff --git a/Engine/Source/Runtime/Engine/Private/SoundNodeRandom.cpp b/Engine/Source/Runtime/Engine/Private/SoundNodeRandom.cpp
--- a/Engine/Source/Runtime/Engine/Private/SoundNodeRandom.cpp
+++ b/Engine/Source/Runtime/Engine/Private/SoundNodeRandom.cpp
@@ -70,10 +70,7 @@ void USoundNodeRandom::ParseNodes( FAudioDevice* AudioDevice, const UPTRINT Node
 	RETRIEVE_SOUNDNODE_PAYLOAD( sizeof( int32 ) );
 	DECLARE_SOUNDNODE_ELEMENT( int32, NodeIndex );
 
-	if( bRandomizeWithoutReplacement == true )
-	{
-		FixHasBeenUsedArray();  // for now prob need this until resave packages has occurred
-	}
+	FixHasBeenUsedArray();  // for now prob need this until resave packages has occurred
 
 #if WITH_EDITOR
 	bool bIsPIESound = (GEditor != nullptr) && ((GEditor->bIsSimulatingInEditor || GEditor->PlayWorld != NULL) && ActiveSound.World != NULL);
