Fixed bounding box of spline mesh convex collision elements

diff --git a/Engine/Source/Runtime/Engine/Private/Components/SplineMeshComponent.cpp b/Engine/Source/Runtime/Engine/Private/Components/SplineMeshComponent.cpp
--- a/Engine/Source/Runtime/Engine/Private/Components/SplineMeshComponent.cpp
+++ b/Engine/Source/Runtime/Engine/Private/Components/SplineMeshComponent.cpp
@@ -924,6 +924,8 @@ void USplineMeshComponent::RecreateCollision()
 				{
 					Point = CalcSliceTransform(GetAxisValue(Point, ForwardAxis)).TransformPosition(Point * Mask);
 				}
+
+				ConvexElem.UpdateElemBox();
 			}
 		}
 
