UE4 - make bounds empty for empty ISMC, github 998, UE-12981

diff --git a/Engine/Source/Runtime/Engine/Private/InstancedStaticMesh.cpp b/Engine/Source/Runtime/Engine/Private/InstancedStaticMesh.cpp
--- a/Engine/Source/Runtime/Engine/Private/InstancedStaticMesh.cpp
+++ b/Engine/Source/Runtime/Engine/Private/InstancedStaticMesh.cpp
@@ -1078,10 +1078,10 @@ bool UInstancedStaticMeshComponent::CanEditSimulatePhysics()
 
 FBoxSphereBounds UInstancedStaticMeshComponent::CalcBounds(const FTransform& BoundTransform) const
 {
-	FMatrix BoundTransformMatrix = BoundTransform.ToMatrixWithScale();
-
 	if(StaticMesh && PerInstanceSMData.Num() > 0)
 	{
+		FMatrix BoundTransformMatrix = BoundTransform.ToMatrixWithScale();
+
 		FBoxSphereBounds RenderBounds = StaticMesh->GetBounds();
 		FBoxSphereBounds NewBounds = RenderBounds.TransformBy(PerInstanceSMData[0].Transform * BoundTransformMatrix);
 
@@ -1094,7 +1094,7 @@ FBoxSphereBounds UInstancedStaticMeshComponent::CalcBounds(const FTransform& Bou
 	}
 	else
 	{
-		return Super::CalcBounds(BoundTransform);
+		return FBoxSphereBounds(BoundTransform.GetLocation(), FVector::ZeroVector, 0.f);
 	}
 }
 
