Tweak how pending kill loop is evaluated when unloading a world

diff --git a/Engine/Source/Runtime/Engine/Private/UnrealEngine.cpp b/Engine/Source/Runtime/Engine/Private/UnrealEngine.cpp
--- a/Engine/Source/Runtime/Engine/Private/UnrealEngine.cpp
+++ b/Engine/Source/Runtime/Engine/Private/UnrealEngine.cpp
@@ -8920,7 +8920,7 @@ bool UEngine::LoadMap( FWorldContext& WorldContext, FURL URL, class UPendingNetG
 		WorldContext.World()->RemoveFromRoot();
 
 		// mark everything else contained in the world to be deleted
-		TArray<UWorld*> CurrentWorlds;
+		TSet<UWorld*> CurrentWorlds;
 		for (auto LevelIt(WorldContext.World()->GetLevelIterator()); LevelIt; ++LevelIt)
 		{
 			const ULevel* Level = *LevelIt;
@@ -8930,15 +8930,12 @@ bool UEngine::LoadMap( FWorldContext& WorldContext, FURL URL, class UPendingNetG
 			}
 		}
 
-		for (TObjectIterator<UObject> It; It; ++It)
+		for (TObjectIterator<UObject> It(RF_PendingKill); It; ++It)
 		{
-			for (const UWorld* World : CurrentWorlds)
+			UWorld* InWorld = It->GetTypedOuter<UWorld>();
+			if (InWorld && CurrentWorlds.Contains(InWorld))
 			{
-				if (It->IsIn(World))
-				{
-					It->MarkPendingKill();
-					break;
-				}
+				It->MarkPendingKill();
 			}
 		}
 
diff --git a/Engine/Source/Runtime/Engine/Private/World.cpp b/Engine/Source/Runtime/Engine/Private/World.cpp
--- a/Engine/Source/Runtime/Engine/Private/World.cpp
+++ b/Engine/Source/Runtime/Engine/Private/World.cpp
@@ -4755,7 +4755,7 @@ UWorld* FSeamlessTravelHandler::Tick()
 			GWorld = NULL;
 
 			// mark everything else contained in the world to be deleted
-			TArray<UWorld*> CurrentWorlds;
+			TSet<UWorld*> CurrentWorlds;
 			for (auto LevelIt(CurrentWorld->GetLevelIterator()); LevelIt; ++LevelIt)
 			{
 				const ULevel* Level = *LevelIt;
@@ -4765,15 +4765,12 @@ UWorld* FSeamlessTravelHandler::Tick()
 				}
 			}
 
-			for (TObjectIterator<UObject> It; It; ++It)
+			for (TObjectIterator<UObject> It(RF_PendingKill); It; ++It)
 			{
-				for (const UWorld* World : CurrentWorlds)
+				UWorld* InWorld = It->GetTypedOuter<UWorld>();
+				if (InWorld && CurrentWorlds.Contains(InWorld))
 				{
-					if (It->IsIn(World))
-					{
-						It->MarkPendingKill();
-						break;
-					}
+					It->MarkPendingKill();
 				}
 			}
 
