Fixed crash when texture painting through Paint Mode from regression due to CL 2469910.

The change to add FCanvasTriangleItem::MaterialRenderProxy was broken - an FCanvasTriangleItem with both null Texture and MaterialRenderProxy would go down the path expecting a MaterialRenderProxy (and subsequently fail).

diff --git a/Engine/Source/Runtime/Engine/Private/UserInterface/CanvasItem.cpp b/Engine/Source/Runtime/Engine/Private/UserInterface/CanvasItem.cpp
--- a/Engine/Source/Runtime/Engine/Private/UserInterface/CanvasItem.cpp
+++ b/Engine/Source/Runtime/Engine/Private/UserInterface/CanvasItem.cpp
@@ -1285,7 +1285,7 @@ void FCanvasBoxItem::SetupBox()
 
 void FCanvasTriangleItem::Draw( class FCanvas* InCanvas )
 {
-	if (Texture != nullptr)
+	if (MaterialRenderProxy == nullptr)
 	{
 		SCOPE_CYCLE_COUNTER(STAT_Canvas_TriTextureItemTime);
 		FBatchedElements* BatchedElements = InCanvas->GetBatchedElements(FCanvas::ET_Triangle, BatchedElementParameters, Texture, BlendMode);
