UE4 - Foliage culled in ortho viewports, UE-13152

diff --git a/Engine/Source/Runtime/Engine/Private/HierarchicalInstancedStaticMesh.cpp b/Engine/Source/Runtime/Engine/Private/HierarchicalInstancedStaticMesh.cpp
--- a/Engine/Source/Runtime/Engine/Private/HierarchicalInstancedStaticMesh.cpp
+++ b/Engine/Source/Runtime/Engine/Private/HierarchicalInstancedStaticMesh.cpp
@@ -1232,6 +1232,7 @@ void FHierarchicalStaticMeshSceneProxy::GetDynamicMeshElements(const TArray<cons
 
 				FMatrix WorldToLocal = GetLocalToWorld().Inverse();
 				bool bUseVectorCull = GUseVectorCull;
+				bool bIsOrtho = false;
 
 				bool bDisableCull = !!CVarDisableCull.GetValueOnRenderThread();
 				ElementParams.ShadowFrustum = !!View->ViewMatrices.GetDynamicMeshElementsShadowCullFrustum;
@@ -1274,6 +1275,7 @@ void FHierarchicalStaticMeshSceneProxy::GetDynamicMeshElements(const TArray<cons
 					}
 					else
 					{
+						bIsOrtho = true;
 						bUseVectorCull = false;
 					}
 				}
@@ -1291,7 +1293,7 @@ void FHierarchicalStaticMeshSceneProxy::GetDynamicMeshElements(const TArray<cons
 				InstanceParams.ViewOriginInLocalZero = WorldToLocal.TransformPosition(View->GetTemporalLODOrigin(0, bMultipleSections));
 				InstanceParams.ViewOriginInLocalOne = WorldToLocal.TransformPosition(View->GetTemporalLODOrigin(1, bMultipleSections));
 
-				float MinSize = CVarFoliageMinimumScreenSize.GetValueOnRenderThread();
+				float MinSize = bIsOrtho ? 0.0f : CVarFoliageMinimumScreenSize.GetValueOnRenderThread();
 				float LODScale = CVarFoliageLODDistanceScale.GetValueOnRenderThread();
 				float LODRandom = CVarRandomLODRange.GetValueOnRenderThread();
 				float MaxDrawDistanceScale = GetCachedScalabilityCVars().ViewDistanceScale;
diff --git a/Engine/Source/Runtime/Engine/Private/InstancedStaticMesh.cpp b/Engine/Source/Runtime/Engine/Private/InstancedStaticMesh.cpp
--- a/Engine/Source/Runtime/Engine/Private/InstancedStaticMesh.cpp
+++ b/Engine/Source/Runtime/Engine/Private/InstancedStaticMesh.cpp
@@ -1725,7 +1725,7 @@ void FInstancedStaticMeshVertexFactoryShaderParameters::SetMesh( FRHICommandList
 
 
 			float SphereRadius = InstancingUserData->MeshRenderData->Bounds.SphereRadius;
-			float MinSize = CVarFoliageMinimumScreenSize.GetValueOnRenderThread();
+			float MinSize = View.ViewMatrices.IsPerspectiveProjection() ? CVarFoliageMinimumScreenSize.GetValueOnRenderThread() : 0.0f;
 			float LODScale = CVarFoliageLODDistanceScale.GetValueOnRenderThread();
 			float LODRandom = CVarRandomLODRange.GetValueOnRenderThread();
 			float MaxDrawDistanceScale = GetCachedScalabilityCVars().ViewDistanceScale;
