UE4 - fix GlobalBoundShaderState for PS4 and rework error testing

diff --git a/Engine/Source/Runtime/Engine/Private/StaticBoundShaderState.cpp b/Engine/Source/Runtime/Engine/Private/StaticBoundShaderState.cpp
--- a/Engine/Source/Runtime/Engine/Private/StaticBoundShaderState.cpp
+++ b/Engine/Source/Runtime/Engine/Private/StaticBoundShaderState.cpp
@@ -13,8 +13,12 @@ TLinkedList<FGlobalBoundShaderStateResource*>*& FGlobalBoundShaderStateResource:
 	return List;
 }
 
-FGlobalBoundShaderStateResource::FGlobalBoundShaderStateResource():
-	GlobalListLink(this)
+FGlobalBoundShaderStateResource::FGlobalBoundShaderStateResource()
+	: GlobalListLink(this)
+	, BoundVertexDeclaration(nullptr)
+	, BoundVertexShader(nullptr)
+	, BoundPixelShader(nullptr)
+	, BoundGeometryShader(nullptr)
 {
 	// Add this resource to the global list in the rendering thread.
 	if(IsInRenderingThread())
@@ -56,6 +60,12 @@ FBoundShaderStateRHIParamRef FGlobalBoundShaderStateResource::GetInitializedRHI(
 	// Create the bound shader state if it hasn't been cached yet.
 	if(!IsValidRef(BoundShaderState))
 	{
+#if DO_CHECK
+		BoundVertexDeclaration = VertexDeclaration;
+		BoundVertexShader = VertexShader;
+		BoundPixelShader = PixelShader;
+		BoundGeometryShader = GeometryShader;
+#endif 
 		BoundShaderState = 
 			RHICreateBoundShaderState(
 			VertexDeclaration,
@@ -65,25 +75,15 @@ FBoundShaderStateRHIParamRef FGlobalBoundShaderStateResource::GetInitializedRHI(
 			PixelShader,
 			GeometryShader);
 	}
-#if !(UE_BUILD_SHIPPING || UE_BUILD_TEST)
-	// Only people working on shaders (and therefore have LogShaders unsuppressed) will want to see these errors
-	else if (!GUsingNullRHI && UE_LOG_ACTIVE(LogShaders, Warning))
-	{
-		FBoundShaderStateRHIRef TempBoundShaderState = 
-			RHICreateBoundShaderState(
-			VertexDeclaration,
-			VertexShader,
-			FHullShaderRHIRef(),
-			FDomainShaderRHIRef(),
-			PixelShader,
-			GeometryShader);
-		// Verify that bound shader state caching is working and that the passed in shaders will actually be used
-		// This will catch cases where one bound shader state is being used with more than one combination of shaders
-		// Otherwise setting the shader will just silently fail once the bound shader state has been initialized with a different shader 
-		// The catch is that this uses the caching mechanism to test for equality, instead of comparing the actual platform dependent shader references.
-		check(TempBoundShaderState == BoundShaderState);
-	}
-#endif
+#if DO_CHECK
+	// Verify that the passed in shaders will actually be used
+	// This will catch cases where one bound shader state is being used with more than one combination of shaders
+	// Otherwise setting the shader will just silently fail once the bound shader state has been initialized with a different shader 
+	check(BoundVertexDeclaration == VertexDeclaration &&
+		BoundVertexShader == VertexShader &&
+		BoundPixelShader == PixelShader &&
+		BoundGeometryShader == GeometryShader);
+#endif 
 
 	return BoundShaderState;
 }
diff --git a/Engine/Source/Runtime/Engine/Public/StaticBoundShaderState.h b/Engine/Source/Runtime/Engine/Public/StaticBoundShaderState.h
--- a/Engine/Source/Runtime/Engine/Public/StaticBoundShaderState.h
+++ b/Engine/Source/Runtime/Engine/Public/StaticBoundShaderState.h
@@ -54,6 +54,13 @@ private:
 
 	// FRenderResource interface.
 	ENGINE_API virtual void ReleaseRHI();
+
+#if DO_CHECK
+	FVertexDeclarationRHIParamRef BoundVertexDeclaration;
+	FVertexShaderRHIParamRef BoundVertexShader;
+	FPixelShaderRHIParamRef BoundPixelShader;
+	FGeometryShaderRHIParamRef BoundGeometryShader;
+#endif 
 };
 
 typedef TGlobalResource<FGlobalBoundShaderStateResource> FGlobalBoundShaderState_Internal;
diff --git a/Engine/Source/Runtime/RHI/Public/RHIMethods.h b/Engine/Source/Runtime/RHI/Public/RHIMethods.h
--- a/Engine/Source/Runtime/RHI/Public/RHIMethods.h
+++ b/Engine/Source/Runtime/RHI/Public/RHIMethods.h
@@ -470,7 +470,8 @@ DEFINE_RHIMETHOD_GLOBAL_1(
 
 /**
  * Creates a bound shader state instance which encapsulates a decl, vertex shader, hull shader, domain shader and pixel shader
- * CAUTION: Even though this is marked as threadsafe, it is only valid to call from the render thread or the RHI thread. Thus is need not be threadsafe on platforms that do not support or aren't using an RHIThread
+ * CAUTION: Even though this is marked as threadsafe, it is only valid to call from the render thread or the RHI thread. It need not be threadsafe unless the RHI support parallel translation.
+ * CAUTION: Platforms that support RHIThread but don't actually have a threadsafe implementation must flush internally with FRHICommandListExecutor::GetImmediateCommandList().ImmediateFlush(EImmediateFlushType::WaitForRHIThread); when the call is from the render thread
  * @param VertexDeclaration - existing vertex decl
  * @param VertexShader - existing vertex shader
  * @param HullShader - existing hull shader
