Fix destructible component ignoring the startAwake and gravity flag. This was turned off during the 3.3.3 upgrade

UE-11214

diff --git a/Engine/Source/Runtime/Engine/Private/Components/DestructibleComponent.cpp b/Engine/Source/Runtime/Engine/Private/Components/DestructibleComponent.cpp
--- a/Engine/Source/Runtime/Engine/Private/Components/DestructibleComponent.cpp
+++ b/Engine/Source/Runtime/Engine/Private/Components/DestructibleComponent.cpp
@@ -345,6 +345,7 @@ void UDestructibleComponent::CreatePhysicsState()
 	// Destructibles are always dynamic or kinematic, and therefore only go into one of the scenes
 	const uint32 SceneType = BodyInstance.UseAsyncScene(PhysScene) ? PST_Async : PST_Sync;
 	NxApexScene* ApexScene = PhysScene->GetApexScene(SceneType);
+	PxScene* PScene = PhysScene->GetPhysXScene(SceneType);
 	check(ApexScene);
 
 	// Create an APEX NxDestructibleActor from the Destructible asset and actor descriptor
@@ -370,23 +371,18 @@ void UDestructibleComponent::CreatePhysicsState()
 
 
 	//  Put to sleep or wake up only if the component is physics-simulated
-	//TODO: APEX 1.3.3 defers adding actors to scene which means this doesn't work. Waiting on flag from them so we can do this properly.
-	/*if (PRootActor != NULL && BodyInstance.bSimulatePhysics)
+	if (PRootActor != NULL && BodyInstance.bSimulatePhysics)
 	{
-		SCOPED_SCENE_WRITE_LOCK(PRootActor->getScene());
+		SCOPED_SCENE_WRITE_LOCK(PScene);	//Question, since apex is defer adding actors do we need to lock? Locking the async scene is expensive!
 
 		PRootActor->setActorFlag(PxActorFlag::eDISABLE_GRAVITY, !BodyInstance.bEnableGravity);
 
 		// Sleep/wake up as appropriate
-		if (BodyInstance.bStartAwake)
+		if (!BodyInstance.bStartAwake)
 		{
-			PRootActor->wakeUp();
+			ApexDestructibleActor->setChunkPhysXActorAwakeState(0, false);
 		}
-		else
-		{
-			PRootActor->putToSleep();
-		}
-	}*/
+	}
 
 	UpdateBounds();
 #endif	// #if WITH_APEX
