- 1st pass function calls in Niagara. Function call nodes link to other scripts. The graphs from those scripts are merged into the main graph and everything is compiled as normal.
- Refactored NiagaraSimulation.cpp a little and removed the requirement for Spawn and Update graphs to have the same attributes.
- Changed colour of attribtues to Green so that function calls could be Blue. Changed Systems Constant nodes to be White.
- Fixed a bug with the Lerp operation.

diff --git a/Engine/Source/Editor/NiagaraEditor/Private/NiagaraCompiler.cpp b/Engine/Source/Editor/NiagaraEditor/Private/NiagaraCompiler.cpp
--- a/Engine/Source/Editor/NiagaraEditor/Private/NiagaraCompiler.cpp
+++ b/Engine/Source/Editor/NiagaraEditor/Private/NiagaraCompiler.cpp
@@ -303,6 +303,21 @@ TNiagaraExprPtr FNiagaraCompiler::GetExternalConstant(const FNiagaraVariableInfo
 	SetOrAddConstant(false, Constant, Default);
 	return Expression_GetExternalConstant(Constant);
 }
+TNiagaraExprPtr FNiagaraCompiler::GetInternalConstant(const FNiagaraVariableInfo& Constant, float Default)
+{
+	SetOrAddConstant(true, Constant, Default);
+	return Expression_GetInternalConstant(Constant);
+}
+TNiagaraExprPtr FNiagaraCompiler::GetInternalConstant(const FNiagaraVariableInfo& Constant, const FVector4& Default)
+{
+	SetOrAddConstant(true, Constant, Default);
+	return Expression_GetInternalConstant(Constant);
+}
+TNiagaraExprPtr FNiagaraCompiler::GetInternalConstant(const FNiagaraVariableInfo& Constant, const FMatrix& Default)
+{
+	SetOrAddConstant(true, Constant, Default);
+	return Expression_GetInternalConstant(Constant);
+}
 
 TNiagaraExprPtr FNiagaraCompiler::GetExternalCurveConstant(const FNiagaraVariableInfo& Constant)
 {
@@ -310,4 +325,309 @@ TNiagaraExprPtr FNiagaraCompiler::GetExternalCurveConstant(const FNiagaraVariabl
 	return Expression_GetExternalConstant(Constant);
 }
 
+// TODO: Refactor this (and some other stuff) into a preprocessing step for use by any compiler?
+bool FNiagaraCompiler::MergeInFunctionNodes()
+{
+	struct FReconnectionInfo
+	{
+	public:
+		UEdGraphPin* From;
+		TArray<UEdGraphPin*> To;
+
+		//Fallback default value if an input connection is not connected.
+		FString FallbackDefault;
+
+		FReconnectionInfo()
+			: From(NULL)
+		{}
+	};
+	TMap<FName, FReconnectionInfo> InputConnections;
+	TMap<FName, FReconnectionInfo> OutputConnections;
+
+	TArray<class UEdGraphPin*> FuncCallInputPins;
+	TArray<class UEdGraphPin*> FuncCallOutputPins;
+
+	//Copies the function graph into the main graph.
+	//Removes the Function call in the main graph and the input and output nodes in the function graph, reconnecting their pins appropriately.
+	auto MergeFunctionIntoMainGraph = [&](UNiagaraNodeFunctionCall* InFunc, UNiagaraGraph* FuncGraph)
+	{
+		InputConnections.Empty();
+		OutputConnections.Empty();
+		FuncCallInputPins.Empty();
+		FuncCallOutputPins.Empty();
+
+		check(InFunc && FuncGraph);
+		if (InFunc->FunctionScript)
+		{
+			//Get all the pins that are connected to the inputs of the function call node in the main graph.
+			InFunc->GetInputPins(FuncCallInputPins);
+			for (UEdGraphPin* FuncCallInputPin : FuncCallInputPins)
+			{
+				FName InputName(*FuncCallInputPin->PinName);
+				FReconnectionInfo& InputConnection = InputConnections.FindOrAdd(InputName);
+				if (FuncCallInputPin->LinkedTo.Num() > 0)
+				{
+					check(FuncCallInputPin->LinkedTo.Num() == 1);
+					UEdGraphPin* LinkFrom = FuncCallInputPin->LinkedTo[0];
+					check(LinkFrom->Direction == EGPD_Output);
+					InputConnection.From = LinkFrom;
+				}
+				else
+				{
+					//This input has no link so we need the default value from the pin.
+					InputConnection.FallbackDefault = FuncCallInputPin->GetDefaultAsString();
+				}
+			}
+			//Get all the pins that are connected to the outputs of the function call node in the main graph.
+			InFunc->GetOutputPins(FuncCallOutputPins);
+			for (UEdGraphPin* FuncCallOutputPin : FuncCallOutputPins)
+			{
+				FName OutputName(*FuncCallOutputPin->PinName);
+				for (UEdGraphPin* LinkTo : FuncCallOutputPin->LinkedTo)
+				{
+					check(LinkTo->Direction == EGPD_Input);
+					FReconnectionInfo& OutputConnection = OutputConnections.FindOrAdd(OutputName);
+					OutputConnection.To.Add(LinkTo);
+				}
+			}
+
+			//Remove the function call node from the graph now that we have everything we need from it.
+			SourceGraph->RemoveNode(InFunc);
+
+			//Keep a list of the Input and Output nodes we see in the function graph so that we can remove (most of) them later.
+			TArray<UEdGraphNode*, TInlineAllocator<64>> ToRemove;
+
+			//Search the nodes in the function graph, finding any connections to input or output nodes.
+			for (UEdGraphNode* FuncGraphNode : FuncGraph->Nodes)
+			{
+				if (UNiagaraNodeInput* InputNode = Cast<UNiagaraNodeInput>(FuncGraphNode))
+				{
+					check(InputNode->Pins.Num() == 1);
+					//Get an array of "To" pins from one or more input nodes referencing each named input.
+					FReconnectionInfo& InputConnection = InputConnections.FindOrAdd(InputNode->Input.Name);
+					if (InputConnection.From)
+					{
+						//We have a connection from the function call so remove the input node and connect to that.
+						ToRemove.Add(InputNode);
+					}
+					else
+					{
+						//This input has no connection from the function call so what do we do here? 
+						//For now we just leave the input node and connect back to it. 
+						//This will mean unconnected pins from the function call will look for constants or attributes. 
+						//In some cases we may want to just take the default value from the function call pin instead?
+						//Maybe have some properties on the function call defining that.
+						InputConnection.From = InputNode->Pins[0];
+					}
+
+					TArray<UEdGraphPin*>& LinkToPins = InputNode->Pins[0]->LinkedTo;
+					for (UEdGraphPin* ToPin : LinkToPins)
+					{
+						check(ToPin->Direction == EGPD_Input);
+						InputConnection.To.Add(ToPin);
+					}
+				}
+				else if (UNiagaraNodeOutput* OutputNode = Cast<UNiagaraNodeOutput>(FuncGraphNode))
+				{
+					//Unlike the input nodes, we don't have the option of keeping these if there is no "From" pin. The default values from the node pins should be used.
+					ToRemove.Add(OutputNode);
+
+					//For each output, get the "From" pin to be reconnected later.
+					for (int32 OutputIdx = 0; OutputIdx < OutputNode->Outputs.Num(); ++OutputIdx)
+					{
+						FName OutputName = OutputNode->Outputs[OutputIdx].Name;
+
+						UEdGraphPin* OutputNodePin = OutputNode->Pins[OutputIdx];
+						check(OutputNodePin->LinkedTo.Num() <= 1);
+						FReconnectionInfo& OutputConnection = OutputConnections.FindOrAdd(OutputName);
+						UEdGraphPin* LinkFromPin = OutputNodePin->LinkedTo.Num() == 1 ? OutputNodePin->LinkedTo[0] : NULL;
+						if (LinkFromPin)
+						{
+							check(LinkFromPin->Direction == EGPD_Output);
+							OutputConnection.From = LinkFromPin;
+						}
+						else
+						{
+							//This output is not connected so links to it in the main graph must use it's default value.
+							OutputConnection.FallbackDefault = OutputNodePin->GetDefaultAsString();
+						}
+					}
+				}
+			}
+
+			//Remove all the In and Out nodes from the function graph.
+			for (UEdGraphNode* Remove : ToRemove)
+			{
+				FuncGraph->RemoveNode(Remove);
+			}
+
+			//Copy the nodes from the function graph over into the main graph.
+			FuncGraph->MoveNodesToAnotherGraph(SourceGraph, false);
+
+			//Finally, do all the reconnection.
+			auto MakeConnection = [&](FReconnectionInfo& Info)
+			{
+				for (UEdGraphPin* LinkTo : Info.To)
+				{
+					if (Info.From)
+					{
+						Info.From->MakeLinkTo(LinkTo);
+					}
+					else
+					{
+						LinkTo->DefaultValue = Info.FallbackDefault;
+					}
+				}
+			};
+			for (TPair<FName, FReconnectionInfo>& ReconnectInfo : InputConnections){ MakeConnection(ReconnectInfo.Value); }
+			for (TPair<FName, FReconnectionInfo>& ReconnectInfo : OutputConnections){ MakeConnection(ReconnectInfo.Value); }
+		}
+	};
+
+	//Helper struct for traversing nested function calls.
+	struct FFunctionContext
+	{
+		//True if this context's function has been merged into the main graph.
+		bool bProcessed;
+		//The index of this context into the ContextPool. 
+		int32 PoolIdx;
+		//Pointer back to the parent context for traversal.
+		FFunctionContext* Parent;
+		//The function call node for this function in the source/parent graph.
+		UNiagaraNodeFunctionCall* Function;
+		//The graph for this function that we are going to merge into the main graph.
+		UNiagaraGraph* FunctionGraph;
+		//The script from which the graph is copied. Used for re entrance check.
+		UNiagaraScript* Script;
+
+		//Contexts for function calls in this function graph.
+		TArray<FFunctionContext*, TInlineAllocator<64>> SubFunctionCalls;
+
+		FFunctionContext()
+			: bProcessed(false)
+			, PoolIdx(INDEX_NONE)
+			, Parent(NULL)
+			, Function(NULL)
+			, FunctionGraph(NULL)
+			, Script(NULL)
+		{
+		}
+
+		/**
+		We don't allow re-entrant functions as this would cause an infinite loop of merging in graphs.
+		Maybe in the future if we allow branching in the VM we can allow this.
+		*/
+		bool CheckForReentrance()const
+		{
+			UNiagaraNodeFunctionCall* Func = Function;
+			FFunctionContext* Curr = Parent;
+			while (Curr)
+			{
+				if (Curr->Script == Script)
+					return true;
+
+				Curr = Curr->Parent;
+			}
+			return false;
+		}
+
+		FString GetCallstack()const
+		{
+			FString Ret;
+			const FFunctionContext* Curr = this;
+			while (Curr)
+			{
+				if (Curr->Script)
+				{
+					Ret.Append(*(Curr->Script->GetPathName()));
+				}
+				else
+				{
+					Ret.Append(TEXT("Unknown"));
+				}
+
+				Ret.Append(TEXT("\n"));
+
+				Curr = Curr->Parent;
+			}
+			return Ret;
+		}
+	};
+
+	//A pool of contexts on the stack to avoid loads of needless, small heap allocations.
+	TArray<FFunctionContext, TInlineAllocator<512>> ContextPool;
+	ContextPool.Reserve(512);
+
+	FFunctionContext RootContext;
+	FFunctionContext* CurrentContext = &RootContext;
+	CurrentContext->FunctionGraph = SourceGraph;
+	CurrentContext->Script = Script;
+
+	//Depth first traversal of all function calls.
+	while (CurrentContext)
+	{
+		//Find any sub functions and process this function call.
+		if (!CurrentContext->bProcessed)
+		{
+			CurrentContext->bProcessed = true;
+
+			//Find any sub functions and check for re-entrance.
+			if (CurrentContext->FunctionGraph)
+			{
+				for (UEdGraphNode* Node : CurrentContext->FunctionGraph->Nodes)
+				{
+					UNiagaraNodeFunctionCall* FuncNode = Cast<UNiagaraNodeFunctionCall>(Node);
+					if (FuncNode)
+					{
+						int32 NewIdx = ContextPool.AddZeroed();
+						FFunctionContext* SubFuncContext = &ContextPool[NewIdx];
+						CurrentContext->SubFunctionCalls.Push(SubFuncContext);
+						SubFuncContext->Parent = CurrentContext;
+						SubFuncContext->Function = FuncNode;
+						SubFuncContext->PoolIdx = NewIdx;
+						SubFuncContext->Script = FuncNode->FunctionScript;
+
+						if (SubFuncContext->CheckForReentrance())
+						{
+							FString Callstack = SubFuncContext->GetCallstack();
+							MessageLog.Error(TEXT("Reentrant function call!\n%s"), *Callstack);
+							return false;
+						}
+
+						//Copy the function graph as we'll be modifying it as we merge in with the main graph.
+						UNiagaraScriptSource* FuncSource = CastChecked<UNiagaraScriptSource>(FuncNode->FunctionScript->Source);
+						check(FuncSource);
+						SubFuncContext->FunctionGraph = CastChecked<UNiagaraGraph>(FEdGraphUtilities::CloneGraph(FuncSource->NodeGraph, NULL, &MessageLog));
+					}
+				}
+			}
+
+			//Merge this function into the main graph now.
+			if (CurrentContext->Function && CurrentContext->FunctionGraph)
+			{
+				MergeFunctionIntoMainGraph(CurrentContext->Function, CurrentContext->FunctionGraph);
+			}
+		}
+
+		if (CurrentContext->SubFunctionCalls.Num() > 0)
+		{
+			//Move to the next sub function.
+			CurrentContext = CurrentContext->SubFunctionCalls.Pop();
+		}
+		else
+		{
+			//Done processing this function so remove it and move back to the parent.
+			if (CurrentContext->PoolIdx != INDEX_NONE)
+			{
+				CurrentContext->FunctionGraph->MarkPendingKill();
+
+				ContextPool.RemoveAtSwap(CurrentContext->PoolIdx);
+			}
+			CurrentContext = CurrentContext->Parent;
+		}
+	}
+
+	return true;
+}
+
 #undef LOCTEXT_NAMESPACE
\ No newline at end of file
diff --git a/Engine/Source/Editor/NiagaraEditor/Private/NiagaraCompiler.h b/Engine/Source/Editor/NiagaraEditor/Private/NiagaraCompiler.h
--- a/Engine/Source/Editor/NiagaraEditor/Private/NiagaraCompiler.h
+++ b/Engine/Source/Editor/NiagaraEditor/Private/NiagaraCompiler.h
@@ -50,6 +50,9 @@ protected:
 	template< typename T >
 	void SetOrAddConstant(bool bInternal, const FNiagaraVariableInfo& Constant, const T& Default);
 
+	/** Searches for Function Call nodes and merges their graphs into the main graph. */
+	bool MergeInFunctionNodes();
+
 public:
 
 	//Begin INiagaraCompiler Interface
@@ -59,6 +62,9 @@ public:
 	virtual TNiagaraExprPtr GetExternalConstant(const FNiagaraVariableInfo& Constant, float Default)override;
 	virtual TNiagaraExprPtr GetExternalConstant(const FNiagaraVariableInfo& Constant, const FVector4& Default)override;
 	virtual TNiagaraExprPtr GetExternalConstant(const FNiagaraVariableInfo& Constant, const FMatrix& Default)override;
+	virtual TNiagaraExprPtr GetInternalConstant(const FNiagaraVariableInfo& Constant, float Default)override;
+	virtual TNiagaraExprPtr GetInternalConstant(const FNiagaraVariableInfo& Constant, const FVector4& Default)override;
+	virtual TNiagaraExprPtr GetInternalConstant(const FNiagaraVariableInfo& Constant, const FMatrix& Default)override;
 	virtual TNiagaraExprPtr GetExternalCurveConstant(const FNiagaraVariableInfo& Constant)override;
 
 	virtual void CheckInputs(FName OpName, TArray<TNiagaraExprPtr>& Inputs)override;
@@ -84,4 +90,4 @@ void FNiagaraCompiler::SetOrAddConstant(bool bInternal, const FNiagaraVariableIn
 		ConstantData.SetOrAddInternal(Constant, Default);
 	else
 		ConstantData.SetOrAddExternal(Constant, Default);
-}
\ No newline at end of file
+}
diff --git a/Engine/Source/Editor/NiagaraEditor/Private/NiagaraCompiler_VectorVM.cpp b/Engine/Source/Editor/NiagaraEditor/Private/NiagaraCompiler_VectorVM.cpp
--- a/Engine/Source/Editor/NiagaraEditor/Private/NiagaraCompiler_VectorVM.cpp
+++ b/Engine/Source/Editor/NiagaraEditor/Private/NiagaraCompiler_VectorVM.cpp
@@ -198,10 +198,22 @@ void FNiagaraCompiler_VectorVM::CompileScript(UNiagaraScript* InScript)
 	Script = InScript;
 	Source = CastChecked<UNiagaraScriptSource>(InScript->Source);
 
+	//Should we roll our own message/error log and put it in a window somewhere?
+	MessageLog.SetSourceName(InScript->GetPathName());
+
 	// Clone the source graph so we can modify it as needed; merging in the child graphs
-	SourceGraph = CastChecked<UNiagaraGraph>(FEdGraphUtilities::CloneGraph(Source->NodeGraph, Source, &MessageLog));
+	SourceGraph = CastChecked<UNiagaraGraph>(FEdGraphUtilities::CloneGraph(Source->NodeGraph, NULL, &MessageLog));
 	FEdGraphUtilities::MergeChildrenGraphsIn(SourceGraph, SourceGraph, /*bRequireSchemaMatch=*/ true);
 
+	if (!MergeInFunctionNodes())
+	{ 
+		//TODO: Insert a good error reporting and compile failure system.
+		InScript->Attributes.Empty();
+		InScript->ByteCode.Empty();
+		InScript->ConstantData.Empty();
+		return;
+	}
+
 	TempRegisters.AddZeroed(VectorVM::NumTempRegisters);
 
 	// Find the output node.
diff --git a/Engine/Source/Editor/UnrealEd/Classes/Niagara/EdGraphSchema_Niagara.h b/Engine/Source/Editor/UnrealEd/Classes/Niagara/EdGraphSchema_Niagara.h
--- a/Engine/Source/Editor/UnrealEd/Classes/Niagara/EdGraphSchema_Niagara.h
+++ b/Engine/Source/Editor/UnrealEd/Classes/Niagara/EdGraphSchema_Niagara.h
@@ -68,8 +68,12 @@ class UEdGraphSchema_Niagara : public UEdGraphSchema
 	UNREALED_API void GetPinDefaultValue(UEdGraphPin* Pin, float& OutDefault)const;
 	UNREALED_API void GetPinDefaultValue(UEdGraphPin* Pin, FVector4& OutDefault)const;
 	UNREALED_API void GetPinDefaultValue(UEdGraphPin* Pin, FMatrix& OutDefault)const;
+	
+	bool IsSystemConstant(const FNiagaraVariableInfo& Variable)const;
 
 	static const FLinearColor NodeTitleColor_Attribute;
 	static const FLinearColor NodeTitleColor_Constant;
+	static const FLinearColor NodeTitleColor_SystemConstant;
+	static const FLinearColor NodeTitleColor_FunctionCall;
 };
 
diff --git a/Engine/Source/Editor/UnrealEd/Classes/Niagara/NiagaraNodeFunctionCall.h b/Engine/Source/Editor/UnrealEd/Classes/Niagara/NiagaraNodeFunctionCall.h
--- a/Engine/Source/Editor/UnrealEd/Classes/Niagara/NiagaraNodeFunctionCall.h
+++ b/Engine/Source/Editor/UnrealEd/Classes/Niagara/NiagaraNodeFunctionCall.h
@@ -1,18 +1,23 @@
+// Copyright 1998-2015 Epic Games, Inc. All Rights Reserved.
+
 #pragma once
 
 #include "NiagaraNodeFunctionCall.generated.h"
 
-UCLASS()
+UCLASS(MinimalAPI)
 class UNiagaraNodeFunctionCall : public UNiagaraNode
 {
 	GENERATED_BODY()
 
+public:
+
 	UPROPERTY(EditAnywhere, Category = "Function")
 	UNiagaraScript* FunctionScript;
 
 
 	//Begin UObject interface
 	virtual void PostLoad()override;
+	virtual void PostEditChangeProperty(struct FPropertyChangedEvent& PropertyChangedEvent)override;
 	//End UObject interface
 
 	// Begin UNiagaraNode interface
@@ -21,9 +26,9 @@ class UNiagaraNodeFunctionCall : public UNiagaraNode
 
 	// Begin EdGraphNode interface
 	virtual void AllocateDefaultPins() override;
-// 	virtual FText GetNodeTitle(ENodeTitleType::Type TitleType) const override;
-// 	virtual FText GetTooltipText() const override;
-// 	virtual FLinearColor GetNodeTitleColor() const override;
+	virtual FText GetNodeTitle(ENodeTitleType::Type TitleType) const override;
+	virtual FText GetTooltipText() const override;
+	virtual FLinearColor GetNodeTitleColor() const override;
 	// End EdGraphNode interface
 
 private:
diff --git a/Engine/Source/Editor/UnrealEd/Private/Niagara/EdGraphSchema_Niagara.cpp b/Engine/Source/Editor/UnrealEd/Private/Niagara/EdGraphSchema_Niagara.cpp
--- a/Engine/Source/Editor/UnrealEd/Private/Niagara/EdGraphSchema_Niagara.cpp
+++ b/Engine/Source/Editor/UnrealEd/Private/Niagara/EdGraphSchema_Niagara.cpp
@@ -13,8 +13,10 @@
 
 #define SNAP_GRID (16) // @todo ensure this is the same as SNodePanel::GetSnapGridSize()
 
-const FLinearColor UEdGraphSchema_Niagara::NodeTitleColor_Attribute = FLinearColor::Blue;
+const FLinearColor UEdGraphSchema_Niagara::NodeTitleColor_Attribute = FLinearColor::Green;
 const FLinearColor UEdGraphSchema_Niagara::NodeTitleColor_Constant = FLinearColor::Red;
+const FLinearColor UEdGraphSchema_Niagara::NodeTitleColor_SystemConstant = FLinearColor::White;
+const FLinearColor UEdGraphSchema_Niagara::NodeTitleColor_FunctionCall = FLinearColor::Blue;
 
 namespace 
 {
@@ -171,17 +173,26 @@ void UEdGraphSchema_Niagara::GetGraphContextActions(FGraphContextMenuBuilder& Co
 
 	//Add a generic Input node to allow getting external constants the current context doesn't know about.
 	{
-		const FText MenuDesc = LOCTEXT("GetSystemConstantUnknown", "Get Constant");
+		const FText MenuDesc = LOCTEXT("GetInput", "Input");
 
 		//TSharedPtr<FNiagaraSchemaAction_NewNode> GetConstAction = AddNewNodeAction(ContextMenuBuilder, *LOCTEXT("System Constants Menu Title", "System Constants").ToString(), MenuDesc, TEXT(""));
-		TSharedPtr<FNiagaraSchemaAction_NewNode> GetConstAction = AddNewNodeAction(ContextMenuBuilder, TEXT(""), MenuDesc, TEXT(""));
+		TSharedPtr<FNiagaraSchemaAction_NewNode> InputAction = AddNewNodeAction(ContextMenuBuilder, TEXT(""), MenuDesc, TEXT(""));
 
 		UNiagaraNodeInput* InputNode = NewObject<UNiagaraNodeInput>(ContextMenuBuilder.OwnerOfTemporaries);
 		InputNode->Input.Name = NAME_None;
 		InputNode->Input.Type = ENiagaraDataType::Vector;
-		GetConstAction->NodeTemplate = InputNode;
+		InputAction->NodeTemplate = InputNode;
 	}
 
+	//Add a function call node
+	{
+		const FText MenuDesc = LOCTEXT("NiagaraFunctionCall", "Function Call");
+
+		TSharedPtr<FNiagaraSchemaAction_NewNode> FunctionCallAction = AddNewNodeAction(ContextMenuBuilder, TEXT(""), MenuDesc, TEXT(""));
+
+		UNiagaraNodeFunctionCall* FunctionCallNode = NewObject<UNiagaraNodeFunctionCall>(ContextMenuBuilder.OwnerOfTemporaries);
+		FunctionCallAction->NodeTemplate = FunctionCallNode;
+	}
 	//TODO: Add quick commands for certain UNiagaraStructs and UNiagaraScripts to be added as functions
 }
 
@@ -326,4 +337,9 @@ void UEdGraphSchema_Niagara::GetPinDefaultValue(UEdGraphPin* Pin, FMatrix& OutDe
 	}
 }
 
+bool UEdGraphSchema_Niagara::IsSystemConstant(const FNiagaraVariableInfo& Variable)const
+{
+	return UNiagaraComponent::GetSystemConstants().Find(Variable) != INDEX_NONE;
+}
+
 #undef LOCTEXT_NAMESPACE
\ No newline at end of file
diff --git a/Engine/Source/Editor/UnrealEd/Private/Niagara/NiagaraEditorCommon.cpp b/Engine/Source/Editor/UnrealEd/Private/Niagara/NiagaraEditorCommon.cpp
--- a/Engine/Source/Editor/UnrealEd/Private/Niagara/NiagaraEditorCommon.cpp
+++ b/Engine/Source/Editor/UnrealEd/Private/Niagara/NiagaraEditorCommon.cpp
@@ -116,8 +116,6 @@ void FNiagaraOpInfo::Init()
 	Op->FriendlyName = NSLOCTEXT("NiagaraOpInfo", "Abs Name", "Abs");
 	Op->Description = NSLOCTEXT("NiagaraOpInfo", "Abs Desc", "Result = |A|");
 	Op->Inputs.Add(FNiagaraOpInOutInfo(A, ENiagaraDataType::Vector, AText, AText, DefaultStr_VecOne));
-	Op->Inputs.Add(FNiagaraOpInOutInfo(B, ENiagaraDataType::Vector, BText, BText, DefaultStr_VecOne));
-	Op->Inputs.Add(FNiagaraOpInOutInfo(C, ENiagaraDataType::Vector, CText, CText, DefaultStr_VecOne));
 	Op->Outputs.Add(FNiagaraOpInOutInfo(Result, ENiagaraDataType::Vector, ResultText, ResultText, DefaultStr_VecOne));
 	Op->OpDelegate.BindStatic(&INiagaraCompiler::Abs);
 
diff --git a/Engine/Source/Editor/UnrealEd/Private/Niagara/NiagaraNodeFunctionCall.cpp b/Engine/Source/Editor/UnrealEd/Private/Niagara/NiagaraNodeFunctionCall.cpp
--- a/Engine/Source/Editor/UnrealEd/Private/Niagara/NiagaraNodeFunctionCall.cpp
+++ b/Engine/Source/Editor/UnrealEd/Private/Niagara/NiagaraNodeFunctionCall.cpp
@@ -1,114 +1,171 @@
+// Copyright 1998-2014 Epic Games, Inc. All Rights Reserved.
 
 #include "UnrealEd.h"
 #include "BlueprintGraphDefinitions.h"
 #include "GraphEditorSettings.h"
+#include "Engine/NiagaraScript.h"
 
 void UNiagaraNodeFunctionCall::PostLoad()
 {
 	Super::PostLoad();
 
-// 	if (FunctionScript)
-// 	{
-// 		FunctionScript->ConditionalPostLoad();
-// 	}
+	if (FunctionScript)
+	{
+		FunctionScript->ConditionalPostLoad();
+	}
 }
 
-// void UNiagaraNodeFunctionCall::PostEditChangeProperty(struct FPropertyChangedEvent& PropertyChangedEvent)
-// {
-// 	ReallocatePins();
-// }
+void UNiagaraNodeFunctionCall::PostEditChangeProperty(struct FPropertyChangedEvent& PropertyChangedEvent)
+{
+	ReallocatePins();
+}
 
 void UNiagaraNodeFunctionCall::ReallocatePins()
 {
-// 	if (!FunctionScript)
-// 		return;
-// 
-// 	const UEdGraphSchema_Niagara* Schema = GetDefault<UEdGraphSchema_Niagara>();
-// 
-// // 	TArray<class UEdGraphPin*> PrevLinks;
-// // 	if (Pins.Num() > 0)
-// // 	{
-// // 		PrevLinks = Pins[0]->LinkedTo;
-// // 	}
-// 
-// 	Pins.Empty();
-// 	
-// 	if (UNiagaraScriptSource* Source = Cast<UNiagaraScriptSource>(FunctionScript->Source))
-// 	{
-// 		if (UNiagaraGraph* Graph = Cast<UNiagaraGraph>(Source->NodeGraph))
-// 		{
-// 			if (UNiagaraNodeOutput* OutNode = Graph->FindOutputNode())
-// 			{
-// 				//Add the input pins.
-// 				TArray<UNiagaraNodeInput*> InputNodes;
-// 				Graph->FindInputNodes(InputNodes);
-// 
-// 				for (UNiagaraNodeInput* InNode : InputNodes)
-// 				{
-// 					check(InNode);
-// 					switch (InNode->Input.Type)
-// 					{
-// 					case ENiagaraDataType::Scalar:
-// 					{
-// 						CreatePin(EGPD_Input, Schema->PC_Float, TEXT(""), NULL, false, false, InNode->Input.Name.ToString());
-// 					}
-// 						break;
-// 					case ENiagaraDataType::Vector:
-// 					{
-// 						CreatePin(EGPD_Input, Schema->PC_Vector, TEXT(""), NULL, false, false, InNode->Input.Name.ToString());
-// 					}
-// 						break;
-// 					case ENiagaraDataType::Matrix:
-// 					{
-// 						CreatePin(EGPD_Input, Schema->PC_Matrix, TEXT(""), NULL, false, false, InNode->Input.Name.ToString());
-// 					}
-// 						break;
-// 					};
-// 				}
-// 
-// 				//Add the output pins.
-// 				for (FNiagaraVariableInfo& Output : OutNode->Outputs)
-// 				{
-// 					switch (Output.Type)
-// 					{
-// 					case ENiagaraDataType::Scalar:
-// 					{
-// 						CreatePin(EGPD_Output, Schema->PC_Float, TEXT(""), NULL, false, false, Output.Name.ToString());
-// 					}
-// 						break;
-// 					case ENiagaraDataType::Vector:
-// 					{
-// 						CreatePin(EGPD_Output, Schema->PC_Vector, TEXT(""), NULL, false, false, Output.Name.ToString());
-// 					}
-// 						break;
-// 					case ENiagaraDataType::Matrix:
-// 					{
-// 						CreatePin(EGPD_Output, Schema->PC_Matrix, TEXT(""), NULL, false, false, Output.Name.ToString());
-// 					}
-// 						break;
-// 					};
-// 				}		
-// 			}
-// 		}
-// 	}
-
-// 	for (UEdGraphPin* Pin : Pins)
-// 	{
-// 		for (UEdGraphPin* PrevLinkedPin : PrevLinks)
-// 		{
-// 			if ((Pin->Direction == EGPD_Input && PrevLinkedPin->Direction == EGPD_Output)
-// 				|| (Pin->Direction == EGPD_Output && PrevLinkedPin->Direction == EGPD_Input))
-// 			{
-// 				if (Pin->PinName == PrevLinkedPin->PinName)
-// 				{
-// 				}
-// 			}			
-// 		}
-// 		if (Pins[0].PinName == Pin->PinName && Schema->CanCreateConnection(Pin, Pins[0]).Response == ECanCreateConnectionResponse::CONNECT_RESPONSE_MAKE)
-// 		{
-// 			Pins[0]->MakeLinkTo(Pin);
-// 		}
-// 	}
+	if (!FunctionScript)
+		return;
+
+	const UEdGraphSchema_Niagara* Schema = GetDefault<UEdGraphSchema_Niagara>();
+	Modify();
+
+	// Break any links to 'orphan' pins
+	for (int32 PinIndex = 0; PinIndex < Pins.Num(); ++PinIndex)
+	{
+		UEdGraphPin* Pin = Pins[PinIndex];
+		TArray<class UEdGraphPin*>& LinkedToRef = Pin->LinkedTo;
+		for (int32 LinkIdx = 0; LinkIdx < LinkedToRef.Num(); LinkIdx++)
+		{
+			UEdGraphPin* OtherPin = LinkedToRef[LinkIdx];
+			// If we are linked to a pin that its owner doesn't know about, break that link
+			if (!OtherPin->GetOwningNode()->Pins.Contains(OtherPin))
+			{
+				Pin->LinkedTo.Remove(OtherPin);
+			}
+		}
+	}
+
+	// Store the old Input and Output pins
+	TArray<UEdGraphPin*> OldInputPins;
+	TArray<UEdGraphPin*> OldOutputPins;
+	GetInputPins(OldInputPins);
+	GetOutputPins(OldOutputPins);
+
+	// Move the existing pins to a saved array
+	TArray<UEdGraphPin*> OldPins(Pins);
+	Pins.Empty();
+	
+	if (UNiagaraScriptSource* Source = Cast<UNiagaraScriptSource>(FunctionScript->Source))
+	{
+		if (UNiagaraGraph* Graph = Cast<UNiagaraGraph>(Source->NodeGraph))
+		{
+			if (UNiagaraNodeOutput* OutNode = Graph->FindOutputNode())
+			{
+				//Add the input pins.
+				TArray<UNiagaraNodeInput*> InputNodes;
+				Graph->FindInputNodes(InputNodes);
+
+				for (UNiagaraNodeInput* InNode : InputNodes)
+				{
+					//Don't expose system constants as inputs.
+					if (Schema->IsSystemConstant(InNode->Input))
+						continue;
+
+					check(InNode);
+					switch (InNode->Input.Type)
+					{
+					case ENiagaraDataType::Scalar:
+					{
+						UEdGraphPin* NewPin = CreatePin(EGPD_Input, Schema->PC_Float, TEXT(""), NULL, false, false, InNode->Input.Name.ToString());
+						NewPin->bDefaultValueIsIgnored = true;
+					}
+						break;
+					case ENiagaraDataType::Vector:
+					{
+						UEdGraphPin* NewPin = CreatePin(EGPD_Input, Schema->PC_Vector, TEXT(""), NULL, false, false, InNode->Input.Name.ToString());
+						NewPin->bDefaultValueIsIgnored = true;
+					}
+						break;
+					case ENiagaraDataType::Matrix:
+					{
+						UEdGraphPin* NewPin = CreatePin(EGPD_Input, Schema->PC_Matrix, TEXT(""), NULL, false, false, InNode->Input.Name.ToString());
+						NewPin->bDefaultValueIsIgnored = true;
+					}
+						break;
+					};
+				}
+
+				//Add the output pins.
+				for (FNiagaraVariableInfo& Output : OutNode->Outputs)
+				{
+					switch (Output.Type)
+					{
+					case ENiagaraDataType::Scalar:
+					{						
+						UEdGraphPin* NewPin = CreatePin(EGPD_Output, Schema->PC_Float, TEXT(""), NULL, false, false, Output.Name.ToString());
+						NewPin->bDefaultValueIsIgnored = true;
+					}
+						break;
+					case ENiagaraDataType::Vector:
+					{						
+						UEdGraphPin* NewPin = CreatePin(EGPD_Output, Schema->PC_Vector, TEXT(""), NULL, false, false, Output.Name.ToString());
+						NewPin->bDefaultValueIsIgnored = true;
+					}
+						break;
+					case ENiagaraDataType::Matrix:
+					{
+						UEdGraphPin* NewPin = CreatePin(EGPD_Output, Schema->PC_Matrix, TEXT(""), NULL, false, false, Output.Name.ToString());
+						NewPin->bDefaultValueIsIgnored = true;
+					}
+						break;
+					};
+				}		
+			}
+		}
+	}
+
+	// Get new Input and Output pins
+	TArray<UEdGraphPin*> NewInputPins;
+	TArray<UEdGraphPin*> NewOutputPins;
+	GetInputPins(NewInputPins);
+	GetOutputPins(NewOutputPins);
+
+	for (int32 PinIndex = 0; PinIndex < OldInputPins.Num(); PinIndex++)
+	{
+		if (PinIndex < NewInputPins.Num())
+		{
+			NewInputPins[PinIndex]->CopyPersistentDataFromOldPin(*OldInputPins[PinIndex]);
+		}
+	}
+
+	for (int32 PinIndex = 0; PinIndex < OldOutputPins.Num(); PinIndex++)
+	{
+		if (PinIndex < NewOutputPins.Num())
+		{
+			NewOutputPins[PinIndex]->CopyPersistentDataFromOldPin(*OldOutputPins[PinIndex]);
+		}
+	}
+
+	OldInputPins.Empty();
+	OldOutputPins.Empty();
+
+	// Throw away the original pins
+	for (int32 OldPinIndex = 0; OldPinIndex < OldPins.Num(); ++OldPinIndex)
+	{
+		UEdGraphPin* OldPin = OldPins[OldPinIndex];
+		OldPin->Modify();
+		OldPin->BreakAllPinLinks();
+
+#if 0
+		UEdGraphNode::ReturnPinToPool(OldPin);
+#else
+		OldPin->Rename(NULL, GetTransientPackage(), REN_None);
+		OldPin->RemoveFromRoot();
+		OldPin->MarkPendingKill();
+#endif
+	}
+	OldPins.Empty();
+
+	GetGraph()->NotifyGraphChanged();
 }
 
 void UNiagaraNodeFunctionCall::AllocateDefaultPins()
@@ -116,48 +173,24 @@ void UNiagaraNodeFunctionCall::AllocateDefaultPins()
 	ReallocatePins();
 }
 
-// FText UNiagaraNodeFunctionCall::GetNodeTitle(ENodeTitleType::Type TitleType) const
-// {
-// 	return /*FunctionScript ? FText::FromString(FunctionScript->GetName()) : */Super::GetNodeTitle(TitleType);
-// }
-// 
-// FText UNiagaraNodeFunctionCall::GetTooltipText()const
-// {
-// 	const FNiagaraOpInfo* OpInfo = FNiagaraOpInfo::GetOpInfo(OpName);
-// 	check(OpInfo);
-// 	return OpInfo->Description;
-// }
-// 
-// FLinearColor UNiagaraNodeFunctionCall::GetNodeTitleColor() const
-// {
-// 	return GetDefault<UGraphEditorSettings>()->FunctionCallNodeTitleColor;
-// }
+FText UNiagaraNodeFunctionCall::GetNodeTitle(ENodeTitleType::Type TitleType) const
+{
+	return FunctionScript ? FText::FromString(FunctionScript->GetName()) : Super::GetNodeTitle(TitleType);
+}
+
+FText UNiagaraNodeFunctionCall::GetTooltipText()const
+{
+	return FunctionScript ? FText::FromString(FunctionScript->GetName()) : FText::FromString(TEXT("Funtion Call"));
+}
+
+FLinearColor UNiagaraNodeFunctionCall::GetNodeTitleColor() const
+{
+	return UEdGraphSchema_Niagara::NodeTitleColor_FunctionCall;
+}
 
 void UNiagaraNodeFunctionCall::Compile(class INiagaraCompiler* Compiler, TArray<FNiagaraNodeResult>& Outputs)
 {
-// 	const FNiagaraOpInfo* OpInfo = FNiagaraOpInfo::GetOpInfo(OpName);
-// 	check(OpInfo);
-// 	int32 NumInputs = OpInfo->Inputs.Num();
-// 	int32 NumOutputs = OpInfo->Outputs.Num();
-// 
-// 	//No need to compile the input pins here. This will be done as a consequence of the function call compilation.
-// // 	TArray<TNiagaraExprPtr> Inputs;
-// // 	for (int32 i = 0; i < NumInputs; ++i)
-// // 	{
-// // 		UEdGraphPin *Pin = Pins[i];
-// // 		check(Pin->Direction == EGPD_Input);
-// // 		Inputs.Add(Compiler->CompilePin(Pin));
-// // 	}
-// 
-// 	TArray<TNiagaraExprPtr> OutputExpressions;
-// 	Compiler->FunctionCall(FunctionScript, OutputExpressions);
-// 
-// 	check(OutputExpressions.Num() == OpInfo->Outputs.Num());
-// 
-// 	for (int32 i = 0; i < NumOutputs; ++i)
-// 	{
-// 		UEdGraphPin *Pin = Pins[NumInputs + i];
-// 		check(Pin->Direction == EGPD_Output);
-// 		Outputs.Add(FNiagaraNodeResult(OutputExpressions[i], Pin));
-// 	}
+	//We don't actually compile function call nodes(yet). They are merged into the source graph. Effectively inlining vs branching.
+	//When we support branching in the VM we can probably look at branching funtion calls.
+	check(false);
 }
\ No newline at end of file
diff --git a/Engine/Source/Editor/UnrealEd/Private/Niagara/NiagaraNodeInput.cpp b/Engine/Source/Editor/UnrealEd/Private/Niagara/NiagaraNodeInput.cpp
--- a/Engine/Source/Editor/UnrealEd/Private/Niagara/NiagaraNodeInput.cpp
+++ b/Engine/Source/Editor/UnrealEd/Private/Niagara/NiagaraNodeInput.cpp
@@ -1,4 +1,4 @@
-// Copyright 1998-2014 Epic Games, Inc. All Rights Reserved.
+// Copyright 1998-2015 Epic Games, Inc. All Rights Reserved.
 
 #include "UnrealEd.h"
 #include "BlueprintGraphDefinitions.h"
@@ -133,9 +133,12 @@ FText UNiagaraNodeInput::GetNodeTitle(ENodeTitleType::Type TitleType) const
 
 FLinearColor UNiagaraNodeInput::GetNodeTitleColor() const
 {
+	const UEdGraphSchema_Niagara* Schema = CastChecked<UEdGraphSchema_Niagara>(GetSchema());
 	if (IsConstant())
 	{
-		return CastChecked<UEdGraphSchema_Niagara>(GetSchema())->NodeTitleColor_Constant;
+		return Schema->IsSystemConstant(Input) ? 
+			CastChecked<UEdGraphSchema_Niagara>(GetSchema())->NodeTitleColor_SystemConstant :
+			CastChecked<UEdGraphSchema_Niagara>(GetSchema())->NodeTitleColor_Constant ;
 	}
 	else
 	{
@@ -148,15 +151,30 @@ void UNiagaraNodeInput::Compile(class INiagaraCompiler* Compiler, TArray<FNiagar
 	//Input nodes in scripts being used as functions by other scripts will actually never get this far.
 	//The compiler will find them and replace them with direct links between the pins going in and the nodes connected to the corresponding input.
 
+	const UEdGraphSchema_Niagara* Schema = CastChecked<UEdGraphSchema_Niagara>(GetSchema());
 	if (IsConstant())
 	{
-		//treat this input as an external constant.
-		switch (Input.Type)
+		if (IsExposedConstant() || Schema->IsSystemConstant(Input))
 		{
-		case ENiagaraDataType::Scalar:	Outputs.Add(FNiagaraNodeResult(Compiler->GetExternalConstant(Input, FloatDefault), Pins[0]));	break;
-		case ENiagaraDataType::Vector:	Outputs.Add(FNiagaraNodeResult(Compiler->GetExternalConstant(Input, VectorDefault), Pins[0]));	break;
-		case ENiagaraDataType::Matrix:	Outputs.Add(FNiagaraNodeResult(Compiler->GetExternalConstant(Input, MatrixDefault), Pins[0]));	break;
-		case ENiagaraDataType::Curve:	Outputs.Add(FNiagaraNodeResult(Compiler->GetExternalCurveConstant(Input), Pins[0]));	break;
+			//treat this input as an external constant.
+			switch (Input.Type)
+			{
+			case ENiagaraDataType::Scalar:	Outputs.Add(FNiagaraNodeResult(Compiler->GetExternalConstant(Input, FloatDefault), Pins[0]));	break;
+			case ENiagaraDataType::Vector:	Outputs.Add(FNiagaraNodeResult(Compiler->GetExternalConstant(Input, VectorDefault), Pins[0]));	break;
+			case ENiagaraDataType::Matrix:	Outputs.Add(FNiagaraNodeResult(Compiler->GetExternalConstant(Input, MatrixDefault), Pins[0]));	break;
+			case ENiagaraDataType::Curve:	Outputs.Add(FNiagaraNodeResult(Compiler->GetExternalCurveConstant(Input), Pins[0]));	break;
+			}
+		}
+		else
+		{
+			//treat this input as an external constant.
+			switch (Input.Type)
+			{
+			case ENiagaraDataType::Scalar:	Outputs.Add(FNiagaraNodeResult(Compiler->GetInternalConstant(Input, FloatDefault), Pins[0]));	break;
+			case ENiagaraDataType::Vector:	Outputs.Add(FNiagaraNodeResult(Compiler->GetInternalConstant(Input, VectorDefault), Pins[0]));	break;
+			case ENiagaraDataType::Matrix:	Outputs.Add(FNiagaraNodeResult(Compiler->GetInternalConstant(Input, MatrixDefault), Pins[0]));	break;
+			case ENiagaraDataType::Curve:	Outputs.Add(FNiagaraNodeResult(Compiler->GetExternalCurveConstant(Input), Pins[0]));	break;//Internal curve Constant?
+			}
 		}
 	}
 	else
diff --git a/Engine/Source/Editor/UnrealEd/Private/Niagara/NiagaraScriptSource.cpp b/Engine/Source/Editor/UnrealEd/Private/Niagara/NiagaraScriptSource.cpp
--- a/Engine/Source/Editor/UnrealEd/Private/Niagara/NiagaraScriptSource.cpp
+++ b/Engine/Source/Editor/UnrealEd/Private/Niagara/NiagaraScriptSource.cpp
@@ -101,29 +101,35 @@ void UNiagaraScriptSource::Compile()
 	FNiagaraEditorModule& NiagaraEditorModule = FModuleManager::Get().LoadModuleChecked<FNiagaraEditorModule>(TEXT("NiagaraEditor"));
 	NiagaraEditorModule.CompileScript(ScriptOwner);
 
-	ExposedVectorConstants.Empty();
+	FNiagaraConstants& ExternalConsts = ScriptOwner->ConstantData.ExternalConstants;
 
-	// grab all constant nodes that are exposed to the editor
-	TArray<UNiagaraNodeInput*> ConstNodes;
-	NodeGraph->GetNodesOfClass<UNiagaraNodeInput>(ConstNodes);
-	for (UNiagaraNodeInput *Node : ConstNodes)
+	const UEdGraphSchema_Niagara* Schema = CastChecked<UEdGraphSchema_Niagara>(NodeGraph->GetSchema());
+	ExposedVectorConstants.Empty();
+	for (int32 ConstIdx = 0; ConstIdx < ExternalConsts.GetNumVectorConstants(); ConstIdx++)
 	{
-		if (Node->IsExposedConstant())
+		FNiagaraVariableInfo Info;
+		FVector4 Value;
+		ExternalConsts.GetVectorConstant(ConstIdx, Value, Info);
+		if (Schema->IsSystemConstant(Info))
 		{
-			if(Node->Input.Type == ENiagaraDataType::Vector)
-			{
-				EditorExposedVectorConstant *Const = new EditorExposedVectorConstant();
-				Const->ConstName = Node->Input.Name;
-				Const->Value = Node->VectorDefault;
-				ExposedVectorConstants.Add( MakeShareable(Const) );
-			}
-			else if (Node->Input.Type == ENiagaraDataType::Curve)
-			{
-				EditorExposedVectorCurveConstant *Const = new EditorExposedVectorCurveConstant();
-				Const->ConstName = Node->Input.Name;
-				ExposedVectorCurveConstants.Add(MakeShareable(Const));
-			}
+			continue;//System constants are "external" but should not be exposed to the editor.
 		}
+			
+		EditorExposedVectorConstant *Const = new EditorExposedVectorConstant();
+		Const->ConstName = Info.Name;
+		Const->Value = Value;
+		ExposedVectorConstants.Add(MakeShareable(Const));
+	}
+
+	for (int32 ConstIdx = 0; ConstIdx < ExternalConsts.GetNumDataObjectConstants(); ConstIdx++)
+	{
+		EditorExposedVectorCurveConstant *Const = new EditorExposedVectorCurveConstant();
+		FNiagaraVariableInfo Info;
+		FNiagaraDataObject* Dummy;
+		ExternalConsts.GetDataObjectConstant(ConstIdx, Dummy, Info);
+		Const->ConstName = Info.Name;
+		//Set default value for this too?
+		ExposedVectorCurveConstants.Add(MakeShareable(Const));
 	}
 }
 
diff --git a/Engine/Source/Editor/UnrealEd/Public/INiagaraCompiler.h b/Engine/Source/Editor/UnrealEd/Public/INiagaraCompiler.h
--- a/Engine/Source/Editor/UnrealEd/Public/INiagaraCompiler.h
+++ b/Engine/Source/Editor/UnrealEd/Public/INiagaraCompiler.h
@@ -23,6 +23,12 @@ public:
 	virtual TNiagaraExprPtr GetExternalConstant(const FNiagaraVariableInfo& Constant, const FVector4& Default) = 0;
 	/** Creates an expression which fetches the named constant. */
 	virtual TNiagaraExprPtr GetExternalConstant(const FNiagaraVariableInfo& Constant, const FMatrix& Default) = 0;
+	/** Creates an expression which fetches the named constant. */
+	virtual TNiagaraExprPtr GetInternalConstant(const FNiagaraVariableInfo& Constant, float Default) = 0;
+	/** Creates an expression which fetches the named constant. */
+	virtual TNiagaraExprPtr GetInternalConstant(const FNiagaraVariableInfo& Constant, const FVector4& Default) = 0;
+	/** Creates an expression which fetches the named constant. */
+	virtual TNiagaraExprPtr GetInternalConstant(const FNiagaraVariableInfo& Constant, const FMatrix& Default) = 0;
 	/** Creates an expression which fetches the named CURVE. */
 	virtual TNiagaraExprPtr GetExternalCurveConstant(const FNiagaraVariableInfo& Constant) = 0;
 
diff --git a/Engine/Source/Runtime/Engine/Classes/Components/NiagaraComponent.h b/Engine/Source/Runtime/Engine/Classes/Components/NiagaraComponent.h
--- a/Engine/Source/Runtime/Engine/Classes/Components/NiagaraComponent.h
+++ b/Engine/Source/Runtime/Engine/Classes/Components/NiagaraComponent.h
@@ -14,10 +14,9 @@ class FNiagaraEmitterParticleData
 {
 public:
 	FNiagaraEmitterParticleData()
-		: CurrentBuffer(0),
-		NumParticles(0),
-		ParticleAllocation(0)
+		: CurrentBuffer(0)
 	{
+		Reset();
 	}
 
 	~FNiagaraEmitterParticleData() {}
@@ -27,27 +26,38 @@ public:
 		ParticleBuffers[CurrentBuffer].Reset(NumExpectedParticles * AttrMap.Num());
 		ParticleBuffers[CurrentBuffer].AddUninitialized(NumExpectedParticles * AttrMap.Num());
 
-		ParticleAllocation = NumExpectedParticles;
+		ParticleAllocation[CurrentBuffer] = NumExpectedParticles;
 	}
 
 	void Reset()
 	{
 		AttrMap.Empty();
-		NumParticles = 0;
-		ParticleAllocation = 0;
+		NumParticles[0] = 0;
+		NumParticles[1] = 0;
+		ParticleAllocation[0] = 0;
+		ParticleAllocation[1] = 0;
 		CurrentBuffer = 0;
 	}
 
 	const FVector4* GetAttributeData(const FNiagaraVariableInfo& AttrID) const
 	{
 		const uint32* Offset = AttrMap.Find(AttrID);
-		return Offset ? ParticleBuffers[CurrentBuffer].GetData() + (*Offset * ParticleAllocation) : NULL;
+		return Offset ? ParticleBuffers[CurrentBuffer].GetData() + (*Offset * ParticleAllocation[CurrentBuffer]) : NULL;
 	}
 
-	FVector4* GetAttributeDataWrite(const FNiagaraVariableInfo& AttrID)
+	void GetAttributeData(const FNiagaraVariableInfo& AttrID, FVector4*& PrevBuffer, FVector4*& CurrBuffer, int32 StartParticle) 
 	{
 		const uint32* Offset = AttrMap.Find(AttrID);
-		return Offset ? ParticleBuffers[CurrentBuffer].GetData() + (*Offset * ParticleAllocation) : NULL;
+		if (Offset)
+		{
+			PrevBuffer = ParticleBuffers[CurrentBuffer ^ 0x1].GetData() + (*Offset * ParticleAllocation[CurrentBuffer ^ 0x1]) + StartParticle;
+			CurrBuffer = ParticleBuffers[CurrentBuffer].GetData() + (*Offset * ParticleAllocation[CurrentBuffer]) + StartParticle;
+		}
+		else
+		{
+			PrevBuffer = NULL;
+			CurrBuffer = NULL;
+		}
 	}
 
 	bool HasAttriubte(const FNiagaraVariableInfo& AttrID)
@@ -66,11 +76,27 @@ public:
 		}
 	}
 
+	void AddAttributes(const TArray<FNiagaraVariableInfo>& Attrs)
+	{
+		for (const FNiagaraVariableInfo& Attr : Attrs)
+		{
+			//TODO: Support attributes of any type. This may as well wait until the move of matrix ops etc over into UNiagaraScriptStructs.
+			uint32 Idx = AttrMap.Num();
+			if (!AttrMap.Find(Attr))
+			{
+				AttrMap.Add(Attr, Idx);
+			}
+		}
+	}
+
+	const TMap<FNiagaraVariableInfo, uint32>& GetAttributes()const { return AttrMap; }
 	int GetNumAttributes()				{ return AttrMap.Num(); }
 
-	uint32 GetNumParticles() const		{ return NumParticles; }
-	uint32 GetParticleAllocation() const { return ParticleAllocation; }
-	void SetNumParticles(uint32 Num)	{ NumParticles = Num;  }
+	uint32 GetNumParticles() const		{ return NumParticles[CurrentBuffer]; }
+	uint32 GetPrevNumParticles() const		{ return NumParticles[CurrentBuffer ^ 0x1]; }
+	uint32 GetParticleAllocation() const { return ParticleAllocation[CurrentBuffer]; }
+	uint32 GetPrevParticleAllocation() const { return ParticleAllocation[CurrentBuffer ^ 0x1]; }
+	void SetNumParticles(uint32 Num)	{ NumParticles[CurrentBuffer] = Num; }
 	void SwapBuffers()					{ CurrentBuffer ^= 0x1; }
 
 	FVector4* GetCurrentBuffer()		{ return ParticleBuffers[CurrentBuffer].GetData(); }
@@ -78,7 +104,9 @@ public:
 
 	int GetBytesUsed()	{ return (ParticleBuffers[0].Num() + ParticleBuffers[1].Num()) * 16 + AttrMap.Num() * 4; }
 private:
-	uint32 CurrentBuffer, NumParticles, ParticleAllocation;
+	uint32 CurrentBuffer;
+	uint32 NumParticles[2];
+	uint32 ParticleAllocation[2];
 	TArray<FVector4> ParticleBuffers[2];
 	TMap<FNiagaraVariableInfo, uint32> AttrMap;
 };
diff --git a/Engine/Source/Runtime/Engine/Classes/Engine/NiagaraScript.h b/Engine/Source/Runtime/Engine/Classes/Engine/NiagaraScript.h
--- a/Engine/Source/Runtime/Engine/Classes/Engine/NiagaraScript.h
+++ b/Engine/Source/Runtime/Engine/Classes/Engine/NiagaraScript.h
@@ -6,6 +6,22 @@
 
 #include "NiagaraScript.generated.h"
 
+#define NIAGARA_INVALID_MEMORY (0xBA)
+
+/** Defines what will happen to unused attributes when a script is run. */
+UENUM()
+enum class EUnusedAttributeBehaviour : uint8
+{
+	/** The previous value of the attribute is copied across. */
+	Copy,
+	/** The attribute is set to zero. */
+	Zero,
+	/** The attribute is untouched. */
+	None,
+	/** The memory for the attribute is set to NIAGARA_INVALID_MEMORY. */
+	MarkInvalid, 
+};
+
 /** Runtime script for a Niagara system */
 UCLASS(MinimalAPI)
 class UNiagaraScript : public UObject
diff --git a/Engine/Source/Runtime/Engine/Classes/Engine/NiagaraSimulation.h b/Engine/Source/Runtime/Engine/Classes/Engine/NiagaraSimulation.h
--- a/Engine/Source/Runtime/Engine/Classes/Engine/NiagaraSimulation.h
+++ b/Engine/Source/Runtime/Engine/Classes/Engine/NiagaraSimulation.h
@@ -13,7 +13,8 @@ NiagaraSimulation.h: Niagara emitter simulation class
 
 DECLARE_CYCLE_STAT(TEXT("Tick"), STAT_NiagaraTick, STATGROUP_Niagara);
 DECLARE_CYCLE_STAT(TEXT("Simulate"), STAT_NiagaraSimulate, STATGROUP_Niagara);
-DECLARE_CYCLE_STAT(TEXT("Spawn + Kill"), STAT_NiagaraSpawnAndKill, STATGROUP_Niagara);
+DECLARE_CYCLE_STAT(TEXT("Spawn"), STAT_NiagaraSpawn, STATGROUP_Niagara);
+DECLARE_CYCLE_STAT(TEXT("Kill"), STAT_NiagaraKill, STATGROUP_Niagara);
 
 
 class NiagaraEffectRenderer;
@@ -144,18 +145,11 @@ private:
 		SpawnRemainder = FloatNumToSpawn - NumToSpawn;
 		return NumToSpawn;
 	}
-
-	/** Run VM to update particle positions */
-	void UpdateParticles(
-		float DeltaSeconds,
-		FVector4* PrevParticles,
-		int32 PrevNumVectorsPerAttribute,
-		FVector4* Particles,
-		int32 NumVectorsPerAttribute,
-		int32 NumParticles
-		);
-
-	int32 SpawnAndKillParticles(int32 NumToSpawn);
+	
+	/** Runs a script in the VM over a specific range of particles. */
+	void RunVMScript(UNiagaraScript* Script, EUnusedAttributeBehaviour UnusedAttribBehaviour);
+	void RunVMScript(UNiagaraScript* Script, EUnusedAttributeBehaviour UnusedAttribBehaviour, uint32 StartParticle);
+	void RunVMScript(UNiagaraScript* Script, EUnusedAttributeBehaviour UnusedAttribBehaviour, uint32 StartParticle, uint32 NumParticles);
 
 	/** Util to move a particle */
 	void MoveParticleToIndex(int32 SrcIndex, int32 DestIndex)
diff --git a/Engine/Source/Runtime/Engine/Classes/Niagara/NiagaraConstantSet.h b/Engine/Source/Runtime/Engine/Classes/Niagara/NiagaraConstantSet.h
--- a/Engine/Source/Runtime/Engine/Classes/Niagara/NiagaraConstantSet.h
+++ b/Engine/Source/Runtime/Engine/Classes/Niagara/NiagaraConstantSet.h
@@ -170,6 +170,32 @@ public:
 		DataObjectConstants.Empty();
 	}
 
+	void GetScalarConstant(int32 Index, float& OutValue, FNiagaraVariableInfo& OutInfo)const
+	{
+		OutInfo = ScalarConstantsInfo[Index]; 
+		OutValue = ScalarConstants[Index];
+	}
+	void GetVectorConstant(int32 Index, FVector4& OutValue, FNiagaraVariableInfo& OutInfo)const
+	{
+		OutInfo = VectorConstantsInfo[Index];
+		OutValue = VectorConstants[Index];
+	}
+	void GetMatrixConstant(int32 Index, FMatrix& OutValue, FNiagaraVariableInfo& OutInfo)const
+	{
+		OutInfo = MatrixConstantsInfo[Index];
+		OutValue = MatrixConstants[Index];
+	}
+	void GetDataObjectConstant(int32 Index, FNiagaraDataObject*& OutValue, FNiagaraVariableInfo& OutInfo)const
+	{
+		OutInfo = DataObjectConstantsInfo[Index];
+		OutValue = DataObjectConstants[Index];
+	}
+
+	const int32 GetNumScalarConstants()const{ return ScalarConstantsInfo.Num(); }
+	const int32 GetNumVectorConstants()const{ return VectorConstantsInfo.Num(); }
+	const int32 GetNumMatrixConstants()const{ return MatrixConstantsInfo.Num(); }
+	const int32 GetNumDataObjectConstants()const{ return DataObjectConstantsInfo.Num(); }
+
 	/** Fills the entire constants set into the constant table. */
 	void AppendToConstantsTable(TArray<FVector4>& ConstantsTable)const
 	{
diff --git a/Engine/Source/Runtime/Engine/Private/NiagaraEffectRenderer.cpp b/Engine/Source/Runtime/Engine/Private/NiagaraEffectRenderer.cpp
--- a/Engine/Source/Runtime/Engine/Private/NiagaraEffectRenderer.cpp
+++ b/Engine/Source/Runtime/Engine/Private/NiagaraEffectRenderer.cpp
@@ -247,6 +247,7 @@ FNiagaraDynamicDataBase *NiagaraEffectRendererSprites::GenerateVertexData(const
 		NewVertex.Position = PosPtr[ParticleIndex];
 		NewVertex.OldPosition = NewVertex.Position;
 		NewVertex.Color = FLinearColor(ColPtr[ParticleIndex]);
+		NewVertex.Color.A = ColPtr[ParticleIndex].W;
 		NewVertex.ParticleId = ParticleId;
 		ParticleId += IdInc;
 		NewVertex.RelativeTime = AgePtr[ParticleIndex].X;
diff --git a/Engine/Source/Runtime/Engine/Private/NiagaraSimulation.cpp b/Engine/Source/Runtime/Engine/Private/NiagaraSimulation.cpp
--- a/Engine/Source/Runtime/Engine/Private/NiagaraSimulation.cpp
+++ b/Engine/Source/Runtime/Engine/Private/NiagaraSimulation.cpp
@@ -33,39 +33,31 @@ FNiagaraSimulation::FNiagaraSimulation(FNiagaraEmitterProperties *InProps, ERHIF
 
 void FNiagaraSimulation::Init()
 {
-	if (Props && Props->UpdateScript && Props->SpawnScript)
+	if (Props && Props->UpdateScript && Props->SpawnScript )
 	{
-		Data.SetAttributes(Props->UpdateScript->Attributes);
-		Constants.Empty();
-		Constants.Merge(Props->ExternalConstants);
-
-		bool bDisable = false;
-		//Temporarily enforcing that spawn scripts have exactly the same attriubtes as update scripts.
-		//To fix this we need to refactor the scripts into source and runtime objects.
-		//The runtime will be a script autogenerated to be a safe combination of all scripts attributes. 
-		//With an options somewhere for handling of attributes not set in a particular script. 
-		//eg. Zeroing for spawn of attributes in update script that are not set explicitly in the spawn script. And pass-through copy for the reverse.
-		if (Props->SpawnScript->Attributes.Num() == Props->UpdateScript->Attributes.Num())
+		if (Props->UpdateScript->Attributes.Num() == 0 || Props->SpawnScript->Attributes.Num() == 0)
 		{
-			for (int32 i = 0; i < Props->SpawnScript->Attributes.Num(); ++i)
+			Data.Reset();
+			bIsEnabled = false;
+			UE_LOG(LogNiagara, Error, TEXT("This emitter cannot be enabled because it's spawn or update script doesn't have any attriubtes.."));
+		}
+		else
+		{
+			//Warn the user if there are any attributes used in the update script that are not initialized in the spawn script.
+			//TODO: We need some window in the effect editor and possibly the graph editor for warnings and errors.
+			for (FNiagaraVariableInfo& Attr : Props->UpdateScript->Attributes)
 			{
-				if (Props->SpawnScript->Attributes[i] != Props->UpdateScript->Attributes[i])
+				int32 FoundIdx;
+				if (!Props->SpawnScript->Attributes.Find(Attr, FoundIdx))
 				{
-					bDisable = true;
-					break;
+					UE_LOG(LogNiagara, Warning, TEXT("Attribute %s is used in the Update script for %s but it is not initialised in the Spawn script!"), *Attr.Name.ToString(), *Props->Name);
 				}
 			}
-		}
-		else
-		{
-			bDisable = true;
-		}
+			Data.AddAttributes(Props->UpdateScript->Attributes);
+			Data.AddAttributes(Props->SpawnScript->Attributes);
 
-		if (bDisable)
-		{
-			Data.Reset();//Clear the attributes to mark the sim as disabled independatly of the user set Enabled flag.
-			bIsEnabled = false;//Also force the user flag to give an indication to the user.
-			UE_LOG(LogNiagara, Error, TEXT("This emitter cannot be enabled because it's spawn script and update script don't have matching attributes."));
+			Constants.Empty();
+			Constants.Merge(Props->ExternalConstants);
 		}
 	}
 	else
@@ -123,44 +115,66 @@ void FNiagaraSimulation::Tick(float DeltaSeconds)
 {
 	SCOPE_CYCLE_COUNTER(STAT_NiagaraTick);
 
+	if (!bIsEnabled)
+		return;
+
 	SimpleTimer TickTime;
 
-	if (Data.GetNumAttributes() == 0)
-		return;
+	check(Data.GetNumAttributes() > 0);
+	check(Props->SpawnScript);
+	check(Props->UpdateScript);
 
 	// Cache the ComponentToWorld transform.
 //	CachedComponentToWorld = Component.GetComponentToWorld();
 
 	Data.SwapBuffers();
-	// Remember the stride of the original data.
-	int32 PrevNumVectorsPerAttribute = Data.GetParticleAllocation();
-	int32 MaxNewParticles = Data.GetNumParticles();
+	Data.SetNumParticles(Data.GetPrevNumParticles());
+
+	int32 OrigNumParticles = Data.GetNumParticles();
 	int32 NumToSpawn = 0;
 
 	// Figure out how many we will spawn.
-	if (Props->SpawnScript)
+	NumToSpawn = CalcNumToSpawn(DeltaSeconds);
+	int32 MaxNewParticles = OrigNumParticles + NumToSpawn;
+	Data.Allocate(MaxNewParticles);
+
+	Age += DeltaSeconds;
+	Constants.SetOrAdd(TEXT("Emitter Age"), FVector4(Age, Age, Age, Age));
+	Constants.SetOrAdd(TEXT("Delta Time"), FVector4(DeltaSeconds, DeltaSeconds, DeltaSeconds, DeltaSeconds));
+
+	// Simulate particles forward by DeltaSeconds.
 	{
-		NumToSpawn = CalcNumToSpawn(DeltaSeconds);
-		MaxNewParticles = Data.GetNumParticles() + NumToSpawn;
-		Data.Allocate(MaxNewParticles);
+		SCOPE_CYCLE_COUNTER(STAT_NiagaraSimulate);
+		RunVMScript(Props->UpdateScript, EUnusedAttributeBehaviour::Copy);
 	}
-
-	if (Props->UpdateScript)
+	
+	//Init new particles with the spawn script.
 	{
-		// Simulate particles forward by DeltaSeconds.
-		UpdateParticles(
-			DeltaSeconds,
-			Data.GetPreviousBuffer(),
-			PrevNumVectorsPerAttribute,
-			Data.GetCurrentBuffer(),
-			MaxNewParticles,
-			Data.GetNumParticles()
-			);
+		SCOPE_CYCLE_COUNTER(STAT_NiagaraSpawn);
+		Data.SetNumParticles(MaxNewParticles);
+		//For now, zero any unused attributes here. But as this is really uninitialized data we should maybe make this a more serious error.
+		RunVMScript(Props->SpawnScript, EUnusedAttributeBehaviour::Zero, OrigNumParticles, NumToSpawn);
 	}
 
-	if (Props->SpawnScript)
+	// Iterate over looking for dead particles and move from the end of the list to the dead location, compacting in the process
 	{
-		SpawnAndKillParticles(NumToSpawn);
+		SCOPE_CYCLE_COUNTER(STAT_NiagaraKill);
+		int32 CurNumParticles = OrigNumParticles = Data.GetNumParticles();
+		int32 ParticleIndex = 0;
+		const FVector4* ParticleRelativeTimes = Data.GetAttributeData(FNiagaraVariableInfo(FName(TEXT("Age")), ENiagaraDataType::Vector));
+		if (ParticleRelativeTimes)
+		{
+			while (ParticleIndex < OrigNumParticles)
+			{
+				if (ParticleRelativeTimes[ParticleIndex].X > 1.0f)
+				{
+					// Particle is dead, move one from the end here. 
+					MoveParticleToIndex(--CurNumParticles, ParticleIndex);
+				}
+				ParticleIndex++;
+			}
+		}
+		Data.SetNumParticles(CurNumParticles);
 	}
 
 	CPUTimeMS = TickTime.GetElapsedMilliseconds();
@@ -170,124 +184,100 @@ void FNiagaraSimulation::Tick(float DeltaSeconds)
 
 }
 
+void FNiagaraSimulation::RunVMScript(UNiagaraScript* Script, EUnusedAttributeBehaviour UnusedAttribBehaviour)
+{
+	RunVMScript(Script, UnusedAttribBehaviour, 0, Data.GetNumParticles());
+}
 
-void FNiagaraSimulation::UpdateParticles(
-	float DeltaSeconds,
-	FVector4* PrevParticles,
-	int32 PrevNumVectorsPerAttribute,
-	FVector4* Particles,
-	int32 NumVectorsPerAttribute,
-	int32 NumParticles
-	)
+void FNiagaraSimulation::RunVMScript(UNiagaraScript* Script, EUnusedAttributeBehaviour UnusedAttribBehaviour, uint32 StartParticle)
 {
-	if (Data.GetNumAttributes() == 0)
+	RunVMScript(Script, UnusedAttribBehaviour, StartParticle, Data.GetNumParticles() - StartParticle);
+}
+
+void FNiagaraSimulation::RunVMScript(UNiagaraScript* Script, EUnusedAttributeBehaviour UnusedAttribBehaviour, uint32 StartParticle, uint32 NumParticles)
+{
+	if (NumParticles == 0)
 		return;
 
-	SCOPE_CYCLE_COUNTER(STAT_NiagaraSimulate);
-	Age += DeltaSeconds;
-	Constants.SetOrAdd(TEXT("Emitter Age"), FVector4(Age, Age, Age, Age));
+	check(Script);
+	check(Script->ByteCode.Num() > 0);
+	check(Script->Attributes.Num() > 0);
+	check(Data.GetNumAttributes() > 0);
+
+	uint32 CurrNumParticles = Data.GetNumParticles();
+
+	check(StartParticle + NumParticles <= CurrNumParticles);
+
+	const FVector4* InputBuffer = Data.GetPreviousBuffer();
+	const FVector4* OutputBuffer = Data.GetCurrentBuffer();
+
+	uint32 InputAttributeStride = Data.GetPrevParticleAllocation();
+	uint32 OutputAttributeStride = Data.GetParticleAllocation();
 
 	VectorRegister* InputRegisters[VectorVM::MaxInputRegisters] = { 0 };
 	VectorRegister* OutputRegisters[VectorVM::MaxOutputRegisters] = { 0 };
-	const int32 NumAttr = Props->UpdateScript->Attributes.Num();
+	
+	//The current script run will be using NumScriptAttrs. We must pick out the Attributes form the simulation data in the order that they appear in the script.
+	//The remaining attributes being handled according to UnusedAttribBehaviour
+	const int32 NumSimulationAttrs = Data.GetNumAttributes();
+	const int32 NumScriptAttrs = Script->Attributes.Num();
 
-	check(NumAttr < VectorVM::MaxInputRegisters);
-	check(NumAttr < VectorVM::MaxOutputRegisters);
+	check(NumScriptAttrs < VectorVM::MaxInputRegisters);
+	check(NumScriptAttrs < VectorVM::MaxOutputRegisters);
 
 	// Setup input and output registers.
-	for (int32 AttrIndex = 0; AttrIndex < NumAttr; ++AttrIndex)
+	for (TPair<FNiagaraVariableInfo, uint32> AttrIndexPair : Data.GetAttributes())
 	{
-		InputRegisters[AttrIndex] = (VectorRegister*)(PrevParticles + AttrIndex * PrevNumVectorsPerAttribute);
-		OutputRegisters[AttrIndex] = (VectorRegister*)(Particles + AttrIndex * NumVectorsPerAttribute);
+		const FNiagaraVariableInfo& AttrInfo = AttrIndexPair.Key;
+
+		FVector4* InBuff;
+		FVector4* OutBuff;
+		Data.GetAttributeData(AttrInfo, InBuff, OutBuff, StartParticle);
+
+		int32 AttrIndex = Script->Attributes.Find(AttrInfo);
+		if (AttrIndex != INDEX_NONE)
+		{
+			//This attribute is used in this script so set it to the correct Input and Output buffer locations.
+			InputRegisters[AttrIndex] = (VectorRegister*)InBuff;
+			OutputRegisters[AttrIndex] = (VectorRegister*)OutBuff;
+		}	
+		else
+		{
+			//This attribute is not used in this script so handle it accordingly.
+			check(OutBuff != NULL);
+			if (UnusedAttribBehaviour == EUnusedAttributeBehaviour::Copy)
+			{
+				check(InBuff != NULL);
+				FMemory::Memcpy(OutBuff, InBuff, NumParticles * sizeof(FVector4));
+			}
+			else if (UnusedAttribBehaviour == EUnusedAttributeBehaviour::Zero)
+			{
+				FMemory::Memzero(OutBuff, NumParticles * sizeof(FVector4));
+			}
+			else if (UnusedAttribBehaviour == EUnusedAttributeBehaviour::MarkInvalid)
+			{
+				FMemory::Memset(OutBuff, NIAGARA_INVALID_MEMORY, NumParticles * sizeof(FVector4));
+			}
+		}
 	}
 
 	//Fill constant table with required emitter constants and internal script constants.
 	TArray<FVector4> ConstantTable;
 	TArray<FNiagaraDataObject *>DataObjTable;
-	Props->UpdateScript->ConstantData.FillConstantTable(Constants, ConstantTable, DataObjTable);
+	Script->ConstantData.FillConstantTable(Constants, ConstantTable, DataObjTable);
 
 	VectorVM::Exec(
-		Props->UpdateScript->ByteCode.GetData(),
+		Script->ByteCode.GetData(),
 		InputRegisters,
-		NumAttr,
+		NumScriptAttrs,
 		OutputRegisters,
-		NumAttr,
+		NumScriptAttrs,
 		ConstantTable.GetData(),
 		DataObjTable.GetData(),
 		NumParticles
 		);
 }
 
-
-
-
-
-int32 FNiagaraSimulation::SpawnAndKillParticles(int32 NumToSpawn)
-{
-	SCOPE_CYCLE_COUNTER(STAT_NiagaraSpawnAndKill);
-	int32 OrigNumParticles = Data.GetNumParticles();
-	int32 CurNumParticles = OrigNumParticles + NumToSpawn;
-
-	// run the spawn graph over all new particles
-	if (Props->SpawnScript && Props->SpawnScript->ByteCode.Num())
-	{
-		VectorRegister* InputRegisters[VectorVM::MaxInputRegisters] = { 0 };
-		VectorRegister* OutputRegisters[VectorVM::MaxOutputRegisters] = { 0 };
-		const int32 NumAttr = Props->SpawnScript->Attributes.Num();
-		const int32 NumVectors = NumToSpawn;
-
-		check(NumAttr < VectorVM::MaxInputRegisters);
-		check(NumAttr < VectorVM::MaxOutputRegisters);
-
-		FVector4 *NewParticlesStart = Data.GetCurrentBuffer() + OrigNumParticles;
-
-		// Setup input and output registers.
-		for (int32 AttrIndex = 0; AttrIndex < NumAttr; ++AttrIndex)
-		{
-			InputRegisters[AttrIndex] = (VectorRegister*)(NewParticlesStart + AttrIndex * Data.GetParticleAllocation());
-			OutputRegisters[AttrIndex] = (VectorRegister*)(NewParticlesStart + AttrIndex * Data.GetParticleAllocation());
-		}
-
-		//Fill constant table with required emitter constants and internal script constants.
-		TArray<FVector4> ConstantTable;
-		TArray<FNiagaraDataObject *>DataObjTable;
-		Props->SpawnScript->ConstantData.FillConstantTable(Constants, ConstantTable, DataObjTable);
-
-		VectorVM::Exec(
-			Props->SpawnScript->ByteCode.GetData(),
-			InputRegisters,
-			NumAttr,
-			OutputRegisters,
-			NumAttr,
-			ConstantTable.GetData(),
-			DataObjTable.GetData(),
-			NumVectors
-			);
-	}
-
-	//TODO: NEED TO HANDLE PARTICLE DEATH WITHOUT RELYING ON AGE ATTRIBUTE. UNLESS WE FORCE IT.
-	//MAYBE WITH BRANCHING IN VM AND A KILL() NODE.
-
-	// Iterate over looking for dead particles and move from the end of the list to the dead location, compacting in the process
-	int32 ParticleIndex = 0;
-	const FVector4* ParticleRelativeTimes = Data.GetAttributeData(FNiagaraVariableInfo(FName(TEXT("Age")), ENiagaraDataType::Vector));
-	if (ParticleRelativeTimes)
-	{
-		while (ParticleIndex < OrigNumParticles)
-		{
-			if (ParticleRelativeTimes[ParticleIndex].X > 1.0f)
-			{
-				// Particle is dead, move one from the end here.
-				MoveParticleToIndex(--CurNumParticles, ParticleIndex);
-			}
-			ParticleIndex++;
-		}
-	}
-
-	Data.SetNumParticles(CurNumParticles);
-	return CurNumParticles;
-}
-
 bool FNiagaraSimulation::CheckAttriubtesForRenderer()
 {
 	bool bOk = true;
