Fixed visual logger automation tests for projects with disabled AI logging in engine (bDisableAILogging flag).

diff --git a/Engine/Source/Runtime/Engine/Private/VisualLogger/VisualLoggerAutomationTests.cpp b/Engine/Source/Runtime/Engine/Private/VisualLogger/VisualLoggerAutomationTests.cpp
--- a/Engine/Source/Runtime/Engine/Private/VisualLogger/VisualLoggerAutomationTests.cpp
+++ b/Engine/Source/Runtime/Engine/Private/VisualLogger/VisualLoggerAutomationTests.cpp
@@ -79,23 +79,30 @@ if ((__Test__)) \
 template<typename TYPE = FVisualLoggerTestDevice>
 struct FTestDeviceContext
 {
-	FTestDeviceContext() 
+	FTestDeviceContext()
 	{
 		Device.Cleanup();
 		FVisualLogger::Get().SetIsRecording(false);
 		FVisualLogger::Get().Cleanup();
 		FVisualLogger::Get().AddDevice(&Device);
+
+		EngineDisableAILoggingFlag = GEngine->bDisableAILogging;
+		GEngine->bDisableAILogging = false;
+
 	}
 
-	~FTestDeviceContext() 
+	~FTestDeviceContext()
 	{
 		FVisualLogger::Get().SetIsRecording(false);
 		FVisualLogger::Get().RemoveDevice(&Device);
 		FVisualLogger::Get().Cleanup();
 		Device.Cleanup();
+
+		GEngine->bDisableAILogging = EngineDisableAILoggingFlag;
 	}
 
 	TYPE Device;
+	uint32 EngineDisableAILoggingFlag : 1;
 };
 
 IMPLEMENT_SIMPLE_AUTOMATION_TEST(FVisualLogTest, "Engine.VisualLogger.Logging simple text", EAutomationTestFlags::ATF_Game | EAutomationTestFlags::ATF_Editor)
