UE4 - fix compile error on shipping builds

diff --git a/Engine/Source/Runtime/Engine/Private/StaticBoundShaderState.cpp b/Engine/Source/Runtime/Engine/Private/StaticBoundShaderState.cpp
--- a/Engine/Source/Runtime/Engine/Private/StaticBoundShaderState.cpp
+++ b/Engine/Source/Runtime/Engine/Private/StaticBoundShaderState.cpp
@@ -15,10 +15,12 @@ TLinkedList<FGlobalBoundShaderStateResource*>*& FGlobalBoundShaderStateResource:
 
 FGlobalBoundShaderStateResource::FGlobalBoundShaderStateResource()
 	: GlobalListLink(this)
+#if DO_CHECK
 	, BoundVertexDeclaration(nullptr)
 	, BoundVertexShader(nullptr)
 	, BoundPixelShader(nullptr)
 	, BoundGeometryShader(nullptr)
+#endif 
 {
 	// Add this resource to the global list in the rendering thread.
 	if(IsInRenderingThread())
