#ue4 - Always clear output FHitResult in RaycastSingle and SweepSingle regardless of whether there is a hit.

UE-13311

diff --git a/Engine/Source/Runtime/Engine/Private/Collision/PhysXCollision.cpp b/Engine/Source/Runtime/Engine/Private/Collision/PhysXCollision.cpp
--- a/Engine/Source/Runtime/Engine/Private/Collision/PhysXCollision.cpp
+++ b/Engine/Source/Runtime/Engine/Private/Collision/PhysXCollision.cpp
@@ -885,6 +885,7 @@ bool RaycastSingle(const UWorld* World, struct FHitResult& OutHit, const FVector
 	SCOPE_CYCLE_COUNTER(STAT_Collision_RaycastSingle);
 	STARTQUERYTIMER();
 
+	OutHit = FHitResult();
 	OutHit.TraceStart = Start;
 	OutHit.TraceEnd = End;
 
@@ -1000,15 +1001,16 @@ bool RaycastSingle(const UWorld* World, struct FHitResult& OutHit, const FVector
 
 bool RaycastMulti(const UWorld* World, TArray<struct FHitResult>& OutHits, const FVector& Start, const FVector& End, ECollisionChannel TraceChannel, const struct FCollisionQueryParams& Params, const struct FCollisionResponseParams& ResponseParams, const struct FCollisionObjectQueryParams& ObjectParams)
 {
-	if ((World == NULL) || (World->GetPhysicsScene() == NULL))
-	{
-		return false;
-	}
 	SCOPE_CYCLE_COUNTER(STAT_Collision_RaycastMultiple);
 	STARTQUERYTIMER();
 
 	OutHits.Reset();
 
+	if ((World == NULL) || (World->GetPhysicsScene() == NULL))
+	{
+		return false;
+	}
+
 	// Track if we get any 'blocking' hits
 	bool bHaveBlockingHit = false;
 
@@ -1239,6 +1241,7 @@ bool GeomSweepSingle(const UWorld* World, const struct FCollisionShape& Collisio
 	SCOPE_CYCLE_COUNTER(STAT_Collision_GeomSweepSingle);
 	STARTQUERYTIMER();
 
+	OutHit = FHitResult();
 	OutHit.TraceStart = Start;
 	OutHit.TraceEnd = End;
 
@@ -1522,13 +1525,14 @@ bool GeomSweepMulti_PhysX(const UWorld* World, const PxGeometry& PGeom, const Px
 
 bool GeomSweepMulti(const UWorld* World, const struct FCollisionShape& CollisionShape, const FQuat& Rot, TArray<FHitResult>& OutHits, FVector Start, FVector End, ECollisionChannel TraceChannel, const struct FCollisionQueryParams& Params, const struct FCollisionResponseParams& ResponseParams, const struct FCollisionObjectQueryParams& ObjectParams)
 {
+	SCOPE_CYCLE_COUNTER(STAT_Collision_GeomSweepMultiple);
+
+	OutHits.Reset();
+
 	if ((World == NULL) || (World->GetPhysicsScene() == NULL))
 	{
 		return false;
 	}
-	SCOPE_CYCLE_COUNTER(STAT_Collision_GeomSweepMultiple);
-
-	OutHits.Reset();
 
 	// Track if we get any 'blocking' hits
 	bool bBlockingHit = false;
@@ -1697,11 +1701,11 @@ bool GeomOverlapMulti_PhysX(const UWorld* World, const PxGeometry& PGeom, const
 template <EQueryInfo::Type InfoType>
 bool GeomOverlapMultiImp(const UWorld* World, const struct FCollisionShape& CollisionShape, const FVector& Pos, const FQuat& Rot, TArray<FOverlapResult>& OutOverlaps, ECollisionChannel TraceChannel, const struct FCollisionQueryParams& Params, const struct FCollisionResponseParams& ResponseParams, const struct FCollisionObjectQueryParams& ObjectParams)
 {
+	SCOPE_CYCLE_COUNTER(STAT_Collision_GeomOverlapMultiple);
 	if ((World == NULL) || (World->GetPhysicsScene() == NULL))
 	{
 		return false;
 	}
-	SCOPE_CYCLE_COUNTER(STAT_Collision_GeomOverlapMultiple);
 
 	// Track if we get any 'blocking' hits
 	bool bHaveBlockingHit = false;
