#ue4 - Don't crash on invalid normal in ConvertQueryImpactHit() (changed to ensure).

- Add better handling of initial overlaps for raycasts.
- Don't allow invalid normal to be returned out of ConvertQueryImpactHit() in any case.

Jira: UE-13194

diff --git a/Engine/Source/Runtime/Engine/Private/Collision/CollisionConversions.cpp b/Engine/Source/Runtime/Engine/Private/Collision/CollisionConversions.cpp
--- a/Engine/Source/Runtime/Engine/Private/Collision/CollisionConversions.cpp
+++ b/Engine/Source/Runtime/Engine/Private/Collision/CollisionConversions.cpp
@@ -26,29 +26,27 @@ DECLARE_CYCLE_STAT(TEXT("ConvertOverlapToHit"), STAT_CollisionConvertOverlapToHi
 DECLARE_CYCLE_STAT(TEXT("ConvertOverlap"), STAT_CollisionConvertOverlap, STATGROUP_Collision);
 
 
-#if !(UE_BUILD_SHIPPING || UE_BUILD_TEST)
+#if !(UE_BUILD_SHIPPING || UE_BUILD_TEST || !WITH_EDITOR)
 /* Validate Normal of OutResult. We're on hunt for invalid normal */
 static void CheckHitResultNormal(const FHitResult& OutResult, const TCHAR* Message, const FVector& Start=FVector::ZeroVector, const FVector& End = FVector::ZeroVector, const PxGeometry* const Geom=NULL)
 {
 	if(!OutResult.bStartPenetrating && !OutResult.Normal.IsNormalized())
 	{
-		UE_LOG(LogPhysics, Warning, TEXT("(%s) Non-normalized OutResult.Normal from capsule conversion: %s (Component- %s)"), Message, *OutResult.Normal.ToString(), *GetNameSafe(OutResult.Component.Get()));
-		// now I'm adding Outresult input
+		UE_LOG(LogPhysics, Warning, TEXT("(%s) Non-normalized OutResult.Normal from hit conversion: %s (Component- %s)"), Message, *OutResult.Normal.ToString(), *GetNameSafe(OutResult.Component.Get()));
 		UE_LOG(LogPhysics, Warning, TEXT("Start Loc(%s), End Loc(%s), Hit Loc(%s), ImpactNormal(%s)"), *Start.ToString(), *End.ToString(), *OutResult.Location.ToString(), *OutResult.ImpactNormal.ToString() );
 		if (Geom != NULL)
 		{
-			// I only seen this crash on capsule
 			if (Geom->getType() == PxGeometryType::eCAPSULE)
 			{
 				const PxCapsuleGeometry * Capsule = (PxCapsuleGeometry*)Geom;
 				UE_LOG(LogPhysics, Warning, TEXT("Capsule radius (%f), Capsule Halfheight (%f)"), Capsule->radius, Capsule->halfHeight);
 			}
 		}
-		check(OutResult.Normal.IsNormalized());
+		ensure(OutResult.Normal.IsNormalized());
 	}
 }
 
-#endif // !(UE_BUILD_SHIPPING || UE_BUILD_TEST)
+#endif // !(UE_BUILD_SHIPPING || UE_BUILD_TEST || !WITH_EDITOR)
 
 
 static FORCEINLINE bool PxQuatIsIdentity(PxQuat const& Q)
@@ -345,7 +343,8 @@ void ConvertQueryImpactHit(const UWorld* World, const PxLocationHit& PHit, FHitR
 	SCOPE_CYCLE_COUNTER(STAT_ConvertQueryImpactHit);
 
 	checkSlow(PHit.flags & PxHitFlag::eDISTANCE);
-	if (Geom != NULL && PHit.hadInitialOverlap())
+	const bool bInitialOverlap = PHit.hadInitialOverlap();
+	if (bInitialOverlap && Geom != nullptr)
 	{
 		ConvertOverlappedShapeToImpactHit(World, PHit, StartLoc, EndLoc, OutResult, *Geom, QueryTM, QueryFilter, bReturnPhysMat);
 		return;
@@ -361,10 +360,16 @@ void ConvertQueryImpactHit(const UWorld* World, const PxLocationHit& PHit, FHitR
 	const FVector SafeLocationToFitShape = StartLoc + (HitTime * TraceStartToEnd);
 
 	// Other info
+	// Caution: we may still have an initial overlap, but with null Geom. This is the case for RayCast results.
 	OutResult.Location = SafeLocationToFitShape;
-	OutResult.ImpactPoint = (PHit.flags & PxHitFlag::ePOSITION) ? P2UVector(PHit.position) : StartLoc;
-	OutResult.Normal = (PHit.flags & PxHitFlag::eNORMAL) ? P2UVector(PHit.normal).GetSafeNormal() : -TraceStartToEnd.GetSafeNormal();
-	OutResult.ImpactNormal = OutResult.Normal;
+
+	const bool bUsePxPoint = ((PHit.flags & PxHitFlag::ePOSITION) && !bInitialOverlap);
+	OutResult.ImpactPoint = bUsePxPoint ? P2UVector(PHit.position) : StartLoc;
+	
+	const bool bUsePxNormal = ((PHit.flags & PxHitFlag::eNORMAL) && !bInitialOverlap);
+	FVector Normal = bUsePxNormal ? P2UVector(PHit.normal).GetSafeNormal() : -TraceStartToEnd.GetSafeNormal();
+	OutResult.Normal = Normal;
+	OutResult.ImpactNormal = Normal;
 
 	OutResult.TraceStart = StartLoc;
 	OutResult.TraceEnd = EndLoc;
@@ -374,13 +379,21 @@ void ConvertQueryImpactHit(const UWorld* World, const PxLocationHit& PHit, FHitR
 	PxFilterData PShapeFilter = PHit.shape->getQueryFilterData();
 	PxSceneQueryHitType::Enum HitType = FPxQueryFilterCallback::CalcQueryHitType(QueryFilter, PShapeFilter);
 	OutResult.bBlockingHit = (HitType == PxSceneQueryHitType::eBLOCK); 
-	OutResult.bStartPenetrating = (PxU32)(PHit.hadInitialOverlap());
+	OutResult.bStartPenetrating = bInitialOverlap;
 
 
 #if !(UE_BUILD_SHIPPING || UE_BUILD_TEST || !WITH_EDITOR)
 	CheckHitResultNormal(OutResult, TEXT("Invalid Normal from ConvertQueryImpactHit"), StartLoc, EndLoc, Geom);
 #endif
 
+	if (bUsePxNormal && !Normal.IsNormalized())
+	{
+		// TraceStartToEnd should never be zero, because of the length restriction in the raycast and sweep tests.
+		Normal = -TraceStartToEnd.GetSafeNormal();
+		OutResult.Normal = Normal;
+		OutResult.ImpactNormal = Normal;
+	}
+
 	PxGeometryType::Enum GeometryType = Geom ? Geom->getType() : PxGeometryType::eINVALID;
 
 	// Special handling for swept-capsule results
