Fix shadow variables

diff --git a/Engine/Plugins/Developer/GitSourceControl/Source/GitSourceControl/Private/GitSourceControlUtils.cpp b/Engine/Plugins/Developer/GitSourceControl/Source/GitSourceControl/Private/GitSourceControlUtils.cpp
--- a/Engine/Plugins/Developer/GitSourceControl/Source/GitSourceControl/Private/GitSourceControlUtils.cpp
+++ b/Engine/Plugins/Developer/GitSourceControl/Source/GitSourceControl/Private/GitSourceControlUtils.cpp
@@ -303,13 +303,15 @@ bool RunCommit(const FString& InPathToGitBinary, const FString& InRepositoryRoot
 	{
 		// Batch files up so we dont exceed command-line limits
 		int32 FileCount = 0;
-		TArray<FString> FilesInBatch;
-		for(int32 FileIndex = 0; FileIndex < GitSourceControlConstants::MaxFilesPerBatch; FileIndex++, FileCount++)
 		{
-			FilesInBatch.Add(InFiles[FileCount]);
+			TArray<FString> FilesInBatch;
+			for(int32 FileIndex = 0; FileIndex < GitSourceControlConstants::MaxFilesPerBatch; FileIndex++, FileCount++)
+			{
+				FilesInBatch.Add(InFiles[FileCount]);
+			}
+			// First batch is a simple "git commit" command with only the first files
+			bResult &= RunCommandInternal(TEXT("commit"), InPathToGitBinary, InRepositoryRoot, InParameters, FilesInBatch, OutResults, OutErrorMessages);
 		}
-		// First batch is a simple "git commit" command with only the first files
-		bResult &= RunCommandInternal(TEXT("commit"), InPathToGitBinary, InRepositoryRoot, InParameters, FilesInBatch, OutResults, OutErrorMessages);
 		
 		TArray<FString> Parameters;
 		for(const auto& Parameter : InParameters)
diff --git a/Engine/Source/Developer/iOS/MetalShaderFormat/Private/MetalUtils.cpp b/Engine/Source/Developer/iOS/MetalShaderFormat/Private/MetalUtils.cpp
--- a/Engine/Source/Developer/iOS/MetalShaderFormat/Private/MetalUtils.cpp
+++ b/Engine/Source/Developer/iOS/MetalShaderFormat/Private/MetalUtils.cpp
@@ -298,7 +298,6 @@ namespace MetalUtils
 			return nullptr;
 		}
 
-		ir_variable* Variable = NULL;
 		if (FCStringAnsi::Strnicmp(Semantic, "SV_", 3) == 0)
 		{
 			FSystemValue* SystemValues = SystemValueTable[Frequency];
@@ -350,7 +349,7 @@ namespace MetalUtils
 			_mesa_glsl_warning(ParseState, "unrecognized system value input '%s'", Semantic);
 		}
 
-		Variable = new(ParseState)ir_variable(
+		ir_variable* Variable = new(ParseState)ir_variable(
 			Type,
 			ralloc_asprintf(ParseState, "IN_%s", Semantic),
 			ir_var_in);
diff --git a/Engine/Source/Editor/GraphEditor/Private/ConnectionDrawingPolicy.cpp b/Engine/Source/Editor/GraphEditor/Private/ConnectionDrawingPolicy.cpp
--- a/Engine/Source/Editor/GraphEditor/Private/ConnectionDrawingPolicy.cpp
+++ b/Engine/Source/Editor/GraphEditor/Private/ConnectionDrawingPolicy.cpp
@@ -287,7 +287,6 @@ void FConnectionDrawingPolicy::DrawConnection(int32 LayerId, const FVector2D& St
 			float ClosestDistanceSquared = FLT_MAX;
 
 			const int32 NumStepsToTest = 16;
-			float TestAlpha = 0.0f;
 			const float StepInterval = 1.0f / (float)NumStepsToTest;
 			FVector2D Point1 = FMath::CubicInterp(P0, P0Tangent, P1, P1Tangent, 0.0f);
 			for (float TestAlpha = 0.0f; TestAlpha < 1.0f; TestAlpha += StepInterval)
diff --git a/Engine/Source/Editor/Kismet/Private/NativeCodeGenerationTool.cpp b/Engine/Source/Editor/Kismet/Private/NativeCodeGenerationTool.cpp
--- a/Engine/Source/Editor/Kismet/Private/NativeCodeGenerationTool.cpp
+++ b/Engine/Source/Editor/Kismet/Private/NativeCodeGenerationTool.cpp
@@ -37,7 +37,7 @@ struct FGeneratedCodeData
 	FString ClassName;
 	TWeakObjectPtr<UBlueprint> Blueprint;
 
-	void GatherUserDefinedDependencies(UBlueprint& Blueprint)
+	void GatherUserDefinedDependencies(UBlueprint& InBlueprint)
 	{
 		TArray<UObject*> ReferencedObjects;
 		{
@@ -45,7 +45,7 @@ struct FGeneratedCodeData
 
 			{
 				TArray<UObject*> ObjectsToCheck;
-				GetObjectsWithOuter(Blueprint.GeneratedClass, ObjectsToCheck, true);
+				GetObjectsWithOuter(InBlueprint.GeneratedClass, ObjectsToCheck, true);
 				for (auto Obj : ObjectsToCheck)
 				{
 					if (IsValid(Obj))
@@ -55,7 +55,7 @@ struct FGeneratedCodeData
 				}
 			}
 
-			for (UClass* Class = Blueprint.GeneratedClass->GetSuperClass(); Class && !Class->HasAnyClassFlags(CLASS_Native); Class = Class->GetSuperClass())
+			for (UClass* Class = InBlueprint.GeneratedClass->GetSuperClass(); Class && !Class->HasAnyClassFlags(CLASS_Native); Class = Class->GetSuperClass())
 			{
 				ReferencedObjects.Add(Class);
 			}
@@ -64,7 +64,7 @@ struct FGeneratedCodeData
 		TypeDependencies.Empty();
 		for (auto Obj : ReferencedObjects)
 		{
-			if (IsValid(Obj) && !Obj->IsIn(Blueprint.GeneratedClass) && (Obj != Blueprint.GeneratedClass))
+			if (IsValid(Obj) && !Obj->IsIn(InBlueprint.GeneratedClass) && (Obj != InBlueprint.GeneratedClass))
 			{
 				if (Obj->IsA<UBlueprintGeneratedClass>() || Obj->IsA<UUserDefinedEnum>() || Obj->IsA<UUserDefinedStruct>())
 				{
diff --git a/Engine/Source/Editor/MaterialEditor/Private/MaterialEditor.cpp b/Engine/Source/Editor/MaterialEditor/Private/MaterialEditor.cpp
--- a/Engine/Source/Editor/MaterialEditor/Private/MaterialEditor.cpp
+++ b/Engine/Source/Editor/MaterialEditor/Private/MaterialEditor.cpp
@@ -2459,21 +2459,21 @@ void FMaterialEditor::OnAssetPostImport(UFactory* InFactory, UObject* InObject)
 			{
 				struct Local
 				{
-					static bool ReferencesImportedObject(UMaterialFunction* Function, UTexture* Texture)
+					static bool ReferencesImportedObject(UMaterialFunction* Function, UTexture* InTexture)
 					{
-						for (UMaterialExpression* Expression : Function->FunctionExpressions)
+						for (UMaterialExpression* MaterialExpression : Function->FunctionExpressions)
 						{
-							if (UMaterialExpressionTextureBase* ExpressionTexture = Cast<UMaterialExpressionTextureBase>(Expression))
+							if (UMaterialExpressionTextureBase* MaterialExpressionTexture = Cast<UMaterialExpressionTextureBase>(MaterialExpression))
 							{
-								if (ExpressionTexture->Texture == Texture)
+								if (MaterialExpressionTexture->Texture == InTexture)
 								{
 									return true;
 								}
 							}
-							else if (UMaterialExpressionMaterialFunctionCall* ExpressionFunctionCall = Cast<UMaterialExpressionMaterialFunctionCall>(Expression))
+							else if (UMaterialExpressionMaterialFunctionCall* MaterialExpressionFunctionCall = Cast<UMaterialExpressionMaterialFunctionCall>(MaterialExpression))
 							{
-								if (ExpressionFunctionCall->MaterialFunction != nullptr &&
-									ReferencesImportedObject(ExpressionFunctionCall->MaterialFunction, Texture))
+								if (MaterialExpressionFunctionCall->MaterialFunction != nullptr &&
+									ReferencesImportedObject(MaterialExpressionFunctionCall->MaterialFunction, InTexture))
 								{
 									return true;
 								}
diff --git a/Engine/Source/Editor/UMGEditor/Private/Designer/SDesignerView.cpp b/Engine/Source/Editor/UMGEditor/Private/Designer/SDesignerView.cpp
--- a/Engine/Source/Editor/UMGEditor/Private/Designer/SDesignerView.cpp
+++ b/Engine/Source/Editor/UMGEditor/Private/Designer/SDesignerView.cpp
@@ -1524,7 +1524,8 @@ void SDesignerView::Tick( const FGeometry& AllottedGeometry, const double InCurr
 		SideRuler->SetCursor(TOptional<FVector2D>());
 	}
 
-	if ( UUserWidget* DefaultWidget = GetDefaultWidget() )
+	DefaultWidget = GetDefaultWidget();
+	if (DefaultWidget)
 	{
 		const bool bNeedDesiredSize = 
 			DefaultWidget->DesignSizeMode == EDesignPreviewSizeMode::Desired || 
diff --git a/Engine/Source/Editor/UnrealEd/Private/AutoReimport/AutoReimportTests.cpp b/Engine/Source/Editor/UnrealEd/Private/AutoReimport/AutoReimportTests.cpp
--- a/Engine/Source/Editor/UnrealEd/Private/AutoReimport/AutoReimportTests.cpp
+++ b/Engine/Source/Editor/UnrealEd/Private/AutoReimport/AutoReimportTests.cpp
@@ -198,12 +198,14 @@ bool FAutoReimportSimpleModifyTest::RunTest(const FString& Parameters)
 	// Start watching the directory
 	TSharedPtr<FAutoReimportTestPayload> Test = MakeShareable(new FAutoReimportTestPayload(WorkingDir));
 
-	TArray<FSrcDstFilenames> Files;
-	Files.Emplace(SrcFilename1, DstFilename);
-
-	if (!CopyTestFiles(*Test, Files))
 	{
-		return false;
+		TArray<FSrcDstFilenames> Files;
+		Files.Emplace(SrcFilename1, DstFilename);
+
+		if (!CopyTestFiles(*Test, Files))
+		{
+			return false;
+		}
 	}
 
 	Test->StartWatching();
@@ -461,11 +463,13 @@ bool FAutoReimportRestartDetectionTest::RunTest(const FString& Parameters)
 	
 	TSharedPtr<FAutoReimportTestPayload> Test = MakeShareable(new FAutoReimportTestPayload(WorkingDir));
 
-	TArray<FSrcDstFilenames> Files;
-	Files.Emplace(SrcFilename, SrcFilename);
-	if (!CopyTestFiles(*Test, Files))
 	{
-		return false;
+		TArray<FSrcDstFilenames> Files;
+		Files.Emplace(SrcFilename, SrcFilename);
+		if (!CopyTestFiles(*Test, Files))
+		{
+			return false;
+		}
 	}
 
 	// Start watching the directory
@@ -483,11 +487,13 @@ bool FAutoReimportRestartDetectionTest::RunTest(const FString& Parameters)
 		Test->Config.bDetectChangesSinceLastRun = true;
 
 		// Modify the file while the watcher isn't running
-		TArray<FSrcDstFilenames> Files;
-		Files.Emplace(SrcFilename2, SrcFilename);
-		if (!CopyTestFiles(*Test, Files))
 		{
-			return;
+			TArray<FSrcDstFilenames> Files;
+			Files.Emplace(SrcFilename2, SrcFilename);
+			if (!CopyTestFiles(*Test, Files))
+			{
+				return;
+			}
 		}
 
 		Test->StartWatching();
@@ -531,10 +537,10 @@ bool FAutoReimportRestartDetectionTest::RunTest(const FString& Parameters)
 			Test->StartWatching();
 			Test->WaitForStartup([=]{
 
-				auto Changes = Test->FileCache->GetOutstandingChanges();
-				if (Changes.Num() != 0)
+				auto OutstandingChanges = Test->FileCache->GetOutstandingChanges();
+				if (OutstandingChanges.Num() != 0)
 				{
-					UE_LOG(LogAutoReimportTests, Error, TEXT("Shouldn't have reported changes when bDetectChangesSinceLastRun is false (%d change(s) received)."), Changes.Num());
+					UE_LOG(LogAutoReimportTests, Error, TEXT("Shouldn't have reported changes when bDetectChangesSinceLastRun is false (%d change(s) received)."), OutstandingChanges.Num());
 					return;
 				}
 
@@ -620,12 +626,14 @@ bool FAutoReimportChangeExtensionsTest::RunTest(const FString& Parameters)
 	Test->WaitForStartup([=]{
 
 		// Add a txt file and a png - we should only be told about the txt
-		TArray<FSrcDstFilenames> Files;
-		Files.Emplace(SrcFilename1, SrcFilename1);
-		Files.Emplace(SrcFilename3, SrcFilename3);
-		if (!CopyTestFiles(*Test, Files))
 		{
-			return;
+			TArray<FSrcDstFilenames> Files;
+			Files.Emplace(SrcFilename1, SrcFilename1);
+			Files.Emplace(SrcFilename3, SrcFilename3);
+			if (!CopyTestFiles(*Test, Files))
+			{
+				return;
+			}
 		}
 
 		ADD_LATENT_AUTOMATION_COMMAND(FDelayedCallbackLatentCommand([=]{
@@ -657,11 +665,13 @@ bool FAutoReimportChangeExtensionsTest::RunTest(const FString& Parameters)
 			Test->StartWatching();
 			Test->WaitForStartup([=]{
 
-				auto Changes = Test->FileCache->GetOutstandingChanges();
-				if (Changes.Num() != 0)
 				{
-					UE_LOG(LogAutoReimportTests, Error, TEXT("Incorrect number of changes reported (%d != 0)."), Changes.Num());
-					return;
+					auto OutstandingChanges = Test->FileCache->GetOutstandingChanges();
+					if (OutstandingChanges.Num() != 0)
+					{
+						UE_LOG(LogAutoReimportTests, Error, TEXT("Incorrect number of changes reported (%d != 0)."), OutstandingChanges.Num());
+						return;
+					}
 				}
 
 				// Add another png - we should only be notified about this one
@@ -673,11 +683,11 @@ bool FAutoReimportChangeExtensionsTest::RunTest(const FString& Parameters)
 				}
 
 				ADD_LATENT_AUTOMATION_COMMAND(FDelayedCallbackLatentCommand([=]{
-					auto Changes = Test->FileCache->GetOutstandingChanges();
+					auto OutstandingChanges = Test->FileCache->GetOutstandingChanges();
 
-					if (Changes.Num() != 1)
+					if (OutstandingChanges.Num() != 1)
 					{
-						UE_LOG(LogAutoReimportTests, Error, TEXT("Incorrect number of changes reported (%d != 1)."), Changes.Num());
+						UE_LOG(LogAutoReimportTests, Error, TEXT("Incorrect number of changes reported (%d != 1)."), OutstandingChanges.Num());
 					}
 				}, 1));
 
@@ -711,13 +721,15 @@ bool FAutoReimportWildcardFiltersTest::RunTest(const FString& Parameters)
 
 	IFileManager::Get().MakeDirectory(*(GetWorkingDir() / TEXT("Content") / TEXT("sub-folder")), false);
 		
-	TArray<FSrcDstFilenames> Files;
-	Files.Emplace(SrcFilename1, DstFilename1);
-	Files.Emplace(SrcFilename2, DstFilename2);
-	Files.Emplace(SrcFilename3, DstFilename3);
-	if (!CopyTestFiles(*Test, Files))
 	{
-		return false;
+		TArray<FSrcDstFilenames> Files;
+		Files.Emplace(SrcFilename1, DstFilename1);
+		Files.Emplace(SrcFilename2, DstFilename2);
+		Files.Emplace(SrcFilename3, DstFilename3);
+		if (!CopyTestFiles(*Test, Files))
+		{
+			return false;
+		}
 	}
 
 	// Start watching the directory
diff --git a/Engine/Source/Editor/UnrealEd/Private/Editor.cpp b/Engine/Source/Editor/UnrealEd/Private/Editor.cpp
--- a/Engine/Source/Editor/UnrealEd/Private/Editor.cpp
+++ b/Engine/Source/Editor/UnrealEd/Private/Editor.cpp
@@ -2623,7 +2623,7 @@ void FReimportManager::UpdateReimportPaths( UObject* Obj, const TArray<FString>&
 	if (Obj)
 	{
 		TArray<FString> UnusedExistingFilenames;
-		auto* Handler = Handlers.FindByPredicate([&](FReimportHandler* Handler){ return Handler->CanReimport(Obj, UnusedExistingFilenames); });
+		auto* Handler = Handlers.FindByPredicate([&](FReimportHandler* InHandler){ return InHandler->CanReimport(Obj, UnusedExistingFilenames); });
 		if (Handler)
 		{
 			(*Handler)->SetReimportPaths(Obj, InFilenames);
diff --git a/Engine/Source/Editor/UnrealEd/Private/Kismet2/BlueprintEditorUtils.cpp b/Engine/Source/Editor/UnrealEd/Private/Kismet2/BlueprintEditorUtils.cpp
--- a/Engine/Source/Editor/UnrealEd/Private/Kismet2/BlueprintEditorUtils.cpp
+++ b/Engine/Source/Editor/UnrealEd/Private/Kismet2/BlueprintEditorUtils.cpp
@@ -1898,12 +1898,12 @@ void FBlueprintEditorUtils::MarkBlueprintAsStructurallyModified(UBlueprint* Blue
 		{
 			for (auto SkelClass : SkelClassesToRecompile)
 			{
-				auto Blueprint = Cast<UBlueprint>(SkelClass->ClassGeneratedBy);
-				if (Blueprint
-					&& Blueprint->Status != BS_BeingCreated
-					&& !Blueprint->bBeingCompiled
-					&& !Blueprint->bIsRegeneratingOnLoad
-					&& Blueprint->bHasBeenRegenerated)
+				auto SkelBlueprint = Cast<UBlueprint>(SkelClass->ClassGeneratedBy);
+				if (SkelBlueprint
+					&& SkelBlueprint->Status != BS_BeingCreated
+					&& !SkelBlueprint->bBeingCompiled
+					&& !SkelBlueprint->bIsRegeneratingOnLoad
+					&& SkelBlueprint->bHasBeenRegenerated)
 				{
 					TArray<UClass*> ChildrenOfClass;
 					GetDerivedClasses(SkelClass, ChildrenOfClass, false);
@@ -1916,7 +1916,7 @@ void FBlueprintEditorUtils::MarkBlueprintAsStructurallyModified(UBlueprint* Blue
 
 					FKismetCompilerOptions CompileOptions;
 					CompileOptions.CompileType = EKismetCompileType::SkeletonOnly;
-					Compiler.CompileBlueprint(Blueprint, CompileOptions, Results);
+					Compiler.CompileBlueprint(SkelBlueprint, CompileOptions, Results);
 
 					SkeletalRecompileChildren(ChildrenOfClass);
 				}
diff --git a/Engine/Source/Editor/UnrealEd/Private/Kismet2/StructureEditorUtils.cpp b/Engine/Source/Editor/UnrealEd/Private/Kismet2/StructureEditorUtils.cpp
--- a/Engine/Source/Editor/UnrealEd/Private/Kismet2/StructureEditorUtils.cpp
+++ b/Engine/Source/Editor/UnrealEd/Private/Kismet2/StructureEditorUtils.cpp
@@ -365,7 +365,7 @@ bool FStructureEditorUtils::ChangeVariableType(UUserDefinedStruct* Struct, FGuid
 
 bool FStructureEditorUtils::ChangeVariableDefaultValue(UUserDefinedStruct* Struct, FGuid VarGuid, const FString& NewDefaultValue)
 {
-	auto ValidateDefaultValue = [](const FStructVariableDescription& VarDesc, const FString& NewDefaultValue) -> bool
+	auto ValidateDefaultValue = [](const FStructVariableDescription& VarDesc, const FString& InNewDefaultValue) -> bool
 	{
 		const UEdGraphSchema_K2* K2Schema = GetDefault<UEdGraphSchema_K2>();
 		const FEdGraphPinType PinType = VarDesc.ToPinType();
@@ -382,7 +382,7 @@ bool FStructureEditorUtils::ChangeVariableDefaultValue(UUserDefinedStruct* Struc
 		}
 		else
 		{
-			bResult = K2Schema->DefaultValueSimpleValidation(PinType, FString(), NewDefaultValue, NULL, FText::GetEmpty());
+			bResult = K2Schema->DefaultValueSimpleValidation(PinType, FString(), InNewDefaultValue, NULL, FText::GetEmpty());
 		}
 		return bResult;
 	};
diff --git a/Engine/Source/Editor/UnrealEd/Private/LevelEditorViewport.cpp b/Engine/Source/Editor/UnrealEd/Private/LevelEditorViewport.cpp
--- a/Engine/Source/Editor/UnrealEd/Private/LevelEditorViewport.cpp
+++ b/Engine/Source/Editor/UnrealEd/Private/LevelEditorViewport.cpp
@@ -3884,17 +3884,17 @@ void FLevelEditorViewportClient::DrawBrushDetails(const FSceneView* View, FPrimi
 
 void FLevelEditorViewportClient::UpdateAudioListener(const FSceneView& View)
 {
-	UWorld* World = GetWorld();
+	UWorld* ViewportWorld = GetWorld();
 
-	if (World)
+	if (ViewportWorld)
 	{
-		if (FAudioDevice* AudioDevice = World->GetAudioDevice())
+		if (FAudioDevice* AudioDevice = ViewportWorld->GetAudioDevice())
 		{
 			FReverbSettings ReverbSettings;
 			FInteriorSettings InteriorSettings;
 			const FVector& ViewLocation = GetViewLocation();
 
-			class AAudioVolume* AudioVolume = World->GetAudioSettings(ViewLocation, &ReverbSettings, &InteriorSettings);
+			class AAudioVolume* AudioVolume = ViewportWorld->GetAudioSettings(ViewLocation, &ReverbSettings, &InteriorSettings);
 
 			FMatrix CameraToWorld = View.ViewMatrices.ViewMatrix.InverseFast();
 			FVector ProjUp = CameraToWorld.TransformVector(FVector(0, 1000, 0));
diff --git a/Engine/Source/Editor/UnrealEd/Private/PlayLevel.cpp b/Engine/Source/Editor/UnrealEd/Private/PlayLevel.cpp
--- a/Engine/Source/Editor/UnrealEd/Private/PlayLevel.cpp
+++ b/Engine/Source/Editor/UnrealEd/Private/PlayLevel.cpp
@@ -2679,7 +2679,6 @@ UGameInstance* UEditorEngine::CreatePIEGameInstance(int32 PIEInstance, bool bInS
 	
 	if (!PieWorldContext->RunAsDedicated)
 	{
-		const ULevelEditorPlaySettings* PlayInSettings = GetDefault<ULevelEditorPlaySettings>();
 		bool bCreateNewAudioDevice = PlayInSettings->IsCreateAudioDeviceForEveryPlayer();
 
 		ViewportClient = NewObject<UGameViewportClient>(this, GameViewportClientClass);
diff --git a/Engine/Source/Editor/UnrealEd/Private/Settings/SettingsClasses.cpp b/Engine/Source/Editor/UnrealEd/Private/Settings/SettingsClasses.cpp
--- a/Engine/Source/Editor/UnrealEd/Private/Settings/SettingsClasses.cpp
+++ b/Engine/Source/Editor/UnrealEd/Private/Settings/SettingsClasses.cpp
@@ -272,8 +272,8 @@ bool FAutoReimportDirectoryConfig::ParseSourceDirectoryAndMountPoint(FString& So
 		}
 
 		// Set the mounted path if necessary
-		auto* Pair = InContext.MountedPaths.FindByPredicate([&](const TPair<FString, FString>& Pair){
-			return SourceDirectory.StartsWith(Pair.Key);
+		auto* Pair = InContext.MountedPaths.FindByPredicate([&](const TPair<FString, FString>& InPair){
+			return SourceDirectory.StartsWith(InPair.Key);
 		});
 
 		if (Pair)
diff --git a/Engine/Source/Runtime/AIModule/Private/AISystem.cpp b/Engine/Source/Runtime/AIModule/Private/AISystem.cpp
--- a/Engine/Source/Runtime/AIModule/Private/AISystem.cpp
+++ b/Engine/Source/Runtime/AIModule/Private/AISystem.cpp
@@ -125,13 +125,13 @@ void UAISystem::RunEQS(const FString& QueryName, UObject* Target)
 #endif
 }
 
-UAISystem::FBlackboardDataToComponentsIterator::FBlackboardDataToComponentsIterator(FBlackboardDataToComponentsMap& BlackboardDataToComponentsMap, UBlackboardData* BlackboardAsset)
+UAISystem::FBlackboardDataToComponentsIterator::FBlackboardDataToComponentsIterator(FBlackboardDataToComponentsMap& InBlackboardDataToComponentsMap, UBlackboardData* BlackboardAsset)
 	: CurrentIteratorIndex(0)
 	, Iterators()
 {
 	while (BlackboardAsset)
 	{
-		Iterators.Add(BlackboardDataToComponentsMap.CreateConstKeyIterator(BlackboardAsset));
+		Iterators.Add(InBlackboardDataToComponentsMap.CreateConstKeyIterator(BlackboardAsset));
 		BlackboardAsset = BlackboardAsset->Parent;
 	}
 	TryMoveIteratorToParentBlackboard();
diff --git a/Engine/Source/Runtime/AIModule/Private/Actions/PawnAction.cpp b/Engine/Source/Runtime/AIModule/Private/Actions/PawnAction.cpp
--- a/Engine/Source/Runtime/AIModule/Private/Actions/PawnAction.cpp
+++ b/Engine/Source/Runtime/AIModule/Private/Actions/PawnAction.cpp
@@ -374,9 +374,9 @@ bool UPawnAction::PushChildAction(UPawnAction& Action)
 // messaging
 //----------------------------------------------------------------------//
 
-void UPawnAction::WaitForMessage(FName MessageType, FAIRequestID RequestID)
+void UPawnAction::WaitForMessage(FName MessageType, FAIRequestID InRequestID)
 {
-	MessageHandlers.Add(FAIMessageObserver::Create(BrainComp, MessageType, RequestID.GetID(), FOnAIMessage::CreateUObject(this, &UPawnAction::HandleAIMessage)));
+	MessageHandlers.Add(FAIMessageObserver::Create(BrainComp, MessageType, InRequestID.GetID(), FOnAIMessage::CreateUObject(this, &UPawnAction::HandleAIMessage)));
 }
 
 //----------------------------------------------------------------------//
diff --git a/Engine/Source/Runtime/AIModule/Private/Perception/AIPerceptionComponent.cpp b/Engine/Source/Runtime/AIModule/Private/Perception/AIPerceptionComponent.cpp
--- a/Engine/Source/Runtime/AIModule/Private/Perception/AIPerceptionComponent.cpp
+++ b/Engine/Source/Runtime/AIModule/Private/Perception/AIPerceptionComponent.cpp
@@ -180,7 +180,7 @@ void UAIPerceptionComponent::OnRegister()
 
 	// this should not be needed but aparently AAIController::PostRegisterAllComponents
 	// gets called component's OnRegister
-	AAIController* AIOwner = Cast<AAIController>(GetOwner());
+	AIOwner = Cast<AAIController>(GetOwner());
 	if (AIOwner)
 	{
 		AIOwner->SetPerceptionComponent(*this);
diff --git a/Engine/Source/Runtime/Core/Private/HAL/ThreadingBase.cpp b/Engine/Source/Runtime/Core/Private/HAL/ThreadingBase.cpp
--- a/Engine/Source/Runtime/Core/Private/HAL/ThreadingBase.cpp
+++ b/Engine/Source/Runtime/Core/Private/HAL/ThreadingBase.cpp
@@ -188,9 +188,9 @@ uint32 FRunnableThread::RunnableTlsSlot = FRunnableThread::GetTlsSlot();
 uint32 FRunnableThread::GetTlsSlot()
 {
 	check( IsInGameThread() );
-	uint32 RunnableTlsSlot = FPlatformTLS::AllocTlsSlot();
-	check( FPlatformTLS::IsValidTlsSlot( RunnableTlsSlot ) );
-	return RunnableTlsSlot;
+	uint32 TlsSlot = FPlatformTLS::AllocTlsSlot();
+	check( FPlatformTLS::IsValidTlsSlot( TlsSlot ) );
+	return TlsSlot;
 }
 
 FRunnableThread::FRunnableThread()
diff --git a/Engine/Source/Runtime/Core/Private/Modules/ModuleManager.cpp b/Engine/Source/Runtime/Core/Private/Modules/ModuleManager.cpp
--- a/Engine/Source/Runtime/Core/Private/Modules/ModuleManager.cpp
+++ b/Engine/Source/Runtime/Core/Private/Modules/ModuleManager.cpp
@@ -191,8 +191,8 @@ void FModuleManager::AddModule(const FName InModuleName)
 	// Make sure module info is added to known modules and proper delegates are fired on exit.
 	struct FAtExit
 	{
-		FAtExit(const FName InModuleName, const TSharedRef<FModuleInfo>& InModuleInfo)
-			: ModuleName(InModuleName)
+		FAtExit(const FName InModule_Name, const TSharedRef<FModuleInfo>& InModuleInfo)
+			: ModuleName(InModule_Name)
 			, ModuleInfo(InModuleInfo)
 		{ }
 
diff --git a/Engine/Source/Runtime/Engine/Classes/Materials/MaterialInstance.h b/Engine/Source/Runtime/Engine/Classes/Materials/MaterialInstance.h
--- a/Engine/Source/Runtime/Engine/Classes/Materials/MaterialInstance.h
+++ b/Engine/Source/Runtime/Engine/Classes/Materials/MaterialInstance.h
@@ -476,8 +476,9 @@ bool UpdateParameterSet(TArray<ParameterType> &Parameters, UMaterial* ParentMate
 		{
 			for (const UMaterialExpression* Expression : ParentMaterial->Expressions)
 			{
-				if (const ExpressionType* ParameterExpression = Cast<const ExpressionType>(Expression))
+				if (Expression->IsA<ExpressionType>())
 				{
+					const ExpressionType* ParameterExpression = CastChecked<const ExpressionType>(Expression);
 					if (ParameterExpression->ParameterName == Parameter.ParameterName)
 					{
 						Parameter.ExpressionGUID = ParameterExpression->ExpressionGUID;
diff --git a/Engine/Source/Runtime/Engine/Private/Materials/Material.cpp b/Engine/Source/Runtime/Engine/Private/Materials/Material.cpp
--- a/Engine/Source/Runtime/Engine/Private/Materials/Material.cpp
+++ b/Engine/Source/Runtime/Engine/Private/Materials/Material.cpp
@@ -1166,8 +1166,9 @@ bool UMaterial::GetGroupName(FName ParameterName, FName& OutDesc) const
 	for (const UMaterialExpression* Expression : Expressions)
 	{
 		// Parameter is a basic Expression Parameter
-		if (const UMaterialExpressionParameter* Parameter = Cast<const UMaterialExpressionParameter>(Expression))
+		if (Expression->IsA<UMaterialExpressionParameter>())
 		{
+			const UMaterialExpressionParameter* Parameter = CastChecked<const UMaterialExpressionParameter>(Expression);
 			if (Parameter->ParameterName == ParameterName)
 			{
 				OutDesc = Parameter->Group;
@@ -1175,8 +1176,9 @@ bool UMaterial::GetGroupName(FName ParameterName, FName& OutDesc) const
 			}
 		}
 		// Parameter is a Texture Sample Parameter
-		else if (const UMaterialExpressionTextureSampleParameter* Parameter = Cast<const UMaterialExpressionTextureSampleParameter>(Expression))
+		else if (Expression->IsA<UMaterialExpressionTextureSampleParameter>())
 		{
+			const UMaterialExpressionTextureSampleParameter* Parameter = CastChecked<const UMaterialExpressionTextureSampleParameter>(Expression);
 			if (Parameter->ParameterName == ParameterName)
 			{
 				OutDesc = Parameter->Group;
@@ -1184,8 +1186,9 @@ bool UMaterial::GetGroupName(FName ParameterName, FName& OutDesc) const
 			}
 		}
 		// Parameter is a Font Sample Parameter
-		else if (const UMaterialExpressionFontSampleParameter* Parameter = Cast<const UMaterialExpressionFontSampleParameter>(Expression))
+		else if (Expression->IsA<UMaterialExpressionFontSampleParameter>())
 		{
+			const UMaterialExpressionFontSampleParameter* Parameter = CastChecked<const UMaterialExpressionFontSampleParameter>(Expression);
 			if (Parameter->ParameterName == ParameterName)
 			{
 				OutDesc = Parameter->Group;
@@ -1193,8 +1196,9 @@ bool UMaterial::GetGroupName(FName ParameterName, FName& OutDesc) const
 			}
 		}
 		// Parameter is a function call
-		else if (const UMaterialExpressionMaterialFunctionCall* FunctionCall = Cast<const UMaterialExpressionMaterialFunctionCall>(Expression))
+		else if (Expression->IsA<UMaterialExpressionMaterialFunctionCall>())
 		{
+			const UMaterialExpressionMaterialFunctionCall* FunctionCall = CastChecked<const UMaterialExpressionMaterialFunctionCall>(Expression);
 			if (FunctionCall->MaterialFunction)
 			{
 				TArray<UMaterialFunction*> Functions;
@@ -1206,8 +1210,9 @@ bool UMaterial::GetGroupName(FName ParameterName, FName& OutDesc) const
 					for (UMaterialExpression* FunctionExpression : Function->FunctionExpressions)
 					{
 						// Parameter is a basic Expression Parameter
-						if (const UMaterialExpressionParameter* Parameter = Cast<const UMaterialExpressionParameter>(FunctionExpression))
+						if (FunctionExpression->IsA<UMaterialExpressionParameter>())
 						{
+							const UMaterialExpressionParameter* Parameter = CastChecked<const UMaterialExpressionParameter>(FunctionExpression);
 							if (Parameter->ParameterName == ParameterName)
 							{
 								OutDesc = Parameter->Group;
@@ -1215,8 +1220,9 @@ bool UMaterial::GetGroupName(FName ParameterName, FName& OutDesc) const
 							}
 						}
 						// Parameter is a Texture Sample Parameter
-						else if (const UMaterialExpressionTextureSampleParameter* Parameter = Cast<const UMaterialExpressionTextureSampleParameter>(FunctionExpression))
+						else if (FunctionExpression->IsA<UMaterialExpressionTextureSampleParameter>())
 						{
+							const UMaterialExpressionTextureSampleParameter* Parameter = CastChecked<const UMaterialExpressionTextureSampleParameter>(FunctionExpression);
 							if (Parameter->ParameterName == ParameterName)
 							{
 								OutDesc = Parameter->Group;
@@ -1224,8 +1230,9 @@ bool UMaterial::GetGroupName(FName ParameterName, FName& OutDesc) const
 							}
 						}
 						// Parameter is a Font Sample Parameter
-						else if (const UMaterialExpressionFontSampleParameter* Parameter = Cast<const UMaterialExpressionFontSampleParameter>(FunctionExpression))
+						else if (FunctionExpression->IsA<UMaterialExpressionFontSampleParameter>())
 						{
+							const UMaterialExpressionFontSampleParameter* Parameter = CastChecked<const UMaterialExpressionFontSampleParameter>(FunctionExpression);
 							if (Parameter->ParameterName == ParameterName)
 							{
 								OutDesc = Parameter->Group;
@@ -1246,8 +1253,9 @@ bool UMaterial::GetParameterDesc(FName ParameterName, FString& OutDesc) const
 	for (const UMaterialExpression* Expression : Expressions)
 	{
 		// Parameter is a basic Expression Parameter
-		if (const UMaterialExpressionParameter* Parameter = Cast<const UMaterialExpressionParameter>(Expression))
+		if (Expression->IsA<UMaterialExpressionParameter>())
 		{
+			const UMaterialExpressionParameter* Parameter = CastChecked<const UMaterialExpressionParameter>(Expression);
 			if (Parameter->ParameterName == ParameterName)
 			{
 				OutDesc = Parameter->Desc;
@@ -1255,8 +1263,9 @@ bool UMaterial::GetParameterDesc(FName ParameterName, FString& OutDesc) const
 			}
 		}
 		// Parameter is a Texture Sample Parameter
-		else if (const UMaterialExpressionTextureSampleParameter* Parameter = Cast<const UMaterialExpressionTextureSampleParameter>(Expression))
+		else if (Expression->IsA<UMaterialExpressionTextureSampleParameter>())
 		{
+			const UMaterialExpressionTextureSampleParameter* Parameter = CastChecked<const UMaterialExpressionTextureSampleParameter>(Expression);
 			if (Parameter->ParameterName == ParameterName)
 			{
 				OutDesc = Parameter->Desc;
@@ -1264,8 +1273,9 @@ bool UMaterial::GetParameterDesc(FName ParameterName, FString& OutDesc) const
 			}
 		}
 		// Parameter is a Font Sample Parameter
-		else if (const UMaterialExpressionFontSampleParameter* Parameter = Cast<const UMaterialExpressionFontSampleParameter>(Expression))
+		else if (Expression->IsA<UMaterialExpressionFontSampleParameter>())
 		{
+			const UMaterialExpressionFontSampleParameter* Parameter = CastChecked<const UMaterialExpressionFontSampleParameter>(Expression);
 			if (Parameter->ParameterName == ParameterName)
 			{
 				OutDesc = Parameter->Desc;
@@ -1273,8 +1283,9 @@ bool UMaterial::GetParameterDesc(FName ParameterName, FString& OutDesc) const
 			}
 		}
 		// Parameter is a function call
-		else if (const UMaterialExpressionMaterialFunctionCall* FunctionCall = Cast<const UMaterialExpressionMaterialFunctionCall>(Expression))
+		else if (Expression->IsA<UMaterialExpressionMaterialFunctionCall>())
 		{
+			const UMaterialExpressionMaterialFunctionCall* FunctionCall = CastChecked<const UMaterialExpressionMaterialFunctionCall>(Expression);
 			if (FunctionCall->MaterialFunction)
 			{
 				TArray<UMaterialFunction*> Functions;
@@ -1286,8 +1297,9 @@ bool UMaterial::GetParameterDesc(FName ParameterName, FString& OutDesc) const
 					for (UMaterialExpression* FunctionExpression : Function->FunctionExpressions)
 					{
 						// Parameter is a basic Expression Parameter
-						if (const UMaterialExpressionParameter* Parameter = Cast<const UMaterialExpressionParameter>(FunctionExpression))
+						if (FunctionExpression->IsA<UMaterialExpressionParameter>())
 						{
+							const UMaterialExpressionParameter* Parameter = CastChecked<const UMaterialExpressionParameter>(FunctionExpression);
 							if (Parameter->ParameterName == ParameterName)
 							{
 								OutDesc = Parameter->Desc;
@@ -1295,8 +1307,9 @@ bool UMaterial::GetParameterDesc(FName ParameterName, FString& OutDesc) const
 							}
 						}
 						// Parameter is a Texture Sample Parameter
-						else if (const UMaterialExpressionTextureSampleParameter* Parameter = Cast<const UMaterialExpressionTextureSampleParameter>(FunctionExpression))
+						else if (FunctionExpression->IsA<UMaterialExpressionTextureSampleParameter>())
 						{
+							const UMaterialExpressionTextureSampleParameter* Parameter = CastChecked<const UMaterialExpressionTextureSampleParameter>(FunctionExpression);
 							if (Parameter->ParameterName == ParameterName)
 							{
 								OutDesc = Parameter->Desc;
@@ -1304,8 +1317,9 @@ bool UMaterial::GetParameterDesc(FName ParameterName, FString& OutDesc) const
 							}
 						}
 						// Parameter is a Font Sample Parameter
-						else if (const UMaterialExpressionFontSampleParameter* Parameter = Cast<const UMaterialExpressionFontSampleParameter>(FunctionExpression))
+						else if (FunctionExpression->IsA<UMaterialExpressionFontSampleParameter>())
 						{
+							const UMaterialExpressionFontSampleParameter* Parameter = CastChecked<const UMaterialExpressionFontSampleParameter>(FunctionExpression);
 							if (Parameter->ParameterName == ParameterName)
 							{
 								OutDesc = Parameter->Desc;
@@ -1325,15 +1339,17 @@ bool UMaterial::GetVectorParameterValue(FName ParameterName, FLinearColor& OutVa
 {
 	for (const UMaterialExpression* Expression : Expressions)
 	{
-		if (const UMaterialExpressionVectorParameter* Parameter = Cast<const UMaterialExpressionVectorParameter>(Expression))
+		if (Expression->IsA<UMaterialExpressionVectorParameter>())
 		{
+			const UMaterialExpressionVectorParameter* Parameter = CastChecked<const UMaterialExpressionVectorParameter>(Expression);
 			if (Parameter->IsNamedParameter(ParameterName, OutValue))
 			{
 				return true;
 			}
 		}
-		else if (const UMaterialExpressionMaterialFunctionCall* FunctionCall = Cast<const UMaterialExpressionMaterialFunctionCall>(Expression))
+		else if (Expression->IsA<UMaterialExpressionMaterialFunctionCall>())
 		{
+			const UMaterialExpressionMaterialFunctionCall* FunctionCall = CastChecked<const UMaterialExpressionMaterialFunctionCall>(Expression);
 			if (FunctionCall->MaterialFunction)
 			{
 				TArray<UMaterialFunction*> Functions;
@@ -1344,8 +1360,9 @@ bool UMaterial::GetVectorParameterValue(FName ParameterName, FLinearColor& OutVa
 				{
 					for (UMaterialExpression* FunctionExpression : Function->FunctionExpressions)
 					{
-						if (const UMaterialExpressionVectorParameter* Parameter = Cast<const UMaterialExpressionVectorParameter>(FunctionExpression))
+						if (FunctionExpression->IsA<UMaterialExpressionVectorParameter>())
 						{
+							const UMaterialExpressionVectorParameter* Parameter = CastChecked<const UMaterialExpressionVectorParameter>(FunctionExpression);
 							if (Parameter->IsNamedParameter(ParameterName, OutValue))
 							{
 								return true;
@@ -1364,15 +1381,17 @@ bool UMaterial::GetScalarParameterValue(FName ParameterName, float& OutValue) co
 {
 	for (const UMaterialExpression* Expression : Expressions)
 	{
-		if (const UMaterialExpressionScalarParameter* Parameter = Cast<const UMaterialExpressionScalarParameter>(Expression))
+		if (Expression->IsA<UMaterialExpressionScalarParameter>())
 		{
+			const UMaterialExpressionScalarParameter* Parameter = CastChecked<const UMaterialExpressionScalarParameter>(Expression);
 			if (Parameter->IsNamedParameter(ParameterName, OutValue))
 			{
 				return true;
 			}
 		}
-		else if (const UMaterialExpressionMaterialFunctionCall* FunctionCall = Cast<const UMaterialExpressionMaterialFunctionCall>(Expression))
+		else if (Expression->IsA<UMaterialExpressionMaterialFunctionCall>())
 		{
+			const UMaterialExpressionMaterialFunctionCall* FunctionCall = CastChecked<const UMaterialExpressionMaterialFunctionCall>(Expression);
 			if (FunctionCall->MaterialFunction)
 			{
 				TArray<UMaterialFunction*> Functions;
@@ -1383,8 +1402,9 @@ bool UMaterial::GetScalarParameterValue(FName ParameterName, float& OutValue) co
 				{
 					for (UMaterialExpression* FunctionExpression : Function->FunctionExpressions)
 					{
-						if (const UMaterialExpressionScalarParameter* Parameter = Cast<const UMaterialExpressionScalarParameter>(FunctionExpression))
+						if (FunctionExpression->IsA<UMaterialExpressionScalarParameter>())
 						{
+							const UMaterialExpressionScalarParameter* Parameter = CastChecked<const UMaterialExpressionScalarParameter>(FunctionExpression);
 							if (Parameter->IsNamedParameter(ParameterName, OutValue))
 							{
 								return true;
@@ -1403,15 +1423,17 @@ bool UMaterial::GetTextureParameterValue(FName ParameterName, UTexture*& OutValu
 {
 	for (const UMaterialExpression* Expression : Expressions)
 	{
-		if (const UMaterialExpressionTextureSampleParameter* Parameter = Cast<const UMaterialExpressionTextureSampleParameter>(Expression))
+		if (Expression->IsA<UMaterialExpressionTextureSampleParameter>())
 		{
+			const UMaterialExpressionTextureSampleParameter* Parameter = CastChecked<const UMaterialExpressionTextureSampleParameter>(Expression);
 			if (Parameter->IsNamedParameter(ParameterName, OutValue))
 			{
 				return true;
 			}
 		}
-		else if (const UMaterialExpressionMaterialFunctionCall* FunctionCall = Cast<const UMaterialExpressionMaterialFunctionCall>(Expression))
+		else if (Expression->IsA<UMaterialExpressionMaterialFunctionCall>())
 		{
+			const UMaterialExpressionMaterialFunctionCall* FunctionCall = CastChecked<const UMaterialExpressionMaterialFunctionCall>(Expression);
 			if (FunctionCall->MaterialFunction)
 			{
 				TArray<UMaterialFunction*> Functions;
@@ -1422,8 +1444,9 @@ bool UMaterial::GetTextureParameterValue(FName ParameterName, UTexture*& OutValu
 				{
 					for (UMaterialExpression* FunctionExpression : Function->FunctionExpressions)
 					{
-						if (const UMaterialExpressionTextureSampleParameter* Parameter = Cast<const UMaterialExpressionTextureSampleParameter>(FunctionExpression))
+						if (FunctionExpression->IsA<UMaterialExpressionTextureSampleParameter>())
 						{
+							const UMaterialExpressionTextureSampleParameter* Parameter = CastChecked<const UMaterialExpressionTextureSampleParameter>(FunctionExpression);
 							if (Parameter->IsNamedParameter(ParameterName, OutValue))
 							{
 								return true;
@@ -1442,15 +1465,17 @@ bool UMaterial::GetFontParameterValue(FName ParameterName, UFont*& OutFontValue,
 {
 	for (const UMaterialExpression* Expression : Expressions)
 	{
-		if (const UMaterialExpressionFontSampleParameter* Parameter = Cast<const UMaterialExpressionFontSampleParameter>(Expression))
+		if (Expression->IsA<UMaterialExpressionFontSampleParameter>())
 		{
+			const UMaterialExpressionFontSampleParameter* Parameter = CastChecked<const UMaterialExpressionFontSampleParameter>(Expression);
 			if (Parameter->IsNamedParameter(ParameterName, OutFontValue, OutFontPage))
 			{
 				return true;
 			}
 		}
-		else if (const UMaterialExpressionMaterialFunctionCall* FunctionCall = Cast<const UMaterialExpressionMaterialFunctionCall>(Expression))
+		else if (Expression->IsA<UMaterialExpressionMaterialFunctionCall>())
 		{
+			const UMaterialExpressionMaterialFunctionCall* FunctionCall = CastChecked<const UMaterialExpressionMaterialFunctionCall>(Expression);
 			if (FunctionCall->MaterialFunction)
 			{
 				TArray<UMaterialFunction*> Functions;
@@ -1461,8 +1486,9 @@ bool UMaterial::GetFontParameterValue(FName ParameterName, UFont*& OutFontValue,
 				{
 					for (UMaterialExpression* FunctionExpression : Function->FunctionExpressions)
 					{
-						if (const UMaterialExpressionFontSampleParameter* Parameter = Cast<const UMaterialExpressionFontSampleParameter>(FunctionExpression))
+						if (FunctionExpression->IsA<UMaterialExpressionFontSampleParameter>())
 						{
+							const UMaterialExpressionFontSampleParameter* Parameter = CastChecked<const UMaterialExpressionFontSampleParameter>(FunctionExpression);
 							if (Parameter->IsNamedParameter(ParameterName, OutFontValue, OutFontPage))
 							{
 								return true;
@@ -1482,15 +1508,17 @@ bool UMaterial::GetStaticSwitchParameterValue(FName ParameterName, bool& OutValu
 {
 	for (const UMaterialExpression* Expression : Expressions)
 	{
-		if (const UMaterialExpressionStaticBoolParameter* Parameter = Cast<const UMaterialExpressionStaticBoolParameter>(Expression))
+		if (Expression->IsA<UMaterialExpressionStaticBoolParameter>())
 		{
+			const UMaterialExpressionStaticBoolParameter* Parameter = CastChecked<const UMaterialExpressionStaticBoolParameter>(Expression);
 			if (Parameter->IsNamedParameter(ParameterName, OutValue, OutExpressionGuid))
 			{
 				return true;
 			}
 		}
-		else if (const UMaterialExpressionMaterialFunctionCall* FunctionCall = Cast<const UMaterialExpressionMaterialFunctionCall>(Expression))
+		else if (Expression->IsA<UMaterialExpressionMaterialFunctionCall>())
 		{
+			const UMaterialExpressionMaterialFunctionCall* FunctionCall = CastChecked<const UMaterialExpressionMaterialFunctionCall>(Expression);
 			if (FunctionCall->MaterialFunction)
 			{
 				TArray<UMaterialFunction*> Functions;
@@ -1501,8 +1529,9 @@ bool UMaterial::GetStaticSwitchParameterValue(FName ParameterName, bool& OutValu
 				{
 					for (UMaterialExpression* FunctionExpression : Function->FunctionExpressions)
 					{
-						if (const UMaterialExpressionStaticBoolParameter* Parameter = Cast<const UMaterialExpressionStaticBoolParameter>(FunctionExpression))
+						if (FunctionExpression->IsA<UMaterialExpressionStaticBoolParameter>())
 						{
+							const UMaterialExpressionStaticBoolParameter* Parameter = CastChecked<const UMaterialExpressionStaticBoolParameter>(FunctionExpression);
 							if (Parameter->IsNamedParameter(ParameterName, OutValue, OutExpressionGuid))
 							{
 								return true;
@@ -1522,15 +1551,17 @@ bool UMaterial::GetStaticComponentMaskParameterValue(FName ParameterName, bool&
 {
 	for (const UMaterialExpression* Expression : Expressions)
 	{
-		if (const UMaterialExpressionStaticComponentMaskParameter* Parameter = Cast<const UMaterialExpressionStaticComponentMaskParameter>(Expression))
+		if (Expression->IsA<UMaterialExpressionStaticComponentMaskParameter>())
 		{
+			const UMaterialExpressionStaticComponentMaskParameter* Parameter = CastChecked<const UMaterialExpressionStaticComponentMaskParameter>(Expression);
 			if (Parameter->IsNamedParameter(ParameterName, OutR, OutG, OutB, OutA, OutExpressionGuid))
 			{
 				return true;
 			}
 		}
-		else if (const UMaterialExpressionMaterialFunctionCall* FunctionCall = Cast<const UMaterialExpressionMaterialFunctionCall>(Expression))
+		else if (Expression->IsA<UMaterialExpressionMaterialFunctionCall>())
 		{
+			const UMaterialExpressionMaterialFunctionCall* FunctionCall = CastChecked<const UMaterialExpressionMaterialFunctionCall>(Expression);
 			if (FunctionCall->MaterialFunction)
 			{
 				TArray<UMaterialFunction*> Functions;
@@ -1541,8 +1572,9 @@ bool UMaterial::GetStaticComponentMaskParameterValue(FName ParameterName, bool&
 				{
 					for (UMaterialExpression* FunctionExpression : Function->FunctionExpressions)
 					{
-						if (const UMaterialExpressionStaticComponentMaskParameter* Parameter = Cast<const UMaterialExpressionStaticComponentMaskParameter>(FunctionExpression))
+						if (FunctionExpression->IsA<UMaterialExpressionStaticComponentMaskParameter>())
 						{
+							const UMaterialExpressionStaticComponentMaskParameter* Parameter = CastChecked<const UMaterialExpressionStaticComponentMaskParameter>(FunctionExpression);
 							if (Parameter->IsNamedParameter(ParameterName, OutR, OutG, OutB, OutA, OutExpressionGuid))
 							{
 								return true;
diff --git a/Engine/Source/Runtime/Landscape/Private/LandscapeSplines.cpp b/Engine/Source/Runtime/Landscape/Private/LandscapeSplines.cpp
--- a/Engine/Source/Runtime/Landscape/Private/LandscapeSplines.cpp
+++ b/Engine/Source/Runtime/Landscape/Private/LandscapeSplines.cpp
@@ -1488,13 +1488,13 @@ void ULandscapeSplineControlPoint::UpdateSplinePoints(bool bUpdateCollision, boo
 	}
 	for (auto& ForeignMeshComponentsPair : ForeignMeshComponentsMap)
 	{
-		ULandscapeSplinesComponent* MeshComponentOuterSplines = ForeignMeshComponentsPair.Key;
+		ULandscapeSplinesComponent* ForeignMeshComponentOuterSplines = ForeignMeshComponentsPair.Key;
 		auto* ForeignMeshComponent = ForeignMeshComponentsPair.Value;
 		if (ForeignMeshComponent != MeshComponent)
 		{
-			MeshComponentOuterSplines->Modify();
+			ForeignMeshComponentOuterSplines->Modify();
 			ForeignMeshComponent->Modify();
-			MeshComponentOuterSplines->RemoveForeignMeshComponent(this, ForeignMeshComponent);
+			ForeignMeshComponentOuterSplines->RemoveForeignMeshComponent(this, ForeignMeshComponent);
 			ForeignMeshComponent->DestroyComponent();
 		}
 	}
diff --git a/Engine/Source/Runtime/NetworkReplayStreaming/HttpNetworkReplayStreaming/Private/HttpNetworkReplayStreaming.cpp b/Engine/Source/Runtime/NetworkReplayStreaming/HttpNetworkReplayStreaming/Private/HttpNetworkReplayStreaming.cpp
--- a/Engine/Source/Runtime/NetworkReplayStreaming/HttpNetworkReplayStreaming/Private/HttpNetworkReplayStreaming.cpp
+++ b/Engine/Source/Runtime/NetworkReplayStreaming/HttpNetworkReplayStreaming/Private/HttpNetworkReplayStreaming.cpp
@@ -718,7 +718,7 @@ void FHttpNetworkReplayStreamer::DeleteFinishedStream( const FString& StreamName
 	Delegate.ExecuteIfBound(false);
 }
 
-void FHttpNetworkReplayStreamer::EnumerateStreams( const FNetworkReplayVersion& ReplayVersion, const FOnEnumerateStreamsComplete& Delegate )
+void FHttpNetworkReplayStreamer::EnumerateStreams( const FNetworkReplayVersion& InReplayVersion, const FOnEnumerateStreamsComplete& Delegate )
 {
 	if ( IsHttpRequestInFlight() )
 	{
@@ -729,7 +729,7 @@ void FHttpNetworkReplayStreamer::EnumerateStreams( const FNetworkReplayVersion&
 	TSharedRef<class IHttpRequest> HttpRequest = FHttpModule::Get().CreateRequest();
 
 	// Enumerate all of the sessions
-	HttpRequest->SetURL( FString::Printf( TEXT( "%senumsessions?App=%s&Version=%u&CL=%u" ), *ServerURL, *ReplayVersion.AppString, ReplayVersion.NetworkVersion, ReplayVersion.Changelist ) );
+	HttpRequest->SetURL( FString::Printf( TEXT( "%senumsessions?App=%s&Version=%u&CL=%u" ), *ServerURL, *InReplayVersion.AppString, InReplayVersion.NetworkVersion, InReplayVersion.Changelist ) );
 	HttpRequest->SetVerb( TEXT( "GET" ) );
 
 	HttpRequest->OnProcessRequestComplete().BindRaw( this, &FHttpNetworkReplayStreamer::HttpEnumerateSessionsFinished );
diff --git a/Engine/Source/Runtime/Online/Voice/Private/Mac/VoiceModuleMac.cpp b/Engine/Source/Runtime/Online/Voice/Private/Mac/VoiceModuleMac.cpp
--- a/Engine/Source/Runtime/Online/Voice/Private/Mac/VoiceModuleMac.cpp
+++ b/Engine/Source/Runtime/Online/Voice/Private/Mac/VoiceModuleMac.cpp
@@ -409,12 +409,12 @@ public:
 		return noErr;
 	}
 	
-	EVoiceCaptureState::Type CopyBuffer(uint8 ReadBuffer, uint8* OutVoiceBuffer, uint32 ReadOffset, uint32* ReadLen, uint32* WriteLen)
+	EVoiceCaptureState::Type CopyBuffer(uint8 InReadBuffer, uint8* OutVoiceBuffer, uint32 InReadOffset, uint32* ReadLen, uint32* WriteLen)
 	{
 		EVoiceCaptureState::Type State = VoiceCaptureState;
 		if ( *ReadLen )
 		{
-			uint8* Data = (uint8*)(BufferList[ReadBuffer].mData) + ReadOffset;
+			uint8* Data = (uint8*)(BufferList[InReadBuffer].mData) + InReadOffset;
 			
 			if ( StreamDesc.mSampleRate == OutputDesc.mSampleRate )
 			{
@@ -434,7 +434,7 @@ public:
 				AudioBufferList InputBuffer;
 				InputBuffer.mNumberBuffers = 1;
 				InputBuffer.mBuffers[0].mNumberChannels =  StreamDesc.mChannelsPerFrame;
-				InputBuffer.mBuffers[0].mDataByteSize = BufferList[ReadBuffer].mDataByteSize;
+				InputBuffer.mBuffers[0].mDataByteSize = BufferList[InReadBuffer].mDataByteSize;
 				InputBuffer.mBuffers[0].mData = Data;
 				
 				FAudioFileIO inUserData;
diff --git a/Engine/Source/Runtime/ShaderCore/Public/ShaderCache.h b/Engine/Source/Runtime/ShaderCore/Public/ShaderCache.h
--- a/Engine/Source/Runtime/ShaderCore/Public/ShaderCache.h
+++ b/Engine/Source/Runtime/ShaderCore/Public/ShaderCache.h
@@ -265,12 +265,12 @@ class SHADERCORE_API FShaderCache
 			
 			friend uint32 GetTypeHash(const FShaderRasterizerState &Key)
 			{
-				uint32 Hash = (*((uint32*)&Key.DepthBias) ^ *((uint32*)&Key.SlopeScaleDepthBias));
-				Hash ^= (Key.FillMode << 8);
-				Hash ^= Key.CullMode;
-				Hash ^= Key.bAllowMSAA ? 2 : 0;
-				Hash ^= Key.bEnableLineAA ? 1 : 0;
-				return Hash;
+				uint32 KeyHash = (*((uint32*)&Key.DepthBias) ^ *((uint32*)&Key.SlopeScaleDepthBias));
+				KeyHash ^= (Key.FillMode << 8);
+				KeyHash ^= Key.CullMode;
+				KeyHash ^= Key.bAllowMSAA ? 2 : 0;
+				KeyHash ^= Key.bEnableLineAA ? 1 : 0;
+				return KeyHash;
 			}
 		};
 		
diff --git a/Engine/Source/Runtime/Slate/Private/Framework/Text/SlateTextLayout.cpp b/Engine/Source/Runtime/Slate/Private/Framework/Text/SlateTextLayout.cpp
--- a/Engine/Source/Runtime/Slate/Private/Framework/Text/SlateTextLayout.cpp
+++ b/Engine/Source/Runtime/Slate/Private/Framework/Text/SlateTextLayout.cpp
@@ -112,7 +112,7 @@ int32 FSlateTextLayout::OnPaint( const FPaintArgs& Args, const FGeometry& Allott
 	return HighestLayerId;
 }
 
-int32 FSlateTextLayout::OnPaintHighlights( const FPaintArgs& Args, const FTextLayout::FLineView& LineView, const TArray<FLineViewHighlight>& Highlights, const FTextBlockStyle& DefaultTextStyle, const FGeometry& AllottedGeometry, const FSlateRect& ClippingRect, FSlateWindowElementList& OutDrawElements, const int32 LayerId, const FWidgetStyle& InWidgetStyle, bool bParentEnabled ) const
+int32 FSlateTextLayout::OnPaintHighlights( const FPaintArgs& Args, const FTextLayout::FLineView& LineView, const TArray<FLineViewHighlight>& Highlights, const FTextBlockStyle& InDefaultTextStyle, const FGeometry& AllottedGeometry, const FSlateRect& ClippingRect, FSlateWindowElementList& OutDrawElements, const int32 LayerId, const FWidgetStyle& InWidgetStyle, bool bParentEnabled ) const
 {
 	int32 CurrentLayerId = LayerId;
 
@@ -121,7 +121,7 @@ int32 FSlateTextLayout::OnPaintHighlights( const FPaintArgs& Args, const FTextLa
 		const TSharedPtr< ISlateLineHighlighter > LineHighlighter = StaticCastSharedPtr< ISlateLineHighlighter >( Highlight.Highlighter );
 		if (LineHighlighter.IsValid())
 		{
-			CurrentLayerId = LineHighlighter->OnPaint( Args, LineView, Highlight.OffsetX, Highlight.Width, DefaultTextStyle, AllottedGeometry, ClippingRect, OutDrawElements, CurrentLayerId, InWidgetStyle, bParentEnabled );
+			CurrentLayerId = LineHighlighter->OnPaint( Args, LineView, Highlight.OffsetX, Highlight.Width, InDefaultTextStyle, AllottedGeometry, ClippingRect, OutDrawElements, CurrentLayerId, InWidgetStyle, bParentEnabled );
 		}
 	}
 
diff --git a/Engine/Source/Runtime/Slate/Private/Framework/Text/SlateWidgetRun.cpp b/Engine/Source/Runtime/Slate/Private/Framework/Text/SlateWidgetRun.cpp
--- a/Engine/Source/Runtime/Slate/Private/Framework/Text/SlateWidgetRun.cpp
+++ b/Engine/Source/Runtime/Slate/Private/Framework/Text/SlateWidgetRun.cpp
@@ -163,8 +163,8 @@ ERunAttributes FSlateWidgetRun::GetRunAttributes() const
 	return ERunAttributes::None;
 }
 
-FSlateWidgetRun::FSlateWidgetRun(const TSharedRef<class FTextLayout>& TextLayout, const FRunInfo& InRunInfo, const TSharedRef< const FString >& InText, const FWidgetRunInfo& InWidgetInfo)
-	: TextLayout(TextLayout)
+FSlateWidgetRun::FSlateWidgetRun(const TSharedRef<class FTextLayout>& InTextLayout, const FRunInfo& InRunInfo, const TSharedRef< const FString >& InText, const FWidgetRunInfo& InWidgetInfo)
+	: TextLayout(InTextLayout)
 	, RunInfo( InRunInfo )
 	, Text( InText )
 	, Range( 0, Text->Len() )
@@ -177,8 +177,8 @@ FSlateWidgetRun::FSlateWidgetRun(const TSharedRef<class FTextLayout>& TextLayout
 	Children.Add(Info.Widget);
 }
 
-FSlateWidgetRun::FSlateWidgetRun(const TSharedRef<class FTextLayout>& TextLayout, const FRunInfo& InRunInfo, const TSharedRef< const FString >& InText, const FWidgetRunInfo& InWidgetInfo, const FTextRange& InRange)
-	: TextLayout(TextLayout)
+FSlateWidgetRun::FSlateWidgetRun(const TSharedRef<class FTextLayout>& InTextLayout, const FRunInfo& InRunInfo, const TSharedRef< const FString >& InText, const FWidgetRunInfo& InWidgetInfo, const FTextRange& InRange)
+	: TextLayout(InTextLayout)
 	, RunInfo(InRunInfo)
 	, Text( InText )
 	, Range( InRange )
diff --git a/Engine/Source/Runtime/Slate/Private/Framework/Text/TextLayout.cpp b/Engine/Source/Runtime/Slate/Private/Framework/Text/TextLayout.cpp
--- a/Engine/Source/Runtime/Slate/Private/Framework/Text/TextLayout.cpp
+++ b/Engine/Source/Runtime/Slate/Private/Framework/Text/TextLayout.cpp
@@ -824,15 +824,15 @@ FTextLocation FTextLayout::GetTextLocationAt( const FLineView& LineView, const F
 	return FTextLocation( LineView.ModelIndex, LineView.Range.EndIndex );
 }
 
-int32 FTextLayout::GetLineViewIndexForTextLocation(const TArray< FTextLayout::FLineView >& LineViews, const FTextLocation& Location, const bool bPerformInclusiveBoundsCheck)
+int32 FTextLayout::GetLineViewIndexForTextLocation(const TArray< FTextLayout::FLineView >& InLineViews, const FTextLocation& Location, const bool bPerformInclusiveBoundsCheck)
 {
 	const int32 LineModelIndex = Location.GetLineIndex();
 	const int32 Offset = Location.GetOffset();
 
 	const FLineModel& LineModel = LineModels[LineModelIndex];
-	for(int32 Index = 0; Index < LineViews.Num(); Index++)
+	for(int32 Index = 0; Index < InLineViews.Num(); Index++)
 	{
-		const FTextLayout::FLineView& LineView = LineViews[Index];
+		const FTextLayout::FLineView& LineView = InLineViews[Index];
 
 		if(LineView.ModelIndex == LineModelIndex)
 		{
@@ -843,7 +843,7 @@ int32 FTextLayout::GetLineViewIndexForTextLocation(const TArray< FTextLayout::FL
 			}
 
 			// If we're the last line, then we also need to test for the end index being part of the range
-			const bool bIsLastLineForModel = Index == (LineViews.Num() - 1) || LineViews[Index + 1].ModelIndex != LineModelIndex;
+			const bool bIsLastLineForModel = Index == (InLineViews.Num() - 1) || InLineViews[Index + 1].ModelIndex != LineModelIndex;
 			if((bIsLastLineForModel || bPerformInclusiveBoundsCheck) && LineView.Range.EndIndex == Offset)
 			{
 				return Index;
@@ -1838,14 +1838,14 @@ void FTextLayout::FRunModel::AppendTextTo(FString& Text, const FTextRange& Range
 	Run->AppendTextTo(Text, Range);
 }
 
-TSharedRef< ILayoutBlock > FTextLayout::FRunModel::CreateBlock( const FBlockDefinition& BlockDefine, float Scale ) const
+TSharedRef< ILayoutBlock > FTextLayout::FRunModel::CreateBlock( const FBlockDefinition& BlockDefine, float InScale ) const
 {
 	const FTextRange& SizeRange = BlockDefine.ActualRange;
 
 	if ( MeasuredRanges.Num() == 0 )
 	{
 		FTextRange RunRange = Run->GetTextRange();
-		return Run->CreateBlock( BlockDefine.ActualRange.BeginIndex, BlockDefine.ActualRange.EndIndex, Run->Measure( SizeRange.BeginIndex, SizeRange.EndIndex, Scale ), BlockDefine.Renderer );
+		return Run->CreateBlock( BlockDefine.ActualRange.BeginIndex, BlockDefine.ActualRange.EndIndex, Run->Measure( SizeRange.BeginIndex, SizeRange.EndIndex, InScale ), BlockDefine.Renderer );
 	}
 
 	int32 StartRangeIndex = 0;
@@ -1910,7 +1910,7 @@ TSharedRef< ILayoutBlock > FTextLayout::FRunModel::CreateBlock( const FBlockDefi
 		}
 		else
 		{
-			BlockSize += Run->Measure( SizeRange.BeginIndex, SizeRange.EndIndex, Scale );
+			BlockSize += Run->Measure( SizeRange.BeginIndex, SizeRange.EndIndex, InScale );
 		}
 	}
 	else
@@ -1921,7 +1921,7 @@ TSharedRef< ILayoutBlock > FTextLayout::FRunModel::CreateBlock( const FBlockDefi
 		}
 		else
 		{
-			BlockSize += Run->Measure( SizeRange.BeginIndex, MeasuredRanges[ StartRangeIndex ].EndIndex, Scale );
+			BlockSize += Run->Measure( SizeRange.BeginIndex, MeasuredRanges[ StartRangeIndex ].EndIndex, InScale );
 		}
 
 		for (int32 Index = StartRangeIndex + 1; Index < EndRangeIndex; Index++)
@@ -1937,7 +1937,7 @@ TSharedRef< ILayoutBlock > FTextLayout::FRunModel::CreateBlock( const FBlockDefi
 		}
 		else
 		{
-			FVector2D Size = Run->Measure( MeasuredRanges[ EndRangeIndex ].BeginIndex, SizeRange.EndIndex, Scale );
+			FVector2D Size = Run->Measure( MeasuredRanges[ EndRangeIndex ].BeginIndex, SizeRange.EndIndex, InScale );
 			BlockSize.X += Size.X;
 			BlockSize.Y = FMath::Max( Size.Y, BlockSize.Y );
 		}
diff --git a/Engine/Source/Runtime/Slate/Private/Widgets/Input/SEditableText.cpp b/Engine/Source/Runtime/Slate/Private/Widgets/Input/SEditableText.cpp
--- a/Engine/Source/Runtime/Slate/Private/Widgets/Input/SEditableText.cpp
+++ b/Engine/Source/Runtime/Slate/Private/Widgets/Input/SEditableText.cpp
@@ -345,8 +345,8 @@ uint32 SEditableText::FTextInputMethodContext::GetTextLength()
 	const TSharedPtr<SEditableText> OwningWidgetPtr = OwningWidget.Pin();
 	if(OwningWidgetPtr.IsValid())
 	{
-		const FString& EditedText = OwningWidgetPtr->EditedText.ToString();
-		TextLength = EditedText.Len();
+		const FString& WidgetEditedText = OwningWidgetPtr->EditedText.ToString();
+		TextLength = WidgetEditedText.Len();
 	}
 	return TextLength;
 }
@@ -416,8 +416,8 @@ void SEditableText::FTextInputMethodContext::GetTextInRange(const uint32 BeginIn
 	const TSharedPtr<SEditableText> OwningWidgetPtr = OwningWidget.Pin();
 	if(OwningWidgetPtr.IsValid())
 	{
-		const FString& EditedText = OwningWidgetPtr->EditedText.ToString();
-		OutString = EditedText.Mid(BeginIndex, Length);
+		const FString& WidgetEditedText = OwningWidgetPtr->EditedText.ToString();
+		OutString = WidgetEditedText.Mid(BeginIndex, Length);
 	}
 }
 
diff --git a/Engine/Source/Runtime/Slate/Private/Widgets/Layout/SResponsiveGridPanel.cpp b/Engine/Source/Runtime/Slate/Private/Widgets/Layout/SResponsiveGridPanel.cpp
--- a/Engine/Source/Runtime/Slate/Private/Widgets/Layout/SResponsiveGridPanel.cpp
+++ b/Engine/Source/Runtime/Slate/Private/Widgets/Layout/SResponsiveGridPanel.cpp
@@ -36,9 +36,9 @@ void SResponsiveGridPanel::ClearChildren()
 	Slots.Empty();
 }
 
-void SResponsiveGridPanel::Construct(const FArguments& InArgs, int32 TotalColumns)
+void SResponsiveGridPanel::Construct(const FArguments& InArgs, int32 InTotalColumns)
 {
-	this->TotalColumns = TotalColumns;
+	TotalColumns = InTotalColumns;
 	ColumnGutter = InArgs._ColumnGutter;
 	RowGutter = InArgs._RowGutter;
 
diff --git a/Engine/Source/Runtime/WebBrowser/Private/WebBrowserWindow.cpp b/Engine/Source/Runtime/WebBrowser/Private/WebBrowserWindow.cpp
--- a/Engine/Source/Runtime/WebBrowser/Private/WebBrowserWindow.cpp
+++ b/Engine/Source/Runtime/WebBrowser/Private/WebBrowserWindow.cpp
@@ -7,15 +7,15 @@
 #include "RHI.h"
 
 #if WITH_CEF3
-FWebBrowserWindow::FWebBrowserWindow(FIntPoint ViewportSize, FString Url, TOptional<FString> ContentsToLoad, bool ShowErrorMessage)
+FWebBrowserWindow::FWebBrowserWindow(FIntPoint InViewportSize, FString InUrl, TOptional<FString> InContentsToLoad, bool bInShowErrorMessage)
 	: DocumentState(EWebBrowserDocumentState::NoDocument)
 	, UpdatableTexture(nullptr)
-	, CurrentUrl(Url)
-	, ViewportSize(ViewportSize)
+	, CurrentUrl(InUrl)
+	, ViewportSize(InViewportSize)
 	, bIsClosing(false)
 	, bHasBeenPainted(false)
-	, ContentsToLoad(ContentsToLoad)
-	, ShowErrorMessage(ShowErrorMessage)
+	, ContentsToLoad(InContentsToLoad)
+	, ShowErrorMessage(bInShowErrorMessage)
 {
 	TextureData.Reserve(ViewportSize.X * ViewportSize.Y * 4);
 	TextureData.SetNumZeroed(ViewportSize.X * ViewportSize.Y * 4);
