Detect when live game switches to final while we're joining + fix pause/unpause channel bugs

diff --git a/Engine/Source/Runtime/Engine/Classes/Engine/ActorChannel.h b/Engine/Source/Runtime/Engine/Classes/Engine/ActorChannel.h
--- a/Engine/Source/Runtime/Engine/Classes/Engine/ActorChannel.h
+++ b/Engine/Source/Runtime/Engine/Classes/Engine/ActorChannel.h
@@ -48,8 +48,8 @@ class ENGINE_API UActorChannel
 	UPROPERTY()
 	AActor* Actor;					// Actor this corresponds to.
 
-	FNetworkGUID ActorNetGUID;		// Actor GUID (useful when we don't have the actor resolved yet). Currently only valid on clients.
-
+	FNetworkGUID	ActorNetGUID;		// Actor GUID (useful when we don't have the actor resolved yet). Currently only valid on clients.
+	float			CustomTimeDilation;
 
 	// Variables.
 	double	RelevantTime;			// Last time this actor was relevant to client.
diff --git a/Engine/Source/Runtime/Engine/Classes/Engine/DemoNetDriver.h b/Engine/Source/Runtime/Engine/Classes/Engine/DemoNetDriver.h
--- a/Engine/Source/Runtime/Engine/Classes/Engine/DemoNetDriver.h
+++ b/Engine/Source/Runtime/Engine/Classes/Engine/DemoNetDriver.h
@@ -130,6 +130,7 @@ public:
 	void TickDemoPlayback( float DeltaSeconds );
 	void SpawnDemoRecSpectator( UNetConnection* Connection );
 	void ResetDemoState();
+	void JumpToEndOfLiveReplay();
 
 	void StopDemo();
 
diff --git a/Engine/Source/Runtime/Engine/Private/DataChannel.cpp b/Engine/Source/Runtime/Engine/Private/DataChannel.cpp
--- a/Engine/Source/Runtime/Engine/Private/DataChannel.cpp
+++ b/Engine/Source/Runtime/Engine/Private/DataChannel.cpp
@@ -1370,6 +1370,7 @@ void UActorChannel::Init( UNetConnection* InConnection, int32 InChannelIndex, bo
 	LastUpdateTime			= Connection->Driver->Time - Connection->Driver->SpawnPrioritySeconds;
 	bActorMustStayDirty		= false;
 	bActorStillInitial		= false;
+	CustomTimeDilation		= 1.0f;
 }
 
 void UActorChannel::SetClosingFlag()
@@ -1882,6 +1883,8 @@ void UActorChannel::ProcessBunch( FInBunch & Bunch )
 		Actor->OnActorChannelOpen(Bunch, Connection);
 
 		RepFlags.bNetInitial = true;
+
+		Actor->CustomTimeDilation = CustomTimeDilation;
 	}
 	else
 	{
diff --git a/Engine/Source/Runtime/Engine/Private/DemoNetDriver.cpp b/Engine/Source/Runtime/Engine/Private/DemoNetDriver.cpp
--- a/Engine/Source/Runtime/Engine/Private/DemoNetDriver.cpp
+++ b/Engine/Source/Runtime/Engine/Private/DemoNetDriver.cpp
@@ -481,6 +481,8 @@ void UDemoNetDriver::TickDispatch(float DeltaSeconds)
 		//	(we want to continue to fly around in real-time)
 		if ( SpectatorController != NULL )
 		{
+			SpectatorController->CustomTimeDilation = 1.0f;
+
 			if ( SpectatorController->GetSpectatorPawn() != NULL )
 			{
 				// Disable collision on the spectator
@@ -938,18 +940,28 @@ void UDemoNetDriver::PauseChannels( const bool bPause )
 	for ( int32 i = ServerConnection->OpenChannels.Num() - 1; i >= 0; i-- )
 	{
 		UChannel* OpenChannel = ServerConnection->OpenChannels[i];
-		if ( OpenChannel != NULL )
+
+		UActorChannel* ActorChannel = Cast< UActorChannel >( OpenChannel );
+
+		if ( ActorChannel == NULL )
 		{
-			UActorChannel* ActorChannel = Cast< UActorChannel >( OpenChannel );
-			if ( ActorChannel != NULL && ActorChannel->GetActor() != NULL )
-			{
-				if ( Cast< APlayerController >( ActorChannel->GetActor() ) == NULL )
-				{
-					// Better way to pause each actor?
-					ActorChannel->GetActor()->CustomTimeDilation = bPause ? 0.0f : 1.0f;
-				}
-			}
+			continue;
+		}
+
+		ActorChannel->CustomTimeDilation = bPause ? 0.0f : 1.0f;
+
+		if ( ActorChannel->GetActor() == SpectatorController )
+		{
+			continue;
 		}
+
+		if ( ActorChannel->GetActor() == NULL )
+		{
+			continue;
+		}
+
+		// Better way to pause each actor?
+		ActorChannel->GetActor()->CustomTimeDilation = ActorChannel->CustomTimeDilation;
 	}
 
 	bChannelsArePaused = bPause;
@@ -985,9 +997,6 @@ bool UDemoNetDriver::ConditionallyReadDemoFrame()
 		return false;
 	} 
 
-	// We're going to read a frame, unpause channels
-	PauseChannels( false );
-
 	const int32 OldFilePos = FileAr->Tell();
 
 	uint32 SavedAbsTimeMS = 0;
@@ -1032,6 +1041,9 @@ bool UDemoNetDriver::ReadDemoFrame( FArchive* Archive )
 {
 	DECLARE_SCOPE_CYCLE_COUNTER(TEXT("ReadDemoFrame time"), STAT_ReplayReadDemoFrameTime, STATGROUP_Net);
 
+	// We're going to read a frame, unpause channels
+	PauseChannels( false );
+
 	while ( true )
 	{
 		uint8 ReadBuffer[ MAX_DEMO_READ_WRITE_BUFFER ];
@@ -1123,6 +1135,27 @@ void UDemoNetDriver::GotoTimeInSeconds( const float TimeInSeconds )
 	}
 }
 
+void UDemoNetDriver::JumpToEndOfLiveReplay()
+{
+	UE_LOG( LogDemo, Log, TEXT( "UDemoNetConnection::JumpToEndOfLiveReplay." ) );
+
+	const uint32 TotalDemoTimeInMS = ReplayStreamer->GetTotalDemoTime();
+
+	DemoTotalTime = (float)TotalDemoTimeInMS / 1000.0f;
+
+	// For now, use no buffer, try to get as close to the end as possible
+	// This will likely mean the world is paused at first, but this might be ok considering it will be as live as it can be
+	const uint32 BufferInMS = 0;//InitialLiveDemoTime != 0 ? 5 * 1000 : 0;
+
+	const uint32 JoinTimeInMS = ReplayStreamer->GetTotalDemoTime() - BufferInMS;
+
+	ReplayStreamer->GotoTimeInMS( JoinTimeInMS, FOnCheckpointReadyDelegate::CreateUObject( this, &UDemoNetDriver::CheckpointReady ) );
+
+	OldDemoCurrentTime	= DemoCurrentTime;		// So we can restore on failure
+	DemoCurrentTime		= (float)( JoinTimeInMS ) / 1000.0f;
+	InitialLiveDemoTime = 0;
+}
+
 void UDemoNetDriver::TickDemoPlayback( float DeltaSeconds )
 {
 	if ( ServerConnection == NULL || ServerConnection->State == USOCK_Closed )
@@ -1131,25 +1164,26 @@ void UDemoNetDriver::TickDemoPlayback( float DeltaSeconds )
 		return;
 	}
 
-	// See if we need to jump into the end (or close to it) of a live game we just joined
+	// See if we need to jump to the end (or close to it) of a live game we just joined
 	if ( InitialLiveDemoTime > 0 )
 	{
-		const uint32 TotalDemoTimeInMS = ReplayStreamer->GetTotalDemoTime();
-
-		// Wait until we have a freshly updated time to join, which allows things to settle down and get a more fresh time
-		if ( TotalDemoTimeInMS == InitialLiveDemoTime )
+		if ( !ReplayStreamer->IsLive() )
 		{
-			return;
+			// The replay was finalized while we joined, so don't try to jump to end
+			InitialLiveDemoTime = 0;
 		}
+		else
+		{
+			const uint32 TotalDemoTimeInMS = ReplayStreamer->GetTotalDemoTime();
 
-		InitialLiveDemoTime = 0;
-
-		DemoTotalTime = (float)TotalDemoTimeInMS / 1000.0f;
-
-		ReplayStreamer->GotoTimeInMS( ReplayStreamer->GetTotalDemoTime() - 10 * 1000, FOnCheckpointReadyDelegate::CreateUObject( this, &UDemoNetDriver::CheckpointReady ) );
+			// Wait for the most recent live time
+			if ( ReplayStreamer->GetTotalDemoTime() == InitialLiveDemoTime )
+			{
+				return;
+			}
 
-		OldDemoCurrentTime = DemoCurrentTime;		// So we can restore on failure
-		DemoCurrentTime = (float)( ReplayStreamer->GetTotalDemoTime() - 10 * 1000 ) / 1000.0f;
+			JumpToEndOfLiveReplay();
+		}
 	}
 
 	if ( GotoCheckpointArchive == NULL && !ReplayStreamer->IsLoadingCheckpoint() )
@@ -1319,18 +1353,34 @@ void UDemoNetDriver::ReplayStreamingReady( bool bSuccess, bool bRecord )
 
 	if ( !bSuccess )
 	{
+		UE_LOG( LogDemo, Warning, TEXT( "UDemoNetConnection::ReplayStreamingReady: Failed." ) );
 		GetWorld()->GetGameInstance()->HandleDemoPlaybackFailure( EDemoPlayFailure::DemoNotFound, FString( EDemoPlayFailure::ToString( EDemoPlayFailure::DemoNotFound ) ) );
 		return;
 	}
 
 	if ( !bRecord )
 	{
+		InitialLiveDemoTime = 0;
+
 		FString Error;
+		
+		const double StartTime = FPlatformTime::Seconds();
+
 		InitConnectInternal( Error );
 
 		if ( ReplayStreamer->IsLive() && ReplayStreamer->GetTotalDemoTime() > 15 * 1000 )
 		{
-			InitialLiveDemoTime = ReplayStreamer->GetTotalDemoTime();
+			// If the load time wasn't very long, jump to end now
+			// Otherwise, defer it until we have a more recent replay time
+			if ( FPlatformTime::Seconds() - StartTime < 10 )
+			{
+				JumpToEndOfLiveReplay();
+			}
+			else
+			{
+				UE_LOG( LogDemo, Log, TEXT( "UDemoNetConnection::ReplayStreamingReady: Deferring checkpoint until next available time." ) );
+				InitialLiveDemoTime = ReplayStreamer->GetTotalDemoTime();
+			}
 		}
 	}
 }
@@ -1379,6 +1429,8 @@ void UDemoNetDriver::LoadCheckpoint()
 		AddNonQueuedActorForScrubbing(SpectatorController);
 	}
 
+	PauseChannels( false );
+
 #if 1
 	// Destroy all non startup actors. They will get restored with the checkpoint
 	for ( FActorIterator It( GetWorld() ); It; ++It )
