JoinStringArray in KismetLibrary

Pull Request
 #1035

thank you derekvanvliet

diff --git a/Engine/Source/Runtime/Engine/Classes/Kismet/KismetStringLibrary.h b/Engine/Source/Runtime/Engine/Classes/Kismet/KismetStringLibrary.h
--- a/Engine/Source/Runtime/Engine/Classes/Kismet/KismetStringLibrary.h
+++ b/Engine/Source/Runtime/Engine/Classes/Kismet/KismetStringLibrary.h
@@ -263,6 +263,15 @@ class ENGINE_API UKismetStringLibrary : public UBlueprintFunctionLibrary
 	 */
 	UFUNCTION(BlueprintPure, Category="Utilities|String")
 	static TArray<FString> ParseIntoArray(const FString& SourceString, const FString& Delimiter = FString(TEXT(" ")), const bool CullEmptyStrings = true);
+	
+	/**
+	 * Concatenates an array of strings into a single string.
+	 * @param SourceArray - The array of strings to concatenate.
+	 * @param Separator - The string used to separate each element.
+	 * @return The final, joined, separated string.
+	 */
+	UFUNCTION(BlueprintPure, Category="Utilities|String")
+	static FString JoinStringArray(const TArray<FString>& SourceArray, const FString&  Separator = FString(TEXT(" ")));
 
 	/**
 	* Returns an array that contains one entry for each character in SourceString
diff --git a/Engine/Source/Runtime/Engine/Private/KismetStringLibrary.cpp b/Engine/Source/Runtime/Engine/Private/KismetStringLibrary.cpp
--- a/Engine/Source/Runtime/Engine/Private/KismetStringLibrary.cpp
+++ b/Engine/Source/Runtime/Engine/Private/KismetStringLibrary.cpp
@@ -285,6 +285,11 @@ TArray<FString> UKismetStringLibrary::ParseIntoArray(const FString& SourceString
 	return SeparatedStrings;
 }
 
+FString UKismetStringLibrary::JoinStringArray(const TArray<FString>& SourceArray, const FString&  Separator)
+{
+	return FString::Join(SourceArray, *Separator);
+}
+
 TArray<FString> UKismetStringLibrary::GetCharacterArrayFromString(const FString& SourceString)
 {
 	TArray<FString> SeparatedChars;
