#ue4 - Reorder some parts of FHitResult to be in a more logical layout.

- Reorder code filling in FHitResult in ConvertQueryImpactHit() to do it in order of struct members.
- Reorder FHitResult ImpactPoint to be in a more logical order within the struct (also affects BreakHitResult).

diff --git a/Engine/Source/Runtime/Engine/Classes/Engine/EngineTypes.h b/Engine/Source/Runtime/Engine/Classes/Engine/EngineTypes.h
--- a/Engine/Source/Runtime/Engine/Classes/Engine/EngineTypes.h
+++ b/Engine/Source/Runtime/Engine/Classes/Engine/EngineTypes.h
@@ -1456,19 +1456,19 @@ struct ENGINE_API FHitResult
 	FVector_NetQuantize Location;
 
 	/**
-	 * Normal of the hit in world space, for the object that was swept. Equal to ImpactNormal for line tests.
-	 * This is computed for capsules and spheres, otherwise it will be the same as ImpactNormal.
-	 * Example: for a sphere trace test, this is a normalized vector pointing in towards the center of the sphere at the point of impact.
+	 * Location in world space of the actual contact of the trace shape (box, sphere, ray, etc) with the impacted object.
+	 * Example: for a sphere trace test, this is the point where the surface of the sphere touches the other object.
 	 */
 	UPROPERTY()
-	FVector_NetQuantizeNormal Normal;
+	FVector_NetQuantize ImpactPoint;
 
 	/**
-	 * Location in world space of the actual contact of the trace shape (box, sphere, ray, etc) with the impacted object.
-	 * Example: for a sphere trace test, this is the point where the surface of the sphere touches the other object.
+	 * Normal of the hit in world space, for the object that was swept. Equal to ImpactNormal for line tests.
+	 * This is computed for capsules and spheres, otherwise it will be the same as ImpactNormal.
+	 * Example: for a sphere trace test, this is a normalized vector pointing in towards the center of the sphere at the point of impact.
 	 */
 	UPROPERTY()
-	FVector_NetQuantize ImpactPoint;
+	FVector_NetQuantizeNormal Normal;
 
 	/**
 	 * Normal of the hit in world space, for the object that was hit by the sweep, if any.
diff --git a/Engine/Source/Runtime/Engine/Classes/Kismet/GameplayStatics.h b/Engine/Source/Runtime/Engine/Classes/Kismet/GameplayStatics.h
--- a/Engine/Source/Runtime/Engine/Classes/Kismet/GameplayStatics.h
+++ b/Engine/Source/Runtime/Engine/Classes/Kismet/GameplayStatics.h
@@ -416,7 +416,7 @@ class ENGINE_API UGameplayStatics : public UBlueprintFunctionLibrary
 	 * @param HitItem		Primitive-specific data recording which item in the primitive was hit
 	 */
 	UFUNCTION(BlueprintPure, Category = "Collision", meta=(NativeBreakFunc))
-	static void BreakHitResult(const struct FHitResult& Hit, bool& bBlockingHit, bool& bInitialOverlap, float& Time, FVector& Location, FVector& Normal, FVector& ImpactPoint, FVector& ImpactNormal, class UPhysicalMaterial*& PhysMat, class AActor*& HitActor, class UPrimitiveComponent*& HitComponent, FName& HitBoneName, int32& HitItem, FVector& TraceStart, FVector& TraceEnd);
+	static void BreakHitResult(const struct FHitResult& Hit, bool& bBlockingHit, bool& bInitialOverlap, float& Time, FVector& Location, FVector& ImpactPoint, FVector& Normal, FVector& ImpactNormal, class UPhysicalMaterial*& PhysMat, class AActor*& HitActor, class UPrimitiveComponent*& HitComponent, FName& HitBoneName, int32& HitItem, FVector& TraceStart, FVector& TraceEnd);
 
 	/** Returns the EPhysicalSurface type of the given Hit. 
 	 * To edit surface type for your project, use ProjectSettings/Physics/PhysicalSurface section
diff --git a/Engine/Source/Runtime/Engine/Private/Collision/CollisionConversions.cpp b/Engine/Source/Runtime/Engine/Private/Collision/CollisionConversions.cpp
--- a/Engine/Source/Runtime/Engine/Private/Collision/CollisionConversions.cpp
+++ b/Engine/Source/Runtime/Engine/Private/Collision/CollisionConversions.cpp
@@ -352,20 +352,25 @@ void ConvertQueryImpactHit(const UWorld* World, const PxLocationHit& PHit, FHitR
 
 	SetHitResultFromShapeAndFaceIndex(PHit.shape,  PHit.actor, PHit.faceIndex, OutResult, bReturnPhysMat);
 
+	// See if this is a 'blocking' hit
+	const PxFilterData PShapeFilter = PHit.shape->getQueryFilterData();
+	const PxSceneQueryHitType::Enum HitType = FPxQueryFilterCallback::CalcQueryHitType(QueryFilter, PShapeFilter);
+	OutResult.bBlockingHit = (HitType == PxSceneQueryHitType::eBLOCK); 
+	OutResult.bStartPenetrating = bInitialOverlap;
+
 	// calculate the hit time
 	const float HitTime = PHit.distance/CheckLength;
+	OutResult.Time = HitTime;
 
 	// figure out where the the "safe" location for this shape is by moving from the startLoc toward the ImpactPoint
 	const FVector TraceStartToEnd = EndLoc - StartLoc;
 	const FVector SafeLocationToFitShape = StartLoc + (HitTime * TraceStartToEnd);
-
-	// Other info
-	// Caution: we may still have an initial overlap, but with null Geom. This is the case for RayCast results.
 	OutResult.Location = SafeLocationToFitShape;
 
 	const bool bUsePxPoint = ((PHit.flags & PxHitFlag::ePOSITION) && !bInitialOverlap);
 	OutResult.ImpactPoint = bUsePxPoint ? P2UVector(PHit.position) : StartLoc;
 	
+	// Caution: we may still have an initial overlap, but with null Geom. This is the case for RayCast results.
 	const bool bUsePxNormal = ((PHit.flags & PxHitFlag::eNORMAL) && !bInitialOverlap);
 	FVector Normal = bUsePxNormal ? P2UVector(PHit.normal).GetSafeNormal() : -TraceStartToEnd.GetSafeNormal();
 	OutResult.Normal = Normal;
@@ -373,13 +378,6 @@ void ConvertQueryImpactHit(const UWorld* World, const PxLocationHit& PHit, FHitR
 
 	OutResult.TraceStart = StartLoc;
 	OutResult.TraceEnd = EndLoc;
-	OutResult.Time = HitTime;
-
-	// See if this is a 'blocking' hit
-	PxFilterData PShapeFilter = PHit.shape->getQueryFilterData();
-	PxSceneQueryHitType::Enum HitType = FPxQueryFilterCallback::CalcQueryHitType(QueryFilter, PShapeFilter);
-	OutResult.bBlockingHit = (HitType == PxSceneQueryHitType::eBLOCK); 
-	OutResult.bStartPenetrating = bInitialOverlap;
 
 
 #if !(UE_BUILD_SHIPPING || UE_BUILD_TEST || !WITH_EDITOR)
diff --git a/Engine/Source/Runtime/Engine/Private/GameplayStatics.cpp b/Engine/Source/Runtime/Engine/Private/GameplayStatics.cpp
--- a/Engine/Source/Runtime/Engine/Private/GameplayStatics.cpp
+++ b/Engine/Source/Runtime/Engine/Private/GameplayStatics.cpp
@@ -593,14 +593,14 @@ UParticleSystemComponent* UGameplayStatics::SpawnEmitterAttached(UParticleSystem
 	return PSC;
 }
 
-void UGameplayStatics::BreakHitResult(const FHitResult& Hit, bool& bBlockingHit, bool& bInitialOverlap, float& Time, FVector& Location, FVector& Normal, FVector& ImpactPoint, FVector& ImpactNormal, UPhysicalMaterial*& PhysMat, AActor*& HitActor, UPrimitiveComponent*& HitComponent, FName& HitBoneName, int32& HitItem, FVector& TraceStart, FVector& TraceEnd)
+void UGameplayStatics::BreakHitResult(const FHitResult& Hit, bool& bBlockingHit, bool& bInitialOverlap, float& Time, FVector& Location, FVector& ImpactPoint, FVector& Normal, FVector& ImpactNormal, UPhysicalMaterial*& PhysMat, AActor*& HitActor, UPrimitiveComponent*& HitComponent, FName& HitBoneName, int32& HitItem, FVector& TraceStart, FVector& TraceEnd)
 {
 	bBlockingHit = Hit.bBlockingHit;
 	bInitialOverlap = Hit.bStartPenetrating;
 	Time = Hit.Time;
 	Location = Hit.Location;
-	Normal = Hit.Normal;
 	ImpactPoint = Hit.ImpactPoint;
+	Normal = Hit.Normal;
 	ImpactNormal = Hit.ImpactNormal;	
 	PhysMat = Hit.PhysMaterial.Get();
 	HitActor = Hit.GetActor();
