#ue4 - Fix recursive UCharacterMovementComponent::StartNewPhysics() allowing UpdatedComponent to change in nested scopes.

diff --git a/Engine/Source/Runtime/Engine/Private/Components/CharacterMovementComponent.cpp b/Engine/Source/Runtime/Engine/Private/Components/CharacterMovementComponent.cpp
--- a/Engine/Source/Runtime/Engine/Private/Components/CharacterMovementComponent.cpp
+++ b/Engine/Source/Runtime/Engine/Private/Components/CharacterMovementComponent.cpp
@@ -1914,6 +1914,7 @@ void UCharacterMovementComponent::StartNewPhysics(float deltaTime, int32 Iterati
 		return;
 	}
 
+	const bool bSavedMovementInProgress = bMovementInProgress;
 	bMovementInProgress = true;
 
 	switch ( MovementMode )
@@ -1944,7 +1945,7 @@ void UCharacterMovementComponent::StartNewPhysics(float deltaTime, int32 Iterati
 		break;
 	}
 
-	bMovementInProgress = false;
+	bMovementInProgress = bSavedMovementInProgress;
 	if ( bDeferUpdateMoveComponent )
 	{
 		SetUpdatedComponent(DeferredUpdatedMoveComponent);
