Allow PlayerStart to be Blueprintable

diff --git a/Engine/Source/Runtime/Engine/Classes/GameFramework/PlayerStart.h b/Engine/Source/Runtime/Engine/Classes/GameFramework/PlayerStart.h
--- a/Engine/Source/Runtime/Engine/Classes/GameFramework/PlayerStart.h
+++ b/Engine/Source/Runtime/Engine/Classes/GameFramework/PlayerStart.h
@@ -9,7 +9,7 @@
  *	
  *	@see https://docs.unrealengine.com/latest/INT/Engine/Actors/PlayerStart/
  */
-UCLASS(ClassGroup=Common, hidecategories=Collision)
+UCLASS(Blueprintable, ClassGroup=Common, hidecategories=Collision)
 class ENGINE_API APlayerStart : public ANavigationObjectBase
 {
 	/** To take more control over PlayerStart selection, you can override the virtual AGameMode::FindPlayerStart and AGameMode::ChoosePlayerStart function */
