Modified fixes for PVS-Studio warnings, mostly adding comments to warnings that were disabled. A few reworked fixes where needed.

diff --git a/Engine/Source/Editor/UnrealEd/Private/Commandlets/DiffPackagesCommandlet.cpp b/Engine/Source/Editor/UnrealEd/Private/Commandlets/DiffPackagesCommandlet.cpp
--- a/Engine/Source/Editor/UnrealEd/Private/Commandlets/DiffPackagesCommandlet.cpp
+++ b/Engine/Source/Editor/UnrealEd/Private/Commandlets/DiffPackagesCommandlet.cpp
@@ -965,6 +965,8 @@ EObjectDiff UDiffPackagesCommandlet::DiffObjects(UObject* ObjA, UObject* ObjB, U
 		return OD_None;
 	}
 
+	// PVS-Studio does not understand the assumption that we must have an ancestor if we have an ObjA or an ObjB, and is
+	// warning about usage of ObjAncestor pointer:
 	UClass* ComparisonClass = (ObjA ? ObjA->GetClass() : (ObjB ? ObjB->GetClass() : ObjAncestor->GetClass())); //-V595
 
 	// complex logic for what kind of differnce this is, if at all
@@ -1123,7 +1125,7 @@ EObjectDiff UDiffPackagesCommandlet::DiffObjects(UObject* ObjA, UObject* ObjB, U
 				{
 					check(ObjB);
 					PropDiff.DiffType = OD_ABConflict;
-					FString FullPath = ObjA ? ObjA->GetFullName(Packages[0]) : ObjB->GetFullName(Packages[1]);
+					FString FullPath = ObjA ? ObjA->GetFullName(Packages[0]) : ObjB->GetFullName(Packages[1]); //-V595 PVS-Studio does not understand the check(objB) above, and is warning about usage of the objB pointer
 
 					// recompose the text relative to each other so that when showing differences of structs,
 					// only properties within the struct that actually changed are shown
diff --git a/Engine/Source/Editor/UnrealEd/Private/Kismet2/BlueprintEditorUtils.cpp b/Engine/Source/Editor/UnrealEd/Private/Kismet2/BlueprintEditorUtils.cpp
--- a/Engine/Source/Editor/UnrealEd/Private/Kismet2/BlueprintEditorUtils.cpp
+++ b/Engine/Source/Editor/UnrealEd/Private/Kismet2/BlueprintEditorUtils.cpp
@@ -7661,7 +7661,10 @@ UK2Node_FunctionResult* FBlueprintEditorUtils::FindOrCreateFunctionResultNode(UK
 		UEdGraph* Graph = InFunctionEntryNode->GetGraph();
 
 		TArray<UK2Node_FunctionResult*> ResultNode;
-		Graph->GetNodesOfClass(ResultNode);
+		if (Graph)
+		{
+			Graph->GetNodesOfClass(ResultNode);
+		}
 
 		if (Graph && ResultNode.Num() == 0)
 		{
diff --git a/Engine/Source/Editor/UnrealEd/Private/LevelEditorViewport.cpp b/Engine/Source/Editor/UnrealEd/Private/LevelEditorViewport.cpp
--- a/Engine/Source/Editor/UnrealEd/Private/LevelEditorViewport.cpp
+++ b/Engine/Source/Editor/UnrealEd/Private/LevelEditorViewport.cpp
@@ -416,6 +416,11 @@ static UObject* GetOrCreateMaterialFromTexture( UTexture* UnrealTexture )
 	UMaterial* UnrealMaterial = (UMaterial*)MaterialFactory->FactoryCreateNew(
 		UMaterial::StaticClass(), Package, *MaterialFullName, RF_Standalone | RF_Public, NULL, GWarn );
 
+	if (UnrealMaterial == nullptr)
+	{
+		return nullptr;
+	}
+
 	const int HSpace = -300;
 
 	// If we were able to figure out the material kind, we need to try and build a complex material
@@ -469,14 +474,11 @@ static UObject* GetOrCreateMaterialFromTexture( UTexture* UnrealTexture )
 		TryAndCreateMaterialInput( UnrealMaterial, EMaterialKind::Normal, NormalTexture, UnrealMaterial->Normal, HSpace, VSpace * 2 );
 	}
 
-	if ( UnrealMaterial != NULL )
-	{
-		// Notify the asset registry
-		FAssetRegistryModule::AssetCreated( UnrealMaterial );
+	// Notify the asset registry
+	FAssetRegistryModule::AssetCreated( UnrealMaterial );
 
-		// Set the dirty flag so this package will get saved later
-		Package->SetDirtyFlag( true );
-	}
+	// Set the dirty flag so this package will get saved later
+	Package->SetDirtyFlag( true );
 
 	UnrealMaterial->ForceRecompileForRendering();
 
diff --git a/Engine/Source/Runtime/AIModule/Private/EnvironmentQuery/EQSRenderingComponent.cpp b/Engine/Source/Runtime/AIModule/Private/EnvironmentQuery/EQSRenderingComponent.cpp
--- a/Engine/Source/Runtime/AIModule/Private/EnvironmentQuery/EQSRenderingComponent.cpp
+++ b/Engine/Source/Runtime/AIModule/Private/EnvironmentQuery/EQSRenderingComponent.cpp
@@ -145,8 +145,9 @@ void FEQSSceneProxy::CollectEQSData(const FEnvQueryResult* ResultItems, const FE
 
 	// using "mid-results" requires manual normalization
 	const bool bUseMidResults = QueryInstance && (QueryInstance->Items.Num() < QueryInstance->DebugData.DebugItems.Num());
-	// no point in checking if QueryInstance != null since bUseMidResults == true guarantees that
-	const TArray<FEnvQueryItem>& Items = bUseMidResults ? QueryInstance->DebugData.DebugItems : ResultItems->Items;
+	// no point in checking if QueryInstance != null since bUseMidResults == true guarantees that, PVS-Studio doesn't
+	// understand that invariant though, and we need to disable V595:
+	const TArray<FEnvQueryItem>& Items = bUseMidResults ? QueryInstance->DebugData.DebugItems : ResultItems->Items; //-V595
 	const TArray<uint8>& RawData = bUseMidResults ? QueryInstance->DebugData.RawData : ResultItems->RawData;
 	const int32 ItemCountLimit = FMath::Clamp(Items.Num(), 0, EQSMaxItemsDrawn);
 	const bool bNoTestsPerformed = QueryInstance != NULL && QueryInstance->CurrentTest <= 0;
diff --git a/Engine/Source/Runtime/AppFramework/Private/Widgets/Testing/STestSuite.cpp b/Engine/Source/Runtime/AppFramework/Private/Widgets/Testing/STestSuite.cpp
--- a/Engine/Source/Runtime/AppFramework/Private/Widgets/Testing/STestSuite.cpp
+++ b/Engine/Source/Runtime/AppFramework/Private/Widgets/Testing/STestSuite.cpp
@@ -52,7 +52,9 @@ namespace
 		// rotation casts
 		auto quat3 = TransformCast<FQuat>(rot);
 		auto rot3 = TransformCast<FRotator>(quat);
-		rot3 = TransformCast<FRotator>(mat);
+		// PVS-Studio complains about double initialization, but that's something we're testing
+		// so we disable the warning:
+		rot3 = TransformCast<FRotator>(mat); //-V519
 
 		// higher level transform casts
 		auto mat3 = TransformCast<FMatrix>(uniScale);
diff --git a/Engine/Source/Runtime/Core/Private/Tests/Internationalization/MetadataTest.cpp b/Engine/Source/Runtime/Core/Private/Tests/Internationalization/MetadataTest.cpp
--- a/Engine/Source/Runtime/Core/Private/Tests/Internationalization/MetadataTest.cpp
+++ b/Engine/Source/Runtime/Core/Private/Tests/Internationalization/MetadataTest.cpp
@@ -233,7 +233,9 @@ bool FMetadataTest::RunTest( const FString& Parameters )
 		// Test assignment
 		{
 			FLocMetadataObject MetadataObjectAClone = *MetadataObjectB;
-			MetadataObjectAClone = *MetadataObjectA;
+			// PVS-Studio complains about double initialization, but that's something we're testing
+			// so we disable the warning
+			MetadataObjectAClone = *MetadataObjectA; //-V519
 			TestTrue( TEXT("MetadataObjectAClone == MetadataObjectA"), MetadataObjectAClone == *MetadataObjectA );
 			TestFalse( TEXT("MetadataObjectAClone == MetadataObjectB"), MetadataObjectAClone == *MetadataObjectB );
 		}
diff --git a/Engine/Source/Runtime/CoreUObject/Private/UObject/Class.cpp b/Engine/Source/Runtime/CoreUObject/Private/UObject/Class.cpp
--- a/Engine/Source/Runtime/CoreUObject/Private/UObject/Class.cpp
+++ b/Engine/Source/Runtime/CoreUObject/Private/UObject/Class.cpp
@@ -2627,7 +2627,11 @@ void UClass::DeferredRegister(UClass *UClassStaticClass,const TCHAR* PackageName
 	Super::DeferredRegister(UClassStaticClass,PackageName,Name);
 
 	// Get stashed registration info.
-	const TCHAR* InClassConfigName = *(TCHAR**)&ClassConfigName;
+
+	// PVS-Studio justifiably complains about this cast, but we expect this to work because we 'know' that 
+	// we're coming from the UClass constructor that is used when 'statically linked'. V580 disables 
+	// a warning that indicates this is an 'odd explicit type casting'.
+	const TCHAR* InClassConfigName = *(TCHAR**)&ClassConfigName; //-V580
 	ClassConfigName = InClassConfigName;
 
 	// Propagate inherited flags.
@@ -3375,7 +3379,12 @@ UClass::UClass
 ,	bCooked( false )
 {
 	// If you add properties here, please update the other constructors and PurgeClass()
-	*(const TCHAR**)&ClassConfigName = InConfigName;
+
+	// We store the pointer to the ConfigName in an FName temporarily, this cast is intentional
+	// as we expect the mis-typed data to get picked up in UClass::DeferredRegister. PVS-Studio
+	// complains about this operation, but AFAIK it is safe (and we've been doing it a long time)
+	// so the warning has been disabled for now:
+	*(const TCHAR**)&ClassConfigName = InConfigName; //-V580
 }
 
 #if WITH_HOT_RELOAD
diff --git a/Engine/Source/Runtime/Engine/Classes/Components/SkyLightComponent.h b/Engine/Source/Runtime/Engine/Classes/Components/SkyLightComponent.h
--- a/Engine/Source/Runtime/Engine/Classes/Components/SkyLightComponent.h
+++ b/Engine/Source/Runtime/Engine/Classes/Components/SkyLightComponent.h
@@ -43,7 +43,9 @@ public:
 		return Size;
 	}
 
-	virtual uint32 GetSizeY() const override
+	// PVS-Studio notices that the implementation of GetSizeX is identical to this one
+	// and warns us. In this case, it is intentional, so we disable the warning:
+	virtual uint32 GetSizeY() const override //-V524
 	{
 		return Size;
 	}
diff --git a/Engine/Source/Runtime/Engine/Classes/Engine/TextureRenderTargetCube.h b/Engine/Source/Runtime/Engine/Classes/Engine/TextureRenderTargetCube.h
--- a/Engine/Source/Runtime/Engine/Classes/Engine/TextureRenderTargetCube.h
+++ b/Engine/Source/Runtime/Engine/Classes/Engine/TextureRenderTargetCube.h
@@ -60,7 +60,9 @@ class UTextureRenderTargetCube : public UTextureRenderTarget
 
 	// Begin UTexture interface.
 	virtual float GetSurfaceWidth() const  override { return SizeX; }
-	virtual float GetSurfaceHeight()const  override { return SizeX; }	
+	// PVS-Studio notices that the implementation of GetSurfaceWidth is identical to this one
+	// and warns us. In this case, it is intentional, so we disable the warning:
+	virtual float GetSurfaceHeight()const  override { return SizeX; }	 //-V524
 	virtual FTextureResource* CreateResource() override;
 	virtual EMaterialValueType GetMaterialType() override;
 	// End UTexture interface.
diff --git a/Engine/Source/Runtime/Engine/Private/Components/CharacterMovementComponent.cpp b/Engine/Source/Runtime/Engine/Private/Components/CharacterMovementComponent.cpp
--- a/Engine/Source/Runtime/Engine/Private/Components/CharacterMovementComponent.cpp
+++ b/Engine/Source/Runtime/Engine/Private/Components/CharacterMovementComponent.cpp
@@ -2727,14 +2727,14 @@ void UCharacterMovementComponent::PhysFlying(float deltaTime, int32 Iterations)
 	FHitResult Hit(1.f);
 	SafeMoveUpdatedComponent(Adjusted, CharacterOwner->GetActorRotation(), true, Hit);
 
-	if( Hit.Time < 1.f && CharacterOwner )
+	if (Hit.Time < 1.f)
 	{
-		const FVector GravDir = FVector(0.f,0.f,-1.f);
+		const FVector GravDir = FVector(0.f, 0.f, -1.f);
 		const FVector VelDir = Velocity.GetSafeNormal();
 		const float UpDown = GravDir | VelDir;
 
 		bool bSteppedUp = false;
-		if( (FMath::Abs(Hit.ImpactNormal.Z) < 0.2f) && (UpDown < 0.5f) && (UpDown > -0.2f) && CanStepUp(Hit))
+		if ((FMath::Abs(Hit.ImpactNormal.Z) < 0.2f) && (UpDown < 0.5f) && (UpDown > -0.2f) && CanStepUp(Hit))
 		{
 			float stepZ = CharacterOwner->GetActorLocation().Z;
 			bSteppedUp = StepUp(GravDir, Adjusted * (1.f - Hit.Time), Hit);
@@ -2743,7 +2743,7 @@ void UCharacterMovementComponent::PhysFlying(float deltaTime, int32 Iterations)
 				OldLocation.Z = CharacterOwner->GetActorLocation().Z + (OldLocation.Z - stepZ);
 			}
 		}
-		
+
 		if (!bSteppedUp)
 		{
 			//adjust and try again
@@ -2751,8 +2751,8 @@ void UCharacterMovementComponent::PhysFlying(float deltaTime, int32 Iterations)
 			SlideAlongSurface(Adjusted, (1.f - Hit.Time), Hit.Normal, Hit, true);
 		}
 	}
-	
-	if( !bJustTeleported && !HasRootMotion() && CharacterOwner )
+
+	if (!bJustTeleported && !HasRootMotion())
 	{
 		Velocity = (CharacterOwner->GetActorLocation() - OldLocation) / deltaTime;
 	}
@@ -2878,11 +2878,11 @@ void UCharacterMovementComponent::StartSwimming(FVector OldLocation, FVector Old
 		}
 		MoveUpdatedComponent(End - CharacterOwner->GetActorLocation(), CharacterOwner->GetActorRotation(), true);
 	}
-	if ( !HasRootMotion() && CharacterOwner && (Velocity.Z > 2.f*SWIMBOBSPEED) && (Velocity.Z < 0.f)) //allow for falling out of water
+	if ( !HasRootMotion() && (Velocity.Z > 2.f*SWIMBOBSPEED) && (Velocity.Z < 0.f)) //allow for falling out of water
 	{
 		Velocity.Z = SWIMBOBSPEED - Velocity.Size2D() * 0.7f; //smooth bobbing
 	}
-	if ( (remainingTime >= MIN_TICK_TIME) && (Iterations < MaxSimulationIterations) && CharacterOwner )
+	if ( (remainingTime >= MIN_TICK_TIME) && (Iterations < MaxSimulationIterations) )
 	{
 		PhysSwimming(remainingTime, Iterations);
 	}
diff --git a/Engine/Source/Runtime/Engine/Private/Components/InterpToMovementComponent.cpp b/Engine/Source/Runtime/Engine/Private/Components/InterpToMovementComponent.cpp
--- a/Engine/Source/Runtime/Engine/Private/Components/InterpToMovementComponent.cpp
+++ b/Engine/Source/Runtime/Engine/Private/Components/InterpToMovementComponent.cpp
@@ -38,7 +38,7 @@ void UInterpToMovementComponent::TickComponent(float DeltaTime, enum ELevelTick
 	Super::TickComponent(DeltaTime, TickType, ThisTickFunction);
 
 	// skip if don't want component updated when not rendered or updated component can't move
-	if (ShouldSkipUpdate(DeltaTime))
+	if (!UpdatedComponent || ShouldSkipUpdate(DeltaTime))
 	{
 		return;
 	}
diff --git a/Engine/Source/Runtime/Engine/Private/Distributions.cpp b/Engine/Source/Runtime/Engine/Private/Distributions.cpp
--- a/Engine/Source/Runtime/Engine/Private/Distributions.cpp
+++ b/Engine/Source/Runtime/Engine/Private/Distributions.cpp
@@ -3621,7 +3621,9 @@ FVector UDistributionVectorUniformCurve::GetMinValue() const
 	return FVector::ZeroVector;
 }
 
-FVector UDistributionVectorUniformCurve::GetMaxValue() const
+// PVS-Studio notices that the implementation of GetMinValue is identical to this one
+// and warns us. In this case, it is intentional, so we disable the warning:
+FVector UDistributionVectorUniformCurve::GetMaxValue() const //-V524
 {
 	check(!TEXT("Don't call me!"));
 	return FVector::ZeroVector;
diff --git a/Engine/Source/Runtime/Engine/Private/Level.cpp b/Engine/Source/Runtime/Engine/Private/Level.cpp
--- a/Engine/Source/Runtime/Engine/Private/Level.cpp
+++ b/Engine/Source/Runtime/Engine/Private/Level.cpp
@@ -59,7 +59,7 @@ void FPrecomputedVisibilityHandler::UpdateVisibilityStats(bool bAllocating) cons
 		}
 	}
 	else
-	{
+	{ //-V523, disabling identical branch warning because PVS-Studio does not understate the stat system in all configurations:
 		DEC_DWORD_STAT_BY(STAT_PrecomputedVisibilityMemory, PrecomputedVisibilityCellBuckets.GetAllocatedSize());
 		for (int32 BucketIndex = 0; BucketIndex < PrecomputedVisibilityCellBuckets.Num(); BucketIndex++)
 		{
diff --git a/Engine/Source/Runtime/Engine/Private/Particles/ParticleComponents.cpp b/Engine/Source/Runtime/Engine/Private/Particles/ParticleComponents.cpp
--- a/Engine/Source/Runtime/Engine/Private/Particles/ParticleComponents.cpp
+++ b/Engine/Source/Runtime/Engine/Private/Particles/ParticleComponents.cpp
@@ -3240,7 +3240,8 @@ FDynamicEmitterDataBase* UParticleSystemComponent::CreateDynamicDataFromReplay(
 		case DET_Mesh:
 			{
 				// Allocate the dynamic data
-				FDynamicMeshEmitterData* NewEmitterData = ::new FDynamicMeshEmitterData(EmitterInstance->CurrentLODLevel->RequiredModule);
+				// PVS-Studio does not understand the checkSlow above, so it is warning us that EmitterInstance->CurrentLODLevel may be null.
+				FDynamicMeshEmitterData* NewEmitterData = ::new FDynamicMeshEmitterData(EmitterInstance->CurrentLODLevel->RequiredModule); //-V595
 
 				// Fill in the source data
 				const FDynamicMeshEmitterReplayData* MeshEmitterReplayData =
diff --git a/Engine/Source/Runtime/Engine/Private/SkeletalMeshComponentPhysics.cpp b/Engine/Source/Runtime/Engine/Private/SkeletalMeshComponentPhysics.cpp
--- a/Engine/Source/Runtime/Engine/Private/SkeletalMeshComponentPhysics.cpp
+++ b/Engine/Source/Runtime/Engine/Private/SkeletalMeshComponentPhysics.cpp
@@ -826,33 +826,36 @@ void USkeletalMeshComponent::InitArticulated(FPhysScene* PhysScene)
 
 #if WITH_PHYSX
 
+	if (PhysScene)
+	{
 		// Get the scene type from the SkeletalMeshComponent's BodyInstance
 		const uint32 SceneType = (bHasBodiesInAsyncScene && PhysScene->HasAsyncScene()) ? PST_Async : PST_Sync;
-			PxScene* PScene = PhysScene->GetPhysXScene(SceneType);
-			SCOPED_SCENE_WRITE_LOCK(PScene);
-
-	// add Aggregate into the scene
-	if(Aggregate && Aggregate->getNbActors() > 0 && PhysScene)
-	{
+		PxScene* PScene = PhysScene->GetPhysXScene(SceneType);
+		SCOPED_SCENE_WRITE_LOCK(PScene);
+		// add Aggregate into the scene
+		if (Aggregate && Aggregate->getNbActors() > 0)
+		{
 			PScene->addAggregate(*Aggregate);
-		
-		// If we've used an aggregate, InitBody would not be able to set awake status as we *must* have a scene
-		// to do that, so we reconcile this here.
-		AActor* Owner = GetOwner();
-		bool bShouldSleep = !BodyInstance.bStartAwake && (Owner && Owner->GetVelocity().SizeSquared() <= KINDA_SMALL_NUMBER);
 
-		for(FBodyInstance* Body : Bodies)
-		{
-			// Creates a DOF constraint if necessary for the body - also requires the scene to exist within the actor
-			Body->CreateDOFLock();
+			// If we've used an aggregate, InitBody would not be able to set awake status as we *must* have a scene
+			// to do that, so we reconcile this here.
+			AActor* Owner = GetOwner();
+			bool bShouldSleep = !BodyInstance.bStartAwake && (Owner && Owner->GetVelocity().SizeSquared() <= KINDA_SMALL_NUMBER);
 
-			// Set to sleep if necessary
-			if(bShouldSleep)
+			for (FBodyInstance* Body : Bodies)
 			{
-				Body->GetPxRigidDynamic_AssumesLocked()->putToSleep();
+				// Creates a DOF constraint if necessary for the body - also requires the scene to exist within the actor
+				Body->CreateDOFLock();
+
+				// Set to sleep if necessary
+				if (bShouldSleep)
+				{
+					Body->GetPxRigidDynamic_AssumesLocked()->putToSleep();
+				}
 			}
 		}
 	}
+
 #endif //WITH_PHYSX
 
 	// Create all the constraints.
@@ -2388,6 +2391,9 @@ void USkeletalMeshComponent::DrawDebugClothCollisions()
 			case FClothCollisionPrimitive::CONVEX:	
 				DrawDebugConvexFromPlanes(CollisionPrims[PrimIndex], Colors[PrimIndex%6]);
 				break;
+
+			case FClothCollisionPrimitive::PLANE:
+				break;
 			}
 
 			//draw a bounding box for double checking
diff --git a/Engine/Source/Runtime/Engine/Private/Texture.cpp b/Engine/Source/Runtime/Engine/Private/Texture.cpp
--- a/Engine/Source/Runtime/Engine/Private/Texture.cpp
+++ b/Engine/Source/Runtime/Engine/Private/Texture.cpp
@@ -599,7 +599,9 @@ uint8* FTextureSource::LockMip(int32 MipIndex)
 						if (RawData->Num() > 0)
 						{
 							LockedMipData = (uint8*)FMemory::Malloc(RawData->Num());
-							FMemory::Memcpy(LockedMipData, RawData->GetData(), RawData->Num());
+							// PVS-Studio does not understand that IImageWrapper::GetRaw's return value validates the pointer, so we disable
+							// the warning that we are using the RawData pointer before checking for null:
+							FMemory::Memcpy(LockedMipData, RawData->GetData(), RawData->Num()); //-V595
 						}
 					}
 					if (RawData == NULL || RawData->Num() == 0)
diff --git a/Engine/Source/Runtime/Engine/Private/Texture2D.cpp b/Engine/Source/Runtime/Engine/Private/Texture2D.cpp
--- a/Engine/Source/Runtime/Engine/Private/Texture2D.cpp
+++ b/Engine/Source/Runtime/Engine/Private/Texture2D.cpp
@@ -1209,7 +1209,9 @@ void FTexture2DResource::InitRHI()
 	if( Owner->PendingMipChangeRequestStatus.GetValue() == TexState_InProgress_Initialization )
 	{
 		bool bSkipRHITextureCreation = false; //Owner->bIsCompositingSource;
-		if (GIsEditor || (!bSkipRHITextureCreation))
+		// PVS - Studio has pointed out that bSkipRHITextureCreation is a codesmell (!bSkipRHITextureCreation is always true)
+		// for now, we have disabled the warning, but consider removing the variable if requirements indicate it will never be used.
+		if (GIsEditor || (!bSkipRHITextureCreation)) //-V560 
 		{
 			if ( CanCreateAsVirtualTexture(Owner, TexCreateFlags) )
 			{
@@ -1297,7 +1299,9 @@ void FTexture2DResource::InitRHI()
 	{
 		// Recreate the texture from the texture contents that were saved by ReleaseRHI.
 		bool bSkipRHITextureCreation = false; //Owner->bIsCompositingSource;
-		if (GIsEditor || (!bSkipRHITextureCreation))
+		// PVS - Studio has pointed out that bSkipRHITextureCreation is a codesmell (!bSkipRHITextureCreation is always true)
+		// for now, we have disabled the warning, but consider removing the variable if requirements indicate it will never be used.
+		if (GIsEditor || (!bSkipRHITextureCreation)) //-V560
 		{
 			FRHIResourceCreateInfo CreateInfo;
 			Texture2DRHI	= RHICreateTexture2D( SizeX, SizeY, EffectiveFormat, Owner->RequestedMips, 1, TexCreateFlags, CreateInfo );
diff --git a/Engine/Source/Runtime/RenderCore/Private/RenderUtils.cpp b/Engine/Source/Runtime/RenderCore/Private/RenderUtils.cpp
--- a/Engine/Source/Runtime/RenderCore/Private/RenderUtils.cpp
+++ b/Engine/Source/Runtime/RenderCore/Private/RenderUtils.cpp
@@ -376,7 +376,9 @@ public:
 	}
 
 	/** Returns the height of the texture in pixels. */
-	virtual uint32 GetSizeY() const override
+	// PVS-Studio notices that the implementation of GetSizeX is identical to this one
+	// and warns us. In this case, it is intentional, so we disable the warning:
+	virtual uint32 GetSizeY() const override //-V524
 	{
 		int32 TextureSize = 1 << (NumMips - 1);
 		return TextureSize;
diff --git a/Engine/Source/Runtime/Renderer/Private/SceneVisibility.cpp b/Engine/Source/Runtime/Renderer/Private/SceneVisibility.cpp
--- a/Engine/Source/Runtime/Renderer/Private/SceneVisibility.cpp
+++ b/Engine/Source/Runtime/Renderer/Private/SceneVisibility.cpp
@@ -1932,7 +1932,9 @@ void FSceneRenderer::ComputeViewVisibility(FRHICommandListImmediate& RHICmdList)
 		const bool bIsParent = ViewState && ViewState->IsViewParent();
 		if ( bIsParent )
 		{
-			ViewState->ParentPrimitives.Empty();
+			// PVS-Studio does not understand the validation of ViewState above, so we're disabling
+			// its warning that ViewState may be null:
+			ViewState->ParentPrimitives.Empty(); //-V595
 		}
 
 		if (ViewState)
diff --git a/Engine/Source/Runtime/UMG/Private/Components/EditableTextBox.cpp b/Engine/Source/Runtime/UMG/Private/Components/EditableTextBox.cpp
--- a/Engine/Source/Runtime/UMG/Private/Components/EditableTextBox.cpp
+++ b/Engine/Source/Runtime/UMG/Private/Components/EditableTextBox.cpp
@@ -20,7 +20,7 @@ UEditableTextBox::UEditableTextBox(const FObjectInitializer& ObjectInitializer)
 	// Grab other defaults from slate arguments.
 	SEditableTextBox::FArguments Defaults;
 	IsReadOnly = Defaults._IsReadOnly.Get();
-	IsPassword = Defaults._IsReadOnly.Get();
+	IsPassword = Defaults._IsPassword.Get();
 	MinimumDesiredWidth = Defaults._MinDesiredWidth.Get();
 	Padding = Defaults._Padding.Get();
 	IsCaretMovedWhenGainFocus = Defaults._IsCaretMovedWhenGainFocus.Get();
diff --git a/Engine/Source/Runtime/Windows/UnrealAudioWasapi/Private/UnrealAudioDeviceWasapi.cpp b/Engine/Source/Runtime/Windows/UnrealAudioWasapi/Private/UnrealAudioDeviceWasapi.cpp
--- a/Engine/Source/Runtime/Windows/UnrealAudioWasapi/Private/UnrealAudioDeviceWasapi.cpp
+++ b/Engine/Source/Runtime/Windows/UnrealAudioWasapi/Private/UnrealAudioDeviceWasapi.cpp
@@ -125,7 +125,7 @@ namespace UAudio
 		Speakers.Reset();
 		uint32 ChanCount = 0;
 		// Build a flag field of the speaker outputs of this device
-		for (uint32 SpeakerTypeIndex = 0; SpeakerTypeIndex < ESpeaker::SPEAKER_TYPE_COUNT, ChanCount < NumChannels; ++SpeakerTypeIndex)
+		for (uint32 SpeakerTypeIndex = 0; SpeakerTypeIndex < ESpeaker::SPEAKER_TYPE_COUNT && ChanCount < NumChannels; ++SpeakerTypeIndex)
 		{
 			if (ChannelMask & SpeakerMaskMapping[SpeakerTypeIndex].XAudio2SpeakerTypeFlag)
 			{
@@ -565,7 +565,7 @@ Cleanup:
 
 		if (!GetDeviceInfo(DataFlow, DeviceIndex, DeviceInfo))
 		{
-			Result = -1;
+			Result = S_FALSE;
 			goto Cleanup;
 		}
 
