Light cleanup pass on PlayerCameraManager, removed UPROPERTY from several members that didn't need it.

diff --git a/Engine/Source/Runtime/Engine/Classes/Camera/PlayerCameraManager.h b/Engine/Source/Runtime/Engine/Classes/Camera/PlayerCameraManager.h
--- a/Engine/Source/Runtime/Engine/Classes/Camera/PlayerCameraManager.h
+++ b/Engine/Source/Runtime/Engine/Classes/Camera/PlayerCameraManager.h
@@ -166,7 +166,7 @@ class ENGINE_API APlayerCameraManager : public AActor
 	GENERATED_UCLASS_BODY()
 
 	/** PlayerController that owns this Camera actor */
-	UPROPERTY()
+	UPROPERTY(transient)
 	class APlayerController* PCOwner;
 
 private_subobject:
@@ -177,83 +177,71 @@ private_subobject:
 
 public:
 	/** Usable to define different camera behaviors. A few simple styles are implemented by default. */
-	UPROPERTY()
 	FName CameraStyle;
 
 	/** FOV to use by default. */
 	UPROPERTY(EditAnywhere, BlueprintReadWrite, Category=PlayerCameraManager)
 	float DefaultFOV;
 
-	/** Value to lock FOV at, if bLockedFOV == true. Ignored otherwise. */
-	UPROPERTY()
+protected:
+	/** Value to lock FOV to, in degrees. Ignored if <= 0, utilized if > 0. */
 	float LockedFOV;
 
+public:
 	/** The default desired width (in world units) of the orthographic view (ignored in Perspective mode) */
 	UPROPERTY(EditAnywhere, BlueprintReadWrite, Category=PlayerCameraManager)
 	float DefaultOrthoWidth;
 
-	/** Value OrthoWidth is locked at, if bLockedOrthoWidth == true. Ignored otherwise. */
-	UPROPERTY()
+protected:
+	/** Value OrthoWidth is locked at, if > 0. Ignored if <= 0. */
 	float LockedOrthoWidth;
 
-	/** True when this camera should use an orthographic perspective instead of FOV */
-	UPROPERTY(EditAnywhere, BlueprintReadWrite, Category=PlayerCameraManager)
-	bool bIsOrthographic;
-
+public:
 	/** Default aspect ratio */
 	UPROPERTY(EditAnywhere, BlueprintReadWrite, Category=PlayerCameraManager)
 	float DefaultAspectRatio;
 
 	/** Color to fade to (when bEnableFading == true). */
-	UPROPERTY()
 	FLinearColor FadeColor;
 
 	/** Amount of fading to apply (when bEnableFading == true). */
-	UPROPERTY()
 	float FadeAmount;
 
 	/** Allows control over scaling individual color channels in the final image (when bEnableColorScaling == true). */
-	UPROPERTY()
 	FVector ColorScale;
 
 	/** Desired color scale which ColorScale will interpolate to (when bEnableColorScaling and bEnableColorScaleInterp == true) */
-	UPROPERTY()
 	FVector DesiredColorScale;
 
 	/** Color scale value at start of interpolation (when bEnableColorScaling and bEnableColorScaleInterp == true) */
-	UPROPERTY()
 	FVector OriginalColorScale;
 
 	/** Total time for color scale interpolation to complete (when bEnableColorScaling and bEnableColorScaleInterp == true) */
-	UPROPERTY()
 	float ColorScaleInterpDuration;
 
 	/** Time at which interpolation started (when bEnableColorScaling and bEnableColorScaleInterp == true) */
-	UPROPERTY()
 	float ColorScaleInterpStartTime;
 
 	/** Cached camera properties. */
-	UPROPERTY()
+	UPROPERTY(transient)
 	struct FCameraCacheEntry CameraCache;
 
-	/** Cached camera properties. */
-	UPROPERTY()
+	/** Cached camera properties, one frame old. */
+	UPROPERTY(transient)
 	struct FCameraCacheEntry LastFrameCameraCache;
 
 	/** Current ViewTarget */
-	UPROPERTY()
+	UPROPERTY(transient)
 	struct FTViewTarget ViewTarget;
 
 	/** Pending view target for blending */
-	UPROPERTY()
+	UPROPERTY(transient)
 	struct FTViewTarget PendingViewTarget;
 
 	/** Time remaining in viewtarget blend. */
-	UPROPERTY()
 	float BlendTimeToGo;
 
 	/** Current view target transition blend parameters. */
-	UPROPERTY()
 	struct FViewTargetTransitionParams BlendParams;
 
 protected:
@@ -275,15 +263,12 @@ public:
 	FVector FreeCamOffset;
 
 	/** Current camera fade alpha range, where X = starting alpha and Y = final alpha (when bEnableFading == true) */
-	UPROPERTY()
 	FVector2D FadeAlpha;
 
 	/** Total duration of the camera fade (when bEnableFading == true) */
-	UPROPERTY()
 	float FadeTime;
 
 	/** Time remaining in camera fade (when bEnableFading == true) */
-	UPROPERTY()
 	float FadeTimeRemaining;
 
 protected:
@@ -300,10 +285,7 @@ protected:
 	UPROPERTY(transient)
 	class UCameraModifier_CameraShake* CachedCameraShakeMod;
 
-	/** By default camera modifiers are not applied to stock debug cameras (e.g. CameraStyle == "FreeCam"). Setting to true will always apply modifiers. */
-	UPROPERTY()
-	bool bAlwaysApplyModifiers;
-	
+
 	////////////////////////
 	// CameraAnim support
 	////////////////////////
@@ -311,11 +293,13 @@ protected:
 	static const int32 MAX_ACTIVE_CAMERA_ANIMS = 8;
 
 	/** Internal pool of camera anim instance objects available for playing camera animations. Defines the max number of camera anims that can play simultaneously. */
-	UPROPERTY()
+	UPROPERTY(transient)
 	class UCameraAnimInst* AnimInstPool[8];    /*MAX_ACTIVE_CAMERA_ANIMS @fixme constant */
 
 	/** Internal list of active post process effects. Parallel array to PostProcessBlendCacheWeights. */
+	UPROPERTY(transient)
 	TArray<struct FPostProcessSettings> PostProcessBlendCache;
+
 	/** Internal list of weights for active post process effects. Parallel array to PostProcessBlendCache. */
 	TArray<float> PostProcessBlendCacheWeights;
 
@@ -328,7 +312,7 @@ protected:
 
 public:
 	/** Array of camera anim instances that are currently playing and in-use */
-	UPROPERTY()
+	UPROPERTY(transient)
 	TArray<class UCameraAnimInst*> ActiveAnims;
 
 	/** Returns active post process info. */
@@ -336,73 +320,67 @@ public:
 	
 protected:
 	/** Array of camera anim instances that are not playing and available to be used. */
-	UPROPERTY()
+	UPROPERTY(transient)
 	TArray<class UCameraAnimInst*> FreeAnims;
 
 	/** Internal. Receives the output of individual camera animations. */
 	UPROPERTY(transient)
 	class ACameraActor* AnimCameraActor;
 
-	/** Tuue if FOV is locked to a constant value (@see LockedFOV) */
-	UPROPERTY()
-	uint32 bLockedFOV:1;
 
-	/** true if OrthoWidth is locked to a constant value (@see LockedOrthoWidth) */
-	UPROPERTY()
-	uint32 bLockedOrthoWidth:1;
+	////////////////////////
+	// Flags
+	////////////////////////
 
 public:
+	/** True when this camera should use an orthographic perspective instead of FOV */
+	UPROPERTY(EditAnywhere, BlueprintReadWrite, Category = PlayerCameraManager)
+	uint32 bIsOrthographic : 1;
+
 	/** True if we should apply FadeColor/FadeAmount to the screen. */
-	UPROPERTY()
-	uint32 bEnableFading:1;
+	uint32 bEnableFading : 1;
 
 	/** True to apply fading of audio alongside the video. */
-	UPROPERTY()
-	uint32 bFadeAudio:1;
-
-protected:
-	/** True if camera fade hold at it's final value when finished. */
-	UPROPERTY()
-	uint32 bHoldFadeWhenFinished: 1;
+	uint32 bFadeAudio : 1;
 
-	UPROPERTY()
-	uint32 bAutoAnimateFade : 1;
-
-public:
 	/** True to turn on scaling of color channels in final image using ColorScale property. */
-	UPROPERTY()
-	uint32 bEnableColorScaling:1;
+	uint32 bEnableColorScaling : 1;
 
 	/** True to smoothly interpolate color scale values when they change. */
-	UPROPERTY()
-	uint32 bEnableColorScaleInterp:1;
+	uint32 bEnableColorScaleInterp : 1;
 
 	/** True if clients are handling setting their own viewtarget and the server should not replicate it (e.g. during certain Matinee sequences) */
-	UPROPERTY()
-	uint32 bClientSimulatingViewTarget:1;
+	uint32 bClientSimulatingViewTarget : 1;
 
 	/** True if server will use camera positions replicated from the client instead of calculating them locally. */
 	UPROPERTY(EditAnywhere, BlueprintReadOnly, Category=PlayerCameraManager)
-	uint32 bUseClientSideCameraUpdates:1;
+	uint32 bUseClientSideCameraUpdates : 1;
 
 	/** For debugging. If true, replicate the client side camera position but don't use it, and draw the positions on the server */
-	UPROPERTY()
-	uint32 bDebugClientSideCamera:1;
+	uint32 bDebugClientSideCamera : 1;
 
 	/** If true, send a camera update to the server on next update. */
-	UPROPERTY(transient)
-	uint32 bShouldSendClientSideCameraUpdate:1;
+	uint32 bShouldSendClientSideCameraUpdate : 1;
 
 	/** 
 	 * True if we did a camera cut this frame. Automatically reset to false every frame.
 	 * This flag affects various things in the renderer (such as whether to use the occlusion queries from last frame, and motion blur).
 	 */
-	UPROPERTY(transient)
-	uint32 bGameCameraCutThisFrame:1;
+	uint32 bGameCameraCutThisFrame : 1;
 
+protected:
+	/** True if camera fade hold at it's final value when finished. */
+	uint32 bHoldFadeWhenFinished : 1;
+
+	uint32 bAutoAnimateFade : 1;
+
+	/** By default camera modifiers are not applied to stock debug cameras (e.g. CameraStyle == "FreeCam"). Setting to true will always apply modifiers. */
+	uint32 bAlwaysApplyModifiers : 1;
+
+public:
 	/** True if camera's orientation should be updated by most recent HMD orientation. */
 	UPROPERTY(EditAnywhere, BlueprintReadWrite, Category = PlayerCameraManager)
-	uint32 bFollowHmdOrientation:1;
+	uint32 bFollowHmdOrientation : 1;
 
 	/** Minimum view pitch, in degrees. */
 	UPROPERTY(EditAnywhere, BlueprintReadWrite, Category=PlayerCameraManager)
diff --git a/Engine/Source/Runtime/Engine/Private/PlayerCameraManager.cpp b/Engine/Source/Runtime/Engine/Private/PlayerCameraManager.cpp
--- a/Engine/Source/Runtime/Engine/Private/PlayerCameraManager.cpp
+++ b/Engine/Source/Runtime/Engine/Private/PlayerCameraManager.cpp
@@ -791,18 +791,17 @@ void APlayerCameraManager::InitializeFor(APlayerController* PC)
 
 float APlayerCameraManager::GetFOVAngle() const
 {
-	return bLockedFOV ? LockedFOV : CameraCache.POV.FOV;
+	return (LockedFOV > 0.f) ? LockedFOV : CameraCache.POV.FOV;
 }
 
 void APlayerCameraManager::SetFOV(float NewFOV)
 {
-	bLockedFOV = true;
 	LockedFOV = NewFOV;
 }
 
 void APlayerCameraManager::UnlockFOV()
 {
-	bLockedFOV = false;
+	LockedFOV = 0.f;
 }
 
 bool APlayerCameraManager::IsOrthographic() const
@@ -812,18 +811,17 @@ bool APlayerCameraManager::IsOrthographic() const
 
 float APlayerCameraManager::GetOrthoWidth() const
 {
-	return bLockedOrthoWidth ? LockedOrthoWidth : DefaultOrthoWidth;
+	return (LockedOrthoWidth > 0.f) ? LockedOrthoWidth : DefaultOrthoWidth;
 }
 
 void APlayerCameraManager::SetOrthoWidth(float OrthoWidth)
 {
-	bLockedOrthoWidth = true;
 	LockedOrthoWidth = OrthoWidth;
 }
 
 void APlayerCameraManager::UnlockOrthoWidth()
 {
-	bLockedOrthoWidth = false;
+	LockedOrthoWidth = 0.f;
 }
 
 void APlayerCameraManager::GetCameraViewPoint(FVector& OutCamLoc, FRotator& OutCamRot) const
