Fixed timer leak in timermanager with looping timers.

diff --git a/Engine/Source/Runtime/Engine/Private/TimerManager.cpp b/Engine/Source/Runtime/Engine/Private/TimerManager.cpp
--- a/Engine/Source/Runtime/Engine/Private/TimerManager.cpp
+++ b/Engine/Source/Runtime/Engine/Private/TimerManager.cpp
@@ -238,6 +238,7 @@ void FTimerManager::InternalSetTimer(FTimerData& NewTimerData, float InRate, boo
 	{
 		NewTimerData.Rate = InRate;
 		NewTimerData.bLoop = InbLoop;
+		NewTimerData.bRequiresDelegate = NewTimerData.TimerDelegate.IsBound();
 
 		const float FirstDelay = (InFirstDelay >= 0.f) ? InFirstDelay : InRate;
 
@@ -265,6 +266,7 @@ void FTimerManager::InternalSetTimerForNextTick(FTimerUnifiedDelegate const& InD
 	FTimerData NewTimerData;
 	NewTimerData.Rate = 0.f;
 	NewTimerData.bLoop = false;
+	NewTimerData.bRequiresDelegate = true;
 	NewTimerData.TimerDelegate = InDelegate;
 	NewTimerData.ExpireTime = InternalTime;
 	NewTimerData.Status = ETimerStatus::Active;
@@ -303,7 +305,7 @@ void FTimerManager::InternalClearTimer(FTimerHandle const& InHandle)
 	}
 }
 
-void FTimerManager::InternalClearTimer(int32 TimerIdx, ETimerStatus::Type TimerStatus)
+void FTimerManager::InternalClearTimer(int32 TimerIdx, ETimerStatus TimerStatus)
 {
 	switch (TimerStatus)
 	{
@@ -434,7 +436,7 @@ void FTimerManager::InternalPauseTimer( FTimerData const* TimerToPause, int32 Ti
 
 	if( TimerToPause && (TimerToPause->Status != ETimerStatus::Paused) )
 	{
-		ETimerStatus::Type PreviousStatus = TimerToPause->Status;
+		ETimerStatus PreviousStatus = TimerToPause->Status;
 
 		// Don't pause the timer if it's currently executing and isn't going to loop
 		if( PreviousStatus != ETimerStatus::Executing || TimerToPause->bLoop )
@@ -542,10 +544,14 @@ void FTimerManager::Tick(float DeltaTime)
 			// Status test needed to ensure it didn't get cleared during execution
 			if( CurrentlyExecutingTimer.bLoop && CurrentlyExecutingTimer.Status == ETimerStatus::Executing )
 			{
-				// Put this timer back on the heap
-				CurrentlyExecutingTimer.ExpireTime += CallCount * CurrentlyExecutingTimer.Rate;
-				CurrentlyExecutingTimer.Status = ETimerStatus::Active;
-				ActiveTimerHeap.HeapPush(CurrentlyExecutingTimer);
+				// if timer requires a delegate, make sure it's still validly bound (i.e. the delegate's object didn't get deleted or something)
+				if (!CurrentlyExecutingTimer.bRequiresDelegate || CurrentlyExecutingTimer.TimerDelegate.IsBound())
+				{
+					// Put this timer back on the heap
+					CurrentlyExecutingTimer.ExpireTime += CallCount * CurrentlyExecutingTimer.Rate;
+					CurrentlyExecutingTimer.Status = ETimerStatus::Active;
+					ActiveTimerHeap.HeapPush(CurrentlyExecutingTimer);
+				}
 			}
 
 			CurrentlyExecutingTimer.Clear();
diff --git a/Engine/Source/Runtime/Engine/Public/TimerManager.h b/Engine/Source/Runtime/Engine/Public/TimerManager.h
--- a/Engine/Source/Runtime/Engine/Public/TimerManager.h
+++ b/Engine/Source/Runtime/Engine/Public/TimerManager.h
@@ -109,24 +109,24 @@ private:
 	int32 Handle;
 };
 
-namespace ETimerStatus
+enum class ETimerStatus : uint8
 {
-	enum Type
-	{
-		Pending,
-		Active,
-		Paused,
-		Executing
-	};
-}
+	Pending,
+	Active,
+	Paused,
+	Executing
+};
 
 struct FTimerData
 {
 	/** If true, this timer will loop indefinitely.  Otherwise, it will be destroyed when it expires. */
-	bool bLoop;
+	uint8 bLoop : 1;
+
+	/** If true, this timer was created with a delegate to call (which means if the delegate becomes invalid, we should invalidate the timer too). */
+	uint8 bRequiresDelegate : 1;
 
 	/** Timer Status */
-	ETimerStatus::Type Status;
+	ETimerStatus Status;
 	
 	/** Time between set and fire, or repeat frequency if looping. */
 	float Rate;
@@ -144,8 +144,11 @@ struct FTimerData
 	FTimerHandle TimerHandle;
 
 	FTimerData()
-		: bLoop(false), Status(ETimerStatus::Active)
-		, Rate(0), ExpireTime(0)
+		: bLoop(false)
+		, bRequiresDelegate(false)
+		, Status(ETimerStatus::Active)
+		, Rate(0)
+		, ExpireTime(0)
 	{}
 
 	/** Operator less, used to sort the heap based on time until execution. **/
@@ -761,7 +764,7 @@ private:
 	void InternalSetTimerForNextTick( FTimerUnifiedDelegate const& InDelegate );
 	void DEPRECATED_InternalClearTimer( FTimerUnifiedDelegate const& InDelegate );
 	void InternalClearTimer( FTimerHandle const& InDelegate );
-	void InternalClearTimer( int32 TimerIdx, ETimerStatus::Type TimerStatus );
+	void InternalClearTimer( int32 TimerIdx, ETimerStatus TimerStatus );
 	void InternalClearAllTimers( void const* Object );
 
 	/** Will find a timer in the active, paused, or pending list. */
@@ -789,7 +792,7 @@ private:
 	/** An internally consistent clock, independent of World.  Advances during ticking. */
 	double InternalTime;
 
-	/** Timer delegate currently being executed.  Used to handle "timer delegates that manipulating timers" cases. */
+	/** Timer delegate currently being executed.  Used to handle "timer delegates that manipulate timers" cases. */
 	FTimerData CurrentlyExecutingTimer;
 
 	/** Set this to GFrameCounter when Timer is ticked. To figure out if Timer has been already ticked or not this frame. */
