UE-6275 - Set SpectatorPlayerController when recording

diff --git a/Engine/Source/Runtime/Engine/Private/DemoNetDriver.cpp b/Engine/Source/Runtime/Engine/Private/DemoNetDriver.cpp
--- a/Engine/Source/Runtime/Engine/Private/DemoNetDriver.cpp
+++ b/Engine/Source/Runtime/Engine/Private/DemoNetDriver.cpp
@@ -1324,9 +1324,9 @@ void UDemoNetDriver::SpawnDemoRecSpectator( UNetConnection* Connection )
 
 	FActorSpawnParameters SpawnInfo;
 	SpawnInfo.ObjectFlags |= RF_Transient;	// We never want these to save into a map
-	APlayerController* Controller = World->SpawnActor<APlayerController>( C, SpawnInfo );
+	SpectatorController = World->SpawnActor<APlayerController>( C, SpawnInfo );
 
-	if ( Controller == NULL )
+	if ( SpectatorController == NULL )
 	{
 		UE_LOG( LogDemo, Error, TEXT( "UDemoNetDriver::SpawnDemoRecSpectator: Failed to spawn demo spectator." ) );
 		return;
@@ -1336,15 +1336,15 @@ void UDemoNetDriver::SpawnDemoRecSpectator( UNetConnection* Connection )
 	{
 		if ( It->IsA( APlayerStart::StaticClass() ) )
 		{
-			Controller->SetInitialLocationAndRotation( It->GetActorLocation(), It->GetActorRotation() );
+			SpectatorController->SetInitialLocationAndRotation( It->GetActorLocation(), It->GetActorRotation() );
 			break;
 		}
 	}
 
-	Controller->SetReplicates( true );
-	Controller->SetAutonomousProxy( true );
+	SpectatorController->SetReplicates( true );
+	SpectatorController->SetAutonomousProxy( true );
 
-	Controller->SetPlayer( Connection );
+	SpectatorController->SetPlayer( Connection );
 }
 
 void UDemoNetDriver::ReplayStreamingReady( bool bSuccess, bool bRecord )
