#ue4 - Remove incorrect map check warning about non-uniform scale on meshes with simple collision. Only spheres and capsules are incompatible, boxes are fine.

diff --git a/Engine/Source/Runtime/Core/Public/Misc/MapErrors.h b/Engine/Source/Runtime/Core/Public/Misc/MapErrors.h
--- a/Engine/Source/Runtime/Core/Public/Misc/MapErrors.h
+++ b/Engine/Source/Runtime/Core/Public/Misc/MapErrors.h
@@ -170,7 +170,7 @@ struct CORE_API FMapErrors
 	/** {Actor} : Static mesh actor has NULL StaticMeshComponent property - please delete!: This warning is caused when there is a static mesh actor with a "None" StaticMeshComponent component.  It is usually found in older maps where duplication was used to create the actor and can be fixed by deleting the actor causing the warning and creating a new one. */
     static FName StaticMeshComponent;
     
-	/** {StaticMesh} has simple collision but is being scaled non-uniformly - collision creation will fail: Simple collision cannot be used with non-uniform scale. Please either fix the scale or the collision type. */
+	/** {StaticMesh} has simple collision but is being scaled non-uniformly - collision creation will fail: sphere or capsule simple collision cannot be used with non-uniform scale. Please either fix the scale or the collision type. */
     static FName SimpleCollisionButNonUniformScale;
 
 	/** More overriden materials {Count} on static mesh component than are referenced {Count} in source mesh {StaticMesh}: **TODO** */
diff --git a/Engine/Source/Runtime/Engine/Private/Components/StaticMeshComponent.cpp b/Engine/Source/Runtime/Engine/Private/Components/StaticMeshComponent.cpp
--- a/Engine/Source/Runtime/Engine/Private/Components/StaticMeshComponent.cpp
+++ b/Engine/Source/Runtime/Engine/Private/Components/StaticMeshComponent.cpp
@@ -384,16 +384,14 @@ void UStaticMeshComponent::CheckForErrors()
 		// Overall scale factor for this mesh.
 		const FVector& TotalScale3D = ComponentToWorld.GetScale3D();
 		if ( !TotalScale3D.IsUniform() &&
-			 (StaticMesh->BodySetup->AggGeom.BoxElems.Num() > 0   ||
-			  StaticMesh->BodySetup->AggGeom.SphylElems.Num() > 0 ||
+			 (StaticMesh->BodySetup->AggGeom.SphylElems.Num() > 0 ||
 			  StaticMesh->BodySetup->AggGeom.SphereElems.Num() > 0) )
-
 		{
 			FFormatNamedArguments Arguments;
 			Arguments.Add(TEXT("MeshName"), FText::FromString(StaticMesh->GetName()));
 			FMessageLog("MapCheck").Warning()
 				->AddToken(FUObjectToken::Create(Owner))
-				->AddToken(FTextToken::Create(FText::Format(LOCTEXT( "MapCheck_Message_SimpleCollisionButNonUniformScale", "'{MeshName}' has simple collision but is being scaled non-uniformly - collision creation will fail" ), Arguments)))
+				->AddToken(FTextToken::Create(FText::Format(LOCTEXT( "MapCheck_Message_SimpleCollisionButNonUniformScale", "'{MeshName}' has sphere or capsule simple collision but is being scaled non-uniformly - collision creation will fail" ), Arguments)))
 				->AddToken(FMapErrorToken::Create(FMapErrors::SimpleCollisionButNonUniformScale));
 		}
 	}
@@ -402,7 +400,7 @@ void UStaticMeshComponent::CheckForErrors()
 	{
 		FMessageLog("MapCheck").Warning()
 			->AddToken(FUObjectToken::Create(this))
-			->AddToken(FTextToken::Create(FText::Format(LOCTEXT( "MapCheck_Message_SimulatePhyNoSimpleCollision", "{0} : Using bSimulatePhysics but StaticMesh has not simple collision."), FText::FromString(GetName()) ) ));
+			->AddToken(FTextToken::Create(FText::Format(LOCTEXT( "MapCheck_Message_SimulatePhyNoSimpleCollision", "{0} : Using bSimulatePhysics but StaticMesh has no simple collision."), FText::FromString(GetName()) ) ));
 	}
 
 	if( Mobility == EComponentMobility::Movable &&
