#ANIM: Allow dynamic addition to rig, but disabled the addition for now
Renamed typo

diff --git a/Engine/Config/BaseEngine.ini b/Engine/Config/BaseEngine.ini
--- a/Engine/Config/BaseEngine.ini
+++ b/Engine/Config/BaseEngine.ini
@@ -278,6 +278,8 @@ bUseStreamingPause=false
 +ActiveClassRedirects=(OldClassName="AnimNode_CopyBone",NewClassName="/Script/Engine.AnimNode_CopyBone")
 +ActiveClassRedirects=(OldClassName="AnimNode_SpringBone",NewClassName="/Script/Engine.AnimNode_SpringBone")
 
++EnumRedirects=(EnumName="EConstraintTransform",OldEnumEntry="EConstraintTransform::Absoluate",NewEnumEntry="EConstraintTransform::Absolute")
+
 ;Navigation redirects
 +ActiveClassRedirects=(OldClassName="NavAreaMeta",NewClassName="/Script/Engine.NavArea_Default",InstanceOnly="true")
 +ActiveClassRedirects=(OldClassName="NavAreaDefinition",NewClassName="/Script/Engine.NavArea")
diff --git a/Engine/Source/Editor/DetailCustomizations/Private/RigDetails.cpp b/Engine/Source/Editor/DetailCustomizations/Private/RigDetails.cpp
--- a/Engine/Source/Editor/DetailCustomizations/Private/RigDetails.cpp
+++ b/Engine/Source/Editor/DetailCustomizations/Private/RigDetails.cpp
@@ -39,23 +39,6 @@ void FRigDetails::CustomizeDetails(IDetailLayoutBuilder& DetailBuilder)
 	TransformBasesPropertyHandle = DetailBuilder.GetProperty("TransformBases");
 	NodesPropertyHandle = DetailBuilder.GetProperty("Nodes");
 
-	// since now we can't really resize the array, we'll just allocate everything here
-	// if you reallocate, TSharedPtr<FString> for combo box won't work because
-	uint32 NumElement = 0;
-	check (FPropertyAccess::Fail != NodesPropertyHandle->AsArray()->GetNumElements(NumElement));
-
-	if ( NumElement > 0 )
-	{
-		DisplayNameTextBoxes.AddZeroed(NumElement);
-	}
-
-	check (FPropertyAccess::Fail != TransformBasesPropertyHandle->AsArray()->GetNumElements(NumElement));
-	if ( NumElement > 0 )
-	{
-		ParentSpaceOptionList.AddZeroed(NumElement);
-		// I need two per each - translation/orientation
-		ParentSpaceComboBoxes.AddZeroed(NumElement*2);
-	}
 
 	TSharedRef<FDetailArrayBuilder> NodeArrayBuilder = MakeShareable(new FDetailArrayBuilder(NodesPropertyHandle.ToSharedRef()));
 	NodeArrayBuilder->OnGenerateArrayElementWidget(FOnGenerateArrayElementWidget::CreateSP(this, &FRigDetails::GenerateNodeArrayElementWidget, &DetailBuilder));
@@ -201,6 +184,12 @@ void FRigDetails::GenerateNodeArrayElementWidget(TSharedRef<IPropertyHandle> Pro
 		]
 	];
 
+	// we're adding on demand because if array expands, we'll have to handle it dynamically
+	if (!DisplayNameTextBoxes.IsValidIndex(ArrayIndex))
+	{
+		DisplayNameTextBoxes.AddZeroed(ArrayIndex-DisplayNameTextBoxes.Num()+1);
+	}
+
 	DisplayNameTextBoxes[ArrayIndex] = DisplayTextBox;
 };
 
@@ -279,9 +268,22 @@ void FRigDetails::GenerateTransformBaseArrayElementWidget(TSharedRef<IPropertyHa
 	check(ParentNameProp_T->GetValueAsDisplayString(ParentNodeName_T) != FPropertyAccess::Fail);
 	check(ParentNameProp_R->GetValueAsDisplayString(ParentNodeName_R) != FPropertyAccess::Fail);
 
+	// we add them on demand since the array can expand or reduce
+	if(!ParentSpaceOptionList.IsValidIndex(ArrayIndex))
+	{
+		ParentSpaceOptionList.AddZeroed(ArrayIndex-ParentSpaceOptionList.Num()+1);
+	}
+
+	// ParentSpaceComboBoxees creates 2 per element - one for translation, and one for rotation
+	if(!ParentSpaceComboBoxes.IsValidIndex(ArrayIndex*2))
+	{
+		int32 NumNewItem = ArrayIndex-(ParentSpaceComboBoxes.Num()/2)+1;
+		ParentSpaceComboBoxes.AddZeroed(NumNewItem*2);
+	}
+
 	// create string list for picking parent node
 	// make sure you don't include itself and find what is curretn selected item
-	TArray<TSharedPtr<FString>> & ParentNodeOptions = ParentSpaceOptionList[ArrayIndex];
+	TArray<TSharedPtr<FString>>& ParentNodeOptions = ParentSpaceOptionList[ArrayIndex];
 
 	ParentNodeOptions.Empty();
 	ParentNodeOptions.Add(MakeShareable(new FString(URig::WorldNodeName.ToString())));
diff --git a/Engine/Source/Runtime/Engine/Classes/Animation/Rig.h b/Engine/Source/Runtime/Engine/Classes/Animation/Rig.h
--- a/Engine/Source/Runtime/Engine/Classes/Animation/Rig.h
+++ b/Engine/Source/Runtime/Engine/Classes/Animation/Rig.h
@@ -76,8 +76,8 @@ namespace EConstraintTransform
 {
 	enum Type
 	{
-		/** Absoluate value. */
-		Absoluate,
+		/** Absolute value. */
+		Absolute,
 		/** Apply relative transform from ref pose. */
 		Relative,
 	};
diff --git a/Engine/Source/Runtime/Engine/Private/Animation/Rig.cpp b/Engine/Source/Runtime/Engine/Private/Animation/Rig.cpp
--- a/Engine/Source/Runtime/Engine/Private/Animation/Rig.cpp
+++ b/Engine/Source/Runtime/Engine/Private/Animation/Rig.cpp
@@ -21,11 +21,50 @@ URig::URig(const FObjectInitializer& ObjectInitializer)
 #if WITH_EDITOR
 void URig::PostEditChangeProperty(FPropertyChangedEvent& PropertyChangedEvent)
 {
-	Super::PostEditChangeProperty(PropertyChangedEvent);
+	if (Nodes.Num() != TransformBases.Num())
+	{
+		int32 NodeNum = Nodes.Num();
+		int32 TransformBasesNum = TransformBases.Num();
+
+		// make sure to add name of nodes if not set
+		for(int32 NewId=TransformBasesNum; NewId <NodeNum; ++NewId)
+		{
+			if (Nodes[NewId].Name == NAME_None)
+			{
+				int32 UniqueIndex = 1;
+
+				while (1) 
+				{
+					// assign name of custom_# 
+					FName PotentialName = *FString::Printf(TEXT("Custom_%d"), UniqueIndex++);
+					if (FindNode(PotentialName) == INDEX_NONE)
+					{
+						Nodes[NewId].Name = PotentialName;
+						Nodes[NewId].ParentName = WorldNodeName;
+						Nodes[NewId].DisplayName = PotentialName.ToString();
+						break;
+					}
+				} 
+			}
+		}
+
+		// if their size doesn't match, make sure to add it
+		if (NodeNum < TransformBasesNum)
+		{
+			// just remove the last elements
+			TransformBases.RemoveAt(NodeNum, TransformBasesNum-NodeNum);
+		}
+		else
+		{
+			for (int32 NewId=TransformBasesNum; NewId <NodeNum; ++NewId)
+			{
+				AddRigConstraint(Nodes[NewId].Name, EControlConstraint::Translation, EConstraintTransform::Absolute, WorldNodeName, 1.f);
+				AddRigConstraint(Nodes[NewId].Name, EControlConstraint::Orientation, EConstraintTransform::Absolute, WorldNodeName, 1.f);
+			}
+		}
+	}
 
-	// verify lots of things
-	// make sure it's unique name
-	// 
+	Super::PostEditChangeProperty(PropertyChangedEvent);
 }
 
 // for now these are privates since we don't have good control yet
@@ -281,14 +320,14 @@ void URig::CreateFromSkeleton(const USkeleton* Skeleton, const TMap<int32, int32
 
 			if (ParentIndex == INDEX_NONE)
 			{
-				AddRigConstraint(BoneName, EControlConstraint::Translation, EConstraintTransform::Absoluate, WorldNodeName, 1.f);
-				AddRigConstraint(BoneName, EControlConstraint::Orientation, EConstraintTransform::Absoluate, WorldNodeName, 1.f);
+				AddRigConstraint(BoneName, EControlConstraint::Translation, EConstraintTransform::Absolute, WorldNodeName, 1.f);
+				AddRigConstraint(BoneName, EControlConstraint::Orientation, EConstraintTransform::Absolute, WorldNodeName, 1.f);
 			}
 			else
 			{
 				FName ParentBoneName = RefSkeleton.GetBoneName(ParentIndex);
-				AddRigConstraint(BoneName, EControlConstraint::Translation, EConstraintTransform::Absoluate, ParentBoneName, 1.f);
-				AddRigConstraint(BoneName, EControlConstraint::Orientation, EConstraintTransform::Absoluate, ParentBoneName, 1.f);
+				AddRigConstraint(BoneName, EControlConstraint::Translation, EConstraintTransform::Absolute, ParentBoneName, 1.f);
+				AddRigConstraint(BoneName, EControlConstraint::Orientation, EConstraintTransform::Absolute, ParentBoneName, 1.f);
 			}
 		}
 	}
