#ue4 - Compute ImpactNormal for collisions with Box shapes. Clean up FindGeomOpposingNormal() code.

UE-11116

diff --git a/Engine/Source/Runtime/Engine/Private/Collision/CollisionConversions.cpp b/Engine/Source/Runtime/Engine/Private/Collision/CollisionConversions.cpp
--- a/Engine/Source/Runtime/Engine/Private/Collision/CollisionConversions.cpp
+++ b/Engine/Source/Runtime/Engine/Private/Collision/CollisionConversions.cpp
@@ -18,6 +18,12 @@ extern TAutoConsoleVariable<int32> CVarShowInitialOverlaps;
 // Sentinel for invalid query results.
 static const PxQueryHit InvalidQueryHit;
 
+static bool IsInvalidFaceIndex(PxU32 faceIndex)
+{
+	checkf(InvalidQueryHit.faceIndex == 0xFFFFffff, TEXT("Engine code needs fixing: PhysX invalid face index sentinel has changed or is not part of default PxQueryHit!"));
+	return (faceIndex == InvalidQueryHit.faceIndex);
+}
+
 // Forward declare, I don't want to move the entire function right now or we lose change history.
 static bool ConvertOverlappedShapeToImpactHit(const UWorld* World, const PxLocationHit& PHit, const FVector& StartLoc, const FVector& EndLoc, FHitResult& OutResult, const PxGeometry& Geom, const PxTransform& QueryTM, const PxFilterData& QueryFilter, bool bReturnPhysMat);
 
@@ -91,17 +97,74 @@ static PxVec3 TransformNormalToShapeSpace(const PxMeshScale& meshScale, const Px
 
 static bool FindSimpleOpposingNormal(const PxLocationHit& PHit, const FVector& TraceDirectionDenorm, FVector& OutNormal)
 {
-	const bool bNormalData = PHit.flags & PxHitFlag::eNORMAL;
-	OutNormal = bNormalData ? P2UVector(PHit.normal) : -TraceDirectionDenorm.GetSafeNormal();
+	// We don't compute anything special, so let calling code figure out the correct fallback.
+	return false;
+}
+
+static bool FindBoxOpposingNormal(const PxLocationHit& PHit, const FVector& TraceDirectionDenorm, FVector& OutNormal)
+{
+	// We require normal info for our algorithm.
+	const bool bNormalData = (PHit.flags & PxHitFlag::eNORMAL);
+	if (!bNormalData)
+	{
+		return false;
+	}
+
+	PxBoxGeometry PxBoxGeom;
+	const bool bReadGeomSuccess = PHit.shape->getBoxGeometry(PxBoxGeom);
+	check(bReadGeomSuccess); // This function should only be used for box geometry
+
+	const PxTransform LocalToWorld = PxShapeExt::getGlobalPose(*PHit.shape, *PHit.actor);
+	
+	// Find which faces were included in the contact normal, and for multiple faces, use the one most opposing the sweep direction.
+	const PxVec3 ContactNormalLocal = LocalToWorld.rotateInv(PHit.normal);
+	const PxVec3 TraceDirDenormWorld = U2PVector(TraceDirectionDenorm);
+	const PxVec3 TraceDirDenormLocal = LocalToWorld.rotateInv(TraceDirDenormWorld);
+
+	PxVec3 BestLocalNormal(ContactNormalLocal);
+	float BestOpposingDot = FLT_MAX;
+
+	for (int32 i=0; i < 3; i++)
+	{
+		// Select axis of face to compare to, based on normal.
+		if (ContactNormalLocal[i] > KINDA_SMALL_NUMBER)
+		{
+			const float TraceDotFaceNormal = TraceDirDenormLocal[i]; // TraceDirDenormLocal.dot(LocalNormal)
+			if (TraceDotFaceNormal < BestOpposingDot)
+			{
+				BestOpposingDot = TraceDotFaceNormal;
+				BestLocalNormal = PxVec3(0.f);
+				BestLocalNormal[i] = 1.f;
+			}
+		}
+		else if (ContactNormalLocal[i] < -KINDA_SMALL_NUMBER)
+		{
+			const float TraceDotFaceNormal = -TraceDirDenormLocal[i]; // TraceDirDenormLocal.dot(LocalNormal)
+			if (TraceDotFaceNormal < BestOpposingDot)
+			{
+				BestOpposingDot = TraceDotFaceNormal;
+				BestLocalNormal = PxVec3(0.f);
+				BestLocalNormal[i] = -1.f;
+			}
+		}
+	}
+
+	// Fill in result
+	const PxVec3 WorldNormal = LocalToWorld.rotate(BestLocalNormal);
+	OutNormal = P2UVector(WorldNormal);
 	return true;
 }
 
 static bool FindHeightFieldOpposingNormal(const PxLocationHit& PHit, const FVector& TraceDirectionDenorm, FVector& OutNormal)
 {
+	if (IsInvalidFaceIndex(PHit.faceIndex))
+	{
+		return false;
+	}
 
 	PxHeightFieldGeometry PHeightFieldGeom;
-	bool bSuccess = PHit.shape->getHeightFieldGeometry(PHeightFieldGeom);
-	check(bSuccess);	//we should only call this function when we have a heightfield
+	const bool bReadGeomSuccess = PHit.shape->getHeightFieldGeometry(PHeightFieldGeom);
+	check(bReadGeomSuccess);	//we should only call this function when we have a heightfield
 	if (PHeightFieldGeom.heightField)
 	{
 		const PxU32 TriIndex = PHit.faceIndex;
@@ -122,6 +185,11 @@ static bool FindHeightFieldOpposingNormal(const PxLocationHit& PHit, const FVect
 
 static bool FindConvexMeshOpposingNormal(const PxLocationHit& PHit, const FVector& TraceDirectionDenorm, FVector& OutNormal)
 {
+	if (IsInvalidFaceIndex(PHit.faceIndex))
+	{
+		return false;
+	}
+
 	PxConvexMeshGeometry PConvexMeshGeom;
 	bool bSuccess = PHit.shape->getConvexMeshGeometry(PConvexMeshGeom);
 	check(bSuccess);	//should only call this function when we have a convex mesh
@@ -160,6 +228,11 @@ static bool FindConvexMeshOpposingNormal(const PxLocationHit& PHit, const FVecto
 
 static bool FindTriMeshOpposingNormal(const PxLocationHit& PHit, const FVector& TraceDirectionDenorm, FVector& OutNormal)
 {
+	if (IsInvalidFaceIndex(PHit.faceIndex))
+	{
+		return false;
+	}
+
 	PxTriangleMeshGeometry PTriMeshGeom;
 	bool bSuccess = PHit.shape->getTriangleMeshGeometry(PTriMeshGeom);
 	check(bSuccess);	//this function should only be called when we have a trimesh
@@ -231,24 +304,22 @@ static bool FindTriMeshOpposingNormal(const PxLocationHit& PHit, const FVector&
  * @param OutNormal - normal we may recompute based on the faceIndex of the hit
  * @return true if we compute a new normal for the geometry.
  */
-static bool FindGeomOpposingNormal(const PxLocationHit& PHit, const FVector& TraceDirectionDenorm, FVector& OutNormal)
+static bool FindGeomOpposingNormal(PxGeometryType::Enum QueryGeomType, const PxLocationHit& PHit, const FVector& TraceDirectionDenorm, FVector& OutNormal)
 {
-	checkf(InvalidQueryHit.faceIndex == 0xFFFFffff, TEXT("Engine code needs fixing: PhysX invalid face index sentinel has changed or is not part of default PxQueryHit!"));
-	if (PHit.faceIndex == InvalidQueryHit.faceIndex)
-	{
-		return false;
-	}
-
-	PxGeometryType::Enum GeomType = PHit.shape->getGeometryType();
-	switch (GeomType)
+	// TODO: can we support other shapes here as well?
+	if (QueryGeomType == PxGeometryType::eCAPSULE || QueryGeomType == PxGeometryType::eSPHERE)
 	{
+		PxGeometryType::Enum GeomType = PHit.shape->getGeometryType();
+		switch (GeomType)
+		{
 		case PxGeometryType::eSPHERE:
-		case PxGeometryType::eBOX:
 		case PxGeometryType::eCAPSULE:		return FindSimpleOpposingNormal(PHit, TraceDirectionDenorm, OutNormal);
+		case PxGeometryType::eBOX:			return FindBoxOpposingNormal(PHit, TraceDirectionDenorm, OutNormal);
 		case PxGeometryType::eCONVEXMESH:	return FindConvexMeshOpposingNormal(PHit, TraceDirectionDenorm, OutNormal);
 		case PxGeometryType::eHEIGHTFIELD:	return FindHeightFieldOpposingNormal(PHit, TraceDirectionDenorm, OutNormal);
 		case PxGeometryType::eTRIANGLEMESH:	return FindTriMeshOpposingNormal(PHit, TraceDirectionDenorm, OutNormal);
 		default: check(false);	//unsupported geom type
+		}
 	}
 
 	return false;
@@ -392,13 +463,8 @@ void ConvertQueryImpactHit(const UWorld* World, const PxLocationHit& PHit, FHitR
 		OutResult.ImpactNormal = Normal;
 	}
 
-	PxGeometryType::Enum GeometryType = Geom ? Geom->getType() : PxGeometryType::eINVALID;
-
-	// Special handling for swept-capsule results
-	if (GeometryType == PxGeometryType::eCAPSULE || GeometryType == PxGeometryType::eSPHERE)
-	{
-		FindGeomOpposingNormal(PHit, TraceStartToEnd, OutResult.ImpactNormal);
-	}
+	const PxGeometryType::Enum SweptGeometryType = Geom ? Geom->getType() : PxGeometryType::eINVALID;
+	FindGeomOpposingNormal(SweptGeometryType, PHit, TraceStartToEnd, OutResult.ImpactNormal);
 	
 	if( PHit.shape->getGeometryType() == PxGeometryType::eHEIGHTFIELD)
 	{
