#ue4 - Character crouch/uncrouch now changes crouch state regardless of whether at the target capsule height. CanCrouch() no longer checks the height condition.

AnswerHub: https://answers.unrealengine.com/questions/180490/471-uncrouch-does-not-work.html#answer-209593

diff --git a/Engine/Source/Runtime/Engine/Private/Character.cpp b/Engine/Source/Runtime/Engine/Private/Character.cpp
--- a/Engine/Source/Runtime/Engine/Private/Character.cpp
+++ b/Engine/Source/Runtime/Engine/Private/Character.cpp
@@ -285,7 +285,7 @@ void ACharacter::OnRep_IsCrouched()
 
 bool ACharacter::CanCrouch()
 {
-	return !bIsCrouched && CharacterMovement && CharacterMovement->CanEverCrouch() && (CapsuleComponent->GetUnscaledCapsuleHalfHeight() > CharacterMovement->CrouchedHalfHeight) && GetRootComponent() && !GetRootComponent()->IsSimulatingPhysics();
+	return !bIsCrouched && CharacterMovement && CharacterMovement->CanEverCrouch() && GetRootComponent() && !GetRootComponent()->IsSimulatingPhysics();
 }
 
 void ACharacter::Crouch(bool bClientSimulation)
diff --git a/Engine/Source/Runtime/Engine/Private/Components/CharacterMovementComponent.cpp b/Engine/Source/Runtime/Engine/Private/Components/CharacterMovementComponent.cpp
--- a/Engine/Source/Runtime/Engine/Private/Components/CharacterMovementComponent.cpp
+++ b/Engine/Source/Runtime/Engine/Private/Components/CharacterMovementComponent.cpp
@@ -1705,14 +1705,19 @@ void UCharacterMovementComponent::Crouch(bool bClientSimulation)
 		return;
 	}
 
-	// Do not perform if collision is already at desired size.
-	if( CharacterOwner->GetCapsuleComponent()->GetUnscaledCapsuleHalfHeight() == CrouchedHalfHeight )
+	if (!CanCrouchInCurrentState())
 	{
 		return;
 	}
 
-	if (!CanCrouchInCurrentState())
+	// See if collision is already at desired size.
+	if (CharacterOwner->GetCapsuleComponent()->GetUnscaledCapsuleHalfHeight() == CrouchedHalfHeight)
 	{
+		if (!bClientSimulation)
+		{
+			CharacterOwner->bIsCrouched = true;
+		}
+		CharacterOwner->OnStartCrouch( 0.f, 0.f );
 		return;
 	}
 
@@ -1781,9 +1786,14 @@ void UCharacterMovementComponent::UnCrouch(bool bClientSimulation)
 
 	ACharacter* DefaultCharacter = CharacterOwner->GetClass()->GetDefaultObject<ACharacter>();
 
-	// Do not perform if collision is already at desired size.
+	// See if collision is already at desired size.
 	if( CharacterOwner->GetCapsuleComponent()->GetUnscaledCapsuleHalfHeight() == DefaultCharacter->GetCapsuleComponent()->GetUnscaledCapsuleHalfHeight() )
 	{
+		if (!bClientSimulation)
+		{
+			CharacterOwner->bIsCrouched = false;
+		}
+		CharacterOwner->OnEndCrouch( 0.f, 0.f );
 		return;
 	}
 
