Fixed crash when opening Level Blueprint for the first time when there is an Actor in the level with the same name as the level.

diff --git a/Engine/Source/Runtime/Engine/Private/Level.cpp b/Engine/Source/Runtime/Engine/Private/Level.cpp
--- a/Engine/Source/Runtime/Engine/Private/Level.cpp
+++ b/Engine/Source/Runtime/Engine/Private/Level.cpp
@@ -536,6 +536,13 @@ void ULevel::PostLoad()
 		UWorld* OuterWorld = Cast<UWorld>(GetOuter());
 		if (LevelScriptBlueprint && OuterWorld && LevelScriptBlueprint->GetFName() != OuterWorld->GetFName())
 		{
+			// The level blueprint must be named the same as the level/world.
+			// If there is already something there with that name, rename it to something else.
+			if (UObject* ExistingObject = StaticFindObject(nullptr, LevelScriptBlueprint->GetOuter(), *OuterWorld->GetName()))
+			{
+				ExistingObject->Rename(nullptr, nullptr, REN_DoNotDirty | REN_DontCreateRedirectors | REN_ForceNoResetLoaders | REN_NonTransactional);
+			}
+
 			// Use LevelScriptBlueprint->GetOuter() instead of NULL to make sure the generated top level objects are moved appropriately
 			LevelScriptBlueprint->Rename(*OuterWorld->GetName(), LevelScriptBlueprint->GetOuter(), REN_DoNotDirty | REN_DontCreateRedirectors | REN_ForceNoResetLoaders | REN_NonTransactional | REN_SkipGeneratedClasses);
 		}
@@ -1639,6 +1646,13 @@ ULevelScriptBlueprint* ULevel::GetLevelScriptBlueprint(bool bDontCreate)
 	const FString LevelScriptName = ULevelScriptBlueprint::CreateLevelScriptNameFromLevel(this);
 	if( !LevelScriptBlueprint && !bDontCreate)
 	{
+		// The level blueprint must be named the same as the level/world.
+		// If there is already something there with that name, rename it to something else.
+		if (UObject* ExistingObject = StaticFindObject(nullptr, this, *LevelScriptName))
+		{
+			ExistingObject->Rename(nullptr, nullptr, REN_DoNotDirty | REN_DontCreateRedirectors | REN_ForceNoResetLoaders | REN_NonTransactional);
+		}
+
 		// If no blueprint is found, create one. 
 		LevelScriptBlueprint = Cast<ULevelScriptBlueprint>(FKismetEditorUtilities::CreateBlueprint(GEngine->LevelScriptActorClass, this, FName(*LevelScriptName), BPTYPE_LevelScript, ULevelScriptBlueprint::StaticClass(), UBlueprintGeneratedClass::StaticClass()));
 
diff --git a/Engine/Source/Runtime/Engine/Private/World.cpp b/Engine/Source/Runtime/Engine/Private/World.cpp
--- a/Engine/Source/Runtime/Engine/Private/World.cpp
+++ b/Engine/Source/Runtime/Engine/Private/World.cpp
@@ -292,6 +292,13 @@ bool UWorld::Rename(const TCHAR* InName, UObject* NewOuter, ERenameFlags Flags)
 			}
 			else
 			{
+				// The level blueprint must be named the same as the level/world.
+				// If there is already something there with that name, rename it to something else.
+				if (UObject* ExistingObject = StaticFindObject(nullptr, LevelScriptBlueprint->GetOuter(), InName))
+				{
+					ExistingObject->Rename(nullptr, nullptr, REN_DoNotDirty | REN_DontCreateRedirectors | REN_ForceNoResetLoaders | REN_NonTransactional);
+				}
+
 				// This is a normal rename. Use LevelScriptBlueprint->GetOuter() instead of NULL to make sure the generated top level objects are moved appropriately
 				if ( !LevelScriptBlueprint->Rename(InName, LevelScriptBlueprint->GetOuter(), Flags) )
 				{
