Dynamic reregistering of PreClothTick based on its actual need to run

diff --git a/Engine/Source/Runtime/Engine/Classes/Components/SkeletalMeshComponent.h b/Engine/Source/Runtime/Engine/Classes/Components/SkeletalMeshComponent.h
--- a/Engine/Source/Runtime/Engine/Classes/Components/SkeletalMeshComponent.h
+++ b/Engine/Source/Runtime/Engine/Classes/Components/SkeletalMeshComponent.h
@@ -803,6 +803,10 @@ public:
 	virtual void InitializeComponent() override;
 	virtual void TickComponent(float DeltaTime, enum ELevelTick TickType, FActorComponentTickFunction *ThisTickFunction) override;
 	virtual void RegisterComponentTickFunctions(bool bRegister) override;
+
+	//Handle registering our pre cloth tick function
+	void RegisterPreClothTick(bool bRegister);
+
 	// End UActorComponent interface.
 
 	// Begin USceneComponent interface.
@@ -1202,6 +1206,12 @@ public:
 	friend class FSkeletalMeshComponentDetails;
 
 private:
+	// Returns whether we need to run the Pre Cloth Tick or not
+	bool ShouldRunPreClothTick() const;
+
+	// Handles registering/unregistering the pre cloth tick as it is needed
+	void UpdatePreClothTickRegisteredState();
+
 	// these are deprecated variables from removing SingleAnimSkeletalComponent
 	// remove if this version goes away : VER_UE4_REMOVE_SINGLENODEINSTANCE
 	// deprecated variable to be re-save
diff --git a/Engine/Source/Runtime/Engine/Private/Components/SkeletalMeshComponent.cpp b/Engine/Source/Runtime/Engine/Private/Components/SkeletalMeshComponent.cpp
--- a/Engine/Source/Runtime/Engine/Private/Components/SkeletalMeshComponent.cpp
+++ b/Engine/Source/Runtime/Engine/Private/Components/SkeletalMeshComponent.cpp
@@ -149,26 +149,42 @@ void USkeletalMeshComponent::RegisterComponentTickFunctions(bool bRegister)
 {
 	Super::RegisterComponentTickFunctions(bRegister);
 
-	if (bRegister)
+	UpdatePreClothTickRegisteredState();
+}
+
+void USkeletalMeshComponent::RegisterPreClothTick(bool bRegister)
+{
+	if (bRegister != PreClothTickFunction.IsTickFunctionRegistered())
 	{
-		if (SetupActorComponentTickFunction(&PreClothTickFunction))
+		if (bRegister)
 		{
-			PreClothTickFunction.Target = this;
-			// Set a prereq for the pre cloth tick to happen after physics is finished
-			if (World != NULL)
+			if (SetupActorComponentTickFunction(&PreClothTickFunction))
 			{
-				PreClothTickFunction.AddPrerequisite(World, World->EndPhysicsTickFunction);
+				PreClothTickFunction.Target = this;
+				// Set a prereq for the pre cloth tick to happen after physics is finished
+				if (World != NULL)
+				{
+					PreClothTickFunction.AddPrerequisite(World, World->EndPhysicsTickFunction);
+				}
 			}
 		}
-	}
-	else
-	{
-		if (PreClothTickFunction.IsTickFunctionRegistered())
+		else
 		{
 			PreClothTickFunction.UnRegisterTickFunction();
 		}
 	}
+}
 
+bool USkeletalMeshComponent::ShouldRunPreClothTick() const
+{
+	return	(bEnablePhysicsOnDedicatedServer || !IsRunningDedicatedServer()) && // Early out with we are on a dedicated server and not running physics
+			(IsSimulatingPhysics() || ShouldBlendPhysicsBones() || (SkeletalMesh && SkeletalMesh->ClothingAssets.Num() > 0));
+}
+
+void USkeletalMeshComponent::UpdatePreClothTickRegisteredState()
+{
+	bool bShouldRunClothTick = ShouldRunPreClothTick();
+	RegisterPreClothTick(bShouldRunClothTick && PrimaryComponentTick.IsTickFunctionRegistered());
 }
 
 bool USkeletalMeshComponent::NeedToSpawnAnimScriptInstance(bool bForceInit) const
@@ -586,6 +602,8 @@ void USkeletalMeshComponent::TickPose(float DeltaTime, bool bNeedsValidRootMotio
 
 void USkeletalMeshComponent::TickComponent(float DeltaTime, enum ELevelTick TickType, FActorComponentTickFunction *ThisTickFunction)
 {
+	UpdatePreClothTickRegisteredState();
+
 	Super::TickComponent(DeltaTime, TickType, ThisTickFunction);
 
 	// Update bOldForceRefPose
@@ -990,11 +1008,6 @@ void USkeletalMeshComponent::PerformAnimationEvaluation(const USkeletalMesh* InS
 	FillSpaceBases(InSkeletalMesh, OutLocalAtoms, OutSpaceBases);
 }
 
-const TCHAR* B(bool b)
-{
-	return b ? TEXT("true") : TEXT("false");
-}
-
 void USkeletalMeshComponent::RefreshBoneTransforms(FActorComponentTickFunction* TickFunction)
 {
 	SCOPE_CYCLE_COUNTER(STAT_RefreshBoneTransforms);
@@ -2124,4 +2137,4 @@ void USkeletalMeshComponent::RefreshActiveVertexAnims()
 	{
 		ActiveVertexAnims.Empty();
 	}
-}
\ No newline at end of file
+}
diff --git a/Engine/Source/Runtime/Engine/Private/SkeletalMeshComponentPhysics.cpp b/Engine/Source/Runtime/Engine/Private/SkeletalMeshComponentPhysics.cpp
--- a/Engine/Source/Runtime/Engine/Private/SkeletalMeshComponentPhysics.cpp
+++ b/Engine/Source/Runtime/Engine/Private/SkeletalMeshComponentPhysics.cpp
@@ -3092,6 +3092,12 @@ void USkeletalMeshComponent::ProcessClothCollisionWithEnvironment()
 
 void USkeletalMeshComponent::PreClothTick(float DeltaTime)
 {
+	//IMPORTANT!
+	//
+	// The decision on whether to use PreClothTick or not is made by ShouldRunPreClothTick()
+	// Any changes that are made to PreClothTick that effect whether it should be run or not
+	// have to be reflected in ShouldRunPreClothTick() as well
+	
 	// if physics is disabled on dedicated server, no reason to be here. 
 	if (!bEnablePhysicsOnDedicatedServer && IsRunningDedicatedServer())
 	{
