Fix stopped components not being restarted  when reimporting sound

diff --git a/Engine/Source/Runtime/Engine/Private/AudioDeviceManager.cpp b/Engine/Source/Runtime/Engine/Private/AudioDeviceManager.cpp
--- a/Engine/Source/Runtime/Engine/Private/AudioDeviceManager.cpp
+++ b/Engine/Source/Runtime/Engine/Private/AudioDeviceManager.cpp
@@ -287,7 +287,6 @@ void FAudioDeviceManager::StopSoundsUsingResource(class USoundWave* InSoundWave,
 	{
 		if (AudioDevice)
 		{
-			TArray<UAudioComponent*> StoppedComponents;
 			AudioDevice->StopSoundsUsingResource(InSoundWave, StoppedComponents);
 		}
 	}
