PR #766: Made it possible to inherit from UControlChannel.

diff --git a/Engine/Source/Runtime/Engine/Classes/Engine/ControlChannel.h b/Engine/Source/Runtime/Engine/Classes/Engine/ControlChannel.h
--- a/Engine/Source/Runtime/Engine/Classes/Engine/ControlChannel.h
+++ b/Engine/Source/Runtime/Engine/Classes/Engine/ControlChannel.h
@@ -9,7 +9,7 @@
  * A channel for exchanging connection control messages.
  */
 UCLASS(transient, customConstructor)
-class UControlChannel
+class ENGINE_API UControlChannel
 	: public UChannel
 {
 	GENERATED_UCLASS_BODY()
@@ -66,8 +66,8 @@ class UControlChannel
 
 
 	/** Handle an incoming bunch. */
-	void ReceivedBunch( FInBunch& Bunch ) override;
+	virtual void ReceivedBunch( FInBunch& Bunch ) override;
 
 	/** Describe the text channel. */
-	FString Describe() override;
+	virtual FString Describe() override;
 };
