PR #681: Marked the UVoiceChannel class for export

diff --git a/Engine/Source/Runtime/Engine/Classes/Engine/VoiceChannel.h b/Engine/Source/Runtime/Engine/Classes/Engine/VoiceChannel.h
--- a/Engine/Source/Runtime/Engine/Classes/Engine/VoiceChannel.h
+++ b/Engine/Source/Runtime/Engine/Classes/Engine/VoiceChannel.h
@@ -10,7 +10,7 @@
 #include "VoiceChannel.generated.h"
 
 UCLASS(transient, customConstructor)
-class UVoiceChannel : public UChannel
+class ENGINE_API UVoiceChannel : public UChannel
 {
 	GENERATED_UCLASS_BODY()
 
@@ -63,5 +63,5 @@ public:
 	 *
 	 * @param VoicePacket the voice packet to send
 	 */
-	void AddVoicePacket(TSharedPtr<class FVoicePacket> VoicePacket);
+	virtual void AddVoicePacket(TSharedPtr<class FVoicePacket> VoicePacket);
 };
