Improved class layouts for camera modifiers and shakes

diff --git a/Engine/Source/Runtime/Engine/Classes/Camera/CameraModifier.h b/Engine/Source/Runtime/Engine/Classes/Camera/CameraModifier.h
--- a/Engine/Source/Runtime/Engine/Classes/Camera/CameraModifier.h
+++ b/Engine/Source/Runtime/Engine/Classes/Camera/CameraModifier.h
@@ -15,6 +15,15 @@ class ENGINE_API UCameraModifier : public UObject
 {
 	GENERATED_UCLASS_BODY()
 
+public:
+	/** If true, enables certain debug visualization features. */
+	UPROPERTY(EditAnywhere, Category = Debug)
+	uint32 bDebug : 1;
+
+	/** If true, no other modifiers of same priority allowed. */
+	UPROPERTY(EditDefaultsOnly, BlueprintReadOnly, Category = CameraModifier)
+	uint32 bExclusive : 1;
+
 protected:
 	/** If true, do not apply this modifier to the camera. */
 	uint32 bDisabled:1;
@@ -22,20 +31,16 @@ protected:
 	/** If true, this modifier will disable itself when finished interpolating out. */
 	uint32 bPendingDisable:1;
 
-	/** Camera this object is associated with. */
-	UPROPERTY(transient, BlueprintReadOnly, Category = CameraModifier)
-	class APlayerCameraManager* CameraOwner;
-
 public:
 	/** Priority value that determines the order in which modifiers are applied. 0 = highest priority, 255 = lowest. */
 	UPROPERTY(EditDefaultsOnly, BlueprintReadOnly, Category = CameraModifier)
-	uint8 Priority;
-
-	/** If true, no other modifiers of same priority allowed. */
-	UPROPERTY(EditDefaultsOnly, BlueprintReadOnly, Category = CameraModifier)
-	uint32 bExclusive:1;
-
+	uint8 Priority;	
+	
 protected:
+	/** Camera this object is associated with. */
+	UPROPERTY(transient, BlueprintReadOnly, Category = CameraModifier)
+	class APlayerCameraManager* CameraOwner;
+
 	/** When blending in, alpha proceeds from 0 to 1 over this time */
 	UPROPERTY(EditDefaultsOnly, BlueprintReadOnly, Category = CameraModifier)
 	float AlphaInTime;
@@ -48,11 +53,6 @@ protected:
 	UPROPERTY(transient, BlueprintReadOnly, Category = CameraModifier)
 	float Alpha;
 
-public:
-	/** If true, enables certain debug visualization features. */
-	UPROPERTY(EditAnywhere, Category=Debug)
-	uint32 bDebug :1;
-
 protected:
 	/** @return Returns the ideal blend alpha for this modifier. Interpolation will seek this value. */
 	virtual float GetTargetAlpha();
diff --git a/Engine/Source/Runtime/Engine/Classes/Camera/CameraShake.h b/Engine/Source/Runtime/Engine/Classes/Camera/CameraShake.h
--- a/Engine/Source/Runtime/Engine/Classes/Camera/CameraShake.h
+++ b/Engine/Source/Runtime/Engine/Classes/Camera/CameraShake.h
@@ -146,10 +146,6 @@ class ENGINE_API UCameraShake : public UObject
 	 * Parameters for defining CameraAnim-driven camera shakes
 	 ************************************************************/
 
-	/** Source camera animation to play. Can be null. */
-	UPROPERTY(EditAnywhere, Category=AnimShake)
-	class UCameraAnim* Anim;
-
 	/** Scalar defining how fast to play the anim. */
 	UPROPERTY(EditAnywhere, Category=AnimShake, meta=(ClampMin = "0.001"))
 	float AnimPlayRate;
@@ -166,40 +162,43 @@ class ENGINE_API UCameraShake : public UObject
 	UPROPERTY(EditAnywhere, Category=AnimShake, meta=(ClampMin = "0.0"))
 	float AnimBlendOutTime;
 
-	/**
-	 * If true, play a random snippet of the animation of length Duration.  Implies bLoop and bRandomStartTime = true for the CameraAnim.
-	 * If false, play the full anim once, non-looped. Useful for getting variety out of a single looped CameraAnim asset.
-	 */
-	UPROPERTY(EditAnywhere, Category=AnimShake)
-	uint32 bRandomAnimSegment:1;
-
 	/** When bRandomAnimSegment is true, this defines how long the anim should play. */
-	UPROPERTY(EditAnywhere, Category=AnimShake, meta=(ClampMin = "0.0", editcondition = "bRandomAnimSegment"))
+	UPROPERTY(EditAnywhere, Category = AnimShake, meta = (ClampMin = "0.0", editcondition = "bRandomAnimSegment"))
 	float RandomAnimSegmentDuration;
 
+	/** Source camera animation to play. Can be null. */
+	UPROPERTY(EditAnywhere, Category = AnimShake)
+	class UCameraAnim* Anim;
+
+	/**
+	* If true, play a random snippet of the animation of length Duration.  Implies bLoop and bRandomStartTime = true for the CameraAnim.
+	* If false, play the full anim once, non-looped. Useful for getting variety out of a single looped CameraAnim asset.
+	*/
+	UPROPERTY(EditAnywhere, Category = AnimShake)
+	uint32 bRandomAnimSegment : 1;
 
 protected:
 
 	// INSTANCE DATA
 	
-	UPROPERTY(transient, BlueprintReadOnly, Category=CameraShake)
-	class APlayerCameraManager* CameraOwner;
-
-	/** Time remaining for oscillation shakes. Less than 0.f means shake infinitely. */
-	float OscillatorTimeRemaining;
-
 	/** True if this shake is currently blending in. */
 	uint32 bBlendingIn:1;
 
-	/** How long this instance has been blending in. */
-	float CurrentBlendInTime;
-
 	/** True if this shake is currently blending out. */
 	uint32 bBlendingOut:1;
 
+	/** How long this instance has been blending in. */
+	float CurrentBlendInTime;
+
 	/** How long this instance has been blending out. */
 	float CurrentBlendOutTime;
 
+	UPROPERTY(transient, BlueprintReadOnly, Category = CameraShake)
+	class APlayerCameraManager* CameraOwner;
+
+	/** Time remaining for oscillation shakes. Less than 0.f means shake infinitely. */
+	float OscillatorTimeRemaining;
+
 	/** Current location sinusoidal offset. */
 	FVector LocSinOffset;
 
@@ -212,13 +211,13 @@ protected:
 	/** Overall intensity scale for this shake instance. */
 	float ShakeScale;
 
+	/** What space to play the shake in before applying to the camera.  Affects both Anim and Oscillation shakes. */
+	ECameraAnimPlaySpace::Type PlaySpace;
+
 	/** The playing instance of the CameraAnim-based shake, if any. */
 	UPROPERTY()
 	class UCameraAnimInst* AnimInst;
 
-	/** What space to play the shake in before applying to the camera.  Affects both Anim and Oscillation shakes. */
-	ECameraAnimPlaySpace::Type PlaySpace;
-
 	/** Matrix defining the playspace, used when PlaySpace == CAPS_UserDefined */
 	FMatrix UserPlaySpaceMatrix;
 
