#ue4 - Fix CapsuleComponent with CapsuleHalfHeight = CapsuleRadius creating collision geometry with slightly taller half height. This should simply create a sphere.

diff --git a/Engine/Source/Runtime/Engine/Private/Components/CapsuleComponent.cpp b/Engine/Source/Runtime/Engine/Private/Components/CapsuleComponent.cpp
--- a/Engine/Source/Runtime/Engine/Private/Components/CapsuleComponent.cpp
+++ b/Engine/Source/Runtime/Engine/Private/Components/CapsuleComponent.cpp
@@ -168,7 +168,7 @@ void UCapsuleComponent::UpdateBodySetup()
 	
 	SE->SetTransform(FTransform::Identity);
 	SE->Radius = CapsuleRadius;
-	SE->Length = 2 * FMath::Max(CapsuleHalfHeight - CapsuleRadius, 1.f);	//SphylElem uses height from center of capsule spheres, but UCapsuleComponent uses halfHeight from end of the sphere
+	SE->Length = 2 * FMath::Max(CapsuleHalfHeight - CapsuleRadius, 0.f);	//SphylElem uses height from center of capsule spheres, but UCapsuleComponent uses halfHeight from end of the sphere
 }
 
 bool UCapsuleComponent::IsZeroExtent() const
