#ue4 - Clean up a comment in FindBoxOpposingNormal() to clarify which normal is being considered.

diff --git a/Engine/Source/Runtime/Engine/Private/Collision/CollisionConversions.cpp b/Engine/Source/Runtime/Engine/Private/Collision/CollisionConversions.cpp
--- a/Engine/Source/Runtime/Engine/Private/Collision/CollisionConversions.cpp
+++ b/Engine/Source/Runtime/Engine/Private/Collision/CollisionConversions.cpp
@@ -129,7 +129,7 @@ static bool FindBoxOpposingNormal(const PxLocationHit& PHit, const FVector& Trac
 		// Select axis of face to compare to, based on normal.
 		if (ContactNormalLocal[i] > KINDA_SMALL_NUMBER)
 		{
-			const float TraceDotFaceNormal = TraceDirDenormLocal[i]; // TraceDirDenormLocal.dot(LocalNormal)
+			const float TraceDotFaceNormal = TraceDirDenormLocal[i]; // TraceDirDenormLocal.dot(BoxFaceNormal)
 			if (TraceDotFaceNormal < BestOpposingDot)
 			{
 				BestOpposingDot = TraceDotFaceNormal;
@@ -139,7 +139,7 @@ static bool FindBoxOpposingNormal(const PxLocationHit& PHit, const FVector& Trac
 		}
 		else if (ContactNormalLocal[i] < -KINDA_SMALL_NUMBER)
 		{
-			const float TraceDotFaceNormal = -TraceDirDenormLocal[i]; // TraceDirDenormLocal.dot(LocalNormal)
+			const float TraceDotFaceNormal = -TraceDirDenormLocal[i]; // TraceDirDenormLocal.dot(BoxFaceNormal)
 			if (TraceDotFaceNormal < BestOpposingDot)
 			{
 				BestOpposingDot = TraceDotFaceNormal;
