Fixed build issue fallout from CL# 2508681

diff --git a/Engine/Source/Runtime/Engine/Public/EdGraph/EdGraphNodeUtils.h b/Engine/Source/Runtime/Engine/Public/EdGraph/EdGraphNodeUtils.h
--- a/Engine/Source/Runtime/Engine/Public/EdGraph/EdGraphNodeUtils.h
+++ b/Engine/Source/Runtime/Engine/Public/EdGraph/EdGraphNodeUtils.h
@@ -102,7 +102,7 @@ public:
 private:
 	/** Helper to lookup from the GConfig the correct ini setting for displaying
 	 *  node and pin titles localized and updates the cache accordingly */
-	FORCEINLINE bool UpdateCacheIternal(FText const& InText) const
+	FORCEINLINE void UpdateCacheIternal(FText const& InText) const
 	{
 		bool bShowNodesAndPinsUnlocalized;
 		GConfig->GetBool( TEXT("Internationalization"), TEXT("ShowNodesAndPinsUnlocalized"), bShowNodesAndPinsUnlocalized, GEditorGameAgnosticIni );
