fixed update when changing r.Shadow.MaxResolution at runtime

diff --git a/Engine/Source/Runtime/Engine/Private/UnrealEngine.cpp b/Engine/Source/Runtime/Engine/Private/UnrealEngine.cpp
--- a/Engine/Source/Runtime/Engine/Private/UnrealEngine.cpp
+++ b/Engine/Source/Runtime/Engine/Private/UnrealEngine.cpp
@@ -321,6 +321,10 @@ void ScalabilityCVarsSinkCallback()
 				UMaterialInstance::AllMaterialsCacheResourceShadersForRendering();
 			}
 		}
+		else
+		{
+			GCachedScalabilityCVars = LocalScalabilityCVars;
+		}
 	}
 }
 
