#ue4 - Make sure Pawn unregisters itself upon EndPlay(), not just during Destroyed().

diff --git a/Engine/Source/Runtime/Engine/Classes/GameFramework/Pawn.h b/Engine/Source/Runtime/Engine/Classes/GameFramework/Pawn.h
--- a/Engine/Source/Runtime/Engine/Classes/GameFramework/Pawn.h
+++ b/Engine/Source/Runtime/Engine/Classes/GameFramework/Pawn.h
@@ -205,6 +205,7 @@ public:
 	virtual void GetActorEyesViewPoint( FVector& Location, FRotator& Rotation ) const override;
 	virtual void OutsideWorldBounds() override;
 	virtual void Destroyed() override;
+	virtual void EndPlay(const EEndPlayReason::Type EndPlayReason) override;
 	virtual void PreInitializeComponents() override;
 	virtual void PostInitializeComponents() override;
 	virtual const AActor* GetNetOwner() const override;
diff --git a/Engine/Source/Runtime/Engine/Private/Pawn.cpp b/Engine/Source/Runtime/Engine/Private/Pawn.cpp
--- a/Engine/Source/Runtime/Engine/Private/Pawn.cpp
+++ b/Engine/Source/Runtime/Engine/Private/Pawn.cpp
@@ -373,11 +373,22 @@ void APawn::PawnClientRestart()
 void APawn::Destroyed()
 {
 	DetachFromControllerPendingDestroy();
-
 	GetWorld()->RemovePawn( this );
 	Super::Destroyed();
 }
 
+void APawn::EndPlay(const EEndPlayReason::Type EndPlayReason)
+{
+	// Only do this once, to not be redundant with Destroyed().
+	if (EndPlayReason != EEndPlayReason::Destroyed)
+	{
+		DetachFromControllerPendingDestroy();
+		GetWorld()->RemovePawn( this );
+	}
+	
+	Super::EndPlay(EndPlayReason);
+}
+
 bool APawn::ShouldTakeDamage(float Damage, FDamageEvent const& DamageEvent, AController* EventInstigator, AActor* DamageCauser) const
 {
 	if ( (Role < ROLE_Authority) || !bCanBeDamaged || !GetWorld()->GetAuthGameMode() || (Damage == 0.f) )
@@ -867,7 +878,6 @@ void APawn::DetachFromControllerPendingDestroy()
 {
 	if ( Controller != NULL && Controller->GetPawn() == this )
 	{
-		AController* OldController = Controller;
 		Controller->PawnPendingDestroy(this);
 		if (Controller != NULL)
 		{
