Removing BT comments

diff --git a/Engine/Source/Runtime/Engine/Private/TextureDerivedData.cpp b/Engine/Source/Runtime/Engine/Private/TextureDerivedData.cpp
--- a/Engine/Source/Runtime/Engine/Private/TextureDerivedData.cpp
+++ b/Engine/Source/Runtime/Engine/Private/TextureDerivedData.cpp
@@ -1291,11 +1291,9 @@ void FTexturePlatformData::Serialize(FArchive& Ar, UTexture* Owner)
 	SerializePlatformData(Ar, this, Owner, bCooking, bStreamable);
 }
 
-// @BLACK_TUSK_CHANGE: BEGIN - [REDACTED_EMAIL] - Pass streamable flag for inlining mips
 void FTexturePlatformData::SerializeCooked(FArchive& Ar, UTexture* Owner, bool bStreamable)
 {
 	SerializePlatformData(Ar, this, Owner, true, bStreamable);
-	// @BLACK_TUSK_CHANGE: END - [REDACTED_EMAIL] 
 	if (Ar.IsLoading() && Mips.Num() > 0)
 	{
 		SizeX = Mips[0].SizeX;
