Deprecated GetRandomPointInRadius navigation functions family and introduced more precise GetRandomReachablePointInRadius and GetRandomPointInNavigableRadius #UE4

- also made sure GetRandomPointInNavigableRadius produced uniformly distributed values, even though current implementation is a bit naive.

diff --git a/Engine/Source/Runtime/Engine/Classes/AI/Navigation/AbstractNavData.h b/Engine/Source/Runtime/Engine/Classes/AI/Navigation/AbstractNavData.h
--- a/Engine/Source/Runtime/Engine/Classes/AI/Navigation/AbstractNavData.h
+++ b/Engine/Source/Runtime/Engine/Classes/AI/Navigation/AbstractNavData.h
@@ -41,7 +41,8 @@ class AAbstractNavData : public ANavigationData
 	// Begin ANavigationData overrides
 	virtual void BatchRaycast(TArray<FNavigationRaycastWork>& Workload, TSharedPtr<const FNavigationQueryFilter> QueryFilter, const UObject* Querier = NULL) const override {};
 	virtual FNavLocation GetRandomPoint(TSharedPtr<const FNavigationQueryFilter> Filter = NULL, const UObject* Querier = NULL) const override { return FNavLocation();  }
-	virtual bool GetRandomPointInRadius(const FVector& Origin, float Radius, FNavLocation& OutResult, TSharedPtr<const FNavigationQueryFilter> Filter = NULL, const UObject* Querier = NULL) const override { return false; }
+	virtual bool GetRandomReachablePointInRadius(const FVector& Origin, float Radius, FNavLocation& OutResult, TSharedPtr<const FNavigationQueryFilter> Filter = NULL, const UObject* Querier = NULL) const override { return false; }
+	virtual bool GetRandomPointInNavigableRadius(const FVector& Origin, float Radius, FNavLocation& OutResult, TSharedPtr<const FNavigationQueryFilter> Filter = NULL, const UObject* Querier = NULL) const override { return false; }
 	virtual bool ProjectPoint(const FVector& Point, FNavLocation& OutLocation, const FVector& Extent, TSharedPtr<const FNavigationQueryFilter> Filter = NULL, const UObject* Querier = NULL) const override { return false; }
 	virtual void BatchProjectPoints(TArray<FNavigationProjectionWork>& Workload, const FVector& Extent, TSharedPtr<const FNavigationQueryFilter> Filter = NULL, const UObject* Querier = NULL) const override {};
 	virtual ENavigationQueryResult::Type CalcPathCost(const FVector& PathStart, const FVector& PathEnd, float& OutPathCost, TSharedPtr<const FNavigationQueryFilter> QueryFilter = NULL, const UObject* Querier = NULL) const override { return ENavigationQueryResult::Invalid; }
diff --git a/Engine/Source/Runtime/Engine/Classes/AI/Navigation/NavigationData.h b/Engine/Source/Runtime/Engine/Classes/AI/Navigation/NavigationData.h
--- a/Engine/Source/Runtime/Engine/Classes/AI/Navigation/NavigationData.h
+++ b/Engine/Source/Runtime/Engine/Classes/AI/Navigation/NavigationData.h
@@ -692,8 +692,12 @@ public:
 
 	virtual FNavLocation GetRandomPoint(TSharedPtr<const FNavigationQueryFilter> Filter = NULL, const UObject* Querier = NULL) const PURE_VIRTUAL(ANavigationData::GetRandomPoint, return FNavLocation(););
 
-	virtual bool GetRandomPointInRadius(const FVector& Origin, float Radius, FNavLocation& OutResult, TSharedPtr<const FNavigationQueryFilter> Filter = NULL, const UObject* Querier = NULL) const PURE_VIRTUAL(ANavigationData::GetRandomPointInRadius, return false;);
+	/** finds a random location in Radius, reachable from Origin */
+	virtual bool GetRandomReachablePointInRadius(const FVector& Origin, float Radius, FNavLocation& OutResult, TSharedPtr<const FNavigationQueryFilter> Filter = NULL, const UObject* Querier = NULL) const PURE_VIRTUAL(ANavigationData::GetRandomReachablePointInRadius, return false;);
 
+	/** finds a random location in navigable space, in given Radius */
+	virtual bool GetRandomPointInNavigableRadius(const FVector& Origin, float Radius, FNavLocation& OutResult, TSharedPtr<const FNavigationQueryFilter> Filter = NULL, const UObject* Querier = NULL) const PURE_VIRTUAL(ANavigationData::GetRandomPointInNavigableRadius, return false;);
+	
 	/**	Tries to project given Point to this navigation type, within given Extent.
 	 *	@param OutLocation if successful this variable will be filed with result
 	 *	@return true if successful, false otherwise
@@ -853,6 +857,10 @@ private:
 	uint16 NavDataUniqueID;
 
 	static uint16 GetNextUniqueID();
+
+public:
+	DEPRECATED(4.8, "GetRandomPointInRadius is deprecated, please use GetRandomReachablePointInRadius")
+	virtual bool GetRandomPointInRadius(const FVector& Origin, float Radius, FNavLocation& OutResult, TSharedPtr<const FNavigationQueryFilter> Filter = NULL, const UObject* Querier = NULL) const PURE_VIRTUAL(ANavigationData::GetRandomPointInRadius, return false;);
 };
 
 struct FAsyncPathFindingQuery : public FPathFindingQuery
diff --git a/Engine/Source/Runtime/Engine/Classes/AI/Navigation/NavigationSystem.h b/Engine/Source/Runtime/Engine/Classes/AI/Navigation/NavigationSystem.h
--- a/Engine/Source/Runtime/Engine/Classes/AI/Navigation/NavigationSystem.h
+++ b/Engine/Source/Runtime/Engine/Classes/AI/Navigation/NavigationSystem.h
@@ -245,10 +245,13 @@ public:
 	
 	UFUNCTION(BlueprintPure, Category="AI|Navigation", meta=(WorldContext="WorldContext" ) )
 	static FVector GetRandomPoint(UObject* WorldContext, ANavigationData* NavData = NULL, TSubclassOf<UNavigationQueryFilter> FilterClass = NULL);
+	
+	UFUNCTION(BlueprintPure, Category = "AI|Navigation", meta = (WorldContext = "WorldContext"))
+	static FVector GetRandomReachablePointInRadius(UObject* WorldContext, const FVector& Origin, float Radius, ANavigationData* NavData = NULL, TSubclassOf<UNavigationQueryFilter> FilterClass = NULL);
 
-	UFUNCTION(BlueprintPure, Category="AI|Navigation", meta=(WorldContext="WorldContext" ) )
-	static FVector GetRandomPointInRadius(UObject* WorldContext, const FVector& Origin, float Radius, ANavigationData* NavData = NULL, TSubclassOf<UNavigationQueryFilter> FilterClass = NULL);
-
+	UFUNCTION(BlueprintPure, Category = "AI|Navigation", meta = (WorldContext = "WorldContext"))
+	static FVector GetRandomPointInNavigableRadius(UObject* WorldContext, const FVector& Origin, float Radius, ANavigationData* NavData = NULL, TSubclassOf<UNavigationQueryFilter> FilterClass = NULL);
+	
 	/** Potentially expensive. Use with caution. Consider using UPathFollowingComponent::GetRemainingPathCost instead */
 	UFUNCTION(BlueprintPure, Category="AI|Navigation", meta=(WorldContext="WorldContext" ) )
 	static ENavigationQueryResult::Type GetPathCost(UObject* WorldContext, const FVector& PathStart, const FVector& PathEnd, float& PathCost, ANavigationData* NavData = NULL, TSubclassOf<UNavigationQueryFilter> FilterClass = NULL);
@@ -413,12 +416,18 @@ public:
 	 *	@return true if any location found, false otherwise */
 	bool GetRandomPoint(FNavLocation& ResultLocation, ANavigationData* NavData = NULL, TSharedPtr<const FNavigationQueryFilter> QueryFilter = NULL);
 
-	/** Finds random point in navigable space restricted to Radius around Origin
+	/** Finds random, reachable point in navigable space restricted to Radius around Origin
 	 *	@param ResultLocation Found point is put here
 	 *	@param NavData If NavData == NULL then MainNavData is used.
 	 *	@return true if any location found, false otherwise */
-	bool GetRandomPointInRadius(const FVector& Origin, float Radius, FNavLocation& ResultLocation, ANavigationData* NavData = NULL, TSharedPtr<const FNavigationQueryFilter> QueryFilter = NULL) const;
+	bool GetRandomReachablePointInRadius(const FVector& Origin, float Radius, FNavLocation& ResultLocation, ANavigationData* NavData = NULL, TSharedPtr<const FNavigationQueryFilter> QueryFilter = NULL) const;
 
+	/** Finds random, point in navigable space restricted to Radius around Origin. Resulting location is not tested for reachability from the Origin
+	 *	@param ResultLocation Found point is put here
+	 *	@param NavData If NavData == NULL then MainNavData is used.
+	 *	@return true if any location found, false otherwise */
+	bool GetRandomPointInNavigableRadius(const FVector& Origin, float Radius, FNavLocation& ResultLocation, ANavigationData* NavData = NULL, TSharedPtr<const FNavigationQueryFilter> QueryFilter = NULL) const;
+	
 	/** Calculates a path from PathStart to PathEnd and retrieves its cost. 
 	 *	@NOTE potentially expensive, so use it with caution */
 	ENavigationQueryResult::Type GetPathCost(const FVector& PathStart, const FVector& PathEnd, float& PathCost, const ANavigationData* NavData = NULL, TSharedPtr<const FNavigationQueryFilter> QueryFilter = NULL) const;
@@ -898,5 +907,10 @@ public:
 	void RegisterCustomLink(INavLinkCustomInterface* CustomLink);
 	DEPRECATED(4.8, "This version is deprecated. Please use the one taking reference to INavLinkCustomInterface rather than a pointer instead.")
 	void UnregisterCustomLink(INavLinkCustomInterface* CustomLink);
+	DEPRECATED(4.8, "GetRandomPointInRadius is deprecated. Use either GetRandomReachablePointInRadius or GetRandomPointInNavigableSpace")
+	UFUNCTION(BlueprintPure, Category = "AI|Navigation", meta = (WorldContext = "WorldContext"))
+	static FVector GetRandomPointInRadius(UObject* WorldContext, const FVector& Origin, float Radius, ANavigationData* NavData = NULL, TSubclassOf<UNavigationQueryFilter> FilterClass = NULL);
+	DEPRECATED(4.8, "GetRandomPointInRadius is deprecated. Use either GetRandomReachablePointInRadius or GetRandomPointInNavigableSpace")
+	bool GetRandomPointInRadius(const FVector& Origin, float Radius, FNavLocation& ResultLocation, ANavigationData* NavData = NULL, TSharedPtr<const FNavigationQueryFilter> QueryFilter = NULL) const;
 };
 
diff --git a/Engine/Source/Runtime/Engine/Classes/AI/Navigation/NavigationTypes.h b/Engine/Source/Runtime/Engine/Classes/AI/Navigation/NavigationTypes.h
--- a/Engine/Source/Runtime/Engine/Classes/AI/Navigation/NavigationTypes.h
+++ b/Engine/Source/Runtime/Engine/Classes/AI/Navigation/NavigationTypes.h
@@ -231,6 +231,9 @@ struct FNavLocation
 	FNavLocation() : Location(FVector::ZeroVector), NodeRef(INVALID_NAVNODEREF) {}
 	explicit FNavLocation(const FVector& InLocation, NavNodeRef InNodeRef = INVALID_NAVNODEREF) 
 		: Location(InLocation), NodeRef(InNodeRef) {}
+
+	/** checks if location has associated navigation node ref */
+	FORCEINLINE bool HasNodeRef() const { return NodeRef != INVALID_NAVNODEREF; }
 };
 
 /** Describes node in navigation path */
diff --git a/Engine/Source/Runtime/Engine/Classes/AI/Navigation/RecastNavMesh.h b/Engine/Source/Runtime/Engine/Classes/AI/Navigation/RecastNavMesh.h
--- a/Engine/Source/Runtime/Engine/Classes/AI/Navigation/RecastNavMesh.h
+++ b/Engine/Source/Runtime/Engine/Classes/AI/Navigation/RecastNavMesh.h
@@ -677,7 +677,8 @@ public:
 
 	// Begin ANavigationData Interface
 	virtual FNavLocation GetRandomPoint(TSharedPtr<const FNavigationQueryFilter> Filter = NULL, const UObject* Querier = NULL) const override;
-	virtual bool GetRandomPointInRadius(const FVector& Origin, float Radius, FNavLocation& OutResult, TSharedPtr<const FNavigationQueryFilter> Filter = NULL, const UObject* Querier = NULL) const override;
+	virtual bool GetRandomReachablePointInRadius(const FVector& Origin, float Radius, FNavLocation& OutResult, TSharedPtr<const FNavigationQueryFilter> Filter = NULL, const UObject* Querier = NULL) const override;
+	virtual bool GetRandomPointInNavigableRadius(const FVector& Origin, float Radius, FNavLocation& OutResult, TSharedPtr<const FNavigationQueryFilter> Filter = NULL, const UObject* Querier = NULL) const override;
 
 	virtual bool ProjectPoint(const FVector& Point, FNavLocation& OutLocation, const FVector& Extent, TSharedPtr<const FNavigationQueryFilter> Filter = NULL, const UObject* Querier = NULL) const override;
 	
diff --git a/Engine/Source/Runtime/Engine/Private/AI/Navigation/NavigationSystem.cpp b/Engine/Source/Runtime/Engine/Private/AI/Navigation/NavigationSystem.cpp
--- a/Engine/Source/Runtime/Engine/Private/AI/Navigation/NavigationSystem.cpp
+++ b/Engine/Source/Runtime/Engine/Private/AI/Navigation/NavigationSystem.cpp
@@ -1026,7 +1026,7 @@ bool UNavigationSystem::GetRandomPoint(FNavLocation& ResultLocation, ANavigation
 	return false;
 }
 
-bool UNavigationSystem::GetRandomPointInRadius(const FVector& Origin, float Radius, FNavLocation& ResultLocation, ANavigationData* NavData, TSharedPtr<const FNavigationQueryFilter> QueryFilter) const
+bool UNavigationSystem::GetRandomReachablePointInRadius(const FVector& Origin, float Radius, FNavLocation& ResultLocation, ANavigationData* NavData, TSharedPtr<const FNavigationQueryFilter> QueryFilter) const
 {
 	SCOPE_CYCLE_COUNTER(STAT_Navigation_QueriesTimeSync);
 
@@ -1035,7 +1035,19 @@ bool UNavigationSystem::GetRandomPointInRadius(const FVector& Origin, float Radi
 		NavData = MainNavData;
 	}
 
-	return NavData != nullptr && NavData->GetRandomPointInRadius(Origin, Radius, ResultLocation, QueryFilter);
+	return NavData != nullptr && NavData->GetRandomReachablePointInRadius(Origin, Radius, ResultLocation, QueryFilter);
+}
+
+bool UNavigationSystem::GetRandomPointInNavigableRadius(const FVector& Origin, float Radius, FNavLocation& ResultLocation, ANavigationData* NavData, TSharedPtr<const FNavigationQueryFilter> QueryFilter) const
+{
+	SCOPE_CYCLE_COUNTER(STAT_Navigation_QueriesTimeSync);
+
+	if (NavData == nullptr)
+	{
+		NavData = MainNavData;
+	}
+
+	return NavData != nullptr && NavData->GetRandomPointInNavigableRadius(Origin, Radius, ResultLocation, QueryFilter);
 }
 
 ENavigationQueryResult::Type UNavigationSystem::GetPathCost(const FVector& PathStart, const FVector& PathEnd, float& OutPathCost, const ANavigationData* NavData, TSharedPtr<const FNavigationQueryFilter> QueryFilter) const
@@ -3385,19 +3397,36 @@ FVector UNavigationSystem::GetRandomPoint(UObject* WorldContextObject, ANavigati
 	return RandomPoint.Location;
 }
 
+FVector UNavigationSystem::GetRandomReachablePointInRadius(UObject* WorldContext, const FVector& Origin, float Radius, ANavigationData* NavData, TSubclassOf<UNavigationQueryFilter> FilterClass)
+{
+	FNavLocation RandomPoint;
 
-FVector UNavigationSystem::GetRandomPointInRadius(UObject* WorldContextObject, const FVector& Origin, float Radius, ANavigationData* NavData, TSubclassOf<UNavigationQueryFilter> FilterClass)
+	UWorld* World = GEngine->GetWorldFromContextObject(WorldContext);
+	UNavigationSystem* NavSys = UNavigationSystem::GetCurrent(World);
+	if (NavSys)
+	{
+		ANavigationData* UseNavData = NavData ? NavData : NavSys->GetMainNavData(FNavigationSystem::DontCreate);
+		if (UseNavData)
+		{
+			NavSys->GetRandomReachablePointInRadius(Origin, Radius, RandomPoint, UseNavData, UNavigationQueryFilter::GetQueryFilter(*UseNavData, FilterClass));
+		}
+	}
+
+	return RandomPoint.Location;
+}
+
+FVector UNavigationSystem::GetRandomPointInNavigableRadius(UObject* WorldContext, const FVector& Origin, float Radius, ANavigationData* NavData, TSubclassOf<UNavigationQueryFilter> FilterClass)
 {
 	FNavLocation RandomPoint;
 
-	UWorld* World = GEngine->GetWorldFromContextObject( WorldContextObject );
+	UWorld* World = GEngine->GetWorldFromContextObject(WorldContext);
 	UNavigationSystem* NavSys = UNavigationSystem::GetCurrent(World);
 	if (NavSys)
 	{
 		ANavigationData* UseNavData = NavData ? NavData : NavSys->GetMainNavData(FNavigationSystem::DontCreate);
 		if (UseNavData)
 		{
-			NavSys->GetRandomPointInRadius(Origin, Radius, RandomPoint, UseNavData, UNavigationQueryFilter::GetQueryFilter(*UseNavData, FilterClass));
+			NavSys->GetRandomPointInNavigableRadius(Origin, Radius, RandomPoint, UseNavData, UNavigationQueryFilter::GetQueryFilter(*UseNavData, FilterClass));
 		}
 	}
 
@@ -3762,3 +3791,13 @@ void UNavigationSystem::UnregisterCustomLink(INavLinkCustomInterface* CustomLink
 		UnregisterCustomLink(*CustomLink);
 	}
 }
+
+FVector UNavigationSystem::GetRandomPointInRadius(UObject* WorldContextObject, const FVector& Origin, float Radius, ANavigationData* NavData, TSubclassOf<UNavigationQueryFilter> FilterClass)
+{
+	return UNavigationSystem::GetRandomReachablePointInRadius(WorldContextObject, Origin, Radius, NavData, FilterClass);
+}
+
+bool UNavigationSystem::GetRandomPointInRadius(const FVector& Origin, float Radius, FNavLocation& ResultLocation, ANavigationData* NavData, TSharedPtr<const FNavigationQueryFilter> QueryFilter) const
+{
+	return GetRandomReachablePointInRadius(Origin, Radius, ResultLocation, NavData, QueryFilter);
+}
\ No newline at end of file
diff --git a/Engine/Source/Runtime/Engine/Private/AI/Navigation/RecastNavMesh.cpp b/Engine/Source/Runtime/Engine/Private/AI/Navigation/RecastNavMesh.cpp
--- a/Engine/Source/Runtime/Engine/Private/AI/Navigation/RecastNavMesh.cpp
+++ b/Engine/Source/Runtime/Engine/Private/AI/Navigation/RecastNavMesh.cpp
@@ -928,15 +928,72 @@ FNavLocation ARecastNavMesh::GetRandomPoint(TSharedPtr<const FNavigationQueryFil
 	return RandomPt;
 }
 
-bool ARecastNavMesh::GetRandomPointInRadius(const FVector& Origin, float Radius, FNavLocation& OutResult, TSharedPtr<const FNavigationQueryFilter> Filter, const UObject* QueryOwner) const
+bool ARecastNavMesh::GetRandomReachablePointInRadius(const FVector& Origin, float Radius, FNavLocation& OutResult, TSharedPtr<const FNavigationQueryFilter> Filter, const UObject* QueryOwner) const
 {
-	bool bSuccess = false;
-	if (RecastNavMeshImpl)
+	if (RecastNavMeshImpl == nullptr || RecastNavMeshImpl->DetourNavMesh == nullptr || Radius <= 0.f)
 	{
-		bSuccess = RecastNavMeshImpl->GetRandomPointInRadius(Origin, Radius, OutResult, GetRightFilterRef(Filter), QueryOwner);
+		return false;
 	}
 
-	return bSuccess;
+	const FNavigationQueryFilter& FilterInstance = GetRightFilterRef(Filter);
+
+	FRecastSpeciaLinkFilter LinkFilter(UNavigationSystem::GetCurrent(GetWorld()), QueryOwner);
+	INITIALIZE_NAVQUERY_WLINKFILTER(NavQuery, FilterInstance.GetMaxSearchNodes(), LinkFilter);
+
+	// inits to "pass all"
+	const dtQueryFilter* QueryFilter = (static_cast<const FRecastQueryFilter*>(FilterInstance.GetImplementation()))->GetAsDetourQueryFilter();
+	ensure(QueryFilter);
+	if (QueryFilter)
+	{
+		// find starting poly
+		// convert start/end pos to Recast coords
+		const float Extent[3] = { Radius, Radius, Radius };
+		float RecastOrigin[3];
+		Unr2RecastVector(Origin, RecastOrigin);
+		NavNodeRef OriginPolyID = INVALID_NAVNODEREF;
+		NavQuery.findNearestPoly(RecastOrigin, Extent, QueryFilter, &OriginPolyID, nullptr);
+
+		dtPolyRef Poly;
+		float RandPt[3];
+		dtStatus Status = NavQuery.findRandomPointAroundCircle(OriginPolyID, RecastOrigin, Radius
+			, QueryFilter, FMath::FRand, &Poly, RandPt);
+
+		if (dtStatusSucceed(Status))
+		{
+			OutResult = FNavLocation(Recast2UnrealPoint(RandPt), Poly);
+			return true;
+		}
+		else
+		{
+			OutResult = FNavLocation(Origin, OriginPolyID);
+		}
+	}
+
+	return false;
+}
+
+bool ARecastNavMesh::GetRandomPointInNavigableRadius(const FVector& Origin, float Radius, FNavLocation& OutResult, TSharedPtr<const FNavigationQueryFilter> Filter, const UObject* Querier) const
+{
+	const FVector ProjectionExtent(NavDataConfig.DefaultQueryExtent.X, NavDataConfig.DefaultQueryExtent.Y, BIG_NUMBER);
+	OutResult = FNavLocation(FNavigationSystem::InvalidLocation);
+	
+	// this is super naive implementation for now. We give it 10 tries, and fail if it's not enough. 
+	// The proper solution would involve processing nearby&in-radius tiles in batches until the whole radius of the query is exhausted
+	static const int32 IterationsLimit = 10;
+	int32 Interation = 0;
+	do 
+	{
+		const float RandomAngle = 2.f * PI * FMath::FRand();
+		const float U = FMath::FRand() + FMath::FRand();
+		const float RandomRadius = Radius * (U > 1 ? 2.f - U : U);
+		const FVector RandomOffset(FMath::Cos(RandomAngle) * RandomRadius, FMath::Sin(RandomAngle) * RandomRadius, 0);
+		FVector RandomLocationInRadius = Origin + RandomOffset;
+
+		// naive implementation 
+		ProjectPoint(RandomLocationInRadius, OutResult, ProjectionExtent, Filter);
+	} while (OutResult.HasNodeRef() == false && ++Interation < IterationsLimit);
+
+	return OutResult.HasNodeRef() == true;
 }
 
 bool ARecastNavMesh::GetRandomPointInCluster(NavNodeRef ClusterRef, FNavLocation& OutLocation) const
diff --git a/Engine/Source/Runtime/Engine/Public/AI/Navigation/PImplRecastNavMesh.h b/Engine/Source/Runtime/Engine/Public/AI/Navigation/PImplRecastNavMesh.h
--- a/Engine/Source/Runtime/Engine/Public/AI/Navigation/PImplRecastNavMesh.h
+++ b/Engine/Source/Runtime/Engine/Public/AI/Navigation/PImplRecastNavMesh.h
@@ -123,6 +123,7 @@ public:
 
 	/** Returns a random location on the navmesh within Radius from Origin. 
 	 *	@return false if no valid navigable location available in specified area */
+	DEPRECATED(4.8, "GetRandomPointInRadius is deprecated, please use ANavigationData::GetRandomReachablePointInRadius")
 	bool GetRandomPointInRadius(const FVector& Origin, float Radius, FNavLocation& OutLocation, const FNavigationQueryFilter& Filter, const UObject* Owner) const;
 
 	/** Returns a random location on the navmesh within cluster */
