Made FloatingPawnMovement component work for AI just like for Player when AI is not following a path #UE4

- to achieve that I've expanded NavAgentInterface with "IsFollowingAPath" function

diff --git a/Engine/Source/Runtime/AIModule/Classes/AIController.h b/Engine/Source/Runtime/AIModule/Classes/AIController.h
--- a/Engine/Source/Runtime/AIModule/Classes/AIController.h
+++ b/Engine/Source/Runtime/AIModule/Classes/AIController.h
@@ -335,6 +335,11 @@ public:
 	//----------------------------------------------------------------------//
 	virtual UAIPerceptionComponent* GetPerceptionComponent() override { return GetAIPerceptionComponent(); }
 
+	//----------------------------------------------------------------------//
+	// INavAgentInterface
+	//----------------------------------------------------------------------//
+	virtual bool IsFollowingAPath() const override;
+
 	//----------------------------------------------------------------------//
 	// Actions
 	//----------------------------------------------------------------------//
diff --git a/Engine/Source/Runtime/AIModule/Classes/Navigation/PathFollowingComponent.h b/Engine/Source/Runtime/AIModule/Classes/Navigation/PathFollowingComponent.h
--- a/Engine/Source/Runtime/AIModule/Classes/Navigation/PathFollowingComponent.h
+++ b/Engine/Source/Runtime/AIModule/Classes/Navigation/PathFollowingComponent.h
@@ -247,6 +247,7 @@ class AIMODULE_API UPathFollowingComponent : public UActorComponent, public IAIR
 	FVector GetPathDestination() const;
 
 	FORCEINLINE const FNavPathSharedPtr GetPath() const { return Path; }
+	FORCEINLINE bool HasValidPath() const { return Path.IsValid() && Path->IsValid(); }
 	bool HasDirectPath() const;
 
 	/** readable name of current status */
diff --git a/Engine/Source/Runtime/AIModule/Private/AIController.cpp b/Engine/Source/Runtime/AIModule/Private/AIController.cpp
--- a/Engine/Source/Runtime/AIModule/Private/AIController.cpp
+++ b/Engine/Source/Runtime/AIModule/Private/AIController.cpp
@@ -757,6 +757,11 @@ bool AAIController::HasPartialPath() const
 	return (PathFollowingComponent != NULL) && (PathFollowingComponent->HasPartialPath());
 }
 
+bool AAIController::IsFollowingAPath() const
+{
+	return (PathFollowingComponent != nullptr) && PathFollowingComponent->HasValidPath();
+}
+
 FVector AAIController::GetImmediateMoveDestination() const
 {
 	return (PathFollowingComponent) ? PathFollowingComponent->GetCurrentTargetLocation() : FVector::ZeroVector;
diff --git a/Engine/Source/Runtime/Engine/Classes/AI/Navigation/NavAgentInterface.h b/Engine/Source/Runtime/Engine/Classes/AI/Navigation/NavAgentInterface.h
--- a/Engine/Source/Runtime/Engine/Classes/AI/Navigation/NavAgentInterface.h
+++ b/Engine/Source/Runtime/Engine/Classes/AI/Navigation/NavAgentInterface.h
@@ -44,6 +44,9 @@ class INavAgentInterface
 	/** Allows delaying repath requests */
 	virtual bool ShouldPostponePathUpdates() const { return false; }
 
+	/** Checks if the agent is actively following a navigation path */
+	virtual bool IsFollowingAPath() const { return false; }
+
 	//----------------------------------------------------------------------//
 	// DEPRECATED
 	//----------------------------------------------------------------------//
diff --git a/Engine/Source/Runtime/Engine/Classes/GameFramework/Controller.h b/Engine/Source/Runtime/Engine/Classes/GameFramework/Controller.h
--- a/Engine/Source/Runtime/Engine/Classes/GameFramework/Controller.h
+++ b/Engine/Source/Runtime/Engine/Classes/GameFramework/Controller.h
@@ -272,6 +272,7 @@ public:
 	virtual FVector GetNavAgentLocation() const override;
 	virtual void GetMoveGoalReachTest(class AActor* MovingActor, const FVector& MoveOffset, FVector& GoalOffset, float& GoalRadius, float& GoalHalfHeight) const override;
 	virtual bool ShouldPostponePathUpdates() const override;
+	virtual bool IsFollowingAPath() const override;
 	// End INavAgentInterface Interface
 
 	/** prepares path following component */
diff --git a/Engine/Source/Runtime/Engine/Private/Controller.cpp b/Engine/Source/Runtime/Engine/Private/Controller.cpp
--- a/Engine/Source/Runtime/Engine/Private/Controller.cpp
+++ b/Engine/Source/Runtime/Engine/Private/Controller.cpp
@@ -562,6 +562,12 @@ bool AController::ShouldPostponePathUpdates() const
 	return Pawn ? Pawn->ShouldPostponePathUpdates() : false;
 }
 
+bool AController::IsFollowingAPath() const
+{
+	UPathFollowingComponent* PathFollowingComp = FindComponentByClass<UPathFollowingComponent>();
+	return (PathFollowingComp != nullptr) && PathFollowingComp->HasValidPath();
+}
+
 void AController::UpdateNavigationComponents()
 {
 	UPathFollowingComponent* PathFollowingComp = FindComponentByClass<UPathFollowingComponent>();
diff --git a/Engine/Source/Runtime/Engine/Private/FloatingPawnMovement.cpp b/Engine/Source/Runtime/Engine/Private/FloatingPawnMovement.cpp
--- a/Engine/Source/Runtime/Engine/Private/FloatingPawnMovement.cpp
+++ b/Engine/Source/Runtime/Engine/Private/FloatingPawnMovement.cpp
@@ -15,7 +15,6 @@ UFloatingPawnMovement::UFloatingPawnMovement(const FObjectInitializer& ObjectIni
 	ResetMoveState();
 }
 
-
 void UFloatingPawnMovement::TickComponent(float DeltaTime, enum ELevelTick TickType, FActorComponentTickFunction *ThisTickFunction)
 {
 	if (ShouldSkipUpdate(DeltaTime))
@@ -33,13 +32,14 @@ void UFloatingPawnMovement::TickComponent(float DeltaTime, enum ELevelTick TickT
 	const AController* Controller = PawnOwner->GetController();
 	if (Controller && Controller->IsLocalController())
 	{
-		if (Controller->IsLocalPlayerController())
+		// apply input for local players but also for AI that's not following a navigation path at the moment
+		if (Controller->IsLocalPlayerController() == true || Controller->IsFollowingAPath() == false)
 		{
 			ApplyControlInputToVelocity(DeltaTime);
 		}
 		// if it's not player controller, but we do have a controller, then it's AI
-		// and we need to limit the speed
-		else if (IsExceedingMaxSpeed(MaxSpeed))
+		// (that's not following a path) and we need to limit the speed
+		else if (IsExceedingMaxSpeed(MaxSpeed) == false)
 		{
 			Velocity = Velocity.GetUnsafeNormal() * MaxSpeed;
 		}
@@ -79,7 +79,6 @@ void UFloatingPawnMovement::TickComponent(float DeltaTime, enum ELevelTick TickT
 	}
 };
 
-
 bool UFloatingPawnMovement::LimitWorldBounds()
 {
 	AWorldSettings* WorldSettings = PawnOwner ? PawnOwner->GetWorldSettings() : NULL;
@@ -98,7 +97,6 @@ bool UFloatingPawnMovement::LimitWorldBounds()
 	return false;
 }
 
-
 void UFloatingPawnMovement::ApplyControlInputToVelocity(float DeltaTime)
 {
 	const FVector ControlAcceleration = GetPendingInputVector().GetClampedToMaxSize(1.f);
@@ -142,7 +140,6 @@ void UFloatingPawnMovement::ApplyControlInputToVelocity(float DeltaTime)
 	ConsumeInputVector();
 }
 
-
 bool UFloatingPawnMovement::ResolvePenetration(const FVector& Adjustment, const FHitResult& Hit, const FRotator& NewRotation)
 {
 	bPositionCorrected |= Super::ResolvePenetration(Adjustment, Hit, NewRotation);
