Non-unity CIS fix #UE4

diff --git a/Engine/Source/Runtime/Engine/Private/AI/Navigation/RecastNavMesh.cpp b/Engine/Source/Runtime/Engine/Private/AI/Navigation/RecastNavMesh.cpp
--- a/Engine/Source/Runtime/Engine/Private/AI/Navigation/RecastNavMesh.cpp
+++ b/Engine/Source/Runtime/Engine/Private/AI/Navigation/RecastNavMesh.cpp
@@ -948,14 +948,13 @@ bool ARecastNavMesh::GetRandomReachablePointInRadius(const FVector& Origin, floa
 		// find starting poly
 		// convert start/end pos to Recast coords
 		const float Extent[3] = { Radius, Radius, Radius };
-		float RecastOrigin[3];
-		Unr2RecastVector(Origin, RecastOrigin);
+		const FVector RecastOrigin = Unreal2RecastPoint(Origin);
 		NavNodeRef OriginPolyID = INVALID_NAVNODEREF;
-		NavQuery.findNearestPoly(RecastOrigin, Extent, QueryFilter, &OriginPolyID, nullptr);
+		NavQuery.findNearestPoly(&RecastOrigin.X, Extent, QueryFilter, &OriginPolyID, nullptr);
 
 		dtPolyRef Poly;
 		float RandPt[3];
-		dtStatus Status = NavQuery.findRandomPointAroundCircle(OriginPolyID, RecastOrigin, Radius
+		dtStatus Status = NavQuery.findRandomPointAroundCircle(OriginPolyID, &RecastOrigin.X, Radius
 			, QueryFilter, FMath::FRand, &Poly, RandPt);
 
 		if (dtStatusSucceed(Status))
