This change modifies timeline components to reflect the timeline node name from the blueprint.
#UE-12450 - Timelines are displayed as [Unamed TimelineComponent] in the Details view
#Branch UE4
#Proj Kismet, Engine

diff --git a/Engine/Source/Editor/Kismet/Private/SSCSEditor.cpp b/Engine/Source/Editor/Kismet/Private/SSCSEditor.cpp
--- a/Engine/Source/Editor/Kismet/Private/SSCSEditor.cpp
+++ b/Engine/Source/Editor/Kismet/Private/SSCSEditor.cpp
@@ -774,7 +774,7 @@ FString FSCSEditorTreeNodeComponentBase::GetDisplayString() const
 	{
 		return VariableName.ToString();
 	}
-	else if (IsNative() && (ComponentTemplate != NULL))
+	else if ( ComponentTemplate != nullptr )
 	{
 		return ComponentTemplate->GetFName().ToString();
 	}
diff --git a/Engine/Source/Runtime/Engine/Private/BlueprintGeneratedClass.cpp b/Engine/Source/Runtime/Engine/Private/BlueprintGeneratedClass.cpp
--- a/Engine/Source/Runtime/Engine/Private/BlueprintGeneratedClass.cpp
+++ b/Engine/Source/Runtime/Engine/Private/BlueprintGeneratedClass.cpp
@@ -444,7 +444,7 @@ void UBlueprintGeneratedClass::CreateComponentsForActor(AActor* Actor) const
 			continue;
 		}
 
-		FName NewName = *(FString::Printf(TEXT("TimelineComp__%d"), Actor->BlueprintCreatedComponents.Num() ) );
+		FName NewName( *UTimelineTemplate::TimelineTemplateNameToVariableName( TimelineTemplate->GetFName() ));
 		UTimelineComponent* NewTimeline = NewObject<UTimelineComponent>(Actor, NewName);
 		NewTimeline->CreationMethod = EComponentCreationMethod::UserConstructionScript; // Indicate it comes from a blueprint so it gets cleared when we rerun construction scripts
 		Actor->BlueprintCreatedComponents.Add(NewTimeline); // Add to array so it gets saved
