Fixed PrimitiveComponent not notifying navigation system properly when bCanEverAffectNavigation value is being changed in the properties view #UE4

addresses UE-12906

diff --git a/Engine/Source/Runtime/Engine/Classes/Components/PrimitiveComponent.h b/Engine/Source/Runtime/Engine/Classes/Components/PrimitiveComponent.h
--- a/Engine/Source/Runtime/Engine/Classes/Components/PrimitiveComponent.h
+++ b/Engine/Source/Runtime/Engine/Classes/Components/PrimitiveComponent.h
@@ -1683,7 +1683,12 @@ public:
 
 	/** set value of bCanEverAffectNavigation flag and update navigation octree if needed */
 	void SetCanEverAffectNavigation(bool bRelevant);
+	
+protected:
+	/** Makes sure navigation system has up to date information regarding component's navigation relevancy and if it can affect navigation at all */
+	void HandleCanEverAffectNavigationChange();
 
+public:
 	DEPRECATED(4.5, "UPrimitiveComponent::DisableNavigationRelevance() is deprecated, use SetCanEverAffectNavigation() instead.")
 	void DisableNavigationRelevance()
 	{
diff --git a/Engine/Source/Runtime/Engine/Private/Components/PrimitiveComponent.cpp b/Engine/Source/Runtime/Engine/Private/Components/PrimitiveComponent.cpp
--- a/Engine/Source/Runtime/Engine/Private/Components/PrimitiveComponent.cpp
+++ b/Engine/Source/Runtime/Engine/Private/Components/PrimitiveComponent.cpp
@@ -585,7 +585,7 @@ void UPrimitiveComponent::PostEditChangeProperty(FPropertyChangedEvent& Property
 	// so it needs to be called directly here
 	if (PropertyThatChanged && PropertyThatChanged->GetFName() == GET_MEMBER_NAME_CHECKED(UPrimitiveComponent, bCanEverAffectNavigation))
 	{
-		UNavigationSystem::UpdateNavOctree(this);
+		HandleCanEverAffectNavigationChange();
 	}
 }
 
@@ -1693,18 +1693,23 @@ void UPrimitiveComponent::SetCanEverAffectNavigation(bool bRelevant)
 	{
 		bCanEverAffectNavigation = bRelevant;
 
-		// update octree if already registered
-		if (bRegistered)
+		HandleCanEverAffectNavigationChange();
+	}
+}
+
+void UPrimitiveComponent::HandleCanEverAffectNavigationChange()
+{
+	// update octree if already registered
+	if (bRegistered)
+	{
+		if (bCanEverAffectNavigation)
 		{
-			if (bRelevant)
-			{
-				bNavigationRelevant = IsNavigationRelevant();
-				UNavigationSystem::OnComponentRegistered(this);
-			}
-			else
-			{
-				UNavigationSystem::OnComponentUnregistered(this);
-			}
+			bNavigationRelevant = IsNavigationRelevant();
+			UNavigationSystem::OnComponentRegistered(this);
+		}
+		else
+		{
+			UNavigationSystem::OnComponentUnregistered(this);
 		}
 	}
 }
