Made sure AudioComponent doesn't even try talking to NavigationSystem #UE4

diff --git a/Engine/Source/Runtime/Engine/Classes/Components/AudioComponent.h b/Engine/Source/Runtime/Engine/Classes/Components/AudioComponent.h
--- a/Engine/Source/Runtime/Engine/Classes/Components/AudioComponent.h
+++ b/Engine/Source/Runtime/Engine/Classes/Components/AudioComponent.h
@@ -271,6 +271,8 @@ public:
 	virtual void Activate(bool bReset=false) override;
 	virtual void Deactivate() override;
 	virtual void OnUpdateTransform(bool bSkipPhysicsMove) override;
+	// do nothing! on purpose
+	virtual void UpdateNavigationData() override {}
 	// End USceneComponent Interface
 
 	// Begin ActorComponent interface.
