Fixed NavRelevantComponents not updating with changes to owning actor #UE4

This was causing an issue where component's navigation influence was not being updated when, for example, parent actor moved.
UE-11716

diff --git a/Engine/Source/Runtime/Engine/Classes/AI/Navigation/NavLinkCustomComponent.h b/Engine/Source/Runtime/Engine/Classes/AI/Navigation/NavLinkCustomComponent.h
--- a/Engine/Source/Runtime/Engine/Classes/AI/Navigation/NavLinkCustomComponent.h
+++ b/Engine/Source/Runtime/Engine/Classes/AI/Navigation/NavLinkCustomComponent.h
@@ -40,7 +40,7 @@ class ENGINE_API UNavLinkCustomComponent : public UNavRelevantComponent, public
 
 	// BEGIN UNavRelevantComponent Interface
 	virtual void GetNavigationData(FNavigationRelevantData& Data) const override;
-	virtual void CalcBounds() override;
+	virtual void CalcAndCacheBounds() const override;
 	// END UNavRelevantComponent Interface
 
 	// BEGIN UActorComponent Interface
diff --git a/Engine/Source/Runtime/Engine/Classes/AI/Navigation/NavModifierComponent.h b/Engine/Source/Runtime/Engine/Classes/AI/Navigation/NavModifierComponent.h
--- a/Engine/Source/Runtime/Engine/Classes/AI/Navigation/NavModifierComponent.h
+++ b/Engine/Source/Runtime/Engine/Classes/AI/Navigation/NavModifierComponent.h
@@ -16,9 +16,9 @@ class UNavModifierComponent : public UNavRelevantComponent
 	UPROPERTY(EditAnywhere, Category = Navigation)
 	TSubclassOf<UNavArea> AreaClass;
 
-	virtual void CalcBounds() override;
+	virtual void CalcAndCacheBounds() const override;
 	virtual void GetNavigationData(FNavigationRelevantData& Data) const override;
 
 protected:
-	FVector ObstacleExtent;
+	mutable FVector ObstacleExtent;
 };
diff --git a/Engine/Source/Runtime/Engine/Classes/AI/Navigation/NavRelevantComponent.h b/Engine/Source/Runtime/Engine/Classes/AI/Navigation/NavRelevantComponent.h
--- a/Engine/Source/Runtime/Engine/Classes/AI/Navigation/NavRelevantComponent.h
+++ b/Engine/Source/Runtime/Engine/Classes/AI/Navigation/NavRelevantComponent.h
@@ -19,20 +19,24 @@ class ENGINE_API UNavRelevantComponent : public UActorComponent, public INavRele
 	virtual FBox GetNavigationBounds() const override;
 	virtual bool IsNavigationRelevant() const override;
 	virtual void UpdateNavigationBounds() override;
+	virtual UObject* GetNavigationParent() const override;
 	// End INavRelevantInterface Interface
 
+	DEPRECATED(4.8, "This function is deprecated. Use CalcAndCacheBounds instead")
 	virtual void CalcBounds();
+	
+	virtual void CalcAndCacheBounds() const;
 
 	UFUNCTION(BlueprintCallable, Category="AI|Navigation")
 	void SetNavigationRelevancy(bool bRelevant);
 
 	/** force refresh in navigation octree */
 	void RefreshNavigationModifiers();
+	
+protected:
 
 	/** bounds for navigation octree */
-	FBox Bounds;
-
-protected:
+	mutable FBox Bounds;
 
 	UPROPERTY()
 	uint32 bNavigationRelevant : 1;
diff --git a/Engine/Source/Runtime/Engine/Classes/AI/Navigation/NavigationSystem.h b/Engine/Source/Runtime/Engine/Classes/AI/Navigation/NavigationSystem.h
--- a/Engine/Source/Runtime/Engine/Classes/AI/Navigation/NavigationSystem.h
+++ b/Engine/Source/Runtime/Engine/Classes/AI/Navigation/NavigationSystem.h
@@ -242,10 +242,7 @@ public:
 	/** Project a point onto the NavigationData */
 	UFUNCTION(BlueprintPure, Category="AI|Navigation", meta=(WorldContext="WorldContext" ) )
 	static FVector ProjectPointToNavigation(UObject* WorldContext, const FVector& Point, ANavigationData* NavData = NULL, TSubclassOf<UNavigationQueryFilter> FilterClass = NULL, const FVector QueryExtent = FVector::ZeroVector);
-	
-	UFUNCTION(BlueprintPure, Category="AI|Navigation", meta=(WorldContext="WorldContext" ) )
-	static FVector GetRandomPoint(UObject* WorldContext, ANavigationData* NavData = NULL, TSubclassOf<UNavigationQueryFilter> FilterClass = NULL);
-	
+		
 	UFUNCTION(BlueprintPure, Category = "AI|Navigation", meta = (WorldContext = "WorldContext"))
 	static FVector GetRandomReachablePointInRadius(UObject* WorldContext, const FVector& Origin, float Radius, ANavigationData* NavData = NULL, TSubclassOf<UNavigationQueryFilter> FilterClass = NULL);
 
@@ -598,7 +595,7 @@ public:
 	void UpdateNavOctreeElement(UObject* ElementOwner, INavRelevantInterface* ElementInterface, int32 UpdateFlags);
 
 	/** force updating parent node and all its children */
-	void UpdateNavOctreeParentChain(UObject* ElementOwner);
+	void UpdateNavOctreeParentChain(UObject* ElementOwner, bool bSkipElementOwnerUpdate = false);
 
 	/** update component bounds in navigation octree and mark only specified area as dirty, doesn't re-export component geometry */
 	bool UpdateNavOctreeElementBounds(UActorComponent* Comp, const FBox& NewBounds, const FBox& DirtyArea);
@@ -912,5 +909,8 @@ public:
 	static FVector GetRandomPointInRadius(UObject* WorldContext, const FVector& Origin, float Radius, ANavigationData* NavData = NULL, TSubclassOf<UNavigationQueryFilter> FilterClass = NULL);
 	DEPRECATED(4.8, "GetRandomPointInRadius is deprecated. Use either GetRandomReachablePointInRadius or GetRandomPointInNavigableSpace")
 	bool GetRandomPointInRadius(const FVector& Origin, float Radius, FNavLocation& ResultLocation, ANavigationData* NavData = NULL, TSharedPtr<const FNavigationQueryFilter> QueryFilter = NULL) const;
+	DEPRECATED(4.8, "GetRandomPointInRadius is deprecated. Use either GetRandomReachablePointInRadius or GetRandomPointInNavigableSpace")
+	UFUNCTION(BlueprintPure, Category = "AI|Navigation", meta = (WorldContext = "WorldContext"))
+	static FVector GetRandomPoint(UObject* WorldContext, ANavigationData* NavData = NULL, TSubclassOf<UNavigationQueryFilter> FilterClass = NULL);
 };
 
diff --git a/Engine/Source/Runtime/Engine/Private/AI/Navigation/NavLinkCustomComponent.cpp b/Engine/Source/Runtime/Engine/Private/AI/Navigation/NavLinkCustomComponent.cpp
--- a/Engine/Source/Runtime/Engine/Private/AI/Navigation/NavLinkCustomComponent.cpp
+++ b/Engine/Source/Runtime/Engine/Private/AI/Navigation/NavLinkCustomComponent.cpp
@@ -98,7 +98,7 @@ void UNavLinkCustomComponent::GetNavigationData(FNavigationRelevantData& Data) c
 	}
 }
 
-void UNavLinkCustomComponent::CalcBounds()
+void UNavLinkCustomComponent::CalcAndCacheBounds() const
 {
 	Bounds = FBox(0);
 	Bounds += GetStartPoint();
diff --git a/Engine/Source/Runtime/Engine/Private/AI/Navigation/NavModifierComponent.cpp b/Engine/Source/Runtime/Engine/Private/AI/Navigation/NavModifierComponent.cpp
--- a/Engine/Source/Runtime/Engine/Private/AI/Navigation/NavModifierComponent.cpp
+++ b/Engine/Source/Runtime/Engine/Private/AI/Navigation/NavModifierComponent.cpp
@@ -10,7 +10,7 @@ UNavModifierComponent::UNavModifierComponent(const FObjectInitializer& ObjectIni
 	AreaClass = UNavArea_Null::StaticClass();
 }
 
-void UNavModifierComponent::CalcBounds()
+void UNavModifierComponent::CalcAndCacheBounds() const
 {
 	const AActor* MyOwner = GetOwner();
 	if (MyOwner)
@@ -24,5 +24,6 @@ void UNavModifierComponent::CalcBounds()
 
 void UNavModifierComponent::GetNavigationData(FNavigationRelevantData& Data) const
 {
+	CalcAndCacheBounds();
 	Data.Modifiers.Add(FAreaNavModifier(Bounds, FTransform::Identity, AreaClass));
 }
diff --git a/Engine/Source/Runtime/Engine/Private/AI/Navigation/NavRelevantComponent.cpp b/Engine/Source/Runtime/Engine/Private/AI/Navigation/NavRelevantComponent.cpp
--- a/Engine/Source/Runtime/Engine/Private/AI/Navigation/NavRelevantComponent.cpp
+++ b/Engine/Source/Runtime/Engine/Private/AI/Navigation/NavRelevantComponent.cpp
@@ -12,7 +12,7 @@ void UNavRelevantComponent::OnRegister()
 {
 	Super::OnRegister();
 
-	CalcBounds();
+	CalcAndCacheBounds();
 
 	UNavigationSystem::OnComponentRegistered(this);
 }
@@ -36,17 +36,29 @@ bool UNavRelevantComponent::IsNavigationRelevant() const
 
 void UNavRelevantComponent::UpdateNavigationBounds()
 {
-	CalcBounds();
+	CalcAndCacheBounds();
 }
 
-void UNavRelevantComponent::CalcBounds()
+UObject* UNavRelevantComponent::GetNavigationParent() const
 {
 	AActor* OwnerActor = GetOwner();
+	return OwnerActor && OwnerActor->GetRootComponent() ? static_cast<UObject*>(OwnerActor->GetRootComponent()) : static_cast<UObject*>(OwnerActor);
+}
+
+void UNavRelevantComponent::CalcBounds()
+{
+	CalcAndCacheBounds();
+}
+
+void UNavRelevantComponent::CalcAndCacheBounds() const
+{
+	const AActor* OwnerActor = GetOwner();
 	const FVector MyLocation = OwnerActor ? OwnerActor->GetActorLocation() : FVector::ZeroVector;
 
 	Bounds = FBox::BuildAABB(MyLocation, FVector(100.0f, 100.0f, 100.0f));
 }
 
+
 void UNavRelevantComponent::SetNavigationRelevancy(bool bRelevant)
 {
 	if (bNavigationRelevant != bRelevant)
diff --git a/Engine/Source/Runtime/Engine/Private/AI/Navigation/NavigationSystem.cpp b/Engine/Source/Runtime/Engine/Private/AI/Navigation/NavigationSystem.cpp
--- a/Engine/Source/Runtime/Engine/Private/AI/Navigation/NavigationSystem.cpp
+++ b/Engine/Source/Runtime/Engine/Private/AI/Navigation/NavigationSystem.cpp
@@ -2473,11 +2473,12 @@ void UNavigationSystem::UpdateNavOctreeElement(UObject* ElementOwner, INavReleva
 		UpdateInfo.PrevBounds = CurrentBounds;
 		UpdateInfo.bHasPrevData = bAlreadyExists;
 	}
+
+	UpdateNavOctreeParentChain(ElementOwner, /*bSkipElementOwnerUpdate=*/ true);
 }
 
-void UNavigationSystem::UpdateNavOctreeParentChain(UObject* ElementOwner)
+void UNavigationSystem::UpdateNavOctreeParentChain(UObject* ElementOwner, bool bSkipElementOwnerUpdate)
 {
-	INavRelevantInterface* ElementInterface = Cast<INavRelevantInterface>(ElementOwner);
 	const int32 UpdateFlags = OctreeUpdate_ParentChain | OctreeUpdate_Refresh;
 
 	TArray<FWeakObjectPtr> ChildNodes;
@@ -2485,10 +2486,15 @@ void UNavigationSystem::UpdateNavOctreeParentChain(UObject* ElementOwner)
 
 	if (ChildNodes.Num() == 0)
 	{
-		UpdateNavOctreeElement(ElementOwner, ElementInterface, UpdateFlags);
+		if (bSkipElementOwnerUpdate == false)
+		{
+			INavRelevantInterface* ElementInterface = Cast<INavRelevantInterface>(ElementOwner);
+			UpdateNavOctreeElement(ElementOwner, ElementInterface, UpdateFlags);
+		}
 		return;
 	}
 
+	INavRelevantInterface* ElementInterface = Cast<INavRelevantInterface>(ElementOwner);
 	TArray<INavRelevantInterface*> ChildNavInterfaces;
 	ChildNavInterfaces.AddZeroed(ChildNodes.Num());
 	
