Minor tweak to previous change reducing redundant nav octree elements unregistering/registering #UE4

diff --git a/Engine/Source/Runtime/Engine/Private/AI/Navigation/NavigationSystem.cpp b/Engine/Source/Runtime/Engine/Private/AI/Navigation/NavigationSystem.cpp
--- a/Engine/Source/Runtime/Engine/Private/AI/Navigation/NavigationSystem.cpp
+++ b/Engine/Source/Runtime/Engine/Private/AI/Navigation/NavigationSystem.cpp
@@ -2508,8 +2508,11 @@ void UNavigationSystem::UpdateNavOctreeParentChain(UObject* ElementOwner, bool b
 		}
 	}
 
-	UnregisterNavOctreeElement(ElementOwner, ElementInterface, UpdateFlags);
-	RegisterNavOctreeElement(ElementOwner, ElementInterface, UpdateFlags);
+	if (bSkipElementOwnerUpdate == false)
+	{
+		UnregisterNavOctreeElement(ElementOwner, ElementInterface, UpdateFlags);
+		RegisterNavOctreeElement(ElementOwner, ElementInterface, UpdateFlags);
+	}
 
 	for (int32 Idx = 0; Idx < ChildNodes.Num(); Idx++)
 	{
