Fixed crash when applying overried material to a destructible mesh instance
#UE-13582

diff --git a/Engine/Source/Runtime/Engine/Private/Components/DestructibleComponent.cpp b/Engine/Source/Runtime/Engine/Private/Components/DestructibleComponent.cpp
--- a/Engine/Source/Runtime/Engine/Private/Components/DestructibleComponent.cpp
+++ b/Engine/Source/Runtime/Engine/Private/Components/DestructibleComponent.cpp
@@ -1382,18 +1382,21 @@ void UDestructibleComponent::SetMaterial(int32 ElementIndex, UMaterialInterface*
 
 #if WITH_APEX
 	// Set new template parameters for the apex actor, so they take effect before fracturing too.
-	physx::apex::NxPhysX3DescTemplate* Template = ApexDestructibleActor->createPhysX3DescTemplate();
-	if(ApexDestructibleActor->getPhysX3Template(*Template))
+	if(ApexDestructibleActor)
 	{
-		UPhysicalMaterial* SimpleMaterial = GetBodyInstance()->GetSimplePhysicalMaterial();
-		check(SimpleMaterial);
-		PxMaterial* PhysxMat = SimpleMaterial->GetPhysXMaterial();
+		physx::apex::NxPhysX3DescTemplate* Template = ApexDestructibleActor->createPhysX3DescTemplate();
+		if(ApexDestructibleActor->getPhysX3Template(*Template))
+		{
+			UPhysicalMaterial* SimpleMaterial = GetBodyInstance()->GetSimplePhysicalMaterial();
+			check(SimpleMaterial);
+			PxMaterial* PhysxMat = SimpleMaterial->GetPhysXMaterial();
 
-		Template->setMaterials(&PhysxMat, 1);
+			Template->setMaterials(&PhysxMat, 1);
 
-		ApexDestructibleActor->setPhysX3Template(Template);
+			ApexDestructibleActor->setPhysX3Template(Template);
+		}
+		Template->release();
 	}
-	Template->release();
 #endif
 }
 
