Fix for UE-12844 - Metallic Setting is not passed through Material Function in Feature Level ES2

Add IsConnected() for material attribute input.
IsNonmetal checks material attribute properties if it is used.

diff --git a/Engine/Source/Runtime/Engine/Private/Materials/MaterialShared.cpp b/Engine/Source/Runtime/Engine/Private/Materials/MaterialShared.cpp
--- a/Engine/Source/Runtime/Engine/Private/Materials/MaterialShared.cpp
+++ b/Engine/Source/Runtime/Engine/Private/Materials/MaterialShared.cpp
@@ -181,6 +181,8 @@ int32 FMaterialAttributesInput::CompileWithDefault(class FMaterialCompiler* Comp
 		}
 	}
 
+	SetConnectedProperty(Property, Ret != INDEX_NONE);
+
 	if( Ret == INDEX_NONE )
 	{
 		Ret = GetDefaultExpressionForMaterialProperty(Compiler, Property);
@@ -800,7 +802,9 @@ bool FMaterialResource::OutputsVelocityOnBasePass() const
 
 bool FMaterialResource::IsNonmetal() const
 {
-	return !Material->Metallic.IsConnected() && !Material->Specular.IsConnected();
+	return !Material->bUseMaterialAttributes ?
+			(!Material->Metallic.IsConnected() && !Material->Specular.IsConnected()) :
+			!(Material->MaterialAttributes.IsConnected(MP_Specular) || Material->MaterialAttributes.IsConnected(MP_Metallic));
 }
 
 bool FMaterialResource::UseLmDirectionality() const
diff --git a/Engine/Source/Runtime/Engine/Public/MaterialExpressionIO.h b/Engine/Source/Runtime/Engine/Public/MaterialExpressionIO.h
--- a/Engine/Source/Runtime/Engine/Public/MaterialExpressionIO.h
+++ b/Engine/Source/Runtime/Engine/Public/MaterialExpressionIO.h
@@ -112,5 +112,19 @@ struct FVector2MaterialInput : FMaterialInput<FVector2D>
 };
 struct FMaterialAttributesInput : FMaterialInput<int32>
 {
+	FMaterialAttributesInput() : FMaterialInput<int32>(), PropertyConnectedBitmask(0)
+	{ 
+		// ensure PropertyConnectedBitmask can contain all properties.
+		static_assert((uint32)(MP_MAX)-1 <= (8 * sizeof(PropertyConnectedBitmask)), "PropertyConnectedBitmask cannot contain entire EMaterialProperty enumeration.");
+	}
+
 	ENGINE_API int32 CompileWithDefault(class FMaterialCompiler* Compiler, EMaterialProperty Property);
+	ENGINE_API bool IsConnected(EMaterialProperty Property) { return ((PropertyConnectedBitmask >> (uint32)Property) & 0x1) != 0; }
+	ENGINE_API void SetConnectedProperty(EMaterialProperty Property, bool bIsConnected) 
+	{
+		PropertyConnectedBitmask = bIsConnected ? PropertyConnectedBitmask | (1 << (uint32)Property) : PropertyConnectedBitmask & ~(1 << (uint32)Property);
+	}
+
+	// each bit corresponds to EMaterialProperty connection status.
+	uint32 PropertyConnectedBitmask;
 };
\ No newline at end of file
