Use PREPROCESSOR_JOIN instead of reinventing the wheel in many locations

diff --git a/Engine/Source/Editor/UnrealEd/Public/Kismet2/CompilerResultsLog.h b/Engine/Source/Editor/UnrealEd/Public/Kismet2/CompilerResultsLog.h
--- a/Engine/Source/Editor/UnrealEd/Public/Kismet2/CompilerResultsLog.h
+++ b/Engine/Source/Editor/UnrealEd/Public/Kismet2/CompilerResultsLog.h
@@ -224,14 +224,12 @@ public:
 	}
 };
 
-#define SCOPED_COMPILER_EVENT_TOKENPASTE_INNER(x,y) x##y
-#define SCOPED_COMPILER_EVENT_TOKENPASTE(x,y) SCOPED_COMPILER_EVENT_TOKENPASTE_INNER(x,y)
 #if STATS
 #define BP_SCOPED_COMPILER_EVENT_STAT(Stat) \
 	SCOPE_CYCLE_COUNTER(Stat); \
-	FScopedCompilerEvent SCOPED_COMPILER_EVENT_TOKENPASTE(ScopedCompilerEvent,__LINE__)(GET_STATDESCRIPTION(Stat))
+	FScopedCompilerEvent PREPROCESSOR_JOIN(ScopedCompilerEvent,__LINE__)(GET_STATDESCRIPTION(Stat))
 #else
 #define BP_SCOPED_COMPILER_EVENT_STAT(Stat) \
-	FScopedCompilerEvent SCOPED_COMPILER_EVENT_TOKENPASTE(ScopedCompilerEvent,__LINE__)(ANSI_TO_TCHAR(#Stat))
+	FScopedCompilerEvent PREPROCESSOR_JOIN(ScopedCompilerEvent,__LINE__)(ANSI_TO_TCHAR(#Stat))
 #endif
 #endif	//#if WITH_EDITOR
diff --git a/Engine/Source/Runtime/Core/Public/Mac/MacPlatformMisc.h b/Engine/Source/Runtime/Core/Public/Mac/MacPlatformMisc.h
--- a/Engine/Source/Runtime/Core/Public/Mac/MacPlatformMisc.h
+++ b/Engine/Source/Runtime/Core/Public/Mac/MacPlatformMisc.h
@@ -224,9 +224,7 @@ private:
 	NSAutoreleasePool*	Pool;
 };
 
-#define SCOPED_AUTORELEASE_POOL_TOKENPASTE_INNER(x,y) x##y
-#define SCOPED_AUTORELEASE_POOL_TOKENPASTE(x,y) SCOPED_AUTORELEASE_POOL_TOKENPASTE_INNER(x,y)
-#define SCOPED_AUTORELEASE_POOL const FScopeAutoreleasePool SCOPED_AUTORELEASE_POOL_TOKENPASTE(Pool,__LINE__);
+#define SCOPED_AUTORELEASE_POOL const FScopeAutoreleasePool PREPROCESSOR_JOIN(Pool,__LINE__);
 
 #endif // __OBJC__
 
diff --git a/Engine/Source/Runtime/Core/Public/Stats/StatsMisc.h b/Engine/Source/Runtime/Core/Public/Stats/StatsMisc.h
--- a/Engine/Source/Runtime/Core/Public/Stats/StatsMisc.h
+++ b/Engine/Source/Runtime/Core/Public/Stats/StatsMisc.h
@@ -74,17 +74,14 @@ protected:
 	EScopeLogTimeUnits Units;
 };
 
-#define SCOPE_LOG_TOKENPASTE_INNER(x,y) x##y
-#define SCOPE_LOG_TOKENPASTE(x,y) SCOPE_LOG_TOKENPASTE_INNER(x,y)
-
 #define SCOPE_LOG_TIME(Name,CumulativePtr) \
-	FScopeLogTime SCOPE_LOG_TOKENPASTE(ScopeLogTime,__LINE__)(Name,CumulativePtr);
+	FScopeLogTime PREPROCESSOR_JOIN(ScopeLogTime,__LINE__)(Name,CumulativePtr);
 
 #define SCOPE_LOG_TIME_IN_SECONDS(Name,CumulativePtr) \
-	FScopeLogTime SCOPE_LOG_TOKENPASTE(ScopeLogTime,__LINE__)(Name, CumulativePtr, FScopeLogTime::ScopeLog_Seconds);
+	FScopeLogTime PREPROCESSOR_JOIN(ScopeLogTime,__LINE__)(Name, CumulativePtr, FScopeLogTime::ScopeLog_Seconds);
 
 #define SCOPE_LOG_TIME_FUNC() \
-	FScopeLogTime SCOPE_LOG_TOKENPASTE(ScopeLogTime,__LINE__)(__FUNCTION__);
+	FScopeLogTime PREPROCESSOR_JOIN(ScopeLogTime,__LINE__)(__FUNCTION__);
 
 #define SCOPE_LOG_TIME_FUNC_WITH_GLOBAL(CumulativePtr) \
-	FScopeLogTime SCOPE_LOG_TOKENPASTE(ScopeLogTime,__LINE__)(__FUNCTION__,CumulativePtr);
+	FScopeLogTime PREPROCESSOR_JOIN(ScopeLogTime,__LINE__)(__FUNCTION__,CumulativePtr);
diff --git a/Engine/Source/Runtime/Engine/Private/PhysicsEngine/PhysXSupport.h b/Engine/Source/Runtime/Engine/Private/PhysicsEngine/PhysXSupport.h
--- a/Engine/Source/Runtime/Engine/Private/PhysicsEngine/PhysXSupport.h
+++ b/Engine/Source/Runtime/Engine/Private/PhysicsEngine/PhysXSupport.h
@@ -78,10 +78,8 @@ private:
 	PxScene* PScene;
 };
 
-#define SCOPED_SCENE_TOKENPASTE_INNER(x,y) x##y
-#define SCOPED_SCENE_TOKENPASTE(x,y) SCOPED_SCENE_TOKENPASTE_INNER(x,y)
-#define SCOPED_SCENE_READ_LOCK( _scene ) FPhysXSceneReadLock SCOPED_SCENE_TOKENPASTE(_rlock,__LINE__)(_scene)
-#define SCOPED_SCENE_WRITE_LOCK( _scene ) FPhysXSceneWriteLock SCOPED_SCENE_TOKENPASTE(_wlock,__LINE__)(_scene)
+#define SCOPED_SCENE_READ_LOCK( _scene ) FPhysXSceneReadLock PREPROCESSOR_JOIN(_rlock,__LINE__)(_scene)
+#define SCOPED_SCENE_WRITE_LOCK( _scene ) FPhysXSceneWriteLock PREPROCESSOR_JOIN(_wlock,__LINE__)(_scene)
 
 #define SCENE_LOCK_READ( _scene )		{ SCOPE_CYCLE_COUNTER(STAT_PhysSceneReadLock); if((_scene) != NULL) { (_scene)->lockRead(); } }
 #define SCENE_UNLOCK_READ( _scene )		{ if((_scene) != NULL) { (_scene)->unlockRead(); } }
diff --git a/Engine/Source/Runtime/Engine/Public/SceneUtils.h b/Engine/Source/Runtime/Engine/Public/SceneUtils.h
--- a/Engine/Source/Runtime/Engine/Public/SceneUtils.h
+++ b/Engine/Source/Runtime/Engine/Public/SceneUtils.h
@@ -52,12 +52,10 @@
 	};
 
 	// Macros to allow for scoping of draw events
-	#define SCOPED_DRAW_TOKENPASTE_INNER(x,y) x##y
-	#define SCOPED_DRAW_TOKENPASTE(x,y) SCOPED_DRAW_TOKENPASTE_INNER(x,y)
-	#define SCOPED_DRAW_EVENT(RHICmdList, Name) FDrawEvent SCOPED_DRAW_TOKENPASTE(Event_##Name,__LINE__); if(GEmitDrawEvents) SCOPED_DRAW_TOKENPASTE(Event_##Name,__LINE__).Start(RHICmdList, TEXT(#Name));
-	#define SCOPED_DRAW_EVENTF(RHICmdList, Name, Format, ...) FDrawEvent SCOPED_DRAW_TOKENPASTE(Event_##Name,__LINE__); if(GEmitDrawEvents) SCOPED_DRAW_TOKENPASTE(Event_##Name,__LINE__).Start(RHICmdList, Format, ##__VA_ARGS__);
-	#define SCOPED_CONDITIONAL_DRAW_EVENT(RHICmdList, Name, Condition) FDrawEvent SCOPED_DRAW_TOKENPASTE(Event_##Name,__LINE__); if(GEmitDrawEvents && (Condition)) SCOPED_DRAW_TOKENPASTE(Event_##Name,__LINE__).Start(RHICmdList, TEXT(#Name));
-	#define SCOPED_CONDITIONAL_DRAW_EVENTF(RHICmdList, Name, Condition, Format, ...) FDrawEvent SCOPED_DRAW_TOKENPASTE(Event_##Name,__LINE__); if(GEmitDrawEvents && (Condition)) SCOPED_DRAW_TOKENPASTE(Event_##Name,__LINE__).Start(RHICmdList, Format, ##__VA_ARGS__);
+	#define SCOPED_DRAW_EVENT(RHICmdList, Name) FDrawEvent PREPROCESSOR_JOIN(Event_##Name,__LINE__); if(GEmitDrawEvents) PREPROCESSOR_JOIN(Event_##Name,__LINE__).Start(RHICmdList, TEXT(#Name));
+	#define SCOPED_DRAW_EVENTF(RHICmdList, Name, Format, ...) FDrawEvent PREPROCESSOR_JOIN(Event_##Name,__LINE__); if(GEmitDrawEvents) PREPROCESSOR_JOIN(Event_##Name,__LINE__).Start(RHICmdList, Format, ##__VA_ARGS__);
+	#define SCOPED_CONDITIONAL_DRAW_EVENT(RHICmdList, Name, Condition) FDrawEvent PREPROCESSOR_JOIN(Event_##Name,__LINE__); if(GEmitDrawEvents && (Condition)) PREPROCESSOR_JOIN(Event_##Name,__LINE__).Start(RHICmdList, TEXT(#Name));
+	#define SCOPED_CONDITIONAL_DRAW_EVENTF(RHICmdList, Name, Condition, Format, ...) FDrawEvent PREPROCESSOR_JOIN(Event_##Name,__LINE__); if(GEmitDrawEvents && (Condition)) PREPROCESSOR_JOIN(Event_##Name,__LINE__).Start(RHICmdList, Format, ##__VA_ARGS__);
 
 #else
 
diff --git a/Engine/Source/Runtime/GameplayAbilities/Public/AbilitySystemLog.h b/Engine/Source/Runtime/GameplayAbilities/Public/AbilitySystemLog.h
--- a/Engine/Source/Runtime/GameplayAbilities/Public/AbilitySystemLog.h
+++ b/Engine/Source/Runtime/GameplayAbilities/Public/AbilitySystemLog.h
@@ -38,9 +38,7 @@ GAMEPLAYABILITIES_API DECLARE_LOG_CATEGORY_EXTERN(VLogAbilitySystem, Warning, Al
 	UE_VLOG_HISTOGRAM(OwnerActor, VLogAbilitySystem, Log, GraphName, LineName, NewPt); \
 }
 
-#define ABILITY_LOG_TOKENPASTE_INNER(x,y) x##y
-#define ABILITY_LOG_TOKENPASTE(x,y) ABILITY_LOG_TOKENPASTE_INNER(x,y)
-#define ABILITY_LOG_SCOPE( Format, ... ) AbilitySystemLogScope ABILITY_LOG_TOKENPASTE(LogScope,__LINE__)( FString::Printf(Format, ##__VA_ARGS__));
+#define ABILITY_LOG_SCOPE( Format, ... ) AbilitySystemLogScope PREPROCESSOR_JOIN(LogScope,__LINE__)( FString::Printf(Format, ##__VA_ARGS__));
 
 struct AbilitySystemLogScope
 {
