Fix memory overwrite caused by UWorld::EditorViews being too small for number of possible view types
Made EditorViews dynamic so future changes in size should not have similar issues
Made EditorViews editor only data

diff --git a/Engine/Source/Runtime/Core/Public/UObject/ObjectVersion.h b/Engine/Source/Runtime/Core/Public/UObject/ObjectVersion.h
--- a/Engine/Source/Runtime/Core/Public/UObject/ObjectVersion.h
+++ b/Engine/Source/Runtime/Core/Public/UObject/ObjectVersion.h
@@ -490,6 +490,8 @@ enum EUnrealEngineObjectUE4Version
 	VER_UE4_LANDSCAPE_SPLINE_CROSS_LEVEL_MESHES,
 	// Deprecate the variables used for sizing in the designer on UUserWidget
 	VER_UE4_DEPRECATE_USER_WIDGET_DESIGN_SIZE,
+	// Make the editor views array dynamically sized
+	VER_UE4_ADD_EDITOR_VIEWS,
 
 	// -----<new versions can be added before this line>-------------------------------------------------
 	// - this needs to be the last line (see note below)
diff --git a/Engine/Source/Runtime/Engine/Classes/Engine/World.h b/Engine/Source/Runtime/Engine/Classes/Engine/World.h
--- a/Engine/Source/Runtime/Engine/Classes/Engine/World.h
+++ b/Engine/Source/Runtime/Engine/Classes/Engine/World.h
@@ -268,18 +268,25 @@ private:
 };
 
 /** Saved editor viewport state information */
+USTRUCT()
 struct ENGINE_API FLevelViewportInfo
 {
+	GENERATED_BODY()
+
 	/** Where the camera is positioned within the viewport. */
+	UPROPERTY()
 	FVector CamPosition;
 
 	/** The camera's position within the viewport. */
+	UPROPERTY()
 	FRotator CamRotation;
 
 	/** The zoom value  for orthographic mode. */
+	UPROPERTY()
 	float CamOrthoZoom;
 
 	/** Whether camera settings have been systematically changed since the last level viewport update. */
+	UPROPERTY()
 	bool CamUpdated;
 
 	FLevelViewportInfo()
@@ -297,7 +304,8 @@ struct ENGINE_API FLevelViewportInfo
 		, CamUpdated(false)
 	{
 	}
-
+	
+	// Needed for backwards compatibility for VER_UE4_ADD_EDITOR_VIEWS, can be removed along with it
 	friend FArchive& operator<<( FArchive& Ar, FLevelViewportInfo& I )
 	{
 		Ar << I.CamPosition;
@@ -691,8 +699,11 @@ public:
 	/** The current renderer feature level of this world */
 	ERHIFeatureLevel::Type						FeatureLevel;
 	
-	/** Saved editor viewport states - one for each view type. Indexed using ELevelViewportType above.							*/
-	FLevelViewportInfo							EditorViews[4];
+#if WITH_EDITORONLY_DATA
+	/** Saved editor viewport states - one for each view type. Indexed using ELevelViewportType from UnrealEdTypes.h.	*/
+	UPROPERTY(NonTransactional)
+	TArray<FLevelViewportInfo>					EditorViews;
+#endif
 
 	/** 
 	 * Set the CurrentLevel for this world. 
diff --git a/Engine/Source/Runtime/Engine/Private/World.cpp b/Engine/Source/Runtime/Engine/Private/World.cpp
--- a/Engine/Source/Runtime/Engine/Private/World.cpp
+++ b/Engine/Source/Runtime/Engine/Private/World.cpp
@@ -50,6 +50,7 @@
 	#include "Editor/UnrealEd/Classes/ThumbnailRendering/WorldThumbnailInfo.h"
 	#include "SlateBasics.h"
 	#include "Editor/Kismet/Public/FindInBlueprintManager.h"
+	#include "Editor/UnrealEd/Classes/Editor/UnrealEdTypes.h"
 #endif
 
 #include "MallocProfiler.h"
@@ -126,6 +127,7 @@ UWorld::UWorld( const FObjectInitializer& ObjectInitializer )
 	TimerManager = new FTimerManager();
 #if WITH_EDITOR
 	bBroadcastSelectionChange = true; //Ed Only
+	EditorViews.SetNum(ELevelViewportType::LVT_MAX);
 #endif // WITH_EDITOR
 
 	FWorldDelegates::OnPostWorldCreation.Broadcast(this);
@@ -145,10 +147,35 @@ void UWorld::Serialize( FArchive& Ar )
 
 	Ar << PersistentLevel;
 
-	Ar << EditorViews[0];
-	Ar << EditorViews[1];
-	Ar << EditorViews[2];
-	Ar << EditorViews[3];
+	if (Ar.UE4Ver() < VER_UE4_ADD_EDITOR_VIEWS)
+	{
+		for (int32 i = 0; i < 4; ++i)
+		{
+			FLevelViewportInfo TempViewportInfo;
+			Ar << TempViewportInfo;
+#if WITH_EDITOR
+			if (Ar.IsLoading())
+			{
+				EditorViews.Add(TempViewportInfo);
+			}
+#endif
+		}
+	}
+#if WITH_EDITOR
+	if ( Ar.IsLoading() )
+	{
+		for (FLevelViewportInfo& ViewportInfo : EditorViews)
+		{
+			ViewportInfo.CamUpdated = true;
+
+			if ( ViewportInfo.CamOrthoZoom == 0.f )
+			{
+				ViewportInfo.CamOrthoZoom = DEFAULT_ORTHOZOOM;
+			}
+		}
+		EditorViews.SetNum(ELevelViewportType::LVT_MAX);
+	}
+#endif
 
 	if (Ar.UE4Ver() < VER_UE4_REMOVE_SAVEGAMESUMMARY)
 	{
