PR #1048: [WorldCollision.h FCollisionShape] Added Is Checks for each shape type! -Rama (Contributed by [REDACTED_PERSON])

diff --git a/Engine/Source/Runtime/Engine/Public/WorldCollision.h b/Engine/Source/Runtime/Engine/Public/WorldCollision.h
--- a/Engine/Source/Runtime/Engine/Public/WorldCollision.h
+++ b/Engine/Source/Runtime/Engine/Public/WorldCollision.h
@@ -76,6 +76,30 @@ struct FCollisionShape
 	{
 		ShapeType = ECollisionShape::Line;
 	}
+ 
+	/** Is the shape currently a Line (Default)? */
+	bool IsLine() const
+	{
+		return ShapeType == ECollisionShape::Line;
+	}
+	
+	/** Is the shape currently a box? */
+	bool IsBox() const
+	{
+		return ShapeType == ECollisionShape::Box;
+	}
+		
+	/** Is the shape currently a sphere? */
+	bool IsSphere() const
+	{
+		return ShapeType == ECollisionShape::Sphere;
+	}	
+		
+	/** Is the shape currently a capsule? */
+	bool IsCapsule() const
+	{
+		return ShapeType == ECollisionShape::Capsule;
+	}
 
 	/** Utility function to Set Box and dimension */
 	void SetBox(const FVector& HalfExtent)
