SkeletalMeshComponent::AddReferenceObjects for skeletons it references in map. This could cause crash in cooks if GC happened at just the right time

diff --git a/Engine/Source/Runtime/Engine/Classes/Animation/Skeleton.h b/Engine/Source/Runtime/Engine/Classes/Animation/Skeleton.h
--- a/Engine/Source/Runtime/Engine/Classes/Animation/Skeleton.h
+++ b/Engine/Source/Runtime/Engine/Classes/Animation/Skeleton.h
@@ -595,6 +595,8 @@ public:
 	virtual void PostInitProperties() override;
 	virtual void Serialize(FArchive& Ar) override;
 
+	ENGINE_API static void AddReferencedObjects(UObject* InThis, FReferenceCollector& Collector);
+
 	/** 
 	 * Create RefLocalPoses from InSkelMesh
 	 * 
diff --git a/Engine/Source/Runtime/Engine/Private/Animation/Skeleton.cpp b/Engine/Source/Runtime/Engine/Private/Animation/Skeleton.cpp
--- a/Engine/Source/Runtime/Engine/Private/Animation/Skeleton.cpp
+++ b/Engine/Source/Runtime/Engine/Private/Animation/Skeleton.cpp
@@ -166,6 +166,20 @@ void USkeleton::Serialize( FArchive& Ar )
 #endif
 }
 
+void USkeleton::AddReferencedObjects(UObject* InThis, FReferenceCollector& Collector)
+{
+	USkeleton* This = CastChecked<USkeleton>(InThis);
+
+#if WITH_EDITORONLY_DATA
+	for (auto Iter = This->AnimRetargetSources.CreateIterator(); Iter; ++Iter)
+	{		
+		Collector.AddReferencedObject(Iter.Value().ReferenceMesh, This);
+	}
+#endif
+
+	Super::AddReferencedObjects( This, Collector );
+}
+
 /** Remove this function when VER_UE4_REFERENCE_SKELETON_REFACTOR is removed. */
 void USkeleton::ConvertToFReferenceSkeleton()
 {
