Expand on AActor::AddComponent documentation

diff --git a/Engine/Source/Runtime/Engine/Classes/GameFramework/Actor.h b/Engine/Source/Runtime/Engine/Classes/GameFramework/Actor.h
--- a/Engine/Source/Runtime/Engine/Classes/GameFramework/Actor.h
+++ b/Engine/Source/Runtime/Engine/Classes/GameFramework/Actor.h
@@ -857,9 +857,14 @@ public:
 	bool HasAuthority() const;
 
 	/** 
-	 * Create a new component given a template name. Template is found in the owning Blueprint.
+	 * BLUEPRINT INTERNAL USE ONLY for Add Component nodes
+	 * Creates a new component owned by this Actor based a component template found in the owning Blueprint.
 	 * Automatic attachment causes the first component created to become the root component, and all subsequent components
-	 * will be attached to the root component.  In manual mode, it is up to the user to attach or set as root.
+	 * will be attached to the root component.  In manual mode, it is up to the user to attach or set as root
+	 * @param TemplateName					The name of the Component Template to use.
+	 * @param bManualAttachment				Whether manual or automatic attachment is to be used
+	 * @param RelativeTransform				The relative transform between the new component and its attach parent (automatic only)
+	 * @param ComponentTemplateContext		Optional UBlueprintGeneratedClass reference to use to find the template in. If null (or not a BPGC), component is sought in this Actor's class
 	 */
 	UFUNCTION(BlueprintCallable, meta=(BlueprintInternalUseOnly = "true", DefaultToSelf="ComponentTemplateContext", HidePin="ComponentTemplateContext"))
 	class UActorComponent* AddComponent(FName TemplateName, bool bManualAttachment, const FTransform& RelativeTransform, const UObject* ComponentTemplateContext);
