UE-13597 Crash when setting a struct variable datatype to DataTableRowHandle

diff --git a/Engine/Source/Runtime/Engine/Classes/Engine/DataTable.h b/Engine/Source/Runtime/Engine/Classes/Engine/DataTable.h
--- a/Engine/Source/Runtime/Engine/Classes/Engine/DataTable.h
+++ b/Engine/Source/Runtime/Engine/Classes/Engine/DataTable.h
@@ -194,7 +194,7 @@ private:
 
 
 /** Handle to a particular row in a table*/
-USTRUCT(BlueprintType)
+USTRUCT()
 struct ENGINE_API FDataTableRowHandle
 {
 	GENERATED_USTRUCT_BODY()
