BP: PIE can mark BP as dirty (UE-13290).

diff --git a/Engine/Source/Runtime/Engine/Private/BlueprintGeneratedClass.cpp b/Engine/Source/Runtime/Engine/Private/BlueprintGeneratedClass.cpp
--- a/Engine/Source/Runtime/Engine/Private/BlueprintGeneratedClass.cpp
+++ b/Engine/Source/Runtime/Engine/Private/BlueprintGeneratedClass.cpp
@@ -250,9 +250,18 @@ void UBlueprintGeneratedClass::ConditionalRecompileClass(TArray<UObject*>* ObjLo
 			const bool bWasRegenerating = GeneratingBP->bIsRegeneratingOnLoad;
 			GeneratingBP->bIsRegeneratingOnLoad = true;
 
-			// Make sure that nodes are up to date, so that we get any updated blueprint signatures
-			FBlueprintEditorUtils::RefreshExternalBlueprintDependencyNodes(GeneratingBP);
+			{
+				UPackage* const Package = Cast<UPackage>(GeneratingBP->GetOutermost());
+				const bool bStartedWithUnsavedChanges = Package != nullptr ? Package->IsDirty() : true;
+
+				// Make sure that nodes are up to date, so that we get any updated blueprint signatures
+				FBlueprintEditorUtils::RefreshExternalBlueprintDependencyNodes(GeneratingBP);
 
+				if (Package != nullptr && Package->IsDirty() && !bStartedWithUnsavedChanges)
+				{
+					Package->SetDirtyFlag(false);
+				}
+			}
 			if ((GeneratingBP->Status != BS_Error) && (GeneratingBP->BlueprintType != EBlueprintType::BPTYPE_MacroLibrary))
 			{
 				FKismetEditorUtilities::RecompileBlueprintBytecode(GeneratingBP, ObjLoaded);
