Engine - Changing OnTouchStarted/Ended to use GetScreenSpacePosition instead of GetLastScreenSpacePosition.

diff --git a/Engine/Source/Runtime/Engine/Private/Slate/SceneViewport.cpp b/Engine/Source/Runtime/Engine/Private/Slate/SceneViewport.cpp
--- a/Engine/Source/Runtime/Engine/Private/Slate/SceneViewport.cpp
+++ b/Engine/Source/Runtime/Engine/Private/Slate/SceneViewport.cpp
@@ -584,7 +584,7 @@ FReply FSceneViewport::OnTouchStarted( const FGeometry& MyGeometry, const FPoint
 		// Switch to the viewport clients world before processing input
 		FScopedConditionalWorldSwitcher WorldSwitcher( ViewportClient );
 
-		const FVector2D TouchPosition = MyGeometry.AbsoluteToLocal(TouchEvent.GetLastScreenSpacePosition());
+		const FVector2D TouchPosition = MyGeometry.AbsoluteToLocal(TouchEvent.GetScreenSpacePosition());
 
 		if( !ViewportClient->InputTouch( this, TouchEvent.GetUserIndex(), TouchEvent.GetPointerIndex(), ETouchType::Began, TouchPosition, FDateTime::Now(), TouchEvent.GetTouchpadIndex()) )
 		{
@@ -632,7 +632,7 @@ FReply FSceneViewport::OnTouchEnded( const FGeometry& MyGeometry, const FPointer
 		// Switch to the viewport clients world before processing input
 		FScopedConditionalWorldSwitcher WorldSwitcher( ViewportClient );
 
-		const FVector2D TouchPosition = MyGeometry.AbsoluteToLocal(TouchEvent.GetLastScreenSpacePosition());
+		const FVector2D TouchPosition = MyGeometry.AbsoluteToLocal(TouchEvent.GetScreenSpacePosition());
 
 		if( !ViewportClient->InputTouch( this, TouchEvent.GetUserIndex(), TouchEvent.GetPointerIndex(), ETouchType::Ended, TouchPosition, FDateTime::Now(), TouchEvent.GetTouchpadIndex()) )
 		{
