Added performance stat for fast delta seriarlize on client

diff --git a/Engine/Source/Runtime/CoreUObject/Private/UObject/CoreNet.cpp b/Engine/Source/Runtime/CoreUObject/Private/UObject/CoreNet.cpp
--- a/Engine/Source/Runtime/CoreUObject/Private/UObject/CoreNet.cpp
+++ b/Engine/Source/Runtime/CoreUObject/Private/UObject/CoreNet.cpp
@@ -8,7 +8,8 @@
 
 DEFINE_LOG_CATEGORY_STATIC(LogCoreNet, Log, All);
 
-DEFINE_STAT(STAT_NetSerializeFast_Array);
+DEFINE_STAT(STAT_NetSerializeFastArray);
+DEFINE_STAT(STAT_NetSerializeFastArray_BuildMap);
 
 /*-----------------------------------------------------------------------------
 	FClassNetCache implementation.
diff --git a/Engine/Source/Runtime/CoreUObject/Public/UObject/UObjectGlobals.h b/Engine/Source/Runtime/CoreUObject/Public/UObject/UObjectGlobals.h
--- a/Engine/Source/Runtime/CoreUObject/Public/UObject/UObjectGlobals.h
+++ b/Engine/Source/Runtime/CoreUObject/Public/UObject/UObjectGlobals.h
@@ -30,7 +30,8 @@ DECLARE_DWORD_COUNTER_STAT_EXTERN(TEXT("FindObjectFast"),STAT_FindObjectFast,STA
  * Network stats counters
  */
 
-DECLARE_CYCLE_STAT_EXTERN(TEXT("NetSerializeFast Array"),STAT_NetSerializeFast_Array,STATGROUP_ServerCPU, COREUOBJECT_API);
+DECLARE_CYCLE_STAT_EXTERN(TEXT("NetSerializeFast Array"),STAT_NetSerializeFastArray,STATGROUP_ServerCPU, COREUOBJECT_API);
+DECLARE_CYCLE_STAT_EXTERN(TEXT("NetSerializeFast Array BuildMap"),STAT_NetSerializeFastArray_BuildMap,STATGROUP_ServerCPU, COREUOBJECT_API);
 
 
 
diff --git a/Engine/Source/Runtime/Engine/Classes/Engine/NetSerialization.h b/Engine/Source/Runtime/Engine/Classes/Engine/NetSerialization.h
--- a/Engine/Source/Runtime/Engine/Classes/Engine/NetSerialization.h
+++ b/Engine/Source/Runtime/Engine/Classes/Engine/NetSerialization.h
@@ -382,6 +382,7 @@ struct FFastArraySerializer
 template< typename Type, typename SerializerType >
 bool FFastArraySerializer::FastArrayDeltaSerialize( TArray<Type> &Items, FNetDeltaSerializeInfo& Parms, SerializerType& ArraySerializer )
 {
+	SCOPE_CYCLE_COUNTER(STAT_NetSerializeFastArray);
 	class UScriptStruct* InnerStruct = Type::StaticStruct();
 
 	if ( Parms.bUpdateUnmappedObjects || Parms.Writer == NULL )
@@ -391,7 +392,7 @@ bool FFastArraySerializer::FastArrayDeltaSerialize( TArray<Type> &Items, FNetDel
 		//---------------
 		if (ArraySerializer.ItemMap.Num() != Items.Num())
 		{
-			SCOPE_CYCLE_COUNTER(STAT_NetSerializeFast_Array);
+			SCOPE_CYCLE_COUNTER(STAT_NetSerializeFastArray_BuildMap);
 			UE_LOG(LogNetSerialization, Verbose, TEXT("FastArrayDeltaSerialize: Recreating Items map. Struct: %s, Items.Num: %d Map.Num: %d"), *InnerStruct->GetOwnerStruct()->GetName(), Items.Num(), ArraySerializer.ItemMap.Num() );
 
 			ArraySerializer.ItemMap.Empty();
