Merge

CL 2505532 Moved custom depth checking to CanBeOccluded()
CL 2500364 Don't consider static meshes with precomputed shadows to be dynamic

diff --git a/Engine/Source/Runtime/Engine/Private/ModelRender.cpp b/Engine/Source/Runtime/Engine/Private/ModelRender.cpp
--- a/Engine/Source/Runtime/Engine/Private/ModelRender.cpp
+++ b/Engine/Source/Runtime/Engine/Private/ModelRender.cpp
@@ -484,7 +484,7 @@ public:
 						{
 							bLightMapped = false;
 						}
-						if(InteractionType != LIT_Dynamic && InteractionType != LIT_CachedSignedDistanceFieldShadowMap2D)
+						if(InteractionType != LIT_Dynamic)
 						{
 							bDynamic = false;
 						}
diff --git a/Engine/Source/Runtime/Engine/Private/SkeletalMesh.cpp b/Engine/Source/Runtime/Engine/Private/SkeletalMesh.cpp
--- a/Engine/Source/Runtime/Engine/Private/SkeletalMesh.cpp
+++ b/Engine/Source/Runtime/Engine/Private/SkeletalMesh.cpp
@@ -3754,6 +3754,7 @@ void ASkeletalMeshActor::PreviewSetAnimPosition(FName SlotName, int32 ChannelInd
 		SkeletalMeshComponent->RefreshSlaveComponents();
 		SkeletalMeshComponent->UpdateComponentToWorld();
 	}
+	SkeletalMeshComponent->FlipEditableSpaceBases();
 }
 
 
@@ -4543,7 +4544,7 @@ FPrimitiveViewRelevance FSkeletalMeshSceneProxy::GetViewRelevance(const FSceneVi
 
 bool FSkeletalMeshSceneProxy::CanBeOccluded() const
 {
-	return !MaterialRelevance.bDisableDepthTest;
+	return !MaterialRelevance.bDisableDepthTest && !ShouldRenderCustomDepth();
 }
 
 /** Util for getting LOD index currently used by this SceneProxy. */
@@ -4649,7 +4650,6 @@ void FSkeletalMeshSceneProxy::UpdateMorphMaterialUsage_GameThread(bool bNeedsMor
 	}
 }
 
-
 USkinnedMeshComponent::USkinnedMeshComponent(const FObjectInitializer& ObjectInitializer)
 	: Super(ObjectInitializer)
 	, AnimUpdateRateParams(NULL)
diff --git a/Engine/Source/Runtime/Engine/Private/StaticMeshRender.cpp b/Engine/Source/Runtime/Engine/Private/StaticMeshRender.cpp
--- a/Engine/Source/Runtime/Engine/Private/StaticMeshRender.cpp
+++ b/Engine/Source/Runtime/Engine/Private/StaticMeshRender.cpp
@@ -851,7 +851,7 @@ void FStaticMeshSceneProxy::OnTransformChanged()
 
 bool FStaticMeshSceneProxy::CanBeOccluded() const
 {
-	return !MaterialRelevance.bDisableDepthTest;
+	return !MaterialRelevance.bDisableDepthTest && !ShouldRenderCustomDepth();
 }
 
 FPrimitiveViewRelevance FStaticMeshSceneProxy::GetViewRelevance(const FSceneView* View)
@@ -934,7 +934,7 @@ void FStaticMeshSceneProxy::GetLightRelevance(const FLightSceneProxy* LightScene
 					bLightMapped = false;
 				}
 
-				if (InteractionType != LIT_Dynamic && InteractionType != LIT_CachedSignedDistanceFieldShadowMap2D)
+				if (InteractionType != LIT_Dynamic)
 				{
 					bDynamic = false;
 				}
diff --git a/Engine/Source/Runtime/Landscape/Private/LandscapeRender.cpp b/Engine/Source/Runtime/Landscape/Private/LandscapeRender.cpp
--- a/Engine/Source/Runtime/Landscape/Private/LandscapeRender.cpp
+++ b/Engine/Source/Runtime/Landscape/Private/LandscapeRender.cpp
@@ -884,7 +884,7 @@ void FLandscapeComponentSceneProxy::GetLightRelevance(const FLightSceneProxy* Li
 			bLightMapped = false;
 		}
 
-		if (InteractionType != LIT_Dynamic && InteractionType != LIT_CachedSignedDistanceFieldShadowMap2D)
+		if (InteractionType != LIT_Dynamic)
 		{
 			bDynamic = false;
 		}
