Temporarily fix monolithic builds by making Niagara module circularly dependent with Engine

diff --git a/Engine/Source/Runtime/Engine/Engine.Build.cs b/Engine/Source/Runtime/Engine/Engine.Build.cs
--- a/Engine/Source/Runtime/Engine/Engine.Build.cs
+++ b/Engine/Source/Runtime/Engine/Engine.Build.cs
@@ -67,13 +67,15 @@ public class Engine : ModuleRules
 				"SlateReflector",
 				"Landscape",
                 "UMG",
-				"Projects"
+				"Projects",
+				"Niagara"
 			}
         );
 
         CircularlyReferencedDependentModules.Add("AIModule");
 		CircularlyReferencedDependentModules.Add("Landscape");
         CircularlyReferencedDependentModules.Add("UMG");
+		CircularlyReferencedDependentModules.Add("Niagara");
 
 		DynamicallyLoadedModuleNames.AddRange(
 			new string[]
