PR #1049: HUD::GetActorsInSelectionRectangle comment fix (Contributed by [REDACTED_PERSON])

diff --git a/Engine/Source/Runtime/Engine/Classes/GameFramework/HUD.h b/Engine/Source/Runtime/Engine/Classes/GameFramework/HUD.h
--- a/Engine/Source/Runtime/Engine/Classes/GameFramework/HUD.h
+++ b/Engine/Source/Runtime/Engine/Classes/GameFramework/HUD.h
@@ -326,8 +326,8 @@ public:
 	 *
 	 * Sample usage:
 	 *
-	 *       TArray<AStaticMeshActor*> ActorsInSelectionRect;
-	 * 		Canvas->GetActorsInSelectionRectangle<AStaticMeshActor>(FirstPoint,SecondPoint,ActorsInSelectionRect);
+	 *      TArray<AStaticMeshActor*> ActorsInSelectionRect;
+	 * 		GetActorsInSelectionRectangle<AStaticMeshActor>(FirstPoint,SecondPoint,ActorsInSelectionRect);
 	 *
 	 *
 	 * @param FirstPoint					The first point, or anchor of the marquee box. Where the dragging of the marquee started in screen space.
