Fix for monolithic build fails

diff --git a/Engine/Source/Runtime/Engine/Classes/Curves/CurveVector.h b/Engine/Source/Runtime/Engine/Classes/Curves/CurveVector.h
--- a/Engine/Source/Runtime/Engine/Classes/Curves/CurveVector.h
+++ b/Engine/Source/Runtime/Engine/Classes/Curves/CurveVector.h
@@ -2,6 +2,7 @@
 
 
 #pragma once
+#include "CurveBase.h"
 #include "CurveVector.generated.h"
 
 UCLASS(BlueprintType, MinimalAPI)
diff --git a/Engine/Source/Runtime/Niagara/Private/NiagaraPrivate.h b/Engine/Source/Runtime/Niagara/Private/NiagaraPrivate.h
--- a/Engine/Source/Runtime/Niagara/Private/NiagaraPrivate.h
+++ b/Engine/Source/Runtime/Niagara/Private/NiagaraPrivate.h
@@ -2,5 +2,5 @@
 
 
 #pragma once
-#include "Engine.h"
+#include "Public/Engine.h"
 #include "VectorVM.h"
\ No newline at end of file
diff --git a/Engine/Source/Runtime/VectorVM/Public/VectorVM.h b/Engine/Source/Runtime/VectorVM/Public/VectorVM.h
--- a/Engine/Source/Runtime/VectorVM/Public/VectorVM.h
+++ b/Engine/Source/Runtime/VectorVM/Public/VectorVM.h
@@ -2,6 +2,7 @@
 
 #pragma once
 #include "Core.h"
+#include "CoreUObject.h"
 
 class FNiagaraDataObject;
 
