UE4 - simple safe fix for UE-13762, an issue with static mesh registration with null render data and speedtree code. Was crashing an automated test.

diff --git a/Engine/Source/Runtime/Engine/Private/Components/StaticMeshComponent.cpp b/Engine/Source/Runtime/Engine/Private/Components/StaticMeshComponent.cpp
--- a/Engine/Source/Runtime/Engine/Private/Components/StaticMeshComponent.cpp
+++ b/Engine/Source/Runtime/Engine/Private/Components/StaticMeshComponent.cpp
@@ -463,22 +463,20 @@ void UStaticMeshComponent::OnRegister()
 				LODData[i].ShadowMap = NULL;
 			}
 		}
-	}
-
-	if (StaticMesh != NULL && StaticMesh->SpeedTreeWind.IsValid() && GetScene())
-	{
-		for (int32 LODIndex = 0; LODIndex < StaticMesh->RenderData->LODResources.Num(); ++LODIndex)
+		if (StaticMesh->SpeedTreeWind.IsValid() && GetScene())
 		{
-			GetScene()->AddSpeedTreeWind(&StaticMesh->RenderData->LODResources[LODIndex].VertexFactory, StaticMesh);
+			for (int32 LODIndex = 0; LODIndex < StaticMesh->RenderData->LODResources.Num(); ++LODIndex)
+			{
+				GetScene()->AddSpeedTreeWind(&StaticMesh->RenderData->LODResources[LODIndex].VertexFactory, StaticMesh);
+			}
 		}
 	}
-
 	Super::OnRegister();
 }
 
 void UStaticMeshComponent::OnUnregister()
 {
-	if (StaticMesh != NULL && StaticMesh->SpeedTreeWind.IsValid() && GetScene())
+	if (StaticMesh && StaticMesh->RenderData && StaticMesh->SpeedTreeWind.IsValid() && GetScene())
 	{
 		for (int32 LODIndex = 0; LODIndex < StaticMesh->RenderData->LODResources.Num(); ++LODIndex)
 		{
