Fix backwards compatibility for EditorViews

diff --git a/Engine/Source/Runtime/Engine/Private/World.cpp b/Engine/Source/Runtime/Engine/Private/World.cpp
--- a/Engine/Source/Runtime/Engine/Private/World.cpp
+++ b/Engine/Source/Runtime/Engine/Private/World.cpp
@@ -149,6 +149,7 @@ void UWorld::Serialize( FArchive& Ar )
 
 	if (Ar.UE4Ver() < VER_UE4_ADD_EDITOR_VIEWS)
 	{
+		EditorViews.SetNum(4);
 		for (int32 i = 0; i < 4; ++i)
 		{
 			FLevelViewportInfo TempViewportInfo;
@@ -156,7 +157,7 @@ void UWorld::Serialize( FArchive& Ar )
 #if WITH_EDITOR
 			if (Ar.IsLoading())
 			{
-				EditorViews.Add(TempViewportInfo);
+				EditorViews[i] = TempViewportInfo;
 			}
 #endif
 		}
