Fix issues with EndOfFrame state being "stuck" after undo

diff --git a/Engine/Source/Runtime/Engine/Classes/Engine/World.h b/Engine/Source/Runtime/Engine/Classes/Engine/World.h
--- a/Engine/Source/Runtime/Engine/Classes/Engine/World.h
+++ b/Engine/Source/Runtime/Engine/Classes/Engine/World.h
@@ -2573,6 +2573,13 @@ public:
 	 */
 	void MarkActorComponentForNeededEndOfFrameUpdate(class UActorComponent* Component, bool bForceGameThread);
 
+	/**
+	 * Updates an ActorComponent's cached state of whether it has been marked for end of frame update based on the current
+	 * state of the World's NeedsEndOfFrameUpdate arrays
+	 * @param Component - Component to update the cached state of
+	 */
+	void UpdateActorComponentEndOfFrameUpdateState(class UActorComponent* Component) const;
+
 	/**
 	 * Send all render updates to the rendering thread.
 	 * @param OutCompletion - all async updates are added to this array, they typically need to be completed before GC or anything else
diff --git a/Engine/Source/Runtime/Engine/Private/Components/ActorComponent.cpp b/Engine/Source/Runtime/Engine/Private/Components/ActorComponent.cpp
--- a/Engine/Source/Runtime/Engine/Private/Components/ActorComponent.cpp
+++ b/Engine/Source/Runtime/Engine/Private/Components/ActorComponent.cpp
@@ -520,6 +520,8 @@ void UActorComponent::PostEditUndo()
 		{
 			Owner->AddOwnedComponent(this);
 		}
+
+		GetWorld()->UpdateActorComponentEndOfFrameUpdateState(this);
 	}
 	Super::PostEditUndo();
 }
diff --git a/Engine/Source/Runtime/Engine/Private/LevelTick.cpp b/Engine/Source/Runtime/Engine/Private/LevelTick.cpp
--- a/Engine/Source/Runtime/Engine/Private/LevelTick.cpp
+++ b/Engine/Source/Runtime/Engine/Private/LevelTick.cpp
@@ -729,6 +729,22 @@ private:
 	}
 };
 
+void UWorld::UpdateActorComponentEndOfFrameUpdateState(UActorComponent* Component) const
+{
+	if (ComponentsThatNeedEndOfFrameUpdate.Contains(Component))
+	{
+		FMarkComponentEndOfFrameUpdateState::Set(Component, EComponentMarkedForEndOfFrameUpdateState::Marked);
+	}
+	else if (ComponentsThatNeedEndOfFrameUpdate_OnGameThread.Contains(Component))
+	{
+		FMarkComponentEndOfFrameUpdateState::Set(Component, EComponentMarkedForEndOfFrameUpdateState::MarkedForGameThread);
+	}
+	else
+	{
+		FMarkComponentEndOfFrameUpdateState::Set(Component, EComponentMarkedForEndOfFrameUpdateState::Unmarked);
+	}
+}
+
 void UWorld::MarkActorComponentForNeededEndOfFrameUpdate(UActorComponent* Component, bool bForceGameThread)
 {
 	check(!bPostTickComponentUpdate); // can't call this while we are doing the updates
