Fix component instance data cache crash during GC due to null Light or Shadow Map refs

diff --git a/Engine/Source/Runtime/Engine/Private/Components/StaticMeshComponent.cpp b/Engine/Source/Runtime/Engine/Private/Components/StaticMeshComponent.cpp
--- a/Engine/Source/Runtime/Engine/Private/Components/StaticMeshComponent.cpp
+++ b/Engine/Source/Runtime/Engine/Private/Components/StaticMeshComponent.cpp
@@ -48,11 +48,17 @@ public:
 		{
 			for (FLightMapRef& LightMapRef : CachedStaticLighting.LODDataLightMap)
 			{
-				LightMapRef->AddReferencedObjects(Collector);
+				if (LightMapRef != nullptr)
+				{
+					LightMapRef->AddReferencedObjects(Collector);
+				}
 			}
 			for (FShadowMapRef& ShadowMapRef : CachedStaticLighting.LODDataShadowMap)
 			{
-				ShadowMapRef->AddReferencedObjects(Collector);
+				if (ShadowMapRef != nullptr)
+				{
+					ShadowMapRef->AddReferencedObjects(Collector);
+				}
 			}		
 		}
 	}
diff --git a/Engine/Source/Runtime/Engine/Private/InstancedStaticMesh.cpp b/Engine/Source/Runtime/Engine/Private/InstancedStaticMesh.cpp
--- a/Engine/Source/Runtime/Engine/Private/InstancedStaticMesh.cpp
+++ b/Engine/Source/Runtime/Engine/Private/InstancedStaticMesh.cpp
@@ -758,11 +758,17 @@ public:
 		{
 			for (FLightMapRef& LightMapRef : CachedStaticLighting.LODDataLightMap)
 			{
-				LightMapRef->AddReferencedObjects(Collector);
+				if (LightMapRef != nullptr)
+				{
+					LightMapRef->AddReferencedObjects(Collector);
+				}
 			}
 			for (FShadowMapRef& ShadowMapRef : CachedStaticLighting.LODDataShadowMap)
 			{
-				ShadowMapRef->AddReferencedObjects(Collector);
+				if (ShadowMapRef != nullptr)
+				{
+					ShadowMapRef->AddReferencedObjects(Collector);
+				}
 			}		
 		}
 	}
