Expose generic axis config property functions
Some efficiency changes in PlayerInput

diff --git a/Engine/Source/Runtime/Engine/Classes/GameFramework/CheatManager.h b/Engine/Source/Runtime/Engine/Classes/GameFramework/CheatManager.h
--- a/Engine/Source/Runtime/Engine/Classes/GameFramework/CheatManager.h
+++ b/Engine/Source/Runtime/Engine/Classes/GameFramework/CheatManager.h
@@ -304,6 +304,14 @@ class ENGINE_API UCheatManager : public UObject
 	UFUNCTION(exec)
 	void SetWorldOrigin();
 
+	/** Exec function to return the mouse sensitivity to its default value */
+	UFUNCTION(exec)
+	void SetMouseSensitivityToDefault();
+
+	/** Backwards compatibility exec function for people used to it instead of using InvertAxisKey */
+	UFUNCTION(exec)
+	void InvertMouse();
+
 	/**
 	 * This will move the player and set their rotation to the passed in values.
 	 * This actually does the location / rotation setting.  Additionally it will set you as ghost as the level may have
diff --git a/Engine/Source/Runtime/Engine/Classes/GameFramework/PlayerInput.h b/Engine/Source/Runtime/Engine/Classes/GameFramework/PlayerInput.h
--- a/Engine/Source/Runtime/Engine/Classes/GameFramework/PlayerInput.h
+++ b/Engine/Source/Runtime/Engine/Classes/GameFramework/PlayerInput.h
@@ -247,9 +247,12 @@ struct FInputAxisKeyMapping
 UCLASS(Within=PlayerController, config=Input, transient)
 class ENGINE_API UPlayerInput : public UObject
 {
-	GENERATED_UCLASS_BODY()
+	GENERATED_BODY()
 
 public:
+
+	UPlayerInput();
+
 	// NOTE: These touch vectors are calculated and set directly, they do not go through the .ini Bindings
 	// Touch locations, from 0..1 (0,0 is top left, 1,1 is bottom right), the Z component is > 0 if the touch is currently held down
 	// @todo: We have 10 touches to match the number of Touch* entries in EKeys (not easy to make this an enum or define or anything)
@@ -262,7 +265,7 @@ public:
 	float MouseSamplingTotal;    /** DirectInput's mouse sampling total time */
 
 
-protected:
+private:
 	TEnumAsByte<EInputEvent> CurrentEvent;
 
 public:
@@ -283,12 +286,18 @@ public:
 	UPROPERTY(config)
 	TArray<FName> InvertedAxis;
 
+	/** Gets the axis properties for a given AxisKey.  Returns if true if AxisKey was found in the AxisConfig array. */
+	bool GetAxisProperties(const FKey AxisKey, FInputAxisProperties& AxisProperties);
+
+	/** Gets the axis properties for a given AxisKey.  Returns if true if AxisKey was found in the AxisConfig array. */
+	void SetAxisProperties(const FKey AxisKey, const FInputAxisProperties& AxisProperties);
+
 	/** Exec function to change the mouse sensitivity */
 	UFUNCTION(exec)
 	void SetMouseSensitivity(const float Sensitivity);
 
 	/** Exec function to return the mouse sensitivity to its default value */
-	UFUNCTION(exec)
+	DEPRECATED(4.8, "SetMouseSensitivityToDefault is deprecated, use SetAxisProperties instead")
 	void SetMouseSensitivityToDefault();
 
 	/** Exec function to add a debug exec command */
@@ -296,6 +305,7 @@ public:
 	void SetBind(FName BindName, const FString& Command);
 
 	/** Returns the mouse sensitivity along the X-axis, or the Y-axis, or 1.0 if none are known. */
+	DEPRECATED(4.8, "GetMouseSensitivity is deprecated, use GetAxisProperties instead")
 	float GetMouseSensitivity();
 
 	/** Returns whether an Axis Key is inverted */
@@ -309,7 +319,7 @@ public:
 	void InvertAxisKey(const FKey AxisKey);
 
 	/** Backwards compatibility exec function for people used to it instead of using InvertAxisKey */
-	UFUNCTION(exec)
+	DEPRECATED(4.8, "InvertMouse is deprecated, use InvertAxisKey(EKeys::MouseY); instead")
 	void InvertMouse();
 
 	/** Exec function to invert an axis mapping */
@@ -341,7 +351,7 @@ public:
 	/** Clear the current cached key maps and rebuild from the source arrays. */
 	void ForceRebuildingKeyMaps(const bool bRestoreDefaults = false);
 
-protected:
+private:
 
 	/** Runtime struct that caches the list of mappings for a given Action Name and the capturing chord if applicable */
 	struct FActionKeyDetails
@@ -491,7 +501,7 @@ public:
 	/** Returns the list of keys mapped to the specified Action Name */
 	const TArray<FInputActionKeyMapping>& GetKeysForAction(const FName ActionName);
 
-protected:
+private:
 	/** 
 	 * Given raw keystate value, returns the "massaged" value. Override for any custom behavior,
 	 * such as input changes dependent on a particular game state.
@@ -551,7 +561,6 @@ protected:
 	FGestureRecognizer GestureRecognizer;
 	friend FGestureRecognizer;
 
-private:
 	/** Static empty array to be able to return from GetKeysFromAction when there are no keys mapped to the requested action name */
 	static const TArray<FInputActionKeyMapping> NoKeyMappings;
 
diff --git a/Engine/Source/Runtime/Engine/Private/CheatManager.cpp b/Engine/Source/Runtime/Engine/Private/CheatManager.cpp
--- a/Engine/Source/Runtime/Engine/Private/CheatManager.cpp
+++ b/Engine/Source/Runtime/Engine/Private/CheatManager.cpp
@@ -1118,6 +1118,33 @@ void UCheatManager::SetWorldOrigin()
 	World->RequestNewWorldOrigin(NewOrigin);
 }
 
+void UCheatManager::SetMouseSensitivityToDefault()
+{
+	UPlayerInput* PlayerInput = GetOuterAPlayerController()->PlayerInput;
+	if (PlayerInput)
+	{
+		// find default sensitivity restore to that
+		for (const FInputAxisConfigEntry& AxisConfigEntry : GetDefault<UInputSettings>()->AxisConfig)
+		{
+			const FKey AxisKey = AxisConfigEntry.AxisKeyName;
+			if (AxisKey == EKeys::MouseX)
+			{
+				PlayerInput->SetMouseSensitivity(AxisConfigEntry.AxisProperties.Sensitivity);
+				break;
+			}
+		}
+	}
+}
+
+void UCheatManager::InvertMouse()
+{
+	UPlayerInput* PlayerInput = GetOuterAPlayerController()->PlayerInput;
+	if (PlayerInput)
+	{
+		PlayerInput->InvertAxisKey(EKeys::MouseY);
+	}
+}
+
 void UCheatManager::LogOutBugItGoToLogFile( const FString& InScreenShotDesc, const FString& InGoString, const FString& InLocString )
 {
 #if ALLOW_DEBUG_FILES
diff --git a/Engine/Source/Runtime/Engine/Private/UserInterface/PlayerInput.cpp b/Engine/Source/Runtime/Engine/Private/UserInterface/PlayerInput.cpp
--- a/Engine/Source/Runtime/Engine/Private/UserInterface/PlayerInput.cpp
+++ b/Engine/Source/Runtime/Engine/Private/UserInterface/PlayerInput.cpp
@@ -66,8 +66,8 @@ struct FDelegateDispatchDetails
 	{}
 };
 
-UPlayerInput::UPlayerInput(const FObjectInitializer& ObjectInitializer)
-	: Super(ObjectInitializer)
+UPlayerInput::UPlayerInput()
+	: Super()
 {
 	SetFlags(RF_Transactional);
 	MouseSamplingTotal = +0.0083f;
@@ -298,35 +298,57 @@ bool UPlayerInput::InputGesture(const FKey Gesture, const EInputEvent Event, con
 	return true;
 }
 
-float UPlayerInput::GetMouseSensitivity()
+bool UPlayerInput::GetAxisProperties(const FKey AxisKey, FInputAxisProperties& OutAxisProperties)
 {
 	ConditionalInitAxisProperties();
 
-	const FInputAxisProperties* const XAxisProps = AxisProperties.Find(EKeys::MouseX);
-	if (XAxisProps)
-	{
-		return XAxisProps->Sensitivity;
-	}
-	const FInputAxisProperties* const YAxisProps = AxisProperties.Find(EKeys::MouseY);
-	if (YAxisProps)
+	const FInputAxisProperties* const AxisProps = AxisProperties.Find(AxisKey);
+	if (AxisProps)
 	{
-		return YAxisProps->Sensitivity;
+		OutAxisProperties = *AxisProps;
+		return true;
 	}
-	return 1.0f;
+
+	return false;
 }
 
-void UPlayerInput::SetMouseSensitivity(const float Sensitivity)
+void UPlayerInput::SetAxisProperties(const FKey AxisKey, const FInputAxisProperties& InAxisProperties)
 {
 	for (FInputAxisConfigEntry& AxisConfigEntry : AxisConfig)
 	{
-		const FKey AxisKey = AxisConfigEntry.AxisKeyName;
-		if (AxisKey == EKeys::MouseX || AxisKey == EKeys::MouseY)
+		if (AxisConfigEntry.AxisKeyName == AxisKey)
 		{
-			AxisConfigEntry.AxisProperties.Sensitivity = Sensitivity;
+			AxisConfigEntry.AxisProperties = InAxisProperties;
 		}
 	}
 
-	AxisProperties.Empty();
+	AxisProperties.Empty(AxisProperties.Num());
+}
+
+float UPlayerInput::GetMouseSensitivity()
+{
+	FInputAxisProperties MouseAxisProps;
+	if (GetAxisProperties(EKeys::MouseX, MouseAxisProps) || GetAxisProperties(EKeys::MouseY, MouseAxisProps))
+	{
+		return MouseAxisProps.Sensitivity;
+	}
+
+	return 1.f;
+}
+
+void UPlayerInput::SetMouseSensitivity(const float Sensitivity)
+{
+	FInputAxisProperties MouseAxisProps;
+	if (GetAxisProperties(EKeys::MouseX, MouseAxisProps))
+	{
+		MouseAxisProps.Sensitivity = Sensitivity;
+		SetAxisProperties(EKeys::MouseX, MouseAxisProps);
+	}
+	if (GetAxisProperties(EKeys::MouseY, MouseAxisProps))
+	{
+		MouseAxisProps.Sensitivity = Sensitivity;
+		SetAxisProperties(EKeys::MouseY, MouseAxisProps);
+	}
 }
 
 void UPlayerInput::SetMouseSensitivityToDefault()
@@ -337,7 +359,9 @@ void UPlayerInput::SetMouseSensitivityToDefault()
 		const FKey AxisKey = AxisConfigEntry.AxisKeyName;
 		if (AxisKey == EKeys::MouseX)
 		{
+			PRAGMA_DISABLE_DEPRECATION_WARNINGS
 			SetMouseSensitivity(AxisConfigEntry.AxisProperties.Sensitivity);
+			PRAGMA_ENABLE_DEPRECATION_WARNINGS
 			break;
 		}
 	}
@@ -448,7 +472,7 @@ struct FAxisDelegate
 void UPlayerInput::AddActionMapping(const FInputActionKeyMapping& KeyMapping)
 {
 	ActionMappings.AddUnique(KeyMapping);
-	ActionKeyMap.Empty();
+	ActionKeyMap.Empty(ActionKeyMap.Num());
 }
 
 void UPlayerInput::RemoveActionMapping(const FInputActionKeyMapping& KeyMapping)
@@ -458,7 +482,7 @@ void UPlayerInput::RemoveActionMapping(const FInputActionKeyMapping& KeyMapping)
 		if (ActionMappings[ActionIndex] == KeyMapping)
 		{
 			ActionMappings.RemoveAtSwap(ActionIndex);
-			ActionKeyMap.Empty();
+			ActionKeyMap.Empty(ActionKeyMap.Num());
 			// we don't break because the mapping may have been in the array twice
 		}
 	}
@@ -467,7 +491,7 @@ void UPlayerInput::RemoveActionMapping(const FInputActionKeyMapping& KeyMapping)
 void UPlayerInput::AddAxisMapping(const FInputAxisKeyMapping& KeyMapping)
 {
 	AxisMappings.AddUnique(KeyMapping);
-	AxisKeyMap.Empty();
+	AxisKeyMap.Empty(AxisKeyMap.Num());
 }
 
 void UPlayerInput::RemoveAxisMapping(const FInputAxisKeyMapping& InKeyMapping)
@@ -479,7 +503,7 @@ void UPlayerInput::RemoveAxisMapping(const FInputAxisKeyMapping& InKeyMapping)
 			&& KeyMapping.Key == InKeyMapping.Key)
 		{
 			AxisMappings.RemoveAtSwap(AxisIndex);
-			AxisKeyMap.Empty();
+			AxisKeyMap.Empty(AxisKeyMap.Num());
 			// we don't break because the mapping may have been in the array twice
 		}
 	}
@@ -490,7 +514,7 @@ void UPlayerInput::AddEngineDefinedActionMapping(const FInputActionKeyMapping& A
 	EngineDefinedActionMappings.AddUnique(ActionMapping);
 	for (TObjectIterator<UPlayerInput> It; It; ++It)
 	{
-		It->ActionKeyMap.Empty();
+		It->ActionKeyMap.Empty(It->ActionKeyMap.Num());
 	}
 }
 
@@ -499,7 +523,7 @@ void UPlayerInput::AddEngineDefinedAxisMapping(const FInputAxisKeyMapping& AxisM
 	EngineDefinedAxisMappings.AddUnique(AxisMapping);
 	for (TObjectIterator<UPlayerInput> It; It; ++It)
 	{
-		It->AxisKeyMap.Empty();
+		It->AxisKeyMap.Empty(It->AxisKeyMap.Num());
 	}
 }
 
@@ -512,9 +536,9 @@ void UPlayerInput::ForceRebuildingKeyMaps(const bool bRestoreDefaults)
 		ActionMappings = GetDefault<UInputSettings>()->ActionMappings;
 	}
 
-	ActionKeyMap.Empty();
-	AxisKeyMap.Empty();
-	AxisProperties.Empty();
+	ActionKeyMap.Empty(ActionKeyMap.Num());
+	AxisKeyMap.Empty(AxisKeyMap.Num());
+	AxisProperties.Empty(AxisProperties.Num());
 }
 
 void UPlayerInput::ConditionalBuildKeyMappings()
@@ -779,8 +803,7 @@ void UPlayerInput::ProcessInputStack(const TArray<UInputComponent*>& InputCompon
 
 		for (uint8 EventIndex = 0; EventIndex < IE_MAX; ++EventIndex)
 		{
-			KeyState->EventCounts[EventIndex] = KeyState->EventAccumulator[EventIndex];
-			KeyState->EventAccumulator[EventIndex].Empty();
+			KeyState->EventCounts[EventIndex] = MoveTemp(KeyState->EventAccumulator[EventIndex]);
 		}
 
 		if ( (KeyState->SampleCountAccumulator > 0) || (Key == EKeys::MouseX) || (Key == EKeys::MouseY) )
