#ANIM: fix crash with changing retargeting translation option for newly imported animation
#ANIM: fix crash with animseqstats

diff --git a/Engine/Source/Runtime/Engine/Private/Animation/AnimEncoding.cpp b/Engine/Source/Runtime/Engine/Private/Animation/AnimEncoding.cpp
--- a/Engine/Source/Runtime/Engine/Private/Animation/AnimEncoding.cpp
+++ b/Engine/Source/Runtime/Engine/Private/Animation/AnimEncoding.cpp
@@ -218,7 +218,10 @@ void AnimationFormat_GetAnimationPose(
 {
 	// decompress the translation component using the proper method
 	checkSlow(Seq.TranslationCodec != NULL);
-	((AnimEncoding*)Seq.TranslationCodec)->GetPoseTranslations(Atoms, TranslationPairs, Seq, Time);
+	if (TranslationPairs.Num() > 0)
+	{
+		((AnimEncoding*)Seq.TranslationCodec)->GetPoseTranslations(Atoms, TranslationPairs, Seq, Time);
+	}
 
 	// decompress the rotation component using the proper method
 	checkSlow(Seq.RotationCodec != NULL);
@@ -525,9 +528,9 @@ void AnimationFormat_GetStats(
 			}
 
 			// Scale.
-			for ( int32 TrackIndex = 0; TrackIndex < NumTransTracks; ++TrackIndex )
+			for ( int32 ScaleIndex = 0; ScaleIndex < NumScaleTracks; ++ScaleIndex )
 			{
-				const int32 NumKeys = Seq->CompressedScaleOffsets.GetOffsetData(TrackIndex, 1);
+				const int32 NumKeys = Seq->CompressedScaleOffsets.GetOffsetData(ScaleIndex, 1);
 				TotalNumScaleKeys += NumKeys;
 				if ( NumKeys == 1 )
 				{
