Engine: Prevented a crash when a skeletal mesh LOD has zero vertices

diff --git a/Engine/Source/Runtime/Engine/Private/SkeletalMesh.cpp b/Engine/Source/Runtime/Engine/Private/SkeletalMesh.cpp
--- a/Engine/Source/Runtime/Engine/Private/SkeletalMesh.cpp
+++ b/Engine/Source/Runtime/Engine/Private/SkeletalMesh.cpp
@@ -249,9 +249,9 @@ FArchive& operator<<(FArchive& Ar,FSkeletalMeshVertexBuffer& VertexBuffer)
 			VertexBuffer.VertexData->Serialize(Ar);	
 
 			// update cached buffer info
-			VertexBuffer.Data = VertexBuffer.VertexData->GetDataPointer();
-			VertexBuffer.Stride = VertexBuffer.VertexData->GetStride();
 			VertexBuffer.NumVertices = VertexBuffer.VertexData->GetNumVertices();
+			VertexBuffer.Data = (VertexBuffer.NumVertices > 0) ? VertexBuffer.VertexData->GetDataPointer() : nullptr;
+			VertexBuffer.Stride = VertexBuffer.VertexData->GetStride();
 		}
 	}
 
@@ -672,9 +672,9 @@ FArchive& operator<<( FArchive& Ar, FSkeletalMeshVertexAPEXClothBuffer& VertexBu
 			VertexBuffer.VertexData->Serialize( Ar );
 
 			// update cached buffer info
-			VertexBuffer.Data = VertexBuffer.VertexData->GetDataPointer();
-			VertexBuffer.Stride = VertexBuffer.VertexData->GetStride();
 			VertexBuffer.NumVertices = VertexBuffer.VertexData->GetNumVertices();
+			VertexBuffer.Data = (VertexBuffer.NumVertices > 0) ? VertexBuffer.VertexData->GetDataPointer() : nullptr;
+			VertexBuffer.Stride = VertexBuffer.VertexData->GetStride();
 		}
 	}
 
