Remove 0-1 constraint on restitution. These can be valid tuning values used by certain games.

diff --git a/Engine/Source/Runtime/Engine/Classes/PhysicalMaterials/PhysicalMaterial.h b/Engine/Source/Runtime/Engine/Classes/PhysicalMaterials/PhysicalMaterial.h
--- a/Engine/Source/Runtime/Engine/Classes/PhysicalMaterials/PhysicalMaterial.h
+++ b/Engine/Source/Runtime/Engine/Classes/PhysicalMaterials/PhysicalMaterial.h
@@ -54,8 +54,8 @@ class ENGINE_API UPhysicalMaterial : public UObject
 	UPROPERTY(EditAnywhere, BlueprintReadWrite, Category = PhysicalMaterial)
 	bool bOverrideFrictionCombineMode;
 
-	/** Restitution or 'bouncyness' of this surface. */
-	UPROPERTY(EditAnywhere, BlueprintReadOnly, Category = PhysicalMaterial, meta = (UIMin=0, UIMax=1, ClampMin=0, ClampMax=1))
+	/** Restitution or 'bounciness of this surface. */
+	UPROPERTY(EditAnywhere, BlueprintReadOnly, Category = PhysicalMaterial)
 	float Restitution;
 
 	/** Restitution combine mode, controls how restitution is computed for multiple materials. */
