When welding to a skeletalmesh, make sure the new shapes are considered part of the bone it was welded to. Fixes jitter for overlapping bodies

diff --git a/Engine/Source/Runtime/Engine/Private/PhysicsEngine/BodyInstance.cpp b/Engine/Source/Runtime/Engine/Private/PhysicsEngine/BodyInstance.cpp
--- a/Engine/Source/Runtime/Engine/Private/PhysicsEngine/BodyInstance.cpp
+++ b/Engine/Source/Runtime/Engine/Private/PhysicsEngine/BodyInstance.cpp
@@ -664,7 +664,7 @@ void FBodyInstance::UpdatePhysicsShapeFilterData(uint32 SkelMeshCompID, bool bUs
 			PxFilterData PComplexQueryData;
 			if (UseCollisionEnabled != ECollisionEnabled::NoCollision)
 			{
-				CreateShapeFilterData(BI->ObjectType, CompID, UseResponse, SkelMeshCompID, BI->InstanceBodyIndex, PSimpleQueryData, PSimFilterData, bUseCCD && !bPhysicsStatic, bUseNotify, bPhysicsStatic);	//CCD is determined by root body in case of welding
+				CreateShapeFilterData(BI->ObjectType, CompID, UseResponse, SkelMeshCompID, InstanceBodyIndex, PSimpleQueryData, PSimFilterData, bUseCCD && !bPhysicsStatic, bUseNotify, bPhysicsStatic);	//InstanceBodyIndex and CCD are determined by root body in case of welding
 				PComplexQueryData = PSimpleQueryData;
 
 				// Build filterdata variations for complex and simple
