Refactor bodysetup shape addition. This will make things like shape sharing easier

diff --git a/Engine/Source/Runtime/Engine/Classes/PhysicsEngine/BodySetup.h b/Engine/Source/Runtime/Engine/Classes/PhysicsEngine/BodySetup.h
--- a/Engine/Source/Runtime/Engine/Classes/PhysicsEngine/BodySetup.h
+++ b/Engine/Source/Runtime/Engine/Classes/PhysicsEngine/BodySetup.h
@@ -227,13 +227,8 @@ public:
 	 *   Add the shapes defined by this body setup to the supplied PxRigidBody. 
 	 */
 	ENGINE_API void AddShapesToRigidActor_AssumesLocked(FBodyInstance* OwningInstance, physx::PxRigidActor* PDestActor, EPhysicsSceneType SceneType, FVector& Scale3D, physx::PxMaterial* SimpleMaterial, TArray<UPhysicalMaterial*>& ComplexMaterials, FShapeData& ShapeData, const FTransform& RelativeTM = FTransform::Identity, TArray<physx::PxShape*>* NewShapes = NULL);
-
-private:
-	void AddSpheresToRigidActor_AssumesLocked(physx::PxRigidActor* PDestActor, const FTransform& RelativeTM, float MinScale, float MinScaleAbs, TArray<physx::PxShape*>* NewShapes) const;
-	void AddBoxesToRigidActor_AssumesLocked(physx::PxRigidActor* PDestActor, const FTransform& RelativeTM, const FVector& Scale3D, const FVector& Scale3DAbs, TArray<physx::PxShape*>* NewShapes) const;
-	void AddSphylsToRigidActor_AssumesLocked(physx::PxRigidActor* PDestActor, const FTransform& RelativeTM, const FVector& Scale3D, const FVector& Scale3DAbs, TArray<physx::PxShape*>* NewShapes) const;
-	void AddConvexElemsToRigidActor_AssumesLocked(physx::PxRigidActor* PDestActor, const FTransform& RelativeTM, const FVector& Scale3D, const FVector& Scale3DAbs, TArray<physx::PxShape*>* NewShapes) const;
-	void AddTriMeshToRigidActor_AssumesLocked(physx::PxRigidActor* PDestActor, const FVector& Scale3D, const FVector& Scale3DAbs, TArray<physx::PxShape*>* NewShapes) const;
 #endif // WITH_PHYSX
 
+	friend struct FAddShapesHelper;
+
 };
diff --git a/Engine/Source/Runtime/Engine/Private/PhysicsEngine/BodySetup.cpp b/Engine/Source/Runtime/Engine/Private/PhysicsEngine/BodySetup.cpp
--- a/Engine/Source/Runtime/Engine/Private/PhysicsEngine/BodySetup.cpp
+++ b/Engine/Source/Runtime/Engine/Private/PhysicsEngine/BodySetup.cpp
@@ -313,282 +313,300 @@ PxMaterial* GetDefaultPhysMaterial()
 	return GEngine->DefaultPhysMaterial->GetPhysXMaterial();
 }
 
-void UBodySetup::AddSpheresToRigidActor_AssumesLocked(PxRigidActor* PDestActor, const FTransform& RelativeTM, float MinScale, float MinScaleAbs, TArray<PxShape*>* NewShapes) const
+/** Helper struct for adding shapes into an actor.*/
+struct FAddShapesHelper
 {
-	float ContactOffsetFactor, MaxContactOffset;
-	GetContactOffsetParams(ContactOffsetFactor, MaxContactOffset);
-	PxMaterial* PDefaultMat = GetDefaultPhysMaterial();
-
-	for (int32 i = 0; i < AggGeom.SphereElems.Num(); i++)
+	FAddShapesHelper(UBodySetup* InBodySetup, FBodyInstance* InOwningInstance, physx::PxRigidActor* InPDestActor, FVector& InScale3D, physx::PxMaterial* InSimpleMaterial, TArray<UPhysicalMaterial*>& InComplexMaterials, FShapeData& InShapeData, const FTransform& InRelativeTM, TArray<physx::PxShape*>* InNewShapes)
+	: BodySetup(InBodySetup)
+	, OwningInstance(InOwningInstance)
+	, PDestActor(InPDestActor)
+	, Scale3D(InScale3D)
+	, SimpleMaterial(InSimpleMaterial)
+	, ComplexMaterials(InComplexMaterials)
+	, ShapeData(InShapeData)
+	, RelativeTM(InRelativeTM)
+	, NewShapes(InNewShapes)
 	{
-		const FKSphereElem& SphereElem = AggGeom.SphereElems[i];
-
-		PxSphereGeometry PSphereGeom;
-		PSphereGeom.radius = (SphereElem.Radius * MinScaleAbs);
-
-		if (ensure(PSphereGeom.isValid()))
+		SetupNonUniformHelper(Scale3D, MinScale, MinScaleAbs, Scale3DAbs);
 		{
-			FVector LocalOrigin = RelativeTM.TransformPosition(SphereElem.Center);
-			PxTransform PLocalPose(U2PVector(LocalOrigin));
-			PLocalPose.p *= MinScale;
+			float MinScaleRelative;
+			float MinScaleAbsRelative;
+			FVector Scale3DAbsRelative;
+			FVector Scale3DRelative = RelativeTM.GetScale3D();
 
-			ensure(PLocalPose.isValid());
-			{
-				PxShape* NewShape = PDestActor->createShape(PSphereGeom, *PDefaultMat, PLocalPose);
-				if (NewShapes)
-				{
-					NewShapes->Add(NewShape);
-				}
+			SetupNonUniformHelper(Scale3DRelative, MinScaleRelative, MinScaleAbsRelative, Scale3DAbsRelative);
 
-				const float ContactOffset = FMath::Min(MaxContactOffset, ContactOffsetFactor * PSphereGeom.radius);
-				NewShape->setContactOffset(ContactOffset);
-			}
-			
-
-			
-		}
-		else
-		{
-			UE_LOG(LogPhysics, Warning, TEXT("AddSpheresToRigidActor: [%s] SphereElem[%d] invalid"), *GetPathNameSafe(GetOuter()), i);
+			MinScaleAbs *= MinScaleAbsRelative;
+			Scale3DAbs.X *= Scale3DAbsRelative.X;
+			Scale3DAbs.Y *= Scale3DAbsRelative.Y;
+			Scale3DAbs.Z *= Scale3DAbsRelative.Z;
 		}
 	}
-}
 
-void UBodySetup::AddBoxesToRigidActor_AssumesLocked(PxRigidActor* PDestActor, const FTransform& RelativeTM, const FVector& Scale3D, const FVector& Scale3DAbs, TArray<PxShape*>* NewShapes) const
-{
-	float ContactOffsetFactor, MaxContactOffset;
-	GetContactOffsetParams(ContactOffsetFactor, MaxContactOffset);
-	PxMaterial* PDefaultMat = GetDefaultPhysMaterial();
+	UBodySetup* BodySetup;
+	FBodyInstance* OwningInstance;
+	physx::PxRigidActor* PDestActor;
+	FVector& Scale3D;
+	const physx::PxMaterial* SimpleMaterial;
+	const TArray<UPhysicalMaterial*>& ComplexMaterials;
+	const FShapeData& ShapeData;
+	const FTransform& RelativeTM;
+	TArray<physx::PxShape*>* NewShapes;
 
-	for (int32 i = 0; i < AggGeom.BoxElems.Num(); i++)
-	{
-		const FKBoxElem& BoxElem = AggGeom.BoxElems[i];
+	float MinScaleAbs;
+	float MinScale;
+	FVector Scale3DAbs;
 
-		PxBoxGeometry PBoxGeom;
-		PBoxGeom.halfExtents.x = (0.5f * BoxElem.X * Scale3DAbs.X);
-		PBoxGeom.halfExtents.y = (0.5f * BoxElem.Y * Scale3DAbs.Y);
-		PBoxGeom.halfExtents.z = (0.5f * BoxElem.Z * Scale3DAbs.Z);
+public:
+	void AddSpheresToRigidActor_AssumesLocked() const
+	{
+		float ContactOffsetFactor, MaxContactOffset;
+		GetContactOffsetParams(ContactOffsetFactor, MaxContactOffset);
+		PxMaterial* PDefaultMat = GetDefaultPhysMaterial();
 
-		FTransform BoxTransform = BoxElem.GetTransform() * RelativeTM;
-		if (PBoxGeom.isValid() && BoxTransform.IsValid())
+		for (int32 i = 0; i < BodySetup->AggGeom.SphereElems.Num(); i++)
 		{
-			PxTransform PLocalPose(U2PTransform(BoxTransform));
-			PLocalPose.p.x *= Scale3D.X;
-			PLocalPose.p.y *= Scale3D.Y;
-			PLocalPose.p.z *= Scale3D.Z;
+			const FKSphereElem& SphereElem = BodySetup->AggGeom.SphereElems[i];
+
+			PxSphereGeometry PSphereGeom;
+			PSphereGeom.radius = (SphereElem.Radius * MinScaleAbs);
 
-			ensure(PLocalPose.isValid());
+			if (ensure(PSphereGeom.isValid()))
 			{
-				PxShape* NewShape = PDestActor->createShape(PBoxGeom, *PDefaultMat, PLocalPose);
+				FVector LocalOrigin = RelativeTM.TransformPosition(SphereElem.Center);
+				PxTransform PLocalPose(U2PVector(LocalOrigin));
+				PLocalPose.p *= MinScale;
 
-				if (NewShapes)
+				ensure(PLocalPose.isValid());
 				{
-					NewShapes->Add(NewShape);
+					const float ContactOffset = FMath::Min(MaxContactOffset, ContactOffsetFactor * PSphereGeom.radius);
+					AttachShape(PSphereGeom, *PDefaultMat, PLocalPose, ContactOffset);
 				}
-
-				const float ContactOffset = FMath::Min(MaxContactOffset, ContactOffsetFactor * PBoxGeom.halfExtents.minElement());
-				NewShape->setContactOffset(ContactOffset);
 			}
-			
-
-			
-		}
-		else
-		{
-			UE_LOG(LogPhysics, Warning, TEXT("AddBoxesToRigidActor: [%s] BoxElems[%d] invalid or has invalid transform"), *GetPathNameSafe(GetOuter()), i);
+			else
+			{
+				UE_LOG(LogPhysics, Warning, TEXT("AddSpheresToRigidActor: [%s] SphereElem[%d] invalid"), *GetPathNameSafe(BodySetup->GetOuter()), i);
+			}
 		}
 	}
-}
-
-void UBodySetup::AddSphylsToRigidActor_AssumesLocked(PxRigidActor* PDestActor, const FTransform& RelativeTM, const FVector& Scale3D, const FVector& Scale3DAbs, TArray<PxShape*>* NewShapes) const
-{
-	float ContactOffsetFactor, MaxContactOffset;
-	GetContactOffsetParams(ContactOffsetFactor, MaxContactOffset);
-	PxMaterial* PDefaultMat = GetDefaultPhysMaterial();
-
-	float ScaleRadius = FMath::Max(Scale3DAbs.X, Scale3DAbs.Y);
-	float ScaleLength = Scale3DAbs.Z;
 
-	for (int32 i = 0; i < AggGeom.SphylElems.Num(); i++)
+	void AddBoxesToRigidActor_AssumesLocked() const
 	{
-		const FKSphylElem& SphylElem = AggGeom.SphylElems[i];
-
-		// this is a bit confusing since radius and height is scaled
-		// first apply the scale first 
-		float Radius = FMath::Max(SphylElem.Radius * ScaleRadius, 0.1f);
-		float Length = SphylElem.Length + SphylElem.Radius * 2.f;
-		float HalfLength = Length * ScaleLength * 0.5f;
-		Radius = FMath::Clamp(Radius, 0.1f, HalfLength);	//radius is capped by half length
-		float HalfHeight = HalfLength - Radius;
-		HalfHeight = FMath::Max(0.1f, HalfHeight);
-
-		PxCapsuleGeometry PCapsuleGeom;
-		PCapsuleGeom.halfHeight = HalfHeight;
-		PCapsuleGeom.radius = Radius;
-
-		if (PCapsuleGeom.isValid())
+		float ContactOffsetFactor, MaxContactOffset;
+		GetContactOffsetParams(ContactOffsetFactor, MaxContactOffset);
+		PxMaterial* PDefaultMat = GetDefaultPhysMaterial();
+
+		for (int32 i = 0; i < BodySetup->AggGeom.BoxElems.Num(); i++)
 		{
-			// The stored sphyl transform assumes the sphyl axis is down Z. In PhysX, it points down X, so we twiddle the matrix a bit here (swap X and Z and negate Y).
-			PxTransform PLocalPose(U2PVector(RelativeTM.TransformPosition(SphylElem.Center)), U2PQuat(SphylElem.Orientation) * U2PSphylBasis);
-			PLocalPose.p.x *= Scale3D.X;
-			PLocalPose.p.y *= Scale3D.Y;
-			PLocalPose.p.z *= Scale3D.Z;
+			const FKBoxElem& BoxElem = BodySetup->AggGeom.BoxElems[i];
+
+			PxBoxGeometry PBoxGeom;
+			PBoxGeom.halfExtents.x = (0.5f * BoxElem.X * Scale3DAbs.X);
+			PBoxGeom.halfExtents.y = (0.5f * BoxElem.Y * Scale3DAbs.Y);
+			PBoxGeom.halfExtents.z = (0.5f * BoxElem.Z * Scale3DAbs.Z);
 
-			ensure(PLocalPose.isValid());
+			FTransform BoxTransform = BoxElem.GetTransform() * RelativeTM;
+			if (PBoxGeom.isValid() && BoxTransform.IsValid())
 			{
-				PxShape* NewShape = PDestActor->createShape(PCapsuleGeom, *PDefaultMat, PLocalPose);
+				PxTransform PLocalPose(U2PTransform(BoxTransform));
+				PLocalPose.p.x *= Scale3D.X;
+				PLocalPose.p.y *= Scale3D.Y;
+				PLocalPose.p.z *= Scale3D.Z;
 
-				if (NewShapes)
+				ensure(PLocalPose.isValid());
 				{
-					NewShapes->Add(NewShape);
+					const float ContactOffset = FMath::Min(MaxContactOffset, ContactOffsetFactor * PBoxGeom.halfExtents.minElement());
+					AttachShape(PBoxGeom, *PDefaultMat, PLocalPose, ContactOffset);
 				}
-
-				const float ContactOffset = FMath::Min(MaxContactOffset, ContactOffsetFactor * PCapsuleGeom.radius);
-				NewShape->setContactOffset(ContactOffset);
 			}
-			
-
-			
-		}
-		else
-		{
-			UE_LOG(LogPhysics, Warning, TEXT("AddSphylsToRigidActor: [%s] SphylElems[%d] invalid"), *GetPathNameSafe(GetOuter()), i);
+			else
+			{
+				UE_LOG(LogPhysics, Warning, TEXT("AddBoxesToRigidActor: [%s] BoxElems[%d] invalid or has invalid transform"), *GetPathNameSafe(BodySetup->GetOuter()), i);
+			}
 		}
 	}
-}
 
-void UBodySetup::AddConvexElemsToRigidActor_AssumesLocked(PxRigidActor* PDestActor, const FTransform& RelativeTM, const FVector& Scale3D, const FVector& Scale3DAbs, TArray<PxShape*>* NewShapes) const
-{
-	float ContactOffsetFactor, MaxContactOffset;
-	GetContactOffsetParams(ContactOffsetFactor, MaxContactOffset);
-	PxMaterial* PDefaultMat = GetDefaultPhysMaterial();
-
-	for (int32 i = 0; i < AggGeom.ConvexElems.Num(); i++)
+	void AddSphylsToRigidActor_AssumesLocked() const
 	{
-		const FKConvexElem& ConvexElem = AggGeom.ConvexElems[i];
+		float ContactOffsetFactor, MaxContactOffset;
+		GetContactOffsetParams(ContactOffsetFactor, MaxContactOffset);
+		PxMaterial* PDefaultMat = GetDefaultPhysMaterial();
 
-		if (ConvexElem.ConvexMesh)
-		{
-			PxTransform PLocalPose;
-			bool bUseNegX = CalcMeshNegScaleCompensation(Scale3D, PLocalPose);
+		float ScaleRadius = FMath::Max(Scale3DAbs.X, Scale3DAbs.Y);
+		float ScaleLength = Scale3DAbs.Z;
 
-			PxConvexMeshGeometry PConvexGeom;
-			PConvexGeom.convexMesh = bUseNegX ? ConvexElem.ConvexMeshNegX : ConvexElem.ConvexMesh;
-			PConvexGeom.scale.scale = U2PVector(Scale3DAbs * ConvexElem.GetTransform().GetScale3D().GetAbs());
-			FTransform ConvexTransform = ConvexElem.GetTransform();
-			if (ConvexTransform.GetScale3D().X < 0 || ConvexTransform.GetScale3D().Y < 0 || ConvexTransform.GetScale3D().Z < 0)
-			{
-				UE_LOG(LogPhysics, Warning, TEXT("AddConvexElemsToRigidActor: [%s] ConvexElem[%d] has negative scale. Not currently supported"), *GetPathNameSafe(GetOuter()), i);
-			}
-			if (ConvexTransform.IsValid())
+		for (int32 i = 0; i < BodySetup->AggGeom.SphylElems.Num(); i++)
+		{
+			const FKSphylElem& SphylElem = BodySetup->AggGeom.SphylElems[i];
+
+			// this is a bit confusing since radius and height is scaled
+			// first apply the scale first 
+			float Radius = FMath::Max(SphylElem.Radius * ScaleRadius, 0.1f);
+			float Length = SphylElem.Length + SphylElem.Radius * 2.f;
+			float HalfLength = Length * ScaleLength * 0.5f;
+			Radius = FMath::Clamp(Radius, 0.1f, HalfLength);	//radius is capped by half length
+			float HalfHeight = HalfLength - Radius;
+			HalfHeight = FMath::Max(0.1f, HalfHeight);
+
+			PxCapsuleGeometry PCapsuleGeom;
+			PCapsuleGeom.halfHeight = HalfHeight;
+			PCapsuleGeom.radius = Radius;
+
+			if (PCapsuleGeom.isValid())
 			{
-				PxTransform PElementTransform = U2PTransform(ConvexTransform * RelativeTM);
-				PLocalPose.q *= PElementTransform.q;
-				PLocalPose.p = PElementTransform.p;
+				// The stored sphyl transform assumes the sphyl axis is down Z. In PhysX, it points down X, so we twiddle the matrix a bit here (swap X and Z and negate Y).
+				PxTransform PLocalPose(U2PVector(RelativeTM.TransformPosition(SphylElem.Center)), U2PQuat(SphylElem.Orientation) * U2PSphylBasis);
 				PLocalPose.p.x *= Scale3D.X;
 				PLocalPose.p.y *= Scale3D.Y;
 				PLocalPose.p.z *= Scale3D.Z;
 
-				if (PConvexGeom.isValid())
+				ensure(PLocalPose.isValid());
 				{
-					PxVec3 PBoundsExtents = PConvexGeom.convexMesh->getLocalBounds().getExtents();
+					const float ContactOffset = FMath::Min(MaxContactOffset, ContactOffsetFactor * PCapsuleGeom.radius);
+					AttachShape(PCapsuleGeom, *PDefaultMat, PLocalPose, ContactOffset);
+				}
+			}
+			else
+			{
+				UE_LOG(LogPhysics, Warning, TEXT("AddSphylsToRigidActor: [%s] SphylElems[%d] invalid"), *GetPathNameSafe(BodySetup->GetOuter()), i);
+			}
+		}
+	}
 
-					ensure(PLocalPose.isValid());
+	void FAddShapesHelper::AddConvexElemsToRigidActor_AssumesLocked() const
+	{
+		float ContactOffsetFactor, MaxContactOffset;
+		GetContactOffsetParams(ContactOffsetFactor, MaxContactOffset);
+		PxMaterial* PDefaultMat = GetDefaultPhysMaterial();
+
+		for (int32 i = 0; i < BodySetup->AggGeom.ConvexElems.Num(); i++)
+		{
+			const FKConvexElem& ConvexElem = BodySetup->AggGeom.ConvexElems[i];
+
+			if (ConvexElem.ConvexMesh)
+			{
+				PxTransform PLocalPose;
+				bool bUseNegX = CalcMeshNegScaleCompensation(Scale3D, PLocalPose);
+
+				PxConvexMeshGeometry PConvexGeom;
+				PConvexGeom.convexMesh = bUseNegX ? ConvexElem.ConvexMeshNegX : ConvexElem.ConvexMesh;
+				PConvexGeom.scale.scale = U2PVector(Scale3DAbs * ConvexElem.GetTransform().GetScale3D().GetAbs());
+				FTransform ConvexTransform = ConvexElem.GetTransform();
+				if (ConvexTransform.GetScale3D().X < 0 || ConvexTransform.GetScale3D().Y < 0 || ConvexTransform.GetScale3D().Z < 0)
+				{
+					UE_LOG(LogPhysics, Warning, TEXT("AddConvexElemsToRigidActor: [%s] ConvexElem[%d] has negative scale. Not currently supported"), *GetPathNameSafe(BodySetup->GetOuter()), i);
+				}
+				if (ConvexTransform.IsValid())
+				{
+					PxTransform PElementTransform = U2PTransform(ConvexTransform * RelativeTM);
+					PLocalPose.q *= PElementTransform.q;
+					PLocalPose.p = PElementTransform.p;
+					PLocalPose.p.x *= Scale3D.X;
+					PLocalPose.p.y *= Scale3D.Y;
+					PLocalPose.p.z *= Scale3D.Z;
+
+					if (PConvexGeom.isValid())
 					{
-						PxShape* NewShape = PDestActor->createShape(PConvexGeom, *PDefaultMat, PLocalPose);
+						PxVec3 PBoundsExtents = PConvexGeom.convexMesh->getLocalBounds().getExtents();
 
-						if (NewShapes)
+						ensure(PLocalPose.isValid());
 						{
-							NewShapes->Add(NewShape);
+							const float ContactOffset = FMath::Min(MaxContactOffset, ContactOffsetFactor * PBoundsExtents.minElement());
+							AttachShape(PConvexGeom, *PDefaultMat, PLocalPose, ContactOffset);
 						}
-
-						const float ContactOffset = FMath::Min(MaxContactOffset, ContactOffsetFactor * PBoundsExtents.minElement());
-						NewShape->setContactOffset(ContactOffset);
 					}
-					
-
-					
+					else
+					{
+						UE_LOG(LogPhysics, Warning, TEXT("AddConvexElemsToRigidActor: [%s] ConvexElem[%d] invalid"), *GetPathNameSafe(BodySetup->GetOuter()), i);
+					}
 				}
 				else
 				{
-					UE_LOG(LogPhysics, Warning, TEXT("AddConvexElemsToRigidActor: [%s] ConvexElem[%d] invalid"), *GetPathNameSafe(GetOuter()), i);
+					UE_LOG(LogPhysics, Warning, TEXT("AddConvexElemsToRigidActor: [%s] ConvexElem[%d] has invalid transform"), *GetPathNameSafe(BodySetup->GetOuter()), i);
 				}
 			}
 			else
 			{
-				UE_LOG(LogPhysics, Warning, TEXT("AddConvexElemsToRigidActor: [%s] ConvexElem[%d] has invalid transform"), *GetPathNameSafe(GetOuter()), i);
+				UE_LOG(LogPhysics, Warning, TEXT("AddConvexElemsToRigidActor: ConvexElem is missing ConvexMesh (%d: %s)"), i, *BodySetup->GetPathName());
 			}
 		}
-		else
-		{
-			UE_LOG(LogPhysics, Warning, TEXT("AddConvexElemsToRigidActor: ConvexElem is missing ConvexMesh (%d: %s)"), i, *GetPathName());
-		}
 	}
-}
-
-void UBodySetup::AddTriMeshToRigidActor_AssumesLocked(PxRigidActor* PDestActor, const FVector& Scale3D, const FVector& Scale3DAbs, TArray<PxShape*>* NewShapes) const
-{
-	float ContactOffsetFactor, MaxContactOffset;
-	GetContactOffsetParams(ContactOffsetFactor, MaxContactOffset);
-	PxMaterial* PDefaultMat = GetDefaultPhysMaterial();
 
-	if(TriMesh || TriMeshNegX)
+	void AddTriMeshToRigidActor_AssumesLocked() const
 	{
-		PxTransform PLocalPose;
-		bool bUseNegX = CalcMeshNegScaleCompensation(Scale3D, PLocalPose);
-
-		// Only case where TriMeshNegX should be null is BSP, which should not require negX version
-		if (bUseNegX && TriMeshNegX == NULL)
-		{
-			UE_LOG(LogPhysics, Log, TEXT("AddTriMeshToRigidActor: Want to use NegX but it doesn't exist! %s"), *GetPathName());
-		}
+		float ContactOffsetFactor, MaxContactOffset;
+		GetContactOffsetParams(ContactOffsetFactor, MaxContactOffset);
+		PxMaterial* PDefaultMat = GetDefaultPhysMaterial();
 
-		PxTriangleMesh* UseTriMesh = bUseNegX ? TriMeshNegX : TriMesh;
-		if (UseTriMesh != NULL)
+		if (BodySetup->TriMesh || BodySetup->TriMeshNegX)
 		{
+			PxTransform PLocalPose;
+			bool bUseNegX = CalcMeshNegScaleCompensation(Scale3D, PLocalPose);
 
-
-			PxTriangleMeshGeometry PTriMeshGeom;
-			PTriMeshGeom.triangleMesh = bUseNegX ? TriMeshNegX : TriMesh;
-			PTriMeshGeom.scale.scale = U2PVector(Scale3DAbs);
-			if (bDoubleSidedGeometry)
+			// Only case where TriMeshNegX should be null is BSP, which should not require negX version
+			if (bUseNegX && BodySetup->TriMeshNegX == NULL)
 			{
-				PTriMeshGeom.meshFlags |= PxMeshGeometryFlag::eDOUBLE_SIDED;
+				UE_LOG(LogPhysics, Log, TEXT("AddTriMeshToRigidActor: Want to use NegX but it doesn't exist! %s"), *BodySetup->GetPathName());
 			}
 
-
-			if (PTriMeshGeom.isValid())
+			PxTriangleMesh* UseTriMesh = bUseNegX ? BodySetup->TriMeshNegX : BodySetup->TriMesh;
+			if (UseTriMesh != NULL)
 			{
-				ensure(PLocalPose.isValid());
+				PxTriangleMeshGeometry PTriMeshGeom;
+				PTriMeshGeom.triangleMesh = bUseNegX ? BodySetup->TriMeshNegX : BodySetup->TriMesh;
+				PTriMeshGeom.scale.scale = U2PVector(Scale3DAbs);
+				if (BodySetup->bDoubleSidedGeometry)
+				{
+					PTriMeshGeom.meshFlags |= PxMeshGeometryFlag::eDOUBLE_SIDED;
+				}
 
-				// Create without 'sim shape' flag, problematic if it's kinematic, and it gets set later anyway.
+
+				if (PTriMeshGeom.isValid())
 				{
-					PxShape* NewShape = PDestActor->createShape(PTriMeshGeom, *PDefaultMat, PxShapeFlag::eSCENE_QUERY_SHAPE | PxShapeFlag::eVISUALIZATION);
+					ensure(PLocalPose.isValid());
 
-					if (NewShape)
+					// Create without 'sim shape' flag, problematic if it's kinematic, and it gets set later anyway.
 					{
-						NewShape->setLocalPose(PLocalPose);
-						NewShape->setContactOffset(MaxContactOffset);
-
-						if (NewShapes)
+						if (PxShape* PNewShape = AttachShape(PTriMeshGeom, *PDefaultMat, PLocalPose, MaxContactOffset))
 						{
-							NewShapes->Add(NewShape);
+							PNewShape->setFlags(PxShapeFlag::eSCENE_QUERY_SHAPE | PxShapeFlag::eVISUALIZATION);
+						}
+						else
+						{
+							UE_LOG(LogPhysics, Log, TEXT("Can't create new mesh shape in AddShapesToRigidActor"));
 						}
 					}
-					else
-					{
-						UE_LOG(LogPhysics, Log, TEXT("Can't create new mesh shape in AddShapesToRigidActor"));
-					}
 				}
-				
+				else
+				{
+					UE_LOG(LogPhysics, Log, TEXT("AddTriMeshToRigidActor: TriMesh invalid"));
+				}
 			}
-			else
+		}
+	}
+
+private:
+
+	PxShape* AttachShape(const PxGeometry& PGeom, const PxMaterial& PMaterial, const PxTransform& PLocalPose, const float ContactOffset) const
+	{
+		PxShape* PNewShape = PDestActor->createShape(PGeom, PMaterial, PLocalPose);
+		if (PNewShape)
+		{
+			if (NewShapes)
 			{
-				UE_LOG(LogPhysics, Log, TEXT("AddTriMeshToRigidActor: TriMesh invalid"));
+				NewShapes->Add(PNewShape);
 			}
+
+			PNewShape->setContactOffset(ContactOffset);
 		}
+
+		return PNewShape;
 	}
-}
+
+
+};
+
+
 
 void UBodySetup::AddShapesToRigidActor_AssumesLocked(FBodyInstance* OwningInstance, physx::PxRigidActor* PDestActor, EPhysicsSceneType SceneType, FVector& Scale3D, physx::PxMaterial* SimpleMaterial, TArray<UPhysicalMaterial*>& ComplexMaterials, FShapeData& ShapeData, const FTransform& RelativeTM, TArray<physx::PxShape*>* NewShapes)
 {
@@ -597,110 +615,90 @@ void UBodySetup::AddShapesToRigidActor_AssumesLocked(FBodyInstance* OwningInstan
 	CreatePhysicsMeshes();
 #endif
 
-	float MinScale;
-	float MinScaleAbs;
-	FVector Scale3DAbs;
 	TArray<PxShape*> Shapes;
-	SetupNonUniformHelper(Scale3D, MinScale, MinScaleAbs, Scale3DAbs);
-
-	{
-		float MinScaleRelative;
-		float MinScaleAbsRelative;
-		FVector Scale3DAbsRelative;
-		FVector Scale3DRelative = RelativeTM.GetScale3D();
-
-		SetupNonUniformHelper(Scale3DRelative, MinScaleRelative, MinScaleAbsRelative, Scale3DAbsRelative);
-
-		MinScaleAbs *= MinScaleAbsRelative;
-		Scale3DAbs.X *= Scale3DAbsRelative.X;
-		Scale3DAbs.Y *= Scale3DAbsRelative.Y;
-		Scale3DAbs.Z *= Scale3DAbsRelative.Z;
-	}
+	FAddShapesHelper AddShapesHelper(this, OwningInstance, PDestActor, Scale3D, SimpleMaterial, ComplexMaterials, ShapeData, RelativeTM, &Shapes);
 
 	// Create shapes for simple collision if we do not want to use the complex collision mesh 
 	// for simple queries as well
 	if (CollisionTraceFlag != ECollisionTraceFlag::CTF_UseComplexAsSimple)
 	{
-		AddSpheresToRigidActor_AssumesLocked(PDestActor, RelativeTM, MinScale, MinScaleAbs, &Shapes);
-		AddBoxesToRigidActor_AssumesLocked(PDestActor, RelativeTM, Scale3D, Scale3DAbs, &Shapes);
-		AddSphylsToRigidActor_AssumesLocked(PDestActor, RelativeTM, Scale3D, Scale3DAbs, &Shapes);
-		AddConvexElemsToRigidActor_AssumesLocked(PDestActor, RelativeTM, Scale3D, Scale3DAbs, &Shapes);
+		AddShapesHelper.AddSpheresToRigidActor_AssumesLocked();
+		AddShapesHelper.AddBoxesToRigidActor_AssumesLocked();
+		AddShapesHelper.AddSphylsToRigidActor_AssumesLocked();
+		AddShapesHelper.AddConvexElemsToRigidActor_AssumesLocked();
 	}
 
 	// Create tri-mesh shape, when we are not using simple collision shapes for 
 	// complex queries as well
-	if( CollisionTraceFlag != ECollisionTraceFlag::CTF_UseSimpleAsComplex )
+	if (CollisionTraceFlag != ECollisionTraceFlag::CTF_UseSimpleAsComplex)
 	{
-		AddTriMeshToRigidActor_AssumesLocked(PDestActor, Scale3D, Scale3DAbs, &Shapes);
+		AddShapesHelper.AddTriMeshToRigidActor_AssumesLocked();
 	}
 
+	for (PxShape* Shape : Shapes)
 	{
-		for (PxShape* Shape : Shapes)
+		FShapeFilterData& Filters = ShapeData.FilterData;
+		// Apply provided materials
+		// If a triangle mesh, need to get array of materials...
+		if (Shape->getGeometryType() == PxGeometryType::eTRIANGLEMESH)
 		{
-			FShapeFilterData& Filters = ShapeData.FilterData;
-			// Apply provided materials
-			// If a triangle mesh, need to get array of materials...
-			if (Shape->getGeometryType() == PxGeometryType::eTRIANGLEMESH)
+			TArray<PxMaterial*> PComplexMats;
+			PComplexMats.AddUninitialized(ComplexMaterials.Num());
+			for (int MatIdx = 0; MatIdx < ComplexMaterials.Num(); MatIdx++)
 			{
-				TArray<PxMaterial*> PComplexMats;
-				PComplexMats.AddUninitialized(ComplexMaterials.Num());
-				for (int MatIdx = 0; MatIdx < ComplexMaterials.Num(); MatIdx++)
-				{
-					check(ComplexMaterials[MatIdx] != NULL);
-					PComplexMats[MatIdx] = ComplexMaterials[MatIdx]->GetPhysXMaterial();
-					check(PComplexMats[MatIdx] != NULL);
-				}
+				check(ComplexMaterials[MatIdx] != NULL);
+				PComplexMats[MatIdx] = ComplexMaterials[MatIdx]->GetPhysXMaterial();
+				check(PComplexMats[MatIdx] != NULL);
+			}
 
-				if (PComplexMats.Num())
-				{
-					Shape->setMaterials(PComplexMats.GetData(), PComplexMats.Num());
-				}
-				else
-				{
-					UE_LOG(LogPhysics, Warning, TEXT("FBodyInstance::UpdatePhysicalMaterials : PComplexMats is empty - falling back on simple physical material."));
-					Shape->setMaterials(&SimpleMaterial, 1);
-				}
+			if (PComplexMats.Num())
+			{
+				Shape->setMaterials(PComplexMats.GetData(), PComplexMats.Num());
+			}
+			else
+			{
+				UE_LOG(LogPhysics, Warning, TEXT("FBodyInstance::UpdatePhysicalMaterials : PComplexMats is empty - falling back on simple physical material."));
+				Shape->setMaterials(&SimpleMaterial, 1);
+			}
 
-				Shape->setQueryFilterData(Filters.QueryComplexFilter);
+			Shape->setQueryFilterData(Filters.QueryComplexFilter);
 
-				if (SceneType == PST_Sync)
-				{
-					Shape->setFlags(ShapeData.SyncShapeFlags | ShapeData.ComplexShapeFlags);
-				}
-				else
-				{
-					Shape->setFlags(ShapeData.AsyncShapeFlags | ShapeData.ComplexShapeFlags);
-				}
+			if (SceneType == PST_Sync)
+			{
+				Shape->setFlags(ShapeData.SyncShapeFlags | ShapeData.ComplexShapeFlags);
 			}
-			// Simple shape, 
 			else
 			{
-				Shape->setMaterials(&SimpleMaterial, 1);
-				Shape->setQueryFilterData(Filters.QuerySimpleFilter);
+				Shape->setFlags(ShapeData.AsyncShapeFlags | ShapeData.ComplexShapeFlags);
+			}
+		}
+		// Simple shape, 
+		else
+		{
+			Shape->setMaterials(&SimpleMaterial, 1);
+			Shape->setQueryFilterData(Filters.QuerySimpleFilter);
 
-				if (SceneType == PST_Sync)
-				{
-					Shape->setFlags(ShapeData.SyncShapeFlags | ShapeData.SimpleShapeFlags);
-				}
-				else
-				{
-					Shape->setFlags(ShapeData.AsyncShapeFlags | ShapeData.SimpleShapeFlags);
-				}
+			if (SceneType == PST_Sync)
+			{
+				Shape->setFlags(ShapeData.SyncShapeFlags | ShapeData.SimpleShapeFlags);
+			}
+			else
+			{
+				Shape->setFlags(ShapeData.AsyncShapeFlags | ShapeData.SimpleShapeFlags);
 			}
-			Shape->setSimulationFilterData(Filters.SimFilter);
+		}
+		Shape->setSimulationFilterData(Filters.SimFilter);
 
-			if (OwningInstance)
+		if (OwningInstance)
+		{
+			if (PxRigidBody* RigidBody = OwningInstance->GetPxRigidBody_AssumesLocked())
 			{
-				if (PxRigidBody* RigidBody = OwningInstance->GetPxRigidBody_AssumesLocked())
-				{
-					RigidBody->setRigidBodyFlags(ShapeData.SyncBodyFlags);
-				}
+				RigidBody->setRigidBodyFlags(ShapeData.SyncBodyFlags);
 			}
 		}
 	}
-	
 
-	if(NewShapes)
+	if (NewShapes)
 	{
 		(*NewShapes) = MoveTemp(Shapes);
 	}
