Fix Mac compile

diff --git a/Engine/Source/Runtime/Engine/Private/CheatManager.cpp b/Engine/Source/Runtime/Engine/Private/CheatManager.cpp
--- a/Engine/Source/Runtime/Engine/Private/CheatManager.cpp
+++ b/Engine/Source/Runtime/Engine/Private/CheatManager.cpp
@@ -17,6 +17,7 @@
 #include "Engine/DebugCameraController.h"
 #include "GameFramework/PlayerState.h"
 #include "GameFramework/GameMode.h"
+#include "GameFramework/PlayerInput.h"
 #include "Components/BrushComponent.h"
 
 DEFINE_LOG_CATEGORY_STATIC(LogCheatManager, Log, All);
