Merging 2514537

//depot/UE4-UT/...

to //depot/UE4/...

diff --git a/Engine/Source/Runtime/Engine/Classes/Engine/BlueprintGeneratedClass.h b/Engine/Source/Runtime/Engine/Classes/Engine/BlueprintGeneratedClass.h
--- a/Engine/Source/Runtime/Engine/Classes/Engine/BlueprintGeneratedClass.h
+++ b/Engine/Source/Runtime/Engine/Classes/Engine/BlueprintGeneratedClass.h
@@ -505,4 +505,16 @@ public:
 
 	// Finds the desired dynamic binding object for this blueprint generated class
 	UDynamicBlueprintBinding* GetDynamicBindingObject(UClass* Class) const;
+
+	/** called to gather blueprint replicated properties */
+	virtual void GetLifetimeBlueprintReplicationList(TArray<class FLifetimeProperty>& OutLifetimeProps) const;
+	/** called prior to replication of an instance of this BP class */
+	virtual void InstancePreReplication(class IRepChangedPropertyTracker& ChangedPropertyTracker) const
+	{
+		UBlueprintGeneratedClass* SuperBPClass = Cast<UBlueprintGeneratedClass>(GetSuperStruct());
+		if (SuperBPClass != NULL)
+		{
+			SuperBPClass->InstancePreReplication(ChangedPropertyTracker);
+		}
+	}
 };
diff --git a/Engine/Source/Runtime/Engine/Private/Actor.cpp b/Engine/Source/Runtime/Engine/Private/Actor.cpp
--- a/Engine/Source/Runtime/Engine/Private/Actor.cpp
+++ b/Engine/Source/Runtime/Engine/Private/Actor.cpp
@@ -725,6 +725,12 @@ void AActor::PreReplication( IRepChangedPropertyTracker & ChangedPropertyTracker
 	}
 
 	DOREPLIFETIME_ACTIVE_OVERRIDE( AActor, ReplicatedMovement, bReplicateMovement );
+
+	UBlueprintGeneratedClass* BPClass = Cast<UBlueprintGeneratedClass>(GetClass());
+	if (BPClass != NULL)
+	{
+		BPClass->InstancePreReplication(ChangedPropertyTracker);
+	}
 }
 
 void AActor::PostActorCreated()
diff --git a/Engine/Source/Runtime/Engine/Private/ActorReplication.cpp b/Engine/Source/Runtime/Engine/Private/ActorReplication.cpp
--- a/Engine/Source/Runtime/Engine/Private/ActorReplication.cpp
+++ b/Engine/Source/Runtime/Engine/Private/ActorReplication.cpp
@@ -226,32 +226,14 @@ void AActor::GatherCurrentMovement()
 	}
 }
 
-static void GetLifetimeBlueprintReplicationList( const AActor* ThisActor, const UBlueprintGeneratedClass* MyClass, TArray< FLifetimeProperty > & OutLifetimeProps )
+void AActor::GetLifetimeReplicatedProps( TArray< FLifetimeProperty > & OutLifetimeProps ) const
 {
-	if ( MyClass == NULL )
-	{
-		return;
-	}
-
-	uint32 PropertiesLeft = MyClass->NumReplicatedProperties;
-
-	for ( TFieldIterator<UProperty> It( MyClass, EFieldIteratorFlags::ExcludeSuper ); It && PropertiesLeft > 0; ++It )
+	UBlueprintGeneratedClass* BPClass = Cast<UBlueprintGeneratedClass>(GetClass());
+	if (BPClass != NULL)
 	{
-		UProperty * Prop = *It;
-		if ( Prop != NULL && Prop->GetPropertyFlags() & CPF_Net )
-		{
-			PropertiesLeft--;
-			OutLifetimeProps.Add( FLifetimeProperty( Prop->RepIndex ) );
-		}
+		BPClass->GetLifetimeBlueprintReplicationList(OutLifetimeProps);
 	}
 
-	return GetLifetimeBlueprintReplicationList( ThisActor, Cast< UBlueprintGeneratedClass >( MyClass->GetSuperStruct() ), OutLifetimeProps );
-}
-
-void AActor::GetLifetimeReplicatedProps( TArray< FLifetimeProperty > & OutLifetimeProps ) const
-{
-	GetLifetimeBlueprintReplicationList( this, Cast< UBlueprintGeneratedClass >( GetClass() ), OutLifetimeProps );
-
 	DOREPLIFETIME( AActor, Role );
 	DOREPLIFETIME( AActor, RemoteRole );
 	DOREPLIFETIME( AActor, Owner );
diff --git a/Engine/Source/Runtime/Engine/Private/BlueprintGeneratedClass.cpp b/Engine/Source/Runtime/Engine/Private/BlueprintGeneratedClass.cpp
--- a/Engine/Source/Runtime/Engine/Private/BlueprintGeneratedClass.cpp
+++ b/Engine/Source/Runtime/Engine/Private/BlueprintGeneratedClass.cpp
@@ -8,6 +8,7 @@
 #include "Engine/SCS_Node.h"
 #include "Engine/LevelScriptActor.h"
 #include "Engine/InheritableComponentHandler.h"
+#include "CoreNet.h"
 
 #if WITH_EDITOR
 #include "BlueprintEditorUtils.h"
@@ -759,4 +760,25 @@ void UBlueprintGeneratedClass::Serialize(FArchive& Ar)
 	{
 		CreatePersistentUberGraphFrame(ClassDefaultObject, true);
 	}
+}
+
+void UBlueprintGeneratedClass::GetLifetimeBlueprintReplicationList(TArray<FLifetimeProperty>& OutLifetimeProps) const
+{
+	uint32 PropertiesLeft = NumReplicatedProperties;
+
+	for (TFieldIterator<UProperty> It(this, EFieldIteratorFlags::ExcludeSuper); It && PropertiesLeft > 0; ++It)
+	{
+		UProperty * Prop = *It;
+		if (Prop != NULL && Prop->GetPropertyFlags() & CPF_Net)
+		{
+			PropertiesLeft--;
+			OutLifetimeProps.Add(FLifetimeProperty(Prop->RepIndex));
+		}
+	}
+
+	UBlueprintGeneratedClass* SuperBPClass = Cast<UBlueprintGeneratedClass>(GetSuperStruct());
+	if (SuperBPClass != NULL)
+	{
+		SuperBPClass->GetLifetimeBlueprintReplicationList(OutLifetimeProps);
+	}
 }
\ No newline at end of file
