Fix build on mac

diff --git a/Engine/Source/Runtime/Engine/Private/PhysicsEngine/BodySetup.cpp b/Engine/Source/Runtime/Engine/Private/PhysicsEngine/BodySetup.cpp
--- a/Engine/Source/Runtime/Engine/Private/PhysicsEngine/BodySetup.cpp
+++ b/Engine/Source/Runtime/Engine/Private/PhysicsEngine/BodySetup.cpp
@@ -473,7 +473,7 @@ public:
 		}
 	}
 
-	void FAddShapesHelper::AddConvexElemsToRigidActor_AssumesLocked() const
+	void AddConvexElemsToRigidActor_AssumesLocked() const
 	{
 		float ContactOffsetFactor, MaxContactOffset;
 		GetContactOffsetParams(ContactOffsetFactor, MaxContactOffset);
