Almost infinite (up to max_int32) loop when a parent component fails to register.

diff --git a/Engine/Source/Runtime/Engine/Private/Actor.cpp b/Engine/Source/Runtime/Engine/Private/Actor.cpp
--- a/Engine/Source/Runtime/Engine/Private/Actor.cpp
+++ b/Engine/Source/Runtime/Engine/Private/Actor.cpp
@@ -3381,6 +3381,7 @@ bool AActor::IncrementalRegisterComponents(int32 NumComponentsToRegister)
 	int32 NumRegisteredComponentsThisRun = 0;
 	TInlineComponentArray<UActorComponent*> Components;
 	GetComponents(Components);
+	TSet<UActorComponent*> RegisteredParents;
 	
 	for (int32 CompIdx = 0; CompIdx < Components.Num() && NumRegisteredComponentsThisRun < NumComponentsToRegister; CompIdx++)
 	{
@@ -3391,6 +3392,14 @@ bool AActor::IncrementalRegisterComponents(int32 NumComponentsToRegister)
 			USceneComponent* ParentComponent = GetUnregisteredParent(Component);
 			if (ParentComponent)
 			{
+				bool bParentAlreadyHandled = false;
+				RegisteredParents.Add(ParentComponent, &bParentAlreadyHandled);
+				if (bParentAlreadyHandled)
+				{
+					UE_LOG(LogActor, Error, TEXT("AActor::IncrementalRegisterComponents parent component '%s' cannot be registered in actor '%s'"), *GetPathNameSafe(ParentComponent), *GetPathName());
+					break;
+				}
+
 				// Register parent first, then return to this component on a next iteration
 				Component = ParentComponent;
 				CompIdx--;
