More compilation failure fixes

diff --git a/Engine/Source/Runtime/Engine/Private/CheatManager.cpp b/Engine/Source/Runtime/Engine/Private/CheatManager.cpp
--- a/Engine/Source/Runtime/Engine/Private/CheatManager.cpp
+++ b/Engine/Source/Runtime/Engine/Private/CheatManager.cpp
@@ -18,6 +18,7 @@
 #include "GameFramework/PlayerState.h"
 #include "GameFramework/GameMode.h"
 #include "GameFramework/PlayerInput.h"
+#include "GameFramework/InputSettings.h"
 #include "Components/BrushComponent.h"
 
 DEFINE_LOG_CATEGORY_STATIC(LogCheatManager, Log, All);
diff --git a/Engine/Source/Runtime/Engine/Private/World.cpp b/Engine/Source/Runtime/Engine/Private/World.cpp
--- a/Engine/Source/Runtime/Engine/Private/World.cpp
+++ b/Engine/Source/Runtime/Engine/Private/World.cpp
@@ -149,7 +149,9 @@ void UWorld::Serialize( FArchive& Ar )
 
 	if (Ar.UE4Ver() < VER_UE4_ADD_EDITOR_VIEWS)
 	{
+#if WITH_EDITOR
 		EditorViews.SetNum(4);
+#endif
 		for (int32 i = 0; i < 4; ++i)
 		{
 			FLevelViewportInfo TempViewportInfo;
