Fix crash caused if ChildActorComponent is added directly to a world without an owning actor

diff --git a/Engine/Source/Runtime/Engine/Private/Components/ChildActorComponent.cpp b/Engine/Source/Runtime/Engine/Private/Components/ChildActorComponent.cpp
--- a/Engine/Source/Runtime/Engine/Private/Components/ChildActorComponent.cpp
+++ b/Engine/Source/Runtime/Engine/Private/Components/ChildActorComponent.cpp
@@ -251,7 +251,7 @@ void UChildActorComponent::CreateChildActor()
 				Params.bNoCollisionFail = true;
 				Params.bDeferConstruction = true; // We defer construction so that we set ParentComponentActor prior to component registration so they appear selected
 				Params.bAllowDuringConstructionScript = true;
-				Params.OverrideLevel = GetOwner()->GetLevel();
+				Params.OverrideLevel = (Actor ? Actor->GetLevel() : nullptr);
 				Params.Name = ChildActorName;
 				if (!HasAllFlags(RF_Transactional))
 				{
@@ -269,7 +269,7 @@ void UChildActorComponent::CreateChildActor()
 					ChildActorName = ChildActor->GetFName();
 
 					// Remember which actor spawned it (for selection in editor etc)
-					ChildActor->ParentComponentActor = GetOwner();
+					ChildActor->ParentComponentActor = Actor;
 
 					ChildActor->AttachRootComponentTo(this);
 
