Fix FBodyInstance readwrite lock not obtaining lock for both actors. This is needed since writing to a static actor requires updating both scenes

diff --git a/Engine/Source/Runtime/Engine/Private/PhysicsEngine/BodyInstance.cpp b/Engine/Source/Runtime/Engine/Private/PhysicsEngine/BodyInstance.cpp
--- a/Engine/Source/Runtime/Engine/Private/PhysicsEngine/BodyInstance.cpp
+++ b/Engine/Source/Runtime/Engine/Private/PhysicsEngine/BodyInstance.cpp
@@ -2768,6 +2768,9 @@ void FBodyInstance::CopyBodyInstancePropertiesFrom(const FBodyInstance* FromInst
 
 void FBodyInstance::ExecuteOnPhysicsReadOnly(TFunctionRef<void()> Func) const
 {
+	//If we are doing a read operation on a static actor there is really no reason to lock both scenes since the data should be the same (as far as queries etc... go)
+	//Because of this our read operations are typically on a dynamic or the sync actor
+
 #if WITH_PHYSX
 	const int32 SceneIndex = RigidActorSync ? SceneIndexSync : SceneIndexAsync;
 	SCOPED_SCENE_READ_LOCK(GetPhysXSceneFromIndex(SceneIndex));
@@ -2777,10 +2780,31 @@ void FBodyInstance::ExecuteOnPhysicsReadOnly(TFunctionRef<void()> Func) const
 
 void FBodyInstance::ExecuteOnPhysicsReadWrite(TFunctionRef<void()> Func) const
 {
+	//If we are doing a write operation we will need to modify both actors to keep them in sync (or it's dynamic).
+	//Because of that write operations on static actors are more expensive and require both locks.
+
 #if WITH_PHYSX
-	const int32 SceneIndex = RigidActorSync ? SceneIndexSync : SceneIndexAsync;
-	SCOPED_SCENE_WRITE_LOCK(GetPhysXSceneFromIndex(SceneIndex));
+	if(RigidActorSync)
+	{
+		SCENE_LOCK_WRITE(GetPhysXSceneFromIndex(SceneIndexSync));
+	}
+	
+	if (RigidActorAsync)
+	{
+		SCENE_LOCK_WRITE(GetPhysXSceneFromIndex(SceneIndexAsync));
+	}
+
 	Func();
+
+	if (RigidActorSync)
+	{
+		SCENE_UNLOCK_WRITE(GetPhysXSceneFromIndex(SceneIndexSync));
+	}
+
+	if (RigidActorAsync)
+	{
+		SCENE_UNLOCK_WRITE(GetPhysXSceneFromIndex(SceneIndexAsync));
+	}
 #endif
 }
 
