Let UPlayerInput::GetMouseSensitivity remain part of API

diff --git a/Engine/Source/Runtime/Engine/Classes/GameFramework/PlayerInput.h b/Engine/Source/Runtime/Engine/Classes/GameFramework/PlayerInput.h
--- a/Engine/Source/Runtime/Engine/Classes/GameFramework/PlayerInput.h
+++ b/Engine/Source/Runtime/Engine/Classes/GameFramework/PlayerInput.h
@@ -305,7 +305,6 @@ public:
 	void SetBind(FName BindName, const FString& Command);
 
 	/** Returns the mouse sensitivity along the X-axis, or the Y-axis, or 1.0 if none are known. */
-	DEPRECATED(4.8, "GetMouseSensitivity is deprecated, use GetAxisProperties instead")
 	float GetMouseSensitivity();
 
 	/** Returns whether an Axis Key is inverted */
