Fix trimesh getting simulation flag turned on. Fallout from refactor

diff --git a/Engine/Source/Runtime/Engine/Private/PhysicsEngine/BodySetup.cpp b/Engine/Source/Runtime/Engine/Private/PhysicsEngine/BodySetup.cpp
--- a/Engine/Source/Runtime/Engine/Private/PhysicsEngine/BodySetup.cpp
+++ b/Engine/Source/Runtime/Engine/Private/PhysicsEngine/BodySetup.cpp
@@ -567,11 +567,7 @@ public:
 
 					// Create without 'sim shape' flag, problematic if it's kinematic, and it gets set later anyway.
 					{
-						if (PxShape* PNewShape = AttachShape(PTriMeshGeom, *PDefaultMat, PLocalPose, MaxContactOffset))
-						{
-							PNewShape->setFlags(PxShapeFlag::eSCENE_QUERY_SHAPE | PxShapeFlag::eVISUALIZATION);
-						}
-						else
+						if (!AttachShape(PTriMeshGeom, *PDefaultMat, PLocalPose, MaxContactOffset, PxShapeFlag::eSCENE_QUERY_SHAPE | PxShapeFlag::eVISUALIZATION))
 						{
 							UE_LOG(LogPhysics, Log, TEXT("Can't create new mesh shape in AddShapesToRigidActor"));
 						}
@@ -587,9 +583,9 @@ public:
 
 private:
 
-	PxShape* AttachShape(const PxGeometry& PGeom, const PxMaterial& PMaterial, const PxTransform& PLocalPose, const float ContactOffset) const
+	PxShape* AttachShape(const PxGeometry& PGeom, const PxMaterial& PMaterial, const PxTransform& PLocalPose, const float ContactOffset, PxShapeFlags PShapeFlags = PxShapeFlag::eVISUALIZATION | PxShapeFlag::eSCENE_QUERY_SHAPE | PxShapeFlag::eSIMULATION_SHAPE) const
 	{
-		PxShape* PNewShape = PDestActor->createShape(PGeom, PMaterial, PLocalPose);
+		PxShape* PNewShape = PDestActor->createShape(PGeom, PMaterial, PLocalPose, PShapeFlags);
 		if (PNewShape)
 		{
 			if (NewShapes)
