Expose GameMode functionality to Blueprints

diff --git a/Engine/Config/BaseGame.ini b/Engine/Config/BaseGame.ini
--- a/Engine/Config/BaseGame.ini
+++ b/Engine/Config/BaseGame.ini
@@ -1,9 +1,6 @@
 [Internationalization]
 +LocalizationPaths=%GAMEDIR%Content/Localization/Game
 
-[/Script/Engine.GameMode]
-InactivePlayerStateLifeSpan=300
-
 [DefaultPlayer]
 Name=Player
 
diff --git a/Engine/Source/Runtime/Engine/Classes/GameFramework/Actor.h b/Engine/Source/Runtime/Engine/Classes/GameFramework/Actor.h
--- a/Engine/Source/Runtime/Engine/Classes/GameFramework/Actor.h
+++ b/Engine/Source/Runtime/Engine/Classes/GameFramework/Actor.h
@@ -2131,6 +2131,10 @@ public:
 	/** Reset actor to initial state - used when restarting level without reloading. */
 	virtual void Reset();
 
+	/** Event called when this Actor is reset to its initial state - used when restarting level without reloading. */
+	UFUNCTION(BlueprintImplementableEvent, Category=Actor, meta=(DisplayName="OnReset"))
+	void K2_OnReset();
+
 	/** Returns the most recent time any of this actor's components were rendered */
 	virtual float GetLastRenderTime() const;
 
diff --git a/Engine/Source/Runtime/Engine/Classes/GameFramework/GameMode.h b/Engine/Source/Runtime/Engine/Classes/GameFramework/GameMode.h
--- a/Engine/Source/Runtime/Engine/Classes/GameFramework/GameMode.h
+++ b/Engine/Source/Runtime/Engine/Classes/GameFramework/GameMode.h
@@ -68,18 +68,23 @@ class ENGINE_API AGameMode : public AInfo
 	// Code to deal with the match state machine
 
 	/** Returns the current match state, this is an accessor to protect the state machine flow */
+	UFUNCTION(BlueprintCallable, Category="Game")
 	FName GetMatchState() const { return MatchState; }
 
 	/** Returns true if the match state is InProgress or later */
+	UFUNCTION(BlueprintCallable, Category="Game")
 	virtual bool HasMatchStarted() const;
 
 	/** Returns true if the match state is InProgress or other gameplay state */
+	UFUNCTION(BlueprintCallable, Category="Game")
 	virtual bool IsMatchInProgress() const;
 
 	/** Returns true if the match state is WaitingPostMatch or later */
+	UFUNCTION(BlueprintCallable, Category="Game")
 	virtual bool HasMatchEnded() const;
 
 	/** Transitions to WaitingToStart and calls BeginPlay on actors. */
+	UFUNCTION(BlueprintCallable, Category="Game")
 	virtual void StartPlay();
 
 	/**
@@ -89,45 +94,55 @@ class ENGINE_API AGameMode : public AInfo
 	virtual void PostSeamlessTravel();
 
 	/** Transition from WaitingToStart to InProgress. You can call this manually, will also get called if ReadyToStartMatch returns true */
+	UFUNCTION(BlueprintCallable, Category="Game")
 	virtual void StartMatch();
 
 	/** Transition from InProgress to WaitingPostMatch. You can call this manually, will also get called if ReadyToEndMatch returns true */
+	UFUNCTION(BlueprintCallable, Category="Game")
 	virtual void EndMatch();
 
 	/** Transition to LeavingMap state. Start the transition out of the current map. Called at start of seamless travel, or right before map change for hard travel. */
 	virtual void StartToLeaveMap();
 
 	/** Restart the game, by default travel to the current map */
+	UFUNCTION(BlueprintCallable, Category="Game")
 	virtual void RestartGame();
 
 	/** Return to main menu, and disconnect any players */
+	UFUNCTION(BlueprintCallable, Category="Game")
 	virtual void ReturnToMainMenuHost();
 
 	/** Report that a match has failed due to unrecoverable error */
+	UFUNCTION(BlueprintCallable, Category="Game")
 	virtual void AbortMatch();
 
 protected:
 
 	/** What match state we are currently in */
-	UPROPERTY()
 	FName MatchState;
 
 	/** Updates the match state and calls the appropriate transition functions */
 	virtual void SetMatchState(FName NewState);
 
+	/** Implementable event to respond to match state changes */
+	UFUNCTION(BlueprintImplementableEvent, Category="Game", meta=(DisplayName="OnSetMatchState"))
+	void K2_OnSetMatchState(FName NewState);
+
 	// Games should override these functions to deal with their game specific logic
 
 	/** Called when the state transitions to WaitingToStart */
 	virtual void HandleMatchIsWaitingToStart();
 
 	/** @return True if ready to Start Match. Games should override this */
-	virtual bool ReadyToStartMatch();
+	UFUNCTION(BlueprintNativeEvent, Category="Game")
+	bool ReadyToStartMatch();
 
 	/** Called when the state transitions to InProgress */
 	virtual void HandleMatchHasStarted();
 
 	/** @return true if ready to End Match. Games should override this */
-	virtual bool ReadyToEndMatch();
+	UFUNCTION(BlueprintNativeEvent, Category="Game")
+	bool ReadyToEndMatch();
 
 	/** Called when the map transitions to WaitingPostMatch */
 	virtual void HandleMatchHasEnded();
@@ -196,35 +211,30 @@ public:
 	int32 NumTravellingPlayers;
 
 	/** Used to assign unique PlayerIDs to each PlayerState. */
-	UPROPERTY()
 	int32 CurrentID;    
 
 	/** The default player name assigned to players that join with no name specified. */
-	UPROPERTY()
-	FString DefaultPlayerName;
-
-	/** Array of available player starts. */
-	UPROPERTY()
-	TArray<class APlayerStart*> PlayerStarts;
+	UPROPERTY(EditAnywhere, Category=GameMode)
+	FText DefaultPlayerName;
 
 	/** Contains strings describing localized game agnostic messages. */
 	UPROPERTY()
 	TSubclassOf<class ULocalMessage> EngineMessageClass;
 
 	/** The class of PlayerController to spawn for players logging in. */
-	UPROPERTY(EditAnywhere, noclear, BlueprintReadWrite, Category=Classes)
+	UPROPERTY(EditAnywhere, noclear, BlueprintReadOnly, Category=Classes)
 	TSubclassOf<class APlayerController> PlayerControllerClass;
 
 	/** The pawn class used by the PlayerController for players when spectating. */
-	UPROPERTY(EditAnywhere, noclear, BlueprintReadWrite, Category=Classes)
+	UPROPERTY(EditAnywhere, noclear, BlueprintReadOnly, Category=Classes)
 	TSubclassOf<class ASpectatorPawn> SpectatorClass;
 
 	/** A PlayerState of this class will be associated with every player to replicate relevant player information to all clients. */
-	UPROPERTY(EditAnywhere, noclear, BlueprintReadWrite, Category=Classes)
+	UPROPERTY(EditAnywhere, noclear, BlueprintReadOnly, Category=Classes)
 	TSubclassOf<class APlayerState> PlayerStateClass;
 
 	/** Class of GameState associated with this GameMode. */
-	UPROPERTY(EditAnywhere, noclear, BlueprintReadWrite, Category=Classes)
+	UPROPERTY(EditAnywhere, noclear, BlueprintReadOnly, Category=Classes)
 	TSubclassOf<class AGameState> GameStateClass;
 
 	/** GameState is used to replicate game state relevant properties to all clients. */
@@ -248,11 +258,11 @@ public:
 protected:
 
 	/** Handy alternate short names for GameMode classes (e.g. "DM" could be an alias for "MyProject.MyGameModeMP_DM". */
-	UPROPERTY(config)
+	UPROPERTY(EditAnywhere, Category=GameMode, AdvancedDisplay)
 	TArray<struct FGameClassShortName> GameModeClassAliases;
 
 	/** Time a playerstate will stick around in an inactive state after a player logout */
-	UPROPERTY(config)
+	UPROPERTY(EditAnywhere, Category=GameMode)
 	float InactivePlayerStateLifeSpan;
 
 	/** If true, dedicated servers will record replays when HandleMatchHasStarted/HandleMatchHasStopped is called */
@@ -273,12 +283,21 @@ public:
 	// End AActor interface
 
 	/** 
+	 * Overridable function to determine whether an Actor should have Reset called when the game has Reset called on it.
+	 * Default implementation returns true
+	 * @param ActorToReset The actor to make a determination for
 	 * @return true if ActorToReset should have Reset() called on it while restarting the game, 
 	 *		   false if the GameMode will manually reset it or if the actor does not need to be reset
+
 	 */
-	virtual bool ShouldReset(AActor* ActorToReset);
+	UFUNCTION(BlueprintNativeEvent, Category="Game")
+	bool ShouldReset(AActor* ActorToReset);
 
-	/** Resets level by calling Reset() on all actors */
+	/** 
+	  * Overridable function called when resetting level.
+	  * Default implementation calls Reset() on all actors except GameMode and Controllers
+	  */
+	UFUNCTION(BlueprintCallable, Category="Game")
 	virtual void ResetLevel();
 
 	/**
@@ -317,6 +336,7 @@ public:
 	void RemovePlayerControllerFromPlayerCount(APlayerController* PC);
 
 	/** Total number of players */
+	UFUNCTION(BlueprintCallable, Category="Game")
 	virtual int32 GetNumPlayers();
 
 	/**
@@ -452,26 +472,24 @@ public:
 	virtual APlayerController* SpawnPlayerController(ENetRole RemoteRole, FVector const& SpawnLocation, FRotator const& SpawnRotation);
 
 	/** @Returns true if NewPlayerController may only join the server as a spectator. */
-	virtual bool MustSpectate(APlayerController* NewPlayerController) const;
+	UFUNCTION(BlueprintNativeEvent, Category="Game")
+	bool MustSpectate(APlayerController* NewPlayerController) const;
 
 	/** returns default pawn class for given controller */
-	virtual UClass* GetDefaultPawnClassForController(AController* InController);
+	UFUNCTION(BlueprintNativeEvent, Category="Game")
+	UClass* GetDefaultPawnClassForController(AController* InController);
 
 	/** Called when StartSpot is selected for spawning NewPlayer to allow optional initialization. */
-	virtual void InitStartSpot(class AActor* StartSpot, AController* NewPlayer);
-
-	/** Register new player start */
-	virtual void AddPlayerStart(APlayerStart* NewPlayerStart);
-
-	/** Remove an existing player start */
-	virtual void RemovePlayerStart(APlayerStart* RemovedPlayerStart);
+	UFUNCTION(BlueprintNativeEvent, Category="Game")
+	void InitStartSpot(class AActor* StartSpot, AController* NewPlayer);
 
 	/**
 	 * @param	NewPlayer - Controller for whom this pawn is spawned
 	 * @param	StartSpot - PlayerStart at which to spawn pawn
 	 * @return	a pawn of the default pawn class
 	 */
-	virtual APawn* SpawnDefaultPawnFor(AController* NewPlayer, class AActor* StartSpot);
+	UFUNCTION(BlueprintNativeEvent, Category="Game")
+	APawn* SpawnDefaultPawnFor(AController* NewPlayer, class AActor* StartSpot);
 
 	/** replicates the current level streaming status to the given PlayerController */
 	virtual void ReplicateStreamingStatus(APlayerController* PC);
@@ -488,6 +506,10 @@ public:
 	/** Called when a Controller with a PlayerState leaves the match. */
 	virtual void Logout( AController* Exiting );
 
+	/** Implementable event when a Controller with a PlayerState leaves the match. */
+	UFUNCTION(BlueprintImplementableEvent, Category="Game", meta=(DisplayName="OnLogout"))
+	void K2_OnLogout(AController* ExitingController);
+
 	/**
 	 * Make sure pawn properties are back to default
 	 * Also a good place to modify them on spawn
@@ -495,10 +517,26 @@ public:
 	virtual void SetPlayerDefaults(APawn* PlayerPawn);
 
 	/** Return whether Viewer is allowed to spectate from the point of view of ViewTarget. */
-	virtual bool CanSpectate( APlayerController* Viewer, APlayerState* ViewTarget );
+	UFUNCTION(BlueprintNativeEvent, Category="Game")
+	bool CanSpectate( APlayerController* Viewer, APlayerState* ViewTarget );
+
+	/** 
+	 * Sets the name for a controller 
+	 * @param Controller	The controller of the player to change the name of
+	 * @param NewName		The name to set the player to
+	 * @param bNameChange	Whether the name is changing or if this is the first time it has been set
+	 */
+	UFUNCTION(BlueprintCallable, Category="Game")
+	virtual void ChangeName( AController* Controller, const FString& NewName, bool bNameChange );
 
-	/* Use reduce damage for teamplay modifications, etc. */
-	virtual void ChangeName( AController* Other, const FString& S, bool bNameChange );
+	/** 
+	 * Overridable event for GameMode blueprint to respond to a change name call
+	 * @param Controller	The controller of the player to change the name of
+	 * @param NewName		The name to set the player to
+	 * @param bNameChange	Whether the name is changing or if this is the first time it has been set
+	 */
+	UFUNCTION(BlueprintImplementableEvent,Category="Game",meta=(DisplayName="OnChangeName"))
+	void K2_OnChangeName( AController* Other, const FString& NewName, bool bNameChange );
 
 	/* Send a player to a URL.*/
 	virtual void SendPlayer( APlayerController* aPlayer, const FString& URL );
@@ -518,11 +556,12 @@ public:
 
 	/** 
 	 * Return the 'best' player start for this player to start from.
-	 * @param Player - is the AController for whom we are choosing a playerstart
-	 * @param IncomingName - specifies the tag of a Playerstart to use
+	 * @param Player The AController for whom we are choosing a Player Start
+	 * @param IncomingName Specifies the tag of a Player Start to use
 	 * @returns Actor chosen as player start (usually a PlayerStart)
 	 */
-	virtual class AActor* FindPlayerStart( AController* Player, const FString& IncomingName = TEXT("") );
+	UFUNCTION(BlueprintNativeEvent, Category="Game")
+	class AActor* FindPlayerStart( AController* Player, const FString& IncomingName = TEXT("") );
 
 	UFUNCTION(BlueprintPure, Category=Game, meta=(DisplayName="FindPlayerStart"))
 	AActor* K2_FindPlayerStart(AController* Player);
@@ -533,10 +572,12 @@ public:
 	* @param Player is the controller for whom we are choosing a playerstart
 	* @returns AActor chosen as player start (usually a PlayerStart)
 	 */
-	virtual class AActor* ChoosePlayerStart( AController* Player );
+	UFUNCTION(BlueprintNativeEvent, Category="Game")
+	class AActor* ChoosePlayerStart( AController* Player );
 
 	/** @return true if it's valid to call RestartPlayer. Will call Player->CanRestartPlayer */
-	virtual bool PlayerCanRestart( APlayerController* Player );
+	UFUNCTION(BlueprintNativeEvent, Category="Game")
+	bool PlayerCanRestart( APlayerController* Player );
 
 	/**
 	 * Used to notify the game type that it is ok to update a player's gameplay
@@ -578,7 +619,7 @@ public:
 	 * called on server during seamless level transitions to get the list of Actors that should be moved into the new level
 	 * PlayerControllers, Role < ROLE_Authority Actors, and any non-Actors that are inside an Actor that is in the list
 	 * (i.e. Object.Outer == Actor in the list)
-	 * are all autmoatically moved regardless of whether they're included here
+	 * are all automatically moved regardless of whether they're included here
 	 * only dynamic actors in the PersistentLevel may be moved (this includes all actors spawned during gameplay)
 	 * this is called for both parts of the transition because actors might change while in the middle (e.g. players might join or leave the game)
 	 * @see also PlayerController::GetSeamlessTravelActorList() (the function that's called on clients)
@@ -591,7 +632,7 @@ public:
 	virtual FString GetRedirectURL(const FString& MapName) const;
 
 	/** 
-	 * used to swap a viewport/connection's PlayerControllers when seamless travelling and the new GameMode's
+	 * used to swap a viewport/connection's PlayerControllers when seamless traveling and the new GameMode's
 	 * controller class is different than the previous
 	 * includes network handling
 	 * @param OldPC - the old PC that should be discarded
@@ -599,6 +640,9 @@ public:
 	 */
 	virtual void SwapPlayerControllers(APlayerController* OldPC, APlayerController* NewPC);
 
+	UFUNCTION(BlueprintImplementableEvent, Category="Game", meta=(DisplayName="SwapPlayerControllers"))
+	void K2_OnSwapPlayerControllers(APlayerController* OldPC, APlayerController* NewPC);
+
 	/** 
 	 * Handles reinitializing players that remained through a seamless level transition
 	 * called from C++ for players that finished loading after the server
@@ -618,17 +662,11 @@ public:
 	/** Given a string, return a fully-qualified GameMode class name */
 	static FString StaticGetFullGameClassName(FString const& Str);
 
-	/** Called periodically, overridden by subclasses */
-	virtual void DefaultTimer();	
-
 	/** Returns true if replays will start/stop during gameplay starting/stopping */
 	virtual bool IsHandlingReplays();
 
 protected:
 
-	/** Handle for efficient management of DefaultTimer timer */
-	FTimerHandle TimerHandle_DefaultTimer;
-
 	/**
 	 * Customize incoming player based on URL options
 	 *
diff --git a/Engine/Source/Runtime/Engine/Classes/GameFramework/PlayerStart.h b/Engine/Source/Runtime/Engine/Classes/GameFramework/PlayerStart.h
--- a/Engine/Source/Runtime/Engine/Classes/GameFramework/PlayerStart.h
+++ b/Engine/Source/Runtime/Engine/Classes/GameFramework/PlayerStart.h
@@ -29,11 +29,6 @@ private_subobject:
 public:
 #endif
 
-	// Begin AActor interface
-	virtual void PostInitializeComponents() override;	
-	virtual void PostUnregisterAllComponents() override;
-	// End AActor interface
-
 #if WITH_EDITORONLY_DATA
 	/** Returns ArrowComponent subobject **/
 	class UArrowComponent* GetArrowComponent() const;
diff --git a/Engine/Source/Runtime/Engine/Private/Actor.cpp b/Engine/Source/Runtime/Engine/Private/Actor.cpp
--- a/Engine/Source/Runtime/Engine/Private/Actor.cpp
+++ b/Engine/Source/Runtime/Engine/Private/Actor.cpp
@@ -1713,7 +1713,11 @@ void AActor::Destroyed()
 }
 
 void AActor::TornOff() {}
-void AActor::Reset() {}
+
+void AActor::Reset()
+{
+	K2_OnReset();
+}
 
 void AActor::FellOutOfWorld(const UDamageType& dmgType)
 {
diff --git a/Engine/Source/Runtime/Engine/Private/Controller.cpp b/Engine/Source/Runtime/Engine/Private/Controller.cpp
--- a/Engine/Source/Runtime/Engine/Private/Controller.cpp
+++ b/Engine/Source/Runtime/Engine/Private/Controller.cpp
@@ -433,7 +433,7 @@ void AController::InitPlayerState()
 			{
 				// don't call SetPlayerName() as that will broadcast entry messages but the GameMode hasn't had a chance
 				// to potentially apply a player/bot name yet
-				PlayerState->PlayerName = GameMode->DefaultPlayerName;
+				PlayerState->PlayerName = GameMode->DefaultPlayerName.ToString();
 			}
 		}
 	}
diff --git a/Engine/Source/Runtime/Engine/Private/GameMode.cpp b/Engine/Source/Runtime/Engine/Private/GameMode.cpp
--- a/Engine/Source/Runtime/Engine/Private/GameMode.cpp
+++ b/Engine/Source/Runtime/Engine/Private/GameMode.cpp
@@ -56,6 +56,7 @@ AGameMode::AGameMode(const FObjectInitializer& ObjectInitializer)
 	CurrentID = 1;
 	PlayerStateClass = APlayerState::StaticClass();
 	MinRespawnDelay = 1.0f;
+	InactivePlayerStateLifeSpan = 300.f;
 }
 
 FString AGameMode::GetNetworkNumber()
@@ -75,6 +76,8 @@ void AGameMode::SwapPlayerControllers(APlayerController* OldPC, APlayerControlle
 		NewPC->NetConnection = OldPC->NetConnection;
 		NewPC->CopyRemoteRoleFrom( OldPC );
 
+		K2_OnSwapPlayerControllers(OldPC, NewPC);
+
 		// send destroy event to old PC immediately if it's local
 		if (Cast<ULocalPlayer>(Player))
 		{
@@ -311,6 +314,8 @@ void AGameMode::Logout( AController* Exiting )
 	APlayerController* PC = Cast<APlayerController>(Exiting);
 	if ( PC != NULL )
 	{
+		K2_OnLogout(Exiting);
+
 		RemovePlayerControllerFromPlayerCount(PC);
 
 		if (GameSession)
@@ -336,7 +341,7 @@ void AGameMode::InitGameState()
 }
 
 
-AActor* AGameMode::FindPlayerStart( AController* Player, const FString& IncomingName )
+AActor* AGameMode::FindPlayerStart_Implementation( AController* Player, const FString& IncomingName )
 {
 	UWorld* World = GetWorld();
 
@@ -398,7 +403,6 @@ void AGameMode::PreInitializeComponents()
 {
 	Super::PreInitializeComponents();
 
-	GetWorldTimerManager().SetTimer(TimerHandle_DefaultTimer, this, &AGameMode::DefaultTimer, GetWorldSettings()->GetEffectiveTimeDilation(), true);
 	FActorSpawnParameters SpawnInfo;
 	SpawnInfo.Instigator = Instigator;
 	SpawnInfo.ObjectFlags |= RF_Transient;	// We never want to save game states or network managers into a map
@@ -511,7 +515,7 @@ void AGameMode::RestartPlayer(AController* NewPlayer)
 #endif	//!UE_WITH_PHYSICS
 }
 
-void AGameMode::InitStartSpot(AActor* StartSpot, AController* NewPlayer)
+void AGameMode::InitStartSpot_Implementation(AActor* StartSpot, AController* NewPlayer)
 {
 }
 
@@ -544,7 +548,7 @@ void AGameMode::HandleMatchIsWaitingToStart()
 	}
 }
 
-bool AGameMode::ReadyToStartMatch()
+bool AGameMode::ReadyToStartMatch_Implementation()
 {
 	// If bDelayed Start is set, wait for a manual match start
 	if (bDelayedStart)
@@ -627,7 +631,7 @@ void AGameMode::HandleMatchHasStarted()
 	}
 }
 
-bool AGameMode::ReadyToEndMatch()
+bool AGameMode::ReadyToEndMatch_Implementation()
 {
 	// By default don't explicitly end match
 	return false;
@@ -739,6 +743,8 @@ void AGameMode::SetMatchState(FName NewState)
 	{
 		GameState->SetMatchState(NewState);
 	}
+
+	K2_OnSetMatchState(NewState);
 }
 
 void AGameMode::Tick(float DeltaSeconds)
@@ -1007,7 +1013,7 @@ FString AGameMode::InitNewPlayer(APlayerController* NewPlayerController, const T
 	FString InName = ParseOption(Options, TEXT("Name")).Left(20);
 	if (InName.IsEmpty())
 	{
-		InName = FString::Printf(TEXT("%s%i"), *DefaultPlayerName, NewPlayerController->PlayerState->PlayerId);
+		InName = FString::Printf(TEXT("%s%i"), *DefaultPlayerName.ToString(), NewPlayerController->PlayerState->PlayerId);
 	}
 
 	ChangeName(NewPlayerController, InName, false);
@@ -1030,7 +1036,7 @@ FString AGameMode::InitNewPlayer(APlayerController* NewPlayerController, const T
 	return ErrorMessage;
 }
 
-bool AGameMode::MustSpectate(APlayerController* NewPlayerController) const
+bool AGameMode::MustSpectate_Implementation(APlayerController* NewPlayerController) const
 {
 	return NewPlayerController->PlayerState->bOnlySpectator;
 }
@@ -1085,7 +1091,7 @@ void AGameMode::Reset()
 }
 
 
-bool AGameMode::ShouldReset(AActor* ActorToReset)
+bool AGameMode::ShouldReset_Implementation(AActor* ActorToReset)
 {
 	return true;
 }
@@ -1308,12 +1314,12 @@ APlayerController* AGameMode::SpawnPlayerController(ENetRole RemoteRole, FVector
 	return NewPC;
 }
 
-UClass* AGameMode::GetDefaultPawnClassForController(AController* InController)
+UClass* AGameMode::GetDefaultPawnClassForController_Implementation(AController* InController)
 {
 	return DefaultPawnClass;
 }
 
-APawn* AGameMode::SpawnDefaultPawnFor(AController* NewPlayer, AActor* StartSpot)
+APawn* AGameMode::SpawnDefaultPawnFor_Implementation(AController* NewPlayer, AActor* StartSpot)
 {
 	// don't allow pawn to be spawned with any pitch or roll
 	FRotator StartRotation(ForceInit);
@@ -1429,7 +1435,7 @@ void AGameMode::StartNewPlayer(APlayerController* NewPlayer)
 	}
 }
 
-bool AGameMode::CanSpectate( APlayerController* Viewer, APlayerState* ViewTarget )
+bool AGameMode::CanSpectate_Implementation( APlayerController* Viewer, APlayerState* ViewTarget )
 {
 	return true;
 }
@@ -1447,7 +1453,6 @@ void AGameMode::SendPlayer( APlayerController* aPlayer, const FString& FURL )
 	aPlayer->ClientTravel( FURL, TRAVEL_Relative );
 }
 
-
 bool AGameMode::GetTravelType()
 {
 	return false;
@@ -1485,34 +1490,24 @@ bool AGameMode::ShouldSpawnAtStartSpot(AController* Player)
 	return ( Player != NULL && Player->StartSpot != NULL );
 }
 
-void AGameMode::AddPlayerStart(APlayerStart* NewPlayerStart)
-{
-	PlayerStarts.AddUnique(NewPlayerStart);
-}
-
-void AGameMode::RemovePlayerStart(APlayerStart* RemovedPlayerStart)
-{
-	PlayerStarts.Remove(RemovedPlayerStart);
-}
-
-AActor* AGameMode::ChoosePlayerStart( AController* Player )
+AActor* AGameMode::ChoosePlayerStart_Implementation( AController* Player )
 {
 	// Choose a player start
 	APlayerStart* FoundPlayerStart = NULL;
 	APawn* PawnToFit = DefaultPawnClass ? DefaultPawnClass->GetDefaultObject<APawn>() : NULL;
 	TArray<APlayerStart*> UnOccupiedStartPoints;
 	TArray<APlayerStart*> OccupiedStartPoints;
-	for (int32 PlayerStartIndex = 0; PlayerStartIndex < PlayerStarts.Num(); ++PlayerStartIndex)
+	for (TActorIterator<APlayerStart> It(GetWorld()); It; ++It)
 	{
-		APlayerStart* PlayerStart = PlayerStarts[PlayerStartIndex];
+		APlayerStart* PlayerStart = *It;
 
-		if (Cast<APlayerStartPIE>( PlayerStart ) != NULL )
+		if (PlayerStart->IsA<APlayerStartPIE>())
 		{
 			// Always prefer the first "Play from Here" PlayerStart, if we find one while in PIE mode
 			FoundPlayerStart = PlayerStart;
 			break;
 		}
-		else if (PlayerStart != NULL)
+		else
 		{
 			FVector ActorLocation = PlayerStart->GetActorLocation();
 			const FRotator ActorRotation = PlayerStart->GetActorRotation();
@@ -1540,7 +1535,7 @@ AActor* AGameMode::ChoosePlayerStart( AController* Player )
 	return FoundPlayerStart;
 }
 
-bool AGameMode::PlayerCanRestart( APlayerController* Player )
+bool AGameMode::PlayerCanRestart_Implementation( APlayerController* Player )
 {
 	if (!IsMatchInProgress() || Player == NULL || Player->IsPendingKillPending())
 	{
@@ -1739,10 +1734,6 @@ FString AGameMode::StaticGetFullGameClassName(FString const& Str)
 	return Str;
 }
 
-void AGameMode::DefaultTimer()
-{	
-}
-
 FString AGameMode::GetRedirectURL(const FString& MapName) const
 {
 	return FString();
diff --git a/Engine/Source/Runtime/Engine/Private/PlayerStart.cpp b/Engine/Source/Runtime/Engine/Private/PlayerStart.cpp
--- a/Engine/Source/Runtime/Engine/Private/PlayerStart.cpp
+++ b/Engine/Source/Runtime/Engine/Private/PlayerStart.cpp
@@ -61,26 +61,6 @@ APlayerStart::APlayerStart(const FObjectInitializer& ObjectInitializer)
 #endif // WITH_EDITORONLY_DATA
 }
 
-void APlayerStart::PostInitializeComponents()
-{
-	Super::PostInitializeComponents();
-	if ( !IsPendingKill()  && GetWorld()->GetAuthGameMode() )
-	{
-		GetWorld()->GetAuthGameMode()->AddPlayerStart(this);
-	}
-}
-
-void APlayerStart::PostUnregisterAllComponents()
-{
-	Super::PostUnregisterAllComponents();
-
-	UWorld* ActorWorld = GetWorld();
-	if ( ActorWorld && ActorWorld->GetAuthGameMode() )
-	{
-		ActorWorld->GetAuthGameMode()->RemovePlayerStart(this);
-	}
-}
-
 #if WITH_EDITORONLY_DATA
 /** Returns ArrowComponent subobject **/
 UArrowComponent* APlayerStart::GetArrowComponent() const { return ArrowComponent; }
diff --git a/Engine/Source/Runtime/Engine/Private/UnrealEngine.cpp b/Engine/Source/Runtime/Engine/Private/UnrealEngine.cpp
--- a/Engine/Source/Runtime/Engine/Private/UnrealEngine.cpp
+++ b/Engine/Source/Runtime/Engine/Private/UnrealEngine.cpp
@@ -8955,6 +8955,14 @@ bool UEngine::LoadMap( FWorldContext& WorldContext, FURL URL, class UPendingNetG
 			}
 		}
 
+		for (TObjectIterator<UObject> It(RF_PendingKill); It; ++It)
+		{
+			if (It->GetTypedOuter<UWorld>())
+			{
+				break;
+			}
+		}
+
 		// Stop all audio to remove references to current level.
 		if (FAudioDevice* AudioDevice = WorldContext.World()->GetAudioDevice())
 		{
