Fix build. Using other overload takes ref instead of ptr

diff --git a/Engine/Source/Runtime/Engine/Private/PhysicsEngine/BodySetup.cpp b/Engine/Source/Runtime/Engine/Private/PhysicsEngine/BodySetup.cpp
--- a/Engine/Source/Runtime/Engine/Private/PhysicsEngine/BodySetup.cpp
+++ b/Engine/Source/Runtime/Engine/Private/PhysicsEngine/BodySetup.cpp
@@ -585,7 +585,7 @@ private:
 
 	PxShape* AttachShape(const PxGeometry& PGeom, const PxMaterial& PMaterial, const PxTransform& PLocalPose, const float ContactOffset, PxShapeFlags PShapeFlags = PxShapeFlag::eVISUALIZATION | PxShapeFlag::eSCENE_QUERY_SHAPE | PxShapeFlag::eSIMULATION_SHAPE) const
 	{
-		PxShape* PNewShape = PDestActor->createShape(PGeom, PMaterial, PLocalPose, PShapeFlags);
+		PxShape* PNewShape = PDestActor->createShape(PGeom, *PMaterial, PLocalPose, PShapeFlags);
 		if (PNewShape)
 		{
 			if (NewShapes)
