Network replay backwards compatibility initial pass + add internal version to replays + better spectator controller persisting across scrubs

diff --git a/Engine/Source/Runtime/Engine/Classes/Engine/ActorChannel.h b/Engine/Source/Runtime/Engine/Classes/Engine/ActorChannel.h
--- a/Engine/Source/Runtime/Engine/Classes/Engine/ActorChannel.h
+++ b/Engine/Source/Runtime/Engine/Classes/Engine/ActorChannel.h
@@ -250,6 +250,8 @@ protected:
 	TSharedRef< FObjectReplicator > & FindOrCreateReplicator(UObject *Obj);
 	bool ObjectHasReplicator(UObject *Obj);	// returns whether we have already created a replicator for this object or not
 
+	void DestroyActorAndComponents();
+
 	virtual bool CleanUp( const bool bForDestroy ) override;
 
 	/** Closes the actor channel but with a 'dormant' flag set so it can be reopened */
diff --git a/Engine/Source/Runtime/Engine/Classes/Engine/DemoNetDriver.h b/Engine/Source/Runtime/Engine/Classes/Engine/DemoNetDriver.h
--- a/Engine/Source/Runtime/Engine/Classes/Engine/DemoNetDriver.h
+++ b/Engine/Source/Runtime/Engine/Classes/Engine/DemoNetDriver.h
@@ -67,16 +67,11 @@ class ENGINE_API UDemoNetDriver : public UNetDriver
 
 	void		SaveCheckpoint();
 
-	bool		bRestoreSpectatorPosition;
-	FVector		SpectatorLocation;
-	FRotator	SpectatorRotation;
-
 	FArchive*	GotoCheckpointArchive;
 	int64		GotoCheckpointSkipExtraTimeInMS;
 
 	void		LoadCheckpoint();
 
-
 private:
 	bool		bIsFastForwarding;
 	bool		bIsLoadingCheckpoint;
diff --git a/Engine/Source/Runtime/Engine/Classes/Engine/PackageMapClient.h b/Engine/Source/Runtime/Engine/Classes/Engine/PackageMapClient.h
--- a/Engine/Source/Runtime/Engine/Classes/Engine/PackageMapClient.h
+++ b/Engine/Source/Runtime/Engine/Classes/Engine/PackageMapClient.h
@@ -68,6 +68,8 @@ public:
 	bool			IsGUIDBroken( const FNetworkGUID& NetGUID, const bool bMustBeRegistered ) const;
 	FString			FullNetGUIDPath( const FNetworkGUID& NetGUID ) const;
 	void			GenerateFullNetGUIDPath_r( const FNetworkGUID& NetGUID, FString& FullPath ) const;
+	void			SetIgnorePackageMismatchOverride( const bool bInIgnorePackageMismatchOverride ) { bIgnorePackageMismatchOverride = bInIgnorePackageMismatchOverride; }
+	bool			ShouldIgnorePackageMismatch() const;
 
 	void			AsyncPackageCallback(const FName& PackageName, UPackage * Package, EAsyncLoadingResult::Type Result);
 	
@@ -81,6 +83,8 @@ public:
 
 	TMap< FName, FNetworkGUID >						PendingAsyncPackages;
 
+	bool											bIgnorePackageMismatchOverride;
+
 #if !(UE_BUILD_SHIPPING || UE_BUILD_TEST)
 	// History for debugging entries in the guid cache
 	TMap<FNetworkGUID, FString>						History;
diff --git a/Engine/Source/Runtime/Engine/Private/DataChannel.cpp b/Engine/Source/Runtime/Engine/Private/DataChannel.cpp
--- a/Engine/Source/Runtime/Engine/Private/DataChannel.cpp
+++ b/Engine/Source/Runtime/Engine/Private/DataChannel.cpp
@@ -1436,6 +1436,28 @@ void UActorChannel::CleanupReplicators( const bool bKeepReplicators )
 	ActorReplicator = NULL;
 }
 
+void UActorChannel::DestroyActorAndComponents()
+{
+	// Destroy any sub-objects we created
+	for ( int32 i = 0; i < CreateSubObjects.Num(); i++ )
+	{
+		if ( CreateSubObjects[i].IsValid() )
+		{
+			UObject *SubObject = CreateSubObjects[i].Get();
+			Actor->OnSubobjectDestroyFromReplication(SubObject);
+			SubObject->MarkPendingKill();
+		}
+	}
+
+	CreateSubObjects.Empty();
+
+	// Destroy the actor
+	if ( Actor != NULL )
+	{
+		Actor->Destroy( true );
+	}
+}
+
 bool UActorChannel::CleanUp( const bool bForDestroy )
 {
 	const bool bIsServer = Connection->Driver->IsServer();
@@ -1467,21 +1489,7 @@ bool UActorChannel::CleanUp( const bool bForDestroy )
 			{
 				UE_LOG(LogNetDormancy, Verbose, TEXT("Channel[%d] '%s' Destroying Actor."), ChIndex, *Describe() );
 
-				// Destroy any sub-objects we created
-				for ( int32 i = 0; i < CreateSubObjects.Num(); i++ )
-				{
-					if ( CreateSubObjects[i].IsValid() )
-					{
-						UObject *SubObject = CreateSubObjects[i].Get();
-						Actor->OnSubobjectDestroyFromReplication(SubObject);
-						SubObject->MarkPendingKill();
-					}
-				}
-
-				CreateSubObjects.Empty();
-
-				// Destroy the actor
-				Actor->Destroy( true );
+				DestroyActorAndComponents();
 			}
 		}
 	}
@@ -1933,9 +1941,25 @@ void UActorChannel::ProcessBunch( FInBunch & Bunch )
 
 		TSharedRef< FObjectReplicator > & Replicator = FindOrCreateReplicator( RepObj );
 
+		if ( Connection->InternalAck && ( Bunch.bOpen || RepObj != Actor  ) )
+		{
+			// Receive network checksum
+			if ( !Replicator->ReadNetworkChecksum( Bunch ) )
+			{
+				UE_LOG( LogNet, Warning, TEXT( "UActorChannel::ProcessBunch: ReadNetworkChecksum failed. RepObj: %s, Channel: %i"), RepObj ? *RepObj->GetFullName() : TEXT( "NULL" ), ChIndex );
+
+				DestroyActorAndComponents();
+				Actor = NULL;
+
+				// Mark as broken so we stop trying to process bunches on this channel
+				Broken = 1;
+				break;
+			}
+		}
+
 		bool bHasUnmapped = false;
 
-		if ( !Replicator.Get().ReceivedBunch( Bunch, RepFlags, bHasUnmapped ) )
+		if ( !Replicator->ReceivedBunch( Bunch, RepFlags, bHasUnmapped ) )
 		{
 			UE_LOG( LogNet, Error, TEXT( "UActorChannel::ProcessBunch: Replicator.ReceivedBunch failed.  Closing connection. RepObj: %s, Channel: %i"), RepObj ? *RepObj->GetFullName() : TEXT( "NULL" ), ChIndex  );
 			Connection->Close();
@@ -2285,6 +2309,11 @@ void UActorChannel::BeginContentBlock( UObject* Obj, FOutBunch &Bunch )
 
 	if ( IsActor )
 	{
+		if ( Connection->InternalAck && OpenPacketId.First == INDEX_NONE )
+		{
+			// Write network checksum
+			ActorReplicator->WriteNetworkChecksum( Bunch );
+		}
 		NETWORK_PROFILER(GNetworkProfiler.TrackBeginContentBlock(Obj, Bunch.GetNumBits() - NumStartingBits));
 		return;
 	}
@@ -2308,6 +2337,12 @@ void UActorChannel::BeginContentBlock( UObject* Obj, FOutBunch &Bunch )
 		}
 	}
 
+	if ( Connection->InternalAck )
+	{
+		// Write network checksum
+		FindOrCreateReplicator( Obj )->WriteNetworkChecksum( Bunch );
+	}
+
 #if !(UE_BUILD_SHIPPING || UE_BUILD_TEST)
 	if (CVarDoReplicationContextString->GetInt() > 0)
 	{
diff --git a/Engine/Source/Runtime/Engine/Private/DataReplication.cpp b/Engine/Source/Runtime/Engine/Private/DataReplication.cpp
--- a/Engine/Source/Runtime/Engine/Private/DataReplication.cpp
+++ b/Engine/Source/Runtime/Engine/Private/DataReplication.cpp
@@ -390,6 +390,16 @@ void FObjectReplicator::ReceivedNak( int32 NakPacketId )
 	}
 }
 
+void FObjectReplicator::WriteNetworkChecksum( FOutBunch& Bunch )
+{
+	RepLayout->WriteNetworkChecksum( Bunch );
+}
+
+bool FObjectReplicator::ReadNetworkChecksum( FInBunch& Bunch )
+{
+	return RepLayout->ReadNetworkChecksum( Bunch );
+}
+
 bool FObjectReplicator::ReceivedBunch( FInBunch &Bunch, const FReplicationFlags& RepFlags, bool & bOutHasUnmapped )
 {
 	UObject* Object = GetObject();
diff --git a/Engine/Source/Runtime/Engine/Private/DemoNetDriver.cpp b/Engine/Source/Runtime/Engine/Private/DemoNetDriver.cpp
--- a/Engine/Source/Runtime/Engine/Private/DemoNetDriver.cpp
+++ b/Engine/Source/Runtime/Engine/Private/DemoNetDriver.cpp
@@ -96,25 +96,28 @@ enum ENetworkVersionHistory
 {
 	HISTORY_INITIAL				= 1,
 	HISTORY_SAVE_ABS_TIME_MS	= 2,			// We now save the abs demo time in ms for each frame (solves accumulation errors)
-	HISTORY_INCREASE_BUFFER		= 3				// Increased buffer size of packets, which invalidates old replays
+	HISTORY_INCREASE_BUFFER		= 3,			// Increased buffer size of packets, which invalidates old replays
+	HISTORY_SAVE_ENGINE_VERSION	= 4				// Now saving engine net version + InternalProtocolVersion
 };
 
 static const uint32 NETWORK_DEMO_MAGIC				= 0x2CF5A13D;
-static const uint32 NETWORK_DEMO_VERSION			= HISTORY_INCREASE_BUFFER;
+static const uint32 NETWORK_DEMO_VERSION			= HISTORY_SAVE_ENGINE_VERSION;
 
 static const uint32 NETWORK_DEMO_METADATA_MAGIC		= 0x3D06B24E;
 static const uint32 NETWORK_DEMO_METADATA_VERSION	= 0;
 
 struct FNetworkDemoHeader
 {
-	uint32	Magic;					// Magic to ensure we're opening the right file.
-	uint32	Version;				// Version number to detect version mismatches.
-	uint32	EngineNetVersion;		// Version of engine networking format
-	FString LevelName;				// Name of level loaded for demo
+	uint32	Magic;						// Magic to ensure we're opening the right file.
+	uint32	Version;					// Version number to detect version mismatches.
+	uint32	InternalProtocolVersion;	// Version of the engine internal network format
+	uint32	EngineNetVersion;			// Version of engine networking format
+	FString LevelName;					// Name of level loaded for demo
 	
 	FNetworkDemoHeader() : 
 		Magic( NETWORK_DEMO_MAGIC ), 
 		Version( NETWORK_DEMO_VERSION ),
+		InternalProtocolVersion( FNetworkVersion::InternalProtocolVersion ),
 		EngineNetVersion( GEngineNetVersion )
 	{}
 
@@ -122,6 +125,8 @@ struct FNetworkDemoHeader
 	{
 		Ar << Header.Magic;
 		Ar << Header.Version;
+		Ar << Header.InternalProtocolVersion;
+		Ar << Header.EngineNetVersion;
 		Ar << Header.LevelName;
 
 		return Ar;
@@ -188,6 +193,8 @@ bool UDemoNetDriver::InitConnect( FNetworkNotify* InNotify, const FURL& ConnectU
 		return false;
 	}
 
+	GuidCache->SetIgnorePackageMismatchOverride( true );
+
 	// Playback, local machine is a client, and the demo stream acts "as if" it's the server.
 	ServerConnection = NewObject<UNetConnection>(GetTransientPackage(), UDemoNetConnection::StaticClass());
 	ServerConnection->InitConnection( this, USOCK_Pending, ConnectURL, 1000000 );
@@ -237,6 +244,14 @@ bool UDemoNetDriver::InitConnectInternal( FString& Error )
 		return false;
 	}
 
+	if ( DemoHeader.InternalProtocolVersion != FNetworkVersion::InternalProtocolVersion )
+	{
+		Error = FString( TEXT( "Engine internal network version is incorrect" ) );
+		UE_LOG( LogDemo, Error, TEXT( "UDemoNetDriver::InitConnect: %s" ), *Error );
+		GameInstance->HandleDemoPlaybackFailure( EDemoPlayFailure::InvalidVersion, Error );
+		return false;
+	}
+
 	// Create fake control channel
 	ServerConnection->CreateChannel( CHTYPE_Control, 1 );
 
@@ -350,6 +365,8 @@ bool UDemoNetDriver::InitListen( FNetworkNotify* InNotify, FURL& ListenURL, bool
 		return false;
 	}
 
+	GuidCache->SetIgnorePackageMismatchOverride( true );
+
 	check( World != NULL );
 
 	class AWorldSettings * WorldSettings = World->GetWorldSettings(); 
@@ -1418,15 +1435,10 @@ void UDemoNetDriver::LoadCheckpoint()
 
 	// Save off the current spectator position
 	// Check for NULL, which can be the case if we haven't played any of the demo yet but want to fast forward (joining live game for example)
-	if ( SpectatorController != NULL && SpectatorController->GetSpectatorPawn() != NULL )
+	if ( SpectatorController != NULL )
 	{
-		bRestoreSpectatorPosition = true;
-
-		SpectatorLocation = SpectatorController->GetSpectatorPawn()->GetActorLocation();
-		SpectatorRotation = SpectatorController->GetControlRotation();//GetSpectatorPawn()->GetActorRotation();
-
 		// Save off the SpectatorController's GUID so that we know not to queue his bunches
-		AddNonQueuedActorForScrubbing(SpectatorController);
+		AddNonQueuedActorForScrubbing( SpectatorController );
 	}
 
 	PauseChannels( false );
@@ -1444,11 +1456,42 @@ void UDemoNetDriver::LoadCheckpoint()
 		{
 			AddNonQueuedActorForScrubbing(*It);
 		}
+		
+		if ( *It == SpectatorController )
+		{
+			continue;
+		}
 
-		if ( !It->IsNetStartupActor() )
+		if ( It->GetOwner() == SpectatorController )
 		{
-			GetWorld()->DestroyActor( *It, true );
+			continue;
 		}
+
+		if ( It->IsNetStartupActor() )
+		{
+			continue;
+		}
+
+		GetWorld()->DestroyActor( *It, true );
+	}
+
+	// Find the SpectatorController on the channels, and make sure shutting down the connection doesn't destroy this actor
+	for ( int32 i = ServerConnection->OpenChannels.Num() - 1; i >= 0; i-- )
+	{
+		UChannel* OpenChannel = ServerConnection->OpenChannels[i];
+		if ( OpenChannel != NULL )
+		{
+			UActorChannel* ActorChannel = Cast< UActorChannel >( OpenChannel );
+			if ( ActorChannel != NULL && ActorChannel->Actor == SpectatorController )
+			{
+				ActorChannel->Actor = NULL;
+			}
+		}
+	}
+
+	if ( ServerConnection->OwningActor == SpectatorController )
+	{
+		ServerConnection->OwningActor = NULL;
 	}
 #else
 	for ( int32 i = ServerConnection->OpenChannels.Num() - 1; i >= 0; i-- )
@@ -1465,8 +1508,6 @@ void UDemoNetDriver::LoadCheckpoint()
 	}
 #endif
 
-	SpectatorController = NULL;
-
 	ServerConnection->Close();
 	ServerConnection->CleanUp();
 
@@ -1479,9 +1520,25 @@ void UDemoNetDriver::LoadCheckpoint()
 	// Create fake control channel
 	ServerConnection->CreateChannel( CHTYPE_Control, 1 );
 
+	// Remember the spectator network guid, so we can persist the spectator player across the checkpoint
+	// (so the state and position persists)
+	const FNetworkGUID SpectatorGUID = GuidCache->NetGUIDLookup.FindRef( SpectatorController );
+
+	// Clean package map to prepare to restore it to the checkpoint state
 	GuidCache->ObjectLookup.Empty();
 	GuidCache->NetGUIDLookup.Empty();
 
+	// Restore the spectator controller packagemap entry (so we find it when we process the checkpoint)
+	if ( SpectatorGUID.IsValid() )
+	{
+		FNetGuidCacheObject& CacheObject = GuidCache->ObjectLookup.FindOrAdd( SpectatorGUID );
+
+		CacheObject.Object = SpectatorController;
+		check( CacheObject.Object != NULL );
+		CacheObject.bNoLoad = true;
+		GuidCache->NetGUIDLookup.Add( SpectatorController, SpectatorGUID );
+	}
+
 	if ( GotoCheckpointArchive->TotalSize() == 0 || GotoCheckpointArchive->TotalSize() == INDEX_NONE )
 	{
 		// This is the very first checkpoint, we'll read the stream from the very beginning in this case
@@ -1529,7 +1586,6 @@ void UDemoNetDriver::LoadCheckpoint()
 		CacheObject.bIgnoreWhenMissing = ( Flags & ( 1 << 1 ) ) ? true : false;		
 
 		GuidCache->ObjectLookup.Add( Guid, CacheObject );
-		//GuidCache->GetObjectFromNetGUID( Guid, false );
 	}
 
 	uint32 SavedAbsTimeMS = 0;
@@ -1554,12 +1610,6 @@ void UDemoNetDriver::LoadCheckpoint()
 	GotoCheckpointSkipExtraTimeInMS = -1;
 	GotoCheckpointArchive			= NULL;
 	bDemoPlaybackDone				= false;
-
-	// Make sure we reset the bRestoreSpectatorPosition. If not, we didn't receive the spectator, which should be replicated in every frame
-	if ( bRestoreSpectatorPosition )
-	{
-		UE_LOG( LogDemo, Warning, TEXT( "UDemoNetConnection::LoadCheckpoint: Spectator wasn't restored." ) );
-	}
 }
 
 bool UDemoNetDriver::ShouldQueueBunchesForActorGUID(FNetworkGUID InGUID) const
@@ -1695,19 +1745,24 @@ void UDemoNetConnection::FlushNet( bool bIgnoreSimulation )
 
 void UDemoNetConnection::HandleClientPlayer( APlayerController* PC, UNetConnection* NetConnection )
 {
+	// If the spectator is the same, assume this is for scrubbing, and we are keeping the old one
+	// (so don't set the position, since we want to persist all that)
+	if ( GetDriver()->SpectatorController == PC )
+	{
+		PC->Role			= ROLE_AutonomousProxy;
+		PC->NetConnection	= NetConnection;
+		LastReceiveTime		= Driver->Time;
+		State				= USOCK_Open;
+		PlayerController	= PC;
+		OwningActor			= PC;
+		return;
+	}
+
 	Super::HandleClientPlayer( PC, NetConnection );
 
 	// Assume this is our special spectator controller
 	GetDriver()->SpectatorController = PC;
 
-	// Persist location and rotation from previous spectator
-	if ( GetDriver()->bRestoreSpectatorPosition )
-	{
-		GetDriver()->SpectatorController->SetInitialLocationAndRotation( GetDriver()->SpectatorLocation, GetDriver()->SpectatorRotation );
-		GetDriver()->bRestoreSpectatorPosition = false;
-		return;
-	}
-
 	for ( FActorIterator It( Driver->World ); It; ++It)
 	{
 		if ( It->IsA( APlayerStart::StaticClass() ) )
diff --git a/Engine/Source/Runtime/Engine/Private/PackageMapClient.cpp b/Engine/Source/Runtime/Engine/Private/PackageMapClient.cpp
--- a/Engine/Source/Runtime/Engine/Private/PackageMapClient.cpp
+++ b/Engine/Source/Runtime/Engine/Private/PackageMapClient.cpp
@@ -649,7 +649,7 @@ FNetworkGUID UPackageMapClient::InternalLoadObject( FArchive & Ar, UObject *& Ob
 					return NetGUID;
 				}
 
-				if ( Package->GetGuid() != PackageGuid && CVarIgnorePackageMismatch.GetValueOnGameThread() == 0 )
+				if ( Package->GetGuid() != PackageGuid && !GuidCache->ShouldIgnorePackageMismatch() )
 				{
 					UE_LOG( LogNetPackageMap, Error, TEXT( "UPackageMapClient::InternalLoadObject: Default object package guid mismatch! PathName: %s, ObjOuter: %s, GUID1: %s, GUID2: %s " ), *PathName, ObjOuter != NULL ? *ObjOuter->GetPathName() : TEXT( "NULL" ), *Package->GetGuid().ToString(), *PackageGuid.ToString() );
 					Object = NULL;
@@ -1263,10 +1263,9 @@ UObject* UPackageMapClient::GetObjectFromNetGUID( const FNetworkGUID& NetGUID, c
 //	FNetGUIDCache
 //----------------------------------------------------------------------------------------
 
-FNetGUIDCache::FNetGUIDCache( UNetDriver * InDriver ) : Driver( InDriver )
+FNetGUIDCache::FNetGUIDCache( UNetDriver * InDriver ) : IsExportingNetGUIDBunch( false ), Driver( InDriver ), bIgnorePackageMismatchOverride( false )
 {
 	UniqueNetIDs[0] = UniqueNetIDs[1] = 0;
-	IsExportingNetGUIDBunch = false;
 }
 
 class FArchiveCountMemGUID : public FArchive
@@ -1653,7 +1652,7 @@ void FNetGUIDCache::AsyncPackageCallback(const FName& PackageName, UPackage * Pa
 		CacheObject->bIsBroken = true;
 		UE_LOG( LogNetPackageMap, Error, TEXT( "AsyncPackageCallback: Package FAILED to load. Path: %s, NetGUID: %s" ), *PackageName.ToString(), *NetGUID.ToString() );
 	}
-	else if ( Package->GetGuid() != CacheObject->PackageGuid && CVarIgnorePackageMismatch.GetValueOnGameThread() == 0 )
+	else if ( Package->GetGuid() != CacheObject->PackageGuid && !ShouldIgnorePackageMismatch() )
 	{
 		CacheObject->bIsBroken = true;
 		UE_LOG( LogNetPackageMap, Error, TEXT( "AsyncPackageCallback: Package GUID mismatch! Path: %s, NetGUID: %s, GUID1: %s, GUID2: %s" ), *PackageName.ToString(), *NetGUID.ToString(), *Package->GetGuid().ToString(), *CacheObject->PackageGuid.ToString() );
@@ -1852,7 +1851,7 @@ UObject* FNetGUIDCache::GetObjectFromNetGUID( const FNetworkGUID& NetGUID, const
 			}
 		}
 
-		if ( Package->GetGuid() != CacheObjectPtr->PackageGuid && CVarIgnorePackageMismatch.GetValueOnGameThread() == 0 )
+		if ( Package->GetGuid() != CacheObjectPtr->PackageGuid && !ShouldIgnorePackageMismatch() )
 		{
 			// If the package guid doesn't match, don't allow it to load
 			CacheObjectPtr->bIsBroken = true;
@@ -2043,6 +2042,11 @@ void FNetGUIDCache::GenerateFullNetGUIDPath_r( const FNetworkGUID& NetGUID, FStr
 	}
 }
 
+bool FNetGUIDCache::ShouldIgnorePackageMismatch() const
+{
+	return bIgnorePackageMismatchOverride || CVarIgnorePackageMismatch.GetValueOnGameThread();
+};
+
 //------------------------------------------------------
 // Debug command to see how many times we've exported each NetGUID
 // Used for measuring inefficiencies. Some duplication is unavoidable since we cannot garuntee atomicicity across multiple channels.
@@ -2081,6 +2085,4 @@ FAutoConsoleCommand	ListNetGUIDExportsCommand(
 	FConsoleCommandDelegate::CreateStatic(ListNetGUIDExports)
 	);
 
-
-
 // ----------------------------------------------------------------
diff --git a/Engine/Source/Runtime/Engine/Private/RepLayout.cpp b/Engine/Source/Runtime/Engine/Private/RepLayout.cpp
--- a/Engine/Source/Runtime/Engine/Private/RepLayout.cpp
+++ b/Engine/Source/Runtime/Engine/Private/RepLayout.cpp
@@ -2180,6 +2180,67 @@ uint16 FRepLayout::AddParentProperty( UProperty * Property, int32 ArrayIndex )
 
 extern bool IsCustomDeltaProperty( UProperty * Property );
 
+static uint32 GetNetFieldChecksum( UField* NetField, const uint32 Checksum )
+{
+	uint32 LocalChecksum = FCrc::StrCrc32( *NetField->GetName(), Checksum );
+
+	LocalChecksum = FCrc::StrCrc32( *NetField->GetClass()->GetName(), LocalChecksum );
+
+	//UE_LOG( LogNetTraffic, Error, TEXT( "GetNetFieldChecksum: %s, %s" ), *NetField->GetName(), *NetField->GetClass()->GetName() );
+
+	UFunction* Function = Cast< UFunction >( NetField );
+	
+	if ( Function != NULL )
+	{
+		// Build checksum based on function parameter list
+		for ( TFieldIterator< UProperty > It( Function ); It && ( It->PropertyFlags & ( CPF_Parm | CPF_ReturnParm ) ) == CPF_Parm; ++It )
+		{
+			LocalChecksum = FCrc::StrCrc32( *It->GetName(), LocalChecksum );
+		}
+
+		return LocalChecksum;
+	}
+
+	UArrayProperty* ArrayProperty = Cast< UArrayProperty >( NetField );
+
+	if ( ArrayProperty != NULL )
+	{
+		return GetNetFieldChecksum( ArrayProperty->Inner, Checksum );
+	}
+
+	UStructProperty* StructProperty = Cast< UStructProperty >( NetField );
+
+	if ( StructProperty != NULL )
+	{
+		LocalChecksum = FCrc::StrCrc32( *StructProperty->Struct->GetName(), LocalChecksum );
+
+		for ( TFieldIterator< UProperty > It( StructProperty->Struct ); It; ++It )
+		{
+			LocalChecksum = GetNetFieldChecksum( *It, LocalChecksum );
+		}
+
+		return LocalChecksum;
+	}
+
+	return FCrc::StrCrc32( *NetField->GetName(), Checksum );
+}
+
+static uint32 GetClassNetworkChecksum( UClass* Class )
+{
+	uint32 Checksum = Class->GetSuperClass() != NULL ? GetClassNetworkChecksum( Class->GetSuperClass() ) : 0;
+
+	Checksum = FCrc::StrCrc32( *Class->GetName(), Checksum );
+
+	for( int32 i = 0; i < Class->NetFields.Num(); i++ )
+	{
+		Checksum = GetNetFieldChecksum( Class->NetFields[i], Checksum );
+	}
+
+	//UE_LOG( LogNetTraffic, Error, TEXT( "GetClassNetworkChecksum: %s, %u" ), *Class->GetName(), Checksum );
+
+	return Checksum;
+}
+
 void FRepLayout::InitFromObjectClass( UClass * InObjectClass )
 {
 	RoleIndex				= -1;
@@ -2309,9 +2370,26 @@ void FRepLayout::InitFromObjectClass( UClass * InObjectClass )
 		}
 	}
 
+	// Initialize NetworkChecksum
+	NetworkChecksum = GetClassNetworkChecksum( InObjectClass );
+
 	Owner = InObjectClass;
 }
 
+void FRepLayout::WriteNetworkChecksum( FOutBunch& Bunch )
+{
+	Bunch << NetworkChecksum;
+}
+
+bool FRepLayout::ReadNetworkChecksum( FInBunch& Bunch )
+{
+	uint32 SerializedChecksum = 0;
+
+	Bunch << SerializedChecksum;
+
+	return SerializedChecksum == NetworkChecksum;
+}
+
 void FRepLayout::InitFromFunction( UFunction * InFunction )
 {
 	int32 RelativeHandle = 0;
diff --git a/Engine/Source/Runtime/Engine/Private/UnrealNetwork.cpp b/Engine/Source/Runtime/Engine/Private/UnrealNetwork.cpp
--- a/Engine/Source/Runtime/Engine/Private/UnrealNetwork.cpp
+++ b/Engine/Source/Runtime/Engine/Private/UnrealNetwork.cpp
@@ -11,10 +11,11 @@ FNetworkVersion::FIsNetworkCompatibleOverride FNetworkVersion::IsNetworkCompatib
 enum ENetworkVersionHistory
 {
 	HISTORY_INITIAL				= 1,
-	HISTORY_INTERNAL_ACK		= 3				// We no longer save packet/channel sequence in stream. We can derive this for 100% reliable connections.
+	HISTORY_INTERNAL_ACK		= 3,				// We no longer save packet/channel sequence in stream. We can derive this for 100% reliable connections.
+	HISTORY_REPLAY_CHECKSUMS	= 4,				// We now save replay checksums into stream for backwards compatibility checks
 };
 
-const uint32 FNetworkVersion::InternalProtocolVersion = HISTORY_INTERNAL_ACK;
+const uint32 FNetworkVersion::InternalProtocolVersion = HISTORY_REPLAY_CHECKSUMS;
 
 uint32 FNetworkVersion::GetLocalNetworkVersion()
 {
diff --git a/Engine/Source/Runtime/Engine/Public/Net/DataReplication.h b/Engine/Source/Runtime/Engine/Public/Net/DataReplication.h
--- a/Engine/Source/Runtime/Engine/Public/Net/DataReplication.h
+++ b/Engine/Source/Runtime/Engine/Public/Net/DataReplication.h
@@ -154,4 +154,7 @@ public:
 	}
 
 	void QueuePropertyRepNotify( UObject* Object, UProperty * Property, const int32 ElementIndex, TArray< uint8 > & MetaData );
+
+	void WriteNetworkChecksum( FOutBunch& Bunch );
+	bool ReadNetworkChecksum( FInBunch& Bunch );
 };
\ No newline at end of file
diff --git a/Engine/Source/Runtime/Engine/Public/Net/RepLayout.h b/Engine/Source/Runtime/Engine/Public/Net/RepLayout.h
--- a/Engine/Source/Runtime/Engine/Public/Net/RepLayout.h
+++ b/Engine/Source/Runtime/Engine/Public/Net/RepLayout.h
@@ -243,7 +243,7 @@ class FRepLayout
 	friend class FRepState;
 
 public:
-	FRepLayout() : FirstNonCustomParent( 0 ), RoleIndex( -1 ), RemoteRoleIndex( -1 ), Owner( NULL ) {}
+	FRepLayout() : FirstNonCustomParent( 0 ), RoleIndex( -1 ), RemoteRoleIndex( -1 ), Owner( NULL ), NetworkChecksum( 0 ) {}
 
 	void OpenAcked( FRepState * RepState ) const;
 
@@ -318,6 +318,9 @@ public:
 	// Serializes all replicated properties of a UObject in or out of an archive (depending on what type of archive it is)
 	ENGINE_API void SerializeObjectReplicatedProperties(UObject* Object, FArchive & Ar) const;
 
+	void WriteNetworkChecksum( FOutBunch& Bunch );
+	bool ReadNetworkChecksum( FInBunch& Bunch );
+
 private:
 	void RebuildConditionalProperties( FRepState * RESTRICT	RepState, const FRepChangedPropertyTracker& ChangedTracker, const FReplicationFlags& RepFlags ) const;
 
@@ -449,4 +452,6 @@ private:
 	int32						RemoteRoleIndex;
 
 	UObject *					Owner;						// Either a UCkass or UFunction
+
+	uint32						NetworkChecksum;
 };
