Better backwards compatibility, we now use the packagemap to handle this. We also now save crc32 of package guid, rather than whole guid.

diff --git a/Engine/Source/Runtime/Engine/Classes/Engine/PackageMapClient.h b/Engine/Source/Runtime/Engine/Classes/Engine/PackageMapClient.h
--- a/Engine/Source/Runtime/Engine/Classes/Engine/PackageMapClient.h
+++ b/Engine/Source/Runtime/Engine/Classes/Engine/PackageMapClient.h
@@ -27,7 +27,7 @@
 class FNetGuidCacheObject
 {
 public:
-	FNetGuidCacheObject() : ReadOnlyTimestamp( 0 ), bNoLoad( 0 ), bIgnoreWhenMissing( 0 ), bIsPending( 0 ), bIsBroken( 0 )
+	FNetGuidCacheObject() : NetworkChecksum( 0 ), PackageChecksum( 0 ), ReadOnlyTimestamp( 0 ), bNoLoad( 0 ), bIgnoreWhenMissing( 0 ), bIsPending( 0 ), bIsBroken( 0 )
 	{
 	}
 
@@ -36,7 +36,8 @@ public:
 	// These fields are set when this guid is static
 	FNetworkGUID				OuterGUID;
 	FName						PathName;
-	FGuid						PackageGuid;				// If this is a package, this is the guid to expect
+	uint32						NetworkChecksum;			// Network checksum saved, used to determine backwards compatible
+	uint32						PackageChecksum;			// If this is a package, this is the guid to expect
 
 	double						ReadOnlyTimestamp;			// Time in second when we should start timing out after going read only
 
@@ -60,7 +61,7 @@ public:
 	void			RegisterNetGUID_Internal( const FNetworkGUID& NetGUID, const FNetGuidCacheObject& CacheObject );
 	void			RegisterNetGUID_Server( const FNetworkGUID& NetGUID, const UObject* Object );
 	void			RegisterNetGUID_Client( const FNetworkGUID& NetGUID, const UObject* Object );
-	void			RegisterNetGUIDFromPath_Client( const FNetworkGUID& NetGUID, const FString& PathName, const FNetworkGUID& OuterGUID, const FGuid& PackageGuid, const bool bNoLoad, const bool bIgnoreWhenMissing );
+	void			RegisterNetGUIDFromPath_Client( const FNetworkGUID& NetGUID, const FString& PathName, const FNetworkGUID& OuterGUID, const uint32 NetworkChecksum, const uint32 PackageChecksum, const bool bNoLoad, const bool bIgnoreWhenMissing );
 	UObject *		GetObjectFromNetGUID( const FNetworkGUID& NetGUID, const bool bIgnoreMustBeMapped );
 	bool			ShouldIgnoreWhenMissing( const FNetworkGUID& NetGUID ) const;
 	bool			IsGUIDRegistered( const FNetworkGUID& NetGUID ) const;
@@ -70,6 +71,10 @@ public:
 	void			GenerateFullNetGUIDPath_r( const FNetworkGUID& NetGUID, FString& FullPath ) const;
 	void			SetIgnorePackageMismatchOverride( const bool bInIgnorePackageMismatchOverride ) { bIgnorePackageMismatchOverride = bInIgnorePackageMismatchOverride; }
 	bool			ShouldIgnorePackageMismatch() const;
+	uint32			GetClassNetworkChecksum( const UClass* Class );
+	uint32			GetNetworkChecksum( const UObject* Obj );
+	void			SetShouldUseNetworkChecksum( const bool bInUseNetworkChecksum );
+	bool			ShouldUseNetworkChecksum() const;
 
 	void			AsyncPackageCallback(const FName& PackageName, UPackage * Package, EAsyncLoadingResult::Type Result);
 	
@@ -77,12 +82,15 @@ public:
 	TMap< TWeakObjectPtr< UObject >, FNetworkGUID >	NetGUIDLookup;
 	int32											UniqueNetIDs[2];
 
+	TMap< TWeakObjectPtr< UClass >, uint32 >		ClassChecksumMap;
+
 	bool											IsExportingNetGUIDBunch;
 
 	UNetDriver *									Driver;
 
 	TMap< FName, FNetworkGUID >						PendingAsyncPackages;
 
+	bool											bUseNetworkChecksum;
 	bool											bIgnorePackageMismatchOverride;
 
 #if !(UE_BUILD_SHIPPING || UE_BUILD_TEST)
diff --git a/Engine/Source/Runtime/Engine/Private/DataChannel.cpp b/Engine/Source/Runtime/Engine/Private/DataChannel.cpp
--- a/Engine/Source/Runtime/Engine/Private/DataChannel.cpp
+++ b/Engine/Source/Runtime/Engine/Private/DataChannel.cpp
@@ -1855,6 +1855,11 @@ void UActorChannel::ReceivedBunch( FInBunch & Bunch )
 
 void UActorChannel::ProcessBunch( FInBunch & Bunch )
 {
+	if ( Broken )
+	{
+		return;
+	}
+
 	const bool bIsServer = Connection->Driver->IsServer();
 
 	FReplicationFlags RepFlags;
@@ -1881,7 +1886,10 @@ void UActorChannel::ProcessBunch( FInBunch & Bunch )
 			check( !bSpawnedNewActor );
 			UE_LOG(LogNet, Warning, TEXT("UActorChannel::ProcessBunch: SerializeNewActor failed to find/spawn actor. Actor: %s, Channel: %i"), NewChannelActor ? *NewChannelActor->GetFullName() : TEXT( "NULL" ), ChIndex);
 			Broken = 1;
-			FNetControlMessage<NMT_ActorChannelFailure>::Send(Connection, ChIndex);
+			if ( !Connection->InternalAck )
+			{
+				FNetControlMessage<NMT_ActorChannelFailure>::Send(Connection, ChIndex);
+			}
 			return;
 		}
 
@@ -1941,22 +1949,6 @@ void UActorChannel::ProcessBunch( FInBunch & Bunch )
 
 		TSharedRef< FObjectReplicator > & Replicator = FindOrCreateReplicator( RepObj );
 
-		if ( Connection->InternalAck && ( Bunch.bOpen || RepObj != Actor  ) )
-		{
-			// Receive network checksum
-			if ( !Replicator->ReadNetworkChecksum( Bunch ) )
-			{
-				UE_LOG( LogNet, Warning, TEXT( "UActorChannel::ProcessBunch: ReadNetworkChecksum failed. RepObj: %s, Channel: %i"), RepObj ? *RepObj->GetFullName() : TEXT( "NULL" ), ChIndex );
-
-				DestroyActorAndComponents();
-				Actor = NULL;
-
-				// Mark as broken so we stop trying to process bunches on this channel
-				Broken = 1;
-				break;
-			}
-		}
-
 		bool bHasUnmapped = false;
 
 		if ( !Replicator->ReceivedBunch( Bunch, RepFlags, bHasUnmapped ) )
@@ -2309,11 +2301,6 @@ void UActorChannel::BeginContentBlock( UObject* Obj, FOutBunch &Bunch )
 
 	if ( IsActor )
 	{
-		if ( Connection->InternalAck && OpenPacketId.First == INDEX_NONE )
-		{
-			// Write network checksum
-			ActorReplicator->WriteNetworkChecksum( Bunch );
-		}
 		NETWORK_PROFILER(GNetworkProfiler.TrackBeginContentBlock(Obj, Bunch.GetNumBits() - NumStartingBits));
 		return;
 	}
@@ -2337,12 +2324,6 @@ void UActorChannel::BeginContentBlock( UObject* Obj, FOutBunch &Bunch )
 		}
 	}
 
-	if ( Connection->InternalAck )
-	{
-		// Write network checksum
-		FindOrCreateReplicator( Obj )->WriteNetworkChecksum( Bunch );
-	}
-
 #if !(UE_BUILD_SHIPPING || UE_BUILD_TEST)
 	if (CVarDoReplicationContextString->GetInt() > 0)
 	{
@@ -2477,8 +2458,13 @@ UObject* UActorChannel::ReadContentBlockHeader(FInBunch & Bunch, bool& bObjectDe
 		// If this is a stably named sub-object, we shouldn't need to create it
 		if ( SubObj == NULL )
 		{
-			UE_LOG( LogNetTraffic, Error, TEXT( "ReadContentBlockHeader: Stably named sub-object not found. Actor: %s" ), *Actor->GetName() );
-			Bunch.SetError();
+			// (ignore though if this is for replays)
+			if ( !Connection->InternalAck )
+			{
+				UE_LOG( LogNetTraffic, Error, TEXT( "ReadContentBlockHeader: Stably named sub-object not found. Actor: %s" ), *Actor->GetName() );
+				Bunch.SetError();
+			}
+
 			return NULL;
 		}
 
@@ -2516,8 +2502,13 @@ UObject* UActorChannel::ReadContentBlockHeader(FInBunch & Bunch, bool& bObjectDe
 		// Valid NetGUID but no class was resolved - this is an error
 		if ( SubObj == NULL )
 		{
-			UE_LOG( LogNetTraffic, Error, TEXT( "UActorChannel::ReadContentBlockHeader: Unable to read sub-object class (SubObj == NULL). Actor: %s" ), *Actor->GetName() );
-			Bunch.SetError();
+			// (unless we're using replays, which could be backwards compatibility kicking in)
+			if ( !Connection->InternalAck )
+			{
+				UE_LOG( LogNetTraffic, Error, TEXT( "UActorChannel::ReadContentBlockHeader: Unable to read sub-object class (SubObj == NULL). Actor: %s" ), *Actor->GetName() );
+				Bunch.SetError();
+			}
+
 			return NULL;
 		}
 	}
diff --git a/Engine/Source/Runtime/Engine/Private/DataReplication.cpp b/Engine/Source/Runtime/Engine/Private/DataReplication.cpp
--- a/Engine/Source/Runtime/Engine/Private/DataReplication.cpp
+++ b/Engine/Source/Runtime/Engine/Private/DataReplication.cpp
@@ -390,16 +390,6 @@ void FObjectReplicator::ReceivedNak( int32 NakPacketId )
 	}
 }
 
-void FObjectReplicator::WriteNetworkChecksum( FOutBunch& Bunch )
-{
-	RepLayout->WriteNetworkChecksum( Bunch );
-}
-
-bool FObjectReplicator::ReadNetworkChecksum( FInBunch& Bunch )
-{
-	return RepLayout->ReadNetworkChecksum( Bunch );
-}
-
 bool FObjectReplicator::ReceivedBunch( FInBunch &Bunch, const FReplicationFlags& RepFlags, bool & bOutHasUnmapped )
 {
 	UObject* Object = GetObject();
diff --git a/Engine/Source/Runtime/Engine/Private/DemoNetDriver.cpp b/Engine/Source/Runtime/Engine/Private/DemoNetDriver.cpp
--- a/Engine/Source/Runtime/Engine/Private/DemoNetDriver.cpp
+++ b/Engine/Source/Runtime/Engine/Private/DemoNetDriver.cpp
@@ -366,6 +366,7 @@ bool UDemoNetDriver::InitListen( FNetworkNotify* InNotify, FURL& ListenURL, bool
 	}
 
 	GuidCache->SetIgnorePackageMismatchOverride( true );
+	GuidCache->SetShouldUseNetworkChecksum( true );
 
 	check( World != NULL );
 
@@ -741,7 +742,8 @@ static void SerializeGuidCache( TSharedPtr< class FNetGUIDCache > GuidCache, FAr
 		*CheckpointArchive << It.Key();
 		*CheckpointArchive << It.Value().OuterGUID;
 		*CheckpointArchive << PathName;
-		*CheckpointArchive << It.Value().PackageGuid;
+		*CheckpointArchive << It.Value().NetworkChecksum;
+		*CheckpointArchive << It.Value().PackageChecksum;
 
 		uint8 Flags = 0;
 		
@@ -1575,7 +1577,8 @@ void UDemoNetDriver::LoadCheckpoint()
 
 		*GotoCheckpointArchive << CacheObject.OuterGUID;
 		*GotoCheckpointArchive << PathName;
-		*GotoCheckpointArchive << CacheObject.PackageGuid;
+		*GotoCheckpointArchive << CacheObject.NetworkChecksum;
+		*GotoCheckpointArchive << CacheObject.PackageChecksum;
 
 		CacheObject.PathName = FName( *PathName );
 
diff --git a/Engine/Source/Runtime/Engine/Private/PackageMapClient.cpp b/Engine/Source/Runtime/Engine/Private/PackageMapClient.cpp
--- a/Engine/Source/Runtime/Engine/Private/PackageMapClient.cpp
+++ b/Engine/Source/Runtime/Engine/Private/PackageMapClient.cpp
@@ -34,6 +34,13 @@ UPackageMapClient::UPackageMapClient(const FObjectInitializer& ObjectInitializer
 {
 }
 
+static uint32 GetPackageChecksum( const UPackage* Package )
+{
+	const FGuid Guid = Package->GetGuid();
+
+	return FCrc::MemCrc32( &Guid, sizeof( FGuid ) );
+}
+
 /**
  *	This is the meat of the PackageMap class which serializes a reference to Object.
  */
@@ -312,8 +319,9 @@ struct FExportFlags
 	{
 		struct
 		{
-			uint8 bHasPath		: 1;
-			uint8 bNoLoad		: 1;
+			uint8 bHasPath				: 1;
+			uint8 bNoLoad				: 1;
+			uint8 bHasNetworkChecksum	: 1;
 		};
 
 		uint8	Value;
@@ -371,6 +379,8 @@ void UPackageMapClient::InternalWriteObject( FArchive & Ar, FNetworkGUID NetGUID
 	//   note: Default NetGUID is implied to always send path
 	FExportFlags ExportFlags;
 
+	ExportFlags.bHasNetworkChecksum = GuidCache->ShouldUseNetworkChecksum() ? 1 : 0;
+
 	if ( NetGUID.IsDefault() )
 	{
 		// Only the client sends default guids
@@ -427,27 +437,32 @@ void UPackageMapClient::InternalWriteObject( FArchive & Ar, FNetworkGUID NetGUID
 		// Serialize Name of object
 		Ar << ObjectPathName;
 
-		FNetGuidCacheObject* CacheObject = GuidCache->ObjectLookup.Find( NetGUID );
+		uint32 NetworkChecksum = 0;
+		uint32 PackageChecksum = 0;
 
-		if ( bIsPackage )
+		if ( ExportFlags.bHasNetworkChecksum )
 		{
-			FGuid PackageGuid = CastChecked< const UPackage >( Object )->GetGuid();
-			Ar << PackageGuid;
-
-			UE_LOG( LogNetPackageMap, VeryVerbose, TEXT( "InternalWriteObject: Package: %s, GUID: %s" ), *ObjectPathName, *PackageGuid.ToString() );
+			NetworkChecksum = GuidCache->GetNetworkChecksum( Object );
+			Ar << NetworkChecksum;
+		}
 
-			if ( CacheObject != NULL )
-			{
-				CacheObject->PackageGuid = PackageGuid;
-			}
+		if ( bIsPackage )
+		{
+			PackageChecksum = GetPackageChecksum( CastChecked< const UPackage >( Object ) );
+			Ar << PackageChecksum;
+			UE_LOG( LogNetPackageMap, VeryVerbose, TEXT( "InternalWriteObject: Package: %s, GUID: %u" ), *ObjectPathName, PackageChecksum );
 		}
 
+		FNetGuidCacheObject* CacheObject = GuidCache->ObjectLookup.Find( NetGUID );
+
 		if ( CacheObject != NULL )
 		{
 			CacheObject->PathName			= FName( *ObjectPathName );
 			CacheObject->OuterGUID			= OuterNetGUID;
 			CacheObject->bNoLoad			= ExportFlags.bNoLoad;
 			CacheObject->bIgnoreWhenMissing = ExportFlags.bNoLoad;
+			CacheObject->NetworkChecksum	= NetworkChecksum;
+			CacheObject->PackageChecksum	= PackageChecksum;
 		}
 
 		if ( GuidCache->IsExportingNetGUIDBunch )
@@ -471,7 +486,7 @@ static void SanityCheckExport(
 	const UObject *			Object, 
 	const FNetworkGUID &	NetGUID, 
 	const FString &			ExpectedPathName, 
-	const FGuid &			ExpectedPackageGuid, 
+	const uint32			ExpectedPackageChecksum,
 	const UObject *			ExpectedOuter, 
 	const FNetworkGUID &	ExpectedOuterGUID,
 	const FExportFlags &	ExportFlags )
@@ -489,9 +504,9 @@ static void SanityCheckExport(
 			UE_LOG( LogNetPackageMap, Warning, TEXT( "SanityCheckExport: CacheObject->OuterGUID != ExpectedOuterGUID. NetGUID: %s, Object: %s, Expected: %s" ), *NetGUID.ToString(), *Object->GetPathName(), *ExpectedPathName );
 		}
 
-		if ( CacheObject->PackageGuid != ExpectedPackageGuid )
+		if ( CacheObject->PackageChecksum != ExpectedPackageChecksum )
 		{
-			UE_LOG( LogNetPackageMap, Warning, TEXT( "SanityCheckExport: CacheObject->PackageGuid != ExpectedPackageGuid. NetGUID: %s, Object: %s, Expected: %s" ), *NetGUID.ToString(), *Object->GetPathName(), *ExpectedPathName );
+			UE_LOG( LogNetPackageMap, Warning, TEXT( "SanityCheckExport: CacheObject->PackageChecksum != ExpectedPackageChecksum. NetGUID: %s, Object: %s, Expected: %s" ), *NetGUID.ToString(), *Object->GetPathName(), *ExpectedPathName );
 		}
 	}
 	else
@@ -587,16 +602,22 @@ FNetworkGUID UPackageMapClient::InternalLoadObject( FArchive & Ar, UObject *& Ob
 		FNetworkGUID OuterGUID = InternalLoadObject( Ar, ObjOuter, InternalLoadObjectRecursionCount + 1 );
 
 		FString PathName;
-		FGuid	PackageGuid;
+		uint32	NetworkChecksum = 0;
+		uint32	PackageChecksum = 0;
 
 		Ar << PathName;
 
+		if ( ExportFlags.bHasNetworkChecksum )
+		{
+			Ar << NetworkChecksum;
+		}
+
 		const bool bIsPackage = NetGUID.IsStatic() && !OuterGUID.IsValid();
 
 		if ( bIsPackage )
 		{
-			Ar << PackageGuid;
-			UE_LOG( LogNetPackageMap, VeryVerbose, TEXT( "InternalLoadObject: Package: %s, GUID: %s" ), *PathName, *PackageGuid.ToString() );
+			Ar << PackageChecksum;
+			UE_LOG( LogNetPackageMap, VeryVerbose, TEXT( "InternalLoadObject: Package: %s, GUID: %u" ), *PathName, PackageChecksum );
 		}
 
 		if ( Ar.IsError() )
@@ -612,7 +633,7 @@ FNetworkGUID UPackageMapClient::InternalLoadObject( FArchive & Ar, UObject *& Ob
 		if ( Object != NULL )
 		{
 			// If we already have the object, just do some sanity checking and return
-			SanityCheckExport( GuidCache.Get(), Object, NetGUID, PathName, PackageGuid, ObjOuter, OuterGUID, ExportFlags );
+			SanityCheckExport( GuidCache.Get(), Object, NetGUID, PathName, PackageChecksum, ObjOuter, OuterGUID, ExportFlags );
 			return NetGUID;
 		}
 
@@ -649,9 +670,9 @@ FNetworkGUID UPackageMapClient::InternalLoadObject( FArchive & Ar, UObject *& Ob
 					return NetGUID;
 				}
 
-				if ( Package->GetGuid() != PackageGuid && !GuidCache->ShouldIgnorePackageMismatch() )
+				if ( !GuidCache->ShouldIgnorePackageMismatch() && GetPackageChecksum( Package ) != PackageChecksum )
 				{
-					UE_LOG( LogNetPackageMap, Error, TEXT( "UPackageMapClient::InternalLoadObject: Default object package guid mismatch! PathName: %s, ObjOuter: %s, GUID1: %s, GUID2: %s " ), *PathName, ObjOuter != NULL ? *ObjOuter->GetPathName() : TEXT( "NULL" ), *Package->GetGuid().ToString(), *PackageGuid.ToString() );
+					UE_LOG( LogNetPackageMap, Error, TEXT( "UPackageMapClient::InternalLoadObject: Default object package guid mismatch! PathName: %s, ObjOuter: %s, GUID1: %u, GUID2: %u " ), *PathName, ObjOuter != NULL ? *ObjOuter->GetPathName() : TEXT( "NULL" ), GetPackageChecksum( Package ), PackageChecksum );
 					Object = NULL;
 					return NetGUID;
 				}
@@ -680,7 +701,7 @@ FNetworkGUID UPackageMapClient::InternalLoadObject( FArchive & Ar, UObject *& Ob
 		const bool bIgnoreWhenMissing = ExportFlags.bNoLoad;
 
 		// Register this path and outer guid combo with the net guid
-		GuidCache->RegisterNetGUIDFromPath_Client( NetGUID, PathName, OuterGUID, PackageGuid, ExportFlags.bNoLoad, bIgnoreWhenMissing );
+		GuidCache->RegisterNetGUIDFromPath_Client( NetGUID, PathName, OuterGUID, NetworkChecksum, PackageChecksum, ExportFlags.bNoLoad, bIgnoreWhenMissing );
 
 		// Try again now that we've registered the path
 		Object = GuidCache->GetObjectFromNetGUID( NetGUID, GuidCache->IsExportingNetGUIDBunch );
@@ -1263,7 +1284,7 @@ UObject* UPackageMapClient::GetObjectFromNetGUID( const FNetworkGUID& NetGUID, c
 //	FNetGUIDCache
 //----------------------------------------------------------------------------------------
 
-FNetGUIDCache::FNetGUIDCache( UNetDriver * InDriver ) : IsExportingNetGUIDBunch( false ), Driver( InDriver ), bIgnorePackageMismatchOverride( false )
+FNetGUIDCache::FNetGUIDCache( UNetDriver * InDriver ) : IsExportingNetGUIDBunch( false ), Driver( InDriver ), bUseNetworkChecksum( false ), bIgnorePackageMismatchOverride( false )
 {
 	UniqueNetIDs[0] = UniqueNetIDs[1] = 0;
 }
@@ -1326,6 +1347,8 @@ void FNetGUIDCache::CleanReferences()
 		checkf( ObjectLookup.FindRef( It.Value() ).Object == It.Key(), TEXT("Failed to validate NetGUIDLookup map in UPackageMap. GUID '%s' was not in the ObjectLookup map with with object '%s'."), *It.Value().ToString(), *It.Key().Get()->GetPathName());
 	}
 
+	ClassChecksumMap.Empty();
+
 	FArchiveCountMemGUID CountBytesAr;
 
 	ObjectLookup.CountBytes( CountBytesAr );
@@ -1571,7 +1594,7 @@ void FNetGUIDCache::RegisterNetGUID_Client( const FNetworkGUID& NetGUID, const U
  *	Associates a net guid with a path, that can be loaded or found later
  *  This function is only called on the client
  */
-void FNetGUIDCache::RegisterNetGUIDFromPath_Client( const FNetworkGUID& NetGUID, const FString& PathName, const FNetworkGUID& OuterGUID, const FGuid& PackageGuid, const bool bNoLoad, const bool bIgnoreWhenMissing )
+void FNetGUIDCache::RegisterNetGUIDFromPath_Client( const FNetworkGUID& NetGUID, const FString& PathName, const FNetworkGUID& OuterGUID, const uint32 NetworkChecksum, const uint32 PackageChecksum, const bool bNoLoad, const bool bIgnoreWhenMissing )
 {
 	check( !IsNetGUIDAuthority() );		// Server never calls this locally
 	check( !NetGUID.IsDefault() );
@@ -1611,7 +1634,8 @@ void FNetGUIDCache::RegisterNetGUIDFromPath_Client( const FNetworkGUID& NetGUID,
 
 	CacheObject.PathName			= FName( *PathName );
 	CacheObject.OuterGUID			= OuterGUID;
-	CacheObject.PackageGuid			= PackageGuid;
+	CacheObject.NetworkChecksum		= NetworkChecksum;
+	CacheObject.PackageChecksum		= PackageChecksum;
 	CacheObject.bNoLoad				= bNoLoad;
 	CacheObject.bIgnoreWhenMissing	= bIgnoreWhenMissing;
 
@@ -1652,10 +1676,10 @@ void FNetGUIDCache::AsyncPackageCallback(const FName& PackageName, UPackage * Pa
 		CacheObject->bIsBroken = true;
 		UE_LOG( LogNetPackageMap, Error, TEXT( "AsyncPackageCallback: Package FAILED to load. Path: %s, NetGUID: %s" ), *PackageName.ToString(), *NetGUID.ToString() );
 	}
-	else if ( Package->GetGuid() != CacheObject->PackageGuid && !ShouldIgnorePackageMismatch() )
+	else if ( !ShouldIgnorePackageMismatch() && GetPackageChecksum( Package ) != CacheObject->PackageChecksum )
 	{
 		CacheObject->bIsBroken = true;
-		UE_LOG( LogNetPackageMap, Error, TEXT( "AsyncPackageCallback: Package GUID mismatch! Path: %s, NetGUID: %s, GUID1: %s, GUID2: %s" ), *PackageName.ToString(), *NetGUID.ToString(), *Package->GetGuid().ToString(), *CacheObject->PackageGuid.ToString() );
+		UE_LOG( LogNetPackageMap, Error, TEXT( "AsyncPackageCallback: Package GUID mismatch! Path: %s, NetGUID: %s, GUID1: %u, GUID2: %u" ), *PackageName.ToString(), *NetGUID.ToString(), GetPackageChecksum( Package ), CacheObject->PackageChecksum );
 	}
 }
 
@@ -1851,11 +1875,24 @@ UObject* FNetGUIDCache::GetObjectFromNetGUID( const FNetworkGUID& NetGUID, const
 			}
 		}
 
-		if ( Package->GetGuid() != CacheObjectPtr->PackageGuid && !ShouldIgnorePackageMismatch() )
+		if ( !ShouldIgnorePackageMismatch() && GetPackageChecksum( Package ) != CacheObjectPtr->PackageChecksum )
 		{
 			// If the package guid doesn't match, don't allow it to load
 			CacheObjectPtr->bIsBroken = true;
-			UE_LOG( LogNetPackageMap, Error, TEXT( "GetObjectFromNetGUID: Package GUID mismatch! Path: %s, NetGUID: %s, GUID1: %s, GUID2: %s" ), *CacheObjectPtr->PathName.ToString(), *NetGUID.ToString(), *Package->GetGuid().ToString(), *CacheObjectPtr->PackageGuid.ToString() );
+			UE_LOG( LogNetPackageMap, Error, TEXT( "GetObjectFromNetGUID: Package GUID mismatch! Path: %s, NetGUID: %s, GUID1: %u, GUID2: %u" ), *CacheObjectPtr->PathName.ToString(), *NetGUID.ToString(), GetPackageChecksum( Package ), CacheObjectPtr->PackageChecksum );
+			return NULL;
+		}
+	}
+
+	if ( CacheObjectPtr->NetworkChecksum != 0 )
+	{
+		const uint32 NetworkChecksum = GetNetworkChecksum( Object );
+
+		if ( CacheObjectPtr->NetworkChecksum != NetworkChecksum )
+		{
+			UE_LOG( LogNetPackageMap, Warning, TEXT( "GetObjectFromNetGUID: Network checksum mismatch. FullNetGUIDPath: %s, %u, %u" ), *FullNetGUIDPath( NetGUID ), CacheObjectPtr->NetworkChecksum, NetworkChecksum );
+
+			CacheObjectPtr->bIsBroken = true;
 			return NULL;
 		}
 	}
@@ -2042,11 +2079,114 @@ void FNetGUIDCache::GenerateFullNetGUIDPath_r( const FNetworkGUID& NetGUID, FStr
 	}
 }
 
+static uint32 GetNetFieldChecksum( const UField* NetField, const uint32 Checksum )
+{
+	uint32 LocalChecksum = FCrc::StrCrc32( *NetField->GetName(), Checksum );
+
+	LocalChecksum = FCrc::StrCrc32( *NetField->GetClass()->GetName(), LocalChecksum );
+
+	const UFunction* Function = Cast< UFunction >( NetField );
+
+	if ( Function != NULL )
+	{
+		// Build checksum based on function parameter list
+		for ( TFieldIterator< UProperty > It( Function ); It && ( It->PropertyFlags & ( CPF_Parm | CPF_ReturnParm ) ) == CPF_Parm; ++It )
+		{
+			LocalChecksum = FCrc::StrCrc32( *It->GetName(), LocalChecksum );
+		}
+
+		return LocalChecksum;
+	}
+
+	const UArrayProperty* ArrayProperty = Cast< UArrayProperty >( NetField );
+
+	if ( ArrayProperty != NULL )
+	{
+		return GetNetFieldChecksum( ArrayProperty->Inner, Checksum );
+	}
+
+	const UStructProperty* StructProperty = Cast< UStructProperty >( NetField );
+
+	if ( StructProperty != NULL )
+	{
+		LocalChecksum = FCrc::StrCrc32( *StructProperty->Struct->GetName(), LocalChecksum );
+
+		for ( TFieldIterator< UProperty > It( StructProperty->Struct ); It; ++It )
+		{
+			LocalChecksum = GetNetFieldChecksum( *It, LocalChecksum );
+		}
+
+		return LocalChecksum;
+	}
+
+	return FCrc::StrCrc32( *NetField->GetName(), Checksum );
+}
+
+uint32 FNetGUIDCache::GetClassNetworkChecksum( const UClass* Class )
+{
+	uint32 Checksum = ClassChecksumMap.FindRef( Class );
+
+	if ( Checksum != 0 )
+	{
+		return Checksum;
+	}
+
+	if ( Class->GetSuperClass() != NULL )
+	{
+		Checksum = GetClassNetworkChecksum( Class->GetSuperClass() );
+	}
+
+	Checksum = FCrc::StrCrc32( *Class->GetName(), Checksum );
+
+	for( int32 i = 0; i < Class->NetFields.Num(); i++ )
+	{
+		Checksum = GetNetFieldChecksum( Class->NetFields[i], Checksum );
+	}
+
+	// Can't be 0
+	if ( Checksum == 0 )
+	{
+		Checksum = 1;
+	}
+
+	ClassChecksumMap.Add( Class, Checksum );
+
+	return Checksum;
+}
+
+uint32 FNetGUIDCache::GetNetworkChecksum( const UObject* Obj )
+{
+	if ( Obj == NULL )
+	{
+		return 0;
+	}
+
+	// If Obj is already a class, we can use that directly
+	const UClass* Class = Cast< UClass >( Obj );
+
+	if ( Class != NULL )
+	{
+		return GetClassNetworkChecksum( Class );
+	}
+
+	return GetClassNetworkChecksum( Obj->GetClass() );
+}
+
 bool FNetGUIDCache::ShouldIgnorePackageMismatch() const
 {
 	return bIgnorePackageMismatchOverride || CVarIgnorePackageMismatch.GetValueOnGameThread();
 };
 
+void FNetGUIDCache::SetShouldUseNetworkChecksum( const bool bInUseNetworkChecksum )
+{
+	bUseNetworkChecksum = bInUseNetworkChecksum;
+}
+
+bool FNetGUIDCache::ShouldUseNetworkChecksum() const
+{
+	return bUseNetworkChecksum;
+}
+
 //------------------------------------------------------
 // Debug command to see how many times we've exported each NetGUID
 // Used for measuring inefficiencies. Some duplication is unavoidable since we cannot garuntee atomicicity across multiple channels.
diff --git a/Engine/Source/Runtime/Engine/Private/RepLayout.cpp b/Engine/Source/Runtime/Engine/Private/RepLayout.cpp
--- a/Engine/Source/Runtime/Engine/Private/RepLayout.cpp
+++ b/Engine/Source/Runtime/Engine/Private/RepLayout.cpp
@@ -2180,67 +2180,6 @@ uint16 FRepLayout::AddParentProperty( UProperty * Property, int32 ArrayIndex )
 
 extern bool IsCustomDeltaProperty( UProperty * Property );
 
-static uint32 GetNetFieldChecksum( UField* NetField, const uint32 Checksum )
-{
-	uint32 LocalChecksum = FCrc::StrCrc32( *NetField->GetName(), Checksum );
-
-	LocalChecksum = FCrc::StrCrc32( *NetField->GetClass()->GetName(), LocalChecksum );
-
-	//UE_LOG( LogNetTraffic, Error, TEXT( "GetNetFieldChecksum: %s, %s" ), *NetField->GetName(), *NetField->GetClass()->GetName() );
-
-	UFunction* Function = Cast< UFunction >( NetField );
-	
-	if ( Function != NULL )
-	{
-		// Build checksum based on function parameter list
-		for ( TFieldIterator< UProperty > It( Function ); It && ( It->PropertyFlags & ( CPF_Parm | CPF_ReturnParm ) ) == CPF_Parm; ++It )
-		{
-			LocalChecksum = FCrc::StrCrc32( *It->GetName(), LocalChecksum );
-		}
-
-		return LocalChecksum;
-	}
-
-	UArrayProperty* ArrayProperty = Cast< UArrayProperty >( NetField );
-
-	if ( ArrayProperty != NULL )
-	{
-		return GetNetFieldChecksum( ArrayProperty->Inner, Checksum );
-	}
-
-	UStructProperty* StructProperty = Cast< UStructProperty >( NetField );
-
-	if ( StructProperty != NULL )
-	{
-		LocalChecksum = FCrc::StrCrc32( *StructProperty->Struct->GetName(), LocalChecksum );
-
-		for ( TFieldIterator< UProperty > It( StructProperty->Struct ); It; ++It )
-		{
-			LocalChecksum = GetNetFieldChecksum( *It, LocalChecksum );
-		}
-
-		return LocalChecksum;
-	}
-
-	return FCrc::StrCrc32( *NetField->GetName(), Checksum );
-}
-
-static uint32 GetClassNetworkChecksum( UClass* Class )
-{
-	uint32 Checksum = Class->GetSuperClass() != NULL ? GetClassNetworkChecksum( Class->GetSuperClass() ) : 0;
-
-	Checksum = FCrc::StrCrc32( *Class->GetName(), Checksum );
-
-	for( int32 i = 0; i < Class->NetFields.Num(); i++ )
-	{
-		Checksum = GetNetFieldChecksum( Class->NetFields[i], Checksum );
-	}
-
-	//UE_LOG( LogNetTraffic, Error, TEXT( "GetClassNetworkChecksum: %s, %u" ), *Class->GetName(), Checksum );
-
-	return Checksum;
-}
-
 void FRepLayout::InitFromObjectClass( UClass * InObjectClass )
 {
 	RoleIndex				= -1;
@@ -2370,26 +2309,9 @@ void FRepLayout::InitFromObjectClass( UClass * InObjectClass )
 		}
 	}
 
-	// Initialize NetworkChecksum
-	NetworkChecksum = GetClassNetworkChecksum( InObjectClass );
-
 	Owner = InObjectClass;
 }
 
-void FRepLayout::WriteNetworkChecksum( FOutBunch& Bunch )
-{
-	Bunch << NetworkChecksum;
-}
-
-bool FRepLayout::ReadNetworkChecksum( FInBunch& Bunch )
-{
-	uint32 SerializedChecksum = 0;
-
-	Bunch << SerializedChecksum;
-
-	return SerializedChecksum == NetworkChecksum;
-}
-
 void FRepLayout::InitFromFunction( UFunction * InFunction )
 {
 	int32 RelativeHandle = 0;
diff --git a/Engine/Source/Runtime/Engine/Private/UnrealNetwork.cpp b/Engine/Source/Runtime/Engine/Private/UnrealNetwork.cpp
--- a/Engine/Source/Runtime/Engine/Private/UnrealNetwork.cpp
+++ b/Engine/Source/Runtime/Engine/Private/UnrealNetwork.cpp
@@ -13,9 +13,10 @@ enum ENetworkVersionHistory
 	HISTORY_INITIAL				= 1,
 	HISTORY_INTERNAL_ACK		= 3,				// We no longer save packet/channel sequence in stream. We can derive this for 100% reliable connections.
 	HISTORY_REPLAY_CHECKSUMS	= 4,				// We now save replay checksums into stream for backwards compatibility checks
+	HISTORY_REPLAY_CHECKSUMS2	= 5,				// Saving package and network checksum to packagemap, no longer save checksum in BeginContentBlockHeader
 };
 
-const uint32 FNetworkVersion::InternalProtocolVersion = HISTORY_REPLAY_CHECKSUMS;
+const uint32 FNetworkVersion::InternalProtocolVersion = HISTORY_REPLAY_CHECKSUMS2;
 
 uint32 FNetworkVersion::GetLocalNetworkVersion()
 {
diff --git a/Engine/Source/Runtime/Engine/Public/Net/DataReplication.h b/Engine/Source/Runtime/Engine/Public/Net/DataReplication.h
--- a/Engine/Source/Runtime/Engine/Public/Net/DataReplication.h
+++ b/Engine/Source/Runtime/Engine/Public/Net/DataReplication.h
@@ -154,7 +154,4 @@ public:
 	}
 
 	void QueuePropertyRepNotify( UObject* Object, UProperty * Property, const int32 ElementIndex, TArray< uint8 > & MetaData );
-
-	void WriteNetworkChecksum( FOutBunch& Bunch );
-	bool ReadNetworkChecksum( FInBunch& Bunch );
 };
\ No newline at end of file
diff --git a/Engine/Source/Runtime/Engine/Public/Net/RepLayout.h b/Engine/Source/Runtime/Engine/Public/Net/RepLayout.h
--- a/Engine/Source/Runtime/Engine/Public/Net/RepLayout.h
+++ b/Engine/Source/Runtime/Engine/Public/Net/RepLayout.h
@@ -243,7 +243,7 @@ class FRepLayout
 	friend class FRepState;
 
 public:
-	FRepLayout() : FirstNonCustomParent( 0 ), RoleIndex( -1 ), RemoteRoleIndex( -1 ), Owner( NULL ), NetworkChecksum( 0 ) {}
+	FRepLayout() : FirstNonCustomParent( 0 ), RoleIndex( -1 ), RemoteRoleIndex( -1 ), Owner( NULL ) {}
 
 	void OpenAcked( FRepState * RepState ) const;
 
@@ -452,6 +452,4 @@ private:
 	int32						RemoteRoleIndex;
 
 	UObject *					Owner;						// Either a UCkass or UFunction
-
-	uint32						NetworkChecksum;
 };
