Fixed warning output in Montage code.

diff --git a/Engine/Source/Runtime/Engine/Private/Animation/AnimInstance.cpp b/Engine/Source/Runtime/Engine/Private/Animation/AnimInstance.cpp
--- a/Engine/Source/Runtime/Engine/Private/Animation/AnimInstance.cpp
+++ b/Engine/Source/Runtime/Engine/Private/Animation/AnimInstance.cpp
@@ -2121,7 +2121,7 @@ float UAnimInstance::Montage_Play(UAnimMontage * MontageToPlay, float InPlayRate
 		else
 		{
 			UE_LOG(LogAnimation, Warning, TEXT("Playing a Montage (%s) for the wrong Skeleton (%s) instead of (%s)."),
-				*GetNameSafe(MontageToPlay), *GetNameSafe(MontageToPlay->GetSkeleton()), *GetNameSafe(CurrentSkeleton));
+				*GetNameSafe(MontageToPlay), *GetNameSafe(CurrentSkeleton), *GetNameSafe(MontageToPlay->GetSkeleton()));
 		}
 	}
 
