#ue4 - Various Actor and SceneComponent methods that use FRotator now have overloads taking rotation as an FQuat.

- Calling code should still function normally (old APIs are maintained).
- FQuat versions are faster if you start with a Quat, since FRotator versions typically will convert to Quat first. For example use GetActorQuat() rather than GetActorRotation() (which itself does a Quat->Rotator conversion).
- In cases where virtual methods using FRotator could be overridden, there are now non-virtual overloads that forward to a single virtual implementation.
   - This ensures derived versions override the version that will be used by engine code.
   - Make sure to call the correct Super function in overrides if currently using it.

diff --git a/Engine/Source/Runtime/Engine/Classes/Components/PrimitiveComponent.h b/Engine/Source/Runtime/Engine/Classes/Components/PrimitiveComponent.h
--- a/Engine/Source/Runtime/Engine/Classes/Components/PrimitiveComponent.h
+++ b/Engine/Source/Runtime/Engine/Classes/Components/PrimitiveComponent.h
@@ -580,15 +580,25 @@ public:
 
 	/**
 	 *  Test the collision of the supplied component at the supplied location/rotation, and determine the set of components that it overlaps.
+	 *  @note This overload taking rotation as a FQuat is slightly faster than the version using FRotator.
+	 *  @note This simply calls the virtual ComponentOverlapMultiImpl() which can be overridden to implement custom behavior.
 	 *  @param  OutOverlaps     Array of overlaps found between this component in specified pose and the world
 	 *  @param  World			World to use for overlap test
-	 *  @param  Pos             Location to place the component's geometry at to test against the world
-	 *  @param  Rot             Rotation to place components' geometry at to test against the world
+	 *  @param  Pos             Location of component's geometry for the test against the world
+	 *  @param  Rot             Rotation of component's geometry for the test against the world
 	 *  @param  TestChannel		The 'channel' that this ray is in, used to determine which components to hit
-	 *	@param	ObjectQueryParams	List of object types it's looking for. When this enters, we do object query with component shape
+	 *  @param	ObjectQueryParams	List of object types it's looking for. When this enters, we do object query with component shape
 	 *  @return true if OutOverlaps contains any blocking results
 	 */
-	virtual bool ComponentOverlapMulti(TArray<struct FOverlapResult>& OutOverlaps, const class UWorld* World, const FVector& Pos, const FRotator& Rot, ECollisionChannel TestChannel, const struct FComponentQueryParams& Params, const struct FCollisionObjectQueryParams& ObjectQueryParams = FCollisionObjectQueryParams::DefaultObjectQueryParam) const;
+	bool ComponentOverlapMulti(TArray<struct FOverlapResult>& OutOverlaps, const class UWorld* World, const FVector& Pos, const FQuat& Rot, ECollisionChannel TestChannel, const struct FComponentQueryParams& Params, const struct FCollisionObjectQueryParams& ObjectQueryParams = FCollisionObjectQueryParams::DefaultObjectQueryParam) const;
+	bool ComponentOverlapMulti(TArray<struct FOverlapResult>& OutOverlaps, const class UWorld* World, const FVector& Pos, const FRotator& Rot, ECollisionChannel TestChannel, const struct FComponentQueryParams& Params, const struct FCollisionObjectQueryParams& ObjectQueryParams = FCollisionObjectQueryParams::DefaultObjectQueryParam) const;
+
+protected:
+
+	// Override this method for custom behavior.
+	virtual bool ComponentOverlapMultiImpl(TArray<struct FOverlapResult>& OutOverlaps, const class UWorld* World, const FVector& Pos, const FQuat& Rot, ECollisionChannel TestChannel, const struct FComponentQueryParams& Params, const struct FCollisionObjectQueryParams& ObjectQueryParams = FCollisionObjectQueryParams::DefaultObjectQueryParam) const;
+
+public:
 
 	/** 
 	 *	Event called when a component hits (or is hit by) something solid. This could happen due to things like Character movement, using Set Location with 'sweep' enabled, or physics simulation.
@@ -1346,19 +1356,11 @@ public:
 	// End UObject interface.
 
 	//Begin USceneComponent Interface
-	/**
-	 * Tries to move the component by a movement vector (Delta) and sets rotation to NewRotation.
-	 * Assumes that the component's current location is valid and that the component does fit in its current Location.
-	 * Dispatches blocking hit notifications (if bSweep is true), and calls UpdateOverlaps() after movement to update overlap state.
-	 * 
-	 * @param Delta			The desired location change in world space.
-	 * @param NewRotation	The new desired rotation in world space.
-	 * @param bSweep		Should we sweep to the destination location, stopping short of the target if blocked by something. Note:If the component has no collision this will have no effect.
-	 * @param Hit			Optional output describing the blocking hit that stopped the move, if any.
-	 * @param MoveFlags		Flags controlling behavior of the move. @see EMoveComponentFlags
-	 * @return				True if some movement occurred, false if no movement occurred.
-	 */
-	virtual bool MoveComponent(const FVector& Delta, const FRotator& NewRotation, bool bSweep, FHitResult* OutHit = NULL, EMoveComponentFlags MoveFlags = MOVECOMP_NoFlags) override;
+
+protected:
+	virtual bool MoveComponentImpl(const FVector& Delta, const FQuat& NewRotation, bool bSweep, FHitResult* OutHit = NULL, EMoveComponentFlags MoveFlags = MOVECOMP_NoFlags) override;
+	
+public:
 	virtual bool IsWorldGeometry() const override;
 	virtual ECollisionEnabled::Type GetCollisionEnabled() const override;
 	virtual ECollisionResponse GetCollisionResponseToChannel(ECollisionChannel Channel) const override;
@@ -1630,13 +1632,23 @@ public:
 	
 	/** 
 	 *  Test the collision of the supplied component at the supplied location/rotation, and determine if it overlaps this component
+	 *  @note This overload taking rotation as a FQuat is slightly faster than the version using FRotator.
+	 *  @note This simply calls the virtual ComponentOverlapComponentImpl() which can be overridden to implement custom behavior.
 	 *  @param  PrimComp        Component to use geometry from to test against this component. Transform of this component is ignored.
 	 *  @param  Pos             Location to place PrimComp geometry at 
 	 *  @param  Rot             Rotation to place PrimComp geometry at 
 	 *  @param	Params			Parameter for trace. TraceTag is only used.
 	 *  @return true if PrimComp overlaps this component at the specified location/rotation
 	 */
-	virtual bool ComponentOverlapComponent(class UPrimitiveComponent* PrimComp, const FVector Pos, const FRotator Rot, const FCollisionQueryParams& Params);
+	bool ComponentOverlapComponent(class UPrimitiveComponent* PrimComp, const FVector Pos, const FQuat& Rot, const FCollisionQueryParams& Params);
+	bool ComponentOverlapComponent(class UPrimitiveComponent* PrimComp, const FVector Pos, const FRotator Rot, const FCollisionQueryParams& Params);
+
+protected:
+
+	// Override this method for custom behavior.
+	virtual bool ComponentOverlapComponentImpl(class UPrimitiveComponent* PrimComp, const FVector Pos, const FQuat& Rot, const FCollisionQueryParams& Params);
+
+public:
 	
 	/** 
 	 *  Test the collision of the supplied Sphere at the supplied location, and determine if it overlaps this component
@@ -1736,3 +1748,27 @@ public:
 
 	bool ContainsData() const;
 };
+
+
+//////////////////////////////////////////////////////////////////////////
+// PrimitiveComponent inlines
+
+FORCEINLINE_DEBUGGABLE bool UPrimitiveComponent::ComponentOverlapMulti(TArray<struct FOverlapResult>& OutOverlaps, const class UWorld* World, const FVector& Pos, const FQuat& Rot, ECollisionChannel TestChannel, const struct FComponentQueryParams& Params, const struct FCollisionObjectQueryParams& ObjectQueryParams) const
+{
+	return ComponentOverlapMultiImpl(OutOverlaps, World, Pos, Rot, TestChannel, Params, ObjectQueryParams);
+}
+
+FORCEINLINE_DEBUGGABLE bool UPrimitiveComponent::ComponentOverlapMulti(TArray<struct FOverlapResult>& OutOverlaps, const class UWorld* World, const FVector& Pos, const FRotator& Rot, ECollisionChannel TestChannel, const struct FComponentQueryParams& Params, const struct FCollisionObjectQueryParams& ObjectQueryParams) const
+{
+	return ComponentOverlapMultiImpl(OutOverlaps, World, Pos, Rot.Quaternion(), TestChannel, Params, ObjectQueryParams);
+}
+
+FORCEINLINE_DEBUGGABLE bool UPrimitiveComponent::ComponentOverlapComponent(class UPrimitiveComponent* PrimComp, const FVector Pos, const FQuat& Rot, const FCollisionQueryParams& Params)
+{
+	return ComponentOverlapComponentImpl(PrimComp, Pos, Rot, Params);
+}
+
+FORCEINLINE_DEBUGGABLE bool UPrimitiveComponent::ComponentOverlapComponent(class UPrimitiveComponent* PrimComp, const FVector Pos, const FRotator Rot, const FCollisionQueryParams& Params)
+{
+	return ComponentOverlapComponentImpl(PrimComp, Pos, Rot.Quaternion(), Params);
+}
diff --git a/Engine/Source/Runtime/Engine/Classes/Components/SceneComponent.h b/Engine/Source/Runtime/Engine/Classes/Components/SceneComponent.h
--- a/Engine/Source/Runtime/Engine/Classes/Components/SceneComponent.h
+++ b/Engine/Source/Runtime/Engine/Classes/Components/SceneComponent.h
@@ -251,6 +251,7 @@ public:
 	UFUNCTION(BlueprintCallable, Category="Utilities|Transformation", meta=(DisplayName="SetRelativeRotation", AdvancedDisplay="bSweep,SweepHitResult"))
 	void K2_SetRelativeRotation(FRotator NewRotation, bool bSweep, FHitResult& SweepHitResult);
 	void SetRelativeRotation(FRotator NewRotation, bool bSweep=false, FHitResult* OutSweepHitResult=nullptr);
+	void SetRelativeRotation(const FQuat& NewRotation, bool bSweep=false, FHitResult* OutSweepHitResult=nullptr);
 
 	/** Set the transform of this component relative to its parent */
 	UFUNCTION(BlueprintCallable, Category="Utilities|Transformation", meta=(DisplayName="SetRelativeTransform"))
@@ -278,6 +279,7 @@ public:
 	UFUNCTION(BlueprintCallable, Category="Utilities|Transformation", meta=(DisplayName="AddRelativeRotation", AdvancedDisplay="bSweep,SweepHitResult"))
 	void K2_AddRelativeRotation(FRotator DeltaRotation, bool bSweep, FHitResult& SweepHitResult);
 	void AddRelativeRotation(FRotator DeltaRotation, bool bSweep=false, FHitResult* OutSweepHitResult=nullptr);
+	void AddRelativeRotation(const FQuat& DeltaRotation, bool bSweep=false, FHitResult* OutSweepHitResult=nullptr);
 
 	/** Adds a delta to the location of this component in its local reference frame */
 	UFUNCTION(BlueprintCallable, Category="Utilities|Transformation", meta=(DisplayName="AddLocalOffset", Keywords="location position"))
@@ -289,6 +291,7 @@ public:
 	UFUNCTION(BlueprintCallable, Category="Utilities|Transformation", meta=(DisplayName="AddLocalRotation", AdvancedDisplay="bSweep,SweepHitResult"))
 	void K2_AddLocalRotation(FRotator DeltaRotation, bool bSweep, FHitResult& SweepHitResult);
 	void AddLocalRotation(FRotator DeltaRotation, bool bSweep=false, FHitResult* OutSweepHitResult=nullptr);
+	void AddLocalRotation(const FQuat& DeltaRotation, bool bSweep=false, FHitResult* OutSweepHitResult=nullptr);
 
 
 	/** Adds a delta to the transform of this component in its local reference frame. Scale is unchanged. */
@@ -307,6 +310,7 @@ public:
 	UFUNCTION(BlueprintCallable, Category="Utilities|Transformation", meta=(DisplayName="SetWorldRotation", AdvancedDisplay="bSweep,SweepHitResult"))
 	void K2_SetWorldRotation(FRotator NewRotation, bool bSweep, FHitResult& SweepHitResult);
 	void SetWorldRotation(FRotator NewRotation, bool bSweep=false, FHitResult* OutSweepHitResult=nullptr);
+	void SetWorldRotation(const FQuat& NewRotation, bool bSweep=false, FHitResult* OutSweepHitResult=nullptr);
 
 
 	/** Set the relative scale of this component to put it at the supplied scale in world space. */
@@ -329,6 +333,7 @@ public:
 	UFUNCTION(BlueprintCallable, Category="Utilities|Transformation", meta=(DisplayName="AddWorldRotation", AdvancedDisplay="bSweep,SweepHitResult"))
 	void K2_AddWorldRotation(FRotator DeltaRotation, bool bSweep, FHitResult& SweepHitResult);
 	void AddWorldRotation(FRotator DeltaRotation, bool bSweep=false, FHitResult* OutSweepHitResult=nullptr);
+	void AddWorldRotation(const FQuat& DeltaRotation, bool bSweep=false, FHitResult* OutSweepHitResult=nullptr);
 
 	/** Adds a delta to the transform of this component in world space. Scale is unchanged. */
 	UFUNCTION(BlueprintCallable, Category="Utilities|Transformation", meta=(DisplayName="AddWorldTransform"))
@@ -525,6 +530,9 @@ public:
 	virtual FName GetComponentInstanceDataType() const override;
 	// End ActorComponent interface
 
+	// Call UpdateComponentToWorld if bWorldToComponentUpdated is false.
+	void ConditionalUpdateComponentToWorld();
+
 	// Begin UObject Interface
 	virtual void Serialize(FArchive& Ar) override;
 	virtual void PostInterpChange(UProperty* PropertyThatChanged) override;
@@ -542,7 +550,7 @@ protected:
 	 * Internal helper, for use from MoveComponent().  Special codepath since the normal setters call MoveComponent.
 	 * @return: true if location or rotation was changed.
 	 */
-	bool InternalSetWorldLocationAndRotation(FVector NewLocation, FRotator NewRotation, bool bNoPhysics=false);
+	bool InternalSetWorldLocationAndRotation(FVector NewLocation, const FQuat& NewQuat, bool bNoPhysics=false);
 
 	virtual void OnUpdateTransform(bool bSkipPhysicsMove);
 
@@ -552,7 +560,7 @@ protected:
 private:
 
 	void PropagateTransformUpdate(bool bTransformChanged, bool bSkipPhysicsMove = false);
-	void UpdateComponentToWorldWithParent(USceneComponent * Parent, bool bSkipPhysicsMove);
+	void UpdateComponentToWorldWithParent(USceneComponent* Parent, bool bSkipPhysicsMove, const FQuat& RelativeRotationQuat);
 
 
 public:
@@ -564,7 +572,9 @@ public:
 	 * Tries to move the component by a movement vector (Delta) and sets rotation to NewRotation.
 	 * Assumes that the component's current location is valid and that the component does fit in its current Location.
 	 * Dispatches blocking hit notifications (if bSweep is true), and calls UpdateOverlaps() after movement to update overlap state.
-	 * 
+	 *
+	 * @note This simply calls the virtual MoveComponentImpl() which can be overridden to implement custom behavior.
+	 * @note The overload taking rotation as an FQuat is slightly faster than the version using FRotator (which will be converted to an FQuat)..
 	 * @param Delta			The desired location change in world space.
 	 * @param NewRotation	The new desired rotation in world space.
 	 * @param bSweep		True to do a swept move, which will stop at a blocking collision. False to do a simple teleport, which will not stop for collisions.
@@ -572,7 +582,15 @@ public:
 	 * @param MoveFlags		Flags controlling behavior of the move. @see EMoveComponentFlags
 	 * @return				True if some movement occurred, false if no movement occurred.
 	 */
-	virtual bool MoveComponent( const FVector& Delta, const FRotator& NewRotation, bool bSweep, FHitResult* Hit=NULL, EMoveComponentFlags MoveFlags = MOVECOMP_NoFlags );
+	bool MoveComponent( const FVector& Delta, const FQuat& NewRotation,    bool bSweep, FHitResult* Hit=NULL, EMoveComponentFlags MoveFlags = MOVECOMP_NoFlags );
+	bool MoveComponent( const FVector& Delta, const FRotator& NewRotation, bool bSweep, FHitResult* Hit=NULL, EMoveComponentFlags MoveFlags = MOVECOMP_NoFlags );
+
+protected:
+
+	// Override this method for custom behavior.
+	virtual bool MoveComponentImpl( const FVector& Delta, const FQuat& NewRotation, bool bSweep, FHitResult* Hit=NULL, EMoveComponentFlags MoveFlags = MOVECOMP_NoFlags );
+
+public:
 
 
 	/** Returns true if movement is currently within the scope of an FScopedMovementUpdate. */
@@ -731,17 +749,18 @@ public:
 	UFUNCTION(BlueprintCallable, Category="Utilities|Transformation", meta=(DisplayName="SetRelativeLocationAndRotation"))
 	void K2_SetRelativeLocationAndRotation(FVector NewLocation, FRotator NewRotation, bool bSweep, FHitResult& SweepHitResult);
 	void SetRelativeLocationAndRotation(FVector NewLocation, FRotator NewRotation, bool bSweep=false, FHitResult* OutSweepHitResult=nullptr);
+	void SetRelativeLocationAndRotation(FVector NewLocation, const FQuat& NewRotation, bool bSweep=false, FHitResult* OutSweepHitResult=nullptr);
 
 	/** Set which parts of the relative transform should be relative to parent, and which should be relative to world */
 	UFUNCTION(BlueprintCallable, Category="Utilities|Transformation")
 	void SetAbsolute(bool bNewAbsoluteLocation = false, bool bNewAbsoluteRotation = false, bool bNewAbsoluteScale = false);
 
-	/** Set the relative location & FRotator rotation of this component to put it at the supplied pose in world space. */
+	/** Set the relative location and FRotator rotation of this component to put it at the supplied pose in world space. */
 	UFUNCTION(BlueprintCallable, Category="Utilities|Transformation", meta=(DisplayName="SetWorldLocationAndRotation"))
 	void K2_SetWorldLocationAndRotation(FVector NewLocation, FRotator NewRotation, bool bSweep, FHitResult& SweepHitResult);
 	void SetWorldLocationAndRotation(FVector NewLocation, FRotator NewRotation, bool bSweep=false, FHitResult* OutSweepHitResult=nullptr);
 
-	/** Set the relative location & FQuat rotation of this component to put it at the supplied pose in world space. */
+	/** Set the relative location and FQuat rotation of this component to put it at the supplied pose in world space. */
 	void SetWorldLocationAndRotation(FVector NewLocation, const FQuat& NewRotation, bool bSweep=false, FHitResult* OutSweepHitResult=nullptr);
 
 	/** Special version of SetWorldLocationAndRotation that does not affect physics. */
@@ -785,10 +804,11 @@ public:
 	 * If the socket is not found, then it returns the component's location and rotation in world space.
 	 * @param InSocketName the name of the socket to find
 	 * @param OutLocation (out) set to the world space location of the socket
-	 * @param OutRotation (optional out) if specified, set to the world space rotation of the socket
+	 * @param OutRotation (out) set to the world space rotation of the socket
 	 * @return whether or not the socket was found
 	 */
 	void GetSocketWorldLocationAndRotation(FName InSocketName, FVector& OutLocation, FRotator& OutRotation) const;
+	void GetSocketWorldLocationAndRotation(FName InSocketName, FVector& OutLocation, FQuat& OutRotation) const;
 
 	/**
 	 * Called to see if it's possible to attach another scene component as a child.
@@ -827,6 +847,65 @@ protected:
 	}
 };
 
+
+//////////////////////////////////////////////////////////////////////////
+// USceneComponent inlines
+
+FORCEINLINE_DEBUGGABLE void USceneComponent::ConditionalUpdateComponentToWorld()
+{
+	if (!bWorldToComponentUpdated)
+	{
+		UpdateComponentToWorld();
+	}
+}
+
+FORCEINLINE_DEBUGGABLE bool USceneComponent::MoveComponent(const FVector& Delta, const FQuat& NewRotation, bool bSweep, FHitResult* Hit, EMoveComponentFlags MoveFlags)
+{
+	return MoveComponentImpl(Delta, NewRotation, bSweep, Hit, MoveFlags);
+}
+
+FORCEINLINE_DEBUGGABLE bool USceneComponent::MoveComponent(const FVector& Delta, const FRotator& NewRotation, bool bSweep, FHitResult* Hit, EMoveComponentFlags MoveFlags)
+{
+	return MoveComponentImpl(Delta, NewRotation.Quaternion(), bSweep, Hit, MoveFlags);
+}
+
+FORCEINLINE_DEBUGGABLE void USceneComponent::SetRelativeLocation(FVector NewLocation, bool bSweep, FHitResult* OutSweepHitResult)
+{
+	SetRelativeLocationAndRotation(NewLocation, RelativeRotation, bSweep, OutSweepHitResult);
+}
+
+FORCEINLINE_DEBUGGABLE void USceneComponent::SetRelativeRotation(FRotator NewRotation, bool bSweep, FHitResult* OutSweepHitResult)
+{
+	SetRelativeLocationAndRotation(RelativeLocation, NewRotation, bSweep, OutSweepHitResult);
+}
+
+FORCEINLINE_DEBUGGABLE void USceneComponent::SetRelativeRotation(const FQuat& NewRotation, bool bSweep, FHitResult* OutSweepHitResult)
+{
+	SetRelativeLocationAndRotation(RelativeLocation, NewRotation, bSweep, OutSweepHitResult);
+}
+
+FORCEINLINE_DEBUGGABLE void USceneComponent::SetRelativeLocationAndRotation(FVector NewLocation, FRotator NewRotation, bool bSweep, FHitResult* OutSweepHitResult)
+{
+	SetRelativeLocationAndRotation(NewLocation, NewRotation.Quaternion(), bSweep, OutSweepHitResult);
+}
+
+FORCEINLINE_DEBUGGABLE void USceneComponent::AddRelativeLocation(FVector DeltaLocation, bool bSweep, FHitResult* OutSweepHitResult)
+{
+	SetRelativeLocationAndRotation(RelativeLocation + DeltaLocation, RelativeRotation, bSweep, OutSweepHitResult);
+}
+
+FORCEINLINE_DEBUGGABLE void USceneComponent::AddRelativeRotation(FRotator DeltaRotation, bool bSweep, FHitResult* OutSweepHitResult)
+{
+	SetRelativeLocationAndRotation(RelativeLocation, RelativeRotation + DeltaRotation, bSweep, OutSweepHitResult);
+}
+
+FORCEINLINE_DEBUGGABLE void USceneComponent::SetWorldRotation(const FRotator NewRotation, bool bSweep, FHitResult* OutSweepHitResult)
+{
+	SetWorldRotation(NewRotation.Quaternion(), bSweep, OutSweepHitResult);
+}
+
+//////////////////////////////////////////////////////////////////////////
+
 /** 
  *  Component instance cached data base class for scene components. 
  *  Stores a list of instance components attached to the 
@@ -962,7 +1041,7 @@ FORCEINLINE const FScopedMovementUpdate::TBlockingHitArray& FScopedMovementUpdat
 	return BlockingHits;
 }
 
-FORCEINLINE void FScopedMovementUpdate::AppendBlockingHit(const FHitResult& Hit)
+FORCEINLINE_DEBUGGABLE void FScopedMovementUpdate::AppendBlockingHit(const FHitResult& Hit)
 {
 	BlockingHits.Add(Hit);
 }
diff --git a/Engine/Source/Runtime/Engine/Classes/Components/SkeletalMeshComponent.h b/Engine/Source/Runtime/Engine/Classes/Components/SkeletalMeshComponent.h
--- a/Engine/Source/Runtime/Engine/Classes/Components/SkeletalMeshComponent.h
+++ b/Engine/Source/Runtime/Engine/Classes/Components/SkeletalMeshComponent.h
@@ -814,21 +814,26 @@ public:
 	virtual FBoxSphereBounds CalcBounds(const FTransform& LocalToWorld) const override;
 	virtual bool IsAnySimulatingPhysics() const override;
 	virtual void OnUpdateTransform(bool bSkipPhysicsMove) override;
-	
+	// End USceneComponent interface.
+
+	// Begin UPrimitiveComponent interface.
+protected:
 	/**
 	 *  Test the collision of the supplied component at the supplied location/rotation, and determine the set of components that it overlaps
 	 *  @param  OutOverlaps     Array of overlaps found between this component in specified pose and the world
 	 *  @param  World			World to use for overlap test
-	 *  @param  Pos             Location to place the component's geometry at to test against the world
-	 *  @param  Rot             Rotation to place components' geometry at to test against the world
+	 *  @param  Pos             Location of the component's geometry for the test against the world
+	 *  @param  Rot             Rotation of the component's geometry for the test against the world
 	 *  @param  TestChannel		The 'channel' that this ray is in, used to determine which components to hit
 	 *	@param	ObjectQueryParams	List of object types it's looking for. When this enters, we do object query with component shape
 	 *  @return TRUE if OutOverlaps contains any blocking results
 	 */
-	virtual bool ComponentOverlapMulti(TArray<struct FOverlapResult>& OutOverlaps, const class UWorld* World, const FVector& Pos, const FRotator& Rot, ECollisionChannel TestChannel, const struct FComponentQueryParams& Params, const struct FCollisionObjectQueryParams& ObjectQueryParams = FCollisionObjectQueryParams::DefaultObjectQueryParam) const override;
-	// End USceneComponent interface.
+	virtual bool ComponentOverlapMultiImpl(TArray<struct FOverlapResult>& OutOverlaps, const class UWorld* World, const FVector& Pos, const FQuat& Rot, ECollisionChannel TestChannel, const struct FComponentQueryParams& Params, const struct FCollisionObjectQueryParams& ObjectQueryParams = FCollisionObjectQueryParams::DefaultObjectQueryParam) const override;
+	
+	virtual bool ComponentOverlapComponentImpl(class UPrimitiveComponent* PrimComp, const FVector Pos, const FQuat& Quat, const FCollisionQueryParams& Params) override;
+
+public:
 
-	// Begin UPrimitiveComponent interface.
 	virtual class UBodySetup* GetBodySetup() override;
 	virtual bool CanEditSimulatePhysics() override;
 	virtual FBodyInstance* GetBodyInstance(FName BoneName = NAME_None, bool bGetWelded = true) const override;
@@ -844,7 +849,6 @@ public:
 	virtual float GetDistanceToCollision(const FVector& Point, FVector& ClosestPointOnCollision) const override;
 	virtual bool LineTraceComponent( FHitResult& OutHit, const FVector Start, const FVector End, const FCollisionQueryParams& Params ) override;
 	virtual bool SweepComponent( FHitResult& OutHit, const FVector Start, const FVector End, const FCollisionShape& CollisionShape, bool bTraceComplex=false) override;
-	virtual bool ComponentOverlapComponent(class UPrimitiveComponent* PrimComp, const FVector Pos, const FRotator FRotator, const FCollisionQueryParams& Params) override;
 	virtual bool OverlapComponent(const FVector& Pos, const FQuat& Rot, const FCollisionShape& CollisionShape) override;
 	virtual void SetSimulatePhysics(bool bEnabled) override;
 	virtual void AddRadialImpulse(FVector Origin, float Radius, float Strength, ERadialImpulseFalloff Falloff, bool bVelChange=false) override;
diff --git a/Engine/Source/Runtime/Engine/Classes/Engine/World.h b/Engine/Source/Runtime/Engine/Classes/Engine/World.h
--- a/Engine/Source/Runtime/Engine/Classes/Engine/World.h
+++ b/Engine/Source/Runtime/Engine/Classes/Engine/World.h
@@ -1687,54 +1687,60 @@ public:
 	// COMPONENT SWEEP
 
 	/**
-	 *  Sweep the geometry of the supplied component, and determine the set of components that it hits
+	 *  Sweep the geometry of the supplied component, and determine the set of components that it hits.
+	 *  @note The overload taking rotation as an FQuat is slightly faster than the version using FRotator (which will be converted to an FQuat)..
 	 *  @param  OutHits         Array of hits found between ray and the world
-	 *  @param  PrimComp        Component to use geometry from to test against the world. Transform of this component is ignored
+	 *  @param  PrimComp        Component's geometry to test against the world. Transform of this component is ignored
 	 *  @param  Start           Start location of the trace
 	 *  @param  End             End location of the trace
-	 *  @param  Rot             Rotation to place PrimComp geometry at to test against the world (rotation remains constant over trace)
+	 *  @param  Rot             Rotation of PrimComp geometry for test against the world (rotation remains constant over sweep)
 	 *  @param  Params          Additional parameters used for the trace
 	 *  @return TRUE if OutHits contains any blocking hit entries
 	 */
-	bool ComponentSweepMulti(TArray<struct FHitResult>& OutHits, class UPrimitiveComponent* PrimComp, const FVector& Start,const FVector& End, const FRotator& Rot, const FComponentQueryParams& Params)  const;
+	bool ComponentSweepMulti(TArray<struct FHitResult>& OutHits, class UPrimitiveComponent* PrimComp, const FVector& Start, const FVector& End, const FQuat& Rot,    const FComponentQueryParams& Params) const;
+	bool ComponentSweepMulti(TArray<struct FHitResult>& OutHits, class UPrimitiveComponent* PrimComp, const FVector& Start, const FVector& End, const FRotator& Rot, const FComponentQueryParams& Params) const;
 
 	// COMPONENT OVERLAP
 
 	/**
 	 *  Test the collision of the supplied component at the supplied location/rotation using object types, and determine the set of components that it overlaps
+	 *  @note The overload taking rotation as an FQuat is slightly faster than the version using FRotator (which will be converted to an FQuat)..
 	 *  @param  OutOverlaps     Array of overlaps found between component in specified pose and the world
-	 *  @param  PrimComp        Component to use geometry from to test against the world. Transform of this component is ignored
-	 *  @param  Pos             Location to place PrimComp geometry at to test against the world
-	 *  @param  Rot             Rotation to place PrimComp geometry at to test against the world
+	 *  @param  PrimComp        Component's geometry to test against the world. Transform of this component is ignored
+	 *  @param  Pos             Location of PrimComp geometry for test against the world
+	 *  @param  Rot             Rotation of PrimComp geometry for test against the world
 	 *	@param	ObjectQueryParams	List of object types it's looking for. When this enters, we do object query with component shape
 	 *  @return TRUE if any hit is found
 	 */
+	bool ComponentOverlapMulti(TArray<struct FOverlapResult>& OutOverlaps, const class UPrimitiveComponent* PrimComp, const FVector& Pos, const FQuat& Rot,    const FComponentQueryParams& Params = FComponentQueryParams::DefaultComponentQueryParams, const FCollisionObjectQueryParams& ObjectQueryParams=FCollisionObjectQueryParams::DefaultObjectQueryParam) const;
 	bool ComponentOverlapMulti(TArray<struct FOverlapResult>& OutOverlaps, const class UPrimitiveComponent* PrimComp, const FVector& Pos, const FRotator& Rot, const FComponentQueryParams& Params = FComponentQueryParams::DefaultComponentQueryParams, const FCollisionObjectQueryParams& ObjectQueryParams=FCollisionObjectQueryParams::DefaultObjectQueryParam) const;
 
 	/**
 	 *  Test the collision of the supplied component at the supplied location/rotation using a specific channel, and determine the set of components that it overlaps
 	 *  @param  OutOverlaps     Array of overlaps found between component in specified pose and the world
-	 *  @param  PrimComp        Component to use geometry from to test against the world. Transform of this component is ignored
-	 *  @param  Pos             Location to place PrimComp geometry at to test against the world
-	 *  @param  Rot             Rotation to place PrimComp geometry at to test against the world
+	 *  @param  PrimComp        Component's geometry to test against the world. Transform of this component is ignored
+	 *  @param  Pos             Location of PrimComp geometry for test against the world
+	 *  @param  Rot             Rotation of PrimComp geometry for test against the world
 	 *  @param  TraceChannel	The 'channel' that this query is in, used to determine which components to hit
 	 *  @return TRUE if OutOverlaps contains any blocking results
 	 */
 	DEPRECATED(4.8, "Use ComponentOverlapMultiByChannel instead.")
 	bool ComponentOverlapMulti(TArray<struct FOverlapResult>& OutOverlaps, const class UPrimitiveComponent* PrimComp, const FVector& Pos, const FRotator& Rot, ECollisionChannel TraceChannel, const FComponentQueryParams& Params, const FCollisionObjectQueryParams& ObjectQueryParams=FCollisionObjectQueryParams::DefaultObjectQueryParam) const
 	{
-		return ComponentOverlapMultiByChannel(OutOverlaps, PrimComp, Pos, Rot, TraceChannel, Params, ObjectQueryParams);
+		return ComponentOverlapMultiByChannel(OutOverlaps, PrimComp, Pos, Rot.Quaternion(), TraceChannel, Params, ObjectQueryParams);
 	}
 
 	/**
 	 *  Test the collision of the supplied component at the supplied location/rotation using a specific channel, and determine the set of components that it overlaps
+	 *  @note The overload taking rotation as an FQuat is slightly faster than the version using FRotator (which will be converted to an FQuat)..
 	 *  @param  OutOverlaps     Array of overlaps found between component in specified pose and the world
-	 *  @param  PrimComp        Component to use geometry from to test against the world. Transform of this component is ignored
-	 *  @param  Pos             Location to place PrimComp geometry at to test against the world
-	 *  @param  Rot             Rotation to place PrimComp geometry at to test against the world
+	 *  @param  PrimComp        Component's geometry to test against the world. Transform of this component is ignored
+	 *  @param  Pos             Location of PrimComp geometry for test against the world
+	 *  @param  Rot             Rotation of PrimComp geometry for test against the world
 	 *  @param  TraceChannel	The 'channel' that this query is in, used to determine which components to hit
 	 *  @return TRUE if OutOverlaps contains any blocking results
 	 */
+	bool ComponentOverlapMultiByChannel(TArray<struct FOverlapResult>& OutOverlaps, const class UPrimitiveComponent* PrimComp, const FVector& Pos, const FQuat& Rot,    ECollisionChannel TraceChannel, const FComponentQueryParams& Params = FComponentQueryParams::DefaultComponentQueryParams, const FCollisionObjectQueryParams& ObjectQueryParams=FCollisionObjectQueryParams::DefaultObjectQueryParam) const;
 	bool ComponentOverlapMultiByChannel(TArray<struct FOverlapResult>& OutOverlaps, const class UPrimitiveComponent* PrimComp, const FVector& Pos, const FRotator& Rot, ECollisionChannel TraceChannel, const FComponentQueryParams& Params = FComponentQueryParams::DefaultComponentQueryParams, const FCollisionObjectQueryParams& ObjectQueryParams=FCollisionObjectQueryParams::DefaultObjectQueryParam) const;
 
 	/**
@@ -3262,3 +3268,28 @@ public:
 private:
 	FWorldDelegates() {}
 };
+
+
+//////////////////////////////////////////////////////////////////////////
+// UWorld inlines:
+
+FORCEINLINE_DEBUGGABLE bool UWorld::ComponentOverlapMulti(TArray<struct FOverlapResult>& OutOverlaps, const class UPrimitiveComponent* PrimComp, const FVector& Pos, const FRotator& Rot, const struct FComponentQueryParams& Params, const struct FCollisionObjectQueryParams& ObjectQueryParams) const
+{
+	// Pass through to FQuat version.
+	return ComponentOverlapMulti(OutOverlaps, PrimComp, Pos, Rot.Quaternion(), Params, ObjectQueryParams);
+}
+
+FORCEINLINE_DEBUGGABLE bool UWorld::ComponentOverlapMultiByChannel(TArray<struct FOverlapResult>& OutOverlaps, const class UPrimitiveComponent* PrimComp, const FVector& Pos, const FRotator& Rot, ECollisionChannel TraceChannel, const FComponentQueryParams& Params /* = FComponentQueryParams::DefaultComponentQueryParams */, const FCollisionObjectQueryParams& ObjectQueryParams/* =FCollisionObjectQueryParams::DefaultObjectQueryParam */) const
+{
+	// Pass through to FQuat version.
+	return ComponentOverlapMultiByChannel(OutOverlaps, PrimComp, Pos, Rot.Quaternion(), TraceChannel, Params);
+}
+
+FORCEINLINE_DEBUGGABLE bool UWorld::ComponentSweepMulti(TArray<struct FHitResult>& OutHits, class UPrimitiveComponent* PrimComp, const FVector& Start, const FVector& End, const FRotator& Rot, const struct FComponentQueryParams& Params) const
+{
+	// Pass through to FQuat version.
+	return ComponentSweepMulti(OutHits, PrimComp, Start, End, Rot.Quaternion(), Params);
+}
+
+
+
diff --git a/Engine/Source/Runtime/Engine/Classes/GameFramework/Actor.h b/Engine/Source/Runtime/Engine/Classes/GameFramework/Actor.h
--- a/Engine/Source/Runtime/Engine/Classes/GameFramework/Actor.h
+++ b/Engine/Source/Runtime/Engine/Classes/GameFramework/Actor.h
@@ -676,6 +676,7 @@ public:
 	 */
 	UFUNCTION(BlueprintCallable, Category="Utilities|Transformation")
 	bool SetActorRotation(FRotator NewRotation);
+	bool SetActorRotation(const FQuat& NewRotation);
 
 	/** 
 	 * Move the actor instantly to the specified location and rotation.
@@ -699,6 +700,7 @@ public:
 	 * @return	Whether the rotation was successfully set.
 	 */
 	bool SetActorLocationAndRotation(FVector NewLocation, FRotator NewRotation, bool bSweep=false, FHitResult* OutSweepHitResult=nullptr);
+	bool SetActorLocationAndRotation(FVector NewLocation, const FQuat& NewRotation, bool bSweep=false, FHitResult* OutSweepHitResult=nullptr);
 
 	/** Set the Actor's world-space scale. */
 	UFUNCTION(BlueprintCallable, Category="Utilities|Transformation")
@@ -758,6 +760,7 @@ public:
 	UFUNCTION(BlueprintCallable, Category="Utilities|Transformation", meta=(DisplayName="AddActorWorldRotation", AdvancedDisplay="bSweep,SweepHitResult"))
 	void K2_AddActorWorldRotation(FRotator DeltaRotation, bool bSweep, FHitResult& SweepHitResult);
 	void AddActorWorldRotation(FRotator DeltaRotation, bool bSweep=false, FHitResult* OutSweepHitResult=nullptr);
+	void AddActorWorldRotation(const FQuat& DeltaRotation, bool bSweep=false, FHitResult* OutSweepHitResult=nullptr);
 
 
 	/** Adds a delta to the transform of this actor in world space. Scale is unchanged. */
@@ -785,6 +788,7 @@ public:
 	UFUNCTION(BlueprintCallable, Category="Utilities|Transformation", meta=(DisplayName="AddActorLocalRotation", AdvancedDisplay="bSweep,SweepHitResult"))
 	void K2_AddActorLocalRotation(FRotator DeltaRotation, bool bSweep, FHitResult& SweepHitResult);
 	void AddActorLocalRotation(FRotator DeltaRotation, bool bSweep=false, FHitResult* OutSweepHitResult=nullptr);
+	void AddActorLocalRotation(const FQuat& DeltaRotation, bool bSweep=false, FHitResult* OutSweepHitResult=nullptr);
 
 
 	/** Adds a delta to the transform of this component in its local reference frame */
@@ -810,6 +814,7 @@ public:
 	UFUNCTION(BlueprintCallable, Category="Utilities|Transformation", meta=(DisplayName="SetActorRelativeRotation", AdvancedDisplay="bSweep,SweepHitResult"))
 	void K2_SetActorRelativeRotation(FRotator NewRelativeRotation, bool bSweep, FHitResult& SweepHitResult);
 	void SetActorRelativeRotation(FRotator NewRelativeRotation, bool bSweep=false, FHitResult* OutSweepHitResult=nullptr);
+	void SetActorRelativeRotation(const FQuat& NewRelativeRotation, bool bSweep=false, FHitResult* OutSweepHitResult=nullptr);
 
 	/**
 	 * Set the actor's RootComponent to the specified relative transform
@@ -2482,7 +2487,9 @@ FORCEINLINE FVector AActor::GetSimpleCollisionCylinderExtent() const
 	FVector GetActorScale() const { return Super::GetActorScale(); } \
 	bool SetActorLocation(const FVector& NewLocation, bool bSweep=false, FHitResult* OutSweepHitResult=nullptr) { return Super::SetActorLocation(NewLocation, bSweep, OutSweepHitResult); } \
 	bool SetActorRotation(FRotator NewRotation) { return Super::SetActorRotation(NewRotation); } \
+	bool SetActorRotation(const FQuat& NewRotation) { return Super::SetActorRotation(NewRotation); } \
 	bool SetActorLocationAndRotation(FVector NewLocation, FRotator NewRotation, bool bSweep=false, FHitResult* OutSweepHitResult=nullptr) { return Super::SetActorLocationAndRotation(NewLocation, NewRotation, bSweep, OutSweepHitResult); } \
+	bool SetActorLocationAndRotation(FVector NewLocation, const FQuat& NewRotation, bool bSweep=false, FHitResult* OutSweepHitResult=nullptr) { return Super::SetActorLocationAndRotation(NewLocation, NewRotation, bSweep, OutSweepHitResult); } \
 	virtual bool TeleportTo( const FVector& DestLocation, const FRotator& DestRotation, bool bIsATest, bool bNoCheck ) override { return Super::TeleportTo(DestLocation, DestRotation, bIsATest, bNoCheck); } \
 	virtual FVector GetVelocity() const override { return Super::GetVelocity(); } \
 	float GetHorizontalDistanceTo(AActor* OtherActor)  { return Super::GetHorizontalDistanceTo(OtherActor); } \
diff --git a/Engine/Source/Runtime/Engine/Classes/PhysicsEngine/BodyInstance.h b/Engine/Source/Runtime/Engine/Classes/PhysicsEngine/BodyInstance.h
--- a/Engine/Source/Runtime/Engine/Classes/PhysicsEngine/BodyInstance.h
+++ b/Engine/Source/Runtime/Engine/Classes/PhysicsEngine/BodyInstance.h
@@ -898,6 +898,7 @@ public:
 
 	/**
 	 *  Determines the set of components that this body instance would overlap with at the supplied location/rotation
+	 *  @note The overload taking rotation as an FQuat is slightly faster than the version using FRotator (which will be converted to an FQuat)..
 	 *  @param  InOutOverlaps   Array of overlaps found between this component in specified pose and the world (new overlaps will be appended, the array is not cleared!)
 	 *  @param  World			World to use for overlap test
 	 *  @param  pWorldToComponent Used to convert the body instance world space position into local space before teleporting it to (Pos, Rot) [optional, use when the body instance isn't centered on the component instance]
@@ -909,6 +910,7 @@ public:
 	 *	@param	ObjectQueryParams	List of object types it's looking for. When this enters, we do object query with component shape
 	 *  @return TRUE if OutOverlaps contains any blocking results
 	 */
+	bool OverlapMulti(TArray<struct FOverlapResult>& InOutOverlaps, const class UWorld* World, const FTransform* pWorldToComponent, const FVector& Pos, const FQuat& Rot,    ECollisionChannel TestChannel, const struct FComponentQueryParams& Params, const struct FCollisionResponseParams& ResponseParams, const FCollisionObjectQueryParams& ObjectQueryParams = FCollisionObjectQueryParams::DefaultObjectQueryParam) const;
 	bool OverlapMulti(TArray<struct FOverlapResult>& InOutOverlaps, const class UWorld* World, const FTransform* pWorldToComponent, const FVector& Pos, const FRotator& Rot, ECollisionChannel TestChannel, const struct FComponentQueryParams& Params, const struct FCollisionResponseParams& ResponseParams, const FCollisionObjectQueryParams& ObjectQueryParams = FCollisionObjectQueryParams::DefaultObjectQueryParam) const;
 
 	/**
@@ -1025,4 +1027,14 @@ struct FBodyInstanceEditorHelpers
 private:
 	FBodyInstanceEditorHelpers() {}
 };
-#endif
\ No newline at end of file
+#endif
+
+//////////////////////////////////////////////////////////////////////////
+// BodyInstance inlines
+
+inline bool FBodyInstance::OverlapMulti(TArray<struct FOverlapResult>& InOutOverlaps, const class UWorld* World, const FTransform* pWorldToComponent, const FVector& Pos, const FRotator& Rot, ECollisionChannel TestChannel, const struct FComponentQueryParams& Params, const struct FCollisionResponseParams& ResponseParams, const struct FCollisionObjectQueryParams& ObjectQueryParams) const
+{
+	// Pass on to FQuat version
+	return OverlapMulti(InOutOverlaps, World, pWorldToComponent, Pos, Rot.Quaternion(), TestChannel, Params, ResponseParams, ObjectQueryParams);
+}
+
diff --git a/Engine/Source/Runtime/Engine/Private/Actor.cpp b/Engine/Source/Runtime/Engine/Private/Actor.cpp
--- a/Engine/Source/Runtime/Engine/Private/Actor.cpp
+++ b/Engine/Source/Runtime/Engine/Private/Actor.cpp
@@ -2667,7 +2667,7 @@ bool AActor::SetActorLocation(const FVector& NewLocation, bool bSweep, FHitResul
 	if (RootComponent)
 	{
 		const FVector Delta = NewLocation - GetActorLocation();
-		return RootComponent->MoveComponent( Delta, GetActorRotation(), bSweep, OutSweepHitResult );
+		return RootComponent->MoveComponent(Delta, GetActorQuat(), bSweep, OutSweepHitResult);
 	}
 	else if (OutSweepHitResult)
 	{
@@ -2681,7 +2681,17 @@ bool AActor::SetActorRotation(FRotator NewRotation)
 {
 	if (RootComponent)
 	{
-		return RootComponent->MoveComponent( FVector::ZeroVector, NewRotation, true );
+		return RootComponent->MoveComponent(FVector::ZeroVector, NewRotation, true);
+	}
+
+	return false;
+}
+
+bool AActor::SetActorRotation(const FQuat& NewRotation)
+{
+	if (RootComponent)
+	{
+		return RootComponent->MoveComponent(FVector::ZeroVector, NewRotation, true);
 	}
 
 	return false;
@@ -2692,7 +2702,22 @@ bool AActor::SetActorLocationAndRotation(FVector NewLocation, FRotator NewRotati
 	if (RootComponent)
 	{
 		const FVector Delta = NewLocation - GetActorLocation();
-		return RootComponent->MoveComponent( Delta, NewRotation, bSweep, OutSweepHitResult );
+		return RootComponent->MoveComponent(Delta, NewRotation, bSweep, OutSweepHitResult);
+	}
+	else if (OutSweepHitResult)
+	{
+		*OutSweepHitResult = FHitResult();
+	}
+
+	return false;
+}
+
+bool AActor::SetActorLocationAndRotation(FVector NewLocation, const FQuat& NewRotation, bool bSweep, FHitResult* OutSweepHitResult)
+{
+	if (RootComponent)
+	{
+		const FVector Delta = NewLocation - GetActorLocation();
+		return RootComponent->MoveComponent(Delta, NewRotation, bSweep, OutSweepHitResult);
 	}
 	else if (OutSweepHitResult)
 	{
@@ -2744,6 +2769,18 @@ void AActor::AddActorWorldRotation(FRotator DeltaRotation, bool bSweep, FHitResu
 	}
 }
 
+void AActor::AddActorWorldRotation(const FQuat& DeltaRotation, bool bSweep, FHitResult* OutSweepHitResult)
+{
+	if (RootComponent)
+	{
+		RootComponent->AddWorldRotation(DeltaRotation, bSweep, OutSweepHitResult);
+	}
+	else if (OutSweepHitResult)
+	{
+		*OutSweepHitResult = FHitResult();
+	}
+}
+
 void AActor::AddActorWorldTransform(const FTransform& DeltaTransform, bool bSweep, FHitResult* OutSweepHitResult)
 {
 	if (RootComponent)
@@ -2809,6 +2846,18 @@ void AActor::AddActorLocalRotation(FRotator DeltaRotation, bool bSweep, FHitResu
 	}
 }
 
+void AActor::AddActorLocalRotation(const FQuat& DeltaRotation, bool bSweep, FHitResult* OutSweepHitResult)
+{
+	if (RootComponent)
+	{
+		RootComponent->AddLocalRotation(DeltaRotation, bSweep, OutSweepHitResult);
+	}
+	else if (OutSweepHitResult)
+	{
+		*OutSweepHitResult = FHitResult();
+	}
+}
+
 void AActor::AddActorLocalTransform(const FTransform& NewTransform, bool bSweep, FHitResult* OutSweepHitResult)
 {
 	if(RootComponent)
@@ -2845,6 +2894,18 @@ void AActor::SetActorRelativeRotation(FRotator NewRelativeRotation, bool bSweep,
 	}
 }
 
+void AActor::SetActorRelativeRotation(const FQuat& NewRelativeRotation, bool bSweep, FHitResult* OutSweepHitResult)
+{
+	if (RootComponent)
+	{
+		RootComponent->SetRelativeRotation(NewRelativeRotation, bSweep, OutSweepHitResult);
+	}
+	else if (OutSweepHitResult)
+	{
+		*OutSweepHitResult = FHitResult();
+	}
+}
+
 void AActor::SetActorRelativeTransform(const FTransform& NewRelativeTransform, bool bSweep, FHitResult* OutSweepHitResult)
 {
 	if (RootComponent)
diff --git a/Engine/Source/Runtime/Engine/Private/Collision/WorldCollision.cpp b/Engine/Source/Runtime/Engine/Private/Collision/WorldCollision.cpp
--- a/Engine/Source/Runtime/Engine/Private/Collision/WorldCollision.cpp
+++ b/Engine/Source/Runtime/Engine/Private/Collision/WorldCollision.cpp
@@ -439,11 +439,11 @@ bool UWorld::OverlapMultiByProfile(TArray<struct FOverlapResult>& OutOverlaps, c
 }
 
 
-bool UWorld::ComponentOverlapMulti(TArray<struct FOverlapResult>& OutOverlaps, const class UPrimitiveComponent* PrimComp, const FVector& Pos, const FRotator& Rot, const struct FComponentQueryParams& Params, const struct FCollisionObjectQueryParams& ObjectQueryParams) const
+bool UWorld::ComponentOverlapMulti(TArray<struct FOverlapResult>& OutOverlaps, const class UPrimitiveComponent* PrimComp, const FVector& Pos, const FQuat& Quat, const struct FComponentQueryParams& Params, const struct FCollisionObjectQueryParams& ObjectQueryParams) const
 {
 	if (PrimComp)
 	{
-		ComponentOverlapMultiByChannel(OutOverlaps, PrimComp, Pos, Rot, PrimComp->GetCollisionObjectType(), Params, ObjectQueryParams);
+		ComponentOverlapMultiByChannel(OutOverlaps, PrimComp, Pos, Quat, PrimComp->GetCollisionObjectType(), Params, ObjectQueryParams);
 		
 		// object query returns true if any hit is found, not only blocking hit
 		return (OutOverlaps.Num() > 0);
@@ -455,11 +455,11 @@ bool UWorld::ComponentOverlapMulti(TArray<struct FOverlapResult>& OutOverlaps, c
 	}
 }
 
-bool UWorld::ComponentOverlapMultiByChannel(TArray<struct FOverlapResult>& OutOverlaps, const class UPrimitiveComponent* PrimComp, const FVector& Pos, const FRotator& Rot, ECollisionChannel TraceChannel, const FComponentQueryParams& Params /* = FComponentQueryParams::DefaultComponentQueryParams */, const FCollisionObjectQueryParams& ObjectQueryParams/* =FCollisionObjectQueryParams::DefaultObjectQueryParam */) const
+bool UWorld::ComponentOverlapMultiByChannel(TArray<struct FOverlapResult>& OutOverlaps, const class UPrimitiveComponent* PrimComp, const FVector& Pos, const FQuat& Quat, ECollisionChannel TraceChannel, const FComponentQueryParams& Params /* = FComponentQueryParams::DefaultComponentQueryParams */, const FCollisionObjectQueryParams& ObjectQueryParams/* =FCollisionObjectQueryParams::DefaultObjectQueryParam */) const
 {
 	if (PrimComp)
 	{
-		return PrimComp->ComponentOverlapMulti(OutOverlaps, this, Pos, Rot, TraceChannel, Params, ObjectQueryParams);
+		return PrimComp->ComponentOverlapMulti(OutOverlaps, this, Pos, Quat, TraceChannel, Params, ObjectQueryParams);
 	}
 	else
 	{
@@ -468,7 +468,7 @@ bool UWorld::ComponentOverlapMultiByChannel(TArray<struct FOverlapResult>& OutOv
 	}
 }
 
-bool UWorld::ComponentSweepMulti(TArray<struct FHitResult>& OutHits, class UPrimitiveComponent* PrimComp, const FVector& Start, const FVector& End, const FRotator& Rot, const struct FComponentQueryParams& Params) const
+bool UWorld::ComponentSweepMulti(TArray<struct FHitResult>& OutHits, class UPrimitiveComponent* PrimComp, const FVector& Start, const FVector& End, const FQuat& Quat, const struct FComponentQueryParams& Params) const
 {
 	if (GetPhysicsScene() == NULL)
 	{
@@ -513,7 +513,7 @@ bool UWorld::ComponentSweepMulti(TArray<struct FHitResult>& OutHits, class UPrim
 		}
 
 		// calculate the test global pose of the actor
-		const PxQuat PGeomRot = U2PQuat(Rot.Quaternion());
+		const PxQuat PGeomRot = U2PQuat(Quat);
 		const PxTransform PGlobalStartPose = PxTransform(U2PVector(Start), PGeomRot);
 		const PxTransform PGlobalEndPose = PxTransform(U2PVector(End), PGeomRot);
 
@@ -558,7 +558,6 @@ bool UWorld::ComponentSweepMulti(TArray<struct FHitResult>& OutHits, class UPrim
 	return bHaveBlockingHit;
 }
 
-
 #if ENABLE_COLLISION_ANALYZER
 
 #include "CollisionAnalyzerModule.h"
diff --git a/Engine/Source/Runtime/Engine/Private/Components/PrimitiveComponent.cpp b/Engine/Source/Runtime/Engine/Private/Components/PrimitiveComponent.cpp
--- a/Engine/Source/Runtime/Engine/Private/Components/PrimitiveComponent.cpp
+++ b/Engine/Source/Runtime/Engine/Private/Components/PrimitiveComponent.cpp
@@ -1318,7 +1318,7 @@ FCollisionShape UPrimitiveComponent::GetCollisionShape(float Inflation) const
 	return FCollisionShape::MakeBox(Bounds.BoxExtent + Inflation);
 }
 
-bool UPrimitiveComponent::MoveComponent( const FVector& Delta, const FRotator& NewRotation, bool bSweep, FHitResult* OutHit, EMoveComponentFlags MoveFlags)
+bool UPrimitiveComponent::MoveComponentImpl( const FVector& Delta, const FQuat& NewRotationQuat, bool bSweep, FHitResult* OutHit, EMoveComponentFlags MoveFlags)
 {
 	SCOPE_CYCLE_COUNTER(STAT_MoveComponentTime);
 
@@ -1354,10 +1354,7 @@ bool UPrimitiveComponent::MoveComponent( const FVector& Delta, const FRotator& N
 		return false;
 	}
 
-	if (!bWorldToComponentUpdated)
-	{
-		UpdateComponentToWorld();
-	}
+	ConditionalUpdateComponentToWorld();
 
 	// Init HitResult
 	FHitResult BlockingHit(1.f);
@@ -1374,8 +1371,7 @@ bool UPrimitiveComponent::MoveComponent( const FVector& Delta, const FRotator& N
 	if (DeltaSizeSq <= MinMovementDistSq)
 	{
 		// Skip if no vector or rotation.
-		const FQuat NewQuat = NewRotation.Quaternion();
-		if( NewQuat.Equals(ComponentToWorld.GetRotation()) )
+		if (NewRotationQuat.Equals(ComponentToWorld.GetRotation()))
 		{
 			// copy to optional output param
 			if (OutHit)
@@ -1397,7 +1393,7 @@ bool UPrimitiveComponent::MoveComponent( const FVector& Delta, const FRotator& N
 	if ( !bSweep )
 	{
 		// not sweeping, just go directly to the new transform
-		bMoved = InternalSetWorldLocationAndRotation(TraceEnd, NewRotation, bSkipPhysicsMove);
+		bMoved = InternalSetWorldLocationAndRotation(TraceEnd, NewRotationQuat, bSkipPhysicsMove);
 	}
 	else
 	{
@@ -1430,7 +1426,7 @@ bool UPrimitiveComponent::MoveComponent( const FVector& Delta, const FRotator& N
 			FComponentQueryParams Params(Name_MoveComponent, Actor);
 			FCollisionResponseParams ResponseParam;
 			InitSweepCollisionParams(Params, ResponseParam);
-			bool const bHadBlockingHit = GetWorld()->ComponentSweepMulti(Hits, this, TraceStart, TraceEnd, GetComponentRotation(), Params);
+			bool const bHadBlockingHit = GetWorld()->ComponentSweepMulti(Hits, this, TraceStart, TraceEnd, GetComponentQuat(), Params);
 
 			if (Hits.Num() > 0)
 			{
@@ -1537,7 +1533,7 @@ bool UPrimitiveComponent::MoveComponent( const FVector& Delta, const FRotator& N
 			// However any rotation was set at the end and not swept, so we can't assume the overlaps at the end location are correct if rotation changed.
 			if (PendingOverlaps.Num() == 0 && CVarAllowCachedOverlaps->GetInt())
 			{
-				if (Actor && Actor->GetRootComponent() == this && AreSymmetricRotations(NewRotation.Quaternion(), ComponentToWorld.GetRotation(), ComponentToWorld.GetScale3D()))
+				if (Actor && Actor->GetRootComponent() == this && AreSymmetricRotations(NewRotationQuat, ComponentToWorld.GetRotation(), ComponentToWorld.GetScale3D()))
 				{
 					// We know we are not overlapping any new components at the end location.
 					// Keep known overlapping child components, as long as we know their overlap status could not have changed (ie they are positioned relative to us).
@@ -1580,7 +1576,7 @@ bool UPrimitiveComponent::MoveComponent( const FVector& Delta, const FRotator& N
 			if (CVarAllowCachedOverlaps->GetInt() && Actor && Actor->GetRootComponent() == this && !IsDeferringMovementUpdates())
 			{
 				// Only if we know that we won't change our overlap status with children by moving.
-				if (AreSymmetricRotations(NewRotation.Quaternion(), ComponentToWorld.GetRotation(), ComponentToWorld.GetScale3D()) &&
+				if (AreSymmetricRotations(NewRotationQuat, ComponentToWorld.GetRotation(), ComponentToWorld.GetScale3D()) &&
 					AreAllCollideableDescendantsRelative())
 				{
 					OverlapsAtEndLocation = OverlappingComponents;
@@ -1590,7 +1586,7 @@ bool UPrimitiveComponent::MoveComponent( const FVector& Delta, const FRotator& N
 		}
 
 		// Update the location.  This will teleport any child components as well (not sweep).
-		bMoved = InternalSetWorldLocationAndRotation(NewLocation, NewRotation, bSkipPhysicsMove);
+		bMoved = InternalSetWorldLocationAndRotation(NewLocation, NewRotationQuat, bSkipPhysicsMove);
 	}
 
 	// Handle overlap notifications.
@@ -1745,7 +1741,7 @@ bool UPrimitiveComponent::SweepComponent(struct FHitResult& OutHit, const FVecto
 	return BodyInstance.Sweep(OutHit, Start, End, CollisionShape, bTraceComplex);
 }
 
-bool UPrimitiveComponent::ComponentOverlapComponent(class UPrimitiveComponent* PrimComp, const FVector Pos, const FRotator Rot, const struct FCollisionQueryParams& Params)
+bool UPrimitiveComponent::ComponentOverlapComponentImpl(class UPrimitiveComponent* PrimComp, const FVector Pos, const FQuat& Quat, const struct FCollisionQueryParams& Params)
 {
 	// if target is skeletalmeshcomponent and do not support singlebody physics
 	USkeletalMeshComponent * OtherComp = Cast<USkeletalMeshComponent>(PrimComp);
@@ -1759,7 +1755,7 @@ bool UPrimitiveComponent::ComponentOverlapComponent(class UPrimitiveComponent* P
 	{
 		TArray<FBodyInstance*> Bodies;
 		Bodies.Add(GetBodyInstance());
-		return BI->OverlapTestForBodies(Pos, Rot.Quaternion(), Bodies);
+		return BI->OverlapTestForBodies(Pos, Quat, Bodies);
 	}
 
 	return false;
@@ -2182,7 +2178,7 @@ void UPrimitiveComponent::UpdateOverlaps(TArray<FOverlapInfo> const* PendingOver
 					FComponentQueryParams Params(NAME_UpdateOverlaps);
 					Params.bTraceAsyncScene = bCheckAsyncSceneOnMove;
 					Params.AddIgnoredActors(MoveIgnoreActors);
-					MyWorld->ComponentOverlapMulti(Overlaps, this, GetComponentLocation(), GetComponentRotation(), Params);
+					MyWorld->ComponentOverlapMulti(Overlaps, this, GetComponentLocation(), GetComponentQuat(), Params);
 
 					for( int32 ResultIdx=0; ResultIdx<Overlaps.Num(); ResultIdx++ )
 					{
@@ -2277,12 +2273,12 @@ void UPrimitiveComponent::UpdateOverlaps(TArray<FOverlapInfo> const* PendingOver
 	}
 }
 
-bool UPrimitiveComponent::ComponentOverlapMulti(TArray<struct FOverlapResult>& OutOverlaps, const UWorld* World, const FVector& Pos, const FRotator& Rot, ECollisionChannel TestChannel, const struct FComponentQueryParams& Params, const struct FCollisionObjectQueryParams& ObjectQueryParams) const
+bool UPrimitiveComponent::ComponentOverlapMultiImpl(TArray<struct FOverlapResult>& OutOverlaps, const UWorld* World, const FVector& Pos, const FQuat& Quat, ECollisionChannel TestChannel, const struct FComponentQueryParams& Params, const struct FCollisionObjectQueryParams& ObjectQueryParams) const
 {
 	FComponentQueryParams ParamsWithSelf = Params;
 	ParamsWithSelf.AddIgnoredComponent(this);
 	OutOverlaps.Reset();
-	return BodyInstance.OverlapMulti(OutOverlaps, World, /*pWorldToComponent=*/ nullptr, Pos, Rot, TestChannel, ParamsWithSelf, FCollisionResponseParams(GetCollisionResponseToChannels()), ObjectQueryParams);
+	return BodyInstance.OverlapMulti(OutOverlaps, World, /*pWorldToComponent=*/ nullptr, Pos, Quat, TestChannel, ParamsWithSelf, FCollisionResponseParams(GetCollisionResponseToChannels()), ObjectQueryParams);
 }
 
 void UPrimitiveComponent::UpdatePhysicsVolume( bool bTriggerNotifiers )
diff --git a/Engine/Source/Runtime/Engine/Private/Components/SceneComponent.cpp b/Engine/Source/Runtime/Engine/Private/Components/SceneComponent.cpp
--- a/Engine/Source/Runtime/Engine/Private/Components/SceneComponent.cpp
+++ b/Engine/Source/Runtime/Engine/Private/Components/SceneComponent.cpp
@@ -318,7 +318,7 @@ void USceneComponent::OnUpdateTransform(bool bSkipPhysicsMove)
 {
 }
 
-void USceneComponent::UpdateComponentToWorldWithParent(USceneComponent * Parent, bool bSkipPhysicsMove)
+void USceneComponent::UpdateComponentToWorldWithParent(USceneComponent* Parent, bool bSkipPhysicsMove, const FQuat& RelativeRotationQuat)
 {
 	// If our parent hasn't been updated before, we'll need walk up our parent attach hierarchy
 	if (Parent && !Parent->bWorldToComponentUpdated)
@@ -335,7 +335,7 @@ void USceneComponent::UpdateComponentToWorldWithParent(USceneComponent * Parent,
 	bWorldToComponentUpdated = true;
 
 	// Calculate the new ComponentToWorld transform
-	const FTransform RelativeTransform(RelativeRotation, RelativeLocation, RelativeScale3D);
+	const FTransform RelativeTransform(RelativeRotationQuat, RelativeLocation, RelativeScale3D);
 	const FTransform NewTransform = CalcNewComponentToWorld(RelativeTransform, Parent);
 
 	// If transform has changed..
@@ -386,7 +386,7 @@ void USceneComponent::OnRegister()
 
 void USceneComponent::UpdateComponentToWorld(bool bSkipPhysicsMove)
 {
-	UpdateComponentToWorldWithParent(AttachParent, bSkipPhysicsMove);
+	UpdateComponentToWorldWithParent(AttachParent, bSkipPhysicsMove, RelativeRotation.Quaternion());
 }
 
 
@@ -689,119 +689,75 @@ void USceneComponent::UpdateBounds()
 #endif
 }
 
-void USceneComponent::SetRelativeLocation(FVector NewLocation, bool bSweep, FHitResult* OutSweepHitResult)
-{
-	SetRelativeLocationAndRotation(NewLocation, RelativeRotation, bSweep, OutSweepHitResult);
-}
 
-void USceneComponent::SetRelativeRotation(FRotator NewRotation, bool bSweep, FHitResult* OutSweepHitResult)
+void USceneComponent::SetRelativeLocationAndRotation(FVector NewLocation, const FQuat& NewRotation, bool bSweep, FHitResult* OutSweepHitResult)
 {
-	SetRelativeLocationAndRotation(RelativeLocation, NewRotation, bSweep, OutSweepHitResult);
-}
-
-void USceneComponent::SetRelativeLocationAndRotation(FVector NewLocation, FRotator NewRotation, bool bSweep, FHitResult* OutSweepHitResult)
-{
-	if(!NewLocation.Equals(RelativeLocation) || !NewRotation.Equals(RelativeRotation))
-	{
-		if (!bWorldToComponentUpdated)
-		{
-			UpdateComponentToWorld();
-		}
+	ConditionalUpdateComponentToWorld();
 
-		const FTransform DesiredRelTransform(NewRotation, NewLocation);
-		const FTransform DesiredWorldTransform = CalcNewComponentToWorld(DesiredRelTransform);
-		const FVector DesiredDelta = DesiredWorldTransform.GetTranslation() - ComponentToWorld.GetTranslation();
+	const FTransform DesiredRelTransform(NewRotation, NewLocation);
+	const FTransform DesiredWorldTransform = CalcNewComponentToWorld(DesiredRelTransform);
+	const FVector DesiredDelta = DesiredWorldTransform.GetTranslation() - ComponentToWorld.GetTranslation();
 
-		// Avoid Quat->Rotator conversion in the case that there is no change in relative->world rotation, common case for root components.
-		const FRotator DesiredRotation = (DesiredRelTransform.GetRotation() == DesiredWorldTransform.GetRotation()) ? NewRotation : DesiredWorldTransform.Rotator();
-
-		MoveComponent(DesiredDelta, DesiredRotation, bSweep, OutSweepHitResult);
-	}
-	else if (OutSweepHitResult)
-	{
-		*OutSweepHitResult = FHitResult();
-	}
+	MoveComponent(DesiredDelta, DesiredWorldTransform.GetRotation(), bSweep, OutSweepHitResult);
 }
 
-void USceneComponent::AddRelativeLocation(FVector DeltaLocation, bool bSweep, FHitResult* OutSweepHitResult)
-{
-	SetRelativeLocationAndRotation(RelativeLocation + DeltaLocation, RelativeRotation, bSweep, OutSweepHitResult);
-}
 
-void USceneComponent::AddRelativeRotation(FRotator DeltaRotation, bool bSweep, FHitResult* OutSweepHitResult)
+void USceneComponent::AddRelativeRotation(const FQuat& DeltaRotation, bool bSweep, FHitResult* OutSweepHitResult)
 {
-	SetRelativeLocationAndRotation(RelativeLocation, RelativeRotation + DeltaRotation, bSweep, OutSweepHitResult);
+	const FQuat CurRelRotQuat = RelativeRotation.Quaternion();
+	const FQuat NewRelRotQuat = DeltaRotation * CurRelRotQuat;
+	SetRelativeLocationAndRotation(RelativeLocation, NewRelRotQuat, bSweep, OutSweepHitResult);
 }
 
 void USceneComponent::AddLocalOffset(FVector DeltaLocation, bool bSweep, FHitResult* OutSweepHitResult)
 {
-	if (!DeltaLocation.IsNearlyZero())
-	{
-		const FVector LocalOffset = RelativeRotation.RotateVector(DeltaLocation);
-
-		SetRelativeLocationAndRotation(RelativeLocation + LocalOffset, RelativeRotation, bSweep, OutSweepHitResult);
-	}
-	else if (OutSweepHitResult)
-	{
-		*OutSweepHitResult = FHitResult();
-	}
+	const FVector LocalOffset = RelativeRotation.RotateVector(DeltaLocation);
+	SetRelativeLocationAndRotation(RelativeLocation + LocalOffset, RelativeRotation, bSweep, OutSweepHitResult);
 }
 
 void USceneComponent::AddLocalRotation(FRotator DeltaRotation, bool bSweep, FHitResult* OutSweepHitResult)
 {
-	if (!DeltaRotation.IsZero())
-	{
-		const FQuat RelativeRotQuat = RelativeRotation.Quaternion();
-		const FQuat DeltaRotQuat = DeltaRotation.Quaternion();
-		const FQuat NewRelRotQuat = RelativeRotQuat * DeltaRotQuat;
+	const FQuat CurRelRotQuat = RelativeRotation.Quaternion();
+	const FQuat NewRelRotQuat = CurRelRotQuat * DeltaRotation.Quaternion();
+	SetRelativeLocationAndRotation(RelativeLocation, NewRelRotQuat, bSweep, OutSweepHitResult);
+}
 
-		SetRelativeLocationAndRotation(RelativeLocation, NewRelRotQuat.Rotator(), bSweep, OutSweepHitResult);
-	}
-	else if (OutSweepHitResult)
-	{
-		*OutSweepHitResult = FHitResult();
-	}
+void USceneComponent::AddLocalRotation(const FQuat& DeltaRotation, bool bSweep, FHitResult* OutSweepHitResult)
+{
+	const FQuat CurRelRotQuat = RelativeRotation.Quaternion();
+	const FQuat NewRelRotQuat = CurRelRotQuat * DeltaRotation;
+	SetRelativeLocationAndRotation(RelativeLocation, NewRelRotQuat, bSweep, OutSweepHitResult);
 }
 
 void USceneComponent::AddLocalTransform(const FTransform& DeltaTransform, bool bSweep, FHitResult* OutSweepHitResult)
 {
 	const FTransform RelativeTransform( RelativeRotation, RelativeLocation, FVector(1,1,1) ); // don't use scaling, so it matches how AddLocalRotation/Offset work
 	const FTransform NewRelTransform = DeltaTransform * RelativeTransform;
-
 	SetRelativeTransform(NewRelTransform, bSweep, OutSweepHitResult);
 }
 
 void USceneComponent::AddWorldOffset(FVector DeltaLocation, bool bSweep, FHitResult* OutSweepHitResult)
 {
-	if (!DeltaLocation.IsZero())
-	{
-		const FVector NewWorldLocation = DeltaLocation + ComponentToWorld.GetTranslation();
-		SetWorldLocation(NewWorldLocation, bSweep, OutSweepHitResult);
-	}
-	else if (OutSweepHitResult)
-	{
-		*OutSweepHitResult = FHitResult();
-	}
+	const FVector NewWorldLocation = DeltaLocation + ComponentToWorld.GetTranslation();
+	SetWorldLocation(NewWorldLocation, bSweep, OutSweepHitResult);
 }
 
 void USceneComponent::AddWorldRotation(FRotator DeltaRotation, bool bSweep, FHitResult* OutSweepHitResult)
 {
-	if (!DeltaRotation.IsZero())
-	{
-		const FQuat NewWorldRotation = DeltaRotation.Quaternion() * ComponentToWorld.GetRotation();
-		SetWorldRotation(NewWorldRotation.Rotator(), bSweep, OutSweepHitResult);
-	}
-	else if (OutSweepHitResult)
-	{
-		*OutSweepHitResult = FHitResult();
-	}
+	const FQuat NewWorldRotation = DeltaRotation.Quaternion() * ComponentToWorld.GetRotation();
+	SetWorldRotation(NewWorldRotation, bSweep, OutSweepHitResult);
+}
+
+void USceneComponent::AddWorldRotation(const FQuat& DeltaRotation, bool bSweep, FHitResult* OutSweepHitResult)
+{
+	const FQuat NewWorldRotation = DeltaRotation * ComponentToWorld.GetRotation();
+	SetWorldRotation(NewWorldRotation, bSweep, OutSweepHitResult);
 }
 
 void USceneComponent::AddWorldTransform(const FTransform& DeltaTransform, bool bSweep, FHitResult* OutSweepHitResult)
 {
 	const FQuat NewWorldRotation = DeltaTransform.GetRotation() * ComponentToWorld.GetRotation();
 	const FVector NewWorldLocation = DeltaTransform.GetTranslation() + ComponentToWorld.GetTranslation();
-
 	SetWorldTransform(FTransform(NewWorldRotation, NewWorldLocation, FVector(1,1,1)));
 }
 
@@ -832,13 +788,13 @@ void USceneComponent::ResetRelativeTransform()
 
 void USceneComponent::SetRelativeTransform(const FTransform& NewTransform, bool bSweep, FHitResult* OutSweepHitResult)
 {
-	SetRelativeLocationAndRotation(NewTransform.GetTranslation(), NewTransform.Rotator(), bSweep, OutSweepHitResult);
+	SetRelativeLocationAndRotation(NewTransform.GetTranslation(), NewTransform.GetRotation(), bSweep, OutSweepHitResult);
 	SetRelativeScale3D(NewTransform.GetScale3D());
 }
 
 FTransform USceneComponent::GetRelativeTransform() const
 {
-	const FTransform RelativeTransform( RelativeRotation, RelativeLocation, RelativeScale3D );
+	const FTransform RelativeTransform( RelativeRotation.Quaternion(), RelativeLocation, RelativeScale3D );
 	return RelativeTransform;
 }
 
@@ -847,7 +803,7 @@ void USceneComponent::SetWorldLocation(FVector NewLocation, bool bSweep, FHitRes
 	FVector NewRelLocation = NewLocation;
 
 	// If attached to something, transform into local space
-	if(AttachParent != NULL && !bAbsoluteLocation)
+	if (AttachParent != NULL && !bAbsoluteLocation)
 	{
 		FTransform ParentToWorld = AttachParent->GetSocketTransform(AttachSocketName);
 		NewRelLocation = ParentToWorld.InverseTransformPosition(NewLocation);
@@ -856,22 +812,23 @@ void USceneComponent::SetWorldLocation(FVector NewLocation, bool bSweep, FHitRes
 	SetRelativeLocation(NewRelLocation, bSweep, OutSweepHitResult);
 }
 
-void USceneComponent::SetWorldRotation(FRotator NewRotation, bool bSweep, FHitResult* OutSweepHitResult)
+void USceneComponent::SetWorldRotation(const FQuat& NewRotation, bool bSweep, FHitResult* OutSweepHitResult)
 {
-	FRotator NewRelRotation = NewRotation;
+	FQuat NewRelRotation = NewRotation;
 
 	// If already attached to something, transform into local space
-	if(AttachParent != NULL && !bAbsoluteRotation)
+	if (AttachParent != NULL && !bAbsoluteRotation)
 	{
-		FTransform ParentToWorld = AttachParent->GetSocketTransform(AttachSocketName);
+		const FQuat ParentToWorldQuat = AttachParent->GetSocketQuaternion(AttachSocketName);
 		// Quat multiplication works reverse way, make sure you do Parent(-1) * World = Local, not World*Parent(-) = Local (the way matrix does)
-		FQuat NewRelQuat = ParentToWorld.GetRotation().Inverse() * NewRotation.Quaternion();
-		NewRelRotation = NewRelQuat.Rotator();
+		const FQuat NewRelQuat = ParentToWorldQuat.Inverse() * NewRotation;
+		NewRelRotation = NewRelQuat;
 	}
 
 	SetRelativeRotation(NewRelRotation, bSweep, OutSweepHitResult);
 }
 
+
 void USceneComponent::SetWorldScale3D(FVector NewScale)
 {
 	FVector NewRelScale = NewScale;
@@ -943,13 +900,12 @@ void USceneComponent::SetWorldLocationAndRotation(FVector NewLocation, const FQu
 		}
 	}
 
-	SetRelativeLocationAndRotation(NewLocation, NewFinalRotation.Rotator(), bSweep, OutSweepHitResult);
+	SetRelativeLocationAndRotation(NewLocation, NewFinalRotation, bSweep, OutSweepHitResult);
 }
 
 void USceneComponent::SetWorldLocationAndRotationNoPhysics(const FVector& NewLocation, const FRotator& NewRotation)
 {
 	// If attached to something, transform into local space
-	FQuat NewFinalRotation = NewRotation.Quaternion();
 	if (AttachParent != NULL)
 	{
 		const FTransform ParentToWorld = AttachParent->GetSocketTransform(AttachSocketName);
@@ -1192,7 +1148,7 @@ void USceneComponent::AttachTo(class USceneComponent* Parent, FName InSocketName
 				//Since the object is physically simulated it can't be the case that it's a child of object A and being attached to object B (at runtime)
 				if (bMaintainWorldPosition == false)	//User tried to attach but physically based so detach. However, if they provided relative coordinates we should still get the correct position
 				{
-					UpdateComponentToWorldWithParent(Parent, false);
+					UpdateComponentToWorldWithParent(Parent, false, RelativeRotation.Quaternion());
 					RelativeLocation = ComponentToWorld.GetLocation();
 					RelativeRotation = ComponentToWorld.GetRotation().Rotator();
 					RelativeScale3D = ComponentToWorld.GetScale3D();
@@ -1599,6 +1555,15 @@ void USceneComponent::GetSocketWorldLocationAndRotation(FName InSocketName, FVec
 	OutRotation = SocketWorldTransform.Rotator();
 }
 
+void USceneComponent::GetSocketWorldLocationAndRotation(FName InSocketName, FVector& OutLocation, FQuat& OutRotation) const
+{
+	FTransform const SocketWorldTransform = GetSocketTransform(InSocketName);
+
+	// assemble output
+	OutLocation = SocketWorldTransform.GetLocation();
+	OutRotation = SocketWorldTransform.GetRotation();
+}
+
 bool USceneComponent::IsWorldGeometry() const
 {
 	return false;
@@ -1748,7 +1713,7 @@ void USceneComponent::UpdatePhysicsVolume( bool bTriggerNotifiers )
 				const UPrimitiveComponent* SelfAsPrimitive = Cast<UPrimitiveComponent>(this);
 				if (SelfAsPrimitive)
 				{
-					MyWorld->ComponentOverlapMultiByChannel(Hits, SelfAsPrimitive, GetComponentLocation(), GetComponentRotation(), GetCollisionObjectType(), Params);
+					MyWorld->ComponentOverlapMultiByChannel(Hits, SelfAsPrimitive, GetComponentLocation(), GetComponentQuat(), GetCollisionObjectType(), Params);
 				}
 				else
 				{
@@ -1813,8 +1778,11 @@ void USceneComponent::BeginDestroy()
 	Super::BeginDestroy();
 }
 
-bool USceneComponent::InternalSetWorldLocationAndRotation(FVector NewLocation, FRotator NewRotation, bool bNoPhysics)
+bool USceneComponent::InternalSetWorldLocationAndRotation(FVector NewLocation, const FQuat& RotationQuat, bool bNoPhysics)
 {
+	check(bWorldToComponentUpdated); // This method is for internal use and this assumption should be valid.
+	FQuat NewRotationQuat(RotationQuat);
+
 	// If attached to something, transform into local space
 	if (AttachParent != NULL)
 	{
@@ -1827,24 +1795,21 @@ bool USceneComponent::InternalSetWorldLocationAndRotation(FVector NewLocation, F
 
 		if (!bAbsoluteRotation)
 		{
-			FQuat const NewQuat = NewRotation.Quaternion();
-
 			// Quat multiplication works reverse way, make sure you do Parent(-1) * World = Local, not World*Parent(-) = Local (the way matrix does)
-			FQuat NewRelQuat = ParentToWorld.GetRotation().Inverse() * NewQuat;
-			NewRotation = NewRelQuat.Rotator();
+			NewRotationQuat = ParentToWorld.GetRotation().Inverse() * NewRotationQuat;
 		}
 	}
 
-	bool bUpdated = false;
-	if (!NewLocation.Equals(RelativeLocation) || !NewRotation.Equals(RelativeRotation))
+	const FRotator NewRotator = NewRotationQuat.Rotator();
+	if (!NewLocation.Equals(RelativeLocation) || !NewRotator.Equals(RelativeRotation))
 	{
 		RelativeLocation = NewLocation;
-		RelativeRotation = NewRotation;
-		UpdateComponentToWorld(bNoPhysics);
-		bUpdated = true;
+		RelativeRotation = NewRotator;
+		UpdateComponentToWorldWithParent(AttachParent, bNoPhysics, NewRotationQuat);
+		return true;
 	}
 
-	return bUpdated;
+	return false;
 }
 
 void USceneComponent::UpdateOverlaps(TArray<FOverlapInfo> const* PendingOverlaps, bool bDoNotifies, const TArray<FOverlapInfo>* OverlapsAtEndLocation)
@@ -1897,7 +1862,7 @@ bool USceneComponent::CheckStaticMobilityAndWarn(const FText& ActionText) const
 	return false;
 }
 
-bool USceneComponent::MoveComponent( const FVector& Delta, const FRotator& NewRotation, bool bSweep, FHitResult* OutHit, EMoveComponentFlags MoveFlags )
+bool USceneComponent::MoveComponentImpl( const FVector& Delta, const FQuat& NewRotation, bool bSweep, FHitResult* OutHit, EMoveComponentFlags MoveFlags )
 {
 	// static things can move before they are registered (e.g. immediately after streaming), but not after.
 	static const FText WarnText = LOCTEXT("InvalidMove", "move");
@@ -1916,17 +1881,13 @@ bool USceneComponent::MoveComponent( const FVector& Delta, const FRotator& NewRo
 		*OutHit = FHitResult(1.f);
 	}
 
-	if (!bWorldToComponentUpdated)
-	{
-		UpdateComponentToWorld();
-	}
+	ConditionalUpdateComponentToWorld();
 
 	// early out for zero case
 	if( Delta.IsZero() )
 	{
 		// Skip if no vector or rotation.
-		const FQuat NewQuat = NewRotation.Quaternion();
-		if( NewQuat.Equals(ComponentToWorld.GetRotation(), SMALL_NUMBER) )
+		if (NewRotation.Equals(ComponentToWorld.GetRotation()))
 		{
 			return true;
 		}
diff --git a/Engine/Source/Runtime/Engine/Private/KismetSystemLibrary.cpp b/Engine/Source/Runtime/Engine/Private/KismetSystemLibrary.cpp
--- a/Engine/Source/Runtime/Engine/Private/KismetSystemLibrary.cpp
+++ b/Engine/Source/Runtime/Engine/Private/KismetSystemLibrary.cpp
@@ -929,7 +929,7 @@ bool UKismetSystemLibrary::ComponentOverlapComponents_DEPRECATED(UPrimitiveCompo
 
 	FCollisionObjectQueryParams::InitType InitType = FCollisionObjectQueryParams::GetCollisionChannelFromOverlapFilter(Filter);
 	check( Component->GetWorld());
-	Component->GetWorld()->ComponentOverlapMulti(Overlaps, Component, ComponentTransform.GetTranslation(), ComponentTransform.GetRotation().Rotator(), Params, FCollisionObjectQueryParams(InitType));
+	Component->GetWorld()->ComponentOverlapMulti(Overlaps, Component, ComponentTransform.GetTranslation(), ComponentTransform.GetRotation(), Params, FCollisionObjectQueryParams(InitType));
 
 	for (int32 OverlapIdx=0; OverlapIdx<Overlaps.Num(); ++OverlapIdx)
 	{
@@ -1142,7 +1142,7 @@ bool UKismetSystemLibrary::ComponentOverlapComponents_NEW(UPrimitiveComponent* C
 	}
 
 	check( Component->GetWorld());
-	Component->GetWorld()->ComponentOverlapMulti(Overlaps, Component, ComponentTransform.GetTranslation(), ComponentTransform.GetRotation().Rotator(), Params, ObjectParams);
+	Component->GetWorld()->ComponentOverlapMulti(Overlaps, Component, ComponentTransform.GetTranslation(), ComponentTransform.GetRotation(), Params, ObjectParams);
 
 	for (int32 OverlapIdx=0; OverlapIdx<Overlaps.Num(); ++OverlapIdx)
 	{
diff --git a/Engine/Source/Runtime/Engine/Private/PhysicsEngine/BodyInstance.cpp b/Engine/Source/Runtime/Engine/Private/PhysicsEngine/BodyInstance.cpp
--- a/Engine/Source/Runtime/Engine/Private/PhysicsEngine/BodyInstance.cpp
+++ b/Engine/Source/Runtime/Engine/Private/PhysicsEngine/BodyInstance.cpp
@@ -3980,7 +3980,7 @@ FTransform RootSpaceToWeldedSpace(const FBodyInstance* BI, const FTransform& Roo
 	return RootTM;
 }
 
-bool FBodyInstance::OverlapMulti(TArray<struct FOverlapResult>& InOutOverlaps, const class UWorld* World, const FTransform* pWorldToComponent, const FVector& Pos, const FRotator& Rot, ECollisionChannel TestChannel, const struct FComponentQueryParams& Params, const struct FCollisionResponseParams& ResponseParams, const struct FCollisionObjectQueryParams& ObjectQueryParams) const
+bool FBodyInstance::OverlapMulti(TArray<struct FOverlapResult>& InOutOverlaps, const class UWorld* World, const FTransform* pWorldToComponent, const FVector& Pos, const FQuat& Quat, ECollisionChannel TestChannel, const struct FComponentQueryParams& Params, const struct FCollisionResponseParams& ResponseParams, const struct FCollisionObjectQueryParams& ObjectQueryParams) const
 {
 	if ( (IsValidBodyInstance() || (WeldParent && WeldParent->IsValidBodyInstance())) == false )
 	{
@@ -3992,7 +3992,7 @@ bool FBodyInstance::OverlapMulti(TArray<struct FOverlapResult>& InOutOverlaps, c
 	bool bHaveBlockingHit = false;
 
 	// Determine how to convert the local space of this body instance to the test space
-	const FTransform ComponentSpaceToTestSpace(Rot, Pos);
+	const FTransform ComponentSpaceToTestSpace(Quat, Pos);
 
 	FTransform BodyInstanceSpaceToTestSpace;
 	if (pWorldToComponent)
diff --git a/Engine/Source/Runtime/Engine/Private/SkeletalMeshComponentPhysics.cpp b/Engine/Source/Runtime/Engine/Private/SkeletalMeshComponentPhysics.cpp
--- a/Engine/Source/Runtime/Engine/Private/SkeletalMeshComponentPhysics.cpp
+++ b/Engine/Source/Runtime/Engine/Private/SkeletalMeshComponentPhysics.cpp
@@ -1850,7 +1850,7 @@ bool USkeletalMeshComponent::SweepComponent( FHitResult& OutHit, const FVector S
 	return bHaveHit;
 }
 
-bool USkeletalMeshComponent::ComponentOverlapComponent(class UPrimitiveComponent* PrimComp,const FVector Pos,const FRotator Rot,const struct FCollisionQueryParams& Params)
+bool USkeletalMeshComponent::ComponentOverlapComponentImpl(class UPrimitiveComponent* PrimComp,const FVector Pos,const FQuat& Quat,const struct FCollisionQueryParams& Params)
 {
 	//we do not support skeletal mesh vs skeletal mesh overlap test
 	if (PrimComp->IsA<USkeletalMeshComponent>())
@@ -1859,10 +1859,10 @@ bool USkeletalMeshComponent::ComponentOverlapComponent(class UPrimitiveComponent
 		return false;
 	}
 
-	if(FBodyInstance* BI = PrimComp->GetBodyInstance())
+	if (FBodyInstance* BI = PrimComp->GetBodyInstance())
 	{
-		return BI->OverlapTestForBodies(Pos, Rot.Quaternion(), Bodies);
-					}
+		return BI->OverlapTestForBodies(Pos, Quat, Bodies);
+	}
 
 	return false;
 }
@@ -1880,7 +1880,7 @@ bool USkeletalMeshComponent::OverlapComponent(const FVector& Pos, const FQuat& R
 	return false;
 }
 
-bool USkeletalMeshComponent::ComponentOverlapMulti(TArray<struct FOverlapResult>& OutOverlaps, const UWorld* World, const FVector& Pos, const FRotator& Rot, ECollisionChannel TestChannel, const struct FComponentQueryParams& Params, const struct FCollisionObjectQueryParams& ObjectQueryParams) const
+bool USkeletalMeshComponent::ComponentOverlapMultiImpl(TArray<struct FOverlapResult>& OutOverlaps, const UWorld* World, const FVector& Pos, const FQuat& Quat, ECollisionChannel TestChannel, const struct FComponentQueryParams& Params, const struct FCollisionObjectQueryParams& ObjectQueryParams) const
 {
 	OutOverlaps.Reset();
 
@@ -1899,7 +1899,7 @@ bool USkeletalMeshComponent::ComponentOverlapMulti(TArray<struct FOverlapResult>
 	for (const FBodyInstance* Body : Bodies)
 	{
 		checkSlow(Body);
-		if (Body->OverlapMulti(OutOverlaps, World, &WorldToComponent, Pos, Rot, TestChannel, ParamsWithSelf, ResponseParams, ObjectQueryParams))
+		if (Body->OverlapMulti(OutOverlaps, World, &WorldToComponent, Pos, Quat, TestChannel, ParamsWithSelf, ResponseParams, ObjectQueryParams))
 		{
 			bHaveBlockingHit = true;
 		}
