#ue4 - MovementComponent refactored to add overloads using FQuat in addition to FRotator.

- Some virtual methods have been changed and overrides will need to be updated. Make sure to call the correct Super function in overrides if currently using it.

diff --git a/Engine/Source/Runtime/Engine/Classes/GameFramework/CharacterMovementComponent.h b/Engine/Source/Runtime/Engine/Classes/GameFramework/CharacterMovementComponent.h
--- a/Engine/Source/Runtime/Engine/Classes/GameFramework/CharacterMovementComponent.h
+++ b/Engine/Source/Runtime/Engine/Classes/GameFramework/CharacterMovementComponent.h
@@ -1351,7 +1351,7 @@ protected:
 	virtual void MaintainHorizontalGroundVelocity();
 
 	/** Overridden to set bJustTeleported to true, so we don't make incorrect velocity calculations based on adjusted movement. */
-	virtual bool ResolvePenetration(const FVector& Adjustment, const FHitResult& Hit, const FRotator& NewRotation) override;
+	virtual bool ResolvePenetrationImpl(const FVector& Adjustment, const FHitResult& Hit, const FQuat& NewRotation) override;
 
 	/** Handle a blocking impact. Calls ApplyImpactPhysicsForces for the hit, if bEnablePhysicsInteraction is true. */
 	virtual void HandleImpact(FHitResult const& Hit, float TimeSlice=0.f, const FVector& MoveDelta = FVector::ZeroVector) override;
diff --git a/Engine/Source/Runtime/Engine/Classes/GameFramework/FloatingPawnMovement.h b/Engine/Source/Runtime/Engine/Classes/GameFramework/FloatingPawnMovement.h
--- a/Engine/Source/Runtime/Engine/Classes/GameFramework/FloatingPawnMovement.h
+++ b/Engine/Source/Runtime/Engine/Classes/GameFramework/FloatingPawnMovement.h
@@ -23,7 +23,10 @@ class ENGINE_API UFloatingPawnMovement : public UPawnMovementComponent
 
 	//Begin UMovementComponent Interface
 	virtual float GetMaxSpeed() const override { return MaxSpeed; }
-	virtual bool ResolvePenetration(const FVector& Adjustment, const FHitResult& Hit, const FRotator& NewRotation) override;
+
+protected:
+	virtual bool ResolvePenetrationImpl(const FVector& Adjustment, const FHitResult& Hit, const FQuat& NewRotation) override;
+public:
 	//End UMovementComponent Interface
 
 	/** Maximum velocity magnitude allowed for the controlled Pawn. */
diff --git a/Engine/Source/Runtime/Engine/Classes/GameFramework/MovementComponent.h b/Engine/Source/Runtime/Engine/Classes/GameFramework/MovementComponent.h
--- a/Engine/Source/Runtime/Engine/Classes/GameFramework/MovementComponent.h
+++ b/Engine/Source/Runtime/Engine/Classes/GameFramework/MovementComponent.h
@@ -241,12 +241,20 @@ public:
 	virtual bool OverlapTest(const FVector& Location, const FQuat& RotationQuat, const ECollisionChannel CollisionChannel, const FCollisionShape& CollisionShape, const AActor* IgnoreActor) const;
 
 	/**
-	 * Moves our UpdatedComponent by the given Delta, and sets rotation to NewRotation.
-	 * Respects the plane constraint, if enabled.
+	 * Moves our UpdatedComponent by the given Delta, and sets rotation to NewRotation. Respects the plane constraint, if enabled.
+	 * @note This simply calls the virtual MoveUpdatedComponentImpl() which can be overridden to implement custom behavior.
+	 * @note The overload taking rotation as an FQuat is slightly faster than the version using FRotator (which will be converted to an FQuat).
 	 * @return True if some movement occurred, false if no movement occurred. Result of any impact will be stored in OutHit.
 	 */
-	virtual bool MoveUpdatedComponent(const FVector& Delta, const FRotator& NewRotation, bool bSweep, FHitResult* OutHit = NULL);
+	bool MoveUpdatedComponent(const FVector& Delta, const FQuat& NewRotation,    bool bSweep, FHitResult* OutHit = NULL);
+	bool MoveUpdatedComponent(const FVector& Delta, const FRotator& NewRotation, bool bSweep, FHitResult* OutHit = NULL);
+
+protected:
 
+	virtual bool MoveUpdatedComponentImpl(const FVector& Delta, const FQuat& NewRotation, bool bSweep, FHitResult* OutHit = NULL);
+
+public:
+	
 	/**
 	 * Moves our UpdatedComponent by the given Delta, and sets rotation to NewRotation.
 	 * Respects the plane constraint, if enabled.
@@ -258,25 +266,35 @@ public:
 	/**
 	 * Calls MoveUpdatedComponent(), handling initial penetrations by calling ResolvePenetration().
 	 * If this adjustment succeeds, the original movement will be attempted again.
+	 * @note The overload taking rotation as an FQuat is slightly faster than the version using FRotator (which will be converted to an FQuat)..
 	 * @return result of the final MoveUpdatedComponent() call.
 	 */
+	bool SafeMoveUpdatedComponent(const FVector& Delta, const FQuat& NewRotation,    bool bSweep, FHitResult& OutHit);
 	bool SafeMoveUpdatedComponent(const FVector& Delta, const FRotator& NewRotation, bool bSweep, FHitResult& OutHit);
 
 	/**
 	 * Calculate a movement adjustment to try to move out of a penetration from a failed move.
-	 * @param Hit - the result of the failed move
+	 * @param Hit the result of the failed move
 	 * @return The adjustment to use after a failed move, or a zero vector if no attempt should be made.
 	 */
 	virtual FVector GetPenetrationAdjustment(const FHitResult& Hit) const;
 	
 	/**
-	 * Try to move out of penetration in an object after a failed move.
-	 * This function should respect the plane constraint if applicable.
-	 * @param Adjustment - the requested adjustment, usually from GetPenetrationAdjustment().
-	 * @param Hit - the result of the failed move
-	 * @return true if the adjustment was successful and the original move should be retried, or false if no repeated attempt should be made.
+	 * Try to move out of penetration in an object after a failed move. This function should respect the plane constraint if applicable.
+	 * @note This simply calls the virtual ResolvePenetrationImpl() which can be overridden to implement custom behavior.
+	 * @note The overload taking rotation as an FQuat is slightly faster than the version using FRotator (which will be converted to an FQuat)..
+	 * @param Adjustment	The requested adjustment, usually from GetPenetrationAdjustment()
+	 * @param Hit			The result of the failed move
+	 * @return True if the adjustment was successful and the original move should be retried, or false if no repeated attempt should be made.
 	 */
-	virtual bool ResolvePenetration(const FVector& Adjustment, const FHitResult& Hit, const FRotator& NewRotation);
+	bool ResolvePenetration(const FVector& Adjustment, const FHitResult& Hit, const FQuat& NewRotation);
+	bool ResolvePenetration(const FVector& Adjustment, const FHitResult& Hit, const FRotator& NewRotation);
+
+protected:
+
+	virtual bool ResolvePenetrationImpl(const FVector& Adjustment, const FHitResult& Hit, const FQuat& NewRotation);
+
+public:
 
 	/**
 	 * Compute a vector to slide along a surface, given an attempted move, time, and normal.
@@ -422,3 +440,28 @@ inline EPlaneConstraintAxisSetting UMovementComponent::GetPlaneConstraintAxisSet
 {
 	return PlaneConstraintAxisSetting;
 }
+
+FORCEINLINE_DEBUGGABLE bool UMovementComponent::SafeMoveUpdatedComponent(const FVector& Delta, const FRotator& NewRotation, bool bSweep, FHitResult& OutHit)
+{
+	return SafeMoveUpdatedComponent(Delta, NewRotation.Quaternion(), bSweep, OutHit);
+}
+
+FORCEINLINE_DEBUGGABLE bool UMovementComponent::MoveUpdatedComponent(const FVector& Delta, const FQuat& NewRotation, bool bSweep, FHitResult* OutHit)
+{
+	return MoveUpdatedComponentImpl(Delta, NewRotation, bSweep, OutHit);
+}
+
+FORCEINLINE_DEBUGGABLE bool UMovementComponent::MoveUpdatedComponent(const FVector& Delta, const FRotator& NewRotation, bool bSweep, FHitResult* OutHit)
+{
+	return MoveUpdatedComponentImpl(Delta, NewRotation.Quaternion(), bSweep, OutHit);
+}
+
+FORCEINLINE_DEBUGGABLE bool UMovementComponent::ResolvePenetration(const FVector& Adjustment, const FHitResult& Hit, const FQuat& NewRotation)
+{
+	return ResolvePenetrationImpl(Adjustment, Hit, NewRotation);
+}
+
+FORCEINLINE_DEBUGGABLE bool UMovementComponent::ResolvePenetration(const FVector& Adjustment, const FHitResult& Hit, const FRotator& NewRotation)
+{
+	return ResolvePenetrationImpl(Adjustment, Hit, NewRotation.Quaternion());
+}
diff --git a/Engine/Source/Runtime/Engine/Private/Components/CharacterMovementComponent.cpp b/Engine/Source/Runtime/Engine/Private/Components/CharacterMovementComponent.cpp
--- a/Engine/Source/Runtime/Engine/Private/Components/CharacterMovementComponent.cpp
+++ b/Engine/Source/Runtime/Engine/Private/Components/CharacterMovementComponent.cpp
@@ -1989,10 +1989,10 @@ float UCharacterMovementComponent::GetMaxSpeed() const
 }
 
 
-bool UCharacterMovementComponent::ResolvePenetration(const FVector& Adjustment, const FHitResult& Hit, const FRotator& NewRotation)
+bool UCharacterMovementComponent::ResolvePenetrationImpl(const FVector& Adjustment, const FHitResult& Hit, const FQuat& NewRotation)
 {
 	// If movement occurs, mark that we teleported, so we don't incorrectly adjust velocity based on a potentially very different movement than our movement direction.
-	bJustTeleported |= Super::ResolvePenetration(Adjustment, Hit, NewRotation);
+	bJustTeleported |= Super::ResolvePenetrationImpl(Adjustment, Hit, NewRotation);
 	return bJustTeleported;
 }
 
diff --git a/Engine/Source/Runtime/Engine/Private/Components/MovementComponent.cpp b/Engine/Source/Runtime/Engine/Private/Components/MovementComponent.cpp
--- a/Engine/Source/Runtime/Engine/Private/Components/MovementComponent.cpp
+++ b/Engine/Source/Runtime/Engine/Private/Components/MovementComponent.cpp
@@ -428,7 +428,7 @@ void UMovementComponent::SnapUpdatedComponentToPlane()
 
 
 
-bool UMovementComponent::MoveUpdatedComponent( const FVector& Delta, const FRotator& NewRotation, bool bSweep, FHitResult* OutHit)
+bool UMovementComponent::MoveUpdatedComponentImpl( const FVector& Delta, const FQuat& NewRotation, bool bSweep, FHitResult* OutHit)
 {
 	if (UpdatedComponent)
 	{
@@ -442,11 +442,11 @@ bool UMovementComponent::MoveUpdatedComponent( const FVector& Delta, const FRota
 
 bool UMovementComponent::K2_MoveUpdatedComponent(FVector Delta, FRotator NewRotation, FHitResult& OutHit, bool bSweep)
 {
-	return SafeMoveUpdatedComponent(Delta, NewRotation, bSweep, OutHit);
+	return SafeMoveUpdatedComponent(Delta, NewRotation.Quaternion(), bSweep, OutHit);
 }
 
 
-bool UMovementComponent::SafeMoveUpdatedComponent(const FVector& Delta, const FRotator& NewRotation, bool bSweep, FHitResult& OutHit)
+bool UMovementComponent::SafeMoveUpdatedComponent(const FVector& Delta, const FQuat& NewRotation, bool bSweep, FHitResult& OutHit)
 {
 	if (UpdatedComponent == NULL)
 	{
@@ -498,7 +498,7 @@ FVector UMovementComponent::GetPenetrationAdjustment(const FHitResult& Hit) cons
 	return ConstrainDirectionToPlane(Result);
 }
 
-bool UMovementComponent::ResolvePenetration(const FVector& ProposedAdjustment, const FHitResult& Hit, const FRotator& NewRotation)
+bool UMovementComponent::ResolvePenetrationImpl(const FVector& ProposedAdjustment, const FHitResult& Hit, const FQuat& NewRotationQuat)
 {
 	// SceneComponent can't be in penetration, so this function really only applies to PrimitiveComponent.
 	const FVector Adjustment = ConstrainDirectionToPlane(ProposedAdjustment);
@@ -514,11 +514,11 @@ bool UMovementComponent::ResolvePenetration(const FVector& ProposedAdjustment, c
 		// We really want to make sure that precision differences or differences between the overlap test and sweep tests don't put us into another overlap,
 		// so make the overlap test a bit more restrictive.
 		const float OverlapInflation = CVarPenetrationOverlapCheckInflation.GetValueOnGameThread();
-		bool bEncroached = OverlapTest(Hit.TraceStart + Adjustment, NewRotation.Quaternion(), UpdatedPrimitive->GetCollisionObjectType(), UpdatedPrimitive->GetCollisionShape(OverlapInflation), ActorOwner);
+		bool bEncroached = OverlapTest(Hit.TraceStart + Adjustment, NewRotationQuat, UpdatedPrimitive->GetCollisionObjectType(), UpdatedPrimitive->GetCollisionShape(OverlapInflation), ActorOwner);
 		if (!bEncroached)
 		{
 			// Move without sweeping.
-			MoveUpdatedComponent(Adjustment, NewRotation, false);
+			MoveUpdatedComponent(Adjustment, NewRotationQuat, false);
 			UE_LOG(LogMovement, Verbose, TEXT("ResolvePenetration: teleport %s by %s"), *ActorOwner->GetName(), *Adjustment.ToString());
 			return true;
 		}
@@ -529,7 +529,7 @@ bool UMovementComponent::ResolvePenetration(const FVector& ProposedAdjustment, c
 
 			// Try sweeping as far as possible...
 			FHitResult SweepOutHit(1.f);
-			bool bMoved = MoveUpdatedComponent(Adjustment, NewRotation, true, &SweepOutHit);
+			bool bMoved = MoveUpdatedComponent(Adjustment, NewRotationQuat, true, &SweepOutHit);
 			UE_LOG(LogMovement, Verbose, TEXT("ResolvePenetration: sweep %s by %s (success = %d)"), *ActorOwner->GetName(), *Adjustment.ToString(), bMoved);
 			
 			// Still stuck?
@@ -540,7 +540,7 @@ bool UMovementComponent::ResolvePenetration(const FVector& ProposedAdjustment, c
 				const FVector CombinedMTD = Adjustment + SecondMTD;
 				if (SecondMTD != Adjustment && !CombinedMTD.IsZero())
 				{
-					bMoved = MoveUpdatedComponent(CombinedMTD, NewRotation, true);
+					bMoved = MoveUpdatedComponent(CombinedMTD, NewRotationQuat, true);
 					UE_LOG(LogMovement, Verbose, TEXT("ResolvePenetration: sweep %s by %s (MTD combo success = %d)"), *ActorOwner->GetName(), *CombinedMTD.ToString(), bMoved);
 				}
 			}
@@ -552,7 +552,7 @@ bool UMovementComponent::ResolvePenetration(const FVector& ProposedAdjustment, c
 				const FVector MoveDelta = ConstrainDirectionToPlane(Hit.TraceEnd - Hit.TraceStart);
 				if (!MoveDelta.IsZero())
 				{
-					bMoved = MoveUpdatedComponent(Adjustment + MoveDelta, NewRotation, true);
+					bMoved = MoveUpdatedComponent(Adjustment + MoveDelta, NewRotationQuat, true);
 					UE_LOG(LogMovement, Verbose, TEXT("ResolvePenetration: sweep %s by %s (adjusted attempt success = %d)"), *ActorOwner->GetName(), *(Adjustment + MoveDelta).ToString(), bMoved);
 				}
 			}	
diff --git a/Engine/Source/Runtime/Engine/Private/FloatingPawnMovement.cpp b/Engine/Source/Runtime/Engine/Private/FloatingPawnMovement.cpp
--- a/Engine/Source/Runtime/Engine/Private/FloatingPawnMovement.cpp
+++ b/Engine/Source/Runtime/Engine/Private/FloatingPawnMovement.cpp
@@ -140,8 +140,8 @@ void UFloatingPawnMovement::ApplyControlInputToVelocity(float DeltaTime)
 	ConsumeInputVector();
 }
 
-bool UFloatingPawnMovement::ResolvePenetration(const FVector& Adjustment, const FHitResult& Hit, const FRotator& NewRotation)
+bool UFloatingPawnMovement::ResolvePenetrationImpl(const FVector& Adjustment, const FHitResult& Hit, const FQuat& NewRotationQuat)
 {
-	bPositionCorrected |= Super::ResolvePenetration(Adjustment, Hit, NewRotation);
+	bPositionCorrected |= Super::ResolvePenetrationImpl(Adjustment, Hit, NewRotationQuat);
 	return bPositionCorrected;
 }
