Engine: Fixed a typo in ActorLabel regarding how to call SetActorLabelUnique

diff --git a/Engine/Source/Runtime/Engine/Classes/GameFramework/Actor.h b/Engine/Source/Runtime/Engine/Classes/GameFramework/Actor.h
--- a/Engine/Source/Runtime/Engine/Classes/GameFramework/Actor.h
+++ b/Engine/Source/Runtime/Engine/Classes/GameFramework/Actor.h
@@ -398,7 +398,7 @@ protected:
 private:
 	/**
 	 * The friendly name for this actor, displayed in the editor.  You should always use AActor::GetActorLabel() to access the actual label to display,
-	 * and call AActor::SetActorLabel() or AActor::SetActorLabelUnique() to change the label.  Never set the label directly.
+	 * and call AActor::SetActorLabel() or FActorLabelUtilities::SetActorLabelUnique() to change the label.  Never set the label directly.
 	 */
 	UPROPERTY()
 	FString ActorLabel;
