#ue4 - Perf: Optimize MovementComponents to reduce Quat->Rotator conversions. Reduce Actor calls for rotation/location and use Component instead (saves branch).

diff --git a/Engine/Source/Runtime/Engine/Private/Components/CharacterMovementComponent.cpp b/Engine/Source/Runtime/Engine/Private/Components/CharacterMovementComponent.cpp
--- a/Engine/Source/Runtime/Engine/Private/Components/CharacterMovementComponent.cpp
+++ b/Engine/Source/Runtime/Engine/Private/Components/CharacterMovementComponent.cpp
@@ -1025,17 +1025,18 @@ void UCharacterMovementComponent::SimulatedTick(float DeltaSeconds)
 
 		if( RootMotionParams.bHasRootMotion )
 		{
-			const FRotator OldRotation = CharacterOwner->GetActorRotation();
-			const FVector OldLocation = CharacterOwner->GetActorLocation();
+			const FQuat OldRotationQuat = UpdatedComponent->GetComponentQuat();
+			const FVector OldLocation = UpdatedComponent->GetComponentLocation();
 			SimulateRootMotion(DeltaSeconds, RootMotionParams.RootMotionTransform);
 			// Root Motion has been used, clear
 			RootMotionParams.Clear();
 
 			// debug
-			if( false )
+			if (false)
 			{
-				const FRotator NewRotation = CharacterOwner->GetActorRotation();
-				const FVector NewLocation = CharacterOwner->GetActorLocation();
+				const FRotator OldRotation = OldRotationQuat.Rotator();
+				const FRotator NewRotation = UpdatedComponent->GetComponentRotation();
+				const FVector NewLocation = UpdatedComponent->GetComponentLocation();
 				DrawDebugCoordinateSystem(GetWorld(), CharacterOwner->GetMesh()->GetComponentLocation() + FVector(0,0,1), NewRotation, 50.f, false);
 				DrawDebugLine(GetWorld(), OldLocation, NewLocation, FColor::Red, true, 10.f);
 
@@ -1109,11 +1110,11 @@ void UCharacterMovementComponent::SimulateRootMotion(float DeltaSeconds, const F
 		// 		SimulateMovement(DeltaSeconds);
 
 		// Apply Root Motion rotation after movement is complete.
-		const FRotator RootMotionRotation = WorldSpaceRootMotionTransform.GetRotation().Rotator();
-		if( !RootMotionRotation.IsNearlyZero() )
+		const FQuat RootMotionRotationQuat = WorldSpaceRootMotionTransform.GetRotation();
+		if( !RootMotionRotationQuat.Equals(FQuat::Identity) )
 		{
-			const FRotator NewActorRotation = (CharacterOwner->GetActorRotation() + RootMotionRotation).GetNormalized();
-			MoveUpdatedComponent(FVector::ZeroVector, NewActorRotation, true);
+			const FQuat NewActorRotationQuat = RootMotionRotationQuat * UpdatedComponent->GetComponentQuat();
+			MoveUpdatedComponent(FVector::ZeroVector, NewActorRotationQuat, true);
 		}
 	}
 }
@@ -1364,20 +1365,20 @@ void UCharacterMovementComponent::UpdateBasedMovement(float DeltaSeconds)
 		{
 			FRotationTranslationMatrix HardRelMatrix(CharacterOwner->GetBasedMovement().Rotation, CharacterOwner->GetBasedMovement().Location);
 			const FMatrix NewWorldTM = HardRelMatrix * NewLocalToWorld;
-			const FRotator NewWorldRot = bIgnoreBaseRotation ? CharacterOwner->GetActorRotation() : NewWorldTM.Rotator();
-			MoveUpdatedComponent( NewWorldTM.GetOrigin() - CharacterOwner->GetActorLocation(), NewWorldRot, true );
+			const FQuat NewWorldRot = bIgnoreBaseRotation ? UpdatedComponent->GetComponentQuat() : NewWorldTM.ToQuat();
+			MoveUpdatedComponent( NewWorldTM.GetOrigin() - UpdatedComponent->GetComponentLocation(), NewWorldRot, true );
 		}
 		else
 		{
-			FQuat FinalQuat = CharacterOwner->GetActorQuat();
+			FQuat FinalQuat = UpdatedComponent->GetComponentQuat();
 			
 			if (bRotationChanged && !bIgnoreBaseRotation)
 			{
 				// Apply change in rotation and pipe through FaceRotation to maintain axis restrictions
-				const FQuat PawnOldQuat = CharacterOwner->GetActorQuat();
+				const FQuat PawnOldQuat = UpdatedComponent->GetComponentQuat();
 				FinalQuat = DeltaQuat * FinalQuat;
 				CharacterOwner->FaceRotation(FinalQuat.Rotator(), 0.f);
-				FinalQuat = CharacterOwner->GetActorQuat();
+				FinalQuat = UpdatedComponent->GetComponentQuat();
 
 				// Pipe through ControlRotation, to affect camera.
 				if (CharacterOwner->Controller)
@@ -1394,21 +1395,21 @@ void UCharacterMovementComponent::UpdateBasedMovement(float DeltaSeconds)
 			CharacterOwner->GetCapsuleComponent()->GetScaledCapsuleSize(Radius, HalfHeight);
 
 			FVector const BaseOffset(0.0f, 0.0f, HalfHeight);
-			FVector const LocalBasePos = OldLocalToWorld.InverseTransformPosition(CharacterOwner->GetActorLocation() - BaseOffset);
+			FVector const LocalBasePos = OldLocalToWorld.InverseTransformPosition(UpdatedComponent->GetComponentLocation() - BaseOffset);
 			FVector const NewWorldPos = ConstrainLocationToPlane(NewLocalToWorld.TransformPosition(LocalBasePos) + BaseOffset);
-			DeltaPosition = ConstrainDirectionToPlane(NewWorldPos - CharacterOwner->GetActorLocation());
+			DeltaPosition = ConstrainDirectionToPlane(NewWorldPos - UpdatedComponent->GetComponentLocation());
 
 			// move attached actor
 			if (bFastAttachedMove)
 			{
 				// we're trusting no other obstacle can prevent the move here
-				UpdatedComponent->SetWorldLocationAndRotation(NewWorldPos, FinalQuat.Rotator(), false);
+				UpdatedComponent->SetWorldLocationAndRotation(NewWorldPos, FinalQuat, false);
 			}
 			else
 			{
 				FHitResult MoveOnBaseHit(1.f);
 				const FVector OldLocation = UpdatedComponent->GetComponentLocation();
-				MoveUpdatedComponent(DeltaPosition, FinalQuat.Rotator(), true, &MoveOnBaseHit);
+				MoveUpdatedComponent(DeltaPosition, FinalQuat, true, &MoveOnBaseHit);
 				if ((UpdatedComponent->GetComponentLocation() - (OldLocation + DeltaPosition)).IsNearlyZero() == false)
 				{
 					OnUnableToFollowBaseMove(DeltaPosition, OldLocation, MoveOnBaseHit);
@@ -1445,7 +1446,7 @@ void UCharacterMovementComponent::UpdateBasedRotation(FRotator &FinalRotation, c
 	FinalRotation.Roll = 0.f;
 	if( Controller )
 	{
-		FinalRotation.Roll = CharacterOwner->GetActorRotation().Roll;
+		FinalRotation.Roll = UpdatedComponent->GetComponentRotation().Roll;
 		FRotator NewRotation = Controller->GetControlRotation();
 		NewRotation.Roll = ControllerRoll;
 		Controller->SetControlRotation(NewRotation);
@@ -1494,7 +1495,7 @@ void UCharacterMovementComponent::PerformMovement(float DeltaSeconds)
 		MaybeUpdateBasedMovement(DeltaSeconds);
 
 		OldVelocity = Velocity;
-		OldLocation = CharacterOwner->GetActorLocation();
+		OldLocation = UpdatedComponent->GetComponentLocation();
 
 		ApplyAccumulatedForces(DeltaSeconds);
 
@@ -1544,7 +1545,7 @@ void UCharacterMovementComponent::PerformMovement(float DeltaSeconds)
 				// Convert Local Space Root Motion to world space. Do it right before used by physics to make sure we use up to date transforms, as translation is relative to rotation.
 				RootMotionParams.Set( SkelMeshComp->ConvertLocalRootMotionToWorld(RootMotionParams.RootMotionTransform) );
 				UE_LOG(LogRootMotion, Log,  TEXT("PerformMovement WorldSpaceRootMotion Translation: %s, Rotation: %s, Actor Facing: %s"),
-					*RootMotionParams.RootMotionTransform.GetTranslation().ToCompactString(), *RootMotionParams.RootMotionTransform.GetRotation().Rotator().ToCompactString(), *CharacterOwner->GetActorForwardVector().ToCompactString());
+					*RootMotionParams.RootMotionTransform.GetTranslation().ToCompactString(), *RootMotionParams.RootMotionTransform.GetRotation().Rotator().ToCompactString(), *UpdatedComponent->GetForwardVector().ToCompactString());
 			}
 
 			// Then turn root motion to velocity to be used by various physics modes.
@@ -1584,19 +1585,20 @@ void UCharacterMovementComponent::PerformMovement(float DeltaSeconds)
 		// Apply Root Motion rotation after movement is complete.
 		if( HasRootMotion() )
 		{
-			const FRotator OldActorRotation = CharacterOwner->GetActorRotation();
-			const FRotator RootMotionRotation = RootMotionParams.RootMotionTransform.GetRotation().Rotator();
-			if( !RootMotionRotation.IsNearlyZero() )
+			const FQuat OldActorRotationQuat = UpdatedComponent->GetComponentQuat();
+			const FQuat RootMotionRotationQuat = RootMotionParams.RootMotionTransform.GetRotation();
+			if( !RootMotionRotationQuat.Equals(FQuat::Identity) )
 			{
-				const FRotator NewActorRotation = (OldActorRotation + RootMotionRotation).GetNormalized();
-				MoveUpdatedComponent(FVector::ZeroVector, NewActorRotation, true);
+				const FQuat NewActorRotationQuat = RootMotionRotationQuat * OldActorRotationQuat;
+				MoveUpdatedComponent(FVector::ZeroVector, NewActorRotationQuat, true);
 			}
 
 			// debug
-			if( false )
+			if (false)
 			{
-				const FVector ResultingLocation = CharacterOwner->GetActorLocation();
-				const FRotator ResultingRotation = CharacterOwner->GetActorRotation();
+				const FRotator OldActorRotation = OldActorRotationQuat.Rotator();
+				const FVector ResultingLocation = UpdatedComponent->GetComponentLocation();
+				const FRotator ResultingRotation = UpdatedComponent->GetComponentRotation();
 
 				// Show current position
 				DrawDebugCoordinateSystem(GetWorld(), CharacterOwner->GetMesh()->GetComponentLocation() + FVector(0,0,1), ResultingRotation, 50.f, false);
@@ -1669,19 +1671,19 @@ void UCharacterMovementComponent::SaveBaseLocation()
 		MovementBaseUtility::GetMovementBaseTransform(MovementBase, CharacterOwner->GetBasedMovement().BoneName, OldBaseLocation, OldBaseQuat);
 
 		// Location
-		const FVector RelativeLocation = CharacterOwner->GetActorLocation() - OldBaseLocation;
+		const FVector RelativeLocation = UpdatedComponent->GetComponentLocation() - OldBaseLocation;
 
 		// Rotation
 		if (bIgnoreBaseRotation)
 		{
 			// Absolute rotation
-			CharacterOwner->SaveRelativeBasedMovement(RelativeLocation, CharacterOwner->GetActorRotation(), false);
+			CharacterOwner->SaveRelativeBasedMovement(RelativeLocation, UpdatedComponent->GetComponentRotation(), false);
 		}
 		else
 		{
 			// Relative rotation
-			const FRotator RelativeRotation = (FQuatRotationMatrix(CharacterOwner->GetActorQuat()) * FQuatRotationMatrix(OldBaseQuat).GetTransposed()).Rotator();
-			CharacterOwner->SaveRelativeBasedMovement(RelativeLocation, RelativeRotation, true);	
+			const FRotator RelativeRotation = (FQuatRotationMatrix(UpdatedComponent->GetComponentQuat()) * FQuatRotationMatrix(OldBaseQuat).GetTransposed()).Rotator();
+			CharacterOwner->SaveRelativeBasedMovement(RelativeLocation, RelativeRotation, true);
 		}
 	}
 }
@@ -1745,7 +1747,7 @@ void UCharacterMovementComponent::Crouch(bool bClientSimulation)
 			FCollisionQueryParams CapsuleParams(NAME_CrouchTrace, false, CharacterOwner);
 			FCollisionResponseParams ResponseParam;
 			InitCollisionParams(CapsuleParams, ResponseParam);
-			const bool bEncroached = GetWorld()->OverlapBlockingTestByChannel(CharacterOwner->GetActorLocation() - FVector(0.f,0.f,ScaledHalfHeightAdjust), FQuat::Identity,
+			const bool bEncroached = GetWorld()->OverlapBlockingTestByChannel(UpdatedComponent->GetComponentLocation() - FVector(0.f,0.f,ScaledHalfHeightAdjust), FQuat::Identity,
 				UpdatedComponent->GetCollisionObjectType(), GetPawnCapsuleCollisionShape(SHRINK_None), CapsuleParams, ResponseParam);
 
 			// If encroached, cancel
@@ -1759,7 +1761,7 @@ void UCharacterMovementComponent::Crouch(bool bClientSimulation)
 		if (bCrouchMaintainsBaseLocation)
 		{
 			// Intentionally not using MoveUpdatedComponent, where a horizontal plane constraint would prevent the base of the capsule from staying at the same spot.
-			UpdatedComponent->MoveComponent(FVector(0.f, 0.f, -ScaledHalfHeightAdjust), CharacterOwner->GetActorRotation(), true);
+			UpdatedComponent->MoveComponent(FVector(0.f, 0.f, -ScaledHalfHeightAdjust), UpdatedComponent->GetComponentQuat(), true);
 		}
 
 		CharacterOwner->bIsCrouched = true;
@@ -1803,7 +1805,7 @@ void UCharacterMovementComponent::UnCrouch(bool bClientSimulation)
 	const float OldUnscaledHalfHeight = CharacterOwner->GetCapsuleComponent()->GetUnscaledCapsuleHalfHeight();
 	const float HalfHeightAdjust = DefaultCharacter->GetCapsuleComponent()->GetUnscaledCapsuleHalfHeight() - OldUnscaledHalfHeight;
 	const float ScaledHalfHeightAdjust = HalfHeightAdjust * ComponentScale;
-	const FVector PawnLocation = CharacterOwner->GetActorLocation();
+	const FVector PawnLocation = UpdatedComponent->GetComponentLocation();
 
 	// Grow to uncrouched size.
 	check(CharacterOwner->GetCapsuleComponent());
@@ -1856,7 +1858,7 @@ void UCharacterMovementComponent::UnCrouch(bool bClientSimulation)
 						if (!bEncroached)
 						{
 							// Intentionally not using MoveUpdatedComponent, where a horizontal plane constraint would prevent the base of the capsule from staying at the same spot.
-							UpdatedComponent->MoveComponent(NewLoc - PawnLocation, CharacterOwner->GetActorRotation(), false);
+							UpdatedComponent->MoveComponent(NewLoc - PawnLocation, UpdatedComponent->GetComponentQuat(), false);
 						}
 					}
 				}
@@ -1885,7 +1887,7 @@ void UCharacterMovementComponent::UnCrouch(bool bClientSimulation)
 			if (!bEncroached)
 			{
 				// Commit the change in location.
-				UpdatedComponent->MoveComponent(StandingLocation - PawnLocation, CharacterOwner->GetActorRotation(), false);
+				UpdatedComponent->MoveComponent(StandingLocation - PawnLocation, UpdatedComponent->GetComponentQuat(), false);
 				bForceNextFloorCheck = true;
 			}
 		}
@@ -2167,8 +2169,8 @@ float UCharacterMovementComponent::ImmersionDepth()
 			FHitResult Hit(1.f);
 			if ( VolumeBrushComp )
 			{
-				const FVector TraceStart = CharacterOwner->GetActorLocation() + FVector(0.f,0.f,CollisionHalfHeight);
-				const FVector TraceEnd = CharacterOwner->GetActorLocation() - FVector(0.f,0.f,CollisionHalfHeight);
+				const FVector TraceStart = UpdatedComponent->GetComponentLocation() + FVector(0.f,0.f,CollisionHalfHeight);
+				const FVector TraceEnd = UpdatedComponent->GetComponentLocation() - FVector(0.f,0.f,CollisionHalfHeight);
 
 				const static FName MovementComp_Character_ImmersionDepthName(TEXT("MovementComp_Character_ImmersionDepth"));
 				FCollisionQueryParams NewTraceParams( MovementComp_Character_ImmersionDepthName, true );
@@ -2268,7 +2270,7 @@ void UCharacterMovementComponent::CalcVelocity(float DeltaTime, float Friction,
 		}
 		else 
 		{
-			Acceleration = MaxAccel * (Velocity.SizeSquared() < SMALL_NUMBER ? CharacterOwner->GetActorForwardVector() : GetSafeNormalPrecise(Velocity));
+			Acceleration = MaxAccel * (Velocity.SizeSquared() < SMALL_NUMBER ? UpdatedComponent->GetForwardVector() : GetSafeNormalPrecise(Velocity));
 		}
 
 		AnalogInputModifier = 1.f;
@@ -2732,10 +2734,10 @@ void UCharacterMovementComponent::PhysFlying(float deltaTime, int32 Iterations)
 	Iterations++;
 	bJustTeleported = false;
 
-	FVector OldLocation = CharacterOwner->GetActorLocation();
+	FVector OldLocation = UpdatedComponent->GetComponentLocation();
 	const FVector Adjusted = Velocity * deltaTime;
 	FHitResult Hit(1.f);
-	SafeMoveUpdatedComponent(Adjusted, CharacterOwner->GetActorRotation(), true, Hit);
+	SafeMoveUpdatedComponent(Adjusted, UpdatedComponent->GetComponentQuat(), true, Hit);
 
 	if (Hit.Time < 1.f)
 	{
@@ -2746,11 +2748,11 @@ void UCharacterMovementComponent::PhysFlying(float deltaTime, int32 Iterations)
 		bool bSteppedUp = false;
 		if ((FMath::Abs(Hit.ImpactNormal.Z) < 0.2f) && (UpDown < 0.5f) && (UpDown > -0.2f) && CanStepUp(Hit))
 		{
-			float stepZ = CharacterOwner->GetActorLocation().Z;
+			float stepZ = UpdatedComponent->GetComponentLocation().Z;
 			bSteppedUp = StepUp(GravDir, Adjusted * (1.f - Hit.Time), Hit);
 			if (bSteppedUp)
 			{
-				OldLocation.Z = CharacterOwner->GetActorLocation().Z + (OldLocation.Z - stepZ);
+				OldLocation.Z = UpdatedComponent->GetComponentLocation().Z + (OldLocation.Z - stepZ);
 			}
 		}
 
@@ -2764,7 +2766,7 @@ void UCharacterMovementComponent::PhysFlying(float deltaTime, int32 Iterations)
 
 	if (!bJustTeleported && !HasRootMotion())
 	{
-		Velocity = (CharacterOwner->GetActorLocation() - OldLocation) / deltaTime;
+		Velocity = (UpdatedComponent->GetComponentLocation() - OldLocation) / deltaTime;
 	}
 }
 
@@ -2785,7 +2787,7 @@ void UCharacterMovementComponent::PhysSwimming(float deltaTime, int32 Iterations
 		Velocity.Z = Velocity.Z * Depth;
 	}
 	Iterations++;
-	FVector OldLocation = CharacterOwner->GetActorLocation();
+	FVector OldLocation = UpdatedComponent->GetComponentLocation();
 	bJustTeleported = false;
 	if( !HasRootMotion() )
 	{
@@ -2814,7 +2816,7 @@ void UCharacterMovementComponent::PhysSwimming(float deltaTime, int32 Iterations
 		bool bSteppedUp = false;
 		if( (FMath::Abs(Hit.ImpactNormal.Z) < 0.2f) && (UpDown < 0.5f) && (UpDown > -0.2f) && CanStepUp(Hit))
 		{
-			float stepZ = CharacterOwner->GetActorLocation().Z;
+			float stepZ = UpdatedComponent->GetComponentLocation().Z;
 			const FVector RealVelocity = Velocity;
 			Velocity.Z = 1.f;	// HACK: since will be moving up, in case pawn leaves the water
 			bSteppedUp = StepUp(GravDir, Adjusted * (1.f - Hit.Time), Hit);
@@ -2826,7 +2828,7 @@ void UCharacterMovementComponent::PhysSwimming(float deltaTime, int32 Iterations
 					StartNewPhysics(remainingTime, Iterations);
 					return;
 				}
-				OldLocation.Z = CharacterOwner->GetActorLocation().Z + (OldLocation.Z - stepZ);
+				OldLocation.Z = UpdatedComponent->GetComponentLocation().Z + (OldLocation.Z - stepZ);
 			}
 			Velocity = RealVelocity;
 		}
@@ -2843,7 +2845,7 @@ void UCharacterMovementComponent::PhysSwimming(float deltaTime, int32 Iterations
 	{
 		bool bWaterJump = !GetPhysicsVolume()->bWaterVolume;
 		float velZ = Velocity.Z;
-		Velocity = (CharacterOwner->GetActorLocation() - OldLocation) / (deltaTime - remainingTime);
+		Velocity = (UpdatedComponent->GetComponentLocation() - OldLocation) / (deltaTime - remainingTime);
 		if (bWaterJump)
 		{
 			Velocity.Z = velZ;
@@ -2872,21 +2874,21 @@ void UCharacterMovementComponent::StartSwimming(FVector OldLocation, FVector Old
 
 	if( !HasRootMotion() && !bJustTeleported )
 	{
-		Velocity = (CharacterOwner->GetActorLocation() - OldLocation)/timeTick; //actual average velocity
+		Velocity = (UpdatedComponent->GetComponentLocation() - OldLocation)/timeTick; //actual average velocity
 		Velocity = 2.f*Velocity - OldVelocity; //end velocity has 2* accel of avg
 		Velocity = Velocity.GetClampedToMaxSize(GetPhysicsVolume()->TerminalVelocity);
 	}
-	const FVector End = FindWaterLine(CharacterOwner->GetActorLocation(), OldLocation);
+	const FVector End = FindWaterLine(UpdatedComponent->GetComponentLocation(), OldLocation);
 	float waterTime = 0.f;
-	if (End != CharacterOwner->GetActorLocation())
+	if (End != UpdatedComponent->GetComponentLocation())
 	{	
-		const float ActualDist = (CharacterOwner->GetActorLocation() - OldLocation).Size();
+		const float ActualDist = (UpdatedComponent->GetComponentLocation() - OldLocation).Size();
 		if (ActualDist > KINDA_SMALL_NUMBER)
 		{
-			waterTime = timeTick * (End - CharacterOwner->GetActorLocation()).Size() / ActualDist;
+			waterTime = timeTick * (End - UpdatedComponent->GetComponentLocation()).Size() / ActualDist;
 			remainingTime += waterTime;
 		}
-		MoveUpdatedComponent(End - CharacterOwner->GetActorLocation(), CharacterOwner->GetActorRotation(), true);
+		MoveUpdatedComponent(End - UpdatedComponent->GetComponentLocation(), UpdatedComponent->GetComponentQuat(), true);
 	}
 	if ( !HasRootMotion() && (Velocity.Z > 2.f*SWIMBOBSPEED) && (Velocity.Z < 0.f)) //allow for falling out of water
 	{
@@ -2900,22 +2902,22 @@ void UCharacterMovementComponent::StartSwimming(FVector OldLocation, FVector Old
 
 float UCharacterMovementComponent::Swim(FVector Delta, FHitResult &Hit)
 {
-	FVector Start = CharacterOwner->GetActorLocation();
+	FVector Start = UpdatedComponent->GetComponentLocation();
 	float airTime = 0.f;
-	SafeMoveUpdatedComponent(Delta, CharacterOwner->GetActorRotation(), true, Hit);
+	SafeMoveUpdatedComponent(Delta, UpdatedComponent->GetComponentQuat(), true, Hit);
 
 	if ( !GetPhysicsVolume()->bWaterVolume ) //then left water
 	{
-		const FVector End = FindWaterLine(Start,CharacterOwner->GetActorLocation());
+		const FVector End = FindWaterLine(Start,UpdatedComponent->GetComponentLocation());
 		const float DesiredDist = Delta.Size();
-		if (End != CharacterOwner->GetActorLocation() && DesiredDist > KINDA_SMALL_NUMBER)
+		if (End != UpdatedComponent->GetComponentLocation() && DesiredDist > KINDA_SMALL_NUMBER)
 		{
-			airTime = (End - CharacterOwner->GetActorLocation()).Size() / DesiredDist;
-			if ( ((CharacterOwner->GetActorLocation() - Start) | (End - CharacterOwner->GetActorLocation())) > 0.f )
+			airTime = (End - UpdatedComponent->GetComponentLocation()).Size() / DesiredDist;
+			if ( ((UpdatedComponent->GetComponentLocation() - Start) | (End - UpdatedComponent->GetComponentLocation())) > 0.f )
 			{
 				airTime = 0.f;
 			}
-			SafeMoveUpdatedComponent(End - CharacterOwner->GetActorLocation(), CharacterOwner->GetActorRotation(), true, Hit);
+			SafeMoveUpdatedComponent(End - UpdatedComponent->GetComponentLocation(), UpdatedComponent->GetComponentQuat(), true, Hit);
 		}
 	}
 	return airTime;
@@ -3020,8 +3022,8 @@ void UCharacterMovementComponent::PhysFalling(float deltaTime, int32 Iterations)
 		const float timeTick = GetSimulationTimeStep(remainingTime, Iterations);
 		remainingTime -= timeTick;
 		
-		const FVector OldLocation = CharacterOwner->GetActorLocation();
-		const FRotator PawnRotation = CharacterOwner->GetActorRotation();
+		const FVector OldLocation = UpdatedComponent->GetComponentLocation();
+		const FQuat PawnRotation = UpdatedComponent->GetComponentQuat();
 		bJustTeleported = false;
 
 		FVector OldVelocity = Velocity;
@@ -3225,7 +3227,7 @@ void UCharacterMovementComponent::PhysFalling(float deltaTime, int32 Iterations)
 						else if (GetPerchRadiusThreshold() > 0.f && Hit.Time == 1.f && OldHitImpactNormal.Z >= WalkableFloorZ)
 						{
 							// We might be in a virtual 'ditch' within our perch radius. This is rare.
-							const FVector PawnLocation = CharacterOwner->GetActorLocation();
+							const FVector PawnLocation = UpdatedComponent->GetComponentLocation();
 							const float ZMovedDist = FMath::Abs(PawnLocation.Z - OldLocation.Z);
 							const float MovedDist2DSq = (PawnLocation - OldLocation).SizeSquared2D();
 							if (ZMovedDist <= 0.2f * timeTick && MovedDist2DSq <= 4.f * timeTick)
@@ -3399,7 +3401,7 @@ void UCharacterMovementComponent::StartFalling(int32 Iterations, float remaining
 {
 	// start falling 
 	const float DesiredDist = Delta.Size();
-	const float ActualDist = (CharacterOwner->GetActorLocation() - subLoc).Size2D();
+	const float ActualDist = (UpdatedComponent->GetComponentLocation() - subLoc).Size2D();
 	remainingTime = (DesiredDist < KINDA_SMALL_NUMBER) 
 					? 0.f
 					: remainingTime + timeTick * (1.f - FMath::Min(1.f,ActualDist/DesiredDist));
@@ -3505,7 +3507,7 @@ void UCharacterMovementComponent::MoveAlongFloor(const FVector& InVelocity, floa
 	// Move along the current floor
 	FHitResult Hit(1.f);
 	FVector RampVector = ComputeGroundMovementDelta(Delta, CurrentFloor.HitResult, CurrentFloor.bLineTrace);
-	SafeMoveUpdatedComponent(RampVector, CharacterOwner->GetActorRotation(), true, Hit);
+	SafeMoveUpdatedComponent(RampVector, UpdatedComponent->GetComponentQuat(), true, Hit);
 	float LastMoveTimeSlice = DeltaSeconds;
 	
 	if (Hit.bStartPenetrating)
@@ -3529,7 +3531,7 @@ void UCharacterMovementComponent::MoveAlongFloor(const FVector& InVelocity, floa
 			const float InitialPercentRemaining = 1.f - PercentTimeApplied;
 			RampVector = ComputeGroundMovementDelta(Delta * InitialPercentRemaining, Hit, false);
 			LastMoveTimeSlice = InitialPercentRemaining * LastMoveTimeSlice;
-			SafeMoveUpdatedComponent(RampVector, CharacterOwner->GetActorRotation(), true, Hit);
+			SafeMoveUpdatedComponent(RampVector, UpdatedComponent->GetComponentQuat(), true, Hit);
 
 			const float SecondHitPercent = Hit.Time * InitialPercentRemaining;
 			PercentTimeApplied = FMath::Clamp(PercentTimeApplied + SecondHitPercent, 0.f, 1.f);
@@ -3660,7 +3662,7 @@ void UCharacterMovementComponent::PhysWalking(float deltaTime, int32 Iterations)
 				const float DesiredDist = Delta.Size();
 				if (DesiredDist > KINDA_SMALL_NUMBER)
 				{
-					const float ActualDist = (CharacterOwner->GetActorLocation() - OldLocation).Size2D();
+					const float ActualDist = (UpdatedComponent->GetComponentLocation() - OldLocation).Size2D();
 					remainingTime += timeTick * (1.f - FMath::Min(1.f,ActualDist/DesiredDist));
 				}
 				StartNewPhysics(remainingTime,Iterations);
@@ -3747,7 +3749,7 @@ void UCharacterMovementComponent::PhysWalking(float deltaTime, int32 Iterations)
 				FHitResult Hit(CurrentFloor.HitResult);
 				Hit.TraceEnd = Hit.TraceStart + FVector(0.f, 0.f, MAX_FLOOR_DIST);
 				const FVector RequestedAdjustment = GetPenetrationAdjustment(Hit);
-				ResolvePenetration(RequestedAdjustment, Hit, CharacterOwner->GetActorRotation());
+				ResolvePenetration(RequestedAdjustment, Hit, UpdatedComponent->GetComponentQuat());
 			}
 
 			// check if just entered water
@@ -3776,12 +3778,12 @@ void UCharacterMovementComponent::PhysWalking(float deltaTime, int32 Iterations)
 			// Make velocity reflect actual move
 			if( !bJustTeleported && !HasRootMotion() && timeTick >= MIN_TICK_TIME)
 			{
-				Velocity = (CharacterOwner->GetActorLocation() - OldLocation) / timeTick;
+				Velocity = (UpdatedComponent->GetComponentLocation() - OldLocation) / timeTick;
 			}
 		}
 
 		// If we didn't move at all this iteration then abort (since future iterations will also be stuck).
-		if (CharacterOwner->GetActorLocation() == OldLocation)
+		if (UpdatedComponent->GetComponentLocation() == OldLocation)
 		{
 			remainingTime = 0.f;
 			break;
@@ -3893,7 +3895,7 @@ void UCharacterMovementComponent::PhysNavWalking(float deltaTime, int32 Iteratio
 		{
 			const bool bSweep = UpdatedPrimitive ? UpdatedPrimitive->bGenerateOverlapEvents : false;
 			FHitResult HitResult;
-			SafeMoveUpdatedComponent(AdjustedDelta, CharacterOwner->GetActorRotation(), bSweep, HitResult);
+			SafeMoveUpdatedComponent(AdjustedDelta, UpdatedComponent->GetComponentQuat(), bSweep, HitResult);
 		}
 
 		// Update velocity to reflect actual move
@@ -4054,7 +4056,7 @@ void UCharacterMovementComponent::AdjustFloorHeight()
 		const float InitialZ = UpdatedComponent->GetComponentLocation().Z;
 		const float AvgFloorDist = (MIN_FLOOR_DIST + MAX_FLOOR_DIST) * 0.5f;
 		const float MoveDist = AvgFloorDist - OldFloorDist;
-		SafeMoveUpdatedComponent( FVector(0.f,0.f,MoveDist), CharacterOwner->GetActorRotation(), true, AdjustHit );
+		SafeMoveUpdatedComponent( FVector(0.f,0.f,MoveDist), UpdatedComponent->GetComponentQuat(), true, AdjustHit );
 		UE_LOG(LogCharacterMovement, VeryVerbose, TEXT("Adjust floor height %.3f (Hit = %d)"), MoveDist, AdjustHit.bBlockingHit);
 
 		if (!AdjustHit.IsValidBlockingHit())
@@ -4190,8 +4192,8 @@ bool UCharacterMovementComponent::TryToLeaveNavWalking()
 	bool bCanTeleport = true;
 	if (CharacterOwner)
 	{
-		FVector CollisionFreeLocation = CharacterOwner->GetActorLocation();
-		bCanTeleport = GetWorld()->FindTeleportSpot(CharacterOwner, CollisionFreeLocation, CharacterOwner->GetActorRotation());
+		FVector CollisionFreeLocation = UpdatedComponent->GetComponentLocation();
+		bCanTeleport = GetWorld()->FindTeleportSpot(CharacterOwner, CollisionFreeLocation, UpdatedComponent->GetComponentRotation());
 		if (bCanTeleport)
 		{
 			CharacterOwner->SetActorLocation(CollisionFreeLocation);
@@ -4303,7 +4305,7 @@ void UCharacterMovementComponent::PhysicsRotation(float DeltaTime)
 		return;
 	}
 
-	const FRotator CurrentRotation = CharacterOwner->GetActorRotation();
+	const FRotator CurrentRotation = UpdatedComponent->GetComponentRotation();
 	FRotator DeltaRot = GetDeltaRotation(DeltaTime);
 	FRotator DesiredRotation = CurrentRotation;
 
@@ -4433,7 +4435,7 @@ bool UCharacterMovementComponent::CheckWaterJump(FVector CheckPoint, FVector& Wa
 	FVector CheckNorm = CheckPoint.GetSafeNormal();
 	float PawnCapsuleRadius, PawnCapsuleHalfHeight;
 	CharacterOwner->GetCapsuleComponent()->GetScaledCapsuleSize(PawnCapsuleRadius, PawnCapsuleHalfHeight);
-	CheckPoint = CharacterOwner->GetActorLocation() + 1.2f * PawnCapsuleRadius * CheckNorm;
+	CheckPoint = UpdatedComponent->GetComponentLocation() + 1.2f * PawnCapsuleRadius * CheckNorm;
 	FVector Extent(PawnCapsuleRadius, PawnCapsuleRadius, PawnCapsuleHalfHeight);
 	FHitResult HitInfo(1.f);
 	static const FName CheckWaterJumpName(TEXT("CheckWaterJump"));
@@ -4442,13 +4444,13 @@ bool UCharacterMovementComponent::CheckWaterJump(FVector CheckPoint, FVector& Wa
 	InitCollisionParams(CapsuleParams, ResponseParam);
 	FCollisionShape CapsuleShape = GetPawnCapsuleCollisionShape(SHRINK_None);
 	const ECollisionChannel CollisionChannel = UpdatedComponent->GetCollisionObjectType();
-	bool bHit = GetWorld()->SweepSingleByChannel( HitInfo, CharacterOwner->GetActorLocation(), CheckPoint, FQuat::Identity, CollisionChannel, CapsuleShape, CapsuleParams, ResponseParam);
+	bool bHit = GetWorld()->SweepSingleByChannel( HitInfo, UpdatedComponent->GetComponentLocation(), CheckPoint, FQuat::Identity, CollisionChannel, CapsuleShape, CapsuleParams, ResponseParam);
 	
 	if ( bHit && !Cast<APawn>(HitInfo.GetActor()) )
 	{
 		// hit a wall - check if it is low enough
 		WallNormal = -1.f * HitInfo.ImpactNormal;
-		FVector Start = CharacterOwner->GetActorLocation();
+		FVector Start = UpdatedComponent->GetComponentLocation();
 		Start.Z += MaxOutOfWaterStepHeight;
 		CheckPoint = Start + 3.2f * PawnCapsuleRadius * WallNormal;
 		FCollisionQueryParams LineParams(CheckWaterJumpName, true, CharacterOwner);
@@ -4529,7 +4531,7 @@ void UCharacterMovementComponent::MoveSmooth(const FVector& InVelocity, const fl
 	else
 	{
 		FHitResult Hit(1.f);
-		SafeMoveUpdatedComponent(Delta, CharacterOwner->GetActorRotation(), true, Hit);
+		SafeMoveUpdatedComponent(Delta, UpdatedComponent->GetComponentQuat(), true, Hit);
 	
 		if (Hit.IsValidBlockingHit())
 		{
@@ -5166,7 +5168,7 @@ bool UCharacterMovementComponent::StepUp(const FVector& InGravDir, const FVector
 
 	// step up - treat as vertical wall
 	FHitResult SweepUpHit(1.f);
-	const FRotator PawnRotation = CharacterOwner->GetActorRotation();
+	const FQuat PawnRotation = UpdatedComponent->GetComponentQuat();
 	SafeMoveUpdatedComponent(-GravDir * StepTravelUpHeight, PawnRotation, true, SweepUpHit);
 
 	// step fwd
@@ -5217,7 +5219,7 @@ bool UCharacterMovementComponent::StepUp(const FVector& InGravDir, const FVector
 	}
 	
 	// Step down
-	SafeMoveUpdatedComponent(GravDir * StepTravelDownHeight, CharacterOwner->GetActorRotation(), true, Hit);
+	SafeMoveUpdatedComponent(GravDir * StepTravelDownHeight, UpdatedComponent->GetComponentQuat(), true, Hit);
 
 	// If step down was initially penetrating abort the step up
 	if (Hit.bStartPenetrating)
@@ -5559,16 +5561,16 @@ void UCharacterMovementComponent::SmoothCorrection(const FVector& OldLocation)
 	FNetworkPredictionData_Client_Character* ClientData = GetPredictionData_Client_Character();
 	if (ClientData && ClientData->bSmoothNetUpdates)
 	{
-		float DistSq = (OldLocation - CharacterOwner->GetActorLocation()).SizeSquared();
+		float DistSq = (OldLocation - UpdatedComponent->GetComponentLocation()).SizeSquared();
 		if (DistSq > FMath::Square(ClientData->MaxSmoothNetUpdateDist))
 		{
 			ClientData->MeshTranslationOffset = (DistSq > FMath::Square(ClientData->NoSmoothNetUpdateDist)) 
 				? FVector::ZeroVector 
-				: ClientData->MeshTranslationOffset + ClientData->MaxSmoothNetUpdateDist * (OldLocation - CharacterOwner->GetActorLocation()).GetSafeNormal();
+				: ClientData->MeshTranslationOffset + ClientData->MaxSmoothNetUpdateDist * (OldLocation - UpdatedComponent->GetComponentLocation()).GetSafeNormal();
 		}
 		else
 		{
-			ClientData->MeshTranslationOffset = ClientData->MeshTranslationOffset + OldLocation - CharacterOwner->GetActorLocation();	
+			ClientData->MeshTranslationOffset = ClientData->MeshTranslationOffset + OldLocation - UpdatedComponent->GetComponentLocation();	
 		}
 	}
 }
@@ -5852,7 +5854,7 @@ void UCharacterMovementComponent::ReplicateMoveToServer(float DeltaTime, const F
 		if (!OverlapTest(OldStartLocation, ClientData->PendingMove->StartRotation.Quaternion(), UpdatedComponent->GetCollisionObjectType(), GetPawnCapsuleCollisionShape(SHRINK_None), CharacterOwner))
 		{
 			FScopedMovementUpdate ScopedMovementUpdate(UpdatedComponent, EScopedUpdate::DeferredUpdates);
-			UE_LOG(LogNetPlayerMovement, VeryVerbose, TEXT("CombineMove: add delta %f + %f and revert from %f %f to %f %f"), DeltaTime, ClientData->PendingMove->DeltaTime, CharacterOwner->GetActorLocation().X, CharacterOwner->GetActorLocation().Y, OldStartLocation.X, OldStartLocation.Y);
+			UE_LOG(LogNetPlayerMovement, VeryVerbose, TEXT("CombineMove: add delta %f + %f and revert from %f %f to %f %f"), DeltaTime, ClientData->PendingMove->DeltaTime, UpdatedComponent->GetComponentLocation().X, UpdatedComponent->GetComponentLocation().Y, OldStartLocation.X, OldStartLocation.Y);
 			
 			// to combine move, first revert pawn position to PendingMove start position, before playing combined move on client
 			const bool bNoCollisionCheck = true;
@@ -5935,7 +5937,7 @@ void UCharacterMovementComponent::ReplicateMoveToServer(float DeltaTime, const F
 	ClientData->ClientUpdateTime = GetWorld()->TimeSeconds;
 
 	UE_LOG(LogNetPlayerMovement, Verbose, TEXT("Client ReplicateMove Time %f Acceleration %s Position %s DeltaTime %f"),
-		NewMove->TimeStamp, *NewMove->Acceleration.ToString(), *CharacterOwner->GetActorLocation().ToString(), DeltaTime);
+		NewMove->TimeStamp, *NewMove->Acceleration.ToString(), *UpdatedComponent->GetComponentLocation().ToString(), DeltaTime);
 
 	// Send to the server
 	CallServerMove( NewMove.Get(), OldMove.Get() );
@@ -6178,7 +6180,7 @@ void UCharacterMovementComponent::ServerMove_Implementation(
 	}
 
 	UE_LOG(LogNetPlayerMovement, Verbose, TEXT("ServerMove Time %f Acceleration %s Position %s DeltaTime %f"),
-			TimeStamp, *Accel.ToString(), *CharacterOwner->GetActorLocation().ToString(), DeltaTime);
+			TimeStamp, *Accel.ToString(), *UpdatedComponent->GetComponentLocation().ToString(), DeltaTime);
 
 	ServerMoveHandleClientError(TimeStamp, DeltaTime, Accel, ClientLoc, ClientMovementBase, ClientBaseBoneName, ClientMovementMode);
 }
@@ -6220,8 +6222,8 @@ void UCharacterMovementComponent::ServerMoveHandleClientError(float ClientTimeSt
 		ServerData->PendingAdjustment.NewVel = Velocity;
 		ServerData->PendingAdjustment.NewBase = MovementBase;
 		ServerData->PendingAdjustment.NewBaseBoneName = CharacterOwner->GetBasedMovement().BoneName;
-		ServerData->PendingAdjustment.NewLoc = CharacterOwner->GetActorLocation();
-		ServerData->PendingAdjustment.NewRot = CharacterOwner->GetActorRotation();
+		ServerData->PendingAdjustment.NewLoc = UpdatedComponent->GetComponentLocation();
+		ServerData->PendingAdjustment.NewRot = UpdatedComponent->GetComponentRotation();
 
 		ServerData->PendingAdjustment.bBaseRelativePosition = MovementBaseUtility::UseRelativeLocation(MovementBase);
 		if (ServerData->PendingAdjustment.bBaseRelativePosition)
@@ -6237,11 +6239,11 @@ void UCharacterMovementComponent::ServerMoveHandleClientError(float ClientTimeSt
 #if !(UE_BUILD_SHIPPING || UE_BUILD_TEST)
 		if (CVarNetShowCorrections.GetValueOnGameThread() != 0)
 		{
-			const FVector LocDiff = CharacterOwner->GetActorLocation() - ClientLoc;
+			const FVector LocDiff = UpdatedComponent->GetComponentLocation() - ClientLoc;
 			UE_LOG(LogNetPlayerMovement, Warning, TEXT("******** Client Error at %f is %f Accel %s LocDiff %s ClientLoc %s, ServerLoc: %s, Base: %s, Bone: %s"),
-				ClientTimeStamp, LocDiff.Size(), *Accel.ToString(), *LocDiff.ToString(), *ClientLoc.ToString(), *CharacterOwner->GetActorLocation().ToString(), *GetNameSafe(MovementBase), *ServerData->PendingAdjustment.NewBaseBoneName.ToString());
+				ClientTimeStamp, LocDiff.Size(), *Accel.ToString(), *LocDiff.ToString(), *ClientLoc.ToString(), *UpdatedComponent->GetComponentLocation().ToString(), *GetNameSafe(MovementBase), *ServerData->PendingAdjustment.NewBaseBoneName.ToString());
 			const float DebugLifetime = CVarNetCorrectionLifetime.GetValueOnGameThread();
-			DrawDebugCapsule(GetWorld(), CharacterOwner->GetActorLocation(), CharacterOwner->GetSimpleCollisionHalfHeight(), CharacterOwner->GetSimpleCollisionRadius(), FQuat::Identity, FColor(100, 255, 100), true, DebugLifetime);
+			DrawDebugCapsule(GetWorld(), UpdatedComponent->GetComponentLocation(), CharacterOwner->GetSimpleCollisionHalfHeight(), CharacterOwner->GetSimpleCollisionRadius(), FQuat::Identity, FColor(100, 255, 100), true, DebugLifetime);
 			DrawDebugCapsule(GetWorld(), ClientLoc                    , CharacterOwner->GetSimpleCollisionHalfHeight(), CharacterOwner->GetSimpleCollisionRadius(), FQuat::Identity, FColor(255, 100, 100), true, DebugLifetime);
 		}
 #endif
@@ -6256,7 +6258,7 @@ void UCharacterMovementComponent::ServerMoveHandleClientError(float ClientTimeSt
 	{
 		if (GetDefault<AGameNetworkManager>()->ClientAuthorativePosition)
 		{
-			const FVector LocDiff = CharacterOwner->GetActorLocation() - ClientLoc;
+			const FVector LocDiff = UpdatedComponent->GetComponentLocation() - ClientLoc;
 			if (!LocDiff.IsZero() || ClientMovementMode != PackNetworkMovementMode())
 			{
 				// Just set the position. On subsequent moves we will resolve initially overlapping conditions.
@@ -6285,7 +6287,7 @@ void UCharacterMovementComponent::ServerMoveHandleClientError(float ClientTimeSt
 bool UCharacterMovementComponent::ServerCheckClientError(float ClientTimeStamp, float DeltaTime, const FVector& Accel, const FVector& ClientWorldLocation, const FVector& RelativeClientLocation, UPrimitiveComponent* ClientMovementBase, FName ClientBaseBoneName, uint8 ClientMovementMode)
 {
 	// Check location difference against global setting
-	const FVector LocDiff = CharacterOwner->GetActorLocation() - ClientWorldLocation;
+	const FVector LocDiff = UpdatedComponent->GetComponentLocation() - ClientWorldLocation;
 	if (GetDefault<AGameNetworkManager>()->ExceedsAllowablePositionError(LocDiff))
 	{
 		return true;
@@ -6528,7 +6530,7 @@ void UCharacterMovementComponent::ClientAdjustPosition_Implementation
 		UE_LOG(LogNetPlayerMovement, Warning, TEXT("******** ClientAdjustPosition Time %f velocity %s position %s NewBase: %s NewBone: %s SavedMoves %d"), TimeStamp, *NewVelocity.ToString(), *NewLocation.ToString(), *GetNameSafe(NewBase), *NewBaseBoneName.ToString(), ClientData->SavedMoves.Num());
 		static const auto CVarCorrectionLifetime = IConsoleManager::Get().FindTConsoleVariableDataFloat(TEXT("p.NetCorrectionLifetime"));
 		const float DebugLifetime = CVarCorrectionLifetime ? CVarCorrectionLifetime->GetValueOnGameThread() : 1.f;
-		DrawDebugCapsule(GetWorld(), CharacterOwner->GetActorLocation()	, CharacterOwner->GetSimpleCollisionHalfHeight(), CharacterOwner->GetSimpleCollisionRadius(), FQuat::Identity, FColor(255, 100, 100), true, DebugLifetime);
+		DrawDebugCapsule(GetWorld(), UpdatedComponent->GetComponentLocation()	, CharacterOwner->GetSimpleCollisionHalfHeight(), CharacterOwner->GetSimpleCollisionRadius(), FQuat::Identity, FColor(255, 100, 100), true, DebugLifetime);
 		DrawDebugCapsule(GetWorld(), NewLocation						, CharacterOwner->GetSimpleCollisionHalfHeight(), CharacterOwner->GetSimpleCollisionRadius(), FQuat::Identity, FColor(100, 255, 100), true, DebugLifetime);
 	}
 #endif //!(UE_BUILD_SHIPPING || UE_BUILD_TEST)
diff --git a/Engine/Source/Runtime/Engine/Private/Components/ProjectileMovementComponent.cpp b/Engine/Source/Runtime/Engine/Private/Components/ProjectileMovementComponent.cpp
--- a/Engine/Source/Runtime/Engine/Private/Components/ProjectileMovementComponent.cpp
+++ b/Engine/Source/Runtime/Engine/Private/Components/ProjectileMovementComponent.cpp
@@ -269,7 +269,7 @@ bool UProjectileMovementComponent::HandleSliding(FHitResult& Hit, float& SubTick
 
 	// Velocity is now parallel to the impact surface.
 	// Perform the move now, before adding gravity/accel again, so we don't just keep hitting the surface.
-	SafeMoveUpdatedComponent(Velocity * SubTickTimeRemaining, UpdatedComponent->GetComponentRotation(), true, Hit);
+	SafeMoveUpdatedComponent(Velocity * SubTickTimeRemaining, UpdatedComponent->GetComponentQuat(), true, Hit);
 
 	if (HasStoppedSimulation())
 	{
diff --git a/Engine/Source/Runtime/Engine/Private/Components/RotatingMovementComponent.cpp b/Engine/Source/Runtime/Engine/Private/Components/RotatingMovementComponent.cpp
--- a/Engine/Source/Runtime/Engine/Private/Components/RotatingMovementComponent.cpp
+++ b/Engine/Source/Runtime/Engine/Private/Components/RotatingMovementComponent.cpp
@@ -36,5 +36,5 @@ void URotatingMovementComponent::TickComponent(float DeltaTime, enum ELevelTick
 	}
 
 	const bool bEnableCollision = false;
-	MoveUpdatedComponent(DeltaLocation, NewRotation.Rotator(), bEnableCollision);
+	MoveUpdatedComponent(DeltaLocation, NewRotation, bEnableCollision);
 }
