#ue4 - Perf: Removed Quat->Rotator conversion in UPrimitiveComponent::SyncComponentToRBPhysics().

diff --git a/Engine/Source/Runtime/Engine/Private/PrimitiveComponentPhysics.cpp b/Engine/Source/Runtime/Engine/Private/PrimitiveComponentPhysics.cpp
--- a/Engine/Source/Runtime/Engine/Private/PrimitiveComponentPhysics.cpp
+++ b/Engine/Source/Runtime/Engine/Private/PrimitiveComponentPhysics.cpp
@@ -582,7 +582,7 @@ void UPrimitiveComponent::SyncComponentToRBPhysics()
 	if(!NewTransform.EqualsNoScale(ComponentToWorld))
 	{
 		const FVector MoveBy = NewTransform.GetLocation() - ComponentToWorld.GetLocation();
-		const FRotator NewRotation = NewTransform.Rotator();
+		const FQuat NewRotation = NewTransform.GetRotation();
 
 		//@warning: do not reference BodyInstance again after calling MoveComponent() - events from the move could have made it unusable (destroying the actor, SetPhysics(), etc)
 		MoveComponent(MoveBy, NewRotation, false, NULL, MOVECOMP_SkipPhysicsMove);
