#ue4 - Perf: Removed Quat->Rotator conversion in UWorld::EncroachingBlockingGeometry()

diff --git a/Engine/Source/Runtime/Engine/Private/LevelActor.cpp b/Engine/Source/Runtime/Engine/Private/LevelActor.cpp
--- a/Engine/Source/Runtime/Engine/Private/LevelActor.cpp
+++ b/Engine/Source/Runtime/Engine/Private/LevelActor.cpp
@@ -703,7 +703,7 @@ bool UWorld::EncroachingBlockingGeometry(AActor* TestActor, FVector TestLocation
 			{
 				// must be registered
 				FComponentQueryParams Params(NAME_EncroachingBlockingGeometry, TestActor);
-				bFoundBlockingHit = ComponentOverlapMultiByChannel(Overlaps, PrimComp, TestLocation, TestActor->GetActorRotation(), BlockingChannel, Params);
+				bFoundBlockingHit = ComponentOverlapMultiByChannel(Overlaps, PrimComp, TestLocation, TestActor->GetActorQuat(), BlockingChannel, Params);
 			}
 			else
 			{
