Support for navmesh chunks streaming in editor worlds
UE-13794

diff --git a/Engine/Source/Runtime/Engine/Classes/AI/Navigation/NavigationSystem.h b/Engine/Source/Runtime/Engine/Classes/AI/Navigation/NavigationSystem.h
--- a/Engine/Source/Runtime/Engine/Classes/AI/Navigation/NavigationSystem.h
+++ b/Engine/Source/Runtime/Engine/Classes/AI/Navigation/NavigationSystem.h
@@ -896,6 +896,9 @@ private:
 	 */
 	ERuntimeGenerationType GetRuntimeGenerationType() const;
 
+	/** Discards all navigation data chunks in all sub-levels */
+	static void DiscardNavigationDataChunks(UWorld* InWorld);
+
 	//----------------------------------------------------------------------//
 	// DEPRECATED
 	//----------------------------------------------------------------------//
diff --git a/Engine/Source/Runtime/Engine/Private/AI/Navigation/NavigationSystem.cpp b/Engine/Source/Runtime/Engine/Private/AI/Navigation/NavigationSystem.cpp
--- a/Engine/Source/Runtime/Engine/Private/AI/Navigation/NavigationSystem.cpp
+++ b/Engine/Source/Runtime/Engine/Private/AI/Navigation/NavigationSystem.cpp
@@ -611,21 +611,19 @@ void UNavigationSystem::OnWorldInitDone(FNavigationSystem::EMode Mode)
 
 		// All navigation actors are registered
 		// Add NavMesh parts from all sub-levels that were streamed in prior NavMesh registration
-		if (World->IsGameWorld())
+		const auto& Levels = World->GetLevels();
+		for (ULevel* Level : Levels)
 		{
-			const auto& Levels = World->GetLevels();
-			for (ULevel* Level : Levels)
+			if (!Level->IsPersistentLevel() && Level->bIsVisible)
 			{
-				if (!Level->IsPersistentLevel() && Level->bIsVisible)
+				for (ANavigationData* NavData : NavDataSet)
 				{
-					for (ANavigationData* NavData : NavDataSet)
-					{
-						NavData->OnStreamingLevelAdded(Level, World);
-					}
+					NavData->OnStreamingLevelAdded(Level, World);
 				}
 			}
 		}
-		else if (World->GetLevels().Num() > 1)
+
+		if (!World->IsGameWorld() && World->GetLevels().Num() > 1)
 		{
 			// this is a bit of a @hack for the time being until we reorganize navigation data instances (4.9?)
 			// the point of this hack is to force editor-time navmesh rebuilding if we have any sublevels
@@ -2048,15 +2046,11 @@ void UNavigationSystem::InitializeForWorld(UWorld* World, FNavigationSystem::EMo
 			NavSys = CreateNavigationSystem(World);
 		}
 
-		// Remove old chunk data from all levels
+		// Remove old/stale chunk data from all levels, when navigation auto-update is enabled
 		// In case navigation system will be created chunks will be regenerated anyway
-		if (Mode == FNavigationSystem::EditorMode)
+		if (Mode == FNavigationSystem::EditorMode && bNavigationAutoUpdateEnabled)
 		{
-			const auto& Levels = World->GetLevels();
-			for (ULevel* Level : Levels)
-			{
-				Level->NavDataChunks.Empty();
-			}
+			DiscardNavigationDataChunks(World);
 		}
 
 		if (NavSys)
@@ -2888,6 +2882,8 @@ void UNavigationSystem::GatherNavigationBounds()
 
 void UNavigationSystem::Build()
 {
+	DiscardNavigationDataChunks(GetWorld());
+	
 	if (IsThereAnywhereToBuildNavigation() == false)
 	{
 		return;
@@ -3362,6 +3358,23 @@ ERuntimeGenerationType UNavigationSystem::GetRuntimeGenerationType() const
 	return RuntimeGenerationType;
 }
 
+void UNavigationSystem::DiscardNavigationDataChunks(UWorld* InWorld)
+{
+	check(InWorld);
+	const auto& Levels = InWorld->GetLevels();
+	for (ULevel* Level : Levels)
+	{
+		for (UNavigationDataChunk* NavChunk : Level->NavDataChunks)
+		{
+			if (NavChunk != nullptr)
+			{
+				NavChunk->MarkPendingKill();
+			}
+		}
+		Level->NavDataChunks.Empty();
+	}
+}
+
 //----------------------------------------------------------------------//
 // Blueprint functions
 //----------------------------------------------------------------------//
diff --git a/Engine/Source/Runtime/Engine/Private/AI/Navigation/RecastNavMesh.cpp b/Engine/Source/Runtime/Engine/Private/AI/Navigation/RecastNavMesh.cpp
--- a/Engine/Source/Runtime/Engine/Private/AI/Navigation/RecastNavMesh.cpp
+++ b/Engine/Source/Runtime/Engine/Private/AI/Navigation/RecastNavMesh.cpp
@@ -1612,7 +1612,7 @@ void ARecastNavMesh::OnStreamingLevelAdded(ULevel* InLevel, UWorld* InWorld)
 	
 	bWantsUpdate = true;
 
-	if (SupportsStreaming() && InWorld->IsGameWorld())
+	if (SupportsStreaming())
 	{
 		URecastNavMeshDataChunk* NavDataChunk = GetNavigationDataChunk(InLevel);
 		if (NavDataChunk)
@@ -1631,7 +1631,7 @@ void ARecastNavMesh::OnStreamingLevelRemoved(ULevel* InLevel, UWorld* InWorld)
 {
 	QUICK_SCOPE_CYCLE_COUNTER(STAT_RecastNavMesh_OnStreamingLevelRemoved);
 	
-	if (SupportsStreaming() && InWorld->IsGameWorld())
+	if (SupportsStreaming())
 	{
 		URecastNavMeshDataChunk* NavDataChunk = GetNavigationDataChunk(InLevel);
 		if (NavDataChunk)
diff --git a/Engine/Source/Runtime/Engine/Private/AI/Navigation/RecastNavMeshDataChunk.cpp b/Engine/Source/Runtime/Engine/Private/AI/Navigation/RecastNavMeshDataChunk.cpp
--- a/Engine/Source/Runtime/Engine/Private/AI/Navigation/RecastNavMeshDataChunk.cpp
+++ b/Engine/Source/Runtime/Engine/Private/AI/Navigation/RecastNavMeshDataChunk.cpp
@@ -38,6 +38,18 @@ FRecastTileData::FRecastTileData(NavNodeRef Ref, int32 DataSize, uint8* RawData,
 	TileCacheRawData = MakeShareable(new FRawData(CacheRawData));
 }
 
+// Helper to duplicate recast raw data
+static uint8* DuplicateRecastRawData(uint8* Src, int32 SrcSize)
+{
+#if WITH_RECAST	
+	uint8* DupData = (uint8*)dtAlloc(SrcSize, DT_ALLOC_PERM);
+#else
+	uint8* DupData = (uint8*)FMemory::Malloc(SrcSize);
+#endif
+	FMemory::Memcpy(DupData, Src, SrcSize);
+	return DupData;
+}
+
 //----------------------------------------------------------------------//
 // URecastNavMeshDataChunk                                                                
 //----------------------------------------------------------------------//
@@ -77,20 +89,17 @@ void URecastNavMeshDataChunk::Serialize(FArchive& Ar)
 			Ar.Seek(RecastNavMeshSizePos + RecastNavMeshSizeBytes);
 		}
 	}
-	else
+	else if (Ar.IsSaving())
 	{
 #if WITH_RECAST
 		SerializeRecastData(Ar, NavMeshVersion);
 #endif// WITH_RECAST
 
-		if (Ar.IsSaving())
-		{
-			int64 CurPos = Ar.Tell();
-			RecastNavMeshSizeBytes = CurPos - RecastNavMeshSizePos;
-			Ar.Seek(RecastNavMeshSizePos);
-			Ar << RecastNavMeshSizeBytes;
-			Ar.Seek(CurPos);
-		}
+		int64 CurPos = Ar.Tell();
+		RecastNavMeshSizeBytes = CurPos - RecastNavMeshSizePos;
+		Ar.Seek(RecastNavMeshSizePos);
+		Ar << RecastNavMeshSizeBytes;
+		Ar.Seek(CurPos);
 	}
 }
 
@@ -151,6 +160,8 @@ TArray<uint32> URecastNavMeshDataChunk::AttachTiles(FPImplRecastNavMesh* NavMesh
 	TArray<uint32> Result;
 	Result.Reserve(Tiles.Num());
 
+	const bool bIsGame = NavMeshImpl->NavMeshOwner->GetWorld()->IsGameWorld();
+
 #if WITH_RECAST	
 	dtNavMesh* NavMesh = NavMeshImpl->DetourNavMesh;
 
@@ -166,9 +177,17 @@ TArray<uint32> URecastNavMeshDataChunk::AttachTiles(FPImplRecastNavMesh* NavMesh
 				continue;
 			}
 			
-			// We don't own tile data anymore
-			TileData.TileDataSize = 0;
-			TileData.TileRawData->RawData = nullptr;
+			if (bIsGame)
+			{
+				// We don't own tile data anymore it will be released by recast navmesh 
+				TileData.TileDataSize = 0;
+				TileData.TileRawData->RawData = nullptr;
+			}
+			else
+			{
+				// In the editor we still need to own data, so make a copy of it
+				TileData.TileRawData->RawData = DuplicateRecastRawData(TileData.TileRawData->RawData, TileData.TileDataSize);
+			}
 			
 			// Attach tile cache layer to target nav mesh
 			if (TileData.TileCacheDataSize > 0)
@@ -183,9 +202,17 @@ TArray<uint32> URecastNavMeshDataChunk::AttachTiles(FPImplRecastNavMesh* NavMesh
 				FNavMeshTileData LayerData(TileData.TileCacheRawData->RawData, TileData.TileCacheDataSize, TileLayerIdx, TileBBox);
 				NavMeshImpl->AddTileCacheLayer(TileX, TileY, TileLayerIdx, LayerData);
 
-				// We don't own tile cache data anymore
-				TileData.TileCacheDataSize = 0;
-				TileData.TileCacheRawData->RawData = nullptr;
+				if (bIsGame)
+				{
+					// We don't own tile cache data anymore it will be released by navmesh
+					TileData.TileCacheDataSize = 0;
+					TileData.TileCacheRawData->RawData = nullptr;
+				}
+				else
+				{
+					// In the editor we still need to own data, so make a copy of it
+					TileData.TileCacheRawData->RawData = DuplicateRecastRawData(TileData.TileCacheRawData->RawData, TileData.TileCacheDataSize);
+				}
 			}
 			
 			Result.Add(NavMesh->decodePolyIdTile(TileData.TileRef));
@@ -202,6 +229,8 @@ TArray<uint32> URecastNavMeshDataChunk::DetachTiles(FPImplRecastNavMesh* NavMesh
 	TArray<uint32> Result;
 	Result.Reserve(Tiles.Num());
 
+	const bool bIsGame = NavMeshImpl->NavMeshOwner->GetWorld()->IsGameWorld();
+
 #if WITH_RECAST	
 	dtNavMesh* NavMesh = NavMeshImpl->DetourNavMesh;
 
@@ -213,18 +242,29 @@ TArray<uint32> URecastNavMeshDataChunk::DetachTiles(FPImplRecastNavMesh* NavMesh
 			const dtMeshTile* MeshTile = NavMesh->getTileByRef(TileData.TileRef);
 			if (MeshTile)
 			{
-				FNavMeshTileData TileCacheData = NavMeshImpl->GetTileCacheLayer(MeshTile->header->x, MeshTile->header->y, MeshTile->header->layer);
-				if (TileCacheData.IsValid())
+				if (bIsGame)
 				{
-					TileData.TileCacheDataSize = TileCacheData.DataSize;
-					TileData.TileCacheRawData->RawData = TileCacheData.Release();
+					FNavMeshTileData TileCacheData = NavMeshImpl->GetTileCacheLayer(MeshTile->header->x, MeshTile->header->y, MeshTile->header->layer);
+					if (TileCacheData.IsValid())
+					{
+						TileData.TileCacheDataSize = TileCacheData.DataSize;
+						TileData.TileCacheRawData->RawData = TileCacheData.Release();
+					}
 				}
 				
 				NavMeshImpl->RemoveTileCacheLayer(MeshTile->header->x, MeshTile->header->y, MeshTile->header->layer);
 			}
 						
-			// Remove tile from navmesh and take ownership of tile raw data
-			NavMesh->removeTile(TileData.TileRef, &TileData.TileRawData->RawData, &TileData.TileDataSize);
+			if (bIsGame)
+			{
+				// Remove tile from navmesh and take ownership of tile raw data
+				NavMesh->removeTile(TileData.TileRef, &TileData.TileRawData->RawData, &TileData.TileDataSize);
+			}
+			else
+			{
+				// In the editor we have a copy of tile data so just release tile in navmesh
+				NavMesh->removeTile(TileData.TileRef, nullptr, nullptr);
+			}
 						
 			Result.Add(NavMesh->decodePolyIdTile(TileData.TileRef));
 		}
@@ -258,16 +298,14 @@ void URecastNavMeshDataChunk::GatherTiles(const FPImplRecastNavMesh* NavMeshImpl
 		if (Tile && Tile->header)
 		{
 			// Make our own copy of tile data
-			uint8* RawTileData = (uint8*)dtAlloc(Tile->dataSize, DT_ALLOC_PERM);
-			FMemory::Memcpy(RawTileData, Tile->data, Tile->dataSize);
+			uint8* RawTileData = DuplicateRecastRawData(Tile->data, Tile->dataSize);
 
 			// Make our own copy of tile cache data
 			uint8* RawTileCacheData = nullptr;
 			FNavMeshTileData TileCacheData = NavMeshImpl->GetTileCacheLayer(Tile->header->x, Tile->header->y, Tile->header->layer);
 			if (TileCacheData.IsValid())
 			{
-				RawTileCacheData = (uint8*)dtAlloc(TileCacheData.DataSize, DT_ALLOC_PERM);
-				FMemory::Memcpy(RawTileCacheData, TileCacheData.GetData(), TileCacheData.DataSize);
+				RawTileCacheData = DuplicateRecastRawData(TileCacheData.GetData(), TileCacheData.DataSize);
 			}
 
 			FRecastTileData RecastTileData(NavMesh->getTileRef(Tile), Tile->dataSize, RawTileData, TileCacheData.DataSize, RawTileCacheData);
