#ue4 - Remove unnecessary check() in USceneComponenent::InternalSetWorldLocationAndRotation().

diff --git a/Engine/Source/Runtime/Engine/Private/Components/SceneComponent.cpp b/Engine/Source/Runtime/Engine/Private/Components/SceneComponent.cpp
--- a/Engine/Source/Runtime/Engine/Private/Components/SceneComponent.cpp
+++ b/Engine/Source/Runtime/Engine/Private/Components/SceneComponent.cpp
@@ -1780,7 +1780,6 @@ void USceneComponent::BeginDestroy()
 
 bool USceneComponent::InternalSetWorldLocationAndRotation(FVector NewLocation, const FQuat& RotationQuat, bool bNoPhysics)
 {
-	check(bWorldToComponentUpdated); // This method is for internal use and this assumption should be valid.
 	FQuat NewRotationQuat(RotationQuat);
 
 	// If attached to something, transform into local space
