non-unity fix

diff --git a/Engine/Source/Runtime/Engine/Private/AI/Navigation/NavigationSystem.cpp b/Engine/Source/Runtime/Engine/Private/AI/Navigation/NavigationSystem.cpp
--- a/Engine/Source/Runtime/Engine/Private/AI/Navigation/NavigationSystem.cpp
+++ b/Engine/Source/Runtime/Engine/Private/AI/Navigation/NavigationSystem.cpp
@@ -8,6 +8,7 @@
 #include "AI/Navigation/NavMeshBoundsVolume.h"
 #include "AI/Navigation/NavRelevantComponent.h"
 #include "AI/Navigation/NavigationInvokerComponent.h"
+#include "AI/Navigation/NavigationDataChunk.h"
 
 #if WITH_RECAST
 #include "RecastNavMeshGenerator.h"
