added vertical compensation for navmesh query extent,
should solve failed navmesh projections in areas with significant height difference between navmesh poly and walkable ground
#ue4

diff --git a/Engine/Source/Runtime/Engine/Classes/AI/Navigation/RecastNavMesh.h b/Engine/Source/Runtime/Engine/Classes/AI/Navigation/RecastNavMesh.h
--- a/Engine/Source/Runtime/Engine/Classes/AI/Navigation/RecastNavMesh.h
+++ b/Engine/Source/Runtime/Engine/Classes/AI/Navigation/RecastNavMesh.h
@@ -634,9 +634,13 @@ public:
 	// Recast runtime params
 	//----------------------------------------------------------------------//
 	/** Euclidean distance heuristic scale used while pathfinding */
-	UPROPERTY(EditAnywhere, Category=Pathfinding, config, meta=(ClampMin = "0.1"))
+	UPROPERTY(EditAnywhere, Category = Query, config, meta = (ClampMin = "0.1"))
 	float HeuristicScale;
 
+	/** Value added to each search height to compensate for error between navmesh polys and walkable geometry  */
+	UPROPERTY(EditAnywhere, Category = Query, config, meta = (ClampMin = "0.0"))
+	float VerticalDeviationFromGroundCompensation;
+
 	/** broadcast for navmesh updates */
 	FOnNavMeshUpdate OnNavMeshUpdate;
 
@@ -922,6 +926,10 @@ public:
 	void RebuildTile(const TArray<FIntPoint>& Tiles);
 
 protected:
+
+	/** Returns query extent including adjustments for voxelization error compensation */
+	FVector GetModifiedQueryExtent(const FVector& QueryExtent) const;
+
 	// @todo docuement
 	UPrimitiveComponent* ConstructRenderingComponentImpl();
 
diff --git a/Engine/Source/Runtime/Engine/Private/AI/Navigation/PImplRecastNavMesh.cpp b/Engine/Source/Runtime/Engine/Private/AI/Navigation/PImplRecastNavMesh.cpp
--- a/Engine/Source/Runtime/Engine/Private/AI/Navigation/PImplRecastNavMesh.cpp
+++ b/Engine/Source/Runtime/Engine/Private/AI/Navigation/PImplRecastNavMesh.cpp
@@ -729,7 +729,7 @@ void FPImplRecastNavMesh::Raycast2D(const FVector& StartLoc, const FVector& EndL
 	FRecastSpeciaLinkFilter LinkFilter(UNavigationSystem::GetCurrent(NavMeshOwner->GetWorld()), Owner);
 	INITIALIZE_NAVQUERY(NavQuery, InQueryFilter.GetMaxSearchNodes(), LinkFilter);
 
-	const FVector& NavExtent = NavMeshOwner->GetDefaultQueryExtent();
+	const FVector NavExtent = NavMeshOwner->GetModifiedQueryExtent(NavMeshOwner->GetDefaultQueryExtent());
 	const float Extent[3] = { NavExtent.X, NavExtent.Z, NavExtent.Y };
 
 	const FVector RecastStart = Unreal2RecastPoint(StartLoc);
@@ -936,7 +936,7 @@ bool FPImplRecastNavMesh::InitPathfinding(const FVector& UnrealStart, const FVec
 	FVector& RecastStart, dtPolyRef& StartPoly,
 	FVector& RecastEnd, dtPolyRef& EndPoly) const
 {
-	const FVector& NavExtent = NavMeshOwner->GetDefaultQueryExtent();
+	const FVector NavExtent = NavMeshOwner->GetModifiedQueryExtent(NavMeshOwner->GetDefaultQueryExtent());
 	const float Extent[3] = { NavExtent.X, NavExtent.Z, NavExtent.Y };
 
 	const FVector RecastStartToProject = Unreal2RecastPoint(UnrealStart);
@@ -1340,7 +1340,8 @@ bool FPImplRecastNavMesh::GetRandomPointInRadius(const FVector& Origin, float Ra
 	{
 		// find starting poly
 		// convert start/end pos to Recast coords
-		const float Extent[3] = {Radius, Radius, Radius};
+		const FVector NavExtent = NavMeshOwner->GetModifiedQueryExtent(FVector(Radius, Radius, Radius));
+		const float Extent[3] = { NavExtent.X, NavExtent.Z, NavExtent.Y };
 		float RecastOrigin[3];
 		Unr2RecastVector(Origin, RecastOrigin);
 		NavNodeRef OriginPolyID = INVALID_NAVNODEREF;
@@ -1406,7 +1407,8 @@ bool FPImplRecastNavMesh::ProjectPointToNavMesh(const FVector& Point, FNavLocati
 	{
 		float ClosestPoint[3];
 
-		FVector RcExtent = Unreal2RecastPoint( Extent ).GetAbs();
+		const FVector ModifiedExtent = NavMeshOwner->GetModifiedQueryExtent(Extent);
+		FVector RcExtent = Unreal2RecastPoint(ModifiedExtent).GetAbs();
 	
 		FVector RcPoint = Unreal2RecastPoint( Point );
 		dtPolyRef PolyRef;
@@ -1417,9 +1419,9 @@ bool FPImplRecastNavMesh::ProjectPointToNavMesh(const FVector& Point, FNavLocati
 			// one last step required due to recast's BVTree imprecision
 			const FVector& UnrealClosestPoint = Recast2UnrVector(ClosestPoint);			
 			const FVector ClosestPointDelta = UnrealClosestPoint - Point;
-			if (FMath::Abs(ClosestPointDelta.X) <= Extent.X &&
-				FMath::Abs(ClosestPointDelta.Y) <= Extent.Y &&
-				FMath::Abs(ClosestPointDelta.Z) <= Extent.Z)
+			if (FMath::Abs(ClosestPointDelta.X) <= ModifiedExtent.X &&
+				FMath::Abs(ClosestPointDelta.Y) <= ModifiedExtent.Y &&
+				FMath::Abs(ClosestPointDelta.Z) <= ModifiedExtent.Z)
 			{
 				bSuccess = true;
 				Result = FNavLocation(UnrealClosestPoint, PolyRef);
@@ -1448,8 +1450,9 @@ bool FPImplRecastNavMesh::ProjectPointMulti(const FVector& Point, TArray<FNavLoc
 	ensure(QueryFilter);
 	if (QueryFilter)
 	{
+		const FVector ModifiedExtent = NavMeshOwner->GetModifiedQueryExtent(Extent);
 		const FVector AdjustedPoint(Point.X, Point.Y, (MaxZ + MinZ) * 0.5f);
-		const FVector AdjustedExtent(Extent.X, Extent.Y, (MaxZ - MinZ) * 0.5f);
+		const FVector AdjustedExtent(ModifiedExtent.X, ModifiedExtent.Y, (MaxZ - MinZ) * 0.5f);
 
 		const FVector RcPoint = Unreal2RecastPoint( AdjustedPoint );
 		const FVector RcExtent = Unreal2RecastPoint( AdjustedExtent ).GetAbs();
@@ -1507,7 +1510,7 @@ NavNodeRef FPImplRecastNavMesh::FindNearestPoly(FVector const& Loc, FVector cons
 		float RecastLoc[3];
 		Unr2RecastVector(Loc, RecastLoc);
 		float RecastExtent[3];
-		Unr2RecastSizeVector(Extent, RecastExtent);
+		Unr2RecastSizeVector(NavMeshOwner->GetModifiedQueryExtent(Extent), RecastExtent);
 
 		NavNodeRef OutRef;
 		dtStatus Status = NavQuery.findNearestPoly(RecastLoc, RecastExtent, QueryFilter, &OutRef, NULL);
@@ -1543,7 +1546,7 @@ bool FPImplRecastNavMesh::GetPolysWithinPathingDistance(FVector const& StartLoc,
 	}
 
 	// @todo this should be configurable in some kind of FindPathQuery structure
-	const FVector& NavExtent = NavMeshOwner->GetDefaultQueryExtent();
+	const FVector NavExtent = NavMeshOwner->GetModifiedQueryExtent(NavMeshOwner->GetDefaultQueryExtent());
 	const float Extent[3] = { NavExtent.X, NavExtent.Z, NavExtent.Y };
 
 	float RecastStartPos[3];
diff --git a/Engine/Source/Runtime/Engine/Private/AI/Navigation/RecastNavMesh.cpp b/Engine/Source/Runtime/Engine/Private/AI/Navigation/RecastNavMesh.cpp
--- a/Engine/Source/Runtime/Engine/Private/AI/Navigation/RecastNavMesh.cpp
+++ b/Engine/Source/Runtime/Engine/Private/AI/Navigation/RecastNavMesh.cpp
@@ -564,6 +564,11 @@ void ARecastNavMesh::PostInitProperties()
 	}
 }
 
+FVector ARecastNavMesh::GetModifiedQueryExtent(const FVector& QueryExtent) const
+{
+	return FVector(QueryExtent.X, QueryExtent.Y, QueryExtent.Z + FMath::Max(0.0f, VerticalDeviationFromGroundCompensation));
+}
+
 void ARecastNavMesh::OnNavAreaAdded(const UClass* NavAreaClass, int32 AgentIndex)
 {
 	Super::OnNavAreaAdded(NavAreaClass, AgentIndex);
@@ -1044,7 +1049,8 @@ void ARecastNavMesh::BatchProjectPoints(TArray<FNavigationProjectionWork>& Workl
 	ensure(QueryFilter);
 	if (QueryFilter)
 	{
-		FVector RcExtent = Unreal2RecastPoint(Extent).GetAbs();
+		const FVector ModifiedExtent = GetModifiedQueryExtent(Extent);
+		FVector RcExtent = Unreal2RecastPoint(ModifiedExtent).GetAbs();
 		float ClosestPoint[3];
 		dtPolyRef PolyRef;
 
@@ -1057,7 +1063,7 @@ void ARecastNavMesh::BatchProjectPoints(TArray<FNavigationProjectionWork>& Workl
 			if (PolyRef > 0)
 			{
 				const FVector& UnrealClosestPoint = Recast2UnrealPoint(ClosestPoint);
-				if (FVector::DistSquared(UnrealClosestPoint, Workload[Idx].Point) <= Extent.SizeSquared())
+				if (FVector::DistSquared(UnrealClosestPoint, Workload[Idx].Point) <= ModifiedExtent.SizeSquared())
 				{
 					Workload[Idx].OutLocation = FNavLocation(UnrealClosestPoint, PolyRef);
 					Workload[Idx].bResult = true;
@@ -1170,7 +1176,7 @@ float ARecastNavMesh::FindDistanceToWall(const FVector& StartLoc, TSharedPtr<con
 		return 0.f;
 	}
 
-	const FVector& NavExtent = GetDefaultQueryExtent();
+	const FVector NavExtent = GetModifiedQueryExtent(GetDefaultQueryExtent());
 	const float Extent[3] = { NavExtent.X, NavExtent.Z, NavExtent.Y };
 
 	const FVector RecastStart = Unreal2RecastPoint(StartLoc);
@@ -1650,7 +1656,7 @@ bool ARecastNavMesh::AdjustLocationWithFilter(const FVector& StartLoc, FVector&
 {
 	INITIALIZE_NAVQUERY(NavQuery, Filter.GetMaxSearchNodes());
 
-	const FVector& NavExtent = GetDefaultQueryExtent();
+	const FVector NavExtent = GetModifiedQueryExtent(GetDefaultQueryExtent());
 	const float Extent[3] = { NavExtent.X, NavExtent.Z, NavExtent.Y };
 
 	const dtQueryFilter* QueryFilter = ((const FRecastQueryFilter*)(Filter.GetImplementation()))->GetAsDetourQueryFilter();
@@ -1843,7 +1849,7 @@ void ARecastNavMesh::BatchRaycast(TArray<FNavigationRaycastWork>& Workload, TSha
 		return;
 	}
 	
-	const FVector& NavExtent = GetDefaultQueryExtent();
+	const FVector NavExtent = GetModifiedQueryExtent(GetDefaultQueryExtent());
 	const float Extent[3] = { NavExtent.X, NavExtent.Z, NavExtent.Y };
 
 	for (auto& WorkItem : Workload)
