fixed early out condition in navigation tile generator using async data gathering
#ue4, merged from dev

diff --git a/Engine/Source/Runtime/Engine/Private/AI/Navigation/RecastNavMeshGenerator.cpp b/Engine/Source/Runtime/Engine/Private/AI/Navigation/RecastNavMeshGenerator.cpp
--- a/Engine/Source/Runtime/Engine/Private/AI/Navigation/RecastNavMeshGenerator.cpp
+++ b/Engine/Source/Runtime/Engine/Private/AI/Navigation/RecastNavMeshGenerator.cpp
@@ -1636,7 +1636,7 @@ bool FRecastTileGenerator::HasDataToBuild() const
 		|| Modifiers.Num()
 		|| OffmeshLinks.Num()
 		|| RawGeometry.Num()
-		|| InclusionBounds.Num();
+		|| (InclusionBounds.Num() && NavigationRelevantData.Num() > 0);
 }
 
 void FRecastTileGenerator::DoWork()
