added option for custom navlinks to ignore navigation parent and force their relevancy
this will prevent navlink from receiving move updates!

#ue4, merged from dev

diff --git a/Engine/Source/Runtime/Engine/Classes/AI/Navigation/NavRelevantComponent.h b/Engine/Source/Runtime/Engine/Classes/AI/Navigation/NavRelevantComponent.h
--- a/Engine/Source/Runtime/Engine/Classes/AI/Navigation/NavRelevantComponent.h
+++ b/Engine/Source/Runtime/Engine/Classes/AI/Navigation/NavRelevantComponent.h
@@ -30,6 +30,9 @@ class ENGINE_API UNavRelevantComponent : public UActorComponent, public INavRele
 	UFUNCTION(BlueprintCallable, Category="AI|Navigation")
 	void SetNavigationRelevancy(bool bRelevant);
 
+	/** force relevancy and skip attaching navigation data to owner's root entry */
+	void ForceNavigationRelevancy(bool bForce);
+
 	/** force refresh in navigation octree */
 	void RefreshNavigationModifiers();
 	
@@ -40,4 +43,8 @@ protected:
 
 	UPROPERTY()
 	uint32 bNavigationRelevant : 1;
+
+	/** attach navigation data to entry for owner's root component (depends on its relevancy) */
+	UPROPERTY()
+	uint32 bAttachToOwnersRoot : 1;
 };
diff --git a/Engine/Source/Runtime/Engine/Private/AI/Navigation/NavRelevantComponent.cpp b/Engine/Source/Runtime/Engine/Private/AI/Navigation/NavRelevantComponent.cpp
--- a/Engine/Source/Runtime/Engine/Private/AI/Navigation/NavRelevantComponent.cpp
+++ b/Engine/Source/Runtime/Engine/Private/AI/Navigation/NavRelevantComponent.cpp
@@ -6,6 +6,7 @@
 UNavRelevantComponent::UNavRelevantComponent(const FObjectInitializer& ObjectInitializer) : Super(ObjectInitializer)
 {
 	bNavigationRelevant = true;
+	bAttachToOwnersRoot = true;
 }
 
 void UNavRelevantComponent::OnRegister()
@@ -41,8 +42,13 @@ void UNavRelevantComponent::UpdateNavigationBounds()
 
 UObject* UNavRelevantComponent::GetNavigationParent() const
 {
-	AActor* OwnerActor = GetOwner();
-	return OwnerActor && OwnerActor->GetRootComponent() ? static_cast<UObject*>(OwnerActor->GetRootComponent()) : static_cast<UObject*>(OwnerActor);
+	if (bAttachToOwnersRoot)
+	{
+		AActor* OwnerActor = GetOwner();
+		return OwnerActor && OwnerActor->GetRootComponent() ? static_cast<UObject*>(OwnerActor->GetRootComponent()) : static_cast<UObject*>(OwnerActor);
+	}
+
+	return nullptr;
 }
 
 void UNavRelevantComponent::CalcBounds()
@@ -58,6 +64,16 @@ void UNavRelevantComponent::CalcAndCacheBounds() const
 	Bounds = FBox::BuildAABB(MyLocation, FVector(100.0f, 100.0f, 100.0f));
 }
 
+void UNavRelevantComponent::ForceNavigationRelevancy(bool bForce)
+{
+	bAttachToOwnersRoot = !bForce;
+	if (bForce)
+	{
+		bNavigationRelevant = true;
+	}
+
+	RefreshNavigationModifiers();
+}
 
 void UNavRelevantComponent::SetNavigationRelevancy(bool bRelevant)
 {
