Skip saving tile cache data in navigation mesh chunks when runtime generation is disabled

diff --git a/Engine/Source/Runtime/Engine/Private/AI/Navigation/RecastNavMeshDataChunk.cpp b/Engine/Source/Runtime/Engine/Private/AI/Navigation/RecastNavMeshDataChunk.cpp
--- a/Engine/Source/Runtime/Engine/Private/AI/Navigation/RecastNavMeshDataChunk.cpp
+++ b/Engine/Source/Runtime/Engine/Private/AI/Navigation/RecastNavMeshDataChunk.cpp
@@ -298,12 +298,17 @@ void URecastNavMeshDataChunk::GatherTiles(const FPImplRecastNavMesh* NavMeshImpl
 			// Make our own copy of tile data
 			uint8* RawTileData = DuplicateRecastRawData(Tile->data, Tile->dataSize);
 
-			// Make our own copy of tile cache data
+			// We need tile cache data only if navmesh supports any kind of runtime generation
+			FNavMeshTileData TileCacheData;
 			uint8* RawTileCacheData = nullptr;
-			FNavMeshTileData TileCacheData = NavMeshImpl->GetTileCacheLayer(Tile->header->x, Tile->header->y, Tile->header->layer);
-			if (TileCacheData.IsValid())
+			if (NavMeshImpl->NavMeshOwner->SupportsRuntimeGeneration())
 			{
-				RawTileCacheData = DuplicateRecastRawData(TileCacheData.GetData(), TileCacheData.DataSize);
+				FNavMeshTileData TileCacheData = NavMeshImpl->GetTileCacheLayer(Tile->header->x, Tile->header->y, Tile->header->layer);
+				if (TileCacheData.IsValid())
+				{
+					// Make our own copy of tile cache data
+					RawTileCacheData = DuplicateRecastRawData(TileCacheData.GetData(), TileCacheData.DataSize);
+				}
 			}
 
 			FRecastTileData RecastTileData(NavMesh->getTileRef(Tile), Tile->dataSize, RawTileData, TileCacheData.DataSize, RawTileCacheData);
