Improve comment on BP Trace Multi funcs to make it clear that the return value is only true if there is a blocking hit

diff --git a/Engine/Source/Runtime/Engine/Classes/Kismet/KismetSystemLibrary.h b/Engine/Source/Runtime/Engine/Classes/Kismet/KismetSystemLibrary.h
--- a/Engine/Source/Runtime/Engine/Classes/Kismet/KismetSystemLibrary.h
+++ b/Engine/Source/Runtime/Engine/Classes/Kismet/KismetSystemLibrary.h
@@ -670,7 +670,7 @@ class ENGINE_API UKismetSystemLibrary : public UBlueprintFunctionLibrary
 	 * @param TraceChannel	The channel to trace
 	 * @param bTraceComplex	True to test against complex collision, false to test against simplified collision.
 	 * @param OutHit		Properties of the trace hit.
-	 * @return				True if there was a hit, false otherwise.
+	 * @return				True if there was a blocking hit, false otherwise.
 	 */
 	UFUNCTION(BlueprintCallable, Category="Collision", meta=(bIgnoreSelf="true", WorldContext="WorldContextObject", AutoCreateRefTerm="ActorsToIgnore", DisplayName = "MultiLineTraceByChannel", Keywords="raycast"))
 	static bool LineTraceMulti_NEW(UObject* WorldContextObject, const FVector Start, const FVector End, ETraceTypeQuery TraceChannel, bool bTraceComplex, const TArray<AActor*>& ActorsToIgnore, EDrawDebugTrace::Type DrawDebugType, TArray<FHitResult>& OutHits, bool bIgnoreSelf);
@@ -701,7 +701,7 @@ class ENGINE_API UKismetSystemLibrary : public UBlueprintFunctionLibrary
 	 * @param TraceChannel	
 	 * @param bTraceComplex	True to test against complex collision, false to test against simplified collision.
 	 * @param OutHits		A list of hits, sorted along the trace from start to finish.  The blocking hit will be the last hit, if there was one.
-	 * @return				True if there was a hit, false otherwise.
+	 * @return				True if there was a blocking hit, false otherwise.
 	 */
 	UFUNCTION(BlueprintCallable, Category="Collision", meta=(bIgnoreSelf="true", WorldContext="WorldContextObject", AutoCreateRefTerm="ActorsToIgnore", DisplayName = "MultiSphereTraceByChannel", Keywords="sweep"))
 	static bool SphereTraceMulti_NEW(UObject* WorldContextObject, const FVector Start, const FVector End, float Radius, ETraceTypeQuery TraceChannel, bool bTraceComplex, const TArray<AActor*>& ActorsToIgnore, EDrawDebugTrace::Type DrawDebugType, TArray<FHitResult>& OutHits, bool bIgnoreSelf);
@@ -733,7 +733,7 @@ class ENGINE_API UKismetSystemLibrary : public UBlueprintFunctionLibrary
 	* @param TraceChannel
 	* @param bTraceComplex	True to test against complex collision, false to test against simplified collision.
 	* @param OutHits		A list of hits, sorted along the trace from start to finish. The blocking hit will be the last hit, if there was one.
-	* @return				True if there was a hit, false otherwise.
+	* @return				True if there was a blocking hit, false otherwise.
 	*/
 	UFUNCTION(BlueprintCallable, Category = "Collision", meta = (bIgnoreSelf = "true", HidePin = "WorldContextObject", DefaultToSelf = "WorldContextObject", AutoCreateRefTerm = "ActorsToIgnore", DisplayName = "MultiBoxTraceByChannel", Keywords="sweep"))
 	static bool BoxTraceMulti(UObject* WorldContextObject, const FVector Start, const FVector End, FVector HalfSize, const FRotator Orientation, ETraceTypeQuery TraceChannel, bool bTraceComplex, const TArray<AActor*>& ActorsToIgnore, EDrawDebugTrace::Type DrawDebugType, TArray<FHitResult>& OutHits, bool bIgnoreSelf);
@@ -768,7 +768,7 @@ class ENGINE_API UKismetSystemLibrary : public UBlueprintFunctionLibrary
 	 * @param TraceChannel	
 	 * @param bTraceComplex	True to test against complex collision, false to test against simplified collision.
 	 * @param OutHits		A list of hits, sorted along the trace from start to finish.  The blocking hit will be the last hit, if there was one.
-	 * @return				True if there was a hit, false otherwise.
+	 * @return				True if there was a blocking hit, false otherwise.
 	 */
 	UFUNCTION(BlueprintCallable, Category="Collision", meta=(bIgnoreSelf="true", WorldContext="WorldContextObject", AutoCreateRefTerm="ActorsToIgnore", DisplayName = "MultiCapsuleTraceByChannel", Keywords="sweep"))
 	static bool CapsuleTraceMulti_NEW(UObject* WorldContextObject, const FVector Start, const FVector End, float Radius, float HalfHeight, ETraceTypeQuery TraceChannel, bool bTraceComplex, const TArray<AActor*>& ActorsToIgnore, EDrawDebugTrace::Type DrawDebugType, TArray<FHitResult>& OutHits, bool bIgnoreSelf);
