Stats - Self/OtherChildren stats should be visible now on the hud, added some stats

diff --git a/Engine/Source/Runtime/Engine/Private/StatsRender2.cpp b/Engine/Source/Runtime/Engine/Private/StatsRender2.cpp
--- a/Engine/Source/Runtime/Engine/Private/StatsRender2.cpp
+++ b/Engine/Source/Runtime/Engine/Private/StatsRender2.cpp
@@ -291,7 +291,11 @@ static int32 RenderCycle( const FComplexStatMessage& Item, class FCanvas* Canvas
 		}
 	}
 
-	Canvas->DrawShadowedString(X + IndentWidth, Y, *ShortenName(*Item.GetDescription()), GetStatRenderGlobals().StatFont, Color);
+	// @TODO yrx 2015-04-17 Move to the stats thread to avoid expensive computation on the game thread.
+	const FString StatDesc = Item.GetDescription();
+	const FString StatDisplay = StatDesc.Len() == 0 ? Item.GetShortName().GetPlainNameString() : StatDesc;
+
+	Canvas->DrawShadowedString(X + IndentWidth, Y, *ShortenName(*StatDisplay), GetStatRenderGlobals().StatFont, Color);
 
 	int32 CurrX = X + GetStatRenderGlobals().AfterNameColumnOffset;
 	// Now append the call count
@@ -663,6 +667,8 @@ static void RenderGroupedWithHierarchy(const FGameThreadHudData& ViewData, FView
  */
 void RenderStats(FViewport* Viewport, class FCanvas* Canvas, int32 X, int32 Y)
 {
+	DECLARE_SCOPE_CYCLE_COUNTER( TEXT( "RenderStats" ), STAT_LoadMap, STATGROUP_StatSystem );
+
 	FGameThreadHudData* ViewData = FHUDGroupGameThreadRenderer::Get().Latest;
 	if (!ViewData)
 	{
