Stats - Improved DrawStatsHUD stats

diff --git a/Engine/Source/Runtime/Engine/Private/GameViewportClient.cpp b/Engine/Source/Runtime/Engine/Private/GameViewportClient.cpp
--- a/Engine/Source/Runtime/Engine/Private/GameViewportClient.cpp
+++ b/Engine/Source/Runtime/Engine/Private/GameViewportClient.cpp
@@ -60,7 +60,10 @@ int32 UGameViewportClient::NumViewportsShowingCollision = 0;
 /** Those sound stat flags which are enabled on this viewport */
 FViewportClient::ESoundShowFlags::Type UGameViewportClient::SoundShowFlags = FViewportClient::ESoundShowFlags::Disabled;
 
-DEFINE_STAT(STAT_UIDrawingTime);
+/**
+ * UI Stats
+ */
+DECLARE_CYCLE_STAT(TEXT("UI Drawing Time"),STAT_UIDrawingTime,STATGROUP_UI);
 
 static TAutoConsoleVariable<int32> CVarSetBlackBordersEnabled(
 	TEXT("r.BlackBorders"),
diff --git a/Engine/Source/Runtime/Engine/Private/UnrealEngine.cpp b/Engine/Source/Runtime/Engine/Private/UnrealEngine.cpp
--- a/Engine/Source/Runtime/Engine/Private/UnrealEngine.cpp
+++ b/Engine/Source/Runtime/Engine/Private/UnrealEngine.cpp
@@ -109,8 +109,6 @@ IMPLEMENT_MODULE( FEngineModule, Engine );
 
 #define LOCTEXT_NAMESPACE "UnrealEngine"
 
-DECLARE_CYCLE_STAT(TEXT("DrawStats"),STAT_DrawStats,STATGROUP_StatSystem);
-
 void FEngineModule::StartupModule()
 {
 	// Setup delegate callback for ProfilingHelpers to access current map name
@@ -6886,14 +6884,13 @@ bool UEngine::ShouldThrottleCPUUsage() const
  */
 void DrawStatsHUD( UWorld* World, FViewport* Viewport, FCanvas* Canvas, UCanvas* CanvasObject, TArray<FDebugDisplayProperty>& DebugProperties, const FVector& ViewLocation, const FRotator& ViewRotation )
 {
+	DECLARE_SCOPE_CYCLE_COUNTER( TEXT( "DrawStatsHUD" ), STAT_DrawStatsHUD, STATGROUP_StatSystem );
+
 	// We cannot draw without a canvas
 	if (Canvas == NULL)
 	{
 		return;
 	}
-#if STATS
-	uint32 DrawStatsBeginTime = FPlatformTime::Cycles();
-#endif
 
 	//@todo joeg: Move this stuff to a function, make safe to use on consoles by
 	// respecting the various safe zones, and make it compile out.
@@ -7222,11 +7219,6 @@ void DrawStatsHUD( UWorld* World, FViewport* Viewport, FCanvas* Canvas, UCanvas*
 		}
 	}
 #endif
-
-#if STATS
-	uint32 DrawStatsEndTime = FPlatformTime::Cycles();
-	SET_CYCLE_COUNTER(STAT_DrawStats, DrawStatsEndTime - DrawStatsBeginTime);
-#endif
 }
 
 
diff --git a/Engine/Source/Runtime/Engine/Public/EngineStats.h b/Engine/Source/Runtime/Engine/Public/EngineStats.h
--- a/Engine/Source/Runtime/Engine/Public/EngineStats.h
+++ b/Engine/Source/Runtime/Engine/Public/EngineStats.h
@@ -151,11 +151,6 @@ DECLARE_DWORD_COUNTER_STAT_EXTERN(TEXT("Partial path"),STAT_Navigation_PartialPa
 DECLARE_FLOAT_ACCUMULATOR_STAT_EXTERN(TEXT("Navmesh cumulative build Time"),STAT_Navigation_CumulativeBuildTime,STATGROUP_Navigation, );
 DECLARE_CYCLE_STAT_EXTERN(TEXT("Navmesh build time"),STAT_Navigation_BuildTime,STATGROUP_Navigation, );
 
-/**
- * UI Stats
- */
-DECLARE_CYCLE_STAT_EXTERN(TEXT("UI Drawing Time"),STAT_UIDrawingTime,STATGROUP_UI, );
-
 /**
 * Canvas Stats
 */
