Remove accidentally checked in debugging code

diff --git a/Engine/Source/Runtime/Engine/Private/UnrealEngine.cpp b/Engine/Source/Runtime/Engine/Private/UnrealEngine.cpp
--- a/Engine/Source/Runtime/Engine/Private/UnrealEngine.cpp
+++ b/Engine/Source/Runtime/Engine/Private/UnrealEngine.cpp
@@ -8957,14 +8957,6 @@ bool UEngine::LoadMap( FWorldContext& WorldContext, FURL URL, class UPendingNetG
 			}
 		}
 
-		for (TObjectIterator<UObject> It(RF_PendingKill); It; ++It)
-		{
-			if (It->GetTypedOuter<UWorld>())
-			{
-				break;
-			}
-		}
-
 		// Stop all audio to remove references to current level.
 		if (FAudioDevice* AudioDevice = WorldContext.World()->GetAudioDevice())
 		{
