Fix destructible component not setting its bodyinstance scene index. This is needed for proper locking

diff --git a/Engine/Source/Runtime/Engine/Private/Components/DestructibleComponent.cpp b/Engine/Source/Runtime/Engine/Private/Components/DestructibleComponent.cpp
--- a/Engine/Source/Runtime/Engine/Private/Components/DestructibleComponent.cpp
+++ b/Engine/Source/Runtime/Engine/Private/Components/DestructibleComponent.cpp
@@ -346,6 +346,9 @@ void UDestructibleComponent::CreatePhysicsState()
 	const uint32 SceneType = BodyInstance.UseAsyncScene(PhysScene) ? PST_Async : PST_Sync;
 	NxApexScene* ApexScene = PhysScene->GetApexScene(SceneType);
 	PxScene* PScene = PhysScene->GetPhysXScene(SceneType);
+
+	BodyInstance.SceneIndexSync = SceneType == PST_Sync ? PhysScene->PhysXSceneIndex[PST_Sync] : 0;
+	BodyInstance.SceneIndexAsync = SceneType == PST_Async ? PhysScene->PhysXSceneIndex[PST_Async] : 0;
 	check(ApexScene);
 
 	// Create an APEX NxDestructibleActor from the Destructible asset and actor descriptor
