CIS fix for iOS

diff --git a/Engine/Source/Runtime/Engine/Public/Net/RepLayout.h b/Engine/Source/Runtime/Engine/Public/Net/RepLayout.h
--- a/Engine/Source/Runtime/Engine/Public/Net/RepLayout.h
+++ b/Engine/Source/Runtime/Engine/Public/Net/RepLayout.h
@@ -10,6 +10,7 @@
 #include "Engine/EngineTypes.h"
 
 class FOutBunch;
+class FInBunch;
 
 class FRepChangedParent
 {
