Expose PCM contant manifold option. Helps in cases where contact point accuracy is important.

diff --git a/Engine/Source/Runtime/Engine/Classes/PhysicsEngine/PhysicsSettings.h b/Engine/Source/Runtime/Engine/Classes/PhysicsEngine/PhysicsSettings.h
--- a/Engine/Source/Runtime/Engine/Classes/PhysicsEngine/PhysicsSettings.h
+++ b/Engine/Source/Runtime/Engine/Classes/PhysicsEngine/PhysicsSettings.h
@@ -116,6 +116,10 @@ class ENGINE_API UPhysicsSettings
 	UPROPERTY(config, EditAnywhere, AdvancedDisplay, Category=Simulation)
 	bool bEnableAsyncScene;
 
+	/** Enables persistent contact manifolds. This will generate fewer contact points, but with more accuracy. Reduces stability of stacking, but can help energy conservation.*/
+	UPROPERTY(config, EditAnywhere, AdvancedDisplay, Category = Simulation)
+	bool bEnablePCM;
+
 	/** Whether to warn when physics locks are used incorrectly. Turning this off is not recommended and should only be used by very advanced users. */
 	UPROPERTY(config, EditAnywhere, AdvancedDisplay, Category = Simulation)
 	bool bWarnMissingLocks;
diff --git a/Engine/Source/Runtime/Engine/Private/PhysicsEngine/PhysScene.cpp b/Engine/Source/Runtime/Engine/Private/PhysicsEngine/PhysScene.cpp
--- a/Engine/Source/Runtime/Engine/Private/PhysicsEngine/PhysScene.cpp
+++ b/Engine/Source/Runtime/Engine/Private/PhysicsEngine/PhysScene.cpp
@@ -1221,6 +1221,12 @@ void FPhysScene::InitPhysScene(uint32 SceneType)
 	PSceneDesc.filterShader = PhysXSimFilterShader;
 	PSceneDesc.simulationEventCallback = SimEventCallback;
 
+	if(UPhysicsSettings::Get()->bEnablePCM)
+	{
+		PSceneDesc.flags |= PxSceneFlag::eENABLE_PCM;
+	}
+	
+
 	//LOC_MOD enable kinematic vs kinematic for APEX destructibles. This is for the kinematic cube moving horizontally in QA-Destructible map to collide with the destructible.
 	// Was this flag turned off in UE4? Do we want to turn it on for both sync and async scenes?
 	PSceneDesc.flags |= PxSceneFlag::eENABLE_KINEMATIC_PAIRS;
diff --git a/Engine/Source/Runtime/Engine/Private/PhysicsEngine/PhysicsSettings.cpp b/Engine/Source/Runtime/Engine/Private/PhysicsEngine/PhysicsSettings.cpp
--- a/Engine/Source/Runtime/Engine/Private/PhysicsEngine/PhysicsSettings.cpp
+++ b/Engine/Source/Runtime/Engine/Private/PhysicsEngine/PhysicsSettings.cpp
@@ -13,6 +13,7 @@ UPhysicsSettings::UPhysicsSettings(const FObjectInitializer& ObjectInitializer)
 	, RagdollAggregateThreshold(4)
 	, TriangleMeshTriangleMinAreaThreshold(5.0f)
 	, bEnableAsyncScene(false)
+	, bEnablePCM(false)
 	, bWarnMissingLocks(true)
 	, bEnable2DPhysics(false)
 	, LockedAxis_DEPRECATED(ESettingsLockedAxis::Invalid)
