Auto assign an actor's root component when users forget to assign it in the constructor (needs to be set, else native components get attached under SCS created ones).

diff --git a/Engine/Source/Runtime/Engine/Private/Actor.cpp b/Engine/Source/Runtime/Engine/Private/Actor.cpp
--- a/Engine/Source/Runtime/Engine/Private/Actor.cpp
+++ b/Engine/Source/Runtime/Engine/Private/Actor.cpp
@@ -2351,6 +2351,45 @@ void AActor::PostEditImport()
 }
 #endif
 
+/** Util that sets up the actor's component hierarchy (when users forget to do so, in their native ctor) */
+static USceneComponent* FixupNativeActorComponents(AActor* Actor)
+{
+	TArray<USceneComponent*> SceneComponents;
+	Actor->GetComponents(SceneComponents);
+
+	USceneComponent* SceneRootComponent = Actor->GetRootComponent();
+	if ((SceneRootComponent == nullptr) && (SceneComponents.Num() > 0))
+	{
+		UE_LOG(LogActor, Warning, TEXT("%s has natively added scene component(s), but none of them were set as the actor's RootComponent - picking one arbitrarily"), *Actor->GetFullName());
+	}
+
+	// if the user forgot to set one of their native components as the root, 
+	// we arbitrarily pick one for them (otherwise the SCS could attempt to 
+	// create its own root, and nest native components under it)
+	for (USceneComponent* Component : SceneComponents)
+	{
+		if ((Component == nullptr) || (Component->AttachParent != nullptr) || (Component->CreationMethod != EComponentCreationMethod::Native))
+		{
+			continue;
+		}
+
+		// if we've already picked a root component (and this was left 
+		// unattached), then attach this one to the root
+		if (SceneRootComponent != nullptr)
+		{
+			UE_LOG(LogActor, Warning, TEXT("Natively added component (%s) was left unattached from the actor's root."), *SceneRootComponent->GetReadableName());
+			Component->AttachTo(SceneRootComponent);
+		}
+		else
+		{
+			SceneRootComponent = Component;
+			Actor->SetRootComponent(Component);
+		}
+	}
+
+	return SceneRootComponent;
+}
+
 void AActor::PostSpawnInitialize(FVector const& SpawnLocation, FRotator const& SpawnRotation, AActor* InOwner, APawn* InInstigator, bool bRemoteOwned, bool bNoFail, bool bDeferConstruction)
 {
 	// General flow here is like so
@@ -2371,11 +2410,12 @@ void AActor::PostSpawnInitialize(FVector const& SpawnLocation, FRotator const& S
 	// Set network role.
 	check(Role == ROLE_Authority);
 	ExchangeNetRoles(bRemoteOwned);
-
+	
+	USceneComponent* SceneRootComponent = FixupNativeActorComponents(this);
 	// Set the actor's location and rotation.
-	if (GetRootComponent() != NULL)
+	if (SceneRootComponent != NULL)
 	{
-		GetRootComponent()->SetWorldLocationAndRotation(SpawnLocation, SpawnRotation);
+		SceneRootComponent->SetWorldLocationAndRotation(SpawnLocation, SpawnRotation);
 	}
 
 	// Call OnComponentCreated on all default (native) components
