#ue4 - Apply ScaleInputAcceleration() to acceleration in UCharacterMovementComponent::MoveAutonomous().

diff --git a/Engine/Source/Runtime/Engine/Private/Components/CharacterMovementComponent.cpp b/Engine/Source/Runtime/Engine/Private/Components/CharacterMovementComponent.cpp
--- a/Engine/Source/Runtime/Engine/Private/Components/CharacterMovementComponent.cpp
+++ b/Engine/Source/Runtime/Engine/Private/Components/CharacterMovementComponent.cpp
@@ -6337,7 +6337,7 @@ void UCharacterMovementComponent::MoveAutonomous
 	UpdateFromCompressedFlags(CompressedFlags);
 	CharacterOwner->CheckJumpInput(DeltaTime);
 
-	Acceleration = ConstrainInputAcceleration(NewAccel);
+	Acceleration = ScaleInputAcceleration( ConstrainInputAcceleration(NewAccel) );
 	AnalogInputModifier = ComputeAnalogInputModifier();
 	
 	PerformMovement(DeltaTime);
