Added support for shared shapes between sync and async scene of static actors

diff --git a/Engine/Source/Runtime/Engine/Classes/PhysicsEngine/BodyInstance.h b/Engine/Source/Runtime/Engine/Classes/PhysicsEngine/BodyInstance.h
--- a/Engine/Source/Runtime/Engine/Classes/PhysicsEngine/BodyInstance.h
+++ b/Engine/Source/Runtime/Engine/Classes/PhysicsEngine/BodyInstance.h
@@ -643,6 +643,8 @@ public:
 
 	bool UseAsyncScene(const FPhysScene* PhysScene) const;
 
+	bool HasSharedShapes() const{ return bHasSharedShapes; }
+
 	/** Indicates whether this body should use the async scene. Must be called before body is init'd, will assert otherwise. Will have no affect if there is no async scene. */
 	void SetUseAsyncScene(bool bNewUseAsyncScene);
 
@@ -801,11 +803,17 @@ public:
 	DEPRECATED(4.8, "Please call ApplyMaterialToShape_AssumesLocked and make sure you obtain the appropriate PhysX scene locks")
 	static void ApplyMaterialToShape(physx::PxShape* PShape, physx::PxMaterial* PSimpleMat, TArray<UPhysicalMaterial*>& ComplexPhysMats)
 	{
-		ApplyMaterialToShape_AssumesLocked(PShape, PSimpleMat, ComplexPhysMats);
+		ApplyMaterialToShape_AssumesLocked(PShape, PSimpleMat, ComplexPhysMats, false);
 	}
 
-	/** Note: This function is not thread safe. Make sure you obtain the appropriate PhysX scene lock before calling it*/
-	static void ApplyMaterialToShape_AssumesLocked(physx::PxShape* PShape, physx::PxMaterial* PSimpleMat, TArray<UPhysicalMaterial*>& ComplexPhysMats);
+	/** 
+	 *  Apply a material directly to the passed in shape. Note this function is very advanced and requires knowledge of shape sharing as well as threading. Note: assumes the appropriate locks have been obtained
+	 *  @param  PShape					The shape we are applying the material to
+	 *  @param  PSimpleMat				The material to use if a simple shape is provided (or complex materials are empty)
+	 *  @param  ComplexPhysMats			The array of materials to apply if a complex shape is provided
+	 *	@param	bSharedShape			If this is true it means you've already detached the shape from all actors that use it (attached shared shapes are not writable).
+	 */
+	static void ApplyMaterialToShape_AssumesLocked(physx::PxShape* PShape, physx::PxMaterial* PSimpleMat, const TArray<UPhysicalMaterial*>& ComplexPhysMats, const bool bSharedShape);
 
 	DEPRECATED(4.8, "Please call ApplyMaterialToInstanceShapes_AssumesLocked and make sure you obtain the appropriate PhysX scene locks")
 	void ApplyMaterialToInstanceShapes(physx::PxMaterial* PSimpleMat, TArray<UPhysicalMaterial*>& ComplexPhysMats)
@@ -1005,6 +1013,9 @@ public:
 	class b2Body* BodyInstancePtr;
 
 #endif	//WITH_BOX2D
+
+private:
+	bool bHasSharedShapes;
 };
 
 template<>
diff --git a/Engine/Source/Runtime/Engine/Classes/PhysicsEngine/BodySetup.h b/Engine/Source/Runtime/Engine/Classes/PhysicsEngine/BodySetup.h
--- a/Engine/Source/Runtime/Engine/Classes/PhysicsEngine/BodySetup.h
+++ b/Engine/Source/Runtime/Engine/Classes/PhysicsEngine/BodySetup.h
@@ -226,7 +226,7 @@ public:
 	/** 
 	 *   Add the shapes defined by this body setup to the supplied PxRigidBody. 
 	 */
-	ENGINE_API void AddShapesToRigidActor_AssumesLocked(FBodyInstance* OwningInstance, physx::PxRigidActor* PDestActor, EPhysicsSceneType SceneType, FVector& Scale3D, physx::PxMaterial* SimpleMaterial, TArray<UPhysicalMaterial*>& ComplexMaterials, FShapeData& ShapeData, const FTransform& RelativeTM = FTransform::Identity, TArray<physx::PxShape*>* NewShapes = NULL);
+	ENGINE_API void AddShapesToRigidActor_AssumesLocked(FBodyInstance* OwningInstance, physx::PxRigidActor* PDestActor, EPhysicsSceneType SceneType, FVector& Scale3D, physx::PxMaterial* SimpleMaterial, TArray<UPhysicalMaterial*>& ComplexMaterials, FShapeData& ShapeData, const FTransform& RelativeTM = FTransform::Identity, TArray<physx::PxShape*>* NewShapes = NULL, bool bShapeSharing = false);
 #endif // WITH_PHYSX
 
 	friend struct FAddShapesHelper;
diff --git a/Engine/Source/Runtime/Engine/Classes/PhysicsEngine/PhysicsSettings.h b/Engine/Source/Runtime/Engine/Classes/PhysicsEngine/PhysicsSettings.h
--- a/Engine/Source/Runtime/Engine/Classes/PhysicsEngine/PhysicsSettings.h
+++ b/Engine/Source/Runtime/Engine/Classes/PhysicsEngine/PhysicsSettings.h
@@ -116,6 +116,10 @@ class ENGINE_API UPhysicsSettings
 	UPROPERTY(config, EditAnywhere, AdvancedDisplay, Category=Simulation)
 	bool bEnableAsyncScene;
 
+	/** Enables shape sharing between sync and async scene for static rigid actors */
+	UPROPERTY(config, EditAnywhere, AdvancedDisplay, Category = Simulation, meta = (editcondition="bEnableAsyncScene"))
+	bool bEnableShapeSharing;
+
 	/** Enables persistent contact manifolds. This will generate fewer contact points, but with more accuracy. Reduces stability of stacking, but can help energy conservation.*/
 	UPROPERTY(config, EditAnywhere, AdvancedDisplay, Category = Simulation)
 	bool bEnablePCM;
diff --git a/Engine/Source/Runtime/Engine/Private/PhysicsEngine/BodyInstance.cpp b/Engine/Source/Runtime/Engine/Private/PhysicsEngine/BodyInstance.cpp
--- a/Engine/Source/Runtime/Engine/Private/PhysicsEngine/BodyInstance.cpp
+++ b/Engine/Source/Runtime/Engine/Private/PhysicsEngine/BodyInstance.cpp
@@ -332,7 +332,7 @@ int32 FBodyInstance::GetAllShapes_AssumesLocked(TArray<PxShape*>& OutShapes) con
 	}
 
 	// grab shapes from async actor
-	if( RigidActorAsync != NULL )
+	if( RigidActorAsync != NULL && !HasSharedShapes())
 	{
 		const int32 NumAsyncShapes = RigidActorAsync->getNbShapes();
 		OutShapes.AddZeroed(NumAsyncShapes);
@@ -351,22 +351,22 @@ void FBodyInstance::UpdateTriMeshVertices(const TArray<FVector> & NewPositions)
 	{
 		ExecuteOnPhysicsReadWrite([&]
 		{
-		BodySetup->UpdateTriMeshVertices(NewPositions);
+			BodySetup->UpdateTriMeshVertices(NewPositions);
 
-		//after updating the vertices we must call setGeometry again to update any shapes referencing the mesh
+			//after updating the vertices we must call setGeometry again to update any shapes referencing the mesh
 
 			TArray<PxShape *> PShapes;
 			const int32 SyncShapeCount = GetAllShapes_AssumesLocked(PShapes);
-		PxTriangleMeshGeometry PTriangleMeshGeometry;
-		for (int32 ShapeIdx = 0; ShapeIdx < PShapes.Num(); ShapeIdx++)
-		{
-			PxShape* PShape = PShapes[ShapeIdx];
-			if (PShape->getGeometryType() == PxGeometryType::eTRIANGLEMESH)
+			PxTriangleMeshGeometry PTriangleMeshGeometry;
+			for (int32 ShapeIdx = 0; ShapeIdx < PShapes.Num(); ShapeIdx++)
 			{
-				PShape->getTriangleMeshGeometry(PTriangleMeshGeometry);
-				PShape->setGeometry(PTriangleMeshGeometry);
+				PxShape* PShape = PShapes[ShapeIdx];
+				if (PShape->getGeometryType() == PxGeometryType::eTRIANGLEMESH)
+				{
+					PShape->getTriangleMeshGeometry(PTriangleMeshGeometry);
+					PShape->setGeometry(PTriangleMeshGeometry);
+				}
 			}
-		}
 		});
 	}
 #endif
@@ -632,6 +632,54 @@ ECollisionEnabled::Type FBodyInstance::GetCollisionEnabled() const
 }
 
 #if WITH_PHYSX
+PxShape* ClonePhysXShape_AssumesLocked(PxShape* PShape)
+{
+	//NOTE: this code is copied directly from ExtSimpleFactory.cpp because physx does not currently have an API for copying a shape or incrementing a ref count.
+	//This is hard to maintain so once they provide us with it we can get rid of this
+
+	PxU16 PMaterialCount = PShape->getNbMaterials();
+
+	TArray<PxMaterial*, TInlineAllocator<64>> PMaterials;
+
+	PMaterials.AddZeroed(PMaterialCount);
+	PShape->getMaterials(&PMaterials[0], PMaterialCount);
+
+	PxShape* PNewShape = GPhysXSDK->createShape(PShape->getGeometry().any(), &PMaterials[0], PMaterialCount, false, PShape->getFlags());
+	PNewShape->setLocalPose(PShape->getLocalPose());
+	PNewShape->setContactOffset(PShape->getContactOffset());
+	PNewShape->setRestOffset(PShape->getRestOffset());
+	PNewShape->setSimulationFilterData(PShape->getSimulationFilterData());
+	PNewShape->setQueryFilterData(PShape->getQueryFilterData());
+
+	return PNewShape;
+}
+
+
+template<typename Lambda>
+void ExecuteOnPxShapeWrite(FBodyInstance* BodyInstance, PxShape* PShape, Lambda Func)
+{
+	const bool bSharedShapes = BodyInstance->HasSharedShapes();
+	if (bSharedShapes)
+	{
+		//we must now create a new shape because calling detach will lose our ref count (there's no way to increment it)
+		//shape sharing is only done on static actors so this code should never execute outside the editor
+		PxShape* PNewShape = ClonePhysXShape_AssumesLocked(PShape);
+		BodyInstance->RigidActorSync->detachShape(*PShape, false);
+		BodyInstance->RigidActorAsync->detachShape(*PShape, false);
+		PShape = PNewShape;
+	}
+
+	Func(PShape);
+
+	if (bSharedShapes)
+	{
+		BodyInstance->RigidActorSync->attachShape(*PShape);
+		BodyInstance->RigidActorAsync->attachShape(*PShape);
+		PShape->release();	//we must have created a new shape so release our reference to it (held by actors)
+	}
+}
+
+
 void FBodyInstance::UpdatePhysicsShapeFilterData(uint32 SkelMeshCompID, bool bUseComplexAsSimple, bool bUseSimpleAsComplex, bool bPhysicsStatic, TEnumAsByte<ECollisionEnabled::Type> * CollisionEnabledOverride, FCollisionResponseContainer * ResponseOverride, bool * bNotifyOverride)
 {
 	ExecuteOnPhysicsReadWrite([&]
@@ -1069,7 +1117,7 @@ struct FInitBodiesHelper
 	void InitBodies_PhysX() const;
 	bool CreateShapesAndActors_PhysX(TArray<PxActor*>& PSyncActors, TArray<PxActor*>& PAsyncActors, TArray<PxActor*>& PDynamicActors, const bool bCanDefer, bool& bDynamicsUseAsyncScene) const;
 	physx::PxRigidActor* CreateActor_PhysX_AssumesLocked(FBodyInstance* Instance, const PxTransform& PTransform) const;
-	bool CreateShapes_PhysX_AssumesLocked(FBodyInstance* Instance, physx::PxRigidActor* PNewDynamic, TArray<PxShape*>& PShapes, int32& ShapesWritten) const;
+	bool CreateShapes_PhysX_AssumesLocked(FBodyInstance* Instance, physx::PxRigidActor* PNewDynamic) const;
 	void AddActorsToScene_PhysX_AssumesLocked(TArray<PxActor*>& PSyncActors, TArray<PxActor*>& PAsyncActors, TArray<PxActor*>& PDynamicActors) const;
 #endif
 
@@ -1135,7 +1183,7 @@ physx::PxRigidActor* FInitBodiesHelper::CreateActor_PhysX_AssumesLocked(FBodyIns
 	return nullptr;
 }
 
-bool FInitBodiesHelper::CreateShapes_PhysX_AssumesLocked(FBodyInstance* Instance, physx::PxRigidActor* PNewDynamic, TArray<PxShape*>& PShapes, int32& ShapesWritten) const
+bool FInitBodiesHelper::CreateShapes_PhysX_AssumesLocked(FBodyInstance* Instance, physx::PxRigidActor* PNewDynamic) const
 {
 	UPhysicalMaterial* SimplePhysMat = Instance->GetSimplePhysicalMaterial();
 	TArray<UPhysicalMaterial*> ComplexPhysMats = Instance->GetComplexPhysicalMaterials();
@@ -1156,15 +1204,15 @@ bool FInitBodiesHelper::CreateShapes_PhysX_AssumesLocked(FBodyInstance* Instance
 
 	bool bInitFail = false;
 
+	const bool bShapeSharing = Instance->HasSharedShapes(); //If we have a static actor we can reuse the shapes between sync and async scene
+	TArray<PxShape*> PSharedShapes;
 	if (Instance->RigidActorSync)
 	{
-		BodySetup->AddShapesToRigidActor_AssumesLocked(Instance, Instance->RigidActorSync, PST_Sync, Instance->Scale3D, PSimpleMat, ComplexPhysMats, ShapeData);
+		BodySetup->AddShapesToRigidActor_AssumesLocked(Instance, Instance->RigidActorSync, PST_Sync, Instance->Scale3D, PSimpleMat, ComplexPhysMats, ShapeData, FTransform::Identity, bShapeSharing ? &PSharedShapes : nullptr, bShapeSharing);
 		bInitFail |= Instance->RigidActorSync->getNbShapes() == 0;
 		Instance->RigidActorSync->userData = &Instance->PhysxUserData;
 		Instance->RigidActorSync->setName(Instance->CharDebugName.IsValid() ? Instance->CharDebugName->GetData() : nullptr);
 
-		ShapesWritten += Instance->RigidActorSync->getShapes(PShapes.GetData() + ShapesWritten, PShapes.Num() - ShapesWritten);
-
 		check(FPhysxUserData::Get<FBodyInstance>(Instance->RigidActorSync->userData) == Instance && FPhysxUserData::Get<FBodyInstance>(Instance->RigidActorSync->userData)->OwnerComponent != NULL);
 
 	}
@@ -1172,13 +1220,21 @@ bool FInitBodiesHelper::CreateShapes_PhysX_AssumesLocked(FBodyInstance* Instance
 	if (Instance->RigidActorAsync)
 	{
 		check(PAsyncScene);
-		BodySetup->AddShapesToRigidActor_AssumesLocked(Instance, Instance->RigidActorAsync, PST_Async, Instance->Scale3D, PSimpleMat, ComplexPhysMats, ShapeData);
+		if(bShapeSharing)
+		{
+			for (PxShape* PShape : PSharedShapes)
+			{
+				Instance->RigidActorAsync->attachShape(*PShape);
+			}
+		}else
+		{
+			BodySetup->AddShapesToRigidActor_AssumesLocked(Instance, Instance->RigidActorAsync, PST_Async, Instance->Scale3D, PSimpleMat, ComplexPhysMats, ShapeData);
+		}
+		
 		bInitFail |= Instance->RigidActorAsync->getNbShapes() == 0;
 		Instance->RigidActorAsync->userData = &Instance->PhysxUserData;
 		Instance->RigidActorAsync->setName(Instance->CharDebugName.IsValid() ? Instance->CharDebugName->GetData() : nullptr);
 
-		ShapesWritten += Instance->RigidActorAsync->getShapes(PShapes.GetData() + ShapesWritten, PShapes.Num() - ShapesWritten);
-
 		check(FPhysxUserData::Get<FBodyInstance>(Instance->RigidActorAsync->userData) == Instance && FPhysxUserData::Get<FBodyInstance>(Instance->RigidActorAsync->userData)->OwnerComponent != NULL);
 	}
 
@@ -1198,11 +1254,6 @@ bool FInitBodiesHelper::CreateShapesAndActors_PhysX(TArray<PxActor*>& PSyncActor
 
 	// Ensure we have the AggGeom inside the body setup so we can calculate the number of shapes
 	BodySetup->CreatePhysicsMeshes();
-	const int32 NumShapes = BodySetup->AggGeom.GetElementCount() * NumBodies;
-	TArray<PxShape*> PShapes;
-	int32 ShapesWritten = 0;
-	PShapes.AddUninitialized(NumShapes);
-
 	for (int32 BodyIdx = 0; BodyIdx < NumBodies; ++BodyIdx)
 	{
 		FBodyInstance* Instance = Bodies[BodyIdx];
@@ -1214,6 +1265,7 @@ bool FInitBodiesHelper::CreateShapesAndActors_PhysX(TArray<PxActor*>& PSyncActor
 		Instance->BodySetup = BodySetup;
 		Instance->Scale3D = Transform.GetScale3D();
 		Instance->CharDebugName = PhysXName;
+		Instance->bHasSharedShapes = bStatic && PhysScene->HasAsyncScene() && UPhysicsSettings::Get()->bEnableShapeSharing;
 
 		// Handle autowelding here to avoid extra work
 		if (Instance->bAutoWeld)
@@ -1252,7 +1304,7 @@ bool FInitBodiesHelper::CreateShapesAndActors_PhysX(TArray<PxActor*>& PSyncActor
 		Instance->PhysxUserData = FPhysxUserData(Instance);
 
 		physx::PxRigidActor* PNewDynamic = CreateActor_PhysX_AssumesLocked(Instance, U2PTransform(Transform));	
-		const bool bInitFail = CreateShapes_PhysX_AssumesLocked(Instance, PNewDynamic, PShapes, ShapesWritten);
+		const bool bInitFail = CreateShapes_PhysX_AssumesLocked(Instance, PNewDynamic);
 
 		if (bInitFail)
 		{
@@ -2002,229 +2054,231 @@ bool FBodyInstance::UpdateBodyScale(const FVector& InScale3D)
 		TArray<PxShape *> PShapes;
 		GetAllShapes_AssumesLocked(PShapes);
 
-	for (int32 ShapeIdx = 0; ShapeIdx < PShapes.Num(); ++ShapeIdx)
-	{
-		PxShape* PShape = PShapes[ShapeIdx];
-		PxGeometryType::Enum GeomType = PShape->getGeometryType();
-
-		if (GeomType == PxGeometryType::eSPHERE)
-		{
-			ScaleMode = EScaleMode::LockedXYZ;	//sphere is most restrictive so we can stop
-			break;
-		}
-		else if (GeomType == PxGeometryType::eCAPSULE)
+		for (int32 ShapeIdx = 0; ShapeIdx < PShapes.Num(); ++ShapeIdx)
 		{
-			ScaleMode = EScaleMode::LockedXY;
+			PxShape* PShape = PShapes[ShapeIdx];
+			PxGeometryType::Enum GeomType = PShape->getGeometryType();
+
+			if (GeomType == PxGeometryType::eSPHERE)
+			{
+				ScaleMode = EScaleMode::LockedXYZ;	//sphere is most restrictive so we can stop
+				break;
+			}
+			else if (GeomType == PxGeometryType::eCAPSULE)
+			{
+				ScaleMode = EScaleMode::LockedXY;
+			}
 		}
-	}
 #endif
 #if WITH_BOX2D
-	if (BodyInstancePtr)
-	{
-		//@TODO: BOX2D: UpdateBodyScale is not implemented yet
-	}
+		if (BodyInstancePtr)
+		{
+			//@TODO: BOX2D: UpdateBodyScale is not implemented yet
+		}
 #endif
 
-	FVector RelativeScale3D;
-	FVector RelativeScale3DAbs;
-	ComputeScalingVectors(ScaleMode, InScale3DAdjusted, OldScale3D, RelativeScale3D, RelativeScale3DAbs, UpdatedScale3D);
+		FVector RelativeScale3D;
+		FVector RelativeScale3DAbs;
+		ComputeScalingVectors(ScaleMode, InScale3DAdjusted, OldScale3D, RelativeScale3D, RelativeScale3DAbs, UpdatedScale3D);
 
-	// Apply scaling
+		// Apply scaling
 #if WITH_PHYSX
-	//we need to allocate all of these here because PhysX insists on using the stack. This is wasteful, but reduces a lot of code duplication
-	PxSphereGeometry PSphereGeom;
-	PxBoxGeometry PBoxGeom;
-	PxCapsuleGeometry PCapsuleGeom;
-	PxConvexMeshGeometry PConvexGeom;
-	PxTriangleMeshGeometry PTriMeshGeom;
+		//we need to allocate all of these here because PhysX insists on using the stack. This is wasteful, but reduces a lot of code duplication
+		PxSphereGeometry PSphereGeom;
+		PxBoxGeometry PBoxGeom;
+		PxCapsuleGeometry PCapsuleGeom;
+		PxConvexMeshGeometry PConvexGeom;
+		PxTriangleMeshGeometry PTriMeshGeom;
 
 		for (int32 ShapeIdx = 0; ShapeIdx < PShapes.Num(); ShapeIdx++)
-	{
-		bool bInvalid = false;	//we only mark invalid if actually found geom and it's invalid scale
-		PxGeometry* UpdatedGeometry = NULL;
-		PxShape* PShape = PShapes[ShapeIdx];
-		PxScene* PScene = PShape->getActor()->getScene();
+		{
+			bool bInvalid = false;	//we only mark invalid if actually found geom and it's invalid scale
+			PxGeometry* UpdatedGeometry = NULL;
+			PxShape* PShape = PShapes[ShapeIdx];
 
-		PxTransform PLocalPose = PShape->getLocalPose();
-		PLocalPose.q.normalize();
-		PxGeometryType::Enum GeomType = PShape->getGeometryType();
+			PxTransform PLocalPose = PShape->getLocalPose();
+			PLocalPose.q.normalize();
+			PxGeometryType::Enum GeomType = PShape->getGeometryType();
 
-		switch (GeomType)
-		{
-			case PxGeometryType::eSPHERE:
+			switch (GeomType)
 			{
-				ensure(ScaleMode == EScaleMode::LockedXYZ);
+				case PxGeometryType::eSPHERE:
+				{
+					ensure(ScaleMode == EScaleMode::LockedXYZ);
 
-				PShape->getSphereGeometry(PSphereGeom);
+					PShape->getSphereGeometry(PSphereGeom);
 
-				PSphereGeom.radius *= RelativeScale3DAbs.X;
-				PLocalPose.p *= RelativeScale3D.X;
+					PSphereGeom.radius *= RelativeScale3DAbs.X;
+					PLocalPose.p *= RelativeScale3D.X;
 
-				if (PSphereGeom.isValid())
-				{
-					UpdatedGeometry = &PSphereGeom;
-					bSuccess = true;
+					if (PSphereGeom.isValid())
+					{
+						UpdatedGeometry = &PSphereGeom;
+						bSuccess = true;
+					}
+					else
+					{
+						bInvalid = true;
+					}
+					break;
 				}
-				else
+				case PxGeometryType::eBOX:
 				{
-					bInvalid = true;
-				}
-				break;
-			}
-			case PxGeometryType::eBOX:
-			{
-				PShape->getBoxGeometry(PBoxGeom);
+					PShape->getBoxGeometry(PBoxGeom);
 
-				PBoxGeom.halfExtents.x *= RelativeScale3DAbs.X;
-				PBoxGeom.halfExtents.y *= RelativeScale3DAbs.Y;
-				PBoxGeom.halfExtents.z *= RelativeScale3DAbs.Z;
-				PLocalPose.p.x *= RelativeScale3D.X;
-				PLocalPose.p.y *= RelativeScale3D.Y;
-				PLocalPose.p.z *= RelativeScale3D.Z;
+					PBoxGeom.halfExtents.x *= RelativeScale3DAbs.X;
+					PBoxGeom.halfExtents.y *= RelativeScale3DAbs.Y;
+					PBoxGeom.halfExtents.z *= RelativeScale3DAbs.Z;
+					PLocalPose.p.x *= RelativeScale3D.X;
+					PLocalPose.p.y *= RelativeScale3D.Y;
+					PLocalPose.p.z *= RelativeScale3D.Z;
 
-				if (PBoxGeom.isValid())
-				{
-					UpdatedGeometry = &PBoxGeom;
-					bSuccess = true;
+					if (PBoxGeom.isValid())
+					{
+						UpdatedGeometry = &PBoxGeom;
+						bSuccess = true;
+					}
+					else
+					{
+						bInvalid = true;
+					}
+					break;
 				}
-				else
+				case PxGeometryType::eCAPSULE:
 				{
-					bInvalid = true;
-				}
-				break;
-			}
-			case PxGeometryType::eCAPSULE:
-			{
-				ensure(ScaleMode == EScaleMode::LockedXY || ScaleMode == EScaleMode::LockedXYZ);
+					ensure(ScaleMode == EScaleMode::LockedXY || ScaleMode == EScaleMode::LockedXYZ);
 
-				PShape->getCapsuleGeometry(PCapsuleGeom);
+					PShape->getCapsuleGeometry(PCapsuleGeom);
 
-				PCapsuleGeom.halfHeight *= RelativeScale3DAbs.Z;
-				PCapsuleGeom.radius *= RelativeScale3DAbs.X;
+					PCapsuleGeom.halfHeight *= RelativeScale3DAbs.Z;
+					PCapsuleGeom.radius *= RelativeScale3DAbs.X;
 
-				PLocalPose.p.x *= RelativeScale3D.X;
-				PLocalPose.p.y *= RelativeScale3D.Y;
-				PLocalPose.p.z *= RelativeScale3D.Z;
+					PLocalPose.p.x *= RelativeScale3D.X;
+					PLocalPose.p.y *= RelativeScale3D.Y;
+					PLocalPose.p.z *= RelativeScale3D.Z;
 
-				if (PCapsuleGeom.isValid())
-				{
-					UpdatedGeometry = &PCapsuleGeom;
-					bSuccess = true;
+					if (PCapsuleGeom.isValid())
+					{
+						UpdatedGeometry = &PCapsuleGeom;
+						bSuccess = true;
+					}
+					else
+					{
+						bInvalid = true;
+					}
+
+					break;
 				}
-				else
+				case PxGeometryType::eCONVEXMESH:
 				{
-					bInvalid = true;
-				}
+					PShape->getConvexMeshGeometry(PConvexGeom);
 
-				break;
-			}
-			case PxGeometryType::eCONVEXMESH:
-			{
-				PShape->getConvexMeshGeometry(PConvexGeom);
-
-				// find which convex elems it is
-				// it would be nice to know if the order of PShapes array index is in the order of createShape
-				// Create convex shapes
-				if (BodySetup.IsValid())
-				{
-					for (int32 i = 0; i < BodySetup->AggGeom.ConvexElems.Num(); i++)
+					// find which convex elems it is
+					// it would be nice to know if the order of PShapes array index is in the order of createShape
+					// Create convex shapes
+					if (BodySetup.IsValid())
 					{
-						FKConvexElem* ConvexElem = &(BodySetup->AggGeom.ConvexElems[i]);
-
-						// found it
-						if (ConvexElem->ConvexMesh == PConvexGeom.convexMesh)
+						for (int32 i = 0; i < BodySetup->AggGeom.ConvexElems.Num(); i++)
 						{
-							// Please note that this one we don't inverse old scale, but just set new one (but we still follow scale mode restriction)
-							FVector NewScale3D = RelativeScale3D * OldScale3D;
-							FVector Scale3DAbs(FMath::Abs(NewScale3D.X), FMath::Abs(NewScale3D.Y), FMath::Abs(NewScale3D.Z)); // magnitude of scale (sign removed)
-
-							PxTransform PNewLocalPose;
-							bool bUseNegX = CalcMeshNegScaleCompensation(NewScale3D, PNewLocalPose);
-
-							PxTransform PElementTransform = U2PTransform(ConvexElem->GetTransform());
-							PNewLocalPose.q *= PElementTransform.q;
-							PNewLocalPose.p += PElementTransform.p;
+							FKConvexElem* ConvexElem = &(BodySetup->AggGeom.ConvexElems[i]);
 
-							PConvexGeom.convexMesh = bUseNegX ? ConvexElem->ConvexMeshNegX : ConvexElem->ConvexMesh;
-							PConvexGeom.scale.scale = U2PVector(Scale3DAbs);
-
-							if (PConvexGeom.isValid())
-							{
-								UpdatedGeometry = &PConvexGeom;
-								bSuccess = true;
-							}
-							else
+							// found it
+							if (ConvexElem->ConvexMesh == PConvexGeom.convexMesh)
 							{
-								bInvalid = true;
+								// Please note that this one we don't inverse old scale, but just set new one (but we still follow scale mode restriction)
+								FVector NewScale3D = RelativeScale3D * OldScale3D;
+								FVector Scale3DAbs(FMath::Abs(NewScale3D.X), FMath::Abs(NewScale3D.Y), FMath::Abs(NewScale3D.Z)); // magnitude of scale (sign removed)
+
+								PxTransform PNewLocalPose;
+								bool bUseNegX = CalcMeshNegScaleCompensation(NewScale3D, PNewLocalPose);
+
+								PxTransform PElementTransform = U2PTransform(ConvexElem->GetTransform());
+								PNewLocalPose.q *= PElementTransform.q;
+								PNewLocalPose.p += PElementTransform.p;
+
+								PConvexGeom.convexMesh = bUseNegX ? ConvexElem->ConvexMeshNegX : ConvexElem->ConvexMesh;
+								PConvexGeom.scale.scale = U2PVector(Scale3DAbs);
+
+								if (PConvexGeom.isValid())
+								{
+									UpdatedGeometry = &PConvexGeom;
+									bSuccess = true;
+								}
+								else
+								{
+									bInvalid = true;
+								}
+								break;
 							}
-							break;
 						}
 					}
-				}
 
-				break;
-			}
-			case PxGeometryType::eTRIANGLEMESH:
-			{
-				PShape->getTriangleMeshGeometry(PTriMeshGeom);
-
-				// Create tri-mesh shape
-				if (BodySetup.IsValid() && (BodySetup->TriMesh != NULL || BodySetup->TriMeshNegX != NULL))
+					break;
+				}
+				case PxGeometryType::eTRIANGLEMESH:
 				{
-					// Please note that this one we don't inverse old scale, but just set new one (but still adjust for scale mode)
-					FVector NewScale3D = RelativeScale3D * OldScale3D;
-					FVector Scale3DAbs(FMath::Abs(NewScale3D.X), FMath::Abs(NewScale3D.Y), FMath::Abs(NewScale3D.Z)); // magnitude of scale (sign removed)
+					PShape->getTriangleMeshGeometry(PTriMeshGeom);
 
-					PxTransform PNewLocalPose;
-					bool bUseNegX = CalcMeshNegScaleCompensation(NewScale3D, PNewLocalPose);
-
-					// Only case where TriMeshNegX should be null is BSP, which should not require negX version
-					if (bUseNegX && BodySetup->TriMeshNegX == NULL)
+					// Create tri-mesh shape
+					if (BodySetup.IsValid() && (BodySetup->TriMesh != NULL || BodySetup->TriMeshNegX != NULL))
 					{
-						UE_LOG(LogPhysics, Warning, TEXT("FBodyInstance::UpdateBodyScale: Want to use NegX but it doesn't exist! %s"), *BodySetup->GetPathName());
-					}
+						// Please note that this one we don't inverse old scale, but just set new one (but still adjust for scale mode)
+						FVector NewScale3D = RelativeScale3D * OldScale3D;
+						FVector Scale3DAbs(FMath::Abs(NewScale3D.X), FMath::Abs(NewScale3D.Y), FMath::Abs(NewScale3D.Z)); // magnitude of scale (sign removed)
 
-					PxTriangleMesh* UseTriMesh = bUseNegX ? BodySetup->TriMeshNegX : BodySetup->TriMesh;
-					if (UseTriMesh != NULL)
-					{
-						PTriMeshGeom.triangleMesh = bUseNegX ? BodySetup->TriMeshNegX : BodySetup->TriMesh;
-						PTriMeshGeom.scale.scale = U2PVector(Scale3DAbs);
+						PxTransform PNewLocalPose;
+						bool bUseNegX = CalcMeshNegScaleCompensation(NewScale3D, PNewLocalPose);
 
-						if (PTriMeshGeom.isValid())
+						// Only case where TriMeshNegX should be null is BSP, which should not require negX version
+						if (bUseNegX && BodySetup->TriMeshNegX == NULL)
 						{
-							UpdatedGeometry = &PTriMeshGeom;
-							bSuccess = true;
-
+							UE_LOG(LogPhysics, Warning, TEXT("FBodyInstance::UpdateBodyScale: Want to use NegX but it doesn't exist! %s"), *BodySetup->GetPathName());
 						}
-						else
+
+						PxTriangleMesh* UseTriMesh = bUseNegX ? BodySetup->TriMeshNegX : BodySetup->TriMesh;
+						if (UseTriMesh != NULL)
 						{
-							bInvalid = true;
+							PTriMeshGeom.triangleMesh = bUseNegX ? BodySetup->TriMeshNegX : BodySetup->TriMesh;
+							PTriMeshGeom.scale.scale = U2PVector(Scale3DAbs);
+
+							if (PTriMeshGeom.isValid())
+							{
+								UpdatedGeometry = &PTriMeshGeom;
+								bSuccess = true;
+
+							}
+							else
+							{
+								bInvalid = true;
+							}
 						}
 					}
+					break;
 				}
-				break;
-			}
-			case PxGeometryType::eHEIGHTFIELD:
+				case PxGeometryType::eHEIGHTFIELD:
+				{
+					// HeightField is only used by Landscape, which does different code path from other primitives
+					break;
+				}
+				default:
+				{
+						   UE_LOG(LogPhysics, Error, TEXT("Unknown geom type."));
+				}
+			}// end switch
+
+			if (UpdatedGeometry)
 			{
-				// HeightField is only used by Landscape, which does different code path from other primitives
-				break;
+				ExecuteOnPxShapeWrite(this, PShape, [&](PxShape* PGivenShape)
+				{
+					PGivenShape->setLocalPose(PLocalPose);
+					PGivenShape->setGeometry(*UpdatedGeometry);
+				});
 			}
-			default:
+			else if (bInvalid)
 			{
-					   UE_LOG(LogPhysics, Error, TEXT("Unknown geom type."));
+				FMessageLog("PIE").Warning(FText::Format(LOCTEXT("PhysicsInvalidScale", "Scale ''{0}'' is not valid on object '{1}'."), FText::FromString(InScale3DAdjusted.ToString()), FText::FromString(GetBodyDebugName())));
 			}
-		}// end switch
-
-		if (UpdatedGeometry)
-		{
-			PShape->setLocalPose(PLocalPose);
-			PShape->setGeometry(*UpdatedGeometry);
 		}
-		else if (bInvalid)
-		{
-			FMessageLog("PIE").Warning(FText::Format(LOCTEXT("PhysicsInvalidScale", "Scale ''{0}'' is not valid on object '{1}'."), FText::FromString(InScale3DAdjusted.ToString()), FText::FromString(GetBodyDebugName())));
-		}
-	}
 	});
 	
 #endif
@@ -3696,7 +3750,7 @@ bool FBodyInstance::LineTrace(struct FHitResult& OutHit, const FVector& Start, c
 
 								// we don't get Shape information when we access via PShape, so I filled it up
 								BestHit.shape = PShape;
-								BestHit.actor = PShape->getActor();
+								BestHit.actor = HasSharedShapes() ? RigidActorSync : PShape->getActor();	//for shared shapes there is no actor, but since it's shared just return the sync actor
 							}
 						}
 					}
@@ -3819,7 +3873,7 @@ bool FBodyInstance::InternalSweepPhysX(struct FHitResult& OutHit, const FVector&
 
 						// we don't get Shape information when we access via PShape, so I filled it up
 						PHit.shape = PShape;
-						PHit.actor = PShape->getActor();
+						PHit.actor = HasSharedShapes() ? RigidActorSync : PShape->getActor();	//in the case of shared shapes getActor will return null. Since the shape is shared we just return the sync actor
 						PxTransform PStartTransform(U2PVector(Start));
 						ConvertQueryImpactHit(OwnerComponentInst->GetWorld(), PHit, OutHit, DeltaMag, QueryFilter, Start, End, NULL, PStartTransform, false, false);
 						return true;
@@ -3846,46 +3900,46 @@ float FBodyInstance::GetDistanceToBody(const FVector& Point, FVector& OutPointOn
 		if (RigidActor->getNbShapes() == 0 || OwnerComponent == NULL)
 		{
 			return;
-	}
+		}
 
 		bEarlyOut = false;
 
-	// Get all the shapes from the actor
-	TArray<PxShape*, TInlineAllocator<16>> PShapes;
+		// Get all the shapes from the actor
+		TArray<PxShape*, TInlineAllocator<16>> PShapes;
 		PShapes.AddZeroed(RigidActor->getNbShapes());
 		int32 NumShapes = RigidActor->getShapes(PShapes.GetData(), PShapes.Num());
 
-	const PxVec3 PPoint = U2PVector(Point);
+		const PxVec3 PPoint = U2PVector(Point);
 
-	// Iterate over each shape
+		// Iterate over each shape
 		for (int32 ShapeIdx = 0; ShapeIdx < PShapes.Num(); ShapeIdx++)
-	{
-		PxShape* PShape = PShapes[ShapeIdx];
-		check(PShape);
-		PxGeometry& PGeom = PShape->getGeometry().any();
+		{
+			PxShape* PShape = PShapes[ShapeIdx];
+			check(PShape);
+			PxGeometry& PGeom = PShape->getGeometry().any();
 			PxTransform PGlobalPose = PxShapeExt::getGlobalPose(*PShape, *RigidActor);
-		PxGeometryType::Enum GeomType = PShape->getGeometryType();
+			PxGeometryType::Enum GeomType = PShape->getGeometryType();
 
-		if (GeomType == PxGeometryType::eTRIANGLEMESH)
-		{
-			// Type unsupported for this function, but some other shapes will probably work. 
-			continue;
-		}
-		bFoundValidBody = true;
+			if (GeomType == PxGeometryType::eTRIANGLEMESH)
+			{
+				// Type unsupported for this function, but some other shapes will probably work. 
+				continue;
+			}
+			bFoundValidBody = true;
 
-		PxVec3 PClosestPoint;
-		float SqrDistance = PxGeometryQuery::pointDistance(PPoint, PGeom, PGlobalPose, &PClosestPoint);
-		// distance has valid data and smaller than mindistance
-			if (SqrDistance > 0.f && MinDistanceSqr > SqrDistance)
-		{
-			MinDistanceSqr = SqrDistance;
-			OutPointOnBody = P2UVector(PClosestPoint);
-		}
-			else if (SqrDistance == 0.f)
-		{
-			MinDistanceSqr = 0.f;
-			break;
-		}
+			PxVec3 PClosestPoint;
+			float SqrDistance = PxGeometryQuery::pointDistance(PPoint, PGeom, PGlobalPose, &PClosestPoint);
+			// distance has valid data and smaller than mindistance
+				if (SqrDistance > 0.f && MinDistanceSqr > SqrDistance)
+			{
+				MinDistanceSqr = SqrDistance;
+				OutPointOnBody = P2UVector(PClosestPoint);
+			}
+				else if (SqrDistance == 0.f)
+			{
+				MinDistanceSqr = 0.f;
+				break;
+			}
 		}	
 	});
 #endif //WITH_PHYSX
@@ -4250,8 +4304,13 @@ void FBodyInstance::SetUseAsyncScene(bool bNewUseAsyncScene)
 	bUseAsyncScene = bNewUseAsyncScene;
 }
 
-void FBodyInstance::ApplyMaterialToShape_AssumesLocked(PxShape* PShape, PxMaterial* PSimpleMat, TArray<UPhysicalMaterial*>& ComplexPhysMats)
+void FBodyInstance::ApplyMaterialToShape_AssumesLocked(PxShape* PShape, PxMaterial* PSimpleMat, const TArray<UPhysicalMaterial*>& ComplexPhysMats, const bool bSharedShape)
 {
+	if(!bSharedShape && !PShape->isExclusive())	//user says the shape is exclusive, but physx says it's shared
+	{
+		UE_LOG(LogPhysics, Warning, TEXT("FBodyInstance::ApplyMaterialToShape_AssumesLocked : Trying to change the physical material of a shared shape. If this is your intention pass bSharedShape = true"));
+	}
+
 	// If a triangle mesh, need to get array of materials...
 	if(PShape->getGeometryType() == PxGeometryType::eTRIANGLEMESH)
 	{
@@ -4270,7 +4329,7 @@ void FBodyInstance::ApplyMaterialToShape_AssumesLocked(PxShape* PShape, PxMateri
 		}
 		else
 		{
-			UE_LOG(LogPhysics, Warning, TEXT("FBodyInstance::UpdatePhysicalMaterials : PComplexMats is empty - falling back on simple physical material."));
+			UE_LOG(LogPhysics, Warning, TEXT("FBodyInstance::ApplyMaterialToShape_AssumesLocked : PComplexMats is empty - falling back on simple physical material."));
 			PShape->setMaterials(&PSimpleMat, 1);
 		}
 
@@ -4291,7 +4350,10 @@ void FBodyInstance::ApplyMaterialToInstanceShapes_AssumesLocked(PxMaterial* PSim
 	{
 		PxShape* PShape = AllShapes[ShapeIdx];
 
-		ApplyMaterialToShape_AssumesLocked(PShape, PSimpleMat, ComplexPhysMats);
+		ExecuteOnPxShapeWrite(this, PShape, [&](PxShape* PNewShape)
+		{
+			ApplyMaterialToShape_AssumesLocked(PNewShape, PSimpleMat, ComplexPhysMats, HasSharedShapes());
+		});		
 	}
 }
 
@@ -4605,90 +4667,93 @@ void FBodyInstance::InitStaticBodies(TArray<FBodyInstance*>& Bodies, TArray<FTra
 	}
 }
 
-void FBodyInstance::SetShapeFlags_AssumesLocked(TEnumAsByte<ECollisionEnabled::Type> UseCollisionEnabled, PxShape* PShape, EPhysicsSceneType SceneType, const bool bUseComplexAsSimple)
+void FBodyInstance::SetShapeFlags_AssumesLocked(TEnumAsByte<ECollisionEnabled::Type> UseCollisionEnabled, PxShape* PInShape, EPhysicsSceneType SceneType, const bool bUseComplexAsSimple)
 {
-	// If query collision is enabled..
-	bool bUpdateMassProperties = false;
-	if(UseCollisionEnabled != ECollisionEnabled::NoCollision)
+	ExecuteOnPxShapeWrite(this, PInShape, [&](PxShape* PShape)
 	{
-		UPrimitiveComponent* OwnerComponentInst = OwnerComponent.Get();
-		AActor* Owner = OwnerComponentInst ? OwnerComponentInst->GetOwner() : NULL;
-		const bool bPhysicsStatic = !OwnerComponentInst || OwnerComponentInst->IsWorldGeometry();
-		
-		// Only perform scene queries in the synchronous scene for static shapes
-		if(bPhysicsStatic)
-		{
-			PShape->setFlag(PxShapeFlag::eSCENE_QUERY_SHAPE, SceneType == PST_Sync);
-		}
-		// If non-static, always enable scene queries
-		else
+		// If query collision is enabled..
+		bool bUpdateMassProperties = false;
+		if (UseCollisionEnabled != ECollisionEnabled::NoCollision)
 		{
-			PShape->setFlag(PxShapeFlag::eSCENE_QUERY_SHAPE, true);
-		}
-
-		// See if we want physics collision
-		bool bSimCollision = (UseCollisionEnabled == ECollisionEnabled::QueryAndPhysics);
+			UPrimitiveComponent* OwnerComponentInst = OwnerComponent.Get();
+			AActor* Owner = OwnerComponentInst ? OwnerComponentInst->GetOwner() : NULL;
+			const bool bPhysicsStatic = !OwnerComponentInst || OwnerComponentInst->IsWorldGeometry();
 
-		// Triangle mesh is 'complex' geom
-		if(PShape->getGeometryType() == PxGeometryType::eTRIANGLEMESH)
-		{
-			// on dynamic objects and objects which don't use complex as simple, tri mesh not used for sim
-			if(!bSimCollision || !bUseComplexAsSimple)
+			// Only perform scene queries in the synchronous scene for static shapes
+			if (bPhysicsStatic)
 			{
-				PShape->setFlag(PxShapeFlag::eSIMULATION_SHAPE, false);
+				PShape->setFlag(PxShapeFlag::eSCENE_QUERY_SHAPE, SceneType == PST_Sync);
 			}
+			// If non-static, always enable scene queries
 			else
 			{
-				PShape->setFlag(PxShapeFlag::eSIMULATION_SHAPE, true);
+				PShape->setFlag(PxShapeFlag::eSCENE_QUERY_SHAPE, true);
 			}
 
-			if(OwnerComponentInst == NULL || !OwnerComponentInst->IsA(UModelComponent::StaticClass()))
-			{
-				PShape->setFlag(PxShapeFlag::eVISUALIZATION, false); // dont draw the tri mesh, we can see it anyway, and its slow
-			}
-		}
-		// Everything else is 'simple'
-		else
-		{
-			// See if we currently have sim collision
-			bool bCurrentSimCollision = (PShape->getFlags() & PxShapeFlag::eSIMULATION_SHAPE);
-			// Enable sim collision
-			if(bSimCollision && !bCurrentSimCollision)
-			{
-				bUpdateMassProperties = true;
-				PShape->setFlag(PxShapeFlag::eSIMULATION_SHAPE, true);
-			}
-			// Disable sim collision
-			else if(!bSimCollision && bCurrentSimCollision)
-			{
-				bUpdateMassProperties = true;
-				PShape->setFlag(PxShapeFlag::eSIMULATION_SHAPE, false);
-			}
+			// See if we want physics collision
+			bool bSimCollision = (UseCollisionEnabled == ECollisionEnabled::QueryAndPhysics);
 
-			// enable swept bounds for CCD for this shape
-			PxRigidBody* PBody = GetPxRigidActor_AssumesLocked()->is<PxRigidBody>();
-			if(bSimCollision && !bPhysicsStatic && bUseCCD && PBody)
+			// Triangle mesh is 'complex' geom
+			if (PShape->getGeometryType() == PxGeometryType::eTRIANGLEMESH)
 			{
-				PBody->setRigidBodyFlag(PxRigidBodyFlag::eENABLE_CCD, true);
+				// on dynamic objects and objects which don't use complex as simple, tri mesh not used for sim
+				if (!bSimCollision || !bUseComplexAsSimple)
+				{
+					PShape->setFlag(PxShapeFlag::eSIMULATION_SHAPE, false);
+				}
+				else
+				{
+					PShape->setFlag(PxShapeFlag::eSIMULATION_SHAPE, true);
+				}
+
+				if (OwnerComponentInst == NULL || !OwnerComponentInst->IsA(UModelComponent::StaticClass()))
+				{
+					PShape->setFlag(PxShapeFlag::eVISUALIZATION, false); // dont draw the tri mesh, we can see it anyway, and its slow
+				}
 			}
-			else if(PBody)
+			// Everything else is 'simple'
+			else
 			{
+				// See if we currently have sim collision
+				bool bCurrentSimCollision = (PShape->getFlags() & PxShapeFlag::eSIMULATION_SHAPE);
+				// Enable sim collision
+				if (bSimCollision && !bCurrentSimCollision)
+				{
+					bUpdateMassProperties = true;
+					PShape->setFlag(PxShapeFlag::eSIMULATION_SHAPE, true);
+				}
+				// Disable sim collision
+				else if (!bSimCollision && bCurrentSimCollision)
+				{
+					bUpdateMassProperties = true;
+					PShape->setFlag(PxShapeFlag::eSIMULATION_SHAPE, false);
+				}
 
-				PBody->setRigidBodyFlag(PxRigidBodyFlag::eENABLE_CCD, false);
+				// enable swept bounds for CCD for this shape
+				PxRigidBody* PBody = GetPxRigidActor_AssumesLocked()->is<PxRigidBody>();
+				if (bSimCollision && !bPhysicsStatic && bUseCCD && PBody)
+				{
+					PBody->setRigidBodyFlag(PxRigidBodyFlag::eENABLE_CCD, true);
+				}
+				else if (PBody)
+				{
+
+					PBody->setRigidBodyFlag(PxRigidBodyFlag::eENABLE_CCD, false);
+				}
 			}
 		}
-	}
-	// No collision enabled
-	else
-	{
-		PShape->setFlag(PxShapeFlag::eSIMULATION_SHAPE, false);
-		PShape->setFlag(PxShapeFlag::eSCENE_QUERY_SHAPE, false);
-	}
+		// No collision enabled
+		else
+		{
+			PShape->setFlag(PxShapeFlag::eSIMULATION_SHAPE, false);
+			PShape->setFlag(PxShapeFlag::eSCENE_QUERY_SHAPE, false);
+		}
 
-	if(bUpdateMassProperties)
-	{
-		UpdateMassProperties();
-	}
+		if (bUpdateMassProperties)
+		{
+			UpdateMassProperties();
+		}
+	});
 }
 
 void FBodyInstance::GetShapeFlags_AssumesLocked(FShapeData& ShapeData, TEnumAsByte<ECollisionEnabled::Type> UseCollisionEnabled, const bool bUseComplexAsSimple /*= false*/)
diff --git a/Engine/Source/Runtime/Engine/Private/PhysicsEngine/BodySetup.cpp b/Engine/Source/Runtime/Engine/Private/PhysicsEngine/BodySetup.cpp
--- a/Engine/Source/Runtime/Engine/Private/PhysicsEngine/BodySetup.cpp
+++ b/Engine/Source/Runtime/Engine/Private/PhysicsEngine/BodySetup.cpp
@@ -316,16 +316,18 @@ PxMaterial* GetDefaultPhysMaterial()
 /** Helper struct for adding shapes into an actor.*/
 struct FAddShapesHelper
 {
-	FAddShapesHelper(UBodySetup* InBodySetup, FBodyInstance* InOwningInstance, physx::PxRigidActor* InPDestActor, FVector& InScale3D, physx::PxMaterial* InSimpleMaterial, TArray<UPhysicalMaterial*>& InComplexMaterials, FShapeData& InShapeData, const FTransform& InRelativeTM, TArray<physx::PxShape*>* InNewShapes)
+	FAddShapesHelper(UBodySetup* InBodySetup, FBodyInstance* InOwningInstance, physx::PxRigidActor* InPDestActor, EPhysicsSceneType InSceneType, FVector& InScale3D, physx::PxMaterial* InSimpleMaterial, TArray<UPhysicalMaterial*>& InComplexMaterials, FShapeData& InShapeData, const FTransform& InRelativeTM, TArray<physx::PxShape*>* InNewShapes, const bool InShapeSharing)
 	: BodySetup(InBodySetup)
 	, OwningInstance(InOwningInstance)
 	, PDestActor(InPDestActor)
+	, SceneType(InSceneType)
 	, Scale3D(InScale3D)
 	, SimpleMaterial(InSimpleMaterial)
 	, ComplexMaterials(InComplexMaterials)
 	, ShapeData(InShapeData)
 	, RelativeTM(InRelativeTM)
 	, NewShapes(InNewShapes)
+	, bShapeSharing(InShapeSharing)
 	{
 		SetupNonUniformHelper(Scale3D, MinScale, MinScaleAbs, Scale3DAbs);
 		{
@@ -346,12 +348,14 @@ struct FAddShapesHelper
 	UBodySetup* BodySetup;
 	FBodyInstance* OwningInstance;
 	physx::PxRigidActor* PDestActor;
+	EPhysicsSceneType SceneType;
 	FVector& Scale3D;
-	const physx::PxMaterial* SimpleMaterial;
+	physx::PxMaterial* SimpleMaterial;
 	const TArray<UPhysicalMaterial*>& ComplexMaterials;
 	const FShapeData& ShapeData;
 	const FTransform& RelativeTM;
 	TArray<physx::PxShape*>* NewShapes;
+	const bool bShapeSharing;
 
 	float MinScaleAbs;
 	float MinScale;
@@ -362,7 +366,6 @@ public:
 	{
 		float ContactOffsetFactor, MaxContactOffset;
 		GetContactOffsetParams(ContactOffsetFactor, MaxContactOffset);
-		PxMaterial* PDefaultMat = GetDefaultPhysMaterial();
 
 		for (int32 i = 0; i < BodySetup->AggGeom.SphereElems.Num(); i++)
 		{
@@ -380,7 +383,7 @@ public:
 				ensure(PLocalPose.isValid());
 				{
 					const float ContactOffset = FMath::Min(MaxContactOffset, ContactOffsetFactor * PSphereGeom.radius);
-					AttachShape(PSphereGeom, *PDefaultMat, PLocalPose, ContactOffset);
+					AttachShape_AssumesLocked(PSphereGeom, PLocalPose, ContactOffset);
 				}
 			}
 			else
@@ -394,7 +397,6 @@ public:
 	{
 		float ContactOffsetFactor, MaxContactOffset;
 		GetContactOffsetParams(ContactOffsetFactor, MaxContactOffset);
-		PxMaterial* PDefaultMat = GetDefaultPhysMaterial();
 
 		for (int32 i = 0; i < BodySetup->AggGeom.BoxElems.Num(); i++)
 		{
@@ -416,7 +418,7 @@ public:
 				ensure(PLocalPose.isValid());
 				{
 					const float ContactOffset = FMath::Min(MaxContactOffset, ContactOffsetFactor * PBoxGeom.halfExtents.minElement());
-					AttachShape(PBoxGeom, *PDefaultMat, PLocalPose, ContactOffset);
+					AttachShape_AssumesLocked(PBoxGeom, PLocalPose, ContactOffset);
 				}
 			}
 			else
@@ -430,7 +432,6 @@ public:
 	{
 		float ContactOffsetFactor, MaxContactOffset;
 		GetContactOffsetParams(ContactOffsetFactor, MaxContactOffset);
-		PxMaterial* PDefaultMat = GetDefaultPhysMaterial();
 
 		float ScaleRadius = FMath::Max(Scale3DAbs.X, Scale3DAbs.Y);
 		float ScaleLength = Scale3DAbs.Z;
@@ -463,7 +464,7 @@ public:
 				ensure(PLocalPose.isValid());
 				{
 					const float ContactOffset = FMath::Min(MaxContactOffset, ContactOffsetFactor * PCapsuleGeom.radius);
-					AttachShape(PCapsuleGeom, *PDefaultMat, PLocalPose, ContactOffset);
+					AttachShape_AssumesLocked(PCapsuleGeom, PLocalPose, ContactOffset);
 				}
 			}
 			else
@@ -473,11 +474,10 @@ public:
 		}
 	}
 
-	void AddConvexElemsToRigidActor_AssumesLocked() const
+	void FAddShapesHelper::AddConvexElemsToRigidActor_AssumesLocked() const
 	{
 		float ContactOffsetFactor, MaxContactOffset;
 		GetContactOffsetParams(ContactOffsetFactor, MaxContactOffset);
-		PxMaterial* PDefaultMat = GetDefaultPhysMaterial();
 
 		for (int32 i = 0; i < BodySetup->AggGeom.ConvexElems.Num(); i++)
 		{
@@ -512,7 +512,7 @@ public:
 						ensure(PLocalPose.isValid());
 						{
 							const float ContactOffset = FMath::Min(MaxContactOffset, ContactOffsetFactor * PBoundsExtents.minElement());
-							AttachShape(PConvexGeom, *PDefaultMat, PLocalPose, ContactOffset);
+							AttachShape_AssumesLocked(PConvexGeom, PLocalPose, ContactOffset);
 						}
 					}
 					else
@@ -536,7 +536,6 @@ public:
 	{
 		float ContactOffsetFactor, MaxContactOffset;
 		GetContactOffsetParams(ContactOffsetFactor, MaxContactOffset);
-		PxMaterial* PDefaultMat = GetDefaultPhysMaterial();
 
 		if (BodySetup->TriMesh || BodySetup->TriMeshNegX)
 		{
@@ -567,7 +566,7 @@ public:
 
 					// Create without 'sim shape' flag, problematic if it's kinematic, and it gets set later anyway.
 					{
-						if (!AttachShape(PTriMeshGeom, *PDefaultMat, PLocalPose, MaxContactOffset, PxShapeFlag::eSCENE_QUERY_SHAPE | PxShapeFlag::eVISUALIZATION))
+						if (!AttachShape_AssumesLocked(PTriMeshGeom, PLocalPose, MaxContactOffset, PxShapeFlag::eSCENE_QUERY_SHAPE | PxShapeFlag::eVISUALIZATION))
 						{
 							UE_LOG(LogPhysics, Log, TEXT("Can't create new mesh shape in AddShapesToRigidActor"));
 						}
@@ -583,9 +582,11 @@ public:
 
 private:
 
-	PxShape* AttachShape(const PxGeometry& PGeom, const PxMaterial& PMaterial, const PxTransform& PLocalPose, const float ContactOffset, PxShapeFlags PShapeFlags = PxShapeFlag::eVISUALIZATION | PxShapeFlag::eSCENE_QUERY_SHAPE | PxShapeFlag::eSIMULATION_SHAPE) const
+	PxShape* AttachShape_AssumesLocked(const PxGeometry& PGeom, const PxTransform& PLocalPose, const float ContactOffset, PxShapeFlags PShapeFlags = PxShapeFlag::eVISUALIZATION | PxShapeFlag::eSCENE_QUERY_SHAPE | PxShapeFlag::eSIMULATION_SHAPE) const
 	{
-		PxShape* PNewShape = PDestActor->createShape(PGeom, PMaterial, PShapeFlags);
+		const PxMaterial* PMaterial = GetDefaultPhysMaterial(); 
+		PxShape* PNewShape = bShapeSharing ? GPhysXSDK->createShape(PGeom, *PMaterial, /*isExclusive =*/ false, PShapeFlags) : PDestActor->createShape(PGeom, *PMaterial, PShapeFlags);
+
 		if (PNewShape)
 		{
 			PNewShape->setLocalPose(PLocalPose);
@@ -596,6 +597,21 @@ private:
 			}
 
 			PNewShape->setContactOffset(ContactOffset);
+
+			const bool bSyncFlags = bShapeSharing || SceneType == PST_Sync;
+			const FShapeFilterData& Filters = ShapeData.FilterData;
+			const bool bComplexShape = PNewShape->getGeometryType() == PxGeometryType::eTRIANGLEMESH;
+
+			PNewShape->setQueryFilterData(bComplexShape ? Filters.QueryComplexFilter : Filters.QuerySimpleFilter);
+			PNewShape->setFlags( (bSyncFlags ? ShapeData.SyncShapeFlags : ShapeData.AsyncShapeFlags) | (bComplexShape ? ShapeData.ComplexShapeFlags : ShapeData.SimpleShapeFlags));
+			PNewShape->setSimulationFilterData(Filters.SimFilter);
+			FBodyInstance::ApplyMaterialToShape_AssumesLocked(PNewShape, SimpleMaterial, ComplexMaterials, bShapeSharing);
+
+			if(bShapeSharing)
+			{
+				PDestActor->attachShape(*PNewShape);
+				PNewShape->release();
+			}
 		}
 
 		return PNewShape;
@@ -606,15 +622,14 @@ private:
 
 
 
-void UBodySetup::AddShapesToRigidActor_AssumesLocked(FBodyInstance* OwningInstance, physx::PxRigidActor* PDestActor, EPhysicsSceneType SceneType, FVector& Scale3D, physx::PxMaterial* SimpleMaterial, TArray<UPhysicalMaterial*>& ComplexMaterials, FShapeData& ShapeData, const FTransform& RelativeTM, TArray<physx::PxShape*>* NewShapes)
+void UBodySetup::AddShapesToRigidActor_AssumesLocked(FBodyInstance* OwningInstance, physx::PxRigidActor* PDestActor, EPhysicsSceneType SceneType, FVector& Scale3D, physx::PxMaterial* SimpleMaterial, TArray<UPhysicalMaterial*>& ComplexMaterials, FShapeData& ShapeData, const FTransform& RelativeTM, TArray<physx::PxShape*>* NewShapes, bool bShapeSharing)
 {
 #if WITH_RUNTIME_PHYSICS_COOKING || WITH_EDITOR
 	// in editor, there are a lot of things relying on body setup to create physics meshes
 	CreatePhysicsMeshes();
 #endif
 
-	TArray<PxShape*> Shapes;
-	FAddShapesHelper AddShapesHelper(this, OwningInstance, PDestActor, Scale3D, SimpleMaterial, ComplexMaterials, ShapeData, RelativeTM, &Shapes);
+	FAddShapesHelper AddShapesHelper(this, OwningInstance, PDestActor, SceneType, Scale3D, SimpleMaterial, ComplexMaterials, ShapeData, RelativeTM, NewShapes, bShapeSharing);
 
 	// Create shapes for simple collision if we do not want to use the complex collision mesh 
 	// for simple queries as well
@@ -633,72 +648,12 @@ void UBodySetup::AddShapesToRigidActor_AssumesLocked(FBodyInstance* OwningInstan
 		AddShapesHelper.AddTriMeshToRigidActor_AssumesLocked();
 	}
 
-	for (PxShape* Shape : Shapes)
+	if (OwningInstance)
 	{
-		FShapeFilterData& Filters = ShapeData.FilterData;
-		// Apply provided materials
-		// If a triangle mesh, need to get array of materials...
-		if (Shape->getGeometryType() == PxGeometryType::eTRIANGLEMESH)
-		{
-			TArray<PxMaterial*> PComplexMats;
-			PComplexMats.AddUninitialized(ComplexMaterials.Num());
-			for (int MatIdx = 0; MatIdx < ComplexMaterials.Num(); MatIdx++)
-			{
-				check(ComplexMaterials[MatIdx] != NULL);
-				PComplexMats[MatIdx] = ComplexMaterials[MatIdx]->GetPhysXMaterial();
-				check(PComplexMats[MatIdx] != NULL);
-			}
-
-			if (PComplexMats.Num())
-			{
-				Shape->setMaterials(PComplexMats.GetData(), PComplexMats.Num());
-			}
-			else
-			{
-				UE_LOG(LogPhysics, Warning, TEXT("FBodyInstance::UpdatePhysicalMaterials : PComplexMats is empty - falling back on simple physical material."));
-				Shape->setMaterials(&SimpleMaterial, 1);
-			}
-
-			Shape->setQueryFilterData(Filters.QueryComplexFilter);
-
-			if (SceneType == PST_Sync)
-			{
-				Shape->setFlags(ShapeData.SyncShapeFlags | ShapeData.ComplexShapeFlags);
-			}
-			else
-			{
-				Shape->setFlags(ShapeData.AsyncShapeFlags | ShapeData.ComplexShapeFlags);
-			}
-		}
-		// Simple shape, 
-		else
+		if (PxRigidBody* RigidBody = OwningInstance->GetPxRigidBody_AssumesLocked())
 		{
-			Shape->setMaterials(&SimpleMaterial, 1);
-			Shape->setQueryFilterData(Filters.QuerySimpleFilter);
-
-			if (SceneType == PST_Sync)
-			{
-				Shape->setFlags(ShapeData.SyncShapeFlags | ShapeData.SimpleShapeFlags);
-			}
-			else
-			{
-				Shape->setFlags(ShapeData.AsyncShapeFlags | ShapeData.SimpleShapeFlags);
-			}
+			RigidBody->setRigidBodyFlags(ShapeData.SyncBodyFlags);
 		}
-		Shape->setSimulationFilterData(Filters.SimFilter);
-
-		if (OwningInstance)
-		{
-			if (PxRigidBody* RigidBody = OwningInstance->GetPxRigidBody_AssumesLocked())
-			{
-				RigidBody->setRigidBodyFlags(ShapeData.SyncBodyFlags);
-			}
-		}
-	}
-
-	if (NewShapes)
-	{
-		(*NewShapes) = MoveTemp(Shapes);
 	}
 }
 
diff --git a/Engine/Source/Runtime/Engine/Private/PhysicsEngine/PhysXSupport.cpp b/Engine/Source/Runtime/Engine/Private/PhysicsEngine/PhysXSupport.cpp
--- a/Engine/Source/Runtime/Engine/Private/PhysicsEngine/PhysXSupport.cpp
+++ b/Engine/Source/Runtime/Engine/Private/PhysicsEngine/PhysXSupport.cpp
@@ -397,10 +397,11 @@ PxFilterFlags PhysXSimFilterShader(	PxFilterObjectAttributes attributes0, PxFilt
 		pairFlags |= (PxPairFlag::eNOTIFY_TOUCH_FOUND | PxPairFlag::eNOTIFY_TOUCH_PERSISTS | PxPairFlag::eNOTIFY_CONTACT_POINTS );
 	}
 
-	if ((FilterFlags0&EPDF_ModifyContacts) || (FilterFlags1&EPDF_ModifyContacts))
+	//TODO: have to temporarily turn this off while getting a fix for shared shape crash
+	/*if ((FilterFlags0&EPDF_ModifyContacts) || (FilterFlags1&EPDF_ModifyContacts))
 	{
 		pairFlags |= (PxPairFlag::eMODIFY_CONTACTS);
-	}
+	}*/	
 
 	return PxFilterFlags();
 }
diff --git a/Engine/Source/Runtime/Engine/Private/PhysicsEngine/PhysicsSettings.cpp b/Engine/Source/Runtime/Engine/Private/PhysicsEngine/PhysicsSettings.cpp
--- a/Engine/Source/Runtime/Engine/Private/PhysicsEngine/PhysicsSettings.cpp
+++ b/Engine/Source/Runtime/Engine/Private/PhysicsEngine/PhysicsSettings.cpp
@@ -13,6 +13,7 @@ UPhysicsSettings::UPhysicsSettings(const FObjectInitializer& ObjectInitializer)
 	, RagdollAggregateThreshold(4)
 	, TriangleMeshTriangleMinAreaThreshold(5.0f)
 	, bEnableAsyncScene(false)
+	, bEnableShapeSharing(false)
 	, bEnablePCM(false)
 	, bWarnMissingLocks(true)
 	, bEnable2DPhysics(false)
