Delete overriden inherited component instances when disabling editable when inherited on a component

diff --git a/Engine/Source/Editor/PropertyEditor/Private/PropertyNode.cpp b/Engine/Source/Editor/PropertyEditor/Private/PropertyNode.cpp
--- a/Engine/Source/Editor/PropertyEditor/Private/PropertyNode.cpp
+++ b/Engine/Source/Editor/PropertyEditor/Private/PropertyNode.cpp
@@ -7,6 +7,7 @@
 #include "ScopedTransaction.h"
 #include "PropertyRestriction.h"
 #include "Editor/UnrealEd/Public/Kismet2/StructureEditorUtils.h"
+#include "Editor/UnrealEd/Public/Kismet2/BlueprintEditorUtils.h"
 #include "Engine/UserDefinedStruct.h"
 
 FPropertySettings& FPropertySettings::Get()
@@ -2332,7 +2333,7 @@ void FPropertyNode::PropagatePropertyChange( UObject* ModifiedObject, const TCHA
 
 	if (Object->HasAnyFlags(RF_ClassDefaultObject|RF_ArchetypeObject))
 	{
-		// Object is a default suobject, collect all instances.
+		// Object is a default subobject, collect all instances.
 		Object->GetArchetypeInstances(ArchetypeInstances);
 	}
 	else if (Object->HasAnyFlags(RF_DefaultSubObject) && Object->GetOuter()->HasAnyFlags(RF_ClassDefaultObject|RF_ArchetypeObject))
@@ -2347,6 +2348,12 @@ void FPropertyNode::PropagatePropertyChange( UObject* ModifiedObject, const TCHA
 		}
 	}
 
+	static FName FNAME_EditableWhenInherited = GET_MEMBER_NAME_CHECKED(UActorComponent,bEditableWhenInherited);
+	if (GetProperty()->GetFName() == FNAME_EditableWhenInherited && ModifiedObject->IsA<UActorComponent>() && FString(TEXT("False")) == NewValue)
+	{
+		FBlueprintEditorUtils::HandleDisableEditableWhenInherited(ModifiedObject, ArchetypeInstances);
+	}
+
 	FPropertyNode* Parent = GetParentNode();
 	UArrayProperty* ParentArrayProp  = Cast<UArrayProperty>(Parent->GetProperty());	
 
diff --git a/Engine/Source/Editor/UnrealEd/Private/Kismet2/BlueprintEditorUtils.cpp b/Engine/Source/Editor/UnrealEd/Private/Kismet2/BlueprintEditorUtils.cpp
--- a/Engine/Source/Editor/UnrealEd/Private/Kismet2/BlueprintEditorUtils.cpp
+++ b/Engine/Source/Editor/UnrealEd/Private/Kismet2/BlueprintEditorUtils.cpp
@@ -7745,5 +7745,24 @@ UK2Node_FunctionResult* FBlueprintEditorUtils::FindOrCreateFunctionResultNode(UK
 	return FunctionResult;
 }
 
+void FBlueprintEditorUtils::HandleDisableEditableWhenInherited(UObject* ModifiedObject, TArray<UObject*>& ArchetypeInstances)
+{
+	for (int32 Index = ArchetypeInstances.Num() - 1; Index >= 0; --Index)
+	{
+		UObject* ArchetypeInstance = ArchetypeInstances[Index];
+		if (ArchetypeInstance != ModifiedObject)
+		{
+			UBlueprintGeneratedClass* BPGC = Cast<UBlueprintGeneratedClass>(ArchetypeInstance->GetOuter());
+			if (BPGC)
+			{
+				UInheritableComponentHandler* ICH = BPGC->GetInheritableComponentHandler(false);
+				check(ICH);
+
+				ICH->RemoveOverridenComponentTemplate(ICH->FindKey(CastChecked<UActorComponent>(ArchetypeInstance)));
+			}
+		}
+	}
+}
+
 #undef LOCTEXT_NAMESPACE
 
diff --git a/Engine/Source/Editor/UnrealEd/Public/Kismet2/BlueprintEditorUtils.h b/Engine/Source/Editor/UnrealEd/Public/Kismet2/BlueprintEditorUtils.h
--- a/Engine/Source/Editor/UnrealEd/Public/Kismet2/BlueprintEditorUtils.h
+++ b/Engine/Source/Editor/UnrealEd/Public/Kismet2/BlueprintEditorUtils.h
@@ -1271,4 +1271,9 @@ public:
 	 * Returns a class name for the specified class that has no automatic suffixes, but is otherwise unmodified.  Class can be nullptr.
 	 */
 	static FString GetClassNameWithoutSuffix(const UClass* Class);
+
+	/**
+	 * Remove overridden component templates from instance component handlers when a parent class disables editable when inherited boolean.
+	 */
+	static void HandleDisableEditableWhenInherited(UObject* ModifiedObject, TArray<UObject*>& ArchetypeInstances);
 };
diff --git a/Engine/Source/Runtime/Engine/Classes/Engine/InheritableComponentHandler.h b/Engine/Source/Runtime/Engine/Classes/Engine/InheritableComponentHandler.h
--- a/Engine/Source/Runtime/Engine/Classes/Engine/InheritableComponentHandler.h
+++ b/Engine/Source/Runtime/Engine/Classes/Engine/InheritableComponentHandler.h
@@ -65,7 +65,13 @@ private:
 	bool IsRecordNecessary(const FComponentOverrideRecord& Record) const;
 
 public:
+
+	// Begin UObject interface
+	virtual void PostLoad() override;
+	// End UObject interface
+
 	UActorComponent* CreateOverridenComponentTemplate(FComponentKey Key);
+	void RemoveOverridenComponentTemplate(FComponentKey Key);
 	void UpdateOwnerClass(UBlueprintGeneratedClass* OwnerClass);
 	void ValidateTemplates();
 	bool IsValid() const;
diff --git a/Engine/Source/Runtime/Engine/Private/InheritableComponentHandler.cpp b/Engine/Source/Runtime/Engine/Private/InheritableComponentHandler.cpp
--- a/Engine/Source/Runtime/Engine/Private/InheritableComponentHandler.cpp
+++ b/Engine/Source/Runtime/Engine/Private/InheritableComponentHandler.cpp
@@ -6,16 +6,32 @@
 
 // UInheritableComponentHandler
 
+void UInheritableComponentHandler::PostLoad()
+{
+	Super::PostLoad();
+
+	for (int32 Index = Records.Num() - 1; Index >= 0; --Index)
+	{
+		FComponentOverrideRecord& Record = Records[Index];
+		if (Record.ComponentTemplate && !CastChecked<UActorComponent>(Record.ComponentTemplate->GetArchetype())->IsEditableWhenInherited())
+		{
+			Record.ComponentTemplate->MarkPendingKill(); // hack needed to be able to identify if NewObject returns this back to us in the future
+			Records.RemoveAtSwap(Index);
+		}
+	}
+}
+
 #if WITH_EDITOR
 UActorComponent* UInheritableComponentHandler::CreateOverridenComponentTemplate(FComponentKey Key)
 {
 	for (int32 Index = 0; Index < Records.Num(); ++Index)
 	{
-		if (Records[Index].ComponentKey.Match(Key))
+		FComponentOverrideRecord& Record = Records[Index];
+		if (Record.ComponentKey.Match(Key))
 		{
-			if (Records[Index].ComponentTemplate)
+			if (Record.ComponentTemplate)
 			{
-				return Records[Index].ComponentTemplate;
+				return Record.ComponentTemplate;
 			}
 			Records.RemoveAtSwap(Index);
 			break;
@@ -33,6 +49,16 @@ UActorComponent* UInheritableComponentHandler::CreateOverridenComponentTemplate(
 	auto NewComponentTemplate = NewObject<UActorComponent>(
 		GetOuter(), BestArchetype->GetClass(), BestArchetype->GetFName(), RF_ArchetypeObject | RF_Public | RF_InheritableComponentTemplate, BestArchetype);
 
+	// HACK: NewObject can return a pre-existing object which will not have been initialized to the archetype.  When we remove the old handlers, we mark them pending
+	//       kill so we can identify that situation here (see UE-13987/UE-13990)
+	if (NewComponentTemplate->IsPendingKill())
+	{
+		NewComponentTemplate->ClearFlags(RF_PendingKill);
+		UEngine::FCopyPropertiesForUnrelatedObjectsParams CopyParams;
+		CopyParams.bDoDelta = false;
+		UEngine::CopyPropertiesForUnrelatedObjects(BestArchetype, NewComponentTemplate, CopyParams);
+	}
+
 	FComponentOverrideRecord NewRecord;
 	NewRecord.ComponentKey = Key;
 	NewRecord.ComponentTemplate = NewComponentTemplate;
@@ -41,6 +67,20 @@ UActorComponent* UInheritableComponentHandler::CreateOverridenComponentTemplate(
 	return NewComponentTemplate;
 }
 
+void UInheritableComponentHandler::RemoveOverridenComponentTemplate(FComponentKey Key)
+{
+	for (int32 Index = 0; Index < Records.Num(); ++Index)
+	{
+		FComponentOverrideRecord& Record = Records[Index];
+		if (Record.ComponentKey.Match(Key))
+		{
+			Record.ComponentTemplate->MarkPendingKill(); // hack needed to be able to identify if NewObject returns this back to us in the future
+			Records.RemoveAtSwap(Index);
+			break;
+		}
+	}
+}
+
 void UInheritableComponentHandler::UpdateOwnerClass(UBlueprintGeneratedClass* OwnerClass)
 {
 	for (auto& Record : Records)
