Make Light Components Blueprintable
Expose Get Light Color as a linear color since the rest of the API uses it that way

diff --git a/Engine/Source/Runtime/Engine/Classes/Components/DirectionalLightComponent.h b/Engine/Source/Runtime/Engine/Classes/Components/DirectionalLightComponent.h
--- a/Engine/Source/Runtime/Engine/Classes/Components/DirectionalLightComponent.h
+++ b/Engine/Source/Runtime/Engine/Classes/Components/DirectionalLightComponent.h
@@ -10,7 +10,7 @@
 /**
  * A light component that has parallel rays. Will provide a uniform lighting across any affected surface (eg. The Sun). This will affect all objects in the defined light-mass importance volume.
  */
-UCLASS(ClassGroup=Lights, hidecategories=(Object, LightProfiles), editinlinenew, meta=(BlueprintSpawnableComponent), MinimalAPI)
+UCLASS(Blueprintable, ClassGroup=Lights, hidecategories=(Object, LightProfiles), editinlinenew, meta=(BlueprintSpawnableComponent), MinimalAPI)
 class UDirectionalLightComponent : public ULightComponent
 {
 	GENERATED_UCLASS_BODY()
diff --git a/Engine/Source/Runtime/Engine/Classes/Components/LightComponentBase.h b/Engine/Source/Runtime/Engine/Classes/Components/LightComponentBase.h
--- a/Engine/Source/Runtime/Engine/Classes/Components/LightComponentBase.h
+++ b/Engine/Source/Runtime/Engine/Classes/Components/LightComponentBase.h
@@ -93,9 +93,14 @@ class ENGINE_API ULightComponentBase : public USceneComponent
 	float DynamicEditorTextureScale;
 #endif
 
+	/** Sets whether this light casts shadows */
 	UFUNCTION(BlueprintCallable, Category="Rendering|Components|Light")
 	void SetCastShadows(bool bNewValue);
 
+	/** Gets the light color as a linear color */
+	UFUNCTION(BlueprintCallable, Category="Rendering|Components|Light")
+	FLinearColor GetLightColor() const;
+
 	virtual void Serialize(FArchive& Ar) override;
 
 	/**
diff --git a/Engine/Source/Runtime/Engine/Classes/Components/PointLightComponent.h b/Engine/Source/Runtime/Engine/Classes/Components/PointLightComponent.h
--- a/Engine/Source/Runtime/Engine/Classes/Components/PointLightComponent.h
+++ b/Engine/Source/Runtime/Engine/Classes/Components/PointLightComponent.h
@@ -8,7 +8,7 @@
 /**
  * A light component which emits light from a single point equally in all directions.
  */
-UCLASS(ClassGroup=(Lights,Common), hidecategories=(Object, LightShafts), editinlinenew, meta=(BlueprintSpawnableComponent))
+UCLASS(Blueprintable, ClassGroup=(Lights,Common), hidecategories=(Object, LightShafts), editinlinenew, meta=(BlueprintSpawnableComponent))
 class ENGINE_API UPointLightComponent : public ULightComponent
 {
 	GENERATED_UCLASS_BODY()
diff --git a/Engine/Source/Runtime/Engine/Classes/Components/SkyLightComponent.h b/Engine/Source/Runtime/Engine/Classes/Components/SkyLightComponent.h
--- a/Engine/Source/Runtime/Engine/Classes/Components/SkyLightComponent.h
+++ b/Engine/Source/Runtime/Engine/Classes/Components/SkyLightComponent.h
@@ -84,7 +84,7 @@ enum ESkyLightSourceType
 	SLS_MAX,
 };
 
-UCLASS(ClassGroup=Lights, HideCategories=(Trigger,Activation,"Components|Activation",Physics), meta=(BlueprintSpawnableComponent))
+UCLASS(Blueprintable, ClassGroup=Lights, HideCategories=(Trigger,Activation,"Components|Activation",Physics), meta=(BlueprintSpawnableComponent))
 class ENGINE_API USkyLightComponent : public ULightComponentBase
 {
 	GENERATED_UCLASS_BODY()
diff --git a/Engine/Source/Runtime/Engine/Classes/Components/SpotLightComponent.h b/Engine/Source/Runtime/Engine/Classes/Components/SpotLightComponent.h
--- a/Engine/Source/Runtime/Engine/Classes/Components/SpotLightComponent.h
+++ b/Engine/Source/Runtime/Engine/Classes/Components/SpotLightComponent.h
@@ -7,7 +7,7 @@
 /**
  * A spot light component emits a directional cone shaped light (Eg a Torch).
  */
-UCLASS(ClassGroup=Lights, hidecategories=Object, editinlinenew, meta=(BlueprintSpawnableComponent))
+UCLASS(Blueprintable, ClassGroup=Lights, hidecategories=Object, editinlinenew, meta=(BlueprintSpawnableComponent))
 class ENGINE_API USpotLightComponent : public UPointLightComponent
 {
 	GENERATED_UCLASS_BODY()
diff --git a/Engine/Source/Runtime/Engine/Private/Components/LightComponent.cpp b/Engine/Source/Runtime/Engine/Private/Components/LightComponent.cpp
--- a/Engine/Source/Runtime/Engine/Private/Components/LightComponent.cpp
+++ b/Engine/Source/Runtime/Engine/Private/Components/LightComponent.cpp
@@ -78,6 +78,11 @@ void ULightComponentBase::SetCastShadows(bool bNewValue)
 	}
 }
 
+FLinearColor ULightComponentBase::GetLightColor() const
+{
+	return FLinearColor(LightColor);
+}
+
 void ULightComponentBase::Serialize(FArchive& Ar)
 {
 	Super::Serialize(Ar);
