Fix for crash in actorcomponent undo/redo in the blueprint editor. Described by QA in UE-8170

diff --git a/Engine/Source/Runtime/Engine/Private/Components/ActorComponent.cpp b/Engine/Source/Runtime/Engine/Private/Components/ActorComponent.cpp
--- a/Engine/Source/Runtime/Engine/Private/Components/ActorComponent.cpp
+++ b/Engine/Source/Runtime/Engine/Private/Components/ActorComponent.cpp
@@ -521,7 +521,10 @@ void UActorComponent::PostEditUndo()
 			Owner->AddOwnedComponent(this);
 		}
 
-		GetWorld()->UpdateActorComponentEndOfFrameUpdateState(this);
+		if (GetWorld())
+		{
+			GetWorld()->UpdateActorComponentEndOfFrameUpdateState(this);
+		}
 	}
 	Super::PostEditUndo();
 }
