#ue4 - Perf: Avoid unnecessary work in APawn::FaceRotation() when none of bUseControllerRotation[Pitch|Yaw|Roll] are enabled.

diff --git a/Engine/Source/Runtime/Engine/Private/Pawn.cpp b/Engine/Source/Runtime/Engine/Private/Pawn.cpp
--- a/Engine/Source/Runtime/Engine/Private/Pawn.cpp
+++ b/Engine/Source/Runtime/Engine/Private/Pawn.cpp
@@ -854,24 +854,28 @@ void APawn::ClientSetRotation( FRotator NewRotation )
 
 void APawn::FaceRotation(FRotator NewControlRotation, float DeltaTime)
 {
-	const FRotator CurrentRotation = GetActorRotation();
-
-	if (!bUseControllerRotationPitch)
+	// Only if we actually are going to use any component of rotation.
+	if (bUseControllerRotationPitch || bUseControllerRotationYaw || bUseControllerRotationRoll)
 	{
-		NewControlRotation.Pitch = CurrentRotation.Pitch;
-	}
+		const FRotator CurrentRotation = GetActorRotation();
 
-	if (!bUseControllerRotationYaw)
-	{
-		NewControlRotation.Yaw = CurrentRotation.Yaw;
-	}
+		if (!bUseControllerRotationPitch)
+		{
+			NewControlRotation.Pitch = CurrentRotation.Pitch;
+		}
 
-	if (!bUseControllerRotationRoll)
-	{
-		NewControlRotation.Roll = CurrentRotation.Roll;
-	}
+		if (!bUseControllerRotationYaw)
+		{
+			NewControlRotation.Yaw = CurrentRotation.Yaw;
+		}
 
-	SetActorRotation(NewControlRotation);
+		if (!bUseControllerRotationRoll)
+		{
+			NewControlRotation.Roll = CurrentRotation.Roll;
+		}
+
+		SetActorRotation(NewControlRotation);
+	}
 }
 
 void APawn::DetachFromControllerPendingDestroy()
