Integrated change from UT to main: (CL 2463239 by [REDACTED_PERSON][REDACTED_HANDLE] on 2015/02/26 18:42:59)
  optimized BoneDataVertexBBuffer updates to only copy the data needed, saves memcpy time in the render thread

diff --git a/Engine/Source/Runtime/Engine/Private/GPUSkinVertexFactory.cpp b/Engine/Source/Runtime/Engine/Private/GPUSkinVertexFactory.cpp
--- a/Engine/Source/Runtime/Engine/Private/GPUSkinVertexFactory.cpp
+++ b/Engine/Source/Runtime/Engine/Private/GPUSkinVertexFactory.cpp
@@ -180,11 +180,11 @@ float* FBoneDataVertexBuffer::LockData()
 	return Data;
 }
 
-void FBoneDataVertexBuffer::UnlockData()
+void FBoneDataVertexBuffer::UnlockData(uint32 SizeInBytes)
 {
 	check(IsValidRef(BoneBuffer));
 	check(IsInRenderingThread() && AllocBlock && AllocSize);
-	FRHICommandListExecutor::GetImmediateCommandList().UpdateVertexBuffer(BoneBuffer.VertexBufferRHI, AllocBlock, AllocSize);
+	FRHICommandListExecutor::GetImmediateCommandList().UpdateVertexBuffer(BoneBuffer.VertexBufferRHI, AllocBlock, SizeInBytes ? SizeInBytes : 4);
 	FMemory::Free(AllocBlock);
 	AllocBlock = nullptr;
 	AllocSize = 0;
@@ -405,7 +405,7 @@ public:
 
 			bool bLocalPerBoneMotionBlur = false;
 			
-			if (FeatureLevel >= ERHIFeatureLevel::SM4 && GPrevPerBoneMotionBlur.IsLocked())
+			if (FeatureLevel >= ERHIFeatureLevel::SM4 && GPrevPerBoneMotionBlur.IsAppendStarted())
 			{
 				// we are in the velocity rendering pass
 
@@ -513,7 +513,7 @@ public:
 
 		bool bIsGPUCached = false;
 
-		if (GEnableGPUSkinCache && GGPUSkinCache.SetVertexStreamFromCache(RHICmdList, BatchElement.UserIndex, Shader, VertexFactory, BatchElement.MinVertexIndex, GPrevPerBoneMotionBlur.IsLocked(), GPUSkinCacheStreamFloatOffset, GPUSkinCacheStreamStride, GPUSkinCacheStreamBuffer))
+		if (GEnableGPUSkinCache && GGPUSkinCache.SetVertexStreamFromCache(RHICmdList, BatchElement.UserIndex, Shader, VertexFactory, BatchElement.MinVertexIndex, GPrevPerBoneMotionBlur.IsAppendStarted(), GPUSkinCacheStreamFloatOffset, GPUSkinCacheStreamStride, GPUSkinCacheStreamBuffer))
 		{
 			bIsGPUCached = true;
 		}
diff --git a/Engine/Source/Runtime/Engine/Private/SkeletalRenderGPUSkin.cpp b/Engine/Source/Runtime/Engine/Private/SkeletalRenderGPUSkin.cpp
--- a/Engine/Source/Runtime/Engine/Private/SkeletalRenderGPUSkin.cpp
+++ b/Engine/Source/Runtime/Engine/Private/SkeletalRenderGPUSkin.cpp
@@ -1106,7 +1106,7 @@ uint32 FPreviousPerBoneMotionBlur::GetSizeX() const
 	return PerChunkBoneMatricesTexture[0].GetSizeX();
 }
 
-bool FPreviousPerBoneMotionBlur::IsLocked() const
+bool FPreviousPerBoneMotionBlur::IsAppendStarted() const
 {
 	return LockedData != 0;
 }
@@ -1118,7 +1118,7 @@ void FPreviousPerBoneMotionBlur::InitIfNeeded()
 		InitResources();
 	}
 }
-void FPreviousPerBoneMotionBlur::LockData()
+void FPreviousPerBoneMotionBlur::StartAppend()
 {
 	check(!LockedData);
 	check(IsInRenderingThread());
@@ -1162,16 +1162,17 @@ uint32 FPreviousPerBoneMotionBlur::AppendData(FBoneSkinning *DataStart, uint32 B
 	}
 }
 
-void FPreviousPerBoneMotionBlur::UnlockData()
+void FPreviousPerBoneMotionBlur::EndAppend()
 {
-	if(IsLocked())
+	if(IsAppendStarted())
 	{
 		check(IsInRenderingThread());
-		LockedTexelPosition.Set(0);
 		LockedTexelCount = 0;
 		LockedData = 0;
 
-		PerChunkBoneMatricesTexture[GetWriteBufferIndex()].UnlockData();
+		// we use float4
+		PerChunkBoneMatricesTexture[GetWriteBufferIndex()].UnlockData(LockedTexelPosition.GetValue() * 4 * sizeof(float));
+		LockedTexelPosition.Set(0);
 
 		AdvanceBufferIndex();
 	}
diff --git a/Engine/Source/Runtime/Engine/Private/SkeletalRenderGPUSkin.h b/Engine/Source/Runtime/Engine/Private/SkeletalRenderGPUSkin.h
--- a/Engine/Source/Runtime/Engine/Private/SkeletalRenderGPUSkin.h
+++ b/Engine/Source/Runtime/Engine/Private/SkeletalRenderGPUSkin.h
@@ -188,10 +188,10 @@ public:
 	uint32 GetSizeX() const;
 
 	/** so we update only during velocity rendering pass */
-	bool IsLocked() const;
+	bool IsAppendStarted() const;
 
 	/** needed before AppendData() ccan be called */
-	ENGINE_API void LockData();
+	ENGINE_API void StartAppend();
 
 	/**
 	 * use between LockData() and UnlockData()
@@ -201,8 +201,8 @@ public:
 	 */
 	uint32 AppendData(FBoneSkinning *DataStart, uint32 BoneCount);
 
-	/** only call if LockData() */
-	ENGINE_API void UnlockData();
+	/** only call if StartAppend() */
+	ENGINE_API void EndAppend();
 
 	/** @return 0 if there should be no bone based motion blur (no previous data available or it's not active) */
 	FBoneDataVertexBuffer* GetReadData();
diff --git a/Engine/Source/Runtime/Engine/Public/GPUSkinVertexFactory.h b/Engine/Source/Runtime/Engine/Public/GPUSkinVertexFactory.h
--- a/Engine/Source/Runtime/Engine/Public/GPUSkinVertexFactory.h
+++ b/Engine/Source/Runtime/Engine/Public/GPUSkinVertexFactory.h
@@ -217,7 +217,7 @@ public:
 	/**
 	* Needs to be called after LockData()
 	*/
-	void UnlockData();
+	void UnlockData(uint32 SizeInBytes);
 
 	/** Returns the size of the buffer in pixels. */
 	uint32 GetSizeX() const;
diff --git a/Engine/Source/Runtime/Renderer/Private/DeferredShadingRenderer.cpp b/Engine/Source/Runtime/Renderer/Private/DeferredShadingRenderer.cpp
--- a/Engine/Source/Runtime/Renderer/Private/DeferredShadingRenderer.cpp
+++ b/Engine/Source/Runtime/Renderer/Private/DeferredShadingRenderer.cpp
@@ -1401,7 +1401,7 @@ bool FDeferredShadingSceneRenderer::RenderBasePass(FRHICommandListImmediate& RHI
 	if (FVelocityRendering::OutputsToGBuffer())
 	{
 		QUICK_SCOPE_CYCLE_COUNTER(STAT_FDeferredShadingSceneRenderer_RenderBasePass_GPrevPerBoneMotionBlur_LockData);
-		GPrevPerBoneMotionBlur.LockData();
+		GPrevPerBoneMotionBlur.StartAppend();
 	}
 
 	if(ViewFamily.EngineShowFlags.LightMapDensity && AllowDebugViewmodes())
@@ -1440,7 +1440,7 @@ bool FDeferredShadingSceneRenderer::RenderBasePass(FRHICommandListImmediate& RHI
 	if (FVelocityRendering::OutputsToGBuffer())
 	{
 		QUICK_SCOPE_CYCLE_COUNTER(STAT_FDeferredShadingSceneRenderer_RenderBasePass_GPrevPerBoneMotionBlur_UnlockData);
-		GPrevPerBoneMotionBlur.UnlockData();
+		GPrevPerBoneMotionBlur.EndAppend();
 	}
 
 	return bDirty;
diff --git a/Engine/Source/Runtime/Renderer/Private/VelocityRendering.cpp b/Engine/Source/Runtime/Renderer/Private/VelocityRendering.cpp
--- a/Engine/Source/Runtime/Renderer/Private/VelocityRendering.cpp
+++ b/Engine/Source/Runtime/Renderer/Private/VelocityRendering.cpp
@@ -738,7 +738,7 @@ void FDeferredShadingSceneRenderer::RenderVelocities(FRHICommandListImmediate& R
 	FPooledRenderTargetDesc Desc = FVelocityRendering::GetRenderTargetDesc();
 	GRenderTargetPool.FindFreeElement(Desc, VelocityRT, TEXT("Velocity"));
 
-	GPrevPerBoneMotionBlur.LockData();
+	GPrevPerBoneMotionBlur.StartAppend();
 
 	BeginVelocityRendering(RHICmdList, VelocityRT, true);
 	
@@ -757,7 +757,7 @@ void FDeferredShadingSceneRenderer::RenderVelocities(FRHICommandListImmediate& R
 	RHICmdList.SetBlendState(TStaticBlendState<>::GetRHI());
 	RHICmdList.SetRasterizerState(TStaticRasterizerState<>::GetRHI());
 	
-	GPrevPerBoneMotionBlur.UnlockData();
+	GPrevPerBoneMotionBlur.EndAppend();
 
 	// to be able to observe results with VisualizeTexture
 	GRenderTargetPool.VisualizeTexture.SetCheckPoint(RHICmdList, VelocityRT);
