#ue4 - Fix Characters using orient-to-movement or controller desired rotation not rotating along with a moving base.

Jira: UE-13797

diff --git a/Engine/Source/Runtime/Engine/Private/Components/CharacterMovementComponent.cpp b/Engine/Source/Runtime/Engine/Private/Components/CharacterMovementComponent.cpp
--- a/Engine/Source/Runtime/Engine/Private/Components/CharacterMovementComponent.cpp
+++ b/Engine/Source/Runtime/Engine/Private/Components/CharacterMovementComponent.cpp
@@ -1376,17 +1376,31 @@ void UCharacterMovementComponent::UpdateBasedMovement(float DeltaSeconds)
 			{
 				// Apply change in rotation and pipe through FaceRotation to maintain axis restrictions
 				const FQuat PawnOldQuat = UpdatedComponent->GetComponentQuat();
-				FinalQuat = DeltaQuat * FinalQuat;
-				CharacterOwner->FaceRotation(FinalQuat.Rotator(), 0.f);
+				const FQuat TargetQuat = DeltaQuat * FinalQuat;
+				FRotator TargetRotator(TargetQuat);
+				CharacterOwner->FaceRotation(TargetRotator, 0.f);
 				FinalQuat = UpdatedComponent->GetComponentQuat();
 
+				if (PawnOldQuat.Equals(FinalQuat))
+				{
+					// Nothing changed. This means we probably are using another rotation mechanism (bOrientToMovement etc). We should still follow the base object.
+					// @todo: This assumes only Yaw is used, currently a valid assumption. This is the only reason FaceRotation() is used above really, aside from being a virtual hook.
+					if (bOrientRotationToMovement || (bUseControllerDesiredRotation && CharacterOwner->Controller))
+					{
+						TargetRotator.Pitch = 0.f;
+						TargetRotator.Roll = 0.f;
+						MoveUpdatedComponent(FVector::ZeroVector, TargetRotator, false);
+						FinalQuat = UpdatedComponent->GetComponentQuat();
+					}
+				}
+
 				// Pipe through ControlRotation, to affect camera.
 				if (CharacterOwner->Controller)
 				{
 					const FQuat PawnDeltaRotation = FinalQuat * PawnOldQuat.Inverse();
 					FRotator FinalRotation = FinalQuat.Rotator();
 					UpdateBasedRotation(FinalRotation, PawnDeltaRotation.Rotator());
-					FinalQuat = FinalRotation.Quaternion();
+					FinalQuat = UpdatedComponent->GetComponentQuat();
 				}
 			}
 
