#ue4 - Save movement mode in CharacterMovement saved move when recording it after movement.

diff --git a/Engine/Source/Runtime/Engine/Private/Components/CharacterMovementComponent.cpp b/Engine/Source/Runtime/Engine/Private/Components/CharacterMovementComponent.cpp
--- a/Engine/Source/Runtime/Engine/Private/Components/CharacterMovementComponent.cpp
+++ b/Engine/Source/Runtime/Engine/Private/Components/CharacterMovementComponent.cpp
@@ -6011,7 +6011,7 @@ void UCharacterMovementComponent::CallServerMove
 			ClientYawPitchINT,
 			ClientMovementBase,
 			ClientBaseBone,
-			PackNetworkMovementMode()
+			NewMove->MovementMode
 			);
 	}
 	else
@@ -6026,7 +6026,7 @@ void UCharacterMovementComponent::CallServerMove
 			ClientYawPitchINT,
 			ClientMovementBase,
 			ClientBaseBone,
-			PackNetworkMovementMode()
+			NewMove->MovementMode
 			);
 	}
 
@@ -7211,6 +7211,7 @@ void FSavedMove_Character::PostUpdate(ACharacter* Character, FSavedMove_Characte
 {
 	// Common code for both recording and after a replay.
 	{
+		MovementMode = Character->GetCharacterMovement()->PackNetworkMovementMode();
 		SavedLocation = Character->GetActorLocation();
 		SavedRotation = Character->GetActorRotation();
 		UPrimitiveComponent* const MovementBase = Character->GetMovementBase();
