PR #856: Fix invalid default value behaviour in asset files  (Contributed by [REDACTED_PERSON])

This fix makes sure that default values from leaf structures do not stump default values specified in structures that aggregate them. To reproduce the bug make two structures (A and B), one with structure B containing a member of type structure A. In structure A add a string, set it's default value to "foo", in structure b, set a different default value to A's string ("bar"). When B i added to a blueprint (C), notice that changing the default value of the blueprint's member back to "bar" (B's default value) is impossible - when the blueprint is compiled the value will be reset to foo.

diff --git a/Engine/Source/Runtime/Engine/Private/UserDefinedStruct.cpp b/Engine/Source/Runtime/Engine/Private/UserDefinedStruct.cpp
--- a/Engine/Source/Runtime/Engine/Private/UserDefinedStruct.cpp
+++ b/Engine/Source/Runtime/Engine/Private/UserDefinedStruct.cpp
@@ -159,7 +159,11 @@ void UUserDefinedStruct::SerializeTaggedProperties(FArchive& Ar, uint8* Data, US
 	}
 	if (bLoadDefaultFirst)
 	{
-		UDDefaultsStruct->CopyScriptStruct(Data, Defaults ? Defaults : StructDefaultMem.GetStructMemory());
+		if (Defaults == nullptr)
+		{
+			Defaults = StructDefaultMem.GetStructMemory();
+		}
+		UDDefaultsStruct->CopyScriptStruct(Data, Defaults);
 	}
 #endif // WITH_EDITOR
 	Super::SerializeTaggedProperties(Ar, Data, DefaultsStruct, Defaults);
