Major network checksum cleanup/bug fixes + stable sorting + debug code (disabled currently)

diff --git a/Engine/Source/Runtime/CoreUObject/Private/UObject/Class.cpp b/Engine/Source/Runtime/CoreUObject/Private/UObject/Class.cpp
--- a/Engine/Source/Runtime/CoreUObject/Private/UObject/Class.cpp
+++ b/Engine/Source/Runtime/CoreUObject/Private/UObject/Class.cpp
@@ -2881,6 +2881,12 @@ void UClass::Link(FArchive& Ar, bool bRelinkExistingProperties)
 		{
 			FORCEINLINE bool operator()( UProperty & A, UProperty & B ) const
 			{
+				// Ensure stable sort
+				if ( A.GetOffset_ForGC() == B.GetOffset_ForGC() )
+				{
+					return A.GetName() < B.GetName();
+				}
+
 				return A.GetOffset_ForGC() < B.GetOffset_ForGC();
 			}
 		};
diff --git a/Engine/Source/Runtime/Engine/Private/PackageMapClient.cpp b/Engine/Source/Runtime/Engine/Private/PackageMapClient.cpp
--- a/Engine/Source/Runtime/Engine/Private/PackageMapClient.cpp
+++ b/Engine/Source/Runtime/Engine/Private/PackageMapClient.cpp
@@ -2079,52 +2079,109 @@ void FNetGUIDCache::GenerateFullNetGUIDPath_r( const FNetworkGUID& NetGUID, FStr
 	}
 }
 
-static uint32 GetNetFieldChecksum( const UField* NetField, const uint32 Checksum )
+static uint32 GetPropertyChecksum( const UProperty* Property, uint32 Checksum );
+
+static void SortProperties( TArray< UProperty* >& Properties )
 {
-	uint32 LocalChecksum = FCrc::StrCrc32( *NetField->GetName(), Checksum );
+	// Sort NetProperties so that their ClassReps are sorted by memory offset
+	struct FCompareUFieldOffsets
+	{
+		FORCEINLINE bool operator()( UProperty & A, UProperty & B ) const
+		{
+			// Ensure stable sort
+			if ( A.GetOffset_ForGC() == B.GetOffset_ForGC() )
+			{
+				return A.GetName() < B.GetName();
+			}
 
-	LocalChecksum = FCrc::StrCrc32( *NetField->GetClass()->GetName(), LocalChecksum );
+			return A.GetOffset_ForGC() < B.GetOffset_ForGC();
+		}
+	};
 
-	const UFunction* Function = Cast< UFunction >( NetField );
+	Sort( Properties.GetData(), Properties.Num(), FCompareUFieldOffsets() );
+}
+
+static uint32 SortedStructFieldsChecksum( const UStruct* Struct, uint32 Checksum )
+{
+	// Generate a list that we can sort, to make sure we process these deterministically
+	TArray< UProperty * > Fields;
 
-	if ( Function != NULL )
+	for ( TFieldIterator< UProperty > It( Struct ); It; ++It )
 	{
-		// Build checksum based on function parameter list
-		for ( TFieldIterator< UProperty > It( Function ); It && ( It->PropertyFlags & ( CPF_Parm | CPF_ReturnParm ) ) == CPF_Parm; ++It )
+		if ( It->PropertyFlags & CPF_RepSkip )
 		{
-			LocalChecksum = FCrc::StrCrc32( *It->GetName(), LocalChecksum );
+			continue;
 		}
 
-		return LocalChecksum;
+		Fields.Add( *It );
 	}
 
-	const UArrayProperty* ArrayProperty = Cast< UArrayProperty >( NetField );
+	// Sort them
+	SortProperties( Fields );
+
+	// Evolve the checksum on the sorted list
+	for ( auto Field : Fields )
+	{
+		Checksum = GetPropertyChecksum( Field, Checksum );
+	}
+
+	return Checksum;
+}
+
+static bool GDebugChecksum = false;
+static int GDebugChecksumIndent = 0;
+
+static uint32 GetPropertyChecksum( const UProperty* Property, uint32 Checksum )
+{
+	if ( GDebugChecksum )
+	{
+		UE_LOG( LogNetPackageMap, Warning, TEXT( "%s%s [%s] [%u] [%u]" ), FCString::Spc( 2 * GDebugChecksumIndent ), *Property->GetName().ToLower(), *Property->GetClass()->GetName().ToLower(), Property->ArrayDim, Checksum );
+	}
+
+	// Evolve checksum on name
+	Checksum = FCrc::StrCrc32( *Property->GetName().ToLower(), Checksum );
+	
+	// Evolve checksum on class name
+	Checksum = FCrc::StrCrc32( *Property->GetClass()->GetName().ToLower(), Checksum );
+
+	// Evolve checksum on array dim (to detect when static arrays change size)
+	// We use a string to be endian agnostic
+	Checksum = FCrc::StrCrc32( *FString::Printf( TEXT( "%u" ), Property->ArrayDim ), Checksum );
 
+	const UArrayProperty* ArrayProperty = Cast< UArrayProperty >( Property );
+
+	// Evolve checksum on array inner
 	if ( ArrayProperty != NULL )
 	{
-		return GetNetFieldChecksum( ArrayProperty->Inner, Checksum );
+		return GetPropertyChecksum( ArrayProperty->Inner, Checksum );
 	}
 
-	const UStructProperty* StructProperty = Cast< UStructProperty >( NetField );
+	const UStructProperty* StructProperty = Cast< UStructProperty >( Property );
 
+	// Evolve checksum on property struct fields
 	if ( StructProperty != NULL )
 	{
-		LocalChecksum = FCrc::StrCrc32( *StructProperty->Struct->GetName(), LocalChecksum );
-
-		for ( TFieldIterator< UProperty > It( StructProperty->Struct ); It; ++It )
+		if ( GDebugChecksum )
 		{
-			LocalChecksum = GetNetFieldChecksum( *It, LocalChecksum );
+			UE_LOG( LogNetPackageMap, Warning, TEXT( "%s [%s] [%u]" ), FCString::Spc( 2 * GDebugChecksumIndent ), *StructProperty->Struct->GetName().ToLower(), Checksum );
 		}
 
-		return LocalChecksum;
+		// Evolve checksum on struct name
+		Checksum = FCrc::StrCrc32( *StructProperty->Struct->GetName().ToLower(), Checksum );
+
+		GDebugChecksumIndent++;
+
+		Checksum = SortedStructFieldsChecksum( StructProperty->Struct, Checksum );
+
+		GDebugChecksumIndent--;
 	}
 
-	return FCrc::StrCrc32( *NetField->GetName(), Checksum );
+	return Checksum;
 }
 
 uint32 FNetGUIDCache::GetClassNetworkChecksum( const UClass* Class )
 {
-	uint32 Checksum = ClassChecksumMap.FindRef( Class );
+	uint32 Checksum = GDebugChecksum ? 0 : ClassChecksumMap.FindRef( Class );
 
 	if ( Checksum != 0 )
 	{
@@ -2136,11 +2193,29 @@ uint32 FNetGUIDCache::GetClassNetworkChecksum( const UClass* Class )
 		Checksum = GetClassNetworkChecksum( Class->GetSuperClass() );
 	}
 
-	Checksum = FCrc::StrCrc32( *Class->GetName(), Checksum );
+	GDebugChecksumIndent++;
+
+	Checksum = FCrc::StrCrc32( *Class->GetName().ToLower(), Checksum );
+
+	// Grab all of the net properties from the replicated fields
+	// (we don't need to check functions here, we do that on demand when considering them for backwards compatibility)
+	TArray< UProperty* > NetProperties;
 
 	for( int32 i = 0; i < Class->NetFields.Num(); i++ )
 	{
-		Checksum = GetNetFieldChecksum( Class->NetFields[i], Checksum );
+		UProperty* Property = Cast< UProperty >( Class->NetFields[i] );
+
+		if ( Property != NULL )
+		{
+			NetProperties.Add( Property );
+		}
+	}
+
+	SortProperties( NetProperties );
+
+	for( int32 i = 0; i < NetProperties.Num(); i++ )
+	{
+		Checksum = GetPropertyChecksum( NetProperties[i], Checksum );
 	}
 
 	// Can't be 0
@@ -2161,15 +2236,26 @@ uint32 FNetGUIDCache::GetNetworkChecksum( const UObject* Obj )
 		return 0;
 	}
 
-	// If Obj is already a class, we can use that directly
-	const UClass* Class = Cast< UClass >( Obj );
+	GDebugChecksum = false;
+	GDebugChecksumIndent = 0;
 
-	if ( Class != NULL )
+#if 0
+	if ( Obj->GetName().Contains( TEXT( "Default__ShooterGameState" ) ) )
 	{
-		return GetClassNetworkChecksum( Class );
+		GDebugChecksum = true;
+
+		UE_LOG( LogNetPackageMap, Warning, TEXT( "OBJECT: %s" ), *Obj->GetName() );
 	}
+#endif
 
-	return GetClassNetworkChecksum( Obj->GetClass() );
+	// If Obj is already a class, we can use that directly
+	const UClass* Class = Cast< UClass >( Obj );
+
+	const uint32 Checksum = ( Class == NULL ) ? GetClassNetworkChecksum( Obj->GetClass() ) : GetClassNetworkChecksum( Class );
+
+	GDebugChecksum = false;
+
+	return Checksum;
 }
 
 bool FNetGUIDCache::ShouldIgnorePackageMismatch() const
diff --git a/Engine/Source/Runtime/Engine/Private/RepLayout.cpp b/Engine/Source/Runtime/Engine/Private/RepLayout.cpp
--- a/Engine/Source/Runtime/Engine/Private/RepLayout.cpp
+++ b/Engine/Source/Runtime/Engine/Private/RepLayout.cpp
@@ -2149,6 +2149,12 @@ int32 FRepLayout::InitFromProperty_r( UProperty * Property, int32 Offset, int32
 		{
 			FORCEINLINE bool operator()( UProperty & A, UProperty & B ) const
 			{
+				// Ensure stable sort
+				if ( A.GetOffset_ForGC() == B.GetOffset_ForGC() )
+				{
+					return A.GetName() < B.GetName();
+				}
+
 				return A.GetOffset_ForGC() < B.GetOffset_ForGC();
 			}
 		};
diff --git a/Engine/Source/Runtime/Engine/Private/UnrealNetwork.cpp b/Engine/Source/Runtime/Engine/Private/UnrealNetwork.cpp
--- a/Engine/Source/Runtime/Engine/Private/UnrealNetwork.cpp
+++ b/Engine/Source/Runtime/Engine/Private/UnrealNetwork.cpp
@@ -14,9 +14,10 @@ enum ENetworkVersionHistory
 	HISTORY_INTERNAL_ACK		= 3,				// We no longer save packet/channel sequence in stream. We can derive this for 100% reliable connections.
 	HISTORY_REPLAY_CHECKSUMS	= 4,				// We now save replay checksums into stream for backwards compatibility checks
 	HISTORY_REPLAY_CHECKSUMS2	= 5,				// Saving package and network checksum to packagemap, no longer save checksum in BeginContentBlockHeader
+	HISTORY_REPLAY_CHECKSUMS3	= 6,				// Various network checksum fixes
 };
 
-const uint32 FNetworkVersion::InternalProtocolVersion = HISTORY_REPLAY_CHECKSUMS2;
+const uint32 FNetworkVersion::InternalProtocolVersion = HISTORY_REPLAY_CHECKSUMS3;
 
 uint32 FNetworkVersion::GetLocalNetworkVersion()
 {
@@ -41,7 +42,7 @@ uint32 FNetworkVersion::GetLocalNetworkVersion()
 	// Hash with internal protocol version
 	uint32 LocalNetworkVersion = FCrc::MemCrc32( &InternalProtocolVersion, sizeof( InternalProtocolVersion ), VersionHash );
 
-#if !(UE_BUILD_SHIPPING || UE_BUILD_TEST)
+#if 0//!(UE_BUILD_SHIPPING || UE_BUILD_TEST)	// DISABLED FOR NOW, MESSES UP COPIED BUILDS
 	if ( !GEngineVersion.IsPromotedBuild() )
 	{
 		// Further hash with machine id if this is a non promoted build
