Rendering: Deleted dead class FFontMaterialRenderProxy

diff --git a/Engine/Source/Runtime/Engine/Private/Materials/MaterialShared.cpp b/Engine/Source/Runtime/Engine/Private/Materials/MaterialShared.cpp
--- a/Engine/Source/Runtime/Engine/Private/Materials/MaterialShared.cpp
+++ b/Engine/Source/Runtime/Engine/Private/Materials/MaterialShared.cpp
@@ -1696,43 +1696,6 @@ bool FLightingDensityMaterialRenderProxy::GetVectorValue(const FName ParameterNa
 	return FColoredMaterialRenderProxy::GetVectorValue(ParameterName, OutValue, Context);
 }
 
-/*-----------------------------------------------------------------------------
-	FFontMaterialRenderProxy
------------------------------------------------------------------------------*/
-
-const class FMaterial* FFontMaterialRenderProxy::GetMaterial(ERHIFeatureLevel::Type InFeatureLevel) const
-{
-	return Parent->GetMaterial(InFeatureLevel);
-}
-
-bool FFontMaterialRenderProxy::GetVectorValue(const FName ParameterName, FLinearColor* OutValue, const FMaterialRenderContext& Context) const
-{
-	return Parent->GetVectorValue(ParameterName, OutValue, Context);
-}
-
-bool FFontMaterialRenderProxy::GetScalarValue(const FName ParameterName, float* OutValue, const FMaterialRenderContext& Context) const
-{
-	return Parent->GetScalarValue(ParameterName, OutValue, Context);
-}
-
-bool FFontMaterialRenderProxy::GetTextureValue(const FName ParameterName,const UTexture** OutValue, const FMaterialRenderContext& Context) const
-{
-	// find the matching font parameter
-	if( ParameterName == FontParamName &&
-		Font->Textures.IsValidIndex(FontPage) )
-	{
-		// use the texture page from the font specified for the parameter
-		UTexture2D* Texture = Font->Textures[FontPage];
-		if( Texture && Texture->Resource )
-		{
-			*OutValue = Texture;
-			return true;
-		}		
-	}
-	// try parent if not valid parameter
-	return Parent->GetTextureValue(ParameterName,OutValue,Context);
-}
-
 /*-----------------------------------------------------------------------------
 	FOverrideSelectionColorMaterialRenderProxy
 -----------------------------------------------------------------------------*/
diff --git a/Engine/Source/Runtime/Engine/Public/MaterialShared.h b/Engine/Source/Runtime/Engine/Public/MaterialShared.h
--- a/Engine/Source/Runtime/Engine/Public/MaterialShared.h
+++ b/Engine/Source/Runtime/Engine/Public/MaterialShared.h
@@ -1385,39 +1385,6 @@ public:
 	virtual bool GetVectorValue(const FName ParameterName, FLinearColor* OutValue, const FMaterialRenderContext& Context) const;
 };
 
-/**
-* A material render proxy for font rendering
-*/
-class FFontMaterialRenderProxy : public FMaterialRenderProxy
-{
-public:
-
-	/** parent material instance for fallbacks */
-	const FMaterialRenderProxy* const Parent;
-	/** font which supplies the texture pages */
-	const class UFont* Font;
-	/** index to the font texture page to use by the intance */
-	const int32 FontPage;
-	/** font parameter name for finding the matching parameter */
-	const FName& FontParamName;
-
-	/** Initialization constructor. */
-	FFontMaterialRenderProxy(const FMaterialRenderProxy* InParent,const class UFont* InFont,const int32 InFontPage, const FName& InFontParamName)
-	:	Parent(InParent)
-	,	Font(InFont)
-	,	FontPage(InFontPage)
-	,	FontParamName(InFontParamName)
-	{
-		check(Parent);
-		check(Font);
-	}
-
-	// FMaterialRenderProxy interface.
-	virtual const class FMaterial* GetMaterial(ERHIFeatureLevel::Type InFeatureLevel) const;
-	virtual bool GetVectorValue(const FName ParameterName, FLinearColor* OutValue, const FMaterialRenderContext& Context) const;
-	virtual bool GetScalarValue(const FName ParameterName, float* OutValue, const FMaterialRenderContext& Context) const;
-	virtual bool GetTextureValue(const FName ParameterName,const UTexture** OutValue, const FMaterialRenderContext& Context) const;
-};
 
 /**
  * A material render proxy which overrides the selection color
