Keeping Actor root components from trying to attach to themselves in FixupNativeActorComponents.

diff --git a/Engine/Source/Runtime/Engine/Private/Actor.cpp b/Engine/Source/Runtime/Engine/Private/Actor.cpp
--- a/Engine/Source/Runtime/Engine/Private/Actor.cpp
+++ b/Engine/Source/Runtime/Engine/Private/Actor.cpp
@@ -2368,7 +2368,10 @@ static USceneComponent* FixupNativeActorComponents(AActor* Actor)
 	// create its own root, and nest native components under it)
 	for (USceneComponent* Component : SceneComponents)
 	{
-		if ((Component == nullptr) || (Component->AttachParent != nullptr) || (Component->CreationMethod != EComponentCreationMethod::Native))
+		if ((Component == nullptr) || 
+			(Component->AttachParent != nullptr) || 
+			(Component->CreationMethod != EComponentCreationMethod::Native) || 
+			(Component == SceneRootComponent))
 		{
 			continue;
 		}
