Factor in multicast functions when building checksum for network compatibility

diff --git a/Engine/Source/Runtime/Engine/Private/PackageMapClient.cpp b/Engine/Source/Runtime/Engine/Private/PackageMapClient.cpp
--- a/Engine/Source/Runtime/Engine/Private/PackageMapClient.cpp
+++ b/Engine/Source/Runtime/Engine/Private/PackageMapClient.cpp
@@ -2128,7 +2128,7 @@ static uint32 SortedStructFieldsChecksum( const UStruct* Struct, uint32 Checksum
 	return Checksum;
 }
 
-static bool GDebugChecksum = false;
+static bool GDebugChecksum		= false;
 static int GDebugChecksumIndent = 0;
 
 static uint32 GetPropertyChecksum( const UProperty* Property, uint32 Checksum )
@@ -2197,11 +2197,12 @@ uint32 FNetGUIDCache::GetClassNetworkChecksum( const UClass* Class )
 
 	Checksum = FCrc::StrCrc32( *Class->GetName().ToLower(), Checksum );
 
-	// Grab all of the net properties from the replicated fields
-	// (we don't need to check functions here, we do that on demand when considering them for backwards compatibility)
 	TArray< UProperty* > NetProperties;
+	TArray< UFunction* > MulticastFunctions;
+
+	// Grab all of the net properties from the replicated fields
 
-	for( int32 i = 0; i < Class->NetFields.Num(); i++ )
+	for ( int32 i = 0; i < Class->NetFields.Num(); i++ )
 	{
 		UProperty* Property = Cast< UProperty >( Class->NetFields[i] );
 
@@ -2209,15 +2210,57 @@ uint32 FNetGUIDCache::GetClassNetworkChecksum( const UClass* Class )
 		{
 			NetProperties.Add( Property );
 		}
+
+		UFunction* Function = Cast< UFunction >( Class->NetFields[i] );
+
+		if ( Function != NULL && ( Function->FunctionFlags & FUNC_NetMulticast ) )
+		{
+			MulticastFunctions.Add( Function );
+		}
 	}
 
+	// Sort properties by offset/name
 	SortProperties( NetProperties );
 
-	for( int32 i = 0; i < NetProperties.Num(); i++ )
+	for ( int32 i = 0; i < NetProperties.Num(); i++ )
 	{
 		Checksum = GetPropertyChecksum( NetProperties[i], Checksum );
 	}
 
+	struct FCompareUFieldNames
+	{
+		FORCEINLINE bool operator()( UField& A, UField& B ) const
+		{
+			return A.GetName() < B.GetName();
+		}
+	};
+
+	// Sort multi-cast functions by name
+	Sort( MulticastFunctions.GetData(), MulticastFunctions.Num(), FCompareUFieldNames() );
+
+	// Grab all multi-cast RPC's (we don't need to worry about client RPC's, except for our spectator player controller, which we'll special case)
+	for ( UFunction* Function : MulticastFunctions )
+	{
+		// Evolve checksum on function name
+		Checksum = FCrc::StrCrc32( *Function->GetName().ToLower(), Checksum );
+
+		TArray< UProperty * > Parms;
+
+		for ( TFieldIterator< UProperty > It( Function ); It && ( It->PropertyFlags & ( CPF_Parm | CPF_ReturnParm ) ) == CPF_Parm; ++It )
+		{
+			Parms.Add( *It );
+		}
+		
+		// Sort parameters by offset/name
+		SortProperties( Parms );
+
+		// Evolve checksum on sorted function parameters
+		for ( UProperty* Parm : Parms )
+		{
+			Checksum = GetPropertyChecksum( Parm, Checksum );
+		}
+	}
+
 	// Can't be 0
 	if ( Checksum == 0 )
 	{
diff --git a/Engine/Source/Runtime/Engine/Private/UnrealNetwork.cpp b/Engine/Source/Runtime/Engine/Private/UnrealNetwork.cpp
--- a/Engine/Source/Runtime/Engine/Private/UnrealNetwork.cpp
+++ b/Engine/Source/Runtime/Engine/Private/UnrealNetwork.cpp
@@ -15,9 +15,10 @@ enum ENetworkVersionHistory
 	HISTORY_REPLAY_CHECKSUMS	= 4,				// We now save replay checksums into stream for backwards compatibility checks
 	HISTORY_REPLAY_CHECKSUMS2	= 5,				// Saving package and network checksum to packagemap, no longer save checksum in BeginContentBlockHeader
 	HISTORY_REPLAY_CHECKSUMS3	= 6,				// Various network checksum fixes
+	HISTORY_REPLAY_CHECKSUMS4	= 7,				// Added multi-cast RPC's to network checksum
 };
 
-const uint32 FNetworkVersion::InternalProtocolVersion = HISTORY_REPLAY_CHECKSUMS3;
+const uint32 FNetworkVersion::InternalProtocolVersion = HISTORY_REPLAY_CHECKSUMS4;
 
 uint32 FNetworkVersion::GetLocalNetworkVersion()
 {
@@ -25,7 +26,7 @@ uint32 FNetworkVersion::GetLocalNetworkVersion()
 	{
 		const uint32 LocalNetworkVersion = GetLocalNetworkVersionOverride.Execute();
 
-		UE_LOG( LogNet, Log, TEXT( "GetLocalNetworkVersionOverride: LocalNetworkVersion: %i" ), LocalNetworkVersion );
+		UE_LOG( LogNet, Log, TEXT( "GetLocalNetworkVersionOverride: LocalNetworkVersion: %u" ), LocalNetworkVersion );
 
 		return LocalNetworkVersion;
 	}
