Added ToString helper to FHitResut

diff --git a/Engine/Source/Runtime/Engine/Classes/Engine/EngineTypes.h b/Engine/Source/Runtime/Engine/Classes/Engine/EngineTypes.h
--- a/Engine/Source/Runtime/Engine/Classes/Engine/EngineTypes.h
+++ b/Engine/Source/Runtime/Engine/Classes/Engine/EngineTypes.h
@@ -1621,6 +1621,8 @@ struct ENGINE_API FHitResult
 		Result.ImpactNormal = -Result.ImpactNormal;
 		return Result;
 	}
+
+	FString ToString() const;
 };
 
 
diff --git a/Engine/Source/Runtime/Engine/Private/EngineTypes.cpp b/Engine/Source/Runtime/Engine/Private/EngineTypes.cpp
--- a/Engine/Source/Runtime/Engine/Private/EngineTypes.cpp
+++ b/Engine/Source/Runtime/Engine/Private/EngineTypes.cpp
@@ -152,3 +152,25 @@ USceneComponent* FComponentReference::GetComponent(AActor* OwningActor) const
 
 	return Result;
 }
+
+FString FHitResult::ToString() const
+{
+	return FString::Printf(TEXT("bBlockingHit:%s bStartPenetrating:%s Time:%f Location:%s ImpactPoint:%s Normal:%s ImpactNormal:%s TraceStart:%s TraceEnd:%s PenetrationDepth:%f Item:%d PhysMaterial:%s Actor:%s Component:%s BoneName:%s FaceIndex:%d"),
+		bBlockingHit == true ? TEXT("True") : TEXT("False"),
+		bStartPenetrating == true ? TEXT("True") : TEXT("False"),
+		Time,
+		*Location.ToString(),
+		*ImpactPoint.ToString(),
+		*Normal.ToString(),
+		*ImpactNormal.ToString(),
+		*TraceStart.ToString(),
+		*TraceEnd.ToString(),
+		PenetrationDepth,
+		Item,
+		PhysMaterial.IsValid() ? *PhysMaterial->GetName() : TEXT("None"),
+		Actor.IsValid() ? *Actor->GetName() : TEXT("None"),
+		Component.IsValid() ? *Component->GetName() : TEXT("None"),
+		BoneName.IsValid() ? *BoneName.ToString() : TEXT("None"),
+		FaceIndex);
+}
+
