UE-14026 Improve BoxExtent comment on BoxComponent

diff --git a/Engine/Source/Runtime/Engine/Classes/Components/BoxComponent.h b/Engine/Source/Runtime/Engine/Classes/Components/BoxComponent.h
--- a/Engine/Source/Runtime/Engine/Classes/Components/BoxComponent.h
+++ b/Engine/Source/Runtime/Engine/Classes/Components/BoxComponent.h
@@ -14,14 +14,14 @@ class ENGINE_API UBoxComponent : public UShapeComponent
 	GENERATED_UCLASS_BODY()
 
 protected:
-	/** The extents of the box **/
+	/** The extents (radii dimensions) of the box **/
 	UPROPERTY(EditAnywhere, BlueprintReadOnly, export, Category=Shape)
 	FVector BoxExtent;
 
 public:
 	/** 
 	 * Change the box extent size. This is the unscaled size, before component scale is applied.
-	 * @param	InBoxExtent: new extent for the box.
+	 * @param	InBoxExtent: new extent (radius) for the box.
 	 * @param	bUpdateOverlaps: if true and this shape is registered and collides, updates touching array for owner actor.
 	 */
 	UFUNCTION(BlueprintCallable, Category="Components|Box")
