Fix non-unity compile

diff --git a/Engine/Source/Runtime/Engine/Private/EngineTypes.cpp b/Engine/Source/Runtime/Engine/Private/EngineTypes.cpp
--- a/Engine/Source/Runtime/Engine/Private/EngineTypes.cpp
+++ b/Engine/Source/Runtime/Engine/Private/EngineTypes.cpp
@@ -1,6 +1,7 @@
 // Copyright 1998-2015 Epic Games, Inc. All Rights Reserved.
 
 #include "EnginePrivate.h"
+#include "PhysicalMaterials/PhysicalMaterial.h"
 
 UEngineBaseTypes::UEngineBaseTypes(const FObjectInitializer& ObjectInitializer)
 	: Super(ObjectInitializer)
