PR #1074: Skeletal Components with no bone specified now use ComponentToWorld instead of FTransform::Identity, which stops them from warping to World Origin.

diff --git a/Engine/Source/Runtime/Engine/Private/PhysicsEngine/PhysicsConstraintComponent.cpp b/Engine/Source/Runtime/Engine/Private/PhysicsEngine/PhysicsConstraintComponent.cpp
--- a/Engine/Source/Runtime/Engine/Private/PhysicsEngine/PhysicsConstraintComponent.cpp
+++ b/Engine/Source/Runtime/Engine/Private/PhysicsEngine/PhysicsConstraintComponent.cpp
@@ -87,10 +87,15 @@ UPrimitiveComponent* UPhysicsConstraintComponent::GetComponentInternal(EConstrai
 
 FTransform UPhysicsConstraintComponent::GetBodyTransformInternal(EConstraintFrame::Type Frame, FName InBoneName) const
 {
-	FTransform ResultTM = FTransform::Identity;
-
-	UPrimitiveComponent* PrimComp  = GetComponentInternal(Frame);
-
+	UPrimitiveComponent* PrimComp = GetComponentInternal(Frame);
+	if(!PrimComp)
+	{
+		return FTransform::Identity;
+	}
+	  
+	//Use ComponentToWorld by default for all components
+	FTransform ResultTM = PrimComp->ComponentToWorld;
+		
 	// Skeletal case
 	USkeletalMeshComponent* SkelComp = Cast<USkeletalMeshComponent>(PrimComp);
 	if(SkelComp != NULL)
@@ -101,11 +106,6 @@ FTransform UPhysicsConstraintComponent::GetBodyTransformInternal(EConstraintFram
 			ResultTM = SkelComp->GetBoneTransform(BoneIndex);
 		}
 	}
-	// Non skeletal (ie single body) case.
-	else if(PrimComp != NULL)
-	{
-		ResultTM = PrimComp->ComponentToWorld;
-	}
 
 	return ResultTM;
 }
