Added support for binary serialization of physx rigid actors. Helps components that have many body instances like InstancedStaticMeshComponent

diff --git a/Engine/Source/Runtime/Core/Public/UObject/ObjectVersion.h b/Engine/Source/Runtime/Core/Public/UObject/ObjectVersion.h
--- a/Engine/Source/Runtime/Core/Public/UObject/ObjectVersion.h
+++ b/Engine/Source/Runtime/Core/Public/UObject/ObjectVersion.h
@@ -494,6 +494,8 @@ enum EUnrealEngineObjectUE4Version
 	VER_UE4_ADD_EDITOR_VIEWS,
 	// Updated foliage to work with either FoliageType assets or blueprint classes
 	VER_UE4_FOLIAGE_WITH_ASSET_OR_CLASS,
+	// Allows PhysicsSerializer to serialize shapes and actors for faster load times
+	VER_UE4_BODYINSTANCE_BINARY_SERIALIZATION,
 
 	// -----<new versions can be added before this line>-------------------------------------------------
 	// - this needs to be the last line (see note below)
diff --git a/Engine/Source/Runtime/Engine/Classes/Components/InstancedStaticMeshComponent.h b/Engine/Source/Runtime/Engine/Classes/Components/InstancedStaticMeshComponent.h
--- a/Engine/Source/Runtime/Engine/Classes/Components/InstancedStaticMeshComponent.h
+++ b/Engine/Source/Runtime/Engine/Classes/Components/InstancedStaticMeshComponent.h
@@ -151,6 +151,10 @@ public:
 	/** Physics representation of the instance bodies */
 	TArray<FBodyInstance*> InstanceBodies;
 
+	/** Serialization of all the InstanceBodies. Helps speed up physics creation time*/
+	UPROPERTY()
+	class UPhysicsSerializer* PhysicsSerializer;
+
 	// Begin UActorComponent interface 
 	virtual FActorComponentInstanceData* GetComponentInstanceData() const override;
 	virtual FName GetComponentInstanceDataType() const override;
diff --git a/Engine/Source/Runtime/Engine/Classes/PhysicsEngine/BodyInstance.h b/Engine/Source/Runtime/Engine/Classes/PhysicsEngine/BodyInstance.h
--- a/Engine/Source/Runtime/Engine/Classes/PhysicsEngine/BodyInstance.h
+++ b/Engine/Source/Runtime/Engine/Classes/PhysicsEngine/BodyInstance.h
@@ -370,6 +370,14 @@ public:
 	TSharedPtr<TArray<ANSICHAR>> CharDebugName;
 #endif	//WITH_PHYSX
 
+	/** Internal use. Physics-engine id of the actor used during serialization. Needs to be outside the ifdef for serialization purposes*/
+	UPROPERTY()
+	uint64 RigidActorSyncId;
+
+	/** Internal use. Physics-engine id of the actor used during serialization.  Needs to be outside the ifdef for serialization purposes*/
+	UPROPERTY()
+	uint64 RigidActorAsyncId;
+
 	/** Per instance data used to initialize dynamic instances */
 	/** Initial physx velocity to apply to dynamic instances */
 	FVector InitialLinearVelocity;
@@ -419,7 +427,7 @@ public:
 	 *	@param InRBScene The physics scene to put the bodies into
 	 *	@param InAggregate Aggregate to place the bodies into, be aware for PhysX we can only have 128 elements in an aggregate!
 	 */
-	static void InitBodies(TArray<FBodyInstance*>& Bodies, TArray<FTransform>& Transforms, class UBodySetup* BodySetup, class UPrimitiveComponent* PrimitiveComp, class FPhysScene* InRBScene, PhysXAggregateType InAggregate = NULL, bool Defer = false);
+	static void InitBodies(TArray<FBodyInstance*>& Bodies, TArray<FTransform>& Transforms, class UBodySetup* BodySetup, class UPrimitiveComponent* PrimitiveComp, class FPhysScene* InRBScene, PhysXAggregateType InAggregate = NULL, bool Defer = false, class UPhysicsSerializer* PhysicsSerializer = nullptr);
 
 	/** Validate a body transform, outputting debug info
 	 *	@param Transform Transform to debug
@@ -625,9 +633,15 @@ public:
 	/** Find the correct PhysicalMaterial for simple geometry on this body */
 	UPhysicalMaterial* GetSimplePhysicalMaterial() const;
 
+	/** Find the correct PhysicalMaterial for simple geometry on a given body and owner. This is really for internal use during serialization */
+	static UPhysicalMaterial* GetSimplePhysicalMaterial(const FBodyInstance* BodyInstance, TWeakObjectPtr<UPrimitiveComponent> Owner, TWeakObjectPtr<UBodySetup> BodySetupPtr);
+
 	/** Get the complex PhysicalMaterial array for this body */
 	TArray<UPhysicalMaterial*> GetComplexPhysicalMaterials() const;
 
+	/** Find the correct PhysicalMaterial for simple geometry on a given body and owner. This is really for internal use during serialization */
+	static void GetComplexPhysicalMaterials(const FBodyInstance* BodyInstance, TWeakObjectPtr<UPrimitiveComponent> Owner, TArray<UPhysicalMaterial*>& OutPhysicalMaterials);
+
 	/** Get the complex PhysicalMaterials for this body */
 	void GetComplexPhysicalMaterials(TArray<UPhysicalMaterial*> &PhysMaterials) const;
 
@@ -1006,6 +1020,7 @@ private:
 	friend class UCollisionProfile;
 	friend class FBodyInstanceCustomization;
 	friend struct FInitBodiesHelper;
+	friend class FDerivedDataPhysXBinarySerializer;
 
 #if WITH_BOX2D
 
diff --git a/Engine/Source/Runtime/Engine/Private/InstancedStaticMesh.cpp b/Engine/Source/Runtime/Engine/Private/InstancedStaticMesh.cpp
--- a/Engine/Source/Runtime/Engine/Private/InstancedStaticMesh.cpp
+++ b/Engine/Source/Runtime/Engine/Private/InstancedStaticMesh.cpp
@@ -10,6 +10,7 @@
 #include "Components/InstancedStaticMeshComponent.h"
 #include "InstancedStaticMesh.h"
 #include "../../Renderer/Private/ScenePrivate.h"
+#include "PhysicsSerializer.h"
 
 const int32 InstancedStaticMeshMaxTexCoord = 8;
 
@@ -729,6 +730,8 @@ UInstancedStaticMeshComponent::UInstancedStaticMeshComponent(const FObjectInitia
 {
 	Mobility = EComponentMobility::Movable;
 	BodyInstance.bSimulatePhysics = false;
+
+	PhysicsSerializer = ObjectInitializer.CreateDefaultSubobject<UPhysicsSerializer>(this, TEXT("PhysicsSerializer"));
 }
 
 #if WITH_EDITOR
@@ -997,6 +1000,15 @@ void UInstancedStaticMeshComponent::CreateAllInstanceBodies()
 		    Instance->CopyBodyInstancePropertiesFrom(&BodyInstance);
 		    Instance->InstanceBodyIndex = i; // Set body index 
 		    Instance->bAutoWeld = false;
+
+#if WITH_PHYSX
+			Instance->RigidActorSyncId = i+1;
+
+			if(GetWorld()->GetPhysicsScene()->HasAsyncScene())
+			{
+				Instance->RigidActorAsyncId = Instance->RigidActorSyncId + NumBodies;
+			}
+#endif
     
 		    // make sure we never enable bSimulatePhysics for ISMComps
 		    Instance->bSimulatePhysics = false;
@@ -1004,7 +1016,20 @@ void UInstancedStaticMeshComponent::CreateAllInstanceBodies()
 
 		if (NumBodies > 0)
 		{
-			FBodyInstance::InitBodies(InstanceBodies, Transforms, BodySetup, this, GetWorld()->GetPhysicsScene(), nullptr, true);
+			TArray<UBodySetup*> BodySetups;
+			TArray<UPhysicalMaterial*> PhysicalMaterials;
+
+			BodySetups.Add(BodySetup);
+			TWeakObjectPtr<UPrimitiveComponent> WeakSelfPtr(this);
+			FBodyInstance::GetComplexPhysicalMaterials(&BodyInstance, WeakSelfPtr, PhysicalMaterials);
+			PhysicalMaterials.Add(FBodyInstance::GetSimplePhysicalMaterial(&BodyInstance, WeakSelfPtr, TWeakObjectPtr<UBodySetup>(BodySetup)));
+
+			PhysicsSerializer->CreatePhysicsData(BodySetups, PhysicalMaterials);
+			
+			FBodyInstance::InitBodies(InstanceBodies, Transforms, BodySetup, this, GetWorld()->GetPhysicsScene(), nullptr, true, PhysicsSerializer);
+
+			//Serialize physics data for fast path cooking
+			PhysicsSerializer->SerializePhysics(InstanceBodies, BodySetups, PhysicalMaterials);
 		}
 	}
 }
diff --git a/Engine/Source/Runtime/Engine/Private/PhysicsEngine/BodyInstance.cpp b/Engine/Source/Runtime/Engine/Private/PhysicsEngine/BodyInstance.cpp
--- a/Engine/Source/Runtime/Engine/Private/PhysicsEngine/BodyInstance.cpp
+++ b/Engine/Source/Runtime/Engine/Private/PhysicsEngine/BodyInstance.cpp
@@ -12,6 +12,9 @@
 	#include "../Collision/CollisionConversions.h"
 #endif // WITH_PHYSX
 
+#include "PhysDerivedData.h"
+#include "PhysicsSerializer.h"
+
 #define LOCTEXT_NAMESPACE "BodyInstance"
 
 #if WITH_BOX2D
@@ -257,6 +260,8 @@ FBodyInstance::FBodyInstance()
 , RigidActorSync(NULL)
 , RigidActorAsync(NULL)
 , BodyAggregate(NULL)
+, RigidActorSyncId(PX_SERIAL_OBJECT_ID_INVALID)
+, RigidActorAsyncId(PX_SERIAL_OBJECT_ID_INVALID)
 , InitialLinearVelocity(0.0f)
 , bWokenExternally(false)
 , PhysxUserData(this)
@@ -1056,7 +1061,7 @@ FVector GetInitialLinearVelocity(const AActor* OwningActor, bool& bComponentAwak
 
 struct FInitBodiesHelper
 {
-	FInitBodiesHelper(const TArray<FBodyInstance*>& InBodies, const TArray<FTransform>& InTransforms, class UBodySetup* InBodySetup, class UPrimitiveComponent* InPrimitiveComp, class FPhysScene* InInRBScene, FBodyInstance::PhysXAggregateType InInAggregate = NULL, bool InDefer = false)
+	FInitBodiesHelper(const TArray<FBodyInstance*>& InBodies, const TArray<FTransform>& InTransforms, class UBodySetup* InBodySetup, class UPrimitiveComponent* InPrimitiveComp, class FPhysScene* InInRBScene, FBodyInstance::PhysXAggregateType InInAggregate = NULL, bool InDefer = false, UPhysicsSerializer* InPhysicsSerializer = nullptr)
 	: Bodies(InBodies)
 	, Transforms(InTransforms)
 	, BodySetup(InBodySetup)
@@ -1064,6 +1069,7 @@ struct FInitBodiesHelper
 	, PhysScene(InInRBScene)
 	, InAggregate(InInAggregate)
 	, bDefer(InDefer)
+	, PhysicsSerializer(InPhysicsSerializer)
 	, DebugName(TEXT(""))
 	, InstanceBlendWeight(-1.f)
 	, bInstanceSimulatePhysics(false)
@@ -1096,6 +1102,7 @@ struct FInitBodiesHelper
 	class FPhysScene* PhysScene;
 	FBodyInstance::PhysXAggregateType InAggregate;
 	const bool bDefer;
+	UPhysicsSerializer* PhysicsSerializer;
 	
 	FString DebugName;
 	TSharedPtr<TArray<ANSICHAR>> PhysXName;
@@ -1117,6 +1124,7 @@ struct FInitBodiesHelper
 	void InitBodies_PhysX() const;
 	bool CreateShapesAndActors_PhysX(TArray<PxActor*>& PSyncActors, TArray<PxActor*>& PAsyncActors, TArray<PxActor*>& PDynamicActors, const bool bCanDefer, bool& bDynamicsUseAsyncScene) const;
 	physx::PxRigidActor* CreateActor_PhysX_AssumesLocked(FBodyInstance* Instance, const PxTransform& PTransform) const;
+	bool GetBinaryData_PhysX_AssumesLocked(FBodyInstance* Instance) const;
 	bool CreateShapes_PhysX_AssumesLocked(FBodyInstance* Instance, physx::PxRigidActor* PNewDynamic) const;
 	void AddActorsToScene_PhysX_AssumesLocked(TArray<PxActor*>& PSyncActors, TArray<PxActor*>& PAsyncActors, TArray<PxActor*>& PDynamicActors) const;
 #endif
@@ -1126,7 +1134,7 @@ struct FInitBodiesHelper
 #endif
 };
 
-void FBodyInstance::InitBodies(TArray<FBodyInstance*>& Bodies, TArray<FTransform>& Transforms, class UBodySetup* BodySetup, class UPrimitiveComponent* PrimitiveComp, class FPhysScene* InRBScene, PhysXAggregateType InAggregate /*= NULL*/, bool bDefer /*= false*/ )
+void FBodyInstance::InitBodies(TArray<FBodyInstance*>& Bodies, TArray<FTransform>& Transforms, class UBodySetup* BodySetup, class UPrimitiveComponent* PrimitiveComp, class FPhysScene* InRBScene, PhysXAggregateType InAggregate /*= NULL*/, bool bDefer /*= false*/, UPhysicsSerializer* PhysicsSerializer )
 {
 	SCOPE_CYCLE_COUNTER(STAT_InitBodies);
 
@@ -1134,7 +1142,7 @@ void FBodyInstance::InitBodies(TArray<FBodyInstance*>& Bodies, TArray<FTransform
 	check(InRBScene);
 	check(Bodies.Num() > 0);
 
-	FInitBodiesHelper InitBodiesHelper(Bodies, Transforms, BodySetup, PrimitiveComp, InRBScene, InAggregate, bDefer);
+	FInitBodiesHelper InitBodiesHelper(Bodies, Transforms, BodySetup, PrimitiveComp, InRBScene, InAggregate, bDefer, PhysicsSerializer);
 	
 #if WITH_PHYSX
 	InitBodiesHelper.InitBodies_PhysX();
@@ -1147,8 +1155,19 @@ void FBodyInstance::InitBodies(TArray<FBodyInstance*>& Bodies, TArray<FTransform
 
 #if WITH_PHYSX
 
+bool FInitBodiesHelper::GetBinaryData_PhysX_AssumesLocked(FBodyInstance* Instance) const
+{
+	Instance->RigidActorSync = PhysicsSerializer->GetRigidActor(Instance->RigidActorSyncId);
+	Instance->RigidActorAsync = PhysicsSerializer->GetRigidActor(Instance->RigidActorAsyncId);
+	
+	return Instance->RigidActorSync || Instance->RigidActorAsync;
+}
+
+
 physx::PxRigidActor* FInitBodiesHelper::CreateActor_PhysX_AssumesLocked(FBodyInstance* Instance, const PxTransform& PTransform) const
 {
+	physx::PxRigidActor* PNewDynamic = nullptr;
+
 	if (bStatic)
 	{
 		Instance->RigidActorSync = GPhysXSDK->createRigidStatic(PTransform);
@@ -1159,16 +1178,8 @@ physx::PxRigidActor* FInitBodiesHelper::CreateActor_PhysX_AssumesLocked(FBodyIns
 	}
 	else
 	{
-		physx::PxRigidActor* PNewDynamic = GPhysXSDK->createRigidDynamic(PTransform);
-		bool bDynamicsUseAsyncScene = Instance->UseAsyncScene(PhysScene);
-
-		// turn off gravity if desired
-		if (!Instance->bEnableGravity)
-		{
-			PNewDynamic->setActorFlag(PxActorFlag::eDISABLE_GRAVITY, true);
-		}
-
-		if (bDynamicsUseAsyncScene)
+		PNewDynamic = GPhysXSDK->createRigidDynamic(PTransform);
+		if (Instance->UseAsyncScene(PhysScene))
 		{
 			Instance->RigidActorAsync = PNewDynamic;
 		}
@@ -1176,11 +1187,9 @@ physx::PxRigidActor* FInitBodiesHelper::CreateActor_PhysX_AssumesLocked(FBodyIns
 		{
 			Instance->RigidActorSync = PNewDynamic;
 		}
-
-		return PNewDynamic;
 	}
 
-	return nullptr;
+		return PNewDynamic;
 }
 
 bool FInitBodiesHelper::CreateShapes_PhysX_AssumesLocked(FBodyInstance* Instance, physx::PxRigidActor* PNewDynamic) const
@@ -1302,29 +1311,66 @@ bool FInitBodiesHelper::CreateShapesAndActors_PhysX(TArray<PxActor*>& PSyncActor
 		}
 
 		Instance->PhysxUserData = FPhysxUserData(Instance);
-
-		physx::PxRigidActor* PNewDynamic = CreateActor_PhysX_AssumesLocked(Instance, U2PTransform(Transform));	
-		const bool bInitFail = CreateShapes_PhysX_AssumesLocked(Instance, PNewDynamic);
-
-		if (bInitFail)
+		
+		static FName PhysicsFormatName(FPlatformProperties::GetPhysicsFormat());
+		
+		physx::PxRigidActor* PNewDynamic = nullptr;
+		bool bFoundBinaryData = false;
+		if(PhysicsSerializer)
 		{
-			UE_LOG(LogPhysics, Log, TEXT("Init Instance %d of Primitive Component %s failed"), BodyIdx, *PrimitiveComp->GetName());
-			if (Instance->RigidActorSync)
-			{
-				Instance->RigidActorSync->release();
-				Instance->RigidActorSync = nullptr;
-			}
+			bFoundBinaryData = GetBinaryData_PhysX_AssumesLocked(Instance);
+		}
+		
+		if(!bFoundBinaryData)
+		{
+			PNewDynamic = CreateActor_PhysX_AssumesLocked(Instance, U2PTransform(Transform));	
+			const bool bInitFail = CreateShapes_PhysX_AssumesLocked(Instance, PNewDynamic);
 
-			if (Instance->RigidActorAsync)
+			if (bInitFail)
 			{
-				Instance->RigidActorAsync->release();
-				Instance->RigidActorAsync = nullptr;
+				UE_LOG(LogPhysics, Log, TEXT("Init Instance %d of Primitive Component %s failed"), BodyIdx, *PrimitiveComp->GetName());
+				if (Instance->RigidActorSync)
+				{
+					Instance->RigidActorSync->release();
+					Instance->RigidActorSync = nullptr;
+				}
+
+				if (Instance->RigidActorAsync)
+				{
+					Instance->RigidActorAsync->release();
+					Instance->RigidActorAsync = nullptr;
+				}
+
+				Instance->OwnerComponent = nullptr;
+				Instance->BodySetup = nullptr;
+
+				continue;
 			}
+		}
 
-			Instance->OwnerComponent = nullptr;
-			Instance->BodySetup = nullptr;
+		if (Instance->RigidActorSync)
+		{
+			Instance->RigidActorSync->userData = &Instance->PhysxUserData;
+			Instance->RigidActorSync->setName(Instance->CharDebugName.IsValid() ? Instance->CharDebugName->GetData() : nullptr);
+			check(FPhysxUserData::Get<FBodyInstance>(Instance->RigidActorSync->userData) == Instance && FPhysxUserData::Get<FBodyInstance>(Instance->RigidActorSync->userData)->OwnerComponent != NULL);
+		}
 
-			continue;
+		if (Instance->RigidActorAsync)
+		{
+			Instance->RigidActorAsync->userData = &Instance->PhysxUserData;
+			Instance->RigidActorAsync->setName(Instance->CharDebugName.IsValid() ? Instance->CharDebugName->GetData() : nullptr);
+			check(FPhysxUserData::Get<FBodyInstance>(Instance->RigidActorAsync->userData) == Instance && FPhysxUserData::Get<FBodyInstance>(Instance->RigidActorAsync->userData)->OwnerComponent != NULL);
+		}
+
+		//handle special stuff only dynamic actors care about
+		if(PNewDynamic)
+		{
+			// turn off gravity if desired
+			if (!Instance->bEnableGravity)
+			{
+				PNewDynamic->setActorFlag(PxActorFlag::eDISABLE_GRAVITY, true);
+				bDynamicsUseAsyncScene = Instance->UseAsyncScene(PhysScene);
+			}
 		}
 
 		if (bCanDefer)
@@ -1360,7 +1406,6 @@ bool FInitBodiesHelper::CreateShapesAndActors_PhysX(TArray<PxActor*>& PSyncActor
 			if (PNewDynamic)
 			{
 				PDynamicActors.Add(PNewDynamic);
-				bDynamicsUseAsyncScene = Instance->UseAsyncScene(PhysScene);
 			}
 
 			// Set the instance to added as we'll add it to the scene in a moment.
@@ -2945,12 +2990,14 @@ void FBodyInstance::SetPhysMaterialOverride( UPhysicalMaterial* NewPhysMaterial
 	UpdatePhysicalMaterials();
 }
 
-
-
-
 UPhysicalMaterial* FBodyInstance::GetSimplePhysicalMaterial() const
 {
-	check( GEngine->DefaultPhysMaterial != NULL );
+	return GetSimplePhysicalMaterial(this, OwnerComponent, BodySetup);
+}
+
+UPhysicalMaterial* FBodyInstance::GetSimplePhysicalMaterial(const FBodyInstance* BodyInstance, TWeakObjectPtr<UPrimitiveComponent> OwnerComp, TWeakObjectPtr<UBodySetup> BodySetupPtr)
+{
+	check(GEngine->DefaultPhysMaterial != NULL);
 
 	// Find the PhysicalMaterial we need to apply to the physics bodies.
 	// (LOW priority) Engine Mat, Material PhysMat, BodySetup Mat, Component Override, Body Override (HIGH priority)
@@ -2958,27 +3005,27 @@ UPhysicalMaterial* FBodyInstance::GetSimplePhysicalMaterial() const
 	UPhysicalMaterial* ReturnPhysMaterial = NULL;
 
 	// BodyInstance override
-	if( PhysMaterialOverride != NULL)	
+	if (BodyInstance->PhysMaterialOverride != NULL)
 	{
-		ReturnPhysMaterial = PhysMaterialOverride;
+		ReturnPhysMaterial = BodyInstance->PhysMaterialOverride;
 		check(!ReturnPhysMaterial || ReturnPhysMaterial->IsValidLowLevel());
 	}
 	// Component override
-	else if( OwnerComponent.IsValid() && OwnerComponent->BodyInstance.PhysMaterialOverride != NULL )
+	else if (OwnerComp.IsValid() && OwnerComp->BodyInstance.PhysMaterialOverride != NULL)
 	{
-		ReturnPhysMaterial = OwnerComponent->BodyInstance.PhysMaterialOverride;
+		ReturnPhysMaterial = OwnerComp->BodyInstance.PhysMaterialOverride;
 		check(!ReturnPhysMaterial || ReturnPhysMaterial->IsValidLowLevel());
 	}
 	// BodySetup
-	else if( BodySetup.IsValid() && BodySetup->PhysMaterial != NULL )
+	else if (BodySetupPtr.IsValid() && BodySetupPtr->PhysMaterial != NULL)
 	{
-		ReturnPhysMaterial = BodySetup->PhysMaterial;
+		ReturnPhysMaterial = BodySetupPtr->PhysMaterial;
 		check(!ReturnPhysMaterial || ReturnPhysMaterial->IsValidLowLevel());
 	}
 	else
 	{
 		// See if the Material has a PhysicalMaterial
-		UMeshComponent* MeshComp = Cast<UMeshComponent>(OwnerComponent.Get());
+		UMeshComponent* MeshComp = Cast<UMeshComponent>(OwnerComp.Get());
 		UPhysicalMaterial* PhysMatFromMaterial = NULL;
 		if (MeshComp != NULL)
 		{
@@ -3012,27 +3059,36 @@ TArray<UPhysicalMaterial*> FBodyInstance::GetComplexPhysicalMaterials() const
 	return PhysMaterials;
 }
 
-void FBodyInstance::GetComplexPhysicalMaterials(TArray<UPhysicalMaterial*> &PhysMaterials) const
+void FBodyInstance::GetComplexPhysicalMaterials(TArray<UPhysicalMaterial*>& PhysMaterials) const
 {
+	GetComplexPhysicalMaterials(this, OwnerComponent, PhysMaterials);
+}
+
+
+void FBodyInstance::GetComplexPhysicalMaterials(const FBodyInstance*, TWeakObjectPtr<UPrimitiveComponent> OwnerComp, TArray<UPhysicalMaterial*>& OutPhysicalMaterials)
+{
+	//NOTE: Binary serialization of physics objects depends on this logic being the same between runs.
+	//      If you make any changes to the returned object update the GUID of the physx binary serializer (FDerivedDataPhysXBinarySerializer)
+
 	check(GEngine->DefaultPhysMaterial != NULL);
 	// See if the Material has a PhysicalMaterial
-	UPrimitiveComponent* PrimComp = OwnerComponent.Get();
-	if(PrimComp)
+	UPrimitiveComponent* PrimComp = OwnerComp.Get();
+	if (PrimComp)
 	{
 		const int32 NumMaterials = PrimComp->GetNumMaterials();
-		PhysMaterials.SetNum(NumMaterials);
+		OutPhysicalMaterials.SetNum(NumMaterials);
 
-		for(int32 MatIdx = 0; MatIdx < NumMaterials; MatIdx++)
+		for (int32 MatIdx = 0; MatIdx < NumMaterials; MatIdx++)
 		{
 			UPhysicalMaterial* PhysMat = GEngine->DefaultPhysMaterial;
 			UMaterialInterface* Material = PrimComp->GetMaterial(MatIdx);
-			if(Material)
+			if (Material)
 			{
 				PhysMat = Material->GetPhysicalMaterial();
 			}
 
 			check(PhysMat != NULL);
-			PhysMaterials[MatIdx] = PhysMat;
+			OutPhysicalMaterials[MatIdx] = PhysMat;
 		}
 	}
 }
diff --git a/Engine/Source/Runtime/Engine/Private/PhysicsEngine/PhysDerivedData.cpp b/Engine/Source/Runtime/Engine/Private/PhysicsEngine/PhysDerivedData.cpp
--- a/Engine/Source/Runtime/Engine/Private/PhysicsEngine/PhysDerivedData.cpp
+++ b/Engine/Source/Runtime/Engine/Private/PhysicsEngine/PhysDerivedData.cpp
@@ -53,7 +53,7 @@ bool FDerivedDataPhysXCooker::Build( TArray<uint8>& OutData )
 	check(Cooker != NULL);
 
 	FMemoryWriter Ar( OutData );
-	uint8 bLittleEndian = PLATFORM_LITTLE_ENDIAN;
+	uint8 bLittleEndian = PLATFORM_LITTLE_ENDIAN;	//TODO: We should pass the target platform into this function and write it. Then swap the endian on the writer so the reader doesn't have to do it at runtime
 	int32 NumConvexElementsCooked = 0;
 	int32 NumMirroredElementsCooked = 0;
 	uint8 bTriMeshCooked = false;
@@ -65,6 +65,8 @@ bool FDerivedDataPhysXCooker::Build( TArray<uint8>& OutData )
 	Ar << bTriMeshCooked;
 	Ar << bMirroredTriMeshCooked;
 
+	//TODO: we must save an id with convex and tri meshes for serialization. We must save this here and patch it up at runtime somehow
+
 	// Cook convex meshes, but only if we are not forcing complex collision to be used as simple collision as well
 	if( BodySetup && BodySetup->CollisionTraceFlag != CTF_UseComplexAsSimple && BodySetup->AggGeom.ConvexElems.Num() > 0 )
 	{
@@ -215,4 +217,43 @@ bool FDerivedDataPhysXCooker::BuildTriMesh( TArray<uint8>& OutData, bool bInMirr
 	return bResult;
 }
 
+FDerivedDataPhysXBinarySerializer::FDerivedDataPhysXBinarySerializer(FName InFormat, const TArray<FBodyInstance*>& InBodies, const TArray<class UBodySetup*>& InBodySetups, const TArray<class UPhysicalMaterial*>& InPhysicalMaterials, const FGuid& InGuid)
+: Bodies(InBodies)
+, BodySetups(InBodySetups)
+, PhysicalMaterials(InPhysicalMaterials)
+, Format(InFormat)
+, DataGuid(InGuid)
+{
+	InitSerializer();
+}
+
+bool FDerivedDataPhysXBinarySerializer::Build(TArray<uint8>& OutData)
+{
+	FMemoryWriter Ar(OutData);
+	
+	//do not serialize anything before the physx data. This is important because physx requires specific alignment. For that to work the physx data must come first in the archive
+	SerializeRigidActors(OutData);
+	
+	// Whatever got cached return true. We want to cache 'failure' too.
+	return true;
+}
+
+void FDerivedDataPhysXBinarySerializer::SerializeRigidActors(TArray<uint8>& OutData)
+{
+	Serializer->SerializeActors(Format, Bodies, BodySetups, PhysicalMaterials, OutData);
+}
+
+void FDerivedDataPhysXBinarySerializer::InitSerializer()
+{
+#if WITH_EDITOR
+	// static here as an optimization
+	static ITargetPlatformManagerModule* TPM = GetTargetPlatformManager();
+	if (TPM)
+	{
+		Serializer = TPM->FindPhysXFormat(Format);
+	}
+#endif
+}
+
+
 #endif	//WITH_PHYSX && WITH_EDITOR
diff --git a/Engine/Source/Runtime/Engine/Private/PhysicsEngine/PhysDerivedData.h b/Engine/Source/Runtime/Engine/Private/PhysicsEngine/PhysDerivedData.h
--- a/Engine/Source/Runtime/Engine/Private/PhysicsEngine/PhysDerivedData.h
+++ b/Engine/Source/Runtime/Engine/Private/PhysicsEngine/PhysDerivedData.h
@@ -86,4 +86,70 @@ private:
 	bool ShouldGenerateNegXTriMeshData();
 };
 
+//////////////////////////////////////////////////////////////////////////
+// PhysX Binary Serialization
+class FDerivedDataPhysXBinarySerializer : public FDerivedDataPluginInterface
+{
+private:
+
+	const TArray<FBodyInstance*>& Bodies;
+	const TArray<class UBodySetup*>& BodySetups;
+	const TArray<class UPhysicalMaterial*>& PhysicalMaterials;
+	FName Format;
+	FGuid DataGuid;
+	const class IPhysXFormat* Serializer;
+	int64 PhysXDataStart;	//important to keep track of this for alignment requirements
+	
+public:
+	FDerivedDataPhysXBinarySerializer(FName InFormat, const TArray<FBodyInstance*>& InBodies, const TArray<class UBodySetup*>& BodySetups, const TArray<class UPhysicalMaterial*>& PhysicalMaterials, const FGuid& InGuid);
+
+	virtual const TCHAR* GetPluginName() const override
+	{
+		return TEXT("PhysXSerializer");
+	}
+
+	virtual const TCHAR* GetVersionString() const override
+	{
+		// This is a version string that mimics the old versioning scheme. If you
+		// want to bump this version, generate a new guid using VS->Tools->Create GUID and
+		// return it here. Ex.
+		return TEXT("EB1DC7A5D72F469B8E0BF5251D90A220");
+	}
+
+	virtual FString GetPluginSpecificCacheKeySuffix() const override
+	{
+		enum { UE_PHYSX_DERIVEDDATA_VER = 3 };
+
+		const uint16 PhysXVersion = ((PX_PHYSICS_VERSION_MAJOR & 0xF) << 12) |
+			((PX_PHYSICS_VERSION_MINOR & 0xF) << 8) |
+			((PX_PHYSICS_VERSION_BUGFIX & 0xF) << 4) |
+			((UE_PHYSX_DERIVEDDATA_VER & 0xF));
+
+		return FString::Printf(TEXT("%s_%s_%hu"),
+			*Format.ToString(),
+			*DataGuid.ToString(),
+			PhysXVersion
+			);
+	}
+
+
+	virtual bool IsBuildThreadsafe() const override
+	{
+		return false;
+	}
+
+	virtual bool Build(TArray<uint8>& OutData) override;
+
+	/** Return true if we can build **/
+	bool CanBuild()
+	{
+		return true;
+	}
+private:
+
+	void SerializeRigidActors(TArray<uint8>& OutData);
+	void InitSerializer();
+
+};
+
 #endif	//WITH_PHYSX && WITH_EDITOR
\ No newline at end of file
diff --git a/Engine/Source/Runtime/Engine/Private/PhysicsEngine/PhysXSupport.cpp b/Engine/Source/Runtime/Engine/Private/PhysicsEngine/PhysXSupport.cpp
--- a/Engine/Source/Runtime/Engine/Private/PhysicsEngine/PhysXSupport.cpp
+++ b/Engine/Source/Runtime/Engine/Private/PhysicsEngine/PhysXSupport.cpp
@@ -6,6 +6,8 @@
 
 #include "EnginePrivate.h"
 #include "PhysicsPublic.h"
+#include "PhysicalMaterials/PhysicalMaterial.h"
+#include "PhysicsEngine/ConvexElem.h"
 
 #if WITH_PHYSX
 
@@ -882,4 +884,41 @@ void FPhysxSharedData::DumpSharedMemoryUsage(FOutputDevice* Ar)
 	}
 }
 
+void AddToCollection(PxCollection* PCollection, PxBase* PBase)
+{
+	if (PBase)
+	{
+		PCollection->add(*PBase);
+	}
+}
+
+PxCollection* MakePhysXCollection(const TArray<UPhysicalMaterial*>& PhysicalMaterials, const TArray<UBodySetup*>& BodySetups, uint64 BaseId)
+{
+	QUICK_SCOPE_CYCLE_COUNTER(STAT_CreateSharedData);
+	PxCollection* PCollection = PxCreateCollection();
+	for (UPhysicalMaterial* PhysicalMaterial : PhysicalMaterials)
+	{
+		if (PhysicalMaterial)
+		{
+			PCollection->add(*PhysicalMaterial->GetPhysXMaterial());
+		}
+	}
+
+	for (UBodySetup* BodySetup : BodySetups)
+	{
+		AddToCollection(PCollection, BodySetup->TriMesh);
+		AddToCollection(PCollection, BodySetup->TriMeshNegX);
+
+		for (const FKConvexElem& ConvexElem : BodySetup->AggGeom.ConvexElems)
+		{
+			AddToCollection(PCollection, ConvexElem.ConvexMesh);
+			AddToCollection(PCollection, ConvexElem.ConvexMeshNegX);
+		}
+	}
+
+	PxSerialization::createSerialObjectIds(*PCollection, PxSerialObjectId(BaseId));
+
+	return PCollection;
+}
+
 #endif // WITH_PHYSX
diff --git a/Engine/Source/Runtime/Engine/Private/PhysicsEngine/PhysXSupport.h b/Engine/Source/Runtime/Engine/Private/PhysicsEngine/PhysXSupport.h
--- a/Engine/Source/Runtime/Engine/Private/PhysicsEngine/PhysXSupport.h
+++ b/Engine/Source/Runtime/Engine/Private/PhysicsEngine/PhysXSupport.h
@@ -426,6 +426,8 @@ private:
 
 };
 
+ENGINE_API PxCollection* MakePhysXCollection(const TArray<UPhysicalMaterial*>& PhysicalMaterials, const TArray<UBodySetup*>& BodySetups, uint64 BaseId);
+
 /** Utility wrapper for a uint8 TArray for loading into PhysX. */
 class FPhysXInputStream : public PxInputStream
 {
diff --git a/Engine/Source/Runtime/Engine/Private/PhysicsEngine/PhysicsSerializer.cpp b/Engine/Source/Runtime/Engine/Private/PhysicsEngine/PhysicsSerializer.cpp
new file mode 100644
--- /dev/null
+++ b/Engine/Source/Runtime/Engine/Private/PhysicsEngine/PhysicsSerializer.cpp
@@ -0,0 +1,168 @@
+// Copyright 1998-2015 Epic Games, Inc. All Rights Reserved.
+
+/*=============================================================================
+	PhysicsSerializer.cpp
+=============================================================================*/ 
+
+#include "EnginePrivate.h"
+#include "PhysicsPublic.h"
+#include "PhysicsSerializer.h"
+#include "PhysDerivedData.h"
+
+UPhysicsSerializer::UPhysicsSerializer(const FObjectInitializer& ObjectInitializer)
+: Super(ObjectInitializer)
+{
+}
+
+FByteBulkData* UPhysicsSerializer::GetBinaryData(FName Format, const TArray<FBodyInstance*>& Bodies, const TArray<class UBodySetup*>& BodySetups, const TArray<class UPhysicalMaterial*>& PhysicalMaterials)
+{
+	
+	QUICK_SCOPE_CYCLE_COUNTER(STAT_GetBinaryData);
+	const bool bContainedData = BinaryFormatData.Contains(Format);
+	FByteBulkData* Result = &BinaryFormatData.GetFormat(Format);
+	Result->SetBulkDataAlignment(PHYSX_SERIALIZATION_ALIGNMENT);
+	if (!bContainedData)
+	{
+#if WITH_EDITOR
+#if WITH_PHYSX
+		TArray<uint8> OutData;
+		FDerivedDataPhysXBinarySerializer* DerivedPhysXSerializer = new FDerivedDataPhysXBinarySerializer(Format, Bodies, BodySetups, PhysicalMaterials, FGuid::NewGuid()); //TODO: Maybe it's worth adding this to the DDC. For now there's a lot of complexity with the guid invalidation so I've left it out.
+		if (DerivedPhysXSerializer->CanBuild())
+		{
+
+			DerivedPhysXSerializer->Build(OutData);
+#endif
+			if (OutData.Num())
+			{
+				Result->Lock(LOCK_READ_WRITE);
+				FMemory::Memcpy(Result->Realloc(OutData.Num()), OutData.GetData(), OutData.Num());
+				Result->Unlock();
+			}
+		}
+		else
+#endif
+		{
+			UE_LOG(LogPhysics, Warning, TEXT("Attempt to use binary physics data but we are unable to."));
+		}
+	}
+
+	return Result->GetBulkDataSize() > 0 ? Result : nullptr;
+}
+
+void UPhysicsSerializer::Serialize( FArchive& Ar )
+{
+	QUICK_SCOPE_CYCLE_COUNTER(STAT_Serialize);
+	Super::Serialize(Ar);
+	if (Ar.UE4Ver() >= VER_UE4_BODYINSTANCE_BINARY_SERIALIZATION)
+	{
+		bool bCooked = Ar.IsCooking();
+		Ar << bCooked;
+
+		if (bCooked)
+		{
+			if (Ar.IsCooking())
+			{
+				TArray<FName> ActualFormatsToSave;
+				ActualFormatsToSave.Add(FPlatformProperties::GetPhysicsFormat());
+				BinaryFormatData.Serialize(Ar, this, &ActualFormatsToSave);
+			}
+			else
+			{
+#if WITH_PHYSX
+				const uint32 Alignment = PHYSX_SERIALIZATION_ALIGNMENT;
+#else
+				const uint32 Alignment = DEFAULT_ALIGNMENT;
+#endif
+				BinaryFormatData.Serialize(Ar, this, nullptr, false, Alignment);
+			}
+		}
+	}
+}
+
+void UPhysicsSerializer::SerializePhysics(const TArray<FBodyInstance*>& Bodies, const TArray<UBodySetup*>& BodySetups, const TArray<UPhysicalMaterial*>& PhysicalMaterials)
+{
+#if WITH_EDITOR
+	GetBinaryData(FPlatformProperties::GetPhysicsFormat(), Bodies, BodySetups, PhysicalMaterials);	//build the binary data
+#endif
+}
+
+void UPhysicsSerializer::CreatePhysicsData(const TArray<UBodySetup*>& BodySetups, const TArray<UPhysicalMaterial*>& PhysicalMaterials)
+{
+#if !WITH_EDITOR
+	QUICK_SCOPE_CYCLE_COUNTER(STAT_PhysicsSerializer_CreatePhysicsData);
+
+	const FName Format(FPlatformProperties::GetPhysicsFormat());
+	if(BinaryFormatData.Contains(Format))
+	{
+		FByteBulkData& BinaryData = BinaryFormatData.GetFormat(Format);
+#if WITH_PHYSX
+		// Read serialized physics data
+		uint8* SerializedData = (uint8*)BinaryData.Lock(LOCK_READ_ONLY);
+
+		FBufferReader Ar(SerializedData, BinaryData.GetBulkDataSize(), false );
+		uint8 bIsLittleEndian;
+		uint64 BaseId;	//the starting index of the shared resources we didn't serialize out
+		Ar << bIsLittleEndian;
+		Ar.SetByteSwapping( PLATFORM_LITTLE_ENDIAN ? !bIsLittleEndian : !!bIsLittleEndian );
+		Ar << BaseId;
+		//Note that PhysX expects the binary data to be 128-byte aligned. Because of this we've padded so find the next spot in memory
+		int32 BytesToPad = PHYSX_SERIALIZATION_ALIGNMENT - (Ar.Tell() % PHYSX_SERIALIZATION_ALIGNMENT);
+		
+		physx::PxSerializationRegistry* PRegistry =  PxSerialization::createSerializationRegistry(*GPhysXSDK);
+		physx::PxCollection* PCollection = nullptr;
+		PxCollection* PExternalData = MakePhysXCollection(PhysicalMaterials, BodySetups, BaseId);		
+		{
+			QUICK_SCOPE_CYCLE_COUNTER(STAT_DeserializePhysics);
+			PCollection = PxSerialization::createCollectionFromBinary(SerializedData + Ar.Tell() + BytesToPad, *PRegistry, PExternalData);	
+		}
+
+		{
+			QUICK_SCOPE_CYCLE_COUNTER(STAT_AddBodiesToMap);
+			const uint32 NumObjects = PCollection->getNbObjects();
+			for (uint32 ObjectIdx = 0; ObjectIdx < NumObjects; ++ObjectIdx)
+			{
+				PxBase& PObject = PCollection->getObject(ObjectIdx);
+				if (PxRigidActor* PRigidActor = PObject.is<PxRigidActor>())
+				{
+					const PxSerialObjectId ObjectId = PCollection->getId(PObject);
+					ActorsMap.Add(ObjectId, PRigidActor);
+				}
+			}
+		}
+
+		PExternalData->release();
+		PCollection->release();
+		PRegistry->release();
+#endif
+	}else
+	{
+#if !WITH_EDITOR
+		UE_LOG(LogPhysics, Warning, TEXT("PhysicsSerializer has no binary data. Body instances will fall back to slow creation path."));
+#endif
+	}
+#endif
+}
+
+void UPhysicsSerializer::BeginDestroy()
+{
+	const FName Format(FPlatformProperties::GetPhysicsFormat());
+	if (BinaryFormatData.Contains(Format))
+	{
+		FByteBulkData& BinaryData = BinaryFormatData.GetFormat(Format);
+		if(BinaryData.IsLocked())
+		{
+			BinaryData.Unlock();
+		}
+	}
+
+	Super::BeginDestroy();
+}
+
+#if WITH_PHYSX
+PxRigidActor* UPhysicsSerializer::GetRigidActor(physx::PxSerialObjectId ObjectId) const
+{
+	PxRigidActor* const* PActor = ActorsMap.Find(ObjectId);
+
+	return PActor ? *PActor : nullptr;
+}
+#endif
\ No newline at end of file
diff --git a/Engine/Source/Runtime/Engine/Public/PhysicsSerializer.h b/Engine/Source/Runtime/Engine/Public/PhysicsSerializer.h
new file mode 100644
--- /dev/null
+++ b/Engine/Source/Runtime/Engine/Public/PhysicsSerializer.h
@@ -0,0 +1,70 @@
+// Copyright 1998-2015 Epic Games, Inc. All Rights Reserved.
+#pragma once
+
+/**
+ * Helper class used to serialize a collection of BodyInstances. This allows the physics engine to serialize whatever expensive computations it needs offline.
+ * A DDC entry will be created per instance and there is some overhead associated with serialization so it's wise to use this for a group of BodyInstances.
+ */
+
+struct FBodyInstance;
+
+namespace physx
+{
+	class PxRigidActor;
+	class PxCollection;
+	class PxSerializationRegistry;
+}
+
+#include "PhysicsSerializer.generated.h"
+
+UCLASS(hidecategories=Object, MinimalAPI)
+class UPhysicsSerializer : public UObject
+{
+	GENERATED_UCLASS_BODY()
+public:
+
+	 /** 
+	  * Performs a binary serialization of the given bodies. This is a physics engine optimization so that various computations can be done offline.
+	  * The serialization stores all physics actors and shapes, as well as references to physical materials and geometry
+	  * @param Bodies The bodies containing the needed physics actors to serialize
+	  * @param BodySetups The various body setups used by Bodies (could be 0). This is needed for keeping geometry out of the serialized data
+	  * @param PhysicalMaterials The physical materials used by Bodies (could be 0). This is needed for keeping physical materials out of the serialized data
+	  */
+	void SerializePhysics(const TArray<FBodyInstance*>& Bodies, const TArray<UBodySetup*>& BodySetups, const TArray<UPhysicalMaterial*>& PhysicalMaterials);
+
+	/** 
+	 * Uses the serialized binary data to recreate physics engine specific data structures. This is an optimization so that FBodyInstance::InitBody doesn't have to do certain work at runtime.
+	 * Assumes that the binary data has already been serialized either by SerializePhysics or from cooked data.
+	 * @param BodySetups The various body setups used by Bodies (could be 0). This is needed for deserialization and pointer fixup
+	 * @param PhysicalMaterials The physical materials used by Bodies (could be 0).This is needed for deserialization and pointer fixup
+	 */
+	void CreatePhysicsData(const TArray<UBodySetup*>& BodySetups, const TArray<UPhysicalMaterial*>& PhysicalMaterials);
+
+	virtual void BeginDestroy();
+
+#if WITH_PHYSX
+	physx::PxRigidActor* GetRigidActor(uint64 ObjectId) const;
+
+	static uint64 GetSerializationId();
+#endif
+
+	virtual void Serialize(FArchive& Ar) override;
+
+private:
+	/** Helper function to find all physical materials that Bodies depend on. During serialization these will be referenced*/
+	void GatherNeededMaterials(const TArray<FBodyInstance*>& Bodies);
+
+	/** Helper function to find all geometry that BodySetups depend on. During serialization these will be referenced*/
+	void GatherNeededGeometry(const TArray<UBodySetup*>& BodySetups);
+
+private:
+	/** Cooked physics data for each format */
+	FFormatContainer BinaryFormatData;
+
+	/** Binary data of the physics actors.*/
+	FByteBulkData* GetBinaryData(FName Format, const TArray<FBodyInstance*>& Bodies, const TArray<class UBodySetup*>& BodySetups, const TArray<class UPhysicalMaterial*>& PhysicalMaterials);
+
+#if WITH_PHYSX
+	TMap <uint64, physx::PxRigidActor*> ActorsMap;
+#endif
+};
diff --git a/Engine/Source/Runtime/PhysXFormats/PhysXFormats.Build.cs b/Engine/Source/Runtime/PhysXFormats/PhysXFormats.Build.cs
--- a/Engine/Source/Runtime/PhysXFormats/PhysXFormats.Build.cs
+++ b/Engine/Source/Runtime/PhysXFormats/PhysXFormats.Build.cs
@@ -17,6 +17,12 @@ public class PhysXFormats : ModuleRules
 			}
 			);
 
+        PrivateIncludePaths.AddRange(
+            new string[] {
+                "Runtime/Engine/Private",   //Not sure why this is its own format. The two modules depend on eachother.
+			}
+        );
+
 		SetupModulePhysXAPEXSupport(Target);
 	}
 }
diff --git a/Engine/Source/Runtime/PhysXFormats/Private/PhysXFormats.cpp b/Engine/Source/Runtime/PhysXFormats/Private/PhysXFormats.cpp
--- a/Engine/Source/Runtime/PhysXFormats/Private/PhysXFormats.cpp
+++ b/Engine/Source/Runtime/PhysXFormats/Private/PhysXFormats.cpp
@@ -9,6 +9,7 @@
 #include "IPhysXFormat.h"
 #include "IPhysXFormatModule.h"
 #include "PhysXFormats.h"
+#include "PhysicsEngine/PhysXSupport.h"
 
 static_assert(WITH_PHYSX, "No point in compiling PhysX cooker, if we don't have PhysX.");
 
@@ -203,6 +204,71 @@ public:
 #endif		// WITH_PHYSX
 	}
 
+	virtual bool SerializeActors(FName Format, const TArray<FBodyInstance*>& Bodies, const TArray<UBodySetup*>& BodySetups, const TArray<UPhysicalMaterial*>& PhysicalMaterials, TArray<uint8>& OutBuffer) const override
+	{
+#if WITH_PHYSX
+		PxSerializationRegistry* PRegistry = PxSerialization::createSerializationRegistry(*GPhysXSDK);
+		PxCollection* PCollection = PxCreateCollection();
+
+		PxBase* PLastObject = nullptr;
+
+		for(FBodyInstance* BodyInstance : Bodies)
+		{
+			if(BodyInstance->RigidActorSync)
+			{
+				PCollection->add(*BodyInstance->RigidActorSync, BodyInstance->RigidActorSyncId);
+				PLastObject = BodyInstance->RigidActorSync;
+			}
+
+			if(BodyInstance->RigidActorAsync)
+			{
+				PCollection->add(*BodyInstance->RigidActorAsync,  BodyInstance->RigidActorAsyncId);
+				PLastObject = BodyInstance->RigidActorAsync;
+			}
+		}
+
+		PxSerialization::createSerialObjectIds(*PCollection, PxSerialObjectId(1));	//we get physx to assign an id for each actor
+
+		//Note that rigid bodies may have assigned ids. It's important to let them go first because we rely on that id for deserialization.
+		//One this is done we must find out the next available ID, and use that for naming the shared resources. We have to save this for deserialization
+		uint64 BaseId = PLastObject ? (PCollection->getId(*PLastObject) + 1) : 1;
+
+		PxCollection* PExceptFor = MakePhysXCollection(PhysicalMaterials, BodySetups, BaseId);
+		
+		for (FBodyInstance* BodyInstance : Bodies)	//and then we mark that id back into the bodyinstance so we can pair the two later
+		{
+			if (BodyInstance->RigidActorSync)
+			{
+				BodyInstance->RigidActorSyncId = PCollection->getId(*BodyInstance->RigidActorSync);
+			}
+
+			if (BodyInstance->RigidActorAsync)
+			{
+				BodyInstance->RigidActorAsyncId = PCollection->getId(*BodyInstance->RigidActorAsync);
+			}
+		}
+
+		//We must store the BaseId for shared resources.
+		FMemoryWriter Ar(OutBuffer);
+		uint8 bIsLittleEndian = PLATFORM_LITTLE_ENDIAN; //TODO: We should pass the target platform into this function and write it. Then swap the endian on the writer so the reader doesn't have to do it at runtime
+		Ar << bIsLittleEndian;
+		Ar << BaseId;
+		//Note that PhysX expects the binary data to be 128-byte aligned. Because of this we must pad
+		int32 BytesToPad = PHYSX_SERIALIZATION_ALIGNMENT - (Ar.Tell() % PHYSX_SERIALIZATION_ALIGNMENT);
+		OutBuffer.AddZeroed(BytesToPad);
+
+		FPhysXOutputStream Buffer(&OutBuffer);
+		PxSerialization::complete(*PCollection, *PRegistry, PExceptFor);
+		PxSerialization::serializeCollectionToBinary(Buffer, *PCollection, *PRegistry, PExceptFor);
+		PCollection->release();
+		PExceptFor->release();
+		PRegistry->release();
+		
+		return true;
+#endif
+		return false;
+	}
+
 };
 
 
diff --git a/Engine/Source/Runtime/PhysXFormats/Private/PhysXFormats.h b/Engine/Source/Runtime/PhysXFormats/Private/PhysXFormats.h
--- a/Engine/Source/Runtime/PhysXFormats/Private/PhysXFormats.h
+++ b/Engine/Source/Runtime/PhysXFormats/Private/PhysXFormats.h
@@ -26,6 +26,7 @@
 #include "PxRenderBuffer.h"
 #include "PxExtensionsAPI.h"
 #include "PxVisualDebuggerExt.h"
+#include "PxCollectionExt.h"
 //#include "PxDefaultCpuDispatcher.h"
 
 // vehicle related header files
@@ -71,10 +72,14 @@ public:
 	virtual PxU32 write(const void* Src, PxU32 Count) override
 	{
 		check(Data);
-		check(Src);
-		int32 CurrentNum = (*Data).Num();
-		(*Data).AddUninitialized(Count);
-		FMemory::Memcpy(&(*Data)[CurrentNum], Src, Count);
+		if(Count)	//PhysX serializer can pass us 0 bytes to write
+		{
+			check(Src);
+			int32 CurrentNum = (*Data).Num();
+			(*Data).AddUninitialized(Count);
+			FMemory::Memcpy(&(*Data)[CurrentNum], Src, Count);
+		}
+		
 		return Count;
 	}
 };
diff --git a/Engine/Source/Runtime/PhysXFormats/Public/IPhysXFormat.h b/Engine/Source/Runtime/PhysXFormats/Public/IPhysXFormat.h
--- a/Engine/Source/Runtime/PhysXFormats/Public/IPhysXFormat.h
+++ b/Engine/Source/Runtime/PhysXFormats/Public/IPhysXFormat.h
@@ -8,7 +8,7 @@
 
 
 /**
- * IPhysXFormat, PhysX cooking abstraction
+ * IPhysXFormat, PhysX cooking and serialization abstraction
 **/
 class IPhysXFormat
 {
@@ -59,6 +59,18 @@ public:
 	 */
 	virtual bool CookHeightField( FName Format, FIntPoint HFSize, float Thickness, const void* Samples, uint32 SamplesStride, TArray<uint8>& OutBuffer ) const = 0;
 
+	/**
+	 * Serializes the BodyInstance
+	 *
+	 * @param Format The desired format
+	 * @param Bodies The bodies containing the needed physx actors to serialize
+	 * @param BodySetups The various body setups used by Bodies (could be just 1). This is needed for keeping geometry out of the serialized data
+	 * @param PhysicalMaterials The physical materials used by Bodies (could be just 1). This is needed for keeping physical materials out of the serialized data
+	 * @param OutBuffer The resulting cooked data
+	 * @return true on success, false otherwise.
+	 */
+	virtual bool SerializeActors( FName Format, const TArray<struct FBodyInstance*>& Bodies, const TArray<class UBodySetup*>& BodySetups, const TArray<class UPhysicalMaterial*>& PhysicalMaterials, TArray<uint8>& OutBuffer ) const = 0;
+
 	/**
 	 * Gets the list of supported formats.
 	 *
@@ -82,3 +94,6 @@ public:
 	 */
 	virtual ~IPhysXFormat( ) { }
 };
+
+// binary serialization requires 128 byte alignment
+#define PHYSX_SERIALIZATION_ALIGNMENT 128
\ No newline at end of file
