Stats - Misc stat rendering tweaks

diff --git a/Engine/Source/Runtime/Engine/Private/StatsRender2.cpp b/Engine/Source/Runtime/Engine/Private/StatsRender2.cpp
--- a/Engine/Source/Runtime/Engine/Private/StatsRender2.cpp
+++ b/Engine/Source/Runtime/Engine/Private/StatsRender2.cpp
@@ -51,9 +51,9 @@ struct FStatFont
 static FStatFont GStatFonts[(int32)EStatFontTypes::NumFonts] =
 {
 	/** Tiny. */
-	FStatFont( 24, 10, 1 ),
+	FStatFont( 36, 10, 1 ),
 	/** Small. */
-	FStatFont( 48, 12, 2 ),
+	FStatFont( 72, 12, 2 ),
 };
 
 void FromString( EStatFontTypes& OutValue, const TCHAR* Buffer )
@@ -132,7 +132,7 @@ struct FStatRenderGlobals
 		if( bNeedRefresh )
 		{
 			// This is the number of W characters to leave spacing for in the stat name column.
-			const FString STATNAME_COLUMN_WIDTH = FString::ChrN( GetNumCharsForStatName(), 'M' );
+			const FString STATNAME_COLUMN_WIDTH = FString::ChrN( GetNumCharsForStatName(), 'S' );
 
 			// This is the number of W characters to leave spacing for in the other stat columns.
 			const FString STATDATA_COLUMN_WIDTH = FString::ChrN( StatFontType == EStatFontTypes::Small ? 8 : 7, 'W' );
@@ -141,6 +141,7 @@ struct FStatRenderGlobals
 			int32 StatColumnSpaceSizeY = 0;
 			int32 TimeColumnSpaceSizeY = 0;
 
+			// @TODO yrx 2015-04-17 Compute on the stats thread to get the exact measurement
 			// Determine where the first column goes.
 			StringSize(StatFont, AfterNameColumnOffset, StatColumnSpaceSizeY, *STATNAME_COLUMN_WIDTH);
 
@@ -667,7 +668,7 @@ static void RenderGroupedWithHierarchy(const FGameThreadHudData& ViewData, FView
  */
 void RenderStats(FViewport* Viewport, class FCanvas* Canvas, int32 X, int32 Y)
 {
-	DECLARE_SCOPE_CYCLE_COUNTER( TEXT( "RenderStats" ), STAT_LoadMap, STATGROUP_StatSystem );
+	DECLARE_SCOPE_CYCLE_COUNTER( TEXT( "RenderStats" ), STAT_RenderStats, STATGROUP_StatSystem );
 
 	FGameThreadHudData* ViewData = FHUDGroupGameThreadRenderer::Get().Latest;
 	if (!ViewData)
