UE4 - fix a few minor things to so that AllowAsyncRenderThreadUpdates 1 works

diff --git a/Engine/Source/Runtime/Engine/Private/Particles/ParticleSystemRender.cpp b/Engine/Source/Runtime/Engine/Private/Particles/ParticleSystemRender.cpp
--- a/Engine/Source/Runtime/Engine/Private/Particles/ParticleSystemRender.cpp
+++ b/Engine/Source/Runtime/Engine/Private/Particles/ParticleSystemRender.cpp
@@ -1243,7 +1243,6 @@ void FDynamicMeshEmitterData::Init( bool bInSelected,
 		InEmitterInstance->SpriteTemplate->LODLevels[InEmitterInstance->CurrentLODLevelIndex]
 		);
 
-	check(IsInGameThread());
 	for (int32 i = 0; i < MeshMaterials.Num(); ++i)
 	{
 		UMaterialInterface* RenderMaterial = MeshMaterials[i];
diff --git a/Engine/Source/Runtime/Engine/Private/StaticMeshRender.cpp b/Engine/Source/Runtime/Engine/Private/StaticMeshRender.cpp
--- a/Engine/Source/Runtime/Engine/Private/StaticMeshRender.cpp
+++ b/Engine/Source/Runtime/Engine/Private/StaticMeshRender.cpp
@@ -1083,9 +1083,19 @@ FStaticMeshSceneProxy::FLODInfo::FLODInfo(const UStaticMeshComponent* InComponen
 		FMaterialResource const* MaterialResource = const_cast<UMaterialInterface const*>(SectionInfo.Material)->GetMaterial_Concurrent(RecursionGuard)->GetMaterialResource(FeatureLevel);
 		if(MaterialResource)
 		{
-			if (MaterialResource->MaterialModifiesMeshPosition_GameThread())
+			if (IsInGameThread())
 			{
-				bUsesMeshModifyingMaterials = true;
+				if (MaterialResource->MaterialModifiesMeshPosition_GameThread())
+				{
+					bUsesMeshModifyingMaterials = true;
+				}
+			}
+			else
+			{
+				if (MaterialResource->MaterialModifiesMeshPosition_RenderThread())
+				{
+					bUsesMeshModifyingMaterials = true;
+				}
 			}
 		}
 	}
