[UE-12232] Need to reset variables exposed to blueprints when rerunning construction scripts (in case past UCS setters have since been removed).

diff --git a/Engine/Source/Runtime/Engine/Private/ActorConstruction.cpp b/Engine/Source/Runtime/Engine/Private/ActorConstruction.cpp
--- a/Engine/Source/Runtime/Engine/Private/ActorConstruction.cpp
+++ b/Engine/Source/Runtime/Engine/Private/ActorConstruction.cpp
@@ -64,16 +64,21 @@ void AActor::ResetPropertiesForConstruction()
 		UStructProperty* StructProp = Cast<UStructProperty>(Prop);
 		UClass* PropClass = CastChecked<UClass>(Prop->GetOuter()); // get the class that added this property
 
+		bool const bCanEditInstanceValue = !Prop->HasAnyPropertyFlags(CPF_DisableEditOnInstance) &&
+			Prop->HasAnyPropertyFlags(CPF_Edit);
+		bool const bCanBeSetInBlueprints = Prop->HasAnyPropertyFlags(CPF_BlueprintVisible) && 
+			!Prop->HasAnyPropertyFlags(CPF_BlueprintReadOnly);
+
 		// First see if it is a random stream, if so reset before running construction script
 		if( (StructProp != NULL) && (StructProp->Struct != NULL) && (StructProp->Struct->GetFName() == RandomStreamName) )
 		{
 			FRandomStream* StreamPtr =  StructProp->ContainerPtrToValuePtr<FRandomStream>(this);
 			StreamPtr->Reset();
 		}
-		// If it is a blueprint added variable that is not editable per-instance, reset to default before running construction script
+		// If it is a blueprint exposed variable that is not editable per-instance, reset to default before running construction script
 		else if( !bIsLevelScriptActor 
-				&& Prop->HasAnyPropertyFlags(CPF_DisableEditOnInstance)
-				&& PropClass->HasAnyClassFlags(CLASS_CompiledFromBlueprint) 
+				&& !bCanEditInstanceValue
+				&& bCanBeSetInBlueprints
 				&& !Prop->IsA(UDelegateProperty::StaticClass()) 
 				&& !Prop->IsA(UMulticastDelegateProperty::StaticClass()) )
 		{
