UE4 - UE-13583, GitHub 1037, minor tweak to instanced static mesh to prevent an update when just calling UpdateTransform

diff --git a/Engine/Source/Runtime/Engine/Classes/Components/HierarchicalInstancedStaticMeshComponent.h b/Engine/Source/Runtime/Engine/Classes/Components/HierarchicalInstancedStaticMeshComponent.h
--- a/Engine/Source/Runtime/Engine/Classes/Components/HierarchicalInstancedStaticMeshComponent.h
+++ b/Engine/Source/Runtime/Engine/Classes/Components/HierarchicalInstancedStaticMeshComponent.h
@@ -125,7 +125,7 @@ public:
 	// UInstancedStaticMesh interface
 	virtual int32 AddInstance(const FTransform& InstanceTransform) override;
 	virtual bool RemoveInstance(int32 InstanceIndex) override;
-	virtual bool UpdateInstanceTransform(int32 InstanceIndex, const FTransform& NewInstanceTransform, bool bWorldSpace) override;
+	virtual bool UpdateInstanceTransform(int32 InstanceIndex, const FTransform& NewInstanceTransform, bool bWorldSpace, bool bMarkRenderStateDirty = false) override;
 	virtual void ClearInstances() override;
 
 
diff --git a/Engine/Source/Runtime/Engine/Classes/Components/InstancedStaticMeshComponent.h b/Engine/Source/Runtime/Engine/Classes/Components/InstancedStaticMeshComponent.h
--- a/Engine/Source/Runtime/Engine/Classes/Components/InstancedStaticMeshComponent.h
+++ b/Engine/Source/Runtime/Engine/Classes/Components/InstancedStaticMeshComponent.h
@@ -108,7 +108,7 @@ class ENGINE_API UInstancedStaticMeshComponent : public UStaticMeshComponent
 	
 	/** Update the transform for the instance specified. Instance is given in local space of this component unless bWorldSpace is set.  Returns True on success. */
 	UFUNCTION(BlueprintCallable, Category = "Components|InstancedStaticMesh")
-	virtual bool UpdateInstanceTransform(int32 InstanceIndex, const FTransform& NewInstanceTransform, bool bWorldSpace=false);
+	virtual bool UpdateInstanceTransform(int32 InstanceIndex, const FTransform& NewInstanceTransform, bool bWorldSpace=false, bool bMarkRenderStateDirty = false);
 
 	/** Remove the instance specified. Returns True on success. */
 	UFUNCTION(BlueprintCallable, Category = "Components|InstancedStaticMesh")
diff --git a/Engine/Source/Runtime/Engine/Private/HierarchicalInstancedStaticMesh.cpp b/Engine/Source/Runtime/Engine/Private/HierarchicalInstancedStaticMesh.cpp
--- a/Engine/Source/Runtime/Engine/Private/HierarchicalInstancedStaticMesh.cpp
+++ b/Engine/Source/Runtime/Engine/Private/HierarchicalInstancedStaticMesh.cpp
@@ -1633,7 +1633,7 @@ bool UHierarchicalInstancedStaticMeshComponent::RemoveInstance(int32 InstanceInd
 	return true;
 }
 
-bool UHierarchicalInstancedStaticMeshComponent::UpdateInstanceTransform(int32 InstanceIndex, const FTransform& NewInstanceTransform, bool bWorldSpace)
+bool UHierarchicalInstancedStaticMeshComponent::UpdateInstanceTransform(int32 InstanceIndex, const FTransform& NewInstanceTransform, bool bWorldSpace, bool bMarkRenderStateDirty)
 {
 	if (!PerInstanceSMData.IsValidIndex(InstanceIndex))
 	{
@@ -1652,7 +1652,7 @@ bool UHierarchicalInstancedStaticMeshComponent::UpdateInstanceTransform(int32 In
 	// Allocate a new instance render order ID, rendered last
 	InstanceReorderTable[InstanceIndex] = PerInstanceSMData.Num()-1 + RemovedInstances.Num();
 
-	bool Result = Super::UpdateInstanceTransform(InstanceIndex, NewInstanceTransform, bWorldSpace);
+	bool Result = Super::UpdateInstanceTransform(InstanceIndex, NewInstanceTransform, bWorldSpace, bMarkRenderStateDirty);
 
 	if (StaticMesh)
 	{
diff --git a/Engine/Source/Runtime/Engine/Private/InstancedStaticMesh.cpp b/Engine/Source/Runtime/Engine/Private/InstancedStaticMesh.cpp
--- a/Engine/Source/Runtime/Engine/Private/InstancedStaticMesh.cpp
+++ b/Engine/Source/Runtime/Engine/Private/InstancedStaticMesh.cpp
@@ -1440,7 +1440,7 @@ bool UInstancedStaticMeshComponent::GetInstanceTransform(int32 InstanceIndex, FT
 	return true;
 }
 
-bool UInstancedStaticMeshComponent::UpdateInstanceTransform(int32 InstanceIndex, const FTransform& NewInstanceTransform, bool bWorldSpace)
+bool UInstancedStaticMeshComponent::UpdateInstanceTransform(int32 InstanceIndex, const FTransform& NewInstanceTransform, bool bWorldSpace, bool bMarkRenderStateDirty)
 {
 	if (!PerInstanceSMData.IsValidIndex(InstanceIndex))
 	{
@@ -1472,7 +1472,10 @@ bool UInstancedStaticMeshComponent::UpdateInstanceTransform(int32 InstanceIndex,
 	PartialNavigationUpdate(InstanceIndex);
 
 	ReleasePerInstanceRenderData();
-	MarkRenderStateDirty();
+	if (bMarkRenderStateDirty)
+	{
+		MarkRenderStateDirty();
+	}
 
 	return true;
 }
