PR #912: Off axis matrix temporal aa fix (Contributed by [REDACTED_PERSON])

diff --git a/Engine/Source/Runtime/Engine/Public/SceneView.h b/Engine/Source/Runtime/Engine/Public/SceneView.h
--- a/Engine/Source/Runtime/Engine/Public/SceneView.h
+++ b/Engine/Source/Runtime/Engine/Public/SceneView.h
@@ -154,6 +154,7 @@ struct FViewMatrices
 		PreShadowTranslation = FVector::ZeroVector;
 		PreViewTranslation = FVector::ZeroVector;
 		ViewOrigin = FVector::ZeroVector;
+		TemporalAASample = FVector2D::ZeroVector;
 	}
 
 	/** ViewToClip : UE4 projection matrix projects such that clip space Z=1 is the near plane, and Z=0 is the infinite far plane. */
@@ -169,13 +170,14 @@ struct FViewMatrices
 	/** During GetDynamicMeshElements this will be the correct cull volume for shadow stuff */
 	const FConvexVolume* GetDynamicMeshElementsShadowCullFrustum;
 	/** If the above is non-null, a translation that is applied to world-space before transforming by one of the shadow matrices. */
-	FVector PreShadowTranslation;
+	FVector		PreShadowTranslation;
 	/** The translation to apply to the world before TranslatedViewProjectionMatrix. Usually it is -ViewOrigin but with rereflections this can differ */
 	FVector		PreViewTranslation;
 	/** To support ortho and other modes this is redundant, in world space */
 	FVector		ViewOrigin;
 	/** Scale applied by the projection matrix in X and Y. */
-	FVector2D ProjectionScale;
+	FVector2D	ProjectionScale;
+	FVector2D	TemporalAASample;
 	/**
 	 * Scale factor to use when computing the size of a sphere in pixels.
 	 * 
diff --git a/Engine/Source/Runtime/Renderer/Private/CompositionLighting/PostProcessAmbientOcclusion.h b/Engine/Source/Runtime/Renderer/Private/CompositionLighting/PostProcessAmbientOcclusion.h
--- a/Engine/Source/Runtime/Renderer/Private/CompositionLighting/PostProcessAmbientOcclusion.h
+++ b/Engine/Source/Runtime/Renderer/Private/CompositionLighting/PostProcessAmbientOcclusion.h
@@ -175,10 +175,10 @@ void FCameraMotionParameters::Set(FRHICommandList& RHICmdList, const FSceneView&
 	FMatrix PrevProj = ViewState->PrevViewMatrices.ProjMatrix;
 
 	// Remove jitter
-	Proj.M[2][0] = 0.0f;
-	Proj.M[2][1] = 0.0f;
-	PrevProj.M[2][0] = 0.0f;
-	PrevProj.M[2][1] = 0.0f;
+	Proj.M[2][0] -= View.ViewMatrices.TemporalAASample.X * 2.0f / View.ViewRect.Width();
+	Proj.M[2][1] -= View.ViewMatrices.TemporalAASample.Y * 2.0f / View.ViewRect.Height();
+	PrevProj.M[2][0] -= ViewState->PrevViewMatrices.TemporalAASample.X * 2.0f / View.ViewRect.Width();
+	PrevProj.M[2][1] -= ViewState->PrevViewMatrices.TemporalAASample.Y * 2.0f / View.ViewRect.Height();
 
 	FVector DeltaTranslation = ViewState->PrevViewMatrices.PreViewTranslation - View.ViewMatrices.PreViewTranslation;
 	FMatrix ViewProj = ( View.ViewMatrices.TranslatedViewMatrix * Proj ).GetTransposed();
diff --git a/Engine/Source/Runtime/Renderer/Private/PostProcess/PostProcessTemporalAA.cpp b/Engine/Source/Runtime/Renderer/Private/PostProcess/PostProcessTemporalAA.cpp
--- a/Engine/Source/Runtime/Renderer/Private/PostProcess/PostProcessTemporalAA.cpp
+++ b/Engine/Source/Runtime/Renderer/Private/PostProcess/PostProcessTemporalAA.cpp
@@ -708,8 +708,8 @@ void FRCPassPostProcessTemporalAA::Process(FRenderingCompositePassContext& Conte
 	// TODO draw separate translucency after jitter has been removed
 
 	// Remove jitter
-	View.ViewMatrices.ProjMatrix.M[2][0] = 0.0f;
-	View.ViewMatrices.ProjMatrix.M[2][1] = 0.0f;
+	View.ViewMatrices.ProjMatrix.M[2][0] -= View.ViewMatrices.TemporalAASample.X * 2.0f / View.ViewRect.Width();
+	View.ViewMatrices.ProjMatrix.M[2][1] -= View.ViewMatrices.TemporalAASample.Y * 2.0f / View.ViewRect.Height();
 
 	// Compute the view projection matrix and its inverse.
 	View.ViewProjectionMatrix = View.ViewMatrices.ViewMatrix * View.ViewMatrices.ProjMatrix;
diff --git a/Engine/Source/Runtime/Renderer/Private/SceneVisibility.cpp b/Engine/Source/Runtime/Renderer/Private/SceneVisibility.cpp
--- a/Engine/Source/Runtime/Renderer/Private/SceneVisibility.cpp
+++ b/Engine/Source/Runtime/Renderer/Private/SceneVisibility.cpp
@@ -1745,14 +1745,11 @@ void FSceneRenderer::PreVisibilityFrameSetup(FRHICommandListImmediate& RHICmdLis
 					SampleY = Halton( Index, 3 ) - 0.5f;
 				}
 
-				View.TemporalJitterPixelsX = SampleX;
-				View.TemporalJitterPixelsY = SampleY;
+				View.ViewMatrices.TemporalAASample.X = SampleX;
+				View.ViewMatrices.TemporalAASample.Y = SampleY;
 
-				float JitterX = ( SampleX ) * 2.0f / View.ViewRect.Width();
-				float JitterY = ( SampleY ) * 2.0f / View.ViewRect.Height();
-
-				View.ViewMatrices.ProjMatrix.M[2][0] += JitterX;
-				View.ViewMatrices.ProjMatrix.M[2][1] += JitterY;
+				View.ViewMatrices.ProjMatrix.M[2][0] += View.ViewMatrices.TemporalAASample.X * 2.0f / View.ViewRect.Width();
+				View.ViewMatrices.ProjMatrix.M[2][1] += View.ViewMatrices.TemporalAASample.Y * 2.0f / View.ViewRect.Height();;
 
 				// Compute the view projection matrix and its inverse.
 				View.ViewProjectionMatrix = View.ViewMatrices.ViewMatrix * View.ViewMatrices.ProjMatrix;
