Expose Blueprint Implementable Event for Player Restart

diff --git a/Engine/Source/Runtime/Engine/Classes/GameFramework/GameMode.h b/Engine/Source/Runtime/Engine/Classes/GameFramework/GameMode.h
--- a/Engine/Source/Runtime/Engine/Classes/GameFramework/GameMode.h
+++ b/Engine/Source/Runtime/Engine/Classes/GameFramework/GameMode.h
@@ -457,7 +457,7 @@ public:
 	virtual void PostLogin( APlayerController* NewPlayer );
 
 	/** Notification that a player has successfully logged in, and has been given a player controller */
-	UFUNCTION(BlueprintImplementableEvent, Category="Game", meta=(DisplayName="PostLogin"))
+	UFUNCTION(BlueprintImplementableEvent, Category="Game", meta=(DisplayName="OnPostLogin"))
 	void K2_PostLogin( APlayerController* NewPlayer );
 
 	/** 
@@ -640,7 +640,7 @@ public:
 	 */
 	virtual void SwapPlayerControllers(APlayerController* OldPC, APlayerController* NewPC);
 
-	UFUNCTION(BlueprintImplementableEvent, Category="Game", meta=(DisplayName="SwapPlayerControllers"))
+	UFUNCTION(BlueprintImplementableEvent, Category="Game", meta=(DisplayName="OnSwapPlayerControllers"))
 	void K2_OnSwapPlayerControllers(APlayerController* OldPC, APlayerController* NewPC);
 
 	/** 
@@ -659,6 +659,9 @@ public:
 	/** Does end of game handling for the online layer */
 	virtual void RestartPlayer(class AController* NewPlayer);
 
+	UFUNCTION(BlueprintImplementableEvent, Category="Game", meta=(DisplayName="OnRestartPlayer"))
+	void K2_OnRestartPlayer(AController* NewPlayer);
+
 	/** Given a string, return a fully-qualified GameMode class name */
 	static FString StaticGetFullGameClassName(FString const& Str);
 
diff --git a/Engine/Source/Runtime/Engine/Private/GameMode.cpp b/Engine/Source/Runtime/Engine/Private/GameMode.cpp
--- a/Engine/Source/Runtime/Engine/Private/GameMode.cpp
+++ b/Engine/Source/Runtime/Engine/Private/GameMode.cpp
@@ -499,6 +499,8 @@ void AGameMode::RestartPlayer(AController* NewPlayer)
 			NewPlayer->SetControlRotation( NewControllerRot );
 
 			SetPlayerDefaults(NewPlayer->GetPawn());
+
+			K2_OnRestartPlayer(NewPlayer);
 		}
 	}
 
