Fix automatic detachment when simulating physics. This is not needed any more with all the welding code.

diff --git a/Engine/Source/Runtime/Engine/Private/Components/PrimitiveComponent.cpp b/Engine/Source/Runtime/Engine/Private/Components/PrimitiveComponent.cpp
--- a/Engine/Source/Runtime/Engine/Private/Components/PrimitiveComponent.cpp
+++ b/Engine/Source/Runtime/Engine/Private/Components/PrimitiveComponent.cpp
@@ -541,14 +541,7 @@ void UPrimitiveComponent::PostEditChangeProperty(FPropertyChangedEvent& Property
 			CachedMaxDrawDistance = LDMaxDrawDistance;
 		}
 
-		if (PropertyName == NAME_SimulatePhysics)
-		{
-			if (AttachParent != NULL && IsSimulatingPhysics())
-			{
-				UE_LOG(LogPrimitiveComponent, Warning, TEXT("%s is simulating physics but has an attach parent. Detaching."), *GetPathName());
-				DetachFromParent();
-			}
-		}
+
 
 		// we need to reregister the primitive if the min draw distance changed to propagate the change to the rendering thread
 		if (PropertyThatChanged->GetFName() == GET_MEMBER_NAME_CHECKED(UPrimitiveComponent, MinDrawDistance))
