Fix welding code not setting the new welding parent during unweld. Also added early out for kinematic case where re-welding is pointless

diff --git a/Engine/Source/Runtime/Engine/Private/PhysicsEngine/BodyInstance.cpp b/Engine/Source/Runtime/Engine/Private/PhysicsEngine/BodyInstance.cpp
--- a/Engine/Source/Runtime/Engine/Private/PhysicsEngine/BodyInstance.cpp
+++ b/Engine/Source/Runtime/Engine/Private/PhysicsEngine/BodyInstance.cpp
@@ -2455,9 +2455,10 @@ void FBodyInstance::SetInstanceSimulatePhysics(bool bSimulate, bool bMaintainPhy
 			{
 				OwnerComponentInst->DetachFromParent(true);
 			}
-			else if (bSimulatePhysics == false)	//if we're switching from kinematic to simulated
+			
+			if (bSimulatePhysics == false)	//if we're switching from kinematic to simulated
 			{
-				//if we are the root in hierarchy we must make sure to update children that are welded
+				//we must make sure to update children that are welded
 				TArray<FBodyInstance*> ChildrenBodies;
 				TArray<FName> ChildrenLabels;
 				OwnerComponentInst->GetWeldedBodies(ChildrenBodies, ChildrenLabels);
@@ -2468,6 +2469,7 @@ void FBodyInstance::SetInstanceSimulatePhysics(bool bSimulate, bool bMaintainPhy
 					if (ChildBI != this)
 					{
 						Weld(ChildBI, ChildBI->OwnerComponent->GetSocketTransform(ChildrenLabels[ChildIdx]));
+						ChildBI->WeldParent = this;
 					}
 				}
 			}
diff --git a/Engine/Source/Runtime/Engine/Private/PrimitiveComponentPhysics.cpp b/Engine/Source/Runtime/Engine/Private/PrimitiveComponentPhysics.cpp
--- a/Engine/Source/Runtime/Engine/Private/PrimitiveComponentPhysics.cpp
+++ b/Engine/Source/Runtime/Engine/Private/PrimitiveComponentPhysics.cpp
@@ -723,8 +723,6 @@ void UPrimitiveComponent::WeldTo(USceneComponent* InParent, FName InSocketName /
 
 void UPrimitiveComponent::UnWeldFromParent()
 {
-
-
 	FBodyInstance* NewRootBI = GetBodyInstance(NAME_None, false);
 	UWorld* CurrentWorld = GetWorld();
 	if (NewRootBI == NULL || NewRootBI->bWelded == false || CurrentWorld == nullptr || IsPendingKill())
@@ -745,11 +743,12 @@ void UPrimitiveComponent::UnWeldFromParent()
 				//create new root
 				RootBI->UnWeld(NewRootBI);	//don't bother fixing up shapes if RootComponent is about to be deleted
 			}
-			
+
 			NewRootBI->bWelded = false;
-			NewRootBI->WeldParent = NULL;
+			const FBodyInstance* PrevWeldParent = NewRootBI->WeldParent;
+			NewRootBI->WeldParent = nullptr;
 
-			bool bHasBodySetup = GetBodySetup() != NULL;
+			bool bHasBodySetup = GetBodySetup() != nullptr;
 
 			//if BodyInstance hasn't already been created we need to initialize it
 			if (bHasBodySetup && NewRootBI->IsValidBodyInstance() == false)
@@ -760,6 +759,11 @@ void UPrimitiveComponent::UnWeldFromParent()
 				NewRootBI->bAutoWeld = bPrevAutoWeld;
 			}
 
+			if(PrevWeldParent == nullptr)	//our parent is kinematic so no need to do any unwelding/rewelding of children
+			{
+				return;
+			}
+
 			//now weld its children to it
 			TArray<FBodyInstance*> ChildrenBodies;
 			TArray<FName> ChildrenLabels;
@@ -774,11 +778,9 @@ void UPrimitiveComponent::UnWeldFromParent()
 					{
 						RootBI->UnWeld(ChildBI);
 					}
-					
-					if (bHasBodySetup)
-					{
-						NewRootBI->Weld(ChildBI, ChildBI->OwnerComponent->GetSocketTransform(ChildrenLabels[ChildIdx]));
-					}
+
+					//At this point, NewRootBI must be kinematic because it's being unwelded. It's up to the code that simulates to call Weld on the children as needed
+					ChildBI->WeldParent = nullptr;	//null because we are currently kinematic
 				}
 			}
 		}
