Blueprints: Added native serialization in UFunction for ubergraph fast calls [UE-10415]

diff --git a/Engine/Source/Editor/KismetCompiler/Private/KismetCompilerVMBackend.cpp b/Engine/Source/Editor/KismetCompiler/Private/KismetCompilerVMBackend.cpp
--- a/Engine/Source/Editor/KismetCompiler/Private/KismetCompilerVMBackend.cpp
+++ b/Engine/Source/Editor/KismetCompiler/Private/KismetCompilerVMBackend.cpp
@@ -715,10 +715,6 @@ public:
 
 				if ((FunctionToCall->FirstPropertyToInit == nullptr) && (FunctionToCall->PostConstructLink == nullptr))
 				{
-					FEventGraphFastCallPair& Pair = *new (CompilerContext.NewClass->FastCallPairs) FEventGraphFastCallPair();
-					Pair.FunctionToPatch = FunctionContext.Function;
-					Pair.EventGraphCallOffset = OffsetWithinUbergraph;
-
 					FunctionContext.Function->EventGraphFunction = FunctionToCall;
 					FunctionContext.Function->EventGraphCallOffset = OffsetWithinUbergraph;
 				}
diff --git a/Engine/Source/Runtime/Core/Public/UObject/ObjectVersion.h b/Engine/Source/Runtime/Core/Public/UObject/ObjectVersion.h
--- a/Engine/Source/Runtime/Core/Public/UObject/ObjectVersion.h
+++ b/Engine/Source/Runtime/Core/Public/UObject/ObjectVersion.h
@@ -496,6 +496,8 @@ enum EUnrealEngineObjectUE4Version
 	VER_UE4_FOLIAGE_WITH_ASSET_OR_CLASS,
 	// Allows PhysicsSerializer to serialize shapes and actors for faster load times
 	VER_UE4_BODYINSTANCE_BINARY_SERIALIZATION,
+	// Added fastcall data serialization directly in UFunction
+	VER_UE4_SERIALIZE_BLUEPRINT_EVENTGRAPH_FASTCALLS_IN_UFUNCTION,
 
 	// -----<new versions can be added before this line>-------------------------------------------------
 	// - this needs to be the last line (see note below)
diff --git a/Engine/Source/Runtime/CoreUObject/Private/UObject/Class.cpp b/Engine/Source/Runtime/CoreUObject/Private/UObject/Class.cpp
--- a/Engine/Source/Runtime/CoreUObject/Private/UObject/Class.cpp
+++ b/Engine/Source/Runtime/CoreUObject/Private/UObject/Class.cpp
@@ -3848,6 +3848,18 @@ void UFunction::Serialize( FArchive& Ar )
 		Ar << RepOffset;
 	}
 
+#if !UE_BLUEPRINT_EVENTGRAPH_FASTCALLS
+	// We need to serialize these values even if the feature is disabled, in order to keep the serialization stream in sync
+	UFunction* EventGraphFunction = nullptr;
+	int32 EventGraphCallOffset = 0;
+#endif
+
+	if (Ar.UE4Ver() >= VER_UE4_SERIALIZE_BLUEPRINT_EVENTGRAPH_FASTCALLS_IN_UFUNCTION)
+	{
+		Ar << EventGraphFunction;
+		Ar << EventGraphCallOffset;
+	}
+
 	// Precomputation.
 	if ((Ar.GetPortFlags() & PPF_Duplicate) != 0)
 	{
diff --git a/Engine/Source/Runtime/CoreUObject/Public/UObject/Class.h b/Engine/Source/Runtime/CoreUObject/Public/UObject/Class.h
--- a/Engine/Source/Runtime/CoreUObject/Public/UObject/Class.h
+++ b/Engine/Source/Runtime/CoreUObject/Public/UObject/Class.h
@@ -1197,10 +1197,10 @@ public:
 	UProperty* FirstPropertyToInit;
 
 #if UE_BLUEPRINT_EVENTGRAPH_FASTCALLS
-	// The event graph this function calls in to
+	// The event graph this function calls in to (persistent)
 	UFunction* EventGraphFunction;
 
-	// The state offset inside of the event graph
+	// The state offset inside of the event graph (persistent)
 	int32 EventGraphCallOffset;
 #endif
 
diff --git a/Engine/Source/Runtime/Engine/Classes/Engine/BlueprintGeneratedClass.h b/Engine/Source/Runtime/Engine/Classes/Engine/BlueprintGeneratedClass.h
--- a/Engine/Source/Runtime/Engine/Classes/Engine/BlueprintGeneratedClass.h
+++ b/Engine/Source/Runtime/Engine/Classes/Engine/BlueprintGeneratedClass.h
@@ -437,8 +437,10 @@ public:
 	UPROPERTY()
 	UFunction* UberGraphFunction;
 
+	// This is a list of event graph call function nodes that are simple (no argument) thunks into the event graph (typically used for animation delegates, etc...)
+	// It is a deprecated list only used for backwards compatibility prior to VER_UE4_SERIALIZE_BLUEPRINT_EVENTGRAPH_FASTCALLS_IN_UFUNCTION.
 	UPROPERTY()
-	TArray<FEventGraphFastCallPair> FastCallPairs;
+	TArray<FEventGraphFastCallPair> FastCallPairs_DEPRECATED;
 
 #if WITH_EDITORONLY_DATA
 	UPROPERTY(Transient)
diff --git a/Engine/Source/Runtime/Engine/Private/BlueprintGeneratedClass.cpp b/Engine/Source/Runtime/Engine/Private/BlueprintGeneratedClass.cpp
--- a/Engine/Source/Runtime/Engine/Private/BlueprintGeneratedClass.cpp
+++ b/Engine/Source/Runtime/Engine/Private/BlueprintGeneratedClass.cpp
@@ -84,7 +84,7 @@ void UBlueprintGeneratedClass::PostLoad()
 
 #if UE_BLUEPRINT_EVENTGRAPH_FASTCALLS
 	// Patch the fast calls (needed as we can't bump engine version to serialize it directly in UFunction right now)
-	for (const FEventGraphFastCallPair& Pair : FastCallPairs)
+	for (const FEventGraphFastCallPair& Pair : FastCallPairs_DEPRECATED)
 	{
 		Pair.FunctionToPatch->EventGraphFunction = UberGraphFunction;
 		Pair.FunctionToPatch->EventGraphCallOffset = Pair.EventGraphCallOffset;
@@ -691,7 +691,7 @@ void UBlueprintGeneratedClass::PurgeClass(bool bRecompilingOnLoad)
 #endif //WITH_EDITOR
 
 #if UE_BLUEPRINT_EVENTGRAPH_FASTCALLS
-	FastCallPairs.Empty();
+	FastCallPairs_DEPRECATED.Empty();
 #endif
 }
 
