Bone Socket particle module now uses it's parent skeletal mesh component if it does not have a skel mesh actor instance param set.
Allows use with anim notifies.

diff --git a/Engine/Source/Runtime/Engine/Private/Particles/ParticleModules_Location.cpp b/Engine/Source/Runtime/Engine/Private/Particles/ParticleModules_Location.cpp
--- a/Engine/Source/Runtime/Engine/Private/Particles/ParticleModules_Location.cpp
+++ b/Engine/Source/Runtime/Engine/Private/Particles/ParticleModules_Location.cpp
@@ -1864,6 +1864,12 @@ USkeletalMeshComponent* UParticleModuleLocationBoneSocket::GetSkeletalMeshCompon
 			//@todo. Warn about this...
 		}
 	}
+
+	if (USkeletalMeshComponent* SkelMesh = Cast<USkeletalMeshComponent>(PSysComp->GetAttachParent()))
+	{
+		return SkelMesh;
+	}
+
 	return NULL;
 }
 
