Map travel failures in PIE will now simply shut down PIE rather than try to load the default map. Avoids potential fatal error if default map subsequently fails to load.

diff --git a/Engine/Source/Editor/UnrealEd/Classes/Editor/EditorEngine.h b/Engine/Source/Editor/UnrealEd/Classes/Editor/EditorEngine.h
--- a/Engine/Source/Editor/UnrealEd/Classes/Editor/EditorEngine.h
+++ b/Engine/Source/Editor/UnrealEd/Classes/Editor/EditorEngine.h
@@ -2694,6 +2694,8 @@ public:
 	DEPRECATED(4.8, "This function is deprecated.  Please use FActorLabelUtilities::SetActorLabelUnique instead")
 	void SetActorLabelUnique( AActor* Actor, const FString& NewActorLabel, const FCachedActorLabels* InExistingActorLabels = nullptr ) const;
 
+	virtual void HandleTravelFailure(UWorld* InWorld, ETravelFailure::Type FailureType, const FString& ErrorString);
+
 private:
 	FTimerHandle CleanupPIEOnlineSessionsTimerHandle;
 };
diff --git a/Engine/Source/Editor/UnrealEd/Private/EditorEngine.cpp b/Engine/Source/Editor/UnrealEd/Private/EditorEngine.cpp
--- a/Engine/Source/Editor/UnrealEd/Private/EditorEngine.cpp
+++ b/Engine/Source/Editor/UnrealEd/Private/EditorEngine.cpp
@@ -6243,4 +6243,27 @@ void FActorLabelUtilities::SetActorLabelUnique(AActor* Actor, const FString& New
 	Actor->SetActorLabel(ModifiedActorLabel);
 }
 
+void UEditorEngine::HandleTravelFailure(UWorld* InWorld, ETravelFailure::Type FailureType, const FString& ErrorString)
+{
+	if (InWorld && InWorld->IsPlayInEditor())
+	{
+		// Default behavior will try to fall back to default map and potentially throw a fatal error if that fails.
+		// Rather than bringing down the whole editor if this happens during a PIE session, just throw a warning and abort the PIE session.
+		{
+			FFormatNamedArguments Arguments;
+			Arguments.Add(TEXT("FailureType"), FText::FromString(ETravelFailure::ToString(FailureType)));
+			Arguments.Add(TEXT("ErrorString"), FText::FromString(ErrorString));
+			FText ErrorMsg = FText::Format(LOCTEXT("PIETravelFailure", "TravelFailure: {FailureType}, Reason for Failure: '{ErrorString}'. Shutting down PIE."), Arguments);
+			UE_LOG(LogNet, Warning, TEXT("%s"), *ErrorMsg.ToString());
+			FMessageLog("PIE").Warning(ErrorMsg);
+		}
+
+		RequestEndPlayMap();
+	}
+	else
+	{
+		Super::HandleTravelFailure(InWorld, FailureType, ErrorString);
+	}
+}
+
 #undef LOCTEXT_NAMESPACE 
diff --git a/Engine/Source/Runtime/Engine/Private/UnrealEngine.cpp b/Engine/Source/Runtime/Engine/Private/UnrealEngine.cpp
--- a/Engine/Source/Runtime/Engine/Private/UnrealEngine.cpp
+++ b/Engine/Source/Runtime/Engine/Private/UnrealEngine.cpp
@@ -7898,35 +7898,20 @@ void UEngine::HandleTravelFailure(UWorld* InWorld, ETravelFailure::Type FailureT
 	if (InWorld == NULL)
 	{
 		UE_LOG(LogNet, Error, TEXT("TravelFailure: %s, Reason for Failure: '%s' with a NULL UWorld"), ETravelFailure::ToString(FailureType), *ErrorString);
-		return;
 	}
+	else
+	{
+		UE_LOG(LogNet, Log, TEXT("TravelFailure: %s, Reason for Failure: '%s'"), ETravelFailure::ToString(FailureType), *ErrorString);
 
-	UE_LOG(LogNet, Log, TEXT("TravelFailure: %s, Reason for Failure: '%s'"), ETravelFailure::ToString(FailureType), *ErrorString);
-
-	// Give the GameInstance a chance to handle the failure.
-	HandleTravelFailure_NotifyGameInstance(InWorld, FailureType);
+		// Give the GameInstance a chance to handle the failure.
+		HandleTravelFailure_NotifyGameInstance(InWorld, FailureType);
 
-	ENetMode NetMode = GetNetMode(InWorld);
+		// Cancel pending net game if there was one
+		CancelPending(InWorld);
 
-	switch (FailureType)
-	{
-	case ETravelFailure::PackageMissing:
-	case ETravelFailure::PackageVersion:
-	case ETravelFailure::NoDownload:
-	case ETravelFailure::NoLevel:
-	case ETravelFailure::InvalidURL:
-	case ETravelFailure::TravelFailure:
-	case ETravelFailure::CheatCommands:
-	case ETravelFailure::PendingNetGameCreateFailure:
-	default:
-		break;
+		// Any of these errors should attempt to load back to some stable map
+		CallHandleDisconnectForFailure(InWorld, InWorld->GetNetDriver());
 	}
-
-	// Cancel pending net game if there was one
-	CancelPending(InWorld);
-
-	// Any of these errors should attempt to load back to some stable map
-	CallHandleDisconnectForFailure(InWorld, InWorld->GetNetDriver());
 }
 
 void UEngine::HandleNetworkFailure(UWorld *World, UNetDriver *NetDriver, ENetworkFailure::Type FailureType, const FString& ErrorString)
