Fix PNG frame export from matinee writing alpha values

diff --git a/Engine/Source/Runtime/Engine/Private/UnrealEngine.cpp b/Engine/Source/Runtime/Engine/Private/UnrealEngine.cpp
--- a/Engine/Source/Runtime/Engine/Private/UnrealEngine.cpp
+++ b/Engine/Source/Runtime/Engine/Private/UnrealEngine.cpp
@@ -10794,6 +10794,13 @@ void FSystemResolution::RequestResolutionChange(int32 InResX, int32 InResY, EWin
 	CVarSystemResolution->Set(*NewValue, ECVF_SetByConsole);
 }
 
+void RemoveAlphaFromColors( TArray<FColor>& Colors )
+{
+	for( FColor& Color : Colors )
+	{
+		Color.A = 255;
+	}
+}
 
 void UEngine::HandleScreenshotCaptured(int32 Width, int32 Height, const TArray<FColor>& Colors)
 {
@@ -10847,7 +10854,12 @@ void UEngine::HandleScreenshotCaptured(int32 Width, int32 Height, const TArray<F
 				if (Filename.Len() > 0)
 				{
 					IImageWrapperPtr ImageWrapper = ImageWrapperModule.CreateImageWrapper(EImageFormat::PNG);
-					if (ImageWrapper.IsValid() && ImageWrapper->SetRaw(&Colors[0], Colors.Num() * sizeof(FColor), Width, Height, ERGBFormat::BGRA, 8))
+
+					TArray<FColor> FinalColors = Colors;
+					// Matinee PNG export does not support writing alpha values
+					RemoveAlphaFromColors( FinalColors );
+
+					if (ImageWrapper.IsValid() && ImageWrapper->SetRaw(&FinalColors[0], FinalColors.Num() * sizeof(FColor), Width, Height, ERGBFormat::BGRA, 8))
 					{
 						FFileHelper::SaveArrayToFile(ImageWrapper->GetCompressed(), *Filename);
 					}
