#ue4 - Warn when trying to move a component with static mobility from a MovementComponent.

diff --git a/Engine/Source/Runtime/Engine/Classes/GameFramework/MovementComponent.h b/Engine/Source/Runtime/Engine/Classes/GameFramework/MovementComponent.h
--- a/Engine/Source/Runtime/Engine/Classes/GameFramework/MovementComponent.h
+++ b/Engine/Source/Runtime/Engine/Classes/GameFramework/MovementComponent.h
@@ -85,6 +85,11 @@ class ENGINE_API UMovementComponent : public UActorComponent
 
 private:
 
+#if WITH_EDITOR
+	// Warned about trying to move something with static mobility.
+	uint32 bEditorWarnedStaticMobilityMove:1;
+#endif
+
 	/**
 	 * Setting that controls behavior when movement is restricted to a 2D plane defined by a specific axis/normal,
 	 * so that movement along the locked axis is not be possible.
diff --git a/Engine/Source/Runtime/Engine/Private/Components/MovementComponent.cpp b/Engine/Source/Runtime/Engine/Private/Components/MovementComponent.cpp
--- a/Engine/Source/Runtime/Engine/Private/Components/MovementComponent.cpp
+++ b/Engine/Source/Runtime/Engine/Private/Components/MovementComponent.cpp
@@ -125,6 +125,11 @@ void UMovementComponent::OnRegister()
 	{
 		SnapUpdatedComponentToPlane();
 	}
+
+#if WITH_EDITOR
+	// Reset so next PIE session warns.
+	bEditorWarnedStaticMobilityMove = false;
+#endif
 }
 
 void UMovementComponent::RegisterComponentTickFunctions(bool bRegister)
@@ -235,10 +240,30 @@ bool UMovementComponent::IsInWater() const
 
 bool UMovementComponent::ShouldSkipUpdate(float DeltaTime) const
 {
-	if (UpdatedComponent == NULL || UpdatedComponent->Mobility != EComponentMobility::Movable)
+	if (UpdatedComponent == nullptr)
 	{
 		return true;
 	}
+		
+	if (UpdatedComponent->Mobility != EComponentMobility::Movable)
+	{
+#if WITH_EDITOR
+		if (!bEditorWarnedStaticMobilityMove)
+		{
+			if (UWorld * World = GetWorld())
+			{
+				if (World->HasBegunPlay() && IsRegistered())
+				{
+					const_cast<UMovementComponent*>(this)->bEditorWarnedStaticMobilityMove = true;
+					FMessageLog("PIE").Warning(FText::Format(LOCTEXT("InvalidMove", "Mobility of {0} : {1} has to be 'Movable' if you'd like to move it with {2}. "),
+						FText::FromString(GetNameSafe(UpdatedComponent->GetOwner())), FText::FromString(UpdatedComponent->GetName()), FText::FromString(GetClass()->GetName())));
+				}
+			}
+		}
+#endif
+
+		return true;
+	}
 
 	if (bUpdateOnlyIfRendered)
 	{
