[UE-13774] Improved fix (instead of CL 2516152), where we repair invalid SCS scene hierarchies (to handle cases where the inherited component hierarchy has been removed and a single scene root node has to be promoted).

diff --git a/Engine/Source/Runtime/Engine/Classes/Engine/SimpleConstructionScript.h b/Engine/Source/Runtime/Engine/Classes/Engine/SimpleConstructionScript.h
--- a/Engine/Source/Runtime/Engine/Classes/Engine/SimpleConstructionScript.h
+++ b/Engine/Source/Runtime/Engine/Classes/Engine/SimpleConstructionScript.h
@@ -92,6 +92,13 @@ private:
 	/** Helper function to find an appropriate child node index that can be promoted to the parent's level */
 	int32 FindPromotableChildNodeIndex(USCS_Node* InParentNode) const;
 
+	/** 
+	 * Helper function to repair invalid SCS hierarchies (like when a 
+	 * Blueprint's super class has had all it components removed, and one of  
+	 * this Blueprint's components is promoted to scene root). 
+	 */
+	void FixupSceneNodeHierarchy();
+
 #if WITH_EDITOR
 	/** Helper function to find the current scene root component template and/or owning SCS node */
 	USceneComponent* GetSceneRootComponentTemplate(USCS_Node** OutSCSNode = nullptr) const;
diff --git a/Engine/Source/Runtime/Engine/Private/SimpleConstructionScript.cpp b/Engine/Source/Runtime/Engine/Private/SimpleConstructionScript.cpp
--- a/Engine/Source/Runtime/Engine/Private/SimpleConstructionScript.cpp
+++ b/Engine/Source/Runtime/Engine/Private/SimpleConstructionScript.cpp
@@ -162,9 +162,6 @@ void USimpleConstructionScript::PostLoad()
 		return;
 	}
 
-	USCS_Node* SceneRootNode = nullptr;
-	USceneComponent* SceneRootComponentTemplate = GetSceneRootComponentTemplate(&SceneRootNode);
-
 	for (NodeIndex=0; NodeIndex < Nodes.Num(); ++NodeIndex)
 	{
 		USCS_Node* Node = Nodes[NodeIndex];
@@ -178,36 +175,8 @@ void USimpleConstructionScript::PostLoad()
 		// Fix up components that may have switched from scene to non-scene type and vice-versa
 		if(Node->ComponentTemplate != nullptr)
 		{
-			// Check to see if switched from non-scene to scene component type
-			if(Node->ComponentTemplate->IsA<USceneComponent>())
-			{
-				if(RootNodes.Contains(Node)
-					&& SceneRootComponentTemplate != nullptr
-					&& Node->ComponentTemplate != SceneRootComponentTemplate
-					&& Node->ParentComponentOrVariableName == NAME_None)
-				{
-					if(SceneRootNode != nullptr)
-					{
-						if(RootNodes.Contains(SceneRootNode))
-						{
-							// Reparent to this BP's root node if it's still in the root set
-							RootNodes.Remove(Node);
-							SceneRootNode->ChildNodes.Add(Node);
-						}
-						else
-						{
-							// Parent to an inherited parent BP's node if not already attached
-							Node->SetParent(SceneRootNode);
-						}
-					}
-					else
-					{
-						// Parent to the native component template if not already attached
-						Node->SetParent(SceneRootComponentTemplate);
-					}
-				}
-			}
-			else
+			// Check to see if switched from scene to a non-scene component type
+			if (!Node->ComponentTemplate->IsA<USceneComponent>())
 			{
 				// Otherwise, check to see if switched from scene to non-scene component type
 				int32 RootNodeIndex = INDEX_NONE;
@@ -309,6 +278,69 @@ void USimpleConstructionScript::PostLoad()
 	}
 }
 
+void USimpleConstructionScript::FixupSceneNodeHierarchy() 
+{
+	// determine the scene's root component, this isn't necessarily a node owned
+	// by this SCS; it could be from a super SCS, or (if SceneRootNode and 
+	// SceneRootComponentTemplate is not) it could be a native component
+	USCS_Node* SceneRootNode = nullptr;
+	USceneComponent* SceneRootComponentTemplate = GetSceneRootComponentTemplate(&SceneRootNode);
+
+	// if there is no scene root (then there shouldn't be anything but the 
+	// default placeholder root).
+	if (SceneRootComponentTemplate == nullptr)
+	{
+		return;
+	}
+
+	bool const bIsSceneRootNative = (SceneRootNode == nullptr);
+	bool const bThisOwnsSceneRoot = !bIsSceneRootNative && RootNodes.Contains(SceneRootNode);
+
+	// iterate backwards so that we can remove nodes from the array as we go
+	for (int32 NodeIndex = RootNodes.Num() - 1; NodeIndex >= 0; --NodeIndex)
+	{
+		USCS_Node* Node = RootNodes[NodeIndex];
+
+		// we only care about the scene component hierarchy (non-scene components 
+		// can share root placement)
+		if ((Node->ComponentTemplate == nullptr) || !Node->ComponentTemplate->IsA<USceneComponent>())
+		{
+			continue;
+		}
+
+		// if this is the scene's root, then we shouldn't fix it up (instead we 
+		// need to be nesting others under this one)
+		if (SceneRootComponentTemplate == Node->ComponentTemplate)
+		{
+			continue;
+		}
+
+		// if this node has a clear parent already defined, then ignore it (I 
+		// imagine that its attachment will be handled elsewhere)
+		if (Node->ParentComponentOrVariableName != NAME_None)
+		{
+			continue;
+		}
+
+		if (bIsSceneRootNative)
+		{
+			// Parent to the native component template if not already attached
+			Node->SetParent(SceneRootComponentTemplate);
+		}
+		else if (bThisOwnsSceneRoot)
+		{
+			// Reparent to this BP's root node if it's still in the root set
+			RootNodes.Remove(Node);
+			SceneRootNode->ChildNodes.Add(Node);
+		}
+		else
+		{
+			// Parent to an inherited parent BP's node if not already attached
+			Node->SetParent(SceneRootNode);
+		}
+	}
+}
+
 void USimpleConstructionScript::FixupRootNodeParentReferences()
 {
 	// Get the BlueprintGeneratedClass that owns the SCS
@@ -316,8 +348,7 @@ void USimpleConstructionScript::FixupRootNodeParentReferences()
 	if(BPGeneratedClass == NULL)
 	{
 		UE_LOG(LogBlueprint, Warning, TEXT("USimpleConstructionScript::FixupRootNodeParentReferences() - owner class is NULL; skipping."));
-
-		// Cannot do fixup without a BPGC
+		// cannot do the rest of fixup without a BPGC
 		return;
 	}
 
@@ -393,6 +424,15 @@ void USimpleConstructionScript::FixupRootNodeParentReferences()
 			}
 		}
 	}
+
+	// call this after we do the above ParentComponentOrVariableName fixup, 
+	// because this operates differently for root nodes that have their 
+	// ParentComponentOrVariableName field cleared
+	//
+	// repairs invalid scene hierarchies (like when this Blueprint has been 
+	// reparented and there is no longer an inherited scene root... meaning one
+	// of the scene component nodes here needs to be promoted)
+	FixupSceneNodeHierarchy();
 }
 
 void USimpleConstructionScript::ExecuteScriptOnActor(AActor* Actor, const FTransform& RootTransform, bool bIsDefaultTransform)
