Blueprints: Added GetObjectName to return the exact object name

diff --git a/Engine/Source/Runtime/Engine/Classes/Kismet/KismetSystemLibrary.h b/Engine/Source/Runtime/Engine/Classes/Kismet/KismetSystemLibrary.h
--- a/Engine/Source/Runtime/Engine/Classes/Kismet/KismetSystemLibrary.h
+++ b/Engine/Source/Runtime/Engine/Classes/Kismet/KismetSystemLibrary.h
@@ -79,6 +79,10 @@ class ENGINE_API UKismetSystemLibrary : public UBlueprintFunctionLibrary
 	UFUNCTION(BlueprintPure, Category = "Utilities")
 	static bool IsValidClass(UClass* Class);
 
+	// Returns the actual object name.
+	UFUNCTION(BlueprintPure, Category = "Utilities")
+	static FString GetObjectName(const UObject* Object);
+
 	// Returns the display name (or actor label), for displaying to end users.  Note:  In editor builds, this is the actor label.  In non-editor builds, this is the actual object name.  This function should not be used to uniquely identify actors!
 	UFUNCTION(BlueprintPure, Category="Utilities")
 	static FString GetDisplayName(const UObject* Object);
diff --git a/Engine/Source/Runtime/Engine/Private/KismetSystemLibrary.cpp b/Engine/Source/Runtime/Engine/Private/KismetSystemLibrary.cpp
--- a/Engine/Source/Runtime/Engine/Private/KismetSystemLibrary.cpp
+++ b/Engine/Source/Runtime/Engine/Private/KismetSystemLibrary.cpp
@@ -38,6 +38,11 @@ bool UKismetSystemLibrary::IsValidClass(UClass* Class)
 	return ::IsValid(Class);
 }
 
+FString UKismetSystemLibrary::GetObjectName(const UObject* Object)
+{
+	return GetNameSafe(Object);
+}
+
 FString UKismetSystemLibrary::GetDisplayName(const UObject* Object)
 {
 #if WITH_EDITOR
