Fixing CIS build error (forgot WITH_EDITOR #if).

diff --git a/Engine/Source/Runtime/Engine/Private/SimpleConstructionScript.cpp b/Engine/Source/Runtime/Engine/Private/SimpleConstructionScript.cpp
--- a/Engine/Source/Runtime/Engine/Private/SimpleConstructionScript.cpp
+++ b/Engine/Source/Runtime/Engine/Private/SimpleConstructionScript.cpp
@@ -280,6 +280,7 @@ void USimpleConstructionScript::PostLoad()
 
 void USimpleConstructionScript::FixupSceneNodeHierarchy() 
 {
+#if WITH_EDITOR
 	// determine the scene's root component, this isn't necessarily a node owned
 	// by this SCS; it could be from a super SCS, or (if SceneRootNode and 
 	// SceneRootComponentTemplate is not) it could be a native component
@@ -339,6 +340,7 @@ void USimpleConstructionScript::FixupSceneNodeHierarchy()
 			Node->SetParent(SceneRootNode);
 		}
 	}
+#endif // #if WITH_EDITOR
 }
 
 void USimpleConstructionScript::FixupRootNodeParentReferences()
