Blueprints: Further clarified the comment on the Blueprint-exposed GetDisplayName method

diff --git a/Engine/Source/Runtime/Engine/Classes/Kismet/KismetSystemLibrary.h b/Engine/Source/Runtime/Engine/Classes/Kismet/KismetSystemLibrary.h
--- a/Engine/Source/Runtime/Engine/Classes/Kismet/KismetSystemLibrary.h
+++ b/Engine/Source/Runtime/Engine/Classes/Kismet/KismetSystemLibrary.h
@@ -83,7 +83,9 @@ class ENGINE_API UKismetSystemLibrary : public UBlueprintFunctionLibrary
 	UFUNCTION(BlueprintPure, Category = "Utilities")
 	static FString GetObjectName(const UObject* Object);
 
-	// Returns the display name (or actor label), for displaying to end users.  Note:  In editor builds, this is the actor label.  In non-editor builds, this is the actual object name.  This function should not be used to uniquely identify actors!
+	// Returns the display name (or actor label), for displaying as a debugging aid.
+	// Note: In editor builds, this is the actor label.  In non-editor builds, this is the actual object name.  This function should not be used to uniquely identify actors!
+	// It is not localized and should not be used for display to an end user of a game.
 	UFUNCTION(BlueprintPure, Category="Utilities")
 	static FString GetDisplayName(const UObject* Object);
 
