UE-9597: Trying to create a custom code class derived from ArrowComponent will result in the build failing in Visual Studio.

Added ENGINE_API for UArrowComponent and UBillboardComponent for now.

diff --git a/Engine/Source/Runtime/Engine/Classes/Components/ArrowComponent.h b/Engine/Source/Runtime/Engine/Classes/Components/ArrowComponent.h
--- a/Engine/Source/Runtime/Engine/Classes/Components/ArrowComponent.h
+++ b/Engine/Source/Runtime/Engine/Classes/Components/ArrowComponent.h
@@ -7,8 +7,8 @@
 /** 
  * A simple arrow rendered using lines. Useful for indicating which way an object is facing.
  */
-UCLASS(ClassGroup=Utility, hidecategories=(Object,LOD,Physics,Lighting,TextureStreaming,Activation,"Components|Activation",Collision), editinlinenew, meta=(BlueprintSpawnableComponent), MinimalAPI)
-class UArrowComponent : public UPrimitiveComponent
+UCLASS(ClassGroup=Utility, hidecategories=(Object,LOD,Physics,Lighting,TextureStreaming,Activation,"Components|Activation",Collision), editinlinenew, meta=(BlueprintSpawnableComponent))
+class ENGINE_API UArrowComponent : public UPrimitiveComponent
 {
 	GENERATED_UCLASS_BODY()
 
@@ -70,7 +70,7 @@ class UArrowComponent : public UPrimitiveComponent
 
 #if WITH_EDITORONLY_DATA
 	/** Set the scale that we use when rendering in-editor */
-	ENGINE_API static void SetEditorScale(float InEditorScale);
+	static void SetEditorScale(float InEditorScale);
 
 	/** The scale we use when rendering in-editor */
 	static float EditorScale;
diff --git a/Engine/Source/Runtime/Engine/Classes/Components/BillboardComponent.h b/Engine/Source/Runtime/Engine/Classes/Components/BillboardComponent.h
--- a/Engine/Source/Runtime/Engine/Classes/Components/BillboardComponent.h
+++ b/Engine/Source/Runtime/Engine/Classes/Components/BillboardComponent.h
@@ -10,8 +10,8 @@ struct FConvexVolume;
 /** 
  * A 2d texture that will be rendered always facing the camera.
  */
-UCLASS(ClassGroup=Rendering, collapsecategories, hidecategories=(Object,Activation,"Components|Activation",Physics,Collision,Lighting,LOD,Mesh,PhysicsVolume), editinlinenew, meta=(BlueprintSpawnableComponent),MinimalAPI)
-class UBillboardComponent : public UPrimitiveComponent
+UCLASS(ClassGroup=Rendering, collapsecategories, hidecategories=(Object,Activation,"Components|Activation",Physics,Collision,Lighting,LOD,Mesh,PhysicsVolume), editinlinenew, meta=(BlueprintSpawnableComponent))
+class ENGINE_API UBillboardComponent : public UPrimitiveComponent
 {
 	GENERATED_UCLASS_BODY()
 
@@ -73,7 +73,7 @@ class UBillboardComponent : public UPrimitiveComponent
 
 #if WITH_EDITORONLY_DATA
 	/** Set the scale that we use when rendering in-editor */
-	ENGINE_API static void SetEditorScale(float InEditorScale);
+	static void SetEditorScale(float InEditorScale);
 
 	/** The scale we use when rendering in-editor */
 	static float EditorScale;
