Due to long compilation times I removed FPRCEINLINE from Visual Logger. Usually we have simple static functions so it won't hurt us much from performance point of view.
Changed NULL macro to nullptr in visual logger files.

diff --git a/Engine/Source/Runtime/AIModule/Classes/EnvironmentQuery/EnvQueryDebugHelpers.h b/Engine/Source/Runtime/AIModule/Classes/EnvironmentQuery/EnvQueryDebugHelpers.h
--- a/Engine/Source/Runtime/AIModule/Classes/EnvironmentQuery/EnvQueryDebugHelpers.h
+++ b/Engine/Source/Runtime/AIModule/Classes/EnvironmentQuery/EnvQueryDebugHelpers.h
@@ -75,7 +75,7 @@ namespace EQSDebug
 	};
 }
 
-FORCEINLINE
+inline
 FArchive& operator<<(FArchive& Ar, FDebugRenderSceneProxy::FSphere& Data)
 {
 	Ar << Data.Radius;
@@ -84,7 +84,7 @@ FArchive& operator<<(FArchive& Ar, FDebugRenderSceneProxy::FSphere& Data)
 	return Ar;
 }
 
-FORCEINLINE
+inline
 FArchive& operator<<(FArchive& Ar, FDebugRenderSceneProxy::FText3d& Data)
 {
 	Ar << Data.Text;
@@ -93,7 +93,7 @@ FArchive& operator<<(FArchive& Ar, FDebugRenderSceneProxy::FText3d& Data)
 	return Ar;
 }
 
-FORCEINLINE
+inline
 FArchive& operator<<(FArchive& Ar, EQSDebug::FItemData& Data)
 {
 	Ar << Data.Desc;
@@ -104,7 +104,7 @@ FArchive& operator<<(FArchive& Ar, EQSDebug::FItemData& Data)
 	return Ar;
 }
 
-FORCEINLINE
+inline
 FArchive& operator<<(FArchive& Ar, EQSDebug::FTestData& Data)
 {
 	Ar << Data.ShortName;
@@ -112,7 +112,7 @@ FArchive& operator<<(FArchive& Ar, EQSDebug::FTestData& Data)
 	return Ar;
 }
 
-FORCEINLINE
+inline
 FArchive& operator<<(FArchive& Ar, EQSDebug::FDebugHelper& Data)
 {
 	Ar << Data.Location;
@@ -122,7 +122,7 @@ FArchive& operator<<(FArchive& Ar, EQSDebug::FDebugHelper& Data)
 	return Ar;
 }
 
-FORCEINLINE
+inline
 FArchive& operator<<(FArchive& Ar, EQSDebug::FQueryData& Data)
 {
 	Ar << Data.Items;
@@ -167,11 +167,11 @@ private:
 };
 
 #if ENABLE_VISUAL_LOG && USE_EQS_DEBUGGER
-FORCEINLINE void UEnvQueryDebugHelpers::LogQuery(FEnvQueryInstance& Query, const FLogCategoryBase& Category, ELogVerbosity::Type Verbosity)
+inline void UEnvQueryDebugHelpers::LogQuery(FEnvQueryInstance& Query, const FLogCategoryBase& Category, ELogVerbosity::Type Verbosity)
 {
-	UWorld *World = NULL;
-	FVisualLogEntry *CurrentEntry = NULL;
-	if (CheckVisualLogInputInternal(Query.Owner.Get(), Category, Verbosity, &World, &CurrentEntry) == false)
+	UWorld *World = nullptr;
+	FVisualLogEntry *CurrentEntry = nullptr;
+	if (FVisualLogger::CheckVisualLogInputInternal(Query.Owner.Get(), Category, Verbosity, &World, &CurrentEntry) == false)
 	{
 		return;
 	}
diff --git a/Engine/Source/Runtime/Engine/Private/VisualLogger/VisualLogger.cpp b/Engine/Source/Runtime/Engine/Private/VisualLogger/VisualLogger.cpp
--- a/Engine/Source/Runtime/Engine/Private/VisualLogger/VisualLogger.cpp
+++ b/Engine/Source/Runtime/Engine/Private/VisualLogger/VisualLogger.cpp
@@ -22,6 +22,387 @@ DEFINE_LOG_CATEGORY(LogVisual);
 
 TMap<UObject*, TArray<TWeakObjectPtr<const UObject> > > FVisualLogger::RedirectionMap;
 
+// Unfortunately needs to be a #define since it uses GET_VARARGS_RESULT which uses the va_list stuff which operates on the
+// current function, so we can't easily call a function
+#define COLLAPSED_LOGF(SerializeFunc) \
+	int32	BufferSize	= 1024; \
+	TCHAR*	Buffer		= nullptr; \
+	int32	Result		= -1; \
+	/* allocate some stack space to use on the first pass, which matches most strings */ \
+	TCHAR	StackBuffer[512]; \
+	TCHAR*	AllocatedBuffer = nullptr; \
+	\
+	/* first, try using the stack buffer */ \
+	Buffer = StackBuffer; \
+	GET_VARARGS_RESULT( Buffer, ARRAY_COUNT(StackBuffer), ARRAY_COUNT(StackBuffer) - 1, Fmt, Fmt, Result ); \
+	\
+	/* if that fails, then use heap allocation to make enough space */ \
+		while(Result == -1) \
+			{ \
+		FMemory::SystemFree(AllocatedBuffer); \
+		/* We need to use malloc here directly as GMalloc might not be safe. */ \
+		Buffer = AllocatedBuffer = (TCHAR*) FMemory::SystemMalloc( BufferSize * sizeof(TCHAR) ); \
+		GET_VARARGS_RESULT( Buffer, BufferSize, BufferSize-1, Fmt, Fmt, Result ); \
+		BufferSize *= 2; \
+			}; \
+	Buffer[Result] = 0; \
+	; \
+	\
+	SerializeFunc; \
+	FMemory::SystemFree(AllocatedBuffer);
+
+inline
+bool FVisualLogger::CheckVisualLogInputInternal(const class UObject* Object, const struct FLogCategoryBase& Category, ELogVerbosity::Type Verbosity, UWorld **World, FVisualLogEntry **CurrentEntry)
+{
+	FVisualLogger& VisualLogger = FVisualLogger::Get();
+	if (!Object || (GEngine && GEngine->bDisableAILogging) || VisualLogger.IsRecording() == false || Object->HasAnyFlags(RF_ClassDefaultObject))
+	{
+		return false;
+	}
+	const FName CategoryName = Category.GetCategoryName();
+	if (VisualLogger.IsBlockedForAllCategories() && VisualLogger.IsWhiteListed(CategoryName) == false)
+	{
+		return false;
+	}
+
+	*World = GEngine->GetWorldFromContextObject(Object, false);
+	if (ensure(*World != nullptr) == false)
+	{
+		return false;
+	}
+
+	*CurrentEntry = VisualLogger.GetEntryToWrite(Object, (*World)->TimeSeconds);
+	if (ensure(CurrentEntry != nullptr) == false)
+	{
+		return false;
+	}
+
+	return true;
+}
+
+inline
+FVisualLogEntry* FVisualLogger::GetEntryToWrite(const class UObject* Object, float TimeStamp, ECreateIfNeeded ShouldCreate)
+{
+	FVisualLogEntry* CurrentEntry = nullptr;
+	UObject * LogOwner = FVisualLogger::FindRedirection(Object);
+
+	bool InitializeNewEntry = false;
+
+	TWeakObjectPtr<UWorld> World = GetWorld(Object);
+
+	if (CurrentEntryPerObject.Contains(LogOwner))
+	{
+		CurrentEntry = &CurrentEntryPerObject[LogOwner];
+		InitializeNewEntry = TimeStamp > CurrentEntry->TimeStamp && ShouldCreate == ECreateIfNeeded::Create;
+		if (World.IsValid())
+		{
+			World->GetTimerManager().ClearTimer(VisualLoggerCleanupTimerHandle);
+			for (auto& CurrentPair : CurrentEntryPerObject)
+			{
+				FVisualLogEntry* Entry = &CurrentPair.Value;
+				if (Entry->TimeStamp >= 0 && Entry->TimeStamp < TimeStamp)
+				{
+					for (auto* Device : OutputDevices)
+					{
+						Device->Serialize(CurrentPair.Key, ObjectToNameMap[CurrentPair.Key], ObjectToClassNameMap[CurrentPair.Key], *Entry);
+					}
+					Entry->Reset();
+				}
+			}
+		}
+	}
+
+	if (!CurrentEntry)
+	{
+		CurrentEntry = &CurrentEntryPerObject.Add(LogOwner);
+		ObjectToNameMap.Add(LogOwner, LogOwner->GetFName());
+		ObjectToClassNameMap.Add(LogOwner, *(LogOwner->GetClass()->GetName()));
+		ObjectToPointerMap.Add(LogOwner, LogOwner);
+		InitializeNewEntry = true;
+	}
+
+	if (InitializeNewEntry)
+	{
+		CurrentEntry->Reset();
+		CurrentEntry->TimeStamp = TimeStamp;
+
+		if (RedirectionMap.Contains(LogOwner))
+		{
+			if (ObjectToPointerMap.Contains(LogOwner) && ObjectToPointerMap[LogOwner].IsValid())
+			{
+				const class AActor* LogOwnerAsActor = Cast<class AActor>(LogOwner);
+				if (LogOwnerAsActor)
+				{
+					LogOwnerAsActor->GrabDebugSnapshot(CurrentEntry);
+				}
+			}
+			for (auto Child : RedirectionMap[LogOwner])
+			{
+				if (Child.IsValid())
+				{
+					const class AActor* ChildAsActor = Cast<class AActor>(Child.Get());
+					if (ChildAsActor)
+					{
+						ChildAsActor->GrabDebugSnapshot(CurrentEntry);
+					}
+				}
+			}
+		}
+		else
+		{
+			const class AActor* ObjectAsActor = Cast<class AActor>(Object);
+			if (ObjectAsActor)
+			{
+				CurrentEntry->Location = ObjectAsActor->GetActorLocation();
+				ObjectAsActor->GrabDebugSnapshot(CurrentEntry);
+			}
+		}
+	}
+
+	if (World.IsValid())
+	{
+		//set next tick timer to flush obsolete/old entries
+		World->GetTimerManager().SetTimer(VisualLoggerCleanupTimerHandle, FTimerDelegate::CreateLambda(
+			[this, World](){
+			for (auto& CurrentPair : CurrentEntryPerObject)
+			{
+				FVisualLogEntry* Entry = &CurrentPair.Value;
+				if (Entry->TimeStamp >= 0 && (!World.IsValid() || Entry->TimeStamp < World->GetTimeSeconds())) // CurrentEntry->TimeStamp == -1 means it's not initialized entry information
+				{
+					for (auto* Device : OutputDevices)
+					{
+						Device->Serialize(CurrentPair.Key, ObjectToNameMap[CurrentPair.Key], ObjectToClassNameMap[CurrentPair.Key], *Entry);
+					}
+					Entry->Reset();
+				}
+			}
+		}
+		), 0.1, false);
+	}
+
+	return CurrentEntry;
+}
+
+inline
+void FVisualLogger::Flush()
+{
+	for (auto &CurrentEntry : CurrentEntryPerObject)
+	{
+		if (CurrentEntry.Value.TimeStamp >= 0)
+		{
+			for (auto* Device : OutputDevices)
+			{
+				Device->Serialize(CurrentEntry.Key, ObjectToNameMap[CurrentEntry.Key], ObjectToClassNameMap[CurrentEntry.Key], CurrentEntry.Value);
+			}
+			CurrentEntry.Value.Reset();
+		}
+	}
+}
+
+
+inline
+VARARG_BODY(void, FVisualLogger::CategorizedLogf, const TCHAR*, VARARG_EXTRA(const class UObject* Object) VARARG_EXTRA(const struct FLogCategoryBase& Category) VARARG_EXTRA(ELogVerbosity::Type Verbosity) VARARG_EXTRA(int32 UniqueLogId))
+{
+	SCOPE_CYCLE_COUNTER(STAT_VisualLog);
+	UWorld *World = nullptr;
+	FVisualLogEntry *CurrentEntry = nullptr;
+	if (CheckVisualLogInputInternal(Object, Category, Verbosity, &World, &CurrentEntry) == false)
+	{
+		return;
+	}
+
+	COLLAPSED_LOGF(
+		CurrentEntry->AddText(Buffer, Category.GetCategoryName(), Verbosity);
+	);
+}
+
+inline
+VARARG_BODY(void, FVisualLogger::GeometryShapeLogf, const TCHAR*, VARARG_EXTRA(const class UObject* Object) VARARG_EXTRA(const struct FLogCategoryBase& Category) VARARG_EXTRA(ELogVerbosity::Type Verbosity) VARARG_EXTRA(int32 UniqueLogId) VARARG_EXTRA(const FVector& Start) VARARG_EXTRA(const FVector& End) VARARG_EXTRA(const FColor& Color))
+{
+	SCOPE_CYCLE_COUNTER(STAT_VisualLog);
+	UWorld *World = nullptr;
+	FVisualLogEntry *CurrentEntry = nullptr;
+	if (CheckVisualLogInputInternal(Object, Category, Verbosity, &World, &CurrentEntry) == false)
+	{
+		return;
+	}
+
+	COLLAPSED_LOGF(
+		CurrentEntry->AddElement(Start, End, Category.GetCategoryName(), Verbosity, Color, Buffer);
+	);
+}
+
+inline
+VARARG_BODY(void, FVisualLogger::GeometryShapeLogf, const TCHAR*, VARARG_EXTRA(const class UObject* Object) VARARG_EXTRA(const struct FLogCategoryBase& Category) VARARG_EXTRA(ELogVerbosity::Type Verbosity) VARARG_EXTRA(int32 UniqueLogId) VARARG_EXTRA(const FVector& Location) VARARG_EXTRA(float Radius) VARARG_EXTRA(const FColor& Color))
+{
+	SCOPE_CYCLE_COUNTER(STAT_VisualLog);
+	UWorld *World = nullptr;
+	FVisualLogEntry *CurrentEntry = nullptr;
+	if (CheckVisualLogInputInternal(Object, Category, Verbosity, &World, &CurrentEntry) == false)
+	{
+		return;
+	}
+
+	COLLAPSED_LOGF(
+		CurrentEntry->AddElement(Location, Category.GetCategoryName(), Verbosity, Color, Buffer, Radius);
+	);
+}
+
+inline
+VARARG_BODY(void, FVisualLogger::GeometryShapeLogf, const TCHAR*, VARARG_EXTRA(const class UObject* Object) VARARG_EXTRA(const struct FLogCategoryBase& Category) VARARG_EXTRA(ELogVerbosity::Type Verbosity) VARARG_EXTRA(int32 UniqueLogId) VARARG_EXTRA(const FBox& Box) VARARG_EXTRA(const FMatrix& Matrix) VARARG_EXTRA(const FColor& Color))
+{
+	SCOPE_CYCLE_COUNTER(STAT_VisualLog);
+	UWorld *World = nullptr;
+	FVisualLogEntry *CurrentEntry = nullptr;
+	if (CheckVisualLogInputInternal(Object, Category, Verbosity, &World, &CurrentEntry) == false)
+	{
+		return;
+	}
+
+	COLLAPSED_LOGF(
+		CurrentEntry->AddElement(Box, Matrix, Category.GetCategoryName(), Verbosity, Color, Buffer);
+	);
+}
+
+inline
+VARARG_BODY(void, FVisualLogger::GeometryShapeLogf, const TCHAR*, VARARG_EXTRA(const class UObject* Object) VARARG_EXTRA(const struct FLogCategoryBase& Category) VARARG_EXTRA(ELogVerbosity::Type Verbosity) VARARG_EXTRA(int32 UniqueLogId) VARARG_EXTRA(const FVector& Orgin) VARARG_EXTRA(const FVector& Direction) VARARG_EXTRA(const float Length) VARARG_EXTRA(const float Angle)  VARARG_EXTRA(const FColor& Color))
+{
+	SCOPE_CYCLE_COUNTER(STAT_VisualLog);
+	UWorld *World = nullptr;
+	FVisualLogEntry *CurrentEntry = nullptr;
+	if (CheckVisualLogInputInternal(Object, Category, Verbosity, &World, &CurrentEntry) == false)
+	{
+		return;
+	}
+
+	COLLAPSED_LOGF(
+		CurrentEntry->AddElement(Orgin, Direction, Length, Angle, Angle, Category.GetCategoryName(), Verbosity, Color, Buffer);
+	);
+}
+
+inline
+VARARG_BODY(void, FVisualLogger::GeometryShapeLogf, const TCHAR*, VARARG_EXTRA(const class UObject* Object) VARARG_EXTRA(const struct FLogCategoryBase& Category) VARARG_EXTRA(ELogVerbosity::Type Verbosity) VARARG_EXTRA(int32 UniqueLogId) VARARG_EXTRA(const FVector& Start) VARARG_EXTRA(const FVector& End) VARARG_EXTRA(const float Radius) VARARG_EXTRA(const FColor& Color))
+{
+	SCOPE_CYCLE_COUNTER(STAT_VisualLog);
+	UWorld *World = nullptr;
+	FVisualLogEntry *CurrentEntry = nullptr;
+	if (CheckVisualLogInputInternal(Object, Category, Verbosity, &World, &CurrentEntry) == false)
+	{
+		return;
+	}
+
+	COLLAPSED_LOGF(
+		CurrentEntry->AddElement(Start, End, Radius, Category.GetCategoryName(), Verbosity, Color, Buffer);
+	);
+}
+
+inline
+VARARG_BODY(void, FVisualLogger::GeometryShapeLogf, const TCHAR*, VARARG_EXTRA(const class UObject* Object) VARARG_EXTRA(const struct FLogCategoryBase& Category) VARARG_EXTRA(ELogVerbosity::Type Verbosity) VARARG_EXTRA(int32 UniqueLogId) VARARG_EXTRA(const FVector& Center) VARARG_EXTRA(float HalfHeight) VARARG_EXTRA(float Radius) VARARG_EXTRA(const FQuat& Rotation) VARARG_EXTRA(const FColor& Color))
+{
+	SCOPE_CYCLE_COUNTER(STAT_VisualLog);
+	UWorld *World = nullptr;
+	FVisualLogEntry *CurrentEntry = nullptr;
+	if (CheckVisualLogInputInternal(Object, Category, Verbosity, &World, &CurrentEntry) == false)
+	{
+		return;
+	}
+
+	COLLAPSED_LOGF(
+		CurrentEntry->AddElement(Center, HalfHeight, Radius, Rotation, Category.GetCategoryName(), Verbosity, Color, Buffer);
+	);
+}
+
+inline
+VARARG_BODY(void, FVisualLogger::HistogramDataLogf, const TCHAR*, VARARG_EXTRA(const class UObject* Object) VARARG_EXTRA(const struct FLogCategoryBase& Category) VARARG_EXTRA(ELogVerbosity::Type Verbosity) VARARG_EXTRA(int32 UniqueLogId) VARARG_EXTRA(FName GraphName) VARARG_EXTRA(FName DataName) VARARG_EXTRA(const FVector2D& Data) VARARG_EXTRA(const FColor& Color))
+{
+	SCOPE_CYCLE_COUNTER(STAT_VisualLog);
+	UWorld *World = nullptr;
+	FVisualLogEntry *CurrentEntry = nullptr;
+	if (CheckVisualLogInputInternal(Object, Category, Verbosity, &World, &CurrentEntry) == false)
+	{
+		return;
+	}
+
+	COLLAPSED_LOGF(
+		CurrentEntry->AddHistogramData(Data, Category.GetCategoryName(), Verbosity, GraphName, DataName);
+	);
+}
+
+#undef COLLAPSED_LOGF
+
+inline
+void FVisualLogger::EventLog(const class UObject* Object, const FName EventTag1, const FVisualLogEventBase& Event1, const FVisualLogEventBase& Event2, const FVisualLogEventBase& Event3, const FVisualLogEventBase& Event4, const FVisualLogEventBase& Event5, const FVisualLogEventBase& Event6)
+{
+	EventLog(Object, EventTag1, Event1, Event2, Event3, Event4, Event5);
+	EventLog(Object, EventTag1, Event6);
+}
+
+inline
+void FVisualLogger::EventLog(const class UObject* Object, const FName EventTag1, const FVisualLogEventBase& Event1, const FVisualLogEventBase& Event2, const FVisualLogEventBase& Event3, const FVisualLogEventBase& Event4, const FVisualLogEventBase& Event5)
+{
+	EventLog(Object, EventTag1, Event1, Event2, Event3, Event4);
+	EventLog(Object, EventTag1, Event5);
+}
+
+inline
+void FVisualLogger::EventLog(const class UObject* Object, const FName EventTag1, const FVisualLogEventBase& Event1, const FVisualLogEventBase& Event2, const FVisualLogEventBase& Event3, const FVisualLogEventBase& Event4)
+{
+	EventLog(Object, EventTag1, Event1, Event2, Event3);
+	EventLog(Object, EventTag1, Event4);
+}
+
+inline
+void FVisualLogger::EventLog(const class UObject* Object, const FName EventTag1, const FVisualLogEventBase& Event1, const FVisualLogEventBase& Event2, const FVisualLogEventBase& Event3)
+{
+	EventLog(Object, EventTag1, Event1, Event2);
+	EventLog(Object, EventTag1, Event3);
+}
+
+inline
+void FVisualLogger::EventLog(const class UObject* Object, const FName EventTag1, const FVisualLogEventBase& Event1, const FVisualLogEventBase& Event2)
+{
+	EventLog(Object, EventTag1, Event1);
+	EventLog(Object, EventTag1, Event2);
+}
+
+inline
+void FVisualLogger::EventLog(const class UObject* LogOwner, const FVisualLogEventBase& Event1, const FName EventTag1, const FName EventTag2, const FName EventTag3, const FName EventTag4, const FName EventTag5, const FName EventTag6)
+{
+	EventLog(LogOwner, EventTag1, Event1, EventTag2, EventTag3, EventTag4, EventTag5, EventTag6);
+}
+
+inline
+void FVisualLogger::EventLog(const class UObject* Object, const FName EventTag1, const FVisualLogEventBase& Event, const FName EventTag2, const FName EventTag3, const FName EventTag4, const FName EventTag5, const FName EventTag6)
+{
+	SCOPE_CYCLE_COUNTER(STAT_VisualLog);
+	UWorld *World = nullptr;
+	FVisualLogEntry *CurrentEntry = nullptr;
+	const FLogCategory<ELogVerbosity::Log, ELogVerbosity::Log> Category(*Event.Name);
+	if (CheckVisualLogInputInternal(Object, Category, ELogVerbosity::Log, &World, &CurrentEntry) == false)
+	{
+		return;
+	}
+
+	int32 Index = CurrentEntry->Events.Find(FVisualLogEvent(Event));
+	if (Index != INDEX_NONE)
+	{
+		CurrentEntry->Events[Index].Counter++;
+	}
+	else
+	{
+		Index = CurrentEntry->AddEvent(Event);
+	}
+
+	CurrentEntry->Events[Index].EventTags.FindOrAdd(EventTag1)++;
+	CurrentEntry->Events[Index].EventTags.FindOrAdd(EventTag2)++;
+	CurrentEntry->Events[Index].EventTags.FindOrAdd(EventTag3)++;
+	CurrentEntry->Events[Index].EventTags.FindOrAdd(EventTag4)++;
+	CurrentEntry->Events[Index].EventTags.FindOrAdd(EventTag5)++;
+	CurrentEntry->Events[Index].EventTags.FindOrAdd(EventTag6)++;
+	CurrentEntry->Events[Index].EventTags.Remove(NAME_None);
+}
+
+
 FVisualLogger::FVisualLogger()
 {
 	BlockAllCategories(false);
@@ -38,17 +419,17 @@ FVisualLogger::FVisualLogger()
 
 UWorld* FVisualLogger::GetWorld(const class UObject* Object)
 {
-	UWorld* World = Object ? GEngine->GetWorldFromContextObject(Object, false) : NULL;
+	UWorld* World = Object ? GEngine->GetWorldFromContextObject(Object, false) : nullptr;
 #if WITH_EDITOR
 	UEditorEngine *EEngine = Cast<UEditorEngine>(GEngine);
-	if (GIsEditor && EEngine != NULL && World == NULL)
+	if (GIsEditor && EEngine != nullptr && World == nullptr)
 	{
 		// lets use PlayWorld during PIE/Simulate and regular world from editor otherwise, to draw debug information
-		World = EEngine->PlayWorld != NULL ? EEngine->PlayWorld : EEngine->GetEditorWorldContext().World();
+		World = EEngine->PlayWorld != nullptr ? EEngine->PlayWorld : EEngine->GetEditorWorldContext().World();
 	}
 
 #endif
-	if (!GIsEditor && World == NULL)
+	if (!GIsEditor && World == nullptr)
 	{
 		World = GEngine->GetWorld();
 	}
@@ -115,7 +496,7 @@ void FVisualLogger::Redirect(UObject* FromObject, UObject* ToObject)
 
 		RedirectionMap.FindOrAdd(NewRedirection).AddUnique(FromObject);
 
-		UE_CVLOG(FromObject != NULL, FromObject, LogVisual, Log, TEXT("Redirected '%s' to '%s'"), *FromObject->GetName(), *NewRedirection->GetName());
+		UE_CVLOG(FromObject != nullptr, FromObject, LogVisual, Log, TEXT("Redirected '%s' to '%s'"), *FromObject->GetName(), *NewRedirection->GetName());
 	}
 }
 
@@ -151,7 +532,7 @@ void FVisualLogger::SetIsRecordingToFile(bool InIsRecording)
 		SetIsRecording(true);
 	}
 
-	UWorld* World = GEngine ? GEngine->GetWorld() : NULL;
+	UWorld* World = GEngine ? GEngine->GetWorld() : nullptr;
 
 	const FString BaseFileName = LogFileNameGetter.IsBound() ? LogFileNameGetter.Execute() : TEXT("VisualLog");
 	const FString MapName = World ? World->GetMapName() : TEXT("");
diff --git a/Engine/Source/Runtime/Engine/Private/VisualLogger/VisualLoggerAutomationTests.cpp b/Engine/Source/Runtime/Engine/Private/VisualLogger/VisualLoggerAutomationTests.cpp
--- a/Engine/Source/Runtime/Engine/Private/VisualLogger/VisualLoggerAutomationTests.cpp
+++ b/Engine/Source/Runtime/Engine/Private/VisualLogger/VisualLoggerAutomationTests.cpp
@@ -13,15 +13,15 @@ namespace
 		// Accessing the game world is only valid for game-only 
 		if (((TestFlags & EAutomationTestFlags::ATF_Editor) || (TestFlags & EAutomationTestFlags::ATF_Game)) == false)
 		{
-			return NULL;
+			return nullptr;
 		}
 		if (GEngine->GetWorldContexts().Num() == 0)
 		{
-			return NULL;
+			return nullptr;
 		}
 		if (GEngine->GetWorldContexts()[0].WorldType != EWorldType::Game && GEngine->GetWorldContexts()[0].WorldType != EWorldType::Editor)
 		{
-			return NULL;
+			return nullptr;
 		}
 
 		return GEngine->GetWorldContexts()[0].World();
@@ -53,7 +53,7 @@ FVisualLoggerTestDevice::FVisualLoggerTestDevice()
 
 void FVisualLoggerTestDevice::Cleanup(bool bReleaseMemory)
 {
-	LastObject = NULL;
+	LastObject = nullptr;
 	LastEntry.Reset();
 }
 
@@ -116,7 +116,7 @@ IMPLEMENT_SIMPLE_AUTOMATION_TEST(FVisualLogTest, "System.Engine.VisualLogger.Log
 bool FVisualLogTest::RunTest(const FString& Parameters)
 {
 	UWorld* World = GetSimpleEngineAutomationTestWorld(GetTestFlags());
-	CHECK_SUCCESS(World != NULL);
+	CHECK_SUCCESS(World != nullptr);
 
 	FTestDeviceContext<FVisualLoggerTestDevice> Context;
 
@@ -124,7 +124,7 @@ bool FVisualLogTest::RunTest(const FString& Parameters)
 	CHECK_FAIL(FVisualLogger::Get().IsRecording());
 
 	UE_VLOG(World, LogVisual, Log, TEXT("Simple text line to test vlog"));
-	CHECK_SUCCESS(Context.Device.LastObject == NULL);
+	CHECK_SUCCESS(Context.Device.LastObject == nullptr);
 	CHECK_SUCCESS(Context.Device.LastEntry.TimeStamp == -1);
 
 	FVisualLogger::Get().SetIsRecording(true);
@@ -139,7 +139,7 @@ bool FVisualLogTest::RunTest(const FString& Parameters)
 		FVisualLogEntry* CurrentEntry = FVisualLogger::Get().GetEntryToWrite(World, World->TimeSeconds, ECreateIfNeeded::DontCreate);
 		
 		{
-			CHECK_SUCCESS(CurrentEntry != NULL);
+			CHECK_SUCCESS(CurrentEntry != nullptr);
 			CHECK_SUCCESS(CurrentEntry->TimeStamp == CurrentTimestamp);
 			CHECK_SUCCESS(CurrentEntry->LogLines.Num() == 1);
 			CHECK_SUCCESS(CurrentEntry->LogLines[0].Category == LogVisual.GetCategoryName());
@@ -148,13 +148,13 @@ bool FVisualLogTest::RunTest(const FString& Parameters)
 			const float NewTimestamp = CurrentTimestamp + 0.1;
 			FVisualLogEntry* NewEntry = FVisualLogger::Get().GetEntryToWrite(World, NewTimestamp); //generate new entry and serialize old one
 			CurrentEntry = &Context.Device.LastEntry;
-			CHECK_SUCCESS(CurrentEntry != NULL);
+			CHECK_SUCCESS(CurrentEntry != nullptr);
 			CHECK_SUCCESS(CurrentEntry->TimeStamp == CurrentTimestamp);
 			CHECK_SUCCESS(CurrentEntry->LogLines.Num() == 1);
 			CHECK_SUCCESS(CurrentEntry->LogLines[0].Category == LogVisual.GetCategoryName());
 			CHECK_SUCCESS(CurrentEntry->LogLines[0].Line == TextToLog);
 
-			CHECK_SUCCESS(NewEntry != NULL);
+			CHECK_SUCCESS(NewEntry != nullptr);
 			CHECK_SUCCESS(NewEntry->TimeStamp - NewTimestamp <= SMALL_NUMBER);
 			CHECK_SUCCESS(NewEntry->LogLines.Num() == 0);
 		}
@@ -170,7 +170,7 @@ IMPLEMENT_SIMPLE_AUTOMATION_TEST(FVisualLogSegmentsTest, "System.Engine.VisualLo
 bool FVisualLogSegmentsTest::RunTest(const FString& Parameters)
 {
 	UWorld* World = GetSimpleEngineAutomationTestWorld(GetTestFlags());
-	CHECK_SUCCESS(World != NULL);
+	CHECK_SUCCESS(World != nullptr);
 
 	FTestDeviceContext<FVisualLoggerTestDevice> Context;
 
@@ -178,7 +178,7 @@ bool FVisualLogSegmentsTest::RunTest(const FString& Parameters)
 	CHECK_FAIL(FVisualLogger::Get().IsRecording());
 
 	UE_VLOG_SEGMENT(World, LogVisual, Log, FVector(0, 0, 0), FVector(1, 0, 0), FColor::Red, TEXT("Simple segment log to test vlog"));
-	CHECK_SUCCESS(Context.Device.LastObject == NULL);
+	CHECK_SUCCESS(Context.Device.LastObject == nullptr);
 	CHECK_SUCCESS(Context.Device.LastEntry.TimeStamp == -1);
 
 	FVisualLogger::Get().SetIsRecording(true);
@@ -187,13 +187,13 @@ bool FVisualLogSegmentsTest::RunTest(const FString& Parameters)
 		const FVector StartPoint(0, 0, 0);
 		const FVector EndPoint(1, 0, 0);
 		UE_VLOG_SEGMENT(World, LogVisual, Log, StartPoint, EndPoint, FColor::Red, TEXT("Simple segment log to test vlog"));
-		CHECK_SUCCESS(Context.Device.LastObject == NULL);
+		CHECK_SUCCESS(Context.Device.LastObject == nullptr);
 		CHECK_SUCCESS(Context.Device.LastEntry.TimeStamp == -1);
 		FVisualLogEntry* CurrentEntry = FVisualLogger::Get().GetEntryToWrite(World, World->TimeSeconds, ECreateIfNeeded::DontCreate);
 
 		float CurrentTimestamp = World->TimeSeconds;
 		{
-			CHECK_SUCCESS(CurrentEntry != NULL);
+			CHECK_SUCCESS(CurrentEntry != nullptr);
 			CHECK_SUCCESS(CurrentEntry->TimeStamp == CurrentTimestamp);
 			CHECK_SUCCESS(CurrentEntry->ElementsToDraw.Num() == 1);
 			CHECK_SUCCESS(CurrentEntry->ElementsToDraw[0].GetType() == EVisualLoggerShapeElement::Segment);
@@ -204,7 +204,7 @@ bool FVisualLogSegmentsTest::RunTest(const FString& Parameters)
 			const float NewTimestamp = CurrentTimestamp + 0.1;
 			FVisualLogEntry* NewEntry = FVisualLogger::Get().GetEntryToWrite(World, NewTimestamp); //generate new entry and serialize old one
 			CurrentEntry = &Context.Device.LastEntry;
-			CHECK_SUCCESS(CurrentEntry != NULL);
+			CHECK_SUCCESS(CurrentEntry != nullptr);
 			CHECK_SUCCESS(CurrentEntry->TimeStamp == CurrentTimestamp);
 			CHECK_SUCCESS(CurrentEntry->ElementsToDraw.Num() == 1);
 			CHECK_SUCCESS(CurrentEntry->ElementsToDraw[0].GetType() == EVisualLoggerShapeElement::Segment);
@@ -212,7 +212,7 @@ bool FVisualLogSegmentsTest::RunTest(const FString& Parameters)
 			CHECK_SUCCESS(CurrentEntry->ElementsToDraw[0].Points[0] == StartPoint);
 			CHECK_SUCCESS(CurrentEntry->ElementsToDraw[0].Points[1] == EndPoint);
 
-			CHECK_SUCCESS(NewEntry != NULL);
+			CHECK_SUCCESS(NewEntry != nullptr);
 			CHECK_SUCCESS(NewEntry->TimeStamp - NewTimestamp <= SMALL_NUMBER);
 			CHECK_SUCCESS(NewEntry->ElementsToDraw.Num() == 0);
 		}
@@ -230,7 +230,7 @@ IMPLEMENT_SIMPLE_AUTOMATION_TEST(FVisualLogEventsTest, "System.Engine.VisualLogg
 bool FVisualLogEventsTest::RunTest(const FString& Parameters)
 {
 	UWorld* World = GetSimpleEngineAutomationTestWorld(GetTestFlags());
-	CHECK_SUCCESS(World != NULL);
+	CHECK_SUCCESS(World != nullptr);
 
 	FTestDeviceContext<FVisualLoggerTestDevice> Context;
 	FVisualLogger::Get().SetIsRecording(true);
@@ -247,13 +247,13 @@ bool FVisualLogEventsTest::RunTest(const FString& Parameters)
 	FVisualLogEntry* CurrentEntry = FVisualLogger::Get().GetEntryToWrite(World, World->TimeSeconds, ECreateIfNeeded::DontCreate);
 	float CurrentTimestamp = World->TimeSeconds;
 	UE_VLOG_EVENTS(World, NAME_None, EventTest);
-	CHECK_SUCCESS(CurrentEntry != NULL);
+	CHECK_SUCCESS(CurrentEntry != nullptr);
 	CHECK_SUCCESS(CurrentEntry->TimeStamp == CurrentTimestamp);
 	CHECK_SUCCESS(CurrentEntry->Events.Num() == 1);
 	CHECK_SUCCESS(CurrentEntry->Events[0].Name == TEXT("EventTest"));
 
 	UE_VLOG_EVENTS(World, NAME_None, EventTest, EventTest2);
-	CHECK_SUCCESS(CurrentEntry != NULL);
+	CHECK_SUCCESS(CurrentEntry != nullptr);
 	CHECK_SUCCESS(CurrentEntry->TimeStamp == CurrentTimestamp);
 	CHECK_SUCCESS(CurrentEntry->Events.Num() == 2);
 	CHECK_SUCCESS(CurrentEntry->Events[0].Counter == 2);
@@ -265,7 +265,7 @@ bool FVisualLogEventsTest::RunTest(const FString& Parameters)
 	UE_VLOG_EVENTS(World, NAME_None, EventTest, EventTest2, EventTest3);
 
 	{
-		CHECK_SUCCESS(CurrentEntry != NULL);
+		CHECK_SUCCESS(CurrentEntry != nullptr);
 		CHECK_SUCCESS(CurrentEntry->TimeStamp == CurrentTimestamp);
 		CHECK_SUCCESS(CurrentEntry->Events.Num() == 3);
 		CHECK_SUCCESS(CurrentEntry->Events[0].Counter == 3);
@@ -282,7 +282,7 @@ bool FVisualLogEventsTest::RunTest(const FString& Parameters)
 		const float NewTimestamp = CurrentTimestamp + 0.1;
 		FVisualLogEntry* NewEntry = FVisualLogger::Get().GetEntryToWrite(World, NewTimestamp); //generate new entry and serialize old one
 		CurrentEntry = &Context.Device.LastEntry;
-		CHECK_SUCCESS(CurrentEntry != NULL);
+		CHECK_SUCCESS(CurrentEntry != nullptr);
 		CHECK_SUCCESS(CurrentEntry->TimeStamp == CurrentTimestamp);
 		CHECK_SUCCESS(CurrentEntry->Events.Num() == 3);
 		CHECK_SUCCESS(CurrentEntry->Events[0].Counter == 3);
@@ -296,7 +296,7 @@ bool FVisualLogEventsTest::RunTest(const FString& Parameters)
 		CHECK_SUCCESS(CurrentEntry->Events[1].UserFriendlyDesc == TEXT("Second simple event for vlog tests"));
 		CHECK_SUCCESS(CurrentEntry->Events[2].UserFriendlyDesc == TEXT("Third simple event for vlog tests"));
 
-		CHECK_SUCCESS(NewEntry != NULL);
+		CHECK_SUCCESS(NewEntry != nullptr);
 		CHECK_SUCCESS(NewEntry->TimeStamp - NewTimestamp <= SMALL_NUMBER);
 		CHECK_SUCCESS(NewEntry->Events.Num() == 0);
 	}
@@ -308,7 +308,7 @@ bool FVisualLogEventsTest::RunTest(const FString& Parameters)
 	CurrentTimestamp = World->TimeSeconds + 0.2;
 	CurrentEntry = FVisualLogger::Get().GetEntryToWrite(World, CurrentTimestamp); //generate new entry and serialize old one
 	UE_VLOG_EVENT_WITH_DATA(World, EventTest, EventTag1);
-	CHECK_SUCCESS(CurrentEntry != NULL);
+	CHECK_SUCCESS(CurrentEntry != nullptr);
 	CHECK_SUCCESS(CurrentEntry->TimeStamp == CurrentTimestamp);
 	CHECK_SUCCESS(CurrentEntry->Events.Num() == 1);
 	CHECK_SUCCESS(CurrentEntry->Events[0].Name == TEXT("EventTest"));
@@ -319,7 +319,7 @@ bool FVisualLogEventsTest::RunTest(const FString& Parameters)
 	CurrentEntry = FVisualLogger::Get().GetEntryToWrite(World, CurrentTimestamp); //generate new entry and serialize old one
 	UE_VLOG_EVENT_WITH_DATA(World, EventTest, EventTag1, EventTag2, EventTag3);
 	UE_VLOG_EVENT_WITH_DATA(World, EventTest, EventTag3);
-	CHECK_SUCCESS(CurrentEntry != NULL);
+	CHECK_SUCCESS(CurrentEntry != nullptr);
 	CHECK_SUCCESS(CurrentEntry->TimeStamp == CurrentTimestamp);
 	CHECK_SUCCESS(CurrentEntry->Events.Num() == 1);
 	CHECK_SUCCESS(CurrentEntry->Events[0].Name == TEXT("EventTest"));
diff --git a/Engine/Source/Runtime/Engine/Private/VisualLogger/VisualLoggerBinaryFileDevice.cpp b/Engine/Source/Runtime/Engine/Private/VisualLogger/VisualLoggerBinaryFileDevice.cpp
--- a/Engine/Source/Runtime/Engine/Private/VisualLogger/VisualLoggerBinaryFileDevice.cpp
+++ b/Engine/Source/Runtime/Engine/Private/VisualLogger/VisualLoggerBinaryFileDevice.cpp
@@ -5,7 +5,7 @@
 #if ENABLE_VISUAL_LOG
 
 FVisualLoggerBinaryFileDevice::FVisualLoggerBinaryFileDevice()
-	: FileArchive(NULL)
+	: FileArchive(nullptr)
 {
 	Cleanup();
 
@@ -25,7 +25,7 @@ void FVisualLoggerBinaryFileDevice::Cleanup(bool bReleaseMemory)
 
 void FVisualLoggerBinaryFileDevice::StartRecordingToFile(float TimeStamp)
 {
-	if (FileArchive != NULL)
+	if (FileArchive != nullptr)
 	{
 		return;
 	}
@@ -40,7 +40,7 @@ void FVisualLoggerBinaryFileDevice::StartRecordingToFile(float TimeStamp)
 
 void FVisualLoggerBinaryFileDevice::StopRecordingToFile(float TimeStamp)
 {
-	if (FileArchive == NULL)
+	if (FileArchive == nullptr)
 	{
 		return;
 	}
@@ -55,7 +55,7 @@ void FVisualLoggerBinaryFileDevice::StopRecordingToFile(float TimeStamp)
 	const int64 TotalSize = FileArchive->TotalSize();
 	FileArchive->Close();
 	delete FileArchive;
-	FileArchive = NULL;
+	FileArchive = nullptr;
 
 	const FString TempFullFilename = FString::Printf(TEXT("%slogs/%s"), *FPaths::GameSavedDir(), *TempFileName);
 	const FString NewFileName = FVisualLoggerHelpers::GenerateFilename(TempFileName, FileName, StartRecordingTime, LastLogTimeStamp);
diff --git a/Engine/Source/Runtime/Engine/Private/VisualLogger/VisualLoggerTypes.cpp b/Engine/Source/Runtime/Engine/Private/VisualLogger/VisualLoggerTypes.cpp
--- a/Engine/Source/Runtime/Engine/Private/VisualLogger/VisualLoggerTypes.cpp
+++ b/Engine/Source/Runtime/Engine/Private/VisualLogger/VisualLoggerTypes.cpp
@@ -45,7 +45,7 @@ FVisualLogEntry::FVisualLogEntry(const class AActor* InActor, TArray<TWeakObject
 		TimeStamp = InActor->GetWorld()->TimeSeconds;
 		Location = InActor->GetActorLocation();
 		InActor->GrabDebugSnapshot(this);
-		if (Children != NULL)
+		if (Children != nullptr)
 		{
 			TWeakObjectPtr<UObject>* WeakActorPtr = Children->GetData();
 			for (int32 Index = 0; Index < Children->Num(); ++Index, ++WeakActorPtr)
@@ -72,7 +72,7 @@ FVisualLogEntry::FVisualLogEntry(float InTimeStamp, FVector InLocation, const UO
 	{
 		AsActor->GrabDebugSnapshot(this);
 	}
-	if (Children != NULL)
+	if (Children != nullptr)
 	{
 		TWeakObjectPtr<UObject>* WeakActorPtr = Children->GetData();
 		for (int32 Index = 0; Index < Children->Num(); ++Index, ++WeakActorPtr)
diff --git a/Engine/Source/Runtime/Engine/Public/VisualLogger/VisualLogger.h b/Engine/Source/Runtime/Engine/Public/VisualLogger/VisualLogger.h
--- a/Engine/Source/Runtime/Engine/Public/VisualLogger/VisualLogger.h
+++ b/Engine/Source/Runtime/Engine/Public/VisualLogger/VisualLogger.h
@@ -136,12 +136,8 @@ public:
 	static void EventLog(const class UObject* LogOwner, const FName EventTag1, const FVisualLogEventBase& Event1, const FVisualLogEventBase& Event2, const FVisualLogEventBase& Event3, const FVisualLogEventBase& Event4);
 	static void EventLog(const class UObject* LogOwner, const FName EventTag1, const FVisualLogEventBase& Event1, const FVisualLogEventBase& Event2, const FVisualLogEventBase& Event3, const FVisualLogEventBase& Event4, const FVisualLogEventBase& Event5);
 	static void EventLog(const class UObject* LogOwner, const FName EventTag1, const FVisualLogEventBase& Event1, const FVisualLogEventBase& Event2, const FVisualLogEventBase& Event3, const FVisualLogEventBase& Event4, const FVisualLogEventBase& Event5, const FVisualLogEventBase& Event6);
-
-	FORCEINLINE
-	static void EventLog(const class UObject* LogOwner, const FVisualLogEventBase& Event1, const FName EventTag1 = NAME_None, const FName EventTag2 = NAME_None, const FName EventTag3 = NAME_None, const FName EventTag4 = NAME_None, const FName EventTag5 = NAME_None, const FName EventTag6 = NAME_None)
-	{
-		EventLog(LogOwner, EventTag1, Event1, EventTag2, EventTag3, EventTag4, EventTag5, EventTag6);
-	}
+	
+	static void EventLog(const class UObject* LogOwner, const FVisualLogEventBase& Event1, const FName EventTag1 = NAME_None, const FName EventTag2 = NAME_None, const FName EventTag3 = NAME_None, const FName EventTag4 = NAME_None, const FName EventTag5 = NAME_None, const FName EventTag6 = NAME_None);
 
 	// static getter
 	static FVisualLogger& Get()
@@ -165,19 +161,19 @@ public:
 	static class UObject* FindRedirection(const class UObject* Object);
 
 	/** blocks all categories from logging. It can be bypassed with white list for categories */
-	FORCEINLINE void BlockAllCategories(bool InBlock) { bBlockedAllCategories = InBlock; }
+	void BlockAllCategories(bool InBlock) { bBlockedAllCategories = InBlock; }
 
 	/** checks if all categories are blocked */
-	FORCEINLINE bool IsBlockedForAllCategories() const { return !!bBlockedAllCategories; }
+	bool IsBlockedForAllCategories() const { return !!bBlockedAllCategories; }
 
 	/** Returns white list for modifications */
-	FORCEINLINE TArray<class FName>& GetWhiteList() { return CategoriesWhiteList; }
+	TArray<class FName>& GetWhiteList() { return CategoriesWhiteList; }
 
-	FORCEINLINE bool IsWhiteListed(const FName& Name) const { return CategoriesWhiteList.Find(Name) != INDEX_NONE; }
+	bool IsWhiteListed(const FName& Name) const { return CategoriesWhiteList.Find(Name) != INDEX_NONE; }
 
-	FORCEINLINE void AddCategortyToWhiteList(class FName Category) { CategoriesWhiteList.AddUnique(Category); }
+	void AddCategortyToWhiteList(class FName Category) { CategoriesWhiteList.AddUnique(Category); }
 
-	FORCEINLINE void ClearWhiteList() { CategoriesWhiteList.Reset();  }
+	void ClearWhiteList() { CategoriesWhiteList.Reset(); }
 
 	/** Generates and returns Id unique for given timestamp - used to connect different logs between (ex. text log with geometry shape) */
 	int32 GetUniqueId(float Timestamp);
@@ -214,10 +210,13 @@ public:
 	/**  Removes previously registered extension */
 	void UnregisterExtension(FName TagName, FVisualLogExtensionInterface* ExtensionInterface) { AllExtensions.Remove(TagName); }
 	/** returns extension identified by given tag */
-	FVisualLogExtensionInterface* GetExtensionForTag(const FName TagName) { return AllExtensions.Contains(TagName) ? AllExtensions[TagName] : NULL; }
+	FVisualLogExtensionInterface* GetExtensionForTag(const FName TagName) { return AllExtensions.Contains(TagName) ? AllExtensions[TagName] : nullptr; }
 	/** Returns reference to map with all registered extension */
 	TMap<FName, FVisualLogExtensionInterface*>& GetAllExtensions() { return AllExtensions; }
 
+	/** internal check for each usage of visual logger */
+	static bool CheckVisualLogInputInternal(const class UObject* Object, const struct FLogCategoryBase& Category, ELogVerbosity::Type Verbosity, UWorld **World, FVisualLogEntry **CurrentEntry);
+
 private:
 	FVisualLogger();
 	virtual void Serialize(const TCHAR* V, ELogVerbosity::Type Verbosity, const class FName& Category) override { ensureMsgf(0, TEXT("Regular serialize is forbiden for visual logs")); }
@@ -260,8 +259,4 @@ protected:
 	FTimerHandle VisualLoggerCleanupTimerHandle;
 };
 
-
-#include "VisualLogger.inl"
-
-
 #endif //ENABLE_VISUAL_LOG
diff --git a/Engine/Source/Runtime/Engine/Public/VisualLogger/VisualLogger.inl b/Engine/Source/Runtime/Engine/Public/VisualLogger/VisualLogger.inl
deleted file mode 100644
diff --git a/Engine/Source/Runtime/Engine/Public/VisualLogger/VisualLoggerTypes.h b/Engine/Source/Runtime/Engine/Public/VisualLogger/VisualLoggerTypes.h
--- a/Engine/Source/Runtime/Engine/Public/VisualLogger/VisualLoggerTypes.h
+++ b/Engine/Source/Runtime/Engine/Public/VisualLogger/VisualLoggerTypes.h
@@ -268,13 +268,13 @@ ENGINE_API  FArchive& operator<<(FArchive& Ar, FVisualLogLine& LogLine);
 ENGINE_API  FArchive& operator<<(FArchive& Ar, FVisualLogStatusCategory& Status);
 ENGINE_API  FArchive& operator<<(FArchive& Ar, FVisualLogEntry& LogEntry);
 
-FORCEINLINE
+inline
 bool operator==(const FVisualLogEvent& Left, const FVisualLogEvent& Right) 
 { 
 	return Left.Name == Right.Name; 
 }
 
-FORCEINLINE
+inline
 FVisualLogEvent::FVisualLogEvent(const FVisualLogEventBase& Event)
 : Counter(1)
 {
@@ -283,7 +283,7 @@ FVisualLogEvent::FVisualLogEvent(const FVisualLogEventBase& Event)
 	Verbosity = Event.Verbosity;
 }
 
-FORCEINLINE
+inline
 FVisualLogEvent& FVisualLogEvent::operator= (const FVisualLogEventBase& Event)
 {
 	Name = Event.Name;
@@ -292,7 +292,7 @@ FVisualLogEvent& FVisualLogEvent::operator= (const FVisualLogEventBase& Event)
 	return *this;
 }
 
-FORCEINLINE
+inline
 FVisualLogLine::FVisualLogLine(const FName& InCategory, ELogVerbosity::Type InVerbosity, const FString& InLine)
 : Line(InLine)
 , Category(InCategory)
@@ -302,7 +302,7 @@ FVisualLogLine::FVisualLogLine(const FName& InCategory, ELogVerbosity::Type InVe
 
 }
 
-FORCEINLINE
+inline
 FVisualLogLine::FVisualLogLine(const FName& InCategory, ELogVerbosity::Type InVerbosity, const FString& InLine, int64 InUserData)
 : Line(InLine)
 , Category(InCategory)
@@ -312,19 +312,19 @@ FVisualLogLine::FVisualLogLine(const FName& InCategory, ELogVerbosity::Type InVe
 
 }
 
-FORCEINLINE
+inline
 void FVisualLogStatusCategory::Add(const FString& Key, const FString& Value)
 {
 	Data.Add(FString(Key).AppendChar(TEXT('|')) + Value);
 }
 
-FORCEINLINE
+inline
 void FVisualLogStatusCategory::AddChild(const FVisualLogStatusCategory& Child)
 {
 	Children.Add(Child);
 }
 
-FORCEINLINE
+inline
 FVisualLogShapeElement::FVisualLogShapeElement(EVisualLoggerShapeElement InType)
 : Verbosity(ELogVerbosity::All)
 , TransformationMatrix(FMatrix::Identity)
@@ -335,7 +335,7 @@ FVisualLogShapeElement::FVisualLogShapeElement(EVisualLoggerShapeElement InType)
 
 }
 
-FORCEINLINE
+inline
 FVisualLogShapeElement::FVisualLogShapeElement(const FString& InDescription, const FColor& InColor, uint16 InThickness, const FName& InCategory)
 : Category(InCategory)
 , Verbosity(ELogVerbosity::All)
@@ -350,32 +350,32 @@ FVisualLogShapeElement::FVisualLogShapeElement(const FString& InDescription, con
 	SetColor(InColor);
 }
 
-FORCEINLINE
+inline
 void FVisualLogShapeElement::SetColor(const FColor& InColor)
 {
 	Color = ((InColor.DWColor() >> 30) << 6)	| (((InColor.DWColor() & 0x00ff0000) >> 22) << 4)	| (((InColor.DWColor() & 0x0000ff00) >> 14) << 2)	| ((InColor.DWColor() & 0x000000ff) >> 6);
 }
 
-FORCEINLINE
+inline
 EVisualLoggerShapeElement FVisualLogShapeElement::GetType() const
 {
 	return Type;
 }
 
-FORCEINLINE
+inline
 void FVisualLogShapeElement::SetType(EVisualLoggerShapeElement InType)
 {
 	Type = InType;
 }
 
-FORCEINLINE
+inline
 FColor FVisualLogShapeElement::GetFColor() const
 {
 	return FColor(((Color & 0xc0) << 24) | ((Color & 0x30) << 18) | ((Color & 0x0c) << 12) | ((Color & 0x03) << 6));
 }
 
 
-FORCEINLINE
+inline
 int32 FVisualLogEntry::FindStatusIndex(const FString& CategoryName)
 {
 	for (int32 TestCategoryIndex = 0; TestCategoryIndex < Status.Num(); TestCategoryIndex++)
