Fixed Pause in PIE in a new editor window.

diff --git a/Engine/Source/Runtime/Engine/Private/LocalPlayer.cpp b/Engine/Source/Runtime/Engine/Private/LocalPlayer.cpp
--- a/Engine/Source/Runtime/Engine/Private/LocalPlayer.cpp
+++ b/Engine/Source/Runtime/Engine/Private/LocalPlayer.cpp
@@ -1053,7 +1053,7 @@ bool ULocalPlayer::HandlePauseCommand( const TCHAR* Cmd, FOutputDevice& Ar, UWor
 	}
 	
 
-	return !ViewportClient->bIsPlayInEditorViewport;
+	return true;
 }
 
 bool ULocalPlayer::HandleListMoveBodyCommand( const TCHAR* Cmd, FOutputDevice& Ar )
