Fixed monolitics compilation issues related to changes in Visual Logger.

diff --git a/Engine/Source/Runtime/Engine/Private/VisualLogger/VisualLogger.cpp b/Engine/Source/Runtime/Engine/Private/VisualLogger/VisualLogger.cpp
--- a/Engine/Source/Runtime/Engine/Private/VisualLogger/VisualLogger.cpp
+++ b/Engine/Source/Runtime/Engine/Private/VisualLogger/VisualLogger.cpp
@@ -2,7 +2,6 @@
 
 #include "EnginePrivate.h"
 #include "VisualLogger/VisualLogger.h"
-#include "VisualLogger/VisualLogger.h"
 #include "VisualLogger/VisualLoggerBinaryFileDevice.h"
 #if WITH_EDITOR
 #	include "Editor/UnrealEd/Public/EditorComponents.h"
@@ -22,36 +21,6 @@ DEFINE_LOG_CATEGORY(LogVisual);
 
 TMap<UObject*, TArray<TWeakObjectPtr<const UObject> > > FVisualLogger::RedirectionMap;
 
-// Unfortunately needs to be a #define since it uses GET_VARARGS_RESULT which uses the va_list stuff which operates on the
-// current function, so we can't easily call a function
-#define COLLAPSED_LOGF(SerializeFunc) \
-	int32	BufferSize	= 1024; \
-	TCHAR*	Buffer		= nullptr; \
-	int32	Result		= -1; \
-	/* allocate some stack space to use on the first pass, which matches most strings */ \
-	TCHAR	StackBuffer[512]; \
-	TCHAR*	AllocatedBuffer = nullptr; \
-	\
-	/* first, try using the stack buffer */ \
-	Buffer = StackBuffer; \
-	GET_VARARGS_RESULT( Buffer, ARRAY_COUNT(StackBuffer), ARRAY_COUNT(StackBuffer) - 1, Fmt, Fmt, Result ); \
-	\
-	/* if that fails, then use heap allocation to make enough space */ \
-		while(Result == -1) \
-			{ \
-		FMemory::SystemFree(AllocatedBuffer); \
-		/* We need to use malloc here directly as GMalloc might not be safe. */ \
-		Buffer = AllocatedBuffer = (TCHAR*) FMemory::SystemMalloc( BufferSize * sizeof(TCHAR) ); \
-		GET_VARARGS_RESULT( Buffer, BufferSize, BufferSize-1, Fmt, Fmt, Result ); \
-		BufferSize *= 2; \
-			}; \
-	Buffer[Result] = 0; \
-	; \
-	\
-	SerializeFunc; \
-	FMemory::SystemFree(AllocatedBuffer);
-
-inline
 bool FVisualLogger::CheckVisualLogInputInternal(const class UObject* Object, const struct FLogCategoryBase& Category, ELogVerbosity::Type Verbosity, UWorld **World, FVisualLogEntry **CurrentEntry)
 {
 	FVisualLogger& VisualLogger = FVisualLogger::Get();
@@ -80,7 +49,7 @@ bool FVisualLogger::CheckVisualLogInputInternal(const class UObject* Object, con
 	return true;
 }
 
-inline
+
 FVisualLogEntry* FVisualLogger::GetEntryToWrite(const class UObject* Object, float TimeStamp, ECreateIfNeeded ShouldCreate)
 {
 	FVisualLogEntry* CurrentEntry = nullptr;
@@ -183,7 +152,7 @@ FVisualLogEntry* FVisualLogger::GetEntryToWrite(const class UObject* Object, flo
 	return CurrentEntry;
 }
 
-inline
+
 void FVisualLogger::Flush()
 {
 	for (auto &CurrentEntry : CurrentEntryPerObject)
@@ -200,178 +169,47 @@ void FVisualLogger::Flush()
 }
 
 
-inline
-VARARG_BODY(void, FVisualLogger::CategorizedLogf, const TCHAR*, VARARG_EXTRA(const class UObject* Object) VARARG_EXTRA(const struct FLogCategoryBase& Category) VARARG_EXTRA(ELogVerbosity::Type Verbosity) VARARG_EXTRA(int32 UniqueLogId))
-{
-	SCOPE_CYCLE_COUNTER(STAT_VisualLog);
-	UWorld *World = nullptr;
-	FVisualLogEntry *CurrentEntry = nullptr;
-	if (CheckVisualLogInputInternal(Object, Category, Verbosity, &World, &CurrentEntry) == false)
-	{
-		return;
-	}
-
-	COLLAPSED_LOGF(
-		CurrentEntry->AddText(Buffer, Category.GetCategoryName(), Verbosity);
-	);
-}
-
-inline
-VARARG_BODY(void, FVisualLogger::GeometryShapeLogf, const TCHAR*, VARARG_EXTRA(const class UObject* Object) VARARG_EXTRA(const struct FLogCategoryBase& Category) VARARG_EXTRA(ELogVerbosity::Type Verbosity) VARARG_EXTRA(int32 UniqueLogId) VARARG_EXTRA(const FVector& Start) VARARG_EXTRA(const FVector& End) VARARG_EXTRA(const FColor& Color))
-{
-	SCOPE_CYCLE_COUNTER(STAT_VisualLog);
-	UWorld *World = nullptr;
-	FVisualLogEntry *CurrentEntry = nullptr;
-	if (CheckVisualLogInputInternal(Object, Category, Verbosity, &World, &CurrentEntry) == false)
-	{
-		return;
-	}
-
-	COLLAPSED_LOGF(
-		CurrentEntry->AddElement(Start, End, Category.GetCategoryName(), Verbosity, Color, Buffer);
-	);
-}
-
-inline
-VARARG_BODY(void, FVisualLogger::GeometryShapeLogf, const TCHAR*, VARARG_EXTRA(const class UObject* Object) VARARG_EXTRA(const struct FLogCategoryBase& Category) VARARG_EXTRA(ELogVerbosity::Type Verbosity) VARARG_EXTRA(int32 UniqueLogId) VARARG_EXTRA(const FVector& Location) VARARG_EXTRA(float Radius) VARARG_EXTRA(const FColor& Color))
-{
-	SCOPE_CYCLE_COUNTER(STAT_VisualLog);
-	UWorld *World = nullptr;
-	FVisualLogEntry *CurrentEntry = nullptr;
-	if (CheckVisualLogInputInternal(Object, Category, Verbosity, &World, &CurrentEntry) == false)
-	{
-		return;
-	}
-
-	COLLAPSED_LOGF(
-		CurrentEntry->AddElement(Location, Category.GetCategoryName(), Verbosity, Color, Buffer, Radius);
-	);
-}
-
-inline
-VARARG_BODY(void, FVisualLogger::GeometryShapeLogf, const TCHAR*, VARARG_EXTRA(const class UObject* Object) VARARG_EXTRA(const struct FLogCategoryBase& Category) VARARG_EXTRA(ELogVerbosity::Type Verbosity) VARARG_EXTRA(int32 UniqueLogId) VARARG_EXTRA(const FBox& Box) VARARG_EXTRA(const FMatrix& Matrix) VARARG_EXTRA(const FColor& Color))
-{
-	SCOPE_CYCLE_COUNTER(STAT_VisualLog);
-	UWorld *World = nullptr;
-	FVisualLogEntry *CurrentEntry = nullptr;
-	if (CheckVisualLogInputInternal(Object, Category, Verbosity, &World, &CurrentEntry) == false)
-	{
-		return;
-	}
-
-	COLLAPSED_LOGF(
-		CurrentEntry->AddElement(Box, Matrix, Category.GetCategoryName(), Verbosity, Color, Buffer);
-	);
-}
-
-inline
-VARARG_BODY(void, FVisualLogger::GeometryShapeLogf, const TCHAR*, VARARG_EXTRA(const class UObject* Object) VARARG_EXTRA(const struct FLogCategoryBase& Category) VARARG_EXTRA(ELogVerbosity::Type Verbosity) VARARG_EXTRA(int32 UniqueLogId) VARARG_EXTRA(const FVector& Orgin) VARARG_EXTRA(const FVector& Direction) VARARG_EXTRA(const float Length) VARARG_EXTRA(const float Angle)  VARARG_EXTRA(const FColor& Color))
-{
-	SCOPE_CYCLE_COUNTER(STAT_VisualLog);
-	UWorld *World = nullptr;
-	FVisualLogEntry *CurrentEntry = nullptr;
-	if (CheckVisualLogInputInternal(Object, Category, Verbosity, &World, &CurrentEntry) == false)
-	{
-		return;
-	}
-
-	COLLAPSED_LOGF(
-		CurrentEntry->AddElement(Orgin, Direction, Length, Angle, Angle, Category.GetCategoryName(), Verbosity, Color, Buffer);
-	);
-}
-
-inline
-VARARG_BODY(void, FVisualLogger::GeometryShapeLogf, const TCHAR*, VARARG_EXTRA(const class UObject* Object) VARARG_EXTRA(const struct FLogCategoryBase& Category) VARARG_EXTRA(ELogVerbosity::Type Verbosity) VARARG_EXTRA(int32 UniqueLogId) VARARG_EXTRA(const FVector& Start) VARARG_EXTRA(const FVector& End) VARARG_EXTRA(const float Radius) VARARG_EXTRA(const FColor& Color))
-{
-	SCOPE_CYCLE_COUNTER(STAT_VisualLog);
-	UWorld *World = nullptr;
-	FVisualLogEntry *CurrentEntry = nullptr;
-	if (CheckVisualLogInputInternal(Object, Category, Verbosity, &World, &CurrentEntry) == false)
-	{
-		return;
-	}
-
-	COLLAPSED_LOGF(
-		CurrentEntry->AddElement(Start, End, Radius, Category.GetCategoryName(), Verbosity, Color, Buffer);
-	);
-}
-
-inline
-VARARG_BODY(void, FVisualLogger::GeometryShapeLogf, const TCHAR*, VARARG_EXTRA(const class UObject* Object) VARARG_EXTRA(const struct FLogCategoryBase& Category) VARARG_EXTRA(ELogVerbosity::Type Verbosity) VARARG_EXTRA(int32 UniqueLogId) VARARG_EXTRA(const FVector& Center) VARARG_EXTRA(float HalfHeight) VARARG_EXTRA(float Radius) VARARG_EXTRA(const FQuat& Rotation) VARARG_EXTRA(const FColor& Color))
-{
-	SCOPE_CYCLE_COUNTER(STAT_VisualLog);
-	UWorld *World = nullptr;
-	FVisualLogEntry *CurrentEntry = nullptr;
-	if (CheckVisualLogInputInternal(Object, Category, Verbosity, &World, &CurrentEntry) == false)
-	{
-		return;
-	}
-
-	COLLAPSED_LOGF(
-		CurrentEntry->AddElement(Center, HalfHeight, Radius, Rotation, Category.GetCategoryName(), Verbosity, Color, Buffer);
-	);
-}
-
-inline
-VARARG_BODY(void, FVisualLogger::HistogramDataLogf, const TCHAR*, VARARG_EXTRA(const class UObject* Object) VARARG_EXTRA(const struct FLogCategoryBase& Category) VARARG_EXTRA(ELogVerbosity::Type Verbosity) VARARG_EXTRA(int32 UniqueLogId) VARARG_EXTRA(FName GraphName) VARARG_EXTRA(FName DataName) VARARG_EXTRA(const FVector2D& Data) VARARG_EXTRA(const FColor& Color))
-{
-	SCOPE_CYCLE_COUNTER(STAT_VisualLog);
-	UWorld *World = nullptr;
-	FVisualLogEntry *CurrentEntry = nullptr;
-	if (CheckVisualLogInputInternal(Object, Category, Verbosity, &World, &CurrentEntry) == false)
-	{
-		return;
-	}
-
-	COLLAPSED_LOGF(
-		CurrentEntry->AddHistogramData(Data, Category.GetCategoryName(), Verbosity, GraphName, DataName);
-	);
-}
-
-#undef COLLAPSED_LOGF
-
-inline
 void FVisualLogger::EventLog(const class UObject* Object, const FName EventTag1, const FVisualLogEventBase& Event1, const FVisualLogEventBase& Event2, const FVisualLogEventBase& Event3, const FVisualLogEventBase& Event4, const FVisualLogEventBase& Event5, const FVisualLogEventBase& Event6)
 {
 	EventLog(Object, EventTag1, Event1, Event2, Event3, Event4, Event5);
 	EventLog(Object, EventTag1, Event6);
 }
 
-inline
+
 void FVisualLogger::EventLog(const class UObject* Object, const FName EventTag1, const FVisualLogEventBase& Event1, const FVisualLogEventBase& Event2, const FVisualLogEventBase& Event3, const FVisualLogEventBase& Event4, const FVisualLogEventBase& Event5)
 {
 	EventLog(Object, EventTag1, Event1, Event2, Event3, Event4);
 	EventLog(Object, EventTag1, Event5);
 }
 
-inline
+
 void FVisualLogger::EventLog(const class UObject* Object, const FName EventTag1, const FVisualLogEventBase& Event1, const FVisualLogEventBase& Event2, const FVisualLogEventBase& Event3, const FVisualLogEventBase& Event4)
 {
 	EventLog(Object, EventTag1, Event1, Event2, Event3);
 	EventLog(Object, EventTag1, Event4);
 }
 
-inline
+
 void FVisualLogger::EventLog(const class UObject* Object, const FName EventTag1, const FVisualLogEventBase& Event1, const FVisualLogEventBase& Event2, const FVisualLogEventBase& Event3)
 {
 	EventLog(Object, EventTag1, Event1, Event2);
 	EventLog(Object, EventTag1, Event3);
 }
 
-inline
+
 void FVisualLogger::EventLog(const class UObject* Object, const FName EventTag1, const FVisualLogEventBase& Event1, const FVisualLogEventBase& Event2)
 {
 	EventLog(Object, EventTag1, Event1);
 	EventLog(Object, EventTag1, Event2);
 }
 
-inline
+
 void FVisualLogger::EventLog(const class UObject* LogOwner, const FVisualLogEventBase& Event1, const FName EventTag1, const FName EventTag2, const FName EventTag3, const FName EventTag4, const FName EventTag5, const FName EventTag6)
 {
 	EventLog(LogOwner, EventTag1, Event1, EventTag2, EventTag3, EventTag4, EventTag5, EventTag6);
 }
 
-inline
+
 void FVisualLogger::EventLog(const class UObject* Object, const FName EventTag1, const FVisualLogEventBase& Event, const FName EventTag2, const FName EventTag3, const FName EventTag4, const FName EventTag5, const FName EventTag6)
 {
 	SCOPE_CYCLE_COUNTER(STAT_VisualLog);
diff --git a/Engine/Source/Runtime/Engine/Public/VisualLogger/VisualLogger.h b/Engine/Source/Runtime/Engine/Public/VisualLogger/VisualLogger.h
--- a/Engine/Source/Runtime/Engine/Public/VisualLogger/VisualLogger.h
+++ b/Engine/Source/Runtime/Engine/Public/VisualLogger/VisualLogger.h
@@ -259,4 +259,163 @@ protected:
 	FTimerHandle VisualLoggerCleanupTimerHandle;
 };
 
+// Unfortunately needs to be a #define since it uses GET_VARARGS_RESULT which uses the va_list stuff which operates on the
+// current function, so we can't easily call a function
+#define COLLAPSED_LOGF(SerializeFunc) \
+	int32	BufferSize	= 1024; \
+	TCHAR*	Buffer		= nullptr; \
+	int32	Result		= -1; \
+	/* allocate some stack space to use on the first pass, which matches most strings */ \
+	TCHAR	StackBuffer[512]; \
+	TCHAR*	AllocatedBuffer = nullptr; \
+	\
+	/* first, try using the stack buffer */ \
+	Buffer = StackBuffer; \
+	GET_VARARGS_RESULT( Buffer, ARRAY_COUNT(StackBuffer), ARRAY_COUNT(StackBuffer) - 1, Fmt, Fmt, Result ); \
+	\
+	/* if that fails, then use heap allocation to make enough space */ \
+			while(Result == -1) \
+						{ \
+		FMemory::SystemFree(AllocatedBuffer); \
+		/* We need to use malloc here directly as GMalloc might not be safe. */ \
+		Buffer = AllocatedBuffer = (TCHAR*) FMemory::SystemMalloc( BufferSize * sizeof(TCHAR) ); \
+		GET_VARARGS_RESULT( Buffer, BufferSize, BufferSize-1, Fmt, Fmt, Result ); \
+		BufferSize *= 2; \
+						}; \
+	Buffer[Result] = 0; \
+	; \
+	\
+	SerializeFunc; \
+	FMemory::SystemFree(AllocatedBuffer);
+
+inline
+VARARG_BODY(void, FVisualLogger::CategorizedLogf, const TCHAR*, VARARG_EXTRA(const class UObject* Object) VARARG_EXTRA(const struct FLogCategoryBase& Category) VARARG_EXTRA(ELogVerbosity::Type Verbosity) VARARG_EXTRA(int32 UniqueLogId))
+{
+	SCOPE_CYCLE_COUNTER(STAT_VisualLog);
+	UWorld *World = nullptr;
+	FVisualLogEntry *CurrentEntry = nullptr;
+	if (CheckVisualLogInputInternal(Object, Category, Verbosity, &World, &CurrentEntry) == false)
+	{
+		return;
+	}
+
+	COLLAPSED_LOGF(
+		CurrentEntry->AddText(Buffer, Category.GetCategoryName(), Verbosity);
+	);
+}
+
+inline
+VARARG_BODY(void, FVisualLogger::GeometryShapeLogf, const TCHAR*, VARARG_EXTRA(const class UObject* Object) VARARG_EXTRA(const struct FLogCategoryBase& Category) VARARG_EXTRA(ELogVerbosity::Type Verbosity) VARARG_EXTRA(int32 UniqueLogId) VARARG_EXTRA(const FVector& Start) VARARG_EXTRA(const FVector& End) VARARG_EXTRA(const FColor& Color))
+{
+	SCOPE_CYCLE_COUNTER(STAT_VisualLog);
+	UWorld *World = nullptr;
+	FVisualLogEntry *CurrentEntry = nullptr;
+	if (CheckVisualLogInputInternal(Object, Category, Verbosity, &World, &CurrentEntry) == false)
+	{
+		return;
+	}
+
+	COLLAPSED_LOGF(
+		CurrentEntry->AddElement(Start, End, Category.GetCategoryName(), Verbosity, Color, Buffer);
+	);
+}
+
+inline
+VARARG_BODY(void, FVisualLogger::GeometryShapeLogf, const TCHAR*, VARARG_EXTRA(const class UObject* Object) VARARG_EXTRA(const struct FLogCategoryBase& Category) VARARG_EXTRA(ELogVerbosity::Type Verbosity) VARARG_EXTRA(int32 UniqueLogId) VARARG_EXTRA(const FVector& Location) VARARG_EXTRA(float Radius) VARARG_EXTRA(const FColor& Color))
+{
+	SCOPE_CYCLE_COUNTER(STAT_VisualLog);
+	UWorld *World = nullptr;
+	FVisualLogEntry *CurrentEntry = nullptr;
+	if (CheckVisualLogInputInternal(Object, Category, Verbosity, &World, &CurrentEntry) == false)
+	{
+		return;
+	}
+
+	COLLAPSED_LOGF(
+		CurrentEntry->AddElement(Location, Category.GetCategoryName(), Verbosity, Color, Buffer, Radius);
+	);
+}
+
+inline
+VARARG_BODY(void, FVisualLogger::GeometryShapeLogf, const TCHAR*, VARARG_EXTRA(const class UObject* Object) VARARG_EXTRA(const struct FLogCategoryBase& Category) VARARG_EXTRA(ELogVerbosity::Type Verbosity) VARARG_EXTRA(int32 UniqueLogId) VARARG_EXTRA(const FBox& Box) VARARG_EXTRA(const FMatrix& Matrix) VARARG_EXTRA(const FColor& Color))
+{
+	SCOPE_CYCLE_COUNTER(STAT_VisualLog);
+	UWorld *World = nullptr;
+	FVisualLogEntry *CurrentEntry = nullptr;
+	if (CheckVisualLogInputInternal(Object, Category, Verbosity, &World, &CurrentEntry) == false)
+	{
+		return;
+	}
+
+	COLLAPSED_LOGF(
+		CurrentEntry->AddElement(Box, Matrix, Category.GetCategoryName(), Verbosity, Color, Buffer);
+	);
+}
+
+inline
+VARARG_BODY(void, FVisualLogger::GeometryShapeLogf, const TCHAR*, VARARG_EXTRA(const class UObject* Object) VARARG_EXTRA(const struct FLogCategoryBase& Category) VARARG_EXTRA(ELogVerbosity::Type Verbosity) VARARG_EXTRA(int32 UniqueLogId) VARARG_EXTRA(const FVector& Orgin) VARARG_EXTRA(const FVector& Direction) VARARG_EXTRA(const float Length) VARARG_EXTRA(const float Angle)  VARARG_EXTRA(const FColor& Color))
+{
+	SCOPE_CYCLE_COUNTER(STAT_VisualLog);
+	UWorld *World = nullptr;
+	FVisualLogEntry *CurrentEntry = nullptr;
+	if (CheckVisualLogInputInternal(Object, Category, Verbosity, &World, &CurrentEntry) == false)
+	{
+		return;
+	}
+
+	COLLAPSED_LOGF(
+		CurrentEntry->AddElement(Orgin, Direction, Length, Angle, Angle, Category.GetCategoryName(), Verbosity, Color, Buffer);
+	);
+}
+
+inline
+VARARG_BODY(void, FVisualLogger::GeometryShapeLogf, const TCHAR*, VARARG_EXTRA(const class UObject* Object) VARARG_EXTRA(const struct FLogCategoryBase& Category) VARARG_EXTRA(ELogVerbosity::Type Verbosity) VARARG_EXTRA(int32 UniqueLogId) VARARG_EXTRA(const FVector& Start) VARARG_EXTRA(const FVector& End) VARARG_EXTRA(const float Radius) VARARG_EXTRA(const FColor& Color))
+{
+	SCOPE_CYCLE_COUNTER(STAT_VisualLog);
+	UWorld *World = nullptr;
+	FVisualLogEntry *CurrentEntry = nullptr;
+	if (CheckVisualLogInputInternal(Object, Category, Verbosity, &World, &CurrentEntry) == false)
+	{
+		return;
+	}
+
+	COLLAPSED_LOGF(
+		CurrentEntry->AddElement(Start, End, Radius, Category.GetCategoryName(), Verbosity, Color, Buffer);
+	);
+}
+
+inline
+VARARG_BODY(void, FVisualLogger::GeometryShapeLogf, const TCHAR*, VARARG_EXTRA(const class UObject* Object) VARARG_EXTRA(const struct FLogCategoryBase& Category) VARARG_EXTRA(ELogVerbosity::Type Verbosity) VARARG_EXTRA(int32 UniqueLogId) VARARG_EXTRA(const FVector& Center) VARARG_EXTRA(float HalfHeight) VARARG_EXTRA(float Radius) VARARG_EXTRA(const FQuat& Rotation) VARARG_EXTRA(const FColor& Color))
+{
+	SCOPE_CYCLE_COUNTER(STAT_VisualLog);
+	UWorld *World = nullptr;
+	FVisualLogEntry *CurrentEntry = nullptr;
+	if (CheckVisualLogInputInternal(Object, Category, Verbosity, &World, &CurrentEntry) == false)
+	{
+		return;
+	}
+
+	COLLAPSED_LOGF(
+		CurrentEntry->AddElement(Center, HalfHeight, Radius, Rotation, Category.GetCategoryName(), Verbosity, Color, Buffer);
+	);
+}
+
+inline
+VARARG_BODY(void, FVisualLogger::HistogramDataLogf, const TCHAR*, VARARG_EXTRA(const class UObject* Object) VARARG_EXTRA(const struct FLogCategoryBase& Category) VARARG_EXTRA(ELogVerbosity::Type Verbosity) VARARG_EXTRA(int32 UniqueLogId) VARARG_EXTRA(FName GraphName) VARARG_EXTRA(FName DataName) VARARG_EXTRA(const FVector2D& Data) VARARG_EXTRA(const FColor& Color))
+{
+	SCOPE_CYCLE_COUNTER(STAT_VisualLog);
+	UWorld *World = nullptr;
+	FVisualLogEntry *CurrentEntry = nullptr;
+	if (CheckVisualLogInputInternal(Object, Category, Verbosity, &World, &CurrentEntry) == false)
+	{
+		return;
+	}
+
+	COLLAPSED_LOGF(
+		CurrentEntry->AddHistogramData(Data, Category.GetCategoryName(), Verbosity, GraphName, DataName);
+	);
+}
+
+#undef COLLAPSED_LOGF
+
 #endif //ENABLE_VISUAL_LOG
