Fix for a lot of platforms build issues regarding the physics serializer

diff --git a/Engine/Source/Runtime/Engine/Private/PhysicsEngine/PhysXSupport.h b/Engine/Source/Runtime/Engine/Private/PhysicsEngine/PhysXSupport.h
--- a/Engine/Source/Runtime/Engine/Private/PhysicsEngine/PhysXSupport.h
+++ b/Engine/Source/Runtime/Engine/Private/PhysicsEngine/PhysXSupport.h
@@ -21,6 +21,11 @@
 #define PHYSX_MEMORY_STATS		0
 #endif
 
+// binary serialization requires 128 byte alignment
+#ifndef PHYSX_SERIALIZATION_ALIGNMENT
+#define PHYSX_SERIALIZATION_ALIGNMENT 128
+#endif
+
 #define PHYSX_MEMORY_STAT_ONLY (0)
 
 #if USE_SCENE_LOCK
diff --git a/Engine/Source/Runtime/PhysXFormats/Public/IPhysXFormat.h b/Engine/Source/Runtime/PhysXFormats/Public/IPhysXFormat.h
--- a/Engine/Source/Runtime/PhysXFormats/Public/IPhysXFormat.h
+++ b/Engine/Source/Runtime/PhysXFormats/Public/IPhysXFormat.h
@@ -94,6 +94,3 @@ public:
 	 */
 	virtual ~IPhysXFormat( ) { }
 };
-
-// binary serialization requires 128 byte alignment
-#define PHYSX_SERIALIZATION_ALIGNMENT 128
\ No newline at end of file
