Fixed graph nodes sometimes not opening a rename when created

Changed the way that renames are initiated to depend on the event structure passed in, rather than keyboard focus
This addresses UE-9950.

diff --git a/Engine/Source/Editor/GraphEditor/Private/SGraphPanel.cpp b/Engine/Source/Editor/GraphEditor/Private/SGraphPanel.cpp
--- a/Engine/Source/Editor/GraphEditor/Private/SGraphPanel.cpp
+++ b/Engine/Source/Editor/GraphEditor/Private/SGraphPanel.cpp
@@ -1332,7 +1332,7 @@ void SGraphPanel::OnGraphChanged(const FEdGraphEditAction& EditAction)
 
 			for (auto Node : EditAction.Nodes)
 			{
-				RegisterActiveTimer(0.f, FWidgetActiveTimerDelegate::CreateStatic(AddNodeDelegateWrapper, this, const_cast<UEdGraphNode*>(Node), (HasKeyboardFocus() || HasFocusedDescendants()) ) );
+				RegisterActiveTimer(0.f, FWidgetActiveTimerDelegate::CreateStatic(AddNodeDelegateWrapper, this, const_cast<UEdGraphNode*>(Node), EditAction.bUserInvoked ) );
 			}
 		}
 		if (bWasSelectAction)
diff --git a/Engine/Source/Runtime/Engine/Private/EdGraph/EdGraph.cpp b/Engine/Source/Runtime/Engine/Private/EdGraph/EdGraph.cpp
--- a/Engine/Source/Runtime/Engine/Private/EdGraph/EdGraph.cpp
+++ b/Engine/Source/Runtime/Engine/Private/EdGraph/EdGraph.cpp
@@ -147,7 +147,7 @@ void UEdGraph::AddNode( UEdGraphNode* NodeToAdd, bool bFromUI/* = false*/, bool
 		AddNodeAction = (EEdGraphActionType)( ((int32)AddNodeAction) | GRAPHACTION_SelectNode );
 	}
 
-	FEdGraphEditAction Action(AddNodeAction, this, NodeToAdd);
+	FEdGraphEditAction Action(AddNodeAction, this, NodeToAdd, bFromUI);
 	
 	NotifyGraphChanged( Action );
 }
diff --git a/Engine/Source/Runtime/Engine/Public/GraphEditAction.h b/Engine/Source/Runtime/Engine/Public/GraphEditAction.h
--- a/Engine/Source/Runtime/Engine/Public/GraphEditAction.h
+++ b/Engine/Source/Runtime/Engine/Public/GraphEditAction.h
@@ -28,15 +28,19 @@ struct FEdGraphEditAction
 	class UEdGraph* Graph;
 	/** [Optional] The object the action occurred on */
 	TSet<const class UEdGraphNode*> Nodes;
+	/** Whether the user invoked this change or not */
+	bool bUserInvoked;
 
 	FEdGraphEditAction()
 		: Action(GRAPHACTION_Default)
 		, Graph (NULL)
+		, bUserInvoked(false)
 	{}
 
-	explicit FEdGraphEditAction(EEdGraphActionType InAction, class UEdGraph* InGraph, class UEdGraphNode* InNode)
+	explicit FEdGraphEditAction(EEdGraphActionType InAction, class UEdGraph* InGraph, class UEdGraphNode* InNode, bool bInUserInvoked)
 		: Action(InAction) 
 		, Graph(InGraph)
+		, bUserInvoked(bInUserInvoked)
 	{
 		Nodes.Add(InNode);
 	}
