Added HSV-based linear color interpolation function (Blueprint and C++)
- Interpolates between colors in HSV space, taking the shortest route between hues.  Also interpolates alpha.
- In C++:  FLinearColor::LerpUsingHSV()
- New Blueprint function:  "Lerp Using HSV (Linear Color)"

diff --git a/Engine/Source/Runtime/Core/Private/Math/Color.cpp b/Engine/Source/Runtime/Core/Private/Math/Color.cpp
--- a/Engine/Source/Runtime/Core/Private/Math/Color.cpp
+++ b/Engine/Source/Runtime/Core/Private/Math/Color.cpp
@@ -255,6 +255,47 @@ FLinearColor FLinearColor::HSVToLinearRGB() const
 						A);
 }
 
+
+FLinearColor FLinearColor::LerpUsingHSV( const FLinearColor& From, const FLinearColor& To, const float Progress )
+{
+	const FLinearColor FromHSV = From.LinearRGBToHSV();
+	const FLinearColor ToHSV = To.LinearRGBToHSV();
+
+	float FromHue = FromHSV.R;
+	float ToHue = ToHSV.R;
+
+	// Take the shortest path to the new hue
+	if( FMath::Abs( FromHue - ToHue ) > 180.0f )
+	{
+		if( ToHue > FromHue )
+		{
+			FromHue += 360.0f;
+		}
+		else
+		{
+			ToHue += 360.0f;
+		}
+	}
+
+	float NewHue = FMath::Lerp( FromHue, ToHue, Progress );
+
+	NewHue = FMath::Fmod( NewHue, 360.0f );
+	if( NewHue < 0.0f )
+	{
+		NewHue += 360.0f;
+	}
+
+	const float NewSaturation = FMath::Lerp( FromHSV.G, ToHSV.G, Progress );
+	const float NewValue = FMath::Lerp( FromHSV.B, ToHSV.B, Progress );
+	FLinearColor Interpolated = FLinearColor( NewHue, NewSaturation, NewValue ).HSVToLinearRGB();
+
+	const float NewAlpha = FMath::Lerp( From.A, To.A, Progress );
+	Interpolated.A = NewAlpha;
+
+	return Interpolated;
+}
+
+
 /**
 * Makes a random but quite nice color.
 */
diff --git a/Engine/Source/Runtime/Core/Public/Math/Color.h b/Engine/Source/Runtime/Core/Public/Math/Color.h
--- a/Engine/Source/Runtime/Core/Public/Math/Color.h
+++ b/Engine/Source/Runtime/Core/Public/Math/Color.h
@@ -242,6 +242,19 @@ struct FLinearColor
 	/** Converts an HSV color to a linear space RGB color */
 	CORE_API FLinearColor HSVToLinearRGB() const;
 
+	/**
+	 * Linearly interpolates between two colors by the specified progress amount.  The interpolation is performed in HSV color space
+	 * taking the shortest path to the new color's hue.  This can give better results than FMath::Lerp(), but is much more expensive.
+	 * The incoming colors are in RGB space, and the output color will be RGB.  The alpha value will also be interpolated.
+	 * 
+	 * @param	From		The color and alpha to interpolate from as linear RGBA
+	 * @param	To			The color and alpha to interpolate to as linear RGBA
+	 * @param	Progress	Scalar interpolation amount (usually between 0.0 and 1.0 inclusive)
+	 * 
+	 * @return	The interpolated color in linear RGB space along with the interpolated alpha value
+	 */
+	static CORE_API FLinearColor LerpUsingHSV( const FLinearColor& From, const FLinearColor& To, const float Progress );
+
 	/** Quantizes the linear color and returns the result as a FColor.  This bypasses the SRGB conversion. */
 	CORE_API FColor Quantize() const;
 
diff --git a/Engine/Source/Runtime/Engine/Classes/Kismet/KismetMathLibrary.h b/Engine/Source/Runtime/Engine/Classes/Kismet/KismetMathLibrary.h
--- a/Engine/Source/Runtime/Engine/Classes/Kismet/KismetMathLibrary.h
+++ b/Engine/Source/Runtime/Engine/Classes/Kismet/KismetMathLibrary.h
@@ -754,6 +754,18 @@ class ENGINE_API UKismetMathLibrary : public UBlueprintFunctionLibrary
 	UFUNCTION(BlueprintPure, meta=(DisplayName = "Lerp (LinearColor)"), Category="Math|Color")
 	static FLinearColor LinearColorLerp(FLinearColor A, FLinearColor B, float Alpha);
 
+	/**
+	 * Linearly interpolates between two colors by the specified Alpha amount (100% of A when Alpha=0 and 100% of B when Alpha=1).  The interpolation is performed in HSV color space taking the shortest path to the new color's hue.  This can give better results than a normal lerp, but is much more expensive.  The incoming colors are in RGB space, and the output color will be RGB.  The alpha value will also be interpolated.
+	 * 
+	 * @param	A		The color and alpha to interpolate from as linear RGBA
+	 * @param	B		The color and alpha to interpolate to as linear RGBA
+	 * @param	Alpha	Scalar interpolation amount (usually between 0.0 and 1.0 inclusive)
+	 * 
+	 * @return	The interpolated color in linear RGB space along with the interpolated alpha value
+	 */
+	UFUNCTION(BlueprintPure, meta=(DisplayName = "Lerp Using HSV (LinearColor)"), Category="Math|Color")
+	static FLinearColor LinearColorLerpUsingHSV(FLinearColor A, FLinearColor B, float Alpha);
+
 	/* Element-wise multiplication of two linear colors (R*R, G*G, B*B, A*A) */
 	UFUNCTION(BlueprintPure, meta=(DisplayName = "LinearColor * (LinearColor)", CompactNodeTitle = "*"), Category="Math|Color")
 	static FLinearColor Multiply_LinearColorLinearColor(FLinearColor A, FLinearColor B);
diff --git a/Engine/Source/Runtime/Engine/Private/KismetMathLibrary.cpp b/Engine/Source/Runtime/Engine/Private/KismetMathLibrary.cpp
--- a/Engine/Source/Runtime/Engine/Private/KismetMathLibrary.cpp
+++ b/Engine/Source/Runtime/Engine/Private/KismetMathLibrary.cpp
@@ -1016,6 +1016,11 @@ FLinearColor UKismetMathLibrary::LinearColorLerp(FLinearColor A, FLinearColor B,
 	return A + Alpha * (B - A);
 }
 
+FLinearColor UKismetMathLibrary::LinearColorLerpUsingHSV(FLinearColor A, FLinearColor B, float Alpha)
+{
+	return FLinearColor::LerpUsingHSV( A, B, Alpha );
+}
+
 FLinearColor UKismetMathLibrary::Multiply_LinearColorLinearColor(FLinearColor A, FLinearColor B)
 {
 	return A * B;
