Update Actor..AttachmentReplication structures in SceneComponent::AttachTo/DetachFromParent to ensure consistency

diff --git a/Engine/Source/Runtime/Engine/Classes/GameFramework/Actor.h b/Engine/Source/Runtime/Engine/Classes/GameFramework/Actor.h
--- a/Engine/Source/Runtime/Engine/Classes/GameFramework/Actor.h
+++ b/Engine/Source/Runtime/Engine/Classes/GameFramework/Actor.h
@@ -231,11 +231,11 @@ public:
 	ENetRole GetRemoteRole() const;
 
 	/** Used for replication of our RootComponent's position and velocity */
-	UPROPERTY(ReplicatedUsing=OnRep_ReplicatedMovement)
+	UPROPERTY(Transient, ReplicatedUsing=OnRep_ReplicatedMovement)
 	struct FRepMovement ReplicatedMovement;
 
 	/** Used for replicating attachment of this actor's RootComponent to another actor. */
-	UPROPERTY(replicatedUsing=OnRep_AttachmentReplication)
+	UPROPERTY(Transient, replicatedUsing=OnRep_AttachmentReplication)
 	struct FRepAttachment AttachmentReplication;
 
 	/** Called on client when updated AttachmentReplication value is received for this actor. */
diff --git a/Engine/Source/Runtime/Engine/Private/Actor.cpp b/Engine/Source/Runtime/Engine/Private/Actor.cpp
--- a/Engine/Source/Runtime/Engine/Private/Actor.cpp
+++ b/Engine/Source/Runtime/Engine/Private/Actor.cpp
@@ -1227,14 +1227,6 @@ void AActor::AttachRootComponentTo(USceneComponent* InParent, FName InSocketName
 	if(RootComponent && InParent)
 	{
 		RootComponent->AttachTo(InParent, InSocketName, AttachLocationType, bWeldSimulatedBodies);
-
-
-		AttachmentReplication.AttachParent = InParent->GetAttachmentRootActor();
-		AttachmentReplication.LocationOffset = RootComponent->RelativeLocation;
-		AttachmentReplication.RotationOffset = RootComponent->RelativeRotation;
-		AttachmentReplication.RelativeScale3D = RootComponent->RelativeScale3D;
-		AttachmentReplication.AttachSocket = InSocketName;
-		AttachmentReplication.AttachComponent = InParent;
 	}
 }
 
@@ -1282,14 +1274,6 @@ void AActor::AttachRootComponentToActor(AActor* InParentActor, FName InSocketNam
 		if (ParentRootComponent)
 		{
 			RootComponent->AttachTo(ParentRootComponent, InSocketName, AttachLocationType, bWeldSimulatedBodies );
-
-
-			AttachmentReplication.AttachParent = InParentActor;
-			AttachmentReplication.LocationOffset = RootComponent->RelativeLocation;
-			AttachmentReplication.RotationOffset = RootComponent->RelativeRotation;
-			AttachmentReplication.RelativeScale3D = RootComponent->RelativeScale3D;
-			AttachmentReplication.AttachSocket = InSocketName;
-			AttachmentReplication.AttachComponent = NULL;
 		}
 	}
 }
@@ -1311,7 +1295,6 @@ void AActor::DetachRootComponentFromParent(bool bMaintainWorldPosition)
 	if(RootComponent)
 	{
 		RootComponent->DetachFromParent(bMaintainWorldPosition);
-		AttachmentReplication.AttachParent = NULL;
 	}
 }
 
diff --git a/Engine/Source/Runtime/Engine/Private/Components/SceneComponent.cpp b/Engine/Source/Runtime/Engine/Private/Components/SceneComponent.cpp
--- a/Engine/Source/Runtime/Engine/Private/Components/SceneComponent.cpp
+++ b/Engine/Source/Runtime/Engine/Private/Components/SceneComponent.cpp
@@ -1269,6 +1269,18 @@ void USceneComponent::AttachTo(class USceneComponent* Parent, FName InSocketName
 		{
 			UpdateOverlaps();
 		}
+
+		AActor* Owner = GetOwner();
+
+		if (Owner && Owner->GetRootComponent() == this)
+		{
+			Owner->AttachmentReplication.AttachParent = AttachParent->GetAttachmentRootActor();
+			Owner->AttachmentReplication.LocationOffset = RelativeLocation;
+			Owner->AttachmentReplication.RotationOffset = RelativeRotation;
+			Owner->AttachmentReplication.RelativeScale3D = RelativeScale3D;
+			Owner->AttachmentReplication.AttachSocket = InSocketName;
+			Owner->AttachmentReplication.AttachComponent = AttachParent;
+		}
 	}
 }
 
@@ -1281,13 +1293,15 @@ void USceneComponent::DetachFromParent(bool bMaintainWorldPosition, bool bCallMo
 {
 	if(AttachParent != NULL)
 	{
+		AActor* Owner = GetOwner();
+
 		if (UPrimitiveComponent * PrimComp = Cast<UPrimitiveComponent>(this))
 		{
 			PrimComp->UnWeldFromParent();
 		}
 
 		// Make sure parent points to us if we're registered
-		checkf(!bRegistered || AttachParent->AttachChildren.Contains(this), TEXT("Attempt to detach SceneComponent '%s' owned by '%s' from AttachParent '%s' while not attached."), *GetName(), (GetOwner() ? *GetOwner()->GetName() : TEXT("Unowned")), *AttachParent->GetName());
+		checkf(!bRegistered || AttachParent->AttachChildren.Contains(this), TEXT("Attempt to detach SceneComponent '%s' owned by '%s' from AttachParent '%s' while not attached."), *GetName(), (Owner ? *Owner->GetName() : TEXT("Unowned")), *AttachParent->GetName());
 
 		if (bCallModify)
 		{
@@ -1303,9 +1317,9 @@ void USceneComponent::DetachFromParent(bool bMaintainWorldPosition, bool bCallMo
 #if WITH_EDITOR
 		if(GEngine)
 		{
-			if(GetOwner() && this == GetOwner()->GetRootComponent())
+			if(Owner && this == Owner->GetRootComponent())
 			{
-				GEngine->BroadcastLevelActorDetached(GetOwner(), AttachParent->GetOwner());
+				GEngine->BroadcastLevelActorDetached(Owner, AttachParent->GetOwner());
 			}
 		}
 #endif
@@ -1330,6 +1344,11 @@ void USceneComponent::DetachFromParent(bool bMaintainWorldPosition, bool bCallMo
 		{
 			UpdateOverlaps();
 		}
+
+		if (Owner && Owner->GetRootComponent() == this)
+		{
+			Owner->AttachmentReplication.AttachParent = nullptr;
+		}
 	}
 }
 
