Reinitialized GScreenshotResolutionX/Y when "shot" command is handled, so that the resolution is correctly set.

diff --git a/Engine/Source/Runtime/Engine/Private/GameViewportClient.cpp b/Engine/Source/Runtime/Engine/Private/GameViewportClient.cpp
--- a/Engine/Source/Runtime/Engine/Private/GameViewportClient.cpp
+++ b/Engine/Source/Runtime/Engine/Private/GameViewportClient.cpp
@@ -2839,6 +2839,8 @@ bool UGameViewportClient::HandleScreenshotCommand( const TCHAR* Cmd, FOutputDevi
 
 		GScreenMessagesRestoreState = GAreScreenMessagesEnabled;
 		GAreScreenMessagesEnabled = false;
+		GScreenshotResolutionX = Viewport->GetSizeXY().X;
+		GScreenshotResolutionY = Viewport->GetSizeXY().Y;
 	}
 	return true;
 }
