Fixed issues with taking screenshots (with f9) in a new editor window.

diff --git a/Engine/Source/Runtime/Engine/Private/GameViewportClient.cpp b/Engine/Source/Runtime/Engine/Private/GameViewportClient.cpp
--- a/Engine/Source/Runtime/Engine/Private/GameViewportClient.cpp
+++ b/Engine/Source/Runtime/Engine/Private/GameViewportClient.cpp
@@ -1303,11 +1303,13 @@ void UGameViewportClient::ProcessScreenShots(FViewport* InViewport)
 		}
 
 		bool bScreenshotSuccessful = false;
+		FIntVector Size(InViewport->GetSizeXY().X, InViewport->GetSizeXY().Y, 0);
 		if( bShowUI && FSlateApplication::IsInitialized() )
 		{
-			FIntVector Size;
 			TSharedRef<SWidget> WindowRef = WindowPtr.ToSharedRef();
 			bScreenshotSuccessful = FSlateApplication::Get().TakeScreenshot( WindowRef, Bitmap, Size);
+			GScreenshotResolutionX = Size.X;
+			GScreenshotResolutionY = Size.Y;
 		}
 		else
 		{
@@ -1318,7 +1320,7 @@ void UGameViewportClient::ProcessScreenShots(FViewport* InViewport)
 		{
 			if (ScreenshotCapturedDelegate.IsBound())
 			{
-				ScreenshotCapturedDelegate.Broadcast(InViewport->GetSizeXY().X, InViewport->GetSizeXY().Y, Bitmap);
+				ScreenshotCapturedDelegate.Broadcast(Size.X, Size.Y, Bitmap);
 			}
 			else
 			{
@@ -1332,7 +1334,7 @@ void UGameViewportClient::ProcessScreenShots(FViewport* InViewport)
 
 				if (PNGScreenshotCapturedDelegate.IsBound() && FPaths::GetExtension(ScreenShotName).ToLower() == TEXT("png"))
 				{
-					PNGScreenshotCapturedDelegate.Execute(InViewport->GetSizeXY().X, InViewport->GetSizeXY().Y, Bitmap, ScreenShotName);
+					PNGScreenshotCapturedDelegate.Execute(Size.X, Size.Y, Bitmap, ScreenShotName);
 				}
 				else
 				{
@@ -1350,7 +1352,7 @@ void UGameViewportClient::ProcessScreenShots(FViewport* InViewport)
 						ScreenShotName = FPaths::GetBaseFilename(ScreenShotName, false);
 						ScreenShotName += TEXT(".bmp");
 					}
-					FFileHelper::CreateBitmap(*ScreenShotName, InViewport->GetSizeXY().X, InViewport->GetSizeXY().Y, Bitmap.GetData(), &SourceRect, &IFileManager::Get(), NULL, bWriteAlpha);
+					FFileHelper::CreateBitmap(*ScreenShotName, Size.X, Size.Y, Bitmap.GetData(), &SourceRect, &IFileManager::Get(), NULL, bWriteAlpha);
 				}
 			}
 		}
