added comments

diff --git a/Engine/Source/Runtime/Engine/Private/SkeletalRenderGPUSkin.h b/Engine/Source/Runtime/Engine/Private/SkeletalRenderGPUSkin.h
--- a/Engine/Source/Runtime/Engine/Private/SkeletalRenderGPUSkin.h
+++ b/Engine/Source/Runtime/Engine/Private/SkeletalRenderGPUSkin.h
@@ -10,6 +10,11 @@
 #include "SkeletalRenderPublic.h"
 #include "GPUSkinVertexFactory.h" 
 
+// 1 for a single buffer, this creates and releases textures every frame (the driver has to keep the reference and need to defer the release, low memory as it only occupies rendered buffers (up to 3 copies), best Xbox360 method?)
+// 2 for double buffering (works well for PC, caused Xbox360 to stall)
+// 3 for triple buffering (works well for PC and Xbox360, wastes a bit of memory)
+#define PER_BONE_BUFFER_COUNT 2
+
 /** 
 * Stores the updated matrices needed to skin the verts.
 * Created by the game thread and sent to the rendering thread as an update 
@@ -160,11 +165,6 @@ private:
 	FSkeletalMeshResource* SkelMeshResource;
 };
 
-// 1 for a single buffer, this creates and releases textures every frame (the driver has to keep the reference and need to defer the release, low memory as it only occupies rendered buffers (up to 3 copies), best Xbox360 method?)
-// 2 for double buffering (works well for PC, caused Xbox360 to stall)
-// 3 for triple buffering (works well for PC and Xbox360, wastes a bit of memory)
-#define PER_BONE_BUFFER_COUNT 2
-
 // only used on the render thread
 class FPreviousPerBoneMotionBlur
 {
@@ -189,7 +189,7 @@ public:
 
 	/** so we update only during velocity rendering pass */
 	bool IsAppendStarted() const;
-
+	
 	/** needed before AppendData() ccan be called */
 	ENGINE_API void StartAppend();
 
@@ -201,7 +201,7 @@ public:
 	 */
 	uint32 AppendData(FBoneSkinning *DataStart, uint32 BoneCount);
 
-	/** only call if StartAppend() */
+	/** only call if StartAppend(), if the append wasn't started it silently ignores the call */
 	ENGINE_API void EndAppend();
 
 	/** @return 0 if there should be no bone based motion blur (no previous data available or it's not active) */
@@ -505,4 +505,5 @@ private:
 
 };
 
+// accessed on the rendering thread[s]
 extern ENGINE_API FPreviousPerBoneMotionBlur GPrevPerBoneMotionBlur;
