Backed off on distance field shadowing bias, which causes leaking in the distance

diff --git a/Engine/Source/Runtime/Engine/Classes/Components/DirectionalLightComponent.h b/Engine/Source/Runtime/Engine/Classes/Components/DirectionalLightComponent.h
--- a/Engine/Source/Runtime/Engine/Classes/Components/DirectionalLightComponent.h
+++ b/Engine/Source/Runtime/Engine/Classes/Components/DirectionalLightComponent.h
@@ -96,13 +96,13 @@ class UDirectionalLightComponent : public ULightComponent
 	uint32 bUseInsetShadowsForMovableObjects : 1;
 
 	/** 0: no DistantShadowCascades, otherwise the count of cascades between WholeSceneDynamicShadowRadius and DistantShadowDistance that are covered by distant shadow cascades. */
-	UPROPERTY(EditAnywhere, BlueprintReadOnly, Category=CascadedShadowMaps, meta=(UIMin = "0", UIMax = "4"), DisplayName = "Far Shadow Cascade Count")
+	UPROPERTY(EditAnywhere, AdvancedDisplay, BlueprintReadOnly, Category=CascadedShadowMaps, meta=(UIMin = "0", UIMax = "4"), DisplayName = "Far Shadow Cascade Count")
 	int32 FarShadowCascadeCount;
 
 	/** 
 	 * Distance at which the far shadow cascade should end.  Far shadows will cover the range between 'Dynamic Shadow Distance' and this distance. 
 	 */
-	UPROPERTY(EditAnywhere, BlueprintReadOnly, Category=CascadedShadowMaps, meta=(UIMin = "0", UIMax = "800000"), DisplayName = "Far Shadow Distance")
+	UPROPERTY(EditAnywhere, AdvancedDisplay, BlueprintReadOnly, Category=CascadedShadowMaps, meta=(UIMin = "0", UIMax = "800000"), DisplayName = "Far Shadow Distance")
 	float FarShadowDistance;
 
 	/** 
diff --git a/Engine/Source/Runtime/Engine/Private/Components/LightComponent.cpp b/Engine/Source/Runtime/Engine/Private/Components/LightComponent.cpp
--- a/Engine/Source/Runtime/Engine/Private/Components/LightComponent.cpp
+++ b/Engine/Source/Runtime/Engine/Private/Components/LightComponent.cpp
@@ -338,7 +338,7 @@ ULightComponent::ULightComponent(const FObjectInitializer& ObjectInitializer)
 	BloomThreshold = 0;
 	BloomTint = FColor::White;
 
-	RayStartOffsetDepthScale = .01f;
+	RayStartOffsetDepthScale = .003f;
 }
 
 
@@ -487,11 +487,19 @@ bool ULightComponent::CanEditChange(const UProperty* InProperty) const
 			}
 		}
 		
+		const bool bIsRayStartOffset = PropertyName == GET_MEMBER_NAME_STRING_CHECKED(ULightComponent, RayStartOffsetDepthScale);
+
 		if (PropertyName == GET_MEMBER_NAME_STRING_CHECKED(ULightComponent, bUseRayTracedDistanceFieldShadows)
-			|| PropertyName == GET_MEMBER_NAME_STRING_CHECKED(ULightComponent, RayStartOffsetDepthScale))
+			|| bIsRayStartOffset)
 		{
 			static const auto CVar = IConsoleManager::Get().FindTConsoleVariableDataInt(TEXT("r.GenerateMeshDistanceFields"));
 			bool bCanEdit = CastShadows && CastDynamicShadows && Mobility != EComponentMobility::Static && CVar->GetValueOnGameThread() != 0;
+
+			if (bIsRayStartOffset)
+			{
+				bCanEdit = bCanEdit && bUseRayTracedDistanceFieldShadows;
+			}
+
 			return bCanEdit;
 		}
 
