added debug code to track down velocity rendeirng issues

diff --git a/Engine/Source/Runtime/Engine/Private/GPUSkinVertexFactory.cpp b/Engine/Source/Runtime/Engine/Private/GPUSkinVertexFactory.cpp
--- a/Engine/Source/Runtime/Engine/Private/GPUSkinVertexFactory.cpp
+++ b/Engine/Source/Runtime/Engine/Private/GPUSkinVertexFactory.cpp
@@ -38,6 +38,16 @@ static FBoneMatricesUniformShaderParameters GBoneUniformStruct;
 	template class FactoryClass<false>;	\
 	template class FactoryClass<true>;
 
+#if !(UE_BUILD_SHIPPING || UE_BUILD_TEST)
+static TAutoConsoleVariable<int32> CVarVelocityTest(
+	TEXT("r.VelocityTest"),
+	0,
+	TEXT("Allows to enable some low level testing code for the velocity rendering (Affects object motion blur and TemporalAA).")
+	TEXT(" 0: off (default)")
+	TEXT(" 1: add random data to the buffer where we store skeletal mesh bone data to test if the code (good to test in PAUSED as well)."),
+	ECVF_Cheat | ECVF_RenderThreadSafe);
+#endif // if !(UE_BUILD_SHIPPING || UE_BUILD_TEST)
+
 /*-----------------------------------------------------------------------------
  FSharedPoolPolicyData
  -----------------------------------------------------------------------------*/
@@ -366,6 +376,13 @@ public:
 		Ar << BoneMatrices;
 		Ar << PreviousBoneMatrices;
 	}
+
+	/** Are we are in the velocity rendering pass or render velocity in the base pass? */
+	bool IsRenderingVelocity() const
+	{
+		return BoneIndexOffset.IsBound();
+	}
+
 	/**
 	* Set any shader data specific to this vertex factory
 	*/
@@ -407,11 +424,24 @@ public:
 			
 			if (FeatureLevel >= ERHIFeatureLevel::SM4 && GPrevPerBoneMotionBlur.IsAppendStarted())
 			{
+				const bool bWorldIsPaused = View.Family->bWorldIsPaused;
+
 				// we are in the velocity rendering pass
 
 				// 0xffffffff or valid index
 				uint32 OldBoneDataIndex = ShaderData.GetOldBoneData(FrameNumber);
 
+#if !(UE_BUILD_SHIPPING || UE_BUILD_TEST)
+				{
+					static const auto MotionBlurDebugVar = IConsoleManager::Get().FindTConsoleVariableDataInt(TEXT("r.MotionBlurDebug"));
+
+					if (MotionBlurDebugVar->GetValueOnRenderThread())
+					{
+						UE_LOG(LogEngine, Log, TEXT("%s"), *ShaderData.GetDebugString(VertexFactory, OldBoneDataIndex));
+					}
+				}
+#endif // !(UE_BUILD_SHIPPING || UE_BUILD_TEST)
+
 				// Read old data if it was written last frame (normal data) or this frame (e.g. split screen)
 				bLocalPerBoneMotionBlur = (OldBoneDataIndex != 0xffffffff);
 
@@ -438,10 +468,19 @@ public:
 						);
 				}
 				FScopeLock Lock(&ShaderData.OldBoneDataLock);
+
+				const FGPUBaseSkinVertexFactory* GPUVertexFactory = (const FGPUBaseSkinVertexFactory*)VertexFactory;
+
 				// if we haven't copied the data yet we skip the update (e.g. split screen)
 				if(ShaderData.IsOldBoneDataUpdateNeeded(FrameNumber))
 				{
-					const FGPUBaseSkinVertexFactory* GPUVertexFactory = (const FGPUBaseSkinVertexFactory*)VertexFactory;
+#if !(UE_BUILD_SHIPPING || UE_BUILD_TEST)
+					if (CVarVelocityTest.GetValueOnRenderThread())
+					{
+						FBoneSkinning Tab[10];
+						GPrevPerBoneMotionBlur.AppendData(Tab, (FMath::Rand() % 55) + 1);
+					}
+#endif // if !(UE_BUILD_SHIPPING || UE_BUILD_TEST)
 
 					// copy the bone data and tell the instance where it can pick it up next frame
 					// append data to a buffer we bind next frame to read old matrix data for motion blur
@@ -449,7 +488,6 @@ public:
 					GPUVertexFactory->SetOldBoneDataStartIndex(FrameNumber, OldBoneDataStartIndex);
 				}
 			}
-
 			SetShaderValue(RHICmdList, Shader->GetVertexShader(), PerBoneMotionBlur, bLocalPerBoneMotionBlur);
 		}
 	}
diff --git a/Engine/Source/Runtime/Engine/Private/SkeletalRenderGPUSkin.cpp b/Engine/Source/Runtime/Engine/Private/SkeletalRenderGPUSkin.cpp
--- a/Engine/Source/Runtime/Engine/Private/SkeletalRenderGPUSkin.cpp
+++ b/Engine/Source/Runtime/Engine/Private/SkeletalRenderGPUSkin.cpp
@@ -36,9 +36,17 @@ DECLARE_CYCLE_STAT(TEXT("Morph Vertex Buffer Update"),STAT_MorphVertexBuffer_Upd
 DECLARE_CYCLE_STAT(TEXT("Morph Vertex Buffer Init"),STAT_MorphVertexBuffer_Init,STATGROUP_MorphTarget);
 DECLARE_CYCLE_STAT(TEXT("Morph Vertex Buffer Apply Delta"),STAT_MorphVertexBuffer_ApplyDelta,STATGROUP_MorphTarget);
 
-
+// accessed on the rendering thread[s]
 FPreviousPerBoneMotionBlur GPrevPerBoneMotionBlur;
 
+static TAutoConsoleVariable<int32> CVarMotionBlurDebug(
+	TEXT("r.MotionBlurDebug"),
+	0,
+	TEXT("Defines if we log debugging output for motion blur rendering.\n")
+	TEXT(" 0: off (default)\n")
+	TEXT(" 1: on"),
+	ECVF_Cheat | ECVF_RenderThreadSafe);
+
 /*-----------------------------------------------------------------------------
 FMorphVertexBuffer
 -----------------------------------------------------------------------------*/
@@ -1091,6 +1099,11 @@ void FPreviousPerBoneMotionBlur::ReleaseResources()
 
 void FPreviousPerBoneMotionBlur::RestoreForPausedMotionBlur()
 {
+	if(CVarMotionBlurDebug.GetValueOnRenderThread())
+	{
+		UE_LOG(LogEngine, Log, TEXT("r.MotionBlurDebug: RestoreForPausedMotionBlur"));
+	}
+
 	if(BufferIndex == 0)
 	{
 		BufferIndex = PER_BONE_BUFFER_COUNT - 1;
@@ -1118,11 +1131,17 @@ void FPreviousPerBoneMotionBlur::InitIfNeeded()
 		InitResources();
 	}
 }
-void FPreviousPerBoneMotionBlur::StartAppend()
+void FPreviousPerBoneMotionBlur::StartAppend(bool bWorldIsPaused)
 {
 	check(!LockedData);
 	check(IsInRenderingThread());
 
+	if (CVarMotionBlurDebug.GetValueOnRenderThread())
+	{
+		UE_LOG(LogEngine, Log, TEXT("r.MotionBlurDebug: BufferIndex=%d/%d Read=%d Write=%d IsLocked=%d"),
+			BufferIndex, PER_BONE_BUFFER_COUNT, GetReadBufferIndex(), GetWriteBufferIndex(), IsAppendStarted() ? 1 : 0);
+	}
+
 	InitIfNeeded();
 
 	FBoneDataVertexBuffer& WriteTexture = PerChunkBoneMatricesTexture[GetWriteBufferIndex()];
@@ -1164,6 +1183,11 @@ uint32 FPreviousPerBoneMotionBlur::AppendData(FBoneSkinning *DataStart, uint32 B
 
 void FPreviousPerBoneMotionBlur::EndAppend()
 {
+	if(CVarMotionBlurDebug.GetValueOnRenderThread())
+	{
+		UE_LOG(LogEngine, Log, TEXT(" "));
+	}
+
 	if(IsAppendStarted())
 	{
 		check(IsInRenderingThread());
@@ -1176,7 +1200,7 @@ void FPreviousPerBoneMotionBlur::EndAppend()
 
 		AdvanceBufferIndex();
 	}
-	
+
 	{
 		static int LogSpawmPrevent = 0;
 
@@ -1202,6 +1226,11 @@ FBoneDataVertexBuffer* FPreviousPerBoneMotionBlur::GetReadData()
 	return &PerChunkBoneMatricesTexture[GetReadBufferIndex()];
 }
 
+FString FPreviousPerBoneMotionBlur::GetDebugString() const
+{
+	return FString::Printf(TEXT("BufferIndex=%d Pos=%d"), GPrevPerBoneMotionBlur.GetWriteBufferIndex(), GPrevPerBoneMotionBlur.LockedTexelPosition.GetValue());
+}
+
 uint32 FPreviousPerBoneMotionBlur::GetReadBufferIndex() const
 {
 	return BufferIndex;
diff --git a/Engine/Source/Runtime/Engine/Private/SkeletalRenderGPUSkin.h b/Engine/Source/Runtime/Engine/Private/SkeletalRenderGPUSkin.h
--- a/Engine/Source/Runtime/Engine/Private/SkeletalRenderGPUSkin.h
+++ b/Engine/Source/Runtime/Engine/Private/SkeletalRenderGPUSkin.h
@@ -189,9 +189,9 @@ public:
 
 	/** so we update only during velocity rendering pass */
 	bool IsAppendStarted() const;
-	
+
 	/** needed before AppendData() ccan be called */
-	ENGINE_API void StartAppend();
+	ENGINE_API void StartAppend(bool bWorldIsPaused);
 
 	/**
 	 * use between LockData() and UnlockData()
@@ -207,6 +207,8 @@ public:
 	/** @return 0 if there should be no bone based motion blur (no previous data available or it's not active) */
 	FBoneDataVertexBuffer* GetReadData();
 
+	FString GetDebugString() const;
+
 private:
 
 	/** Stores the bone information with one frame delay, required for per bone motion blur. Buffered to avoid stalls on texture Lock() */
diff --git a/Engine/Source/Runtime/Engine/Public/GPUSkinVertexFactory.h b/Engine/Source/Runtime/Engine/Public/GPUSkinVertexFactory.h
--- a/Engine/Source/Runtime/Engine/Public/GPUSkinVertexFactory.h
+++ b/Engine/Source/Runtime/Engine/Public/GPUSkinVertexFactory.h
@@ -243,6 +243,7 @@ public:
 			INC_DWORD_STAT_BY( STAT_SkeletalMeshMotionBlurSkinningMemory, ComputeMemorySize());
 			const int32 TileBufferSize = ComputeMemorySize();
 			FRHIResourceCreateInfo CreateInfo;
+			// BUF_Dynamic as it can be too large (e.g. 100 characters each 100 bones with 48 bytes per bone is about 1/2 MB) to do it as BUF_Volatile
 			BoneBuffer.VertexBufferRHI = RHICreateVertexBuffer( TileBufferSize, BUF_Dynamic | BUF_ShaderResource, CreateInfo );
 			BoneBuffer.VertexBufferSRV = RHICreateShaderResourceView( BoneBuffer.VertexBufferRHI, sizeof(FVector4), PF_A32B32G32R32F );
 		}
@@ -377,8 +378,19 @@ public:
 			return BoneBuffer;
 		}
 
+#if !(UE_BUILD_SHIPPING || UE_BUILD_TEST)
+		FString GetDebugString(const FVertexFactory* VertexFactory, uint32 OldBoneDataIndex) const
+		{
+			return FString::Printf(TEXT("r.MotionBlurDebug: OldBoneDataIndex[%p] (%d:%d %d:%d), => %d"),
+				VertexFactory,
+				OldBoneFrameNumber[0], OldBoneDataStartIndex[0],
+				OldBoneFrameNumber[1], OldBoneDataStartIndex[1],
+				OldBoneDataIndex);
+		}
+#endif // !(UE_BUILD_SHIPPING || UE_BUILD_TEST)
+
 		mutable FCriticalSection OldBoneDataLock;
-	private:
+private:
 
 		// the following members can be stored in less bytes (after some adjustments)
 
diff --git a/Engine/Source/Runtime/Renderer/Private/DeferredShadingRenderer.cpp b/Engine/Source/Runtime/Renderer/Private/DeferredShadingRenderer.cpp
--- a/Engine/Source/Runtime/Renderer/Private/DeferredShadingRenderer.cpp
+++ b/Engine/Source/Runtime/Renderer/Private/DeferredShadingRenderer.cpp
@@ -1401,7 +1401,7 @@ bool FDeferredShadingSceneRenderer::RenderBasePass(FRHICommandListImmediate& RHI
 	if (FVelocityRendering::OutputsToGBuffer())
 	{
 		QUICK_SCOPE_CYCLE_COUNTER(STAT_FDeferredShadingSceneRenderer_RenderBasePass_GPrevPerBoneMotionBlur_LockData);
-		GPrevPerBoneMotionBlur.StartAppend();
+		GPrevPerBoneMotionBlur.StartAppend(ViewFamily.bWorldIsPaused);
 	}
 
 	if(ViewFamily.EngineShowFlags.LightMapDensity && AllowDebugViewmodes())
diff --git a/Engine/Source/Runtime/Renderer/Private/PostProcess/PostProcessMotionBlur.cpp b/Engine/Source/Runtime/Renderer/Private/PostProcess/PostProcessMotionBlur.cpp
--- a/Engine/Source/Runtime/Renderer/Private/PostProcess/PostProcessMotionBlur.cpp
+++ b/Engine/Source/Runtime/Renderer/Private/PostProcess/PostProcessMotionBlur.cpp
@@ -1416,6 +1416,7 @@ void FRCPassPostProcessVisualizeMotionBlur::Process(FRenderingCompositePassConte
 	}
 
 	const FSceneView& View = Context.View;
+	const FSceneViewFamily& ViewFamily = *(View.Family);
 
 	FIntPoint TexSize = InputDesc->Extent;
 
@@ -1462,6 +1463,50 @@ void FRCPassPostProcessVisualizeMotionBlur::Process(FRenderingCompositePassConte
 		*VertexShader,
 		EDRF_UseTriangleOptimization);
 
+	// this is a helper class for FCanvas to be able to get screen size
+	class FRenderTargetTemp : public FRenderTarget
+	{
+	public:
+		const FSceneView& View;
+		const FTexture2DRHIRef Texture;
+
+		FRenderTargetTemp(const FSceneView& InView, const FTexture2DRHIRef InTexture)
+			: View(InView), Texture(InTexture)
+		{
+		}
+		virtual FIntPoint GetSizeXY() const
+		{
+			return View.ViewRect.Size();
+		};
+		virtual const FTexture2DRHIRef& GetRenderTargetTexture() const
+		{
+			return Texture;
+		}
+	} TempRenderTarget(View, (const FTexture2DRHIRef&)DestRenderTarget.TargetableTexture);
+
+	FCanvas Canvas(&TempRenderTarget, NULL, ViewFamily.CurrentRealTime, ViewFamily.CurrentWorldTime, ViewFamily.DeltaWorldTime, Context.GetFeatureLevel());
+
+	float X = 30;
+	float Y = 8;
+	const float YStep = 14;
+	const float ColumnWidth = 250;
+
+	FString Line;
+
+	Line = FString::Printf(TEXT("Visualize MotionBlur"));
+	Canvas.DrawShadowedString(X, Y += YStep, *Line, GetStatsFont(), FLinearColor(1, 1, 1));
+	
+	Line = FString::Printf(TEXT("%d"), ViewFamily.FrameNumber);
+	Canvas.DrawShadowedString(X, Y += YStep, TEXT("FrameNumber:"), GetStatsFont(), FLinearColor(1, 1, 1));
+	Canvas.DrawShadowedString(X + ColumnWidth, Y, *Line, GetStatsFont(), FLinearColor(1, 1, 1));
+
+	Line = FString::Printf(TEXT("%d"), ViewFamily.bWorldIsPaused ? 1 : 0);
+	Canvas.DrawShadowedString(X, Y += YStep, TEXT("WorldIsPaused:"), GetStatsFont(), FLinearColor(1, 1, 1));
+	Canvas.DrawShadowedString(X + ColumnWidth, Y, *Line, GetStatsFont(), FLinearColor(1, 1, 1));
+
+	Canvas.Flush_RenderThread(Context.RHICmdList);
+
+
 	Context.RHICmdList.CopyToResolveTarget(DestRenderTarget.TargetableTexture, DestRenderTarget.ShaderResourceTexture, false, FResolveParams());
 }
 
diff --git a/Engine/Source/Runtime/Renderer/Private/VelocityRendering.cpp b/Engine/Source/Runtime/Renderer/Private/VelocityRendering.cpp
--- a/Engine/Source/Runtime/Renderer/Private/VelocityRendering.cpp
+++ b/Engine/Source/Runtime/Renderer/Private/VelocityRendering.cpp
@@ -738,6 +738,15 @@ void FDeferredShadingSceneRenderer::RenderVelocities(FRHICommandListImmediate& R
 	FPooledRenderTargetDesc Desc = FVelocityRendering::GetRenderTargetDesc();
 	GRenderTargetPool.FindFreeElement(Desc, VelocityRT, TEXT("Velocity"));
 
+	{
+		static const auto MotionBlurDebugVar = IConsoleManager::Get().FindTConsoleVariableDataInt(TEXT("r.MotionBlurDebug"));
+
+		if(MotionBlurDebugVar->GetValueOnRenderThread())
+		{
+			UE_LOG(LogEngine, Log, TEXT("r.MotionBlurDebug: FrameNumber=%d Pause=%d"), ViewFamily.FrameNumber, ViewFamily.bWorldIsPaused ? 1 : 0);
+		}
+	}
+
 	{
 		GPrevPerBoneMotionBlur.StartAppend(ViewFamily.bWorldIsPaused);
 
