#ue4 - Clean up and organize Character Movement Component settings.

diff --git a/Engine/Source/Runtime/Engine/Classes/AI/Navigation/NavigationTypes.h b/Engine/Source/Runtime/Engine/Classes/AI/Navigation/NavigationTypes.h
--- a/Engine/Source/Runtime/Engine/Classes/AI/Navigation/NavigationTypes.h
+++ b/Engine/Source/Runtime/Engine/Classes/AI/Navigation/NavigationTypes.h
@@ -358,19 +358,19 @@ struct ENGINE_API FNavAgentProperties : public FMovementProperties
 	GENERATED_USTRUCT_BODY()
 
 	/** Radius of the capsule used for navigation/pathfinding. */
-	UPROPERTY(EditAnywhere, BlueprintReadWrite, Category = MovementProperties)
+	UPROPERTY(EditAnywhere, BlueprintReadWrite, Category=MovementProperties, meta=(DisplayName="Nav Agent Radius"))
 	float AgentRadius;
 
 	/** Total height of the capsule used for navigation/pathfinding. */
-	UPROPERTY(EditAnywhere, BlueprintReadWrite, Category = MovementProperties)
+	UPROPERTY(EditAnywhere, BlueprintReadWrite, Category=MovementProperties, meta=(DisplayName="Nav Agent Height"))
 	float AgentHeight;
 
 	/** Step height to use, or -1 for default value from navdata's config. */
-	UPROPERTY(EditAnywhere, BlueprintReadWrite, Category = MovementProperties)
+	UPROPERTY(EditAnywhere, BlueprintReadWrite, Category=MovementProperties, meta=(DisplayName="Nav Agent Step Height"))
 	float AgentStepHeight;
 
 	/** Scale factor to apply to height of bounds when searching for navmesh to project to when nav walking */
-	UPROPERTY(EditAnywhere, BlueprintReadWrite, Category = MovementProperties)
+	UPROPERTY(EditAnywhere, BlueprintReadWrite, Category=MovementProperties)
 	float NavWalkingSearchHeightScale;
 
 	FNavAgentProperties(float Radius = -1.f, float Height = -1.f)
diff --git a/Engine/Source/Runtime/Engine/Classes/GameFramework/CharacterMovementComponent.h b/Engine/Source/Runtime/Engine/Classes/GameFramework/CharacterMovementComponent.h
--- a/Engine/Source/Runtime/Engine/Classes/GameFramework/CharacterMovementComponent.h
+++ b/Engine/Source/Runtime/Engine/Classes/GameFramework/CharacterMovementComponent.h
@@ -137,16 +137,20 @@ protected:
 
 public:
 
+	/** Custom gravity scale. Gravity is multiplied by this amount for the character. */
+	UPROPERTY(Category="Character Movement (General Settings)", EditAnywhere, BlueprintReadWrite)
+	float GravityScale;
+
 	/** Maximum height character can step up */
-	UPROPERTY(Category="Character Movement", EditAnywhere, BlueprintReadWrite, meta=(ClampMin="0", UIMin="0"))
+	UPROPERTY(Category="Character Movement: Walking", EditAnywhere, BlueprintReadWrite, meta=(ClampMin="0", UIMin="0"))
 	float MaxStepHeight;
 
 	/** Initial velocity (instantaneous vertical acceleration) when jumping. */
-	UPROPERTY(Category="Character Movement", EditAnywhere, BlueprintReadWrite, meta=(DisplayName="Jump Z Velocity", ClampMin="0", UIMin="0"))
+	UPROPERTY(Category="Character Movement: Jumping / Falling", EditAnywhere, BlueprintReadWrite, meta=(DisplayName="Jump Z Velocity", ClampMin="0", UIMin="0"))
 	float JumpZVelocity;
 
 	/** Fraction of JumpZVelocity to use when automatically "jumping off" of a base actor that's not allowed to be a base for a character. (For example, if you're not allowed to stand on other players.) */
-	UPROPERTY(Category="Character Movement", EditAnywhere, BlueprintReadWrite, AdvancedDisplay, meta=(ClampMin="0", UIMin="0"))
+	UPROPERTY(Category="Character Movement: Jumping / Falling", EditAnywhere, BlueprintReadWrite, AdvancedDisplay, meta=(ClampMin="0", UIMin="0"))
 	float JumpOffJumpZFactor;
 
 private:
@@ -154,13 +158,13 @@ private:
 	/**
 	 * Max angle in degrees of a walkable surface. Any greater than this and it is too steep to be walkable.
 	 */
-	UPROPERTY(Category="Character Movement", EditAnywhere, meta=(ClampMin="0.0", ClampMax="90.0", UIMin = "0.0", UIMax = "90.0"))
+	UPROPERTY(Category="Character Movement: Walking", EditAnywhere, meta=(ClampMin="0.0", ClampMax="90.0", UIMin = "0.0", UIMax = "90.0"))
 	float WalkableFloorAngle;
 
 	/**
 	 * Minimum Z value for floor normal. If less, not a walkable surface. Computed from WalkableFloorAngle.
 	 */
-	UPROPERTY(Category="Character Movement", VisibleAnywhere)
+	UPROPERTY(Category="Character Movement: Walking", VisibleAnywhere)
 	float WalkableFloorZ;
 
 public:
@@ -175,7 +179,7 @@ public:
 	 * This is automatically replicated through the Character owner and for client-server movement functions.
 	 * @see SetMovementMode(), CustomMovementMode
 	 */
-	UPROPERTY(Category=MovementMode, BlueprintReadOnly)
+	UPROPERTY(Category="Character Movement: MovementMode", BlueprintReadOnly)
 	TEnumAsByte<enum EMovementMode> MovementMode;
 
 	/**
@@ -183,7 +187,7 @@ public:
 	 * This is automatically replicated through the Character owner and for client-server movement functions.
 	 * @see SetMovementMode()
 	 */
-	UPROPERTY(Category=MovementMode, BlueprintReadOnly)
+	UPROPERTY(Category="Character Movement: MovementMode", BlueprintReadOnly)
 	uint8 CustomMovementMode;
 
 	/** Saved location of object we are standing on, for UpdateBasedMovement() to determine if base moved in the last frame, and therefore pawn needs an update. */
@@ -192,100 +196,105 @@ public:
 	/** Saved location of object we are standing on, for UpdateBasedMovement() to determine if base moved in the last frame, and therefore pawn needs an update. */
 	FQuat OldBaseQuat;
 
-	/** Custom gravity scale. Gravity is multiplied by this amount for the character. */
-	UPROPERTY(Category="Character Movement", EditAnywhere, BlueprintReadWrite)
-	float GravityScale;
-
 	/**
 	 * Coefficient of friction.
 	 * This property allows you to control how much friction is applied when moving across the ground, applying an opposing force that scales with current velocity.
 	 * This can be used to simulate slippery surfaces such as ice or oil by changing the value (possibly based on the material pawn is standing on).
 	 * @see BrakingDecelerationWalking
 	 */
-	UPROPERTY(Category="Character Movement", EditAnywhere, BlueprintReadWrite, meta=(ClampMin="0", UIMin="0"))
+	UPROPERTY(Category="Character Movement: Walking", EditAnywhere, BlueprintReadWrite, meta=(ClampMin="0", UIMin="0"))
 	float GroundFriction;
 
 	/** The maximum ground speed when walking. Also determines maximum lateral speed when falling. */
-	UPROPERTY(Category="Character Movement", EditAnywhere, BlueprintReadWrite, meta=(ClampMin="0", UIMin="0"))
+	UPROPERTY(Category="Character Movement: Walking", EditAnywhere, BlueprintReadWrite, meta=(ClampMin="0", UIMin="0"))
 	float MaxWalkSpeed;
 
 	/** The maximum ground speed when walking and crouched. */
-	UPROPERTY(Category="Character Movement", EditAnywhere, BlueprintReadWrite, meta=(ClampMin="0", UIMin="0"))
+	UPROPERTY(Category="Character Movement: Walking", EditAnywhere, BlueprintReadWrite, meta=(ClampMin="0", UIMin="0"))
 	float MaxWalkSpeedCrouched;
 
-	/** The maximum speed when using Custom movement mode. */
-	UPROPERTY(Category="Character Movement", EditAnywhere, BlueprintReadWrite, meta=(ClampMin="0", UIMin="0"))
-	float MaxCustomMovementSpeed;
-
 	/** The maximum swimming speed. */
-	UPROPERTY(Category="Character Movement", EditAnywhere, BlueprintReadWrite, meta=(ClampMin="0", UIMin="0"))
+	UPROPERTY(Category="Character Movement: Swimming", EditAnywhere, BlueprintReadWrite, meta=(ClampMin="0", UIMin="0"))
 	float MaxSwimSpeed;
 
 	/** The maximum flying speed. */
-	UPROPERTY(Category="Character Movement", EditAnywhere, BlueprintReadWrite, meta=(ClampMin="0", UIMin="0"))
+	UPROPERTY(Category="Character Movement: Flying", EditAnywhere, BlueprintReadWrite, meta=(ClampMin="0", UIMin="0"))
 	float MaxFlySpeed;
 
+	/** The maximum speed when using Custom movement mode. */
+	UPROPERTY(Category="Character Movement: Custom Movement", EditAnywhere, BlueprintReadWrite, meta=(ClampMin="0", UIMin="0"))
+	float MaxCustomMovementSpeed;
+
 	/** Max Acceleration (rate of change of velocity) */
-	UPROPERTY(Category="Character Movement", EditAnywhere, BlueprintReadWrite, meta=(ClampMin="0", UIMin="0"))
+	UPROPERTY(Category="Character Movement (General Settings)", EditAnywhere, BlueprintReadWrite, meta=(ClampMin="0", UIMin="0"))
 	float MaxAcceleration;
 
 	/**
 	 * Deceleration when walking and not applying acceleration. This is a constant opposing force that directly lowers velocity by a constant value.
-	 * @see GroundFriction
+	 * @see GroundFriction, MaxAcceleration
 	 */
-	UPROPERTY(Category="Character Movement", EditAnywhere, BlueprintReadWrite, meta=(ClampMin="0", UIMin="0"))
+	UPROPERTY(Category="Character Movement: Walking", EditAnywhere, BlueprintReadWrite, meta=(ClampMin="0", UIMin="0"))
 	float BrakingDecelerationWalking;
 
-	/** Lateral deceleration when falling and not applying acceleration. */
-	UPROPERTY(Category="Character Movement", EditAnywhere, BlueprintReadWrite, meta=(ClampMin="0", UIMin="0"))
+	/**
+	 * Lateral deceleration when falling and not applying acceleration.
+	 * @see MaxAcceleration
+	 */
+	UPROPERTY(Category="Character Movement: Jumping / Falling", EditAnywhere, BlueprintReadWrite, meta=(ClampMin="0", UIMin="0"))
 	float BrakingDecelerationFalling;
 
-	/** Deceleration when swimming and not applying acceleration. */
-	UPROPERTY(Category="Character Movement", EditAnywhere, BlueprintReadWrite, meta=(ClampMin="0", UIMin="0"))
+	/**
+	 * Deceleration when swimming and not applying acceleration.
+	 * @see MaxAcceleration
+	 */
+	UPROPERTY(Category="Character Movement: Swimming", EditAnywhere, BlueprintReadWrite, meta=(ClampMin="0", UIMin="0"))
 	float BrakingDecelerationSwimming;
 
-	/** Deceleration when flying and not applying acceleration. */
-	UPROPERTY(Category="Character Movement", EditAnywhere, BlueprintReadWrite, meta=(ClampMin="0", UIMin="0"))
+	/**
+	 * Deceleration when flying and not applying acceleration.
+	 * @see MaxAcceleration
+	 */
+	UPROPERTY(Category="Character Movement: Flying", EditAnywhere, BlueprintReadWrite, meta=(ClampMin="0", UIMin="0"))
 	float BrakingDecelerationFlying;
 
 	/**
 	 * When falling, amount of lateral movement control available to the character.
 	 * 0 = no control, 1 = full control at max speed of MaxWalkSpeed.
 	 */
-	UPROPERTY(Category="Character Movement", EditAnywhere, BlueprintReadWrite, meta=(ClampMin="0", UIMin="0"))
+	UPROPERTY(Category="Character Movement: Jumping / Falling", EditAnywhere, BlueprintReadWrite, meta=(ClampMin="0", UIMin="0"))
 	float AirControl;
 
 	/**
 	 * When falling, multiplier applied to AirControl when lateral velocity is less than AirControlBoostVelocityThreshold.
 	 * Setting this to zero will disable air control boosting. Final result is clamped at 1.
 	 */
-	UPROPERTY(Category="Character Movement", EditAnywhere, BlueprintReadWrite, meta=(ClampMin="0", UIMin="0"))
+	UPROPERTY(Category="Character Movement: Jumping / Falling", EditAnywhere, BlueprintReadWrite, meta=(ClampMin="0", UIMin="0"))
 	float AirControlBoostMultiplier;
 
 	/**
 	 * When falling, if lateral velocity magnitude is less than this value, AirControl is multiplied by AirControlBoostMultiplier.
 	 * Setting this to zero will disable air control boosting.
 	 */
-	UPROPERTY(Category="Character Movement", EditAnywhere, BlueprintReadWrite, meta=(ClampMin="0", UIMin="0"))
+	UPROPERTY(Category="Character Movement: Jumping / Falling", EditAnywhere, BlueprintReadWrite, meta=(ClampMin="0", UIMin="0"))
 	float AirControlBoostVelocityThreshold;
 
 	/** Friction to apply to lateral air movement when falling. */
-	UPROPERTY(Category="Character Movement", EditAnywhere, BlueprintReadWrite, meta=(ClampMin="0", UIMin="0"))
+	UPROPERTY(Category="Character Movement: Jumping / Falling", EditAnywhere, BlueprintReadWrite, meta=(ClampMin="0", UIMin="0"))
 	float FallingLateralFriction;
 
 	/** Collision half-height when crouching (component scale is applied separately) */
-	UPROPERTY(Category="Character Movement", EditAnywhere, BlueprintReadOnly, meta=(ClampMin="0", UIMin="0"))
+	UPROPERTY(Category="Character Movement (General Settings)", EditAnywhere, BlueprintReadOnly, meta=(ClampMin="0", UIMin="0"))
 	float CrouchedHalfHeight;
 
 	/** Water buoyancy. A ratio (1.0 = neutral buoyancy, 0.0 = no buoyancy) */
-	UPROPERTY(Category="Character Movement", EditAnywhere, BlueprintReadWrite)
+	UPROPERTY(Category="Character Movement: Swimming", EditAnywhere, BlueprintReadWrite)
 	float Buoyancy;
 
 	/**
 	 * Don't allow the character to perch on the edge of a surface if the contact is this close to the edge of the capsule.
 	 * Note that characters will not fall off if they are within MaxStepHeight of a walkable surface below.
 	 */
-	UPROPERTY(Category="Character Movement", EditAnywhere, BlueprintReadWrite, AdvancedDisplay, meta=(ClampMin="0", UIMin="0"))
+	UPROPERTY(Category="Character Movement: Walking", EditAnywhere, BlueprintReadWrite, AdvancedDisplay, meta=(ClampMin="0", UIMin="0"))
 	float PerchRadiusThreshold;
 
 	/**
@@ -293,22 +302,22 @@ public:
 	 * Note that we still enforce MaxStepHeight to start the step up; this just allows the character to hang off the edge or step slightly higher off the floor.
 	 * (@see PerchRadiusThreshold)
 	 */
-	UPROPERTY(Category="Character Movement", EditAnywhere, BlueprintReadWrite, AdvancedDisplay, meta=(ClampMin="0", UIMin="0"))
+	UPROPERTY(Category="Character Movement: Walking", EditAnywhere, BlueprintReadWrite, AdvancedDisplay, meta=(ClampMin="0", UIMin="0"))
 	float PerchAdditionalHeight;
 
 	/** Change in rotation per second, used when UseControllerDesiredRotation or OrientRotationToMovement are true. */
-	UPROPERTY(Category="Character Movement", EditAnywhere, BlueprintReadWrite)
+	UPROPERTY(Category="Character Movement (General Settings)", EditAnywhere, BlueprintReadWrite)
 	FRotator RotationRate;
 
 	/** If true, smoothly rotate the Character toward the Controller's desired rotation, using RotationRate as the rate of rotation change. Overridden by OrientRotationToMovement. */
-	UPROPERTY(Category="Character Movement", EditAnywhere, BlueprintReadWrite, AdvancedDisplay)
+	UPROPERTY(Category="Character Movement (General Settings)", EditAnywhere, BlueprintReadWrite, AdvancedDisplay)
 	uint32 bUseControllerDesiredRotation:1;
 
 	/**
 	 * If true, rotate the Character toward the direction of acceleration, using RotationRate as the rate of rotation change. Overrides UseControllerDesiredRotation.
 	 * Normally you will want to make sure that other settings are cleared, such as bUseControllerRotationYaw on the Character.
 	 */
-	UPROPERTY(Category="Character Movement", EditAnywhere, BlueprintReadWrite)
+	UPROPERTY(Category="Character Movement (General Settings)", EditAnywhere, BlueprintReadWrite)
 	uint32 bOrientRotationToMovement:1;
 
 protected:
@@ -329,7 +338,7 @@ public:
 	 *
 	 * @see FScopedMovementUpdate
 	 */
-	UPROPERTY(Category="Character Movement", EditAnywhere, AdvancedDisplay)
+	UPROPERTY(Category="Character Movement (General Settings)", EditAnywhere, AdvancedDisplay)
 	uint32 bEnableScopedMovementUpdates:1;
 
 	/** Ignores size of acceleration component, and forces max acceleration to drive character at full velocity. */
@@ -342,14 +351,14 @@ public:
 	 * Characters that are spawned without a Controller but with this flag enabled will initialize the movement mode to DefaultLandMovementMode or DefaultWaterMovementMode appropriately.
 	 * @see DefaultLandMovementMode, DefaultWaterMovementMode
 	 */
-	UPROPERTY(Category="Character Movement", EditAnywhere, BlueprintReadWrite, AdvancedDisplay)
+	UPROPERTY(Category="Character Movement (General Settings)", EditAnywhere, BlueprintReadWrite, AdvancedDisplay)
 	uint32 bRunPhysicsWithNoController:1;
 
 	/**
 	 * Force the Character in MOVE_Walking to do a check for a valid floor even if he hasn't moved. Cleared after next floor check.
 	 * Normally if bAlwaysCheckFloor is false we try to avoid the floor check unless some conditions are met, but this can be used to force the next check to always run.
 	 */
-	UPROPERTY(Category="Character Movement", VisibleInstanceOnly, BlueprintReadWrite, AdvancedDisplay)
+	UPROPERTY(Category="Character Movement: Walking", VisibleInstanceOnly, BlueprintReadWrite, AdvancedDisplay)
 	uint32 bForceNextFloorCheck:1;
 
 	/** If true, the capsule needs to be shrunk on this simulated proxy, to avoid replication rounding putting us in geometry.
@@ -358,11 +367,11 @@ public:
 	uint32 bShrinkProxyCapsule:1;
 
 	/** If true, Character can walk off a ledge. */
-	UPROPERTY(Category="Character Movement", EditAnywhere, BlueprintReadWrite)
+	UPROPERTY(Category="Character Movement: Walking", EditAnywhere, BlueprintReadWrite)
 	uint32 bCanWalkOffLedges:1;
 
 	/** If true, Character can walk off a ledge when crouching. */
-	UPROPERTY(Category="Character Movement", EditAnywhere, BlueprintReadWrite)
+	UPROPERTY(Category="Character Movement: Walking", EditAnywhere, BlueprintReadWrite)
 	uint32 bCanWalkOffLedgesWhenCrouching:1;
 
 public:
@@ -376,64 +385,64 @@ public:
 	USceneComponent* DeferredUpdatedMoveComponent;
 
 	/** Maximum step height for getting out of water */
-	UPROPERTY(Category="Character Movement", EditAnywhere, BlueprintReadWrite, AdvancedDisplay, meta=(ClampMin="0", UIMin="0"))
+	UPROPERTY(Category="Character Movement: Swimming", EditAnywhere, BlueprintReadWrite, AdvancedDisplay, meta=(ClampMin="0", UIMin="0"))
 	float MaxOutOfWaterStepHeight;
 
 	/** Z velocity applied when pawn tries to get out of water */
-	UPROPERTY(Category="Character Movement", EditAnywhere, BlueprintReadWrite, AdvancedDisplay)
+	UPROPERTY(Category="Character Movement: Swimming", EditAnywhere, BlueprintReadWrite, AdvancedDisplay)
 	float OutofWaterZ;
 
 	/** Mass of pawn (for when momentum is imparted to it). */
-	UPROPERTY(Category="Character Movement", EditAnywhere, BlueprintReadWrite, meta=(ClampMin="0", UIMin="0"))
+	UPROPERTY(Category="Character Movement (General Settings)", EditAnywhere, BlueprintReadWrite, meta=(ClampMin="0", UIMin="0"))
 	float Mass;
 
 	/** If enabled, the player will interact with physics objects when walking into them. */
-	UPROPERTY(Category="Agent Physics", EditAnywhere, BlueprintReadWrite)
+	UPROPERTY(Category="Character Movement: Physics Interaction", EditAnywhere, BlueprintReadWrite)
 	bool bEnablePhysicsInteraction;
 
 	/** If enabled, the TouchForceFactor is applied per kg mass of the affected object. */
-	UPROPERTY(Category="Agent Physics", EditAnywhere, BlueprintReadWrite, meta=(editcondition = "bEnablePhysicsInteraction"))
+	UPROPERTY(Category="Character Movement: Physics Interaction", EditAnywhere, BlueprintReadWrite, meta=(editcondition = "bEnablePhysicsInteraction"))
 	bool bTouchForceScaledToMass;
 
 	/** If enabled, the PushForceFactor is applied per kg mass of the affected object. */
-	UPROPERTY(Category="Agent Physics", EditAnywhere, BlueprintReadWrite, meta=(editcondition = "bEnablePhysicsInteraction"))
+	UPROPERTY(Category="Character Movement: Physics Interaction", EditAnywhere, BlueprintReadWrite, meta=(editcondition = "bEnablePhysicsInteraction"))
 	bool bPushForceScaledToMass;
 
 	/** If enabled, the applied push force will try to get the physics object to the same velocity than the player, not faster. This will only
 		scale the force down, it will never apply more force than defined by PushForceFactor. */
-	UPROPERTY(Category="Agent Physics", EditAnywhere, BlueprintReadWrite, meta=(editcondition = "bEnablePhysicsInteraction"))
+	UPROPERTY(Category="Character Movement: Physics Interaction", EditAnywhere, BlueprintReadWrite, meta=(editcondition = "bEnablePhysicsInteraction"))
 	bool bScalePushForceToVelocity;
 
 	/** Force applied to objects we stand on (due to Mass and Gravity) is scaled by this amount. */
-	UPROPERTY(Category="Agent Physics", EditAnywhere, BlueprintReadWrite, meta=(editcondition = "bEnablePhysicsInteraction"))
+	UPROPERTY(Category="Character Movement: Physics Interaction", EditAnywhere, BlueprintReadWrite, meta=(editcondition = "bEnablePhysicsInteraction"))
 	float StandingDownwardForceScale;
 
 	/** Initial impulse force to apply when the player bounces into a blocking physics object. */
-	UPROPERTY(Category="Agent Physics", EditAnywhere, BlueprintReadWrite, meta=(editcondition = "bEnablePhysicsInteraction"))
+	UPROPERTY(Category="Character Movement: Physics Interaction", EditAnywhere, BlueprintReadWrite, meta=(editcondition = "bEnablePhysicsInteraction"))
 	float InitialPushForceFactor;
 
 	/** Force to apply when the player collides with a blocking physics object. */
-	UPROPERTY(Category="Agent Physics", EditAnywhere, BlueprintReadWrite, meta=(editcondition = "bEnablePhysicsInteraction"))
+	UPROPERTY(Category="Character Movement: Physics Interaction", EditAnywhere, BlueprintReadWrite, meta=(editcondition = "bEnablePhysicsInteraction"))
 	float PushForceFactor;
 
 	/** Z-Offset for the position the force is applied to. 0.0f is the center of the physics object, 1.0f is the top and -1.0f is the bottom of the object. */
-	UPROPERTY(Category="Agent Physics", EditAnywhere, BlueprintReadWrite, meta=(UIMin = "-1.0", UIMax = "1.0"), meta=(editcondition = "bEnablePhysicsInteraction"))
+	UPROPERTY(Category="Character Movement: Physics Interaction", EditAnywhere, BlueprintReadWrite, meta=(UIMin = "-1.0", UIMax = "1.0"), meta=(editcondition = "bEnablePhysicsInteraction"))
 	float PushForcePointZOffsetFactor;
 
 	/** Force to apply to physics objects that are touched by the player. */
-	UPROPERTY(Category="Agent Physics", EditAnywhere, BlueprintReadWrite, meta=(editcondition = "bEnablePhysicsInteraction"))
+	UPROPERTY(Category="Character Movement: Physics Interaction", EditAnywhere, BlueprintReadWrite, meta=(editcondition = "bEnablePhysicsInteraction"))
 	float TouchForceFactor;
 
 	/** Minimum Force applied to touched physics objects. If < 0.0f, there is no minimum. */
-	UPROPERTY(Category="Agent Physics", EditAnywhere, BlueprintReadWrite, meta=(editcondition = "bEnablePhysicsInteraction"))
+	UPROPERTY(Category="Character Movement: Physics Interaction", EditAnywhere, BlueprintReadWrite, meta=(editcondition = "bEnablePhysicsInteraction"))
 	float MinTouchForce;
 
 	/** Maximum force applied to touched physics objects. If < 0.0f, there is no maximum. */
-	UPROPERTY(Category="Agent Physics", EditAnywhere, BlueprintReadWrite, meta=(editcondition = "bEnablePhysicsInteraction"))
+	UPROPERTY(Category="Character Movement: Physics Interaction", EditAnywhere, BlueprintReadWrite, meta=(editcondition = "bEnablePhysicsInteraction"))
 	float MaxTouchForce;
 
 	/** Force per kg applied constantly to all overlapping components. */
-	UPROPERTY(Category="Agent Physics", EditAnywhere, BlueprintReadWrite, meta=(editcondition = "bEnablePhysicsInteraction"))
+	UPROPERTY(Category="Character Movement: Physics Interaction", EditAnywhere, BlueprintReadWrite, meta=(editcondition = "bEnablePhysicsInteraction"))
 	float RepulsionForce;
 
 
@@ -503,7 +512,7 @@ public:
 	 * WARNING: if (MaxSimulationTimeStep * MaxSimulationIterations) is too low for the min framerate, the last simulation step may exceed MaxSimulationTimeStep to complete the simulation.
 	 * @see MaxSimulationIterations
 	 */
-	UPROPERTY(Category="Character Movement", EditAnywhere, BlueprintReadWrite, AdvancedDisplay, meta=(ClampMin="0.0166", ClampMax="0.50", UIMin="0.0166", UIMax="0.50"))
+	UPROPERTY(Category="Character Movement (General Settings)", EditAnywhere, BlueprintReadWrite, AdvancedDisplay, meta=(ClampMin="0.0166", ClampMax="0.50", UIMin="0.0166", UIMax="0.50"))
 	float MaxSimulationTimeStep;
 
 	/**
@@ -514,19 +523,19 @@ public:
 	 * WARNING: if (MaxSimulationTimeStep * MaxSimulationIterations) is too low for the min framerate, the last simulation step may exceed MaxSimulationTimeStep to complete the simulation.
 	 * @see MaxSimulationTimeStep
 	 */
-	UPROPERTY(Category="Character Movement", EditAnywhere, BlueprintReadWrite, AdvancedDisplay, meta=(ClampMin="1", ClampMax="25", UIMin="1", UIMax="25"))
+	UPROPERTY(Category="Character Movement (General Settings)", EditAnywhere, BlueprintReadWrite, AdvancedDisplay, meta=(ClampMin="1", ClampMax="25", UIMin="1", UIMax="25"))
 	int32 MaxSimulationIterations;
 
 	/** Used in determining if pawn is going off ledge.  If the ledge is "shorter" than this value then the pawn will be able to walk off it. **/
-	UPROPERTY(Category="Character Movement", EditAnywhere, BlueprintReadWrite, AdvancedDisplay)
+	UPROPERTY(Category="Character Movement: Walking", EditAnywhere, BlueprintReadWrite, AdvancedDisplay)
 	float LedgeCheckThreshold;
 
 	/** When exiting water, jump if control pitch angle is this high or above. */
-	UPROPERTY(Category="Character Movement", EditAnywhere, BlueprintReadWrite, AdvancedDisplay)
+	UPROPERTY(Category="Character Movement: Swimming", EditAnywhere, BlueprintReadWrite, AdvancedDisplay)
 	float JumpOutOfWaterPitch;
 
 	/** Information about the floor the Character is standing on (updated only during walking movement). */
-	UPROPERTY(Category="Character Movement", VisibleInstanceOnly, BlueprintReadOnly)
+	UPROPERTY(Category="Character Movement: Walking", VisibleInstanceOnly, BlueprintReadOnly)
 	FFindFloorResult CurrentFloor;
 
 	/**
@@ -534,7 +543,7 @@ public:
 	 * @see DefaultWaterMovementMode
 	 * @see bRunPhysicsWithNoController
 	 */
-	UPROPERTY(Category="Character Movement", EditAnywhere, BlueprintReadWrite)
+	UPROPERTY(Category="Character Movement (General Settings)", EditAnywhere, BlueprintReadWrite)
 	TEnumAsByte<enum EMovementMode> DefaultLandMovementMode;
 
 	/**
@@ -542,7 +551,7 @@ public:
 	 * @see DefaultLandMovementMode
 	 * @see bRunPhysicsWithNoController
 	 */
-	UPROPERTY(Category="Character Movement", EditAnywhere, BlueprintReadWrite)
+	UPROPERTY(Category="Character Movement (General Settings)", EditAnywhere, BlueprintReadWrite)
 	TEnumAsByte<enum EMovementMode> DefaultWaterMovementMode;
 
 	/** Ground movement mode to switch to after falling */
@@ -553,19 +562,19 @@ public:
 	 * If true, walking movement always maintains horizontal velocity when moving up ramps, which causes movement up ramps to be faster parallel to the ramp surface.
 	 * If false, then walking movement maintains velocity magnitude parallel to the ramp surface.
 	 */
-	UPROPERTY(Category="Character Movement", EditAnywhere, BlueprintReadWrite)
+	UPROPERTY(Category="Character Movement: Walking", EditAnywhere, BlueprintReadWrite)
 	uint32 bMaintainHorizontalGroundVelocity:1;
 
 	/** If true, impart the base actor's X velocity when falling off it (which includes jumping) */
-	UPROPERTY(Category="Character Movement", EditAnywhere, BlueprintReadWrite)
+	UPROPERTY(Category="Character Movement: Jumping / Falling", EditAnywhere, BlueprintReadWrite)
 	uint32 bImpartBaseVelocityX:1;
 
 	/** If true, impart the base actor's Y velocity when falling off it (which includes jumping) */
-	UPROPERTY(Category="Character Movement", EditAnywhere, BlueprintReadWrite)
+	UPROPERTY(Category="Character Movement: Jumping / Falling", EditAnywhere, BlueprintReadWrite)
 	uint32 bImpartBaseVelocityY:1;
 
 	/** If true, impart the base actor's Z velocity when falling off it (which includes jumping) */
-	UPROPERTY(Category="Character Movement", EditAnywhere, BlueprintReadWrite)
+	UPROPERTY(Category="Character Movement: Jumping / Falling", EditAnywhere, BlueprintReadWrite)
 	uint32 bImpartBaseVelocityZ:1;
 
 	/**
@@ -573,11 +582,11 @@ public:
 	 * Only those components of the velocity allowed by the separate component settings (bImpartBaseVelocityX etc) will be applied.
 	 * @see bImpartBaseVelocityX, bImpartBaseVelocityY, bImpartBaseVelocityZ
 	 */
-	UPROPERTY(Category="Character Movement", EditAnywhere, BlueprintReadWrite)
+	UPROPERTY(Category="Character Movement: Jumping / Falling", EditAnywhere, BlueprintReadWrite)
 	uint32 bImpartBaseAngularVelocity:1;
 
 	/** Used by movement code to determine if a change in position is based on normal movement or a teleport. If not a teleport, velocity can be recomputed based on the change in position. */
-	UPROPERTY(Category="Character Movement", Transient, VisibleInstanceOnly, BlueprintReadWrite)
+	UPROPERTY(Category="Character Movement (General Settings)", Transient, VisibleInstanceOnly, BlueprintReadWrite)
 	uint32 bJustTeleported:1;
 
 	/** True when a network replication update is received for simulated proxies. */
@@ -589,7 +598,7 @@ public:
 	uint32 bNetworkMovementModeChanged:1;
 
 	/** if true, event NotifyJumpApex() to CharacterOwner's controller when at apex of jump.  Is cleared when event is triggered. */
-	UPROPERTY(Category = "Character Movement", EditAnywhere, BlueprintReadWrite)
+	UPROPERTY(Category="Character Movement: Jumping / Falling", EditAnywhere, BlueprintReadWrite)
 	uint32 bNotifyApex:1;
 
 	/** Instantly stop when in flying mode and no acceleration is being applied. */
@@ -597,7 +606,7 @@ public:
 	uint32 bCheatFlying:1;
 
 	/** If true, try to crouch (or keep crouching) on next update. If false, try to stop crouching on next update. */
-	UPROPERTY(Category="Character Movement", VisibleInstanceOnly, BlueprintReadOnly)
+	UPROPERTY(Category="Character Movement (General Settings)", VisibleInstanceOnly, BlueprintReadOnly)
 	uint32 bWantsToCrouch:1;
 
 	/**
@@ -605,11 +614,15 @@ public:
 	 * The same behavior applies when the character uncrouches: if true, the base is kept in the same location and the center moves up. If false, the capsule grows and only moves up if the base impacts something.
 	 * By default this variable is set when the movement mode changes: set to true when walking and false otherwise. Feel free to override the behavior when the movement mode changes.
 	 */
-	UPROPERTY(Category="Character Movement", VisibleInstanceOnly, BlueprintReadWrite, AdvancedDisplay)
+	UPROPERTY(Category="Character Movement (General Settings)", VisibleInstanceOnly, BlueprintReadWrite, AdvancedDisplay)
 	uint32 bCrouchMaintainsBaseLocation:1;
 
-	/** If true, the pawn ignores the effects of changes in its base's rotation on its rotation. */
-	UPROPERTY(Category="Character Movement", EditAnywhere, BlueprintReadWrite)
+	/**
+	 * Whether the character ignores changes in rotation of the base it is standing on.
+	 * If true, the character maintains current world rotation.
+	 * If false, the character rotates with the moving base.
+	 */
+	UPROPERTY(Category="Character Movement: Walking", EditAnywhere, BlueprintReadWrite)
 	uint32 bIgnoreBaseRotation:1;
 
 	/** 
@@ -620,17 +633,18 @@ public:
 	uint32 bFastAttachedMove:1;
 
 	/**
-	 * True to always force floor checks for stationary Characters.
-	 * Normally floor checks are avoided if possible when not moving, but this can be used to force them if there are use-cases where they are being skipped erroneously.
+	 * Whether we always force floor checks for stationary Characters while walking.
+	 * Normally floor checks are avoided if possible when not moving, but this can be used to force them if there are use-cases where they are being skipped erroneously
+	 * (such as objects moving up into the character from below).
 	 */
-	UPROPERTY(Category="Character Movement", EditAnywhere, BlueprintReadWrite, AdvancedDisplay)
+	UPROPERTY(Category="Character Movement: Walking", EditAnywhere, BlueprintReadWrite, AdvancedDisplay)
 	uint32 bAlwaysCheckFloor:1;
 
 	/**
 	 * Performs floor checks as if the character is using a shape with a flat base.
 	 * This avoids the situation where characters slowly lower off the side of a ledge (as their capsule 'balances' on the edge).
 	 */
-	UPROPERTY(Category="Character Movement", EditAnywhere, BlueprintReadWrite, AdvancedDisplay)
+	UPROPERTY(Category="Character Movement: Walking", EditAnywhere, BlueprintReadWrite, AdvancedDisplay)
 	uint32 bUseFlatBaseForFloorChecks:1;
 
 	/** Used to prevent reentry of JumpOff() */
@@ -642,7 +656,7 @@ public:
 	uint32 bWantsToLeaveNavWalking:1;
 
 	/** If set, component will use RVO avoidance */
-	UPROPERTY(Category="Avoidance", EditAnywhere, BlueprintReadOnly)
+	UPROPERTY(Category="Character Movement: Avoidance", EditAnywhere, BlueprintReadOnly)
 	uint32 bUseRVOAvoidance:1;
 
 	/**
@@ -650,7 +664,7 @@ public:
 	 * If true, acceleration is applied when path following to reach the target velocity.
 	 * If false, path following velocity is set directly, disregarding acceleration.
 	 */
-	UPROPERTY(Category="Character Movement", EditAnywhere, BlueprintReadWrite, AdvancedDisplay)
+	UPROPERTY(Category="Character Movement (General Settings)", EditAnywhere, BlueprintReadWrite, AdvancedDisplay)
 	uint32 bRequestedMoveUseAcceleration:1;
 
 protected:
@@ -676,7 +690,7 @@ protected:
 	uint32 bDeferUpdateBasedMovement : 1;
 
 	/** Whether to raycast to underlying geometry to better conform navmesh-walking characters */
-	UPROPERTY(Category = "NavMesh Movement", EditAnywhere, BlueprintReadOnly)
+	UPROPERTY(Category="Character Movement: NavMesh Movement", EditAnywhere, BlueprintReadOnly)
 	uint32 bProjectNavMeshWalking : 1;
 
 	/** forced avoidance velocity, used when AvoidanceLockTimer is > 0 */
@@ -687,7 +701,7 @@ protected:
 
 public:
 
-	UPROPERTY(Category = "Avoidance", EditAnywhere, BlueprintReadOnly)
+	UPROPERTY(Category="Character Movement: Avoidance", EditAnywhere, BlueprintReadOnly)
 	float AvoidanceConsiderationRadius;
 
 	/**
@@ -698,29 +712,29 @@ public:
 	FVector RequestedVelocity;
 
 	/** No default value, for now it's assumed to be valid if GetAvoidanceManager() returns non-NULL. */
-	UPROPERTY(Category="Avoidance", VisibleAnywhere, BlueprintReadOnly, AdvancedDisplay)
+	UPROPERTY(Category="Character Movement: Avoidance", VisibleAnywhere, BlueprintReadOnly, AdvancedDisplay)
 	int32 AvoidanceUID;
 
 	/** Moving actor's group mask */
-	UPROPERTY(Category="Avoidance", EditAnywhere, BlueprintReadOnly, AdvancedDisplay)
+	UPROPERTY(Category="Character Movement: Avoidance", EditAnywhere, BlueprintReadOnly, AdvancedDisplay)
 	FNavAvoidanceMask AvoidanceGroup;
 	UFUNCTION(BlueprintCallable, Category="Pawn|Components|CharacterMovement")
 	void SetAvoidanceGroup(int32 GroupFlags);
 
 	/** Will avoid other agents if they are in one of specified groups */
-	UPROPERTY(Category="Avoidance", EditAnywhere, BlueprintReadOnly, AdvancedDisplay)
+	UPROPERTY(Category="Character Movement: Avoidance", EditAnywhere, BlueprintReadOnly, AdvancedDisplay)
 	FNavAvoidanceMask GroupsToAvoid;
 	UFUNCTION(BlueprintCallable, Category="Pawn|Components|CharacterMovement")
 	void SetGroupsToAvoid(int32 GroupFlags);
 
 	/** Will NOT avoid other agents if they are in one of specified groups, higher priority than GroupsToAvoid */
-	UPROPERTY(Category="Avoidance", EditAnywhere, BlueprintReadOnly, AdvancedDisplay)
+	UPROPERTY(Category="Character Movement: Avoidance", EditAnywhere, BlueprintReadOnly, AdvancedDisplay)
 	FNavAvoidanceMask GroupsToIgnore;
 	UFUNCTION(BlueprintCallable, Category="Pawn|Components|CharacterMovement")
 	void SetGroupsToIgnore(int32 GroupFlags);
 
 	/** De facto default value 0.5 (due to that being the default in the avoidance registration function), indicates RVO behavior. */
-	UPROPERTY(Category="Avoidance", EditAnywhere, BlueprintReadOnly)
+	UPROPERTY(Category="Character Movement: Avoidance", EditAnywhere, BlueprintReadOnly)
 	float AvoidanceWeight;
 
 	/** Temporarily holds launch velocity when pawn is to be launched so it happens at end of movement. */
@@ -734,32 +748,32 @@ public:
 	FHitResult CachedProjectedNavMeshHitResult;
 
 	/** How often we should raycast to project from navmesh to underlying geometry */
-	UPROPERTY(Category = "NavMesh Movement", EditAnywhere, BlueprintReadWrite, meta=(editcondition = "bProjectNavMeshWalking"))
+	UPROPERTY(Category="Character Movement: NavMesh Movement", EditAnywhere, BlueprintReadWrite, meta=(editcondition = "bProjectNavMeshWalking"))
 	float NavMeshProjectionInterval;
 
 	UPROPERTY(Transient)
 	float NavMeshProjectionTimer;
 
 	/** Speed at which to interpolate agent navmesh offset between traces. 0: Instant (no interp) > 0: Interp speed") */
-	UPROPERTY(Category = "NavMesh Movement", EditAnywhere, BlueprintReadWrite, meta=(editcondition = "bProjectNavMeshWalking", ClampMin="0", UIMin="0"))
+	UPROPERTY(Category="Character Movement: NavMesh Movement", EditAnywhere, BlueprintReadWrite, meta=(editcondition = "bProjectNavMeshWalking", ClampMin="0", UIMin="0"))
 	float NavMeshProjectionInterpSpeed;
 
 	/**
 	 * Scale of the total capsule height to use for projection from navmesh to underlying geometry in the upward direction.
 	 * In other words, start the trace at [CapsuleHeight * NavMeshProjectionHeightScaleUp] above nav mesh.
 	 */
-	UPROPERTY(Category = "NavMesh Movement", EditAnywhere, BlueprintReadWrite, meta=(editcondition = "bProjectNavMeshWalking", ClampMin="0", UIMin="0"))
+	UPROPERTY(Category="Character Movement: NavMesh Movement", EditAnywhere, BlueprintReadWrite, meta=(editcondition = "bProjectNavMeshWalking", ClampMin="0", UIMin="0"))
 	float NavMeshProjectionHeightScaleUp;
 
 	/**
 	 * Scale of the total capsule height to use for projection from navmesh to underlying geometry in the downward direction.
 	 * In other words, trace down to [CapsuleHeight * NavMeshProjectionHeightScaleDown] below nav mesh.
 	 */
-	UPROPERTY(Category = "NavMesh Movement", EditAnywhere, BlueprintReadWrite, meta=(editcondition = "bProjectNavMeshWalking", ClampMin="0", UIMin="0"))
+	UPROPERTY(Category="Character Movement: NavMesh Movement", EditAnywhere, BlueprintReadWrite, meta=(editcondition = "bProjectNavMeshWalking", ClampMin="0", UIMin="0"))
 	float NavMeshProjectionHeightScaleDown;
 
 	/** Change avoidance state and registers in RVO manager if needed */
-	UFUNCTION(BlueprintCallable, Category = "Pawn|Components|CharacterMovement", meta = (UnsafeDuringActorConstruction = "true"))
+	UFUNCTION(BlueprintCallable, Category="Pawn|Components|CharacterMovement", meta = (UnsafeDuringActorConstruction = "true"))
 	void SetAvoidanceEnabled(bool bEnable);
 
 	/** Get the Character that owns UpdatedComponent. */
diff --git a/Engine/Source/Runtime/Engine/Classes/GameFramework/NavMovementComponent.h b/Engine/Source/Runtime/Engine/Classes/GameFramework/NavMovementComponent.h
--- a/Engine/Source/Runtime/Engine/Classes/GameFramework/NavMovementComponent.h
+++ b/Engine/Source/Runtime/Engine/Classes/GameFramework/NavMovementComponent.h
@@ -25,7 +25,7 @@ class ENGINE_API UNavMovementComponent : public UMovementComponent
 	GENERATED_UCLASS_BODY()
 
 	/** Properties that define how the component can move. */
-	UPROPERTY(EditAnywhere, BlueprintReadWrite, Category=MovementComponent, meta=(DisplayName="Movement Capabilities", Keywords="Nav Agent"))
+	UPROPERTY(EditAnywhere, BlueprintReadWrite, Category="Movement Capabilities", meta=(DisplayName="Movement Capabilities", Keywords="Nav Agent"))
 	FNavAgentProperties NavAgentProps;
 
 protected:
