Fix minor misleading comments/parameter naming (copy-paste error)

diff --git a/Engine/Source/Runtime/Engine/Classes/GameFramework/Actor.h b/Engine/Source/Runtime/Engine/Classes/GameFramework/Actor.h
--- a/Engine/Source/Runtime/Engine/Classes/GameFramework/Actor.h
+++ b/Engine/Source/Runtime/Engine/Classes/GameFramework/Actor.h
@@ -299,8 +299,8 @@ public:
 	/** Called on the actor when a new subobject is dynamically created via replication */
 	virtual void OnSubobjectCreatedFromReplication(UObject *NewSubobject);
 
-	/** Called on the actor when a new subobject is dynamically created via replication */
-	virtual void OnSubobjectDestroyFromReplication(UObject *NewSubobject);
+	/** Called on the actor when a subobject is dynamically destroyed via replication */
+	virtual void OnSubobjectDestroyFromReplication(UObject *Subobject);
 
 	/** Called on the actor right before replication occurs */
 	virtual void PreReplication( IRepChangedPropertyTracker & ChangedPropertyTracker );
diff --git a/Engine/Source/Runtime/Engine/Private/ActorReplication.cpp b/Engine/Source/Runtime/Engine/Private/ActorReplication.cpp
--- a/Engine/Source/Runtime/Engine/Private/ActorReplication.cpp
+++ b/Engine/Source/Runtime/Engine/Private/ActorReplication.cpp
@@ -301,11 +301,11 @@ void AActor::OnSubobjectCreatedFromReplication(UObject *NewSubobject)
 	}
 }
 
-/** Called on the actor when a new subobject is dynamically created via replication */
-void AActor::OnSubobjectDestroyFromReplication(UObject *NewSubobject)
+/** Called on the actor when a subobject is dynamically destroyed via replication */
+void AActor::OnSubobjectDestroyFromReplication(UObject *Subobject)
 {
-	check(NewSubobject);
-	if ( UActorComponent * Component = Cast<UActorComponent>(NewSubobject) )
+	check(Subobject);
+	if ( UActorComponent * Component = Cast<UActorComponent>(Subobject) )
 	{
 		Component->DestroyComponent();
 	}
