Allow Blueprint Components to call Destroy on themselves

diff --git a/Engine/Source/Runtime/Engine/Private/Components/ActorComponent.cpp b/Engine/Source/Runtime/Engine/Private/Components/ActorComponent.cpp
--- a/Engine/Source/Runtime/Engine/Private/Components/ActorComponent.cpp
+++ b/Engine/Source/Runtime/Engine/Private/Components/ActorComponent.cpp
@@ -938,7 +938,7 @@ void UActorComponent::OnComponentDestroyed()
 
 void UActorComponent::K2_DestroyComponent(UObject* Object)
 {
-	AActor* Owner = GetOwner();
+	AActor* Owner = (Object != this ? GetOwner() : nullptr);
 	if (Owner == NULL || Owner == Object)
 	{
 		DestroyComponent();
