Fixed FStaticMeshInstanceBuffer in game having NumInstances == 0, which broke distance field lighting and mobile instanced emulation

diff --git a/Engine/Source/Runtime/Engine/Private/InstancedStaticMesh.cpp b/Engine/Source/Runtime/Engine/Private/InstancedStaticMesh.cpp
--- a/Engine/Source/Runtime/Engine/Private/InstancedStaticMesh.cpp
+++ b/Engine/Source/Runtime/Engine/Private/InstancedStaticMesh.cpp
@@ -117,7 +117,9 @@ void FStaticMeshInstanceBuffer::Init(UInstancedStaticMeshComponent* InComponent,
 void FStaticMeshInstanceBuffer::InitFromPreallocatedData(UInstancedStaticMeshComponent* InComponent, FStaticMeshInstanceData& Other)
 {
 	QUICK_SCOPE_CYCLE_COUNTER(STAT_FStaticMeshInstanceBuffer_InitFromPreallocatedData);
+	const uint32 NewNumInstances = Other.Num();
 	AllocateData(Other);
+	NumInstances = NewNumInstances;
 	SetupCPUAccess(InComponent);
 	InstanceData->SetAllowCPUAccess(true);		// GDC demo hack!
 }
