Add InitStaticBodies for instanced static mesh component

diff --git a/Engine/Source/Runtime/Engine/Classes/PhysicsEngine/BodyInstance.h b/Engine/Source/Runtime/Engine/Classes/PhysicsEngine/BodyInstance.h
--- a/Engine/Source/Runtime/Engine/Classes/PhysicsEngine/BodyInstance.h
+++ b/Engine/Source/Runtime/Engine/Classes/PhysicsEngine/BodyInstance.h
@@ -61,6 +61,7 @@ namespace EDOFMode
 }
 
 struct FCollisionNotifyInfo;
+template <bool bCompileStatic> struct FInitBodiesHelper;
 
 USTRUCT()
 struct ENGINE_API FCollisionResponse
@@ -436,14 +437,15 @@ public:
 	 */
 	static bool ValidateTransform(const FTransform &Transform, const FString& DebugName, const UBodySetup* Setup);
 
-	/** Standalone path to batch initialise large amounts of static bodies, which will be deferred till the next scene update for fast scene addition.
+	/** Standalone path to batch initialize large amounts of static bodies, which will be deferred till the next scene update for fast scene addition.
 	 *	@param Bodies
 	 *	@param Transforms
 	 *	@param BodySetup
 	 *	@param PrimitiveComp
 	 *	@param InRBScene
+	 *  @param PhysicsSerializer
 	 */
-	static void InitStaticBodies(TArray<FBodyInstance*>& Bodies, TArray<FTransform>& Transforms, class UBodySetup* BodySetup, class UPrimitiveComponent* PrimitiveComp, class FPhysScene* InRBScene);
+	static void InitStaticBodies(TArray<FBodyInstance*>& Bodies, TArray<FTransform>& Transforms, class UBodySetup* BodySetup, class UPrimitiveComponent* PrimitiveComp, class FPhysScene* InRBScene, UPhysicsSerializer* PhysicsSerializer);
 
 	/** Obtains the appropriate PhysX scene lock for READING and executes the passed in lambda. */
 	void ExecuteOnPhysicsReadOnly(TFunctionRef<void()> Func) const;
@@ -1019,7 +1021,9 @@ private:
 
 	friend class UCollisionProfile;
 	friend class FBodyInstanceCustomization;
-	friend struct FInitBodiesHelper;
+	
+	friend struct FInitBodiesHelper<true>;
+	friend struct FInitBodiesHelper<false>;
 	friend class FDerivedDataPhysXBinarySerializer;
 
 #if WITH_BOX2D
diff --git a/Engine/Source/Runtime/Engine/Private/InstancedStaticMesh.cpp b/Engine/Source/Runtime/Engine/Private/InstancedStaticMesh.cpp
--- a/Engine/Source/Runtime/Engine/Private/InstancedStaticMesh.cpp
+++ b/Engine/Source/Runtime/Engine/Private/InstancedStaticMesh.cpp
@@ -984,10 +984,9 @@ void UInstancedStaticMeshComponent::CreateAllInstanceBodies()
 {
 	QUICK_SCOPE_CYCLE_COUNTER(STAT_UInstancedStaticMeshComponent_CreateAllInstanceBodies);
 
-	UBodySetup* BodySetup = GetBodySetup();
-	if (BodySetup)
+	if (UBodySetup* BodySetup = GetBodySetup())
 	{
-	    int32 NumBodies = PerInstanceSMData.Num();
+	    const int32 NumBodies = PerInstanceSMData.Num();
 	    InstanceBodies.SetNumUninitialized(NumBodies);
     
 	    TArray<FTransform> Transforms;
@@ -1028,7 +1027,14 @@ void UInstancedStaticMeshComponent::CreateAllInstanceBodies()
 
 			PhysicsSerializer->CreatePhysicsData(BodySetups, PhysicalMaterials);
 			
-			FBodyInstance::InitBodies(InstanceBodies, Transforms, BodySetup, this, GetWorld()->GetPhysicsScene(), nullptr, true, PhysicsSerializer);
+			if(Mobility == EComponentMobility::Static)
+			{
+				FBodyInstance::InitStaticBodies(InstanceBodies, Transforms, BodySetup, this, GetWorld()->GetPhysicsScene(), PhysicsSerializer);
+			}else
+			{
+				FBodyInstance::InitBodies(InstanceBodies, Transforms, BodySetup, this, GetWorld()->GetPhysicsScene(), nullptr, true, PhysicsSerializer);
+			}
+			
 
 			//Serialize physics data for fast path cooking
 			PhysicsSerializer->SerializePhysics(InstanceBodies, BodySetups, PhysicalMaterials);
diff --git a/Engine/Source/Runtime/Engine/Private/PhysicsEngine/BodyInstance.cpp b/Engine/Source/Runtime/Engine/Private/PhysicsEngine/BodyInstance.cpp
--- a/Engine/Source/Runtime/Engine/Private/PhysicsEngine/BodyInstance.cpp
+++ b/Engine/Source/Runtime/Engine/Private/PhysicsEngine/BodyInstance.cpp
@@ -1059,6 +1059,7 @@ FVector GetInitialLinearVelocity(const AActor* OwningActor, bool& bComponentAwak
 	return InitialLinVel;
 }
 
+template <bool bCompileStatic>
 struct FInitBodiesHelper
 {
 	FInitBodiesHelper(const TArray<FBodyInstance*>& InBodies, const TArray<FTransform>& InTransforms, class UBodySetup* InBodySetup, class UPrimitiveComponent* InPrimitiveComp, class FPhysScene* InInRBScene, FBodyInstance::PhysXAggregateType InInAggregate = NULL, bool InDefer = false, UPhysicsSerializer* InPhysicsSerializer = nullptr)
@@ -1081,8 +1082,8 @@ struct FInitBodiesHelper
 
 		PhysXName = GetDebugDebugName(PrimitiveComp, BodySetup, DebugName);
 
-		bStatic = PrimitiveComp == nullptr || PrimitiveComp->Mobility != EComponentMobility::Movable;
-		SkelMeshComp = GetSkeletalMeshComponentAndProperties(PrimitiveComp, BodySetup, InstanceBlendWeight, bInstanceSimulatePhysics, bComponentAwake);
+		bStatic = bCompileStatic || PrimitiveComp == nullptr || PrimitiveComp->Mobility != EComponentMobility::Movable;
+		SkelMeshComp = bCompileStatic ? nullptr : GetSkeletalMeshComponentAndProperties(PrimitiveComp, BodySetup, InstanceBlendWeight, bInstanceSimulatePhysics, bComponentAwake);
 		
 		const AActor* OwningActor = PrimitiveComp ? PrimitiveComp->GetOwner() : nullptr;
 		InitialLinVel = GetInitialLinearVelocity(OwningActor, bComponentAwake);
@@ -1121,556 +1122,546 @@ struct FInitBodiesHelper
 	PxScene* PSyncScene;
 	PxScene* PAsyncScene;
 
-	void InitBodies_PhysX() const;
-	bool CreateShapesAndActors_PhysX(TArray<PxActor*>& PSyncActors, TArray<PxActor*>& PAsyncActors, TArray<PxActor*>& PDynamicActors, const bool bCanDefer, bool& bDynamicsUseAsyncScene) const;
-	physx::PxRigidActor* CreateActor_PhysX_AssumesLocked(FBodyInstance* Instance, const PxTransform& PTransform) const;
-	bool GetBinaryData_PhysX_AssumesLocked(FBodyInstance* Instance) const;
-	bool CreateShapes_PhysX_AssumesLocked(FBodyInstance* Instance, physx::PxRigidActor* PNewDynamic) const;
-	void AddActorsToScene_PhysX_AssumesLocked(TArray<PxActor*>& PSyncActors, TArray<PxActor*>& PAsyncActors, TArray<PxActor*>& PDynamicActors) const;
-#endif
-
-#if WITH_BOX2D
-	void InitBodies_Box2D() const;
-#endif
-};
-
-void FBodyInstance::InitBodies(TArray<FBodyInstance*>& Bodies, TArray<FTransform>& Transforms, class UBodySetup* BodySetup, class UPrimitiveComponent* PrimitiveComp, class FPhysScene* InRBScene, PhysXAggregateType InAggregate /*= NULL*/, bool bDefer /*= false*/, UPhysicsSerializer* PhysicsSerializer )
-{
-	SCOPE_CYCLE_COUNTER(STAT_InitBodies);
-
-	check(BodySetup);
-	check(InRBScene);
-	check(Bodies.Num() > 0);
-
-	FInitBodiesHelper InitBodiesHelper(Bodies, Transforms, BodySetup, PrimitiveComp, InRBScene, InAggregate, bDefer, PhysicsSerializer);
-	
-#if WITH_PHYSX
-	InitBodiesHelper.InitBodies_PhysX();
-#endif
-
-#if WITH_BOX2D
-	InitBodiesHelper.InitBodies_Box2D();
-#endif
-}
-
-#if WITH_PHYSX
-
-bool FInitBodiesHelper::GetBinaryData_PhysX_AssumesLocked(FBodyInstance* Instance) const
-{
-	Instance->RigidActorSync = PhysicsSerializer->GetRigidActor(Instance->RigidActorSyncId);
-	Instance->RigidActorAsync = PhysicsSerializer->GetRigidActor(Instance->RigidActorAsyncId);
-	
-	return Instance->RigidActorSync || Instance->RigidActorAsync;
-}
+	bool GetBinaryData_PhysX_AssumesLocked(FBodyInstance* Instance) const
+	{
+		Instance->RigidActorSync = PhysicsSerializer->GetRigidActor(Instance->RigidActorSyncId);
+		Instance->RigidActorAsync = PhysicsSerializer->GetRigidActor(Instance->RigidActorAsyncId);
 
+		return Instance->RigidActorSync || Instance->RigidActorAsync;
+	}
 
-physx::PxRigidActor* FInitBodiesHelper::CreateActor_PhysX_AssumesLocked(FBodyInstance* Instance, const PxTransform& PTransform) const
-{
-	physx::PxRigidActor* PNewDynamic = nullptr;
 
-	if (bStatic)
-	{
-		Instance->RigidActorSync = GPhysXSDK->createRigidStatic(PTransform);
-		if (PAsyncScene)
-		{
-			Instance->RigidActorAsync = GPhysXSDK->createRigidStatic(PTransform);
-		}
-	}
-	else
+	physx::PxRigidActor* CreateActor_PhysX_AssumesLocked(FBodyInstance* Instance, const PxTransform& PTransform) const
 	{
-		PNewDynamic = GPhysXSDK->createRigidDynamic(PTransform);
-		if (Instance->UseAsyncScene(PhysScene))
+		physx::PxRigidActor* PNewDynamic = nullptr;
+
+		if (bCompileStatic || bStatic)
 		{
-			Instance->RigidActorAsync = PNewDynamic;
+			Instance->RigidActorSync = GPhysXSDK->createRigidStatic(PTransform);
+			if (PAsyncScene)
+			{
+				Instance->RigidActorAsync = GPhysXSDK->createRigidStatic(PTransform);
+			}
 		}
 		else
 		{
-			Instance->RigidActorSync = PNewDynamic;
+			PNewDynamic = GPhysXSDK->createRigidDynamic(PTransform);
+			if (Instance->UseAsyncScene(PhysScene))
+			{
+				Instance->RigidActorAsync = PNewDynamic;
+			}
+			else
+			{
+				Instance->RigidActorSync = PNewDynamic;
+			}
 		}
-	}
 
 		return PNewDynamic;
-}
-
-bool FInitBodiesHelper::CreateShapes_PhysX_AssumesLocked(FBodyInstance* Instance, physx::PxRigidActor* PNewDynamic) const
-{
-	UPhysicalMaterial* SimplePhysMat = Instance->GetSimplePhysicalMaterial();
-	TArray<UPhysicalMaterial*> ComplexPhysMats = Instance->GetComplexPhysicalMaterials();
-	PxMaterial* PSimpleMat = SimplePhysMat->GetPhysXMaterial();
-
-	FShapeData ShapeData;
-	Instance->GetFilterData_AssumesLocked(ShapeData);
-	Instance->GetShapeFlags_AssumesLocked(ShapeData, Instance->CollisionEnabled, BodySetup->CollisionTraceFlag == CTF_UseComplexAsSimple);
-
-	if (PNewDynamic)
-	{
-		if (!Instance->ShouldInstanceSimulatingPhysics())
-		{
-			ShapeData.SyncBodyFlags |= PxRigidBodyFlag::eKINEMATIC;
-		}
-		ShapeData.SyncBodyFlags |= PxRigidBodyFlag::eUSE_KINEMATIC_TARGET_FOR_SCENE_QUERIES;
 	}
 
-	bool bInitFail = false;
-
-	const bool bShapeSharing = Instance->HasSharedShapes(); //If we have a static actor we can reuse the shapes between sync and async scene
-	TArray<PxShape*> PSharedShapes;
-	if (Instance->RigidActorSync)
+	bool CreateShapes_PhysX_AssumesLocked(FBodyInstance* Instance, physx::PxRigidActor* PNewDynamic) const
 	{
-		BodySetup->AddShapesToRigidActor_AssumesLocked(Instance, Instance->RigidActorSync, PST_Sync, Instance->Scale3D, PSimpleMat, ComplexPhysMats, ShapeData, FTransform::Identity, bShapeSharing ? &PSharedShapes : nullptr, bShapeSharing);
-		bInitFail |= Instance->RigidActorSync->getNbShapes() == 0;
-		Instance->RigidActorSync->userData = &Instance->PhysxUserData;
-		Instance->RigidActorSync->setName(Instance->CharDebugName.IsValid() ? Instance->CharDebugName->GetData() : nullptr);
-
-		check(FPhysxUserData::Get<FBodyInstance>(Instance->RigidActorSync->userData) == Instance && FPhysxUserData::Get<FBodyInstance>(Instance->RigidActorSync->userData)->OwnerComponent != NULL);
+		UPhysicalMaterial* SimplePhysMat = Instance->GetSimplePhysicalMaterial();
+		TArray<UPhysicalMaterial*> ComplexPhysMats = Instance->GetComplexPhysicalMaterials();
+		PxMaterial* PSimpleMat = SimplePhysMat->GetPhysXMaterial();
 
-	}
+		FShapeData ShapeData;
+		Instance->GetFilterData_AssumesLocked(ShapeData);
+		Instance->GetShapeFlags_AssumesLocked(ShapeData, Instance->CollisionEnabled, BodySetup->CollisionTraceFlag == CTF_UseComplexAsSimple);
 
-	if (Instance->RigidActorAsync)
-	{
-		check(PAsyncScene);
-		if(bShapeSharing)
+		if (!bCompileStatic && PNewDynamic)
 		{
-			for (PxShape* PShape : PSharedShapes)
+			if (!Instance->ShouldInstanceSimulatingPhysics())
 			{
-				Instance->RigidActorAsync->attachShape(*PShape);
+				ShapeData.SyncBodyFlags |= PxRigidBodyFlag::eKINEMATIC;
 			}
-		}else
-		{
-			BodySetup->AddShapesToRigidActor_AssumesLocked(Instance, Instance->RigidActorAsync, PST_Async, Instance->Scale3D, PSimpleMat, ComplexPhysMats, ShapeData);
+			ShapeData.SyncBodyFlags |= PxRigidBodyFlag::eUSE_KINEMATIC_TARGET_FOR_SCENE_QUERIES;
 		}
-		
-		bInitFail |= Instance->RigidActorAsync->getNbShapes() == 0;
-		Instance->RigidActorAsync->userData = &Instance->PhysxUserData;
-		Instance->RigidActorAsync->setName(Instance->CharDebugName.IsValid() ? Instance->CharDebugName->GetData() : nullptr);
-
-		check(FPhysxUserData::Get<FBodyInstance>(Instance->RigidActorAsync->userData) == Instance && FPhysxUserData::Get<FBodyInstance>(Instance->RigidActorAsync->userData)->OwnerComponent != NULL);
-	}
-
-	return bInitFail;
-}
 
-bool FInitBodiesHelper::CreateShapesAndActors_PhysX(TArray<PxActor*>& PSyncActors,TArray<PxActor*>& PAsyncActors,TArray<PxActor*>& PDynamicActors,const bool bCanDefer, bool& bDynamicsUseAsyncScene) const
-{
-	const int32 NumBodies = Bodies.Num();
-	PSyncActors.Reserve(NumBodies);
-
-	if (PAsyncScene)
-	{
-		// Only allocate this if we can have async bodies
-		PAsyncActors.Reserve(NumBodies);
-	}
+		bool bInitFail = false;
 
-	// Ensure we have the AggGeom inside the body setup so we can calculate the number of shapes
-	BodySetup->CreatePhysicsMeshes();
-	for (int32 BodyIdx = 0; BodyIdx < NumBodies; ++BodyIdx)
-	{
-		FBodyInstance* Instance = Bodies[BodyIdx];
-		const FTransform& Transform = Transforms[BodyIdx];
+		const bool bShapeSharing = Instance->HasSharedShapes(); //If we have a static actor we can reuse the shapes between sync and async scene
+		TArray<PxShape*> PSharedShapes;
+		if (Instance->RigidActorSync)
+		{
+			BodySetup->AddShapesToRigidActor_AssumesLocked(Instance, Instance->RigidActorSync, PST_Sync, Instance->Scale3D, PSimpleMat, ComplexPhysMats, ShapeData, FTransform::Identity, bShapeSharing ? &PSharedShapes : nullptr, bShapeSharing);
+			bInitFail |= Instance->RigidActorSync->getNbShapes() == 0;
+			Instance->RigidActorSync->userData = &Instance->PhysxUserData;
+			Instance->RigidActorSync->setName(Instance->CharDebugName.IsValid() ? Instance->CharDebugName->GetData() : nullptr);
 
-		FBodyInstance::ValidateTransform(Transform, DebugName, BodySetup);
+			check(FPhysxUserData::Get<FBodyInstance>(Instance->RigidActorSync->userData) == Instance && FPhysxUserData::Get<FBodyInstance>(Instance->RigidActorSync->userData)->OwnerComponent != NULL);
 
-		Instance->OwnerComponent = PrimitiveComp;
-		Instance->BodySetup = BodySetup;
-		Instance->Scale3D = Transform.GetScale3D();
-		Instance->CharDebugName = PhysXName;
-		Instance->bHasSharedShapes = bStatic && PhysScene->HasAsyncScene() && UPhysicsSettings::Get()->bEnableShapeSharing;
+		}
 
-		// Handle autowelding here to avoid extra work
-		if (Instance->bAutoWeld)
+		if (Instance->RigidActorAsync)
 		{
-			UPrimitiveComponent * ParentPrimComponent = PrimitiveComp ? Cast<UPrimitiveComponent>(PrimitiveComp->AttachParent) : NULL;
-
-			//if we have a parent we will now do the weld and exit any further initialization
-			if (ParentPrimComponent && PrimitiveComp->GetWorld() && PrimitiveComp->GetWorld()->IsGameWorld())
+			check(PAsyncScene);
+			if (bShapeSharing)
 			{
-				if (PrimitiveComp->WeldToImplementation(ParentPrimComponent, PrimitiveComp->AttachSocketName, false))	//welded new simulated body so initialization is done
+				for (PxShape* PShape : PSharedShapes)
 				{
-					return false;
+					Instance->RigidActorAsync->attachShape(*PShape);
 				}
 			}
-		}
+			else
+			{
+				BodySetup->AddShapesToRigidActor_AssumesLocked(Instance, Instance->RigidActorAsync, PST_Async, Instance->Scale3D, PSimpleMat, ComplexPhysMats, ShapeData);
+			}
 
-		// Don't process if we've already got a body
-		if (Instance->GetPxRigidActor_AssumesLocked())
-		{
-			Instance->OwnerComponent = nullptr;
-			Instance->BodySetup = nullptr;
+			bInitFail |= Instance->RigidActorAsync->getNbShapes() == 0;
+			Instance->RigidActorAsync->userData = &Instance->PhysxUserData;
+			Instance->RigidActorAsync->setName(Instance->CharDebugName.IsValid() ? Instance->CharDebugName->GetData() : nullptr);
 
-			continue;
+			check(FPhysxUserData::Get<FBodyInstance>(Instance->RigidActorAsync->userData) == Instance && FPhysxUserData::Get<FBodyInstance>(Instance->RigidActorAsync->userData)->OwnerComponent != NULL);
 		}
 
-		// Set sim parameters for bodies from skeletal mesh components
-		if (SkelMeshComp)
-		{
-			Instance->bSimulatePhysics = bInstanceSimulatePhysics;
-			if (InstanceBlendWeight != -1.0f)
-			{
-				Instance->PhysicsBlendWeight = InstanceBlendWeight;
-			}
-		}
+		return bInitFail;
+	}
 
-		Instance->PhysxUserData = FPhysxUserData(Instance);
-		
-		static FName PhysicsFormatName(FPlatformProperties::GetPhysicsFormat());
-		
-		physx::PxRigidActor* PNewDynamic = nullptr;
-		bool bFoundBinaryData = false;
-		if(PhysicsSerializer)
+	bool CreateShapesAndActors_PhysX(TArray<PxActor*>& PSyncActors, TArray<PxActor*>& PAsyncActors, TArray<PxActor*>& PDynamicActors, const bool bCanDefer, bool& bDynamicsUseAsyncScene) const
+	{
+		const int32 NumBodies = Bodies.Num();
+		PSyncActors.Reserve(NumBodies);
+
+		if (PAsyncScene)
 		{
-			bFoundBinaryData = GetBinaryData_PhysX_AssumesLocked(Instance);
+			// Only allocate this if we can have async bodies
+			PAsyncActors.Reserve(NumBodies);
 		}
-		
-		if(!bFoundBinaryData)
+
+		// Ensure we have the AggGeom inside the body setup so we can calculate the number of shapes
+		BodySetup->CreatePhysicsMeshes();
+		for (int32 BodyIdx = 0; BodyIdx < NumBodies; ++BodyIdx)
 		{
-			PNewDynamic = CreateActor_PhysX_AssumesLocked(Instance, U2PTransform(Transform));	
-			const bool bInitFail = CreateShapes_PhysX_AssumesLocked(Instance, PNewDynamic);
+			FBodyInstance* Instance = Bodies[BodyIdx];
+			const FTransform& Transform = Transforms[BodyIdx];
+
+			FBodyInstance::ValidateTransform(Transform, DebugName, BodySetup);
+
+			Instance->OwnerComponent = PrimitiveComp;
+			Instance->BodySetup = BodySetup;
+			Instance->Scale3D = Transform.GetScale3D();
+			Instance->CharDebugName = PhysXName;
+			Instance->bHasSharedShapes = bStatic && PhysScene->HasAsyncScene() && UPhysicsSettings::Get()->bEnableShapeSharing;
 
-			if (bInitFail)
+			// Handle autowelding here to avoid extra work
+			if (!bCompileStatic && Instance->bAutoWeld)
 			{
-				UE_LOG(LogPhysics, Log, TEXT("Init Instance %d of Primitive Component %s failed"), BodyIdx, *PrimitiveComp->GetName());
-				if (Instance->RigidActorSync)
-				{
-					Instance->RigidActorSync->release();
-					Instance->RigidActorSync = nullptr;
-				}
+				UPrimitiveComponent * ParentPrimComponent = PrimitiveComp ? Cast<UPrimitiveComponent>(PrimitiveComp->AttachParent) : NULL;
 
-				if (Instance->RigidActorAsync)
+				//if we have a parent we will now do the weld and exit any further initialization
+				if (ParentPrimComponent && PrimitiveComp->GetWorld() && PrimitiveComp->GetWorld()->IsGameWorld())
 				{
-					Instance->RigidActorAsync->release();
-					Instance->RigidActorAsync = nullptr;
+					if (PrimitiveComp->WeldToImplementation(ParentPrimComponent, PrimitiveComp->AttachSocketName, false))	//welded new simulated body so initialization is done
+					{
+						return false;
+					}
 				}
+			}
 
+			// Don't process if we've already got a body
+			if (Instance->GetPxRigidActor_AssumesLocked())
+			{
 				Instance->OwnerComponent = nullptr;
 				Instance->BodySetup = nullptr;
 
 				continue;
 			}
-		}
 
-		if (Instance->RigidActorSync)
-		{
-			Instance->RigidActorSync->userData = &Instance->PhysxUserData;
-			Instance->RigidActorSync->setName(Instance->CharDebugName.IsValid() ? Instance->CharDebugName->GetData() : nullptr);
-			check(FPhysxUserData::Get<FBodyInstance>(Instance->RigidActorSync->userData) == Instance && FPhysxUserData::Get<FBodyInstance>(Instance->RigidActorSync->userData)->OwnerComponent != NULL);
-		}
-
-		if (Instance->RigidActorAsync)
-		{
-			Instance->RigidActorAsync->userData = &Instance->PhysxUserData;
-			Instance->RigidActorAsync->setName(Instance->CharDebugName.IsValid() ? Instance->CharDebugName->GetData() : nullptr);
-			check(FPhysxUserData::Get<FBodyInstance>(Instance->RigidActorAsync->userData) == Instance && FPhysxUserData::Get<FBodyInstance>(Instance->RigidActorAsync->userData)->OwnerComponent != NULL);
-		}
-
-		//handle special stuff only dynamic actors care about
-		if(PNewDynamic)
-		{
-			// turn off gravity if desired
-			if (!Instance->bEnableGravity)
+			// Set sim parameters for bodies from skeletal mesh components
+			if (!bCompileStatic && SkelMeshComp)
 			{
-				PNewDynamic->setActorFlag(PxActorFlag::eDISABLE_GRAVITY, true);
-				bDynamicsUseAsyncScene = Instance->UseAsyncScene(PhysScene);
+				Instance->bSimulatePhysics = bInstanceSimulatePhysics;
+				if (InstanceBlendWeight != -1.0f)
+				{
+					Instance->PhysicsBlendWeight = InstanceBlendWeight;
+				}
 			}
-		}
 
-		if (bCanDefer)
-		{
-			if (PNewDynamic)
+			Instance->PhysxUserData = FPhysxUserData(Instance);
+
+			static FName PhysicsFormatName(FPlatformProperties::GetPhysicsFormat());
+
+			physx::PxRigidActor* PNewDynamic = nullptr;
+			bool bFoundBinaryData = false;
+			if (PhysicsSerializer)
 			{
-				PhysScene->DeferAddActor(Instance, PNewDynamic, Instance->UseAsyncScene(PhysScene) ? PST_Async : PST_Sync);
+				bFoundBinaryData = GetBinaryData_PhysX_AssumesLocked(Instance);
 			}
-			else
+
+			if (!bFoundBinaryData)
 			{
-				if (Instance->RigidActorSync)
-				{
-					PhysScene->DeferAddActor(Instance, Instance->RigidActorSync, PST_Sync);
-				}
-				if (Instance->RigidActorAsync)
+				PNewDynamic = CreateActor_PhysX_AssumesLocked(Instance, U2PTransform(Transform));
+				const bool bInitFail = CreateShapes_PhysX_AssumesLocked(Instance, PNewDynamic);
+
+				if (bInitFail)
 				{
-					PhysScene->DeferAddActor(Instance, Instance->RigidActorAsync, PST_Async);
+					UE_LOG(LogPhysics, Log, TEXT("Init Instance %d of Primitive Component %s failed"), BodyIdx, *PrimitiveComp->GetName());
+					if (Instance->RigidActorSync)
+					{
+						Instance->RigidActorSync->release();
+						Instance->RigidActorSync = nullptr;
+					}
+
+					if (Instance->RigidActorAsync)
+					{
+						Instance->RigidActorAsync->release();
+						Instance->RigidActorAsync = nullptr;
+					}
+
+					Instance->OwnerComponent = nullptr;
+					Instance->BodySetup = nullptr;
+
+					continue;
 				}
 			}
-		}
-		else
-		{
+
 			if (Instance->RigidActorSync)
 			{
-				PSyncActors.Add(Instance->RigidActorSync);
+				Instance->RigidActorSync->userData = &Instance->PhysxUserData;
+				Instance->RigidActorSync->setName(Instance->CharDebugName.IsValid() ? Instance->CharDebugName->GetData() : nullptr);
+				check(FPhysxUserData::Get<FBodyInstance>(Instance->RigidActorSync->userData) == Instance && FPhysxUserData::Get<FBodyInstance>(Instance->RigidActorSync->userData)->OwnerComponent != NULL);
 			}
 
 			if (Instance->RigidActorAsync)
 			{
-				PAsyncActors.Add(Instance->RigidActorAsync);
+				Instance->RigidActorAsync->userData = &Instance->PhysxUserData;
+				Instance->RigidActorAsync->setName(Instance->CharDebugName.IsValid() ? Instance->CharDebugName->GetData() : nullptr);
+				check(FPhysxUserData::Get<FBodyInstance>(Instance->RigidActorAsync->userData) == Instance && FPhysxUserData::Get<FBodyInstance>(Instance->RigidActorAsync->userData)->OwnerComponent != NULL);
 			}
 
-			if (PNewDynamic)
+			//handle special stuff only dynamic actors care about
+			if (!bCompileStatic && PNewDynamic)
 			{
-				PDynamicActors.Add(PNewDynamic);
+				// turn off gravity if desired
+				if (!Instance->bEnableGravity)
+				{
+					PNewDynamic->setActorFlag(PxActorFlag::eDISABLE_GRAVITY, true);
+					bDynamicsUseAsyncScene = Instance->UseAsyncScene(PhysScene);
+				}
 			}
 
-			// Set the instance to added as we'll add it to the scene in a moment.
-			// This makes sure if it ends up in one of the deferred queues later it 
-			// functions correctly
-			Instance->CurrentSceneState = BodyInstanceSceneState::Added;
-		}
+			if (bCompileStatic || bCanDefer)
+			{
+				if (!bCompileStatic && PNewDynamic)
+				{
+					PhysScene->DeferAddActor(Instance, PNewDynamic, Instance->UseAsyncScene(PhysScene) ? PST_Async : PST_Sync);
+				}
+				else
+				{
+					if (Instance->RigidActorSync)
+					{
+						PhysScene->DeferAddActor(Instance, Instance->RigidActorSync, PST_Sync);
+					}
+					if (Instance->RigidActorAsync)
+					{
+						PhysScene->DeferAddActor(Instance, Instance->RigidActorAsync, PST_Async);
+					}
+				}
+			}
+			else
+			{
+				if (Instance->RigidActorSync)
+				{
+					PSyncActors.Add(Instance->RigidActorSync);
+				}
 
-		Instance->SceneIndexSync = PhysScene->PhysXSceneIndex[PST_Sync];
-		Instance->SceneIndexAsync = PAsyncScene ? PhysScene->PhysXSceneIndex[PST_Async] : 0;
-		Instance->InitialLinearVelocity = InitialLinVel;
-		Instance->bWokenExternally = bComponentAwake;
-	}
+				if (Instance->RigidActorAsync)
+				{
+					PAsyncActors.Add(Instance->RigidActorAsync);
+				}
 
-	return true;
-}
+				if (PNewDynamic)
+				{
+					PDynamicActors.Add(PNewDynamic);
+				}
 
-void FInitBodiesHelper::AddActorsToScene_PhysX_AssumesLocked(TArray<PxActor*>& PSyncActors, TArray<PxActor*>& PAsyncActors, TArray<PxActor*>& PDynamicActors) const
-{
-	SCOPE_CYCLE_COUNTER(STAT_BulkSceneAdd);
+				// Set the instance to added as we'll add it to the scene in a moment.
+				// This makes sure if it ends up in one of the deferred queues later it 
+				// functions correctly
+				Instance->CurrentSceneState = BodyInstanceSceneState::Added;
+			}
+
+			Instance->SceneIndexSync = PhysScene->PhysXSceneIndex[PST_Sync];
+			Instance->SceneIndexAsync = PAsyncScene ? PhysScene->PhysXSceneIndex[PST_Async] : 0;
+			Instance->InitialLinearVelocity = InitialLinVel;
+			Instance->bWokenExternally = bComponentAwake;
+		}
 
+		return true;
+	}
 
-	// Use the aggregate if it exists, has enough slots and is in the correct scene (or no scene)
-	if (InAggregate && PDynamicActors.Num() > 0 && (InAggregate->getMaxNbActors() - InAggregate->getNbActors()) >= (uint32)PDynamicActors.Num())
+	void AddActorsToScene_PhysX_AssumesLocked(TArray<PxActor*>& PSyncActors, TArray<PxActor*>& PAsyncActors, TArray<PxActor*>& PDynamicActors) const
 	{
+		SCOPE_CYCLE_COUNTER(STAT_BulkSceneAdd);
 
-		if (InAggregate->getScene() == nullptr || InAggregate->getScene() == PDynamicActors[0]->getScene())
+
+		// Use the aggregate if it exists, has enough slots and is in the correct scene (or no scene)
+		if (InAggregate && PDynamicActors.Num() > 0 && (InAggregate->getMaxNbActors() - InAggregate->getNbActors()) >= (uint32)PDynamicActors.Num())
 		{
-			for (PxActor* Actor : PDynamicActors)
+
+			if (InAggregate->getScene() == nullptr || InAggregate->getScene() == PDynamicActors[0]->getScene())
 			{
-				InAggregate->addActor(*Actor);
+				for (PxActor* Actor : PDynamicActors)
+				{
+					InAggregate->addActor(*Actor);
+				}
 			}
 		}
-	}
-	else
-	{
-
-		for (PxActor* Actor : PSyncActors)
+		else
 		{
-			PSyncScene->addActor(*Actor);
+
+			for (PxActor* Actor : PSyncActors)
+			{
+				PSyncScene->addActor(*Actor);
+			}
+
+			if (PAsyncScene)
+			{
+				for (PxActor* Actor : PAsyncActors)
+				{
+					PAsyncScene->addActor(*Actor);
+				}
+			}
 		}
 
-		if (PAsyncScene)
+		// Set up dynamic instance data
+		if (!bCompileStatic && !bStatic)
 		{
-			for (PxActor* Actor : PAsyncActors)
+			SCOPE_CYCLE_COUNTER(STAT_InitBodyPostAdd);
+			for (int32 BodyIdx = 0, NumBodies = Bodies.Num(); BodyIdx < NumBodies; ++BodyIdx)
 			{
-				PAsyncScene->addActor(*Actor);
+				FBodyInstance* Instance = Bodies[BodyIdx];
+				Instance->InitDynamicProperties_AssumesLocked();
 			}
 		}
 	}
 
-	// Set up dynamic instance data
-	if (!bStatic)
+	void InitBodies_PhysX() const
 	{
-		SCOPE_CYCLE_COUNTER(STAT_InitBodyPostAdd);
-		for (int32 BodyIdx = 0, NumBodies = Bodies.Num(); BodyIdx < NumBodies; ++BodyIdx)
+		TArray<PxActor*> PSyncActors;
+		TArray<PxActor*> PAsyncActors;
+		TArray<PxActor*> PDynamicActors;
+
+		// No reason to defer if there's an existing aggregate.
+		const bool bCanDefer = bCompileStatic || bDefer && !InAggregate;
+		bool bDynamicsUseAsync = false;
+		if (CreateShapesAndActors_PhysX(PSyncActors, PAsyncActors, PDynamicActors, bCanDefer, bDynamicsUseAsync) == false)
 		{
-			FBodyInstance* Instance = Bodies[BodyIdx];
-			Instance->InitDynamicProperties_AssumesLocked();
+			return;
 		}
-	}
-}
-
-void FInitBodiesHelper::InitBodies_PhysX() const
-{
-	TArray<PxActor*> PSyncActors;
-	TArray<PxActor*> PAsyncActors;
-	TArray<PxActor*> PDynamicActors;
-	
-	// No reason to defer if there's an existing aggregate.
-	const bool bCanDefer = bDefer && !InAggregate;
-	bool bDynamicsUseAsync = false;
-	if(CreateShapesAndActors_PhysX(PSyncActors, PAsyncActors, PDynamicActors, bCanDefer, bDynamicsUseAsync) == false)
-	{
-		return;
-	}
 
-	if (!bCanDefer)
-	{
-		const bool bAddingToSyncScene = (PSyncActors.Num() || (PDynamicActors.Num() && !bDynamicsUseAsync)) && PSyncScene;
-		const bool bAddingToAsyncScene = (PAsyncActors.Num() || (PDynamicActors.Num() && bDynamicsUseAsync)) && PAsyncScene;
+		if (!bCompileStatic && !bCanDefer)
+		{
+			const bool bAddingToSyncScene = (PSyncActors.Num() || (PDynamicActors.Num() && !bDynamicsUseAsync)) && PSyncScene;
+			const bool bAddingToAsyncScene = (PAsyncActors.Num() || (PDynamicActors.Num() && bDynamicsUseAsync)) && PAsyncScene;
 
-		SCOPED_SCENE_WRITE_LOCK(bAddingToSyncScene ? PSyncScene : nullptr);
-		SCOPED_SCENE_WRITE_LOCK(bAddingToAsyncScene ? PAsyncScene : nullptr);
+			SCOPED_SCENE_WRITE_LOCK(bAddingToSyncScene ? PSyncScene : nullptr);
+			SCOPED_SCENE_WRITE_LOCK(bAddingToAsyncScene ? PAsyncScene : nullptr);
 
-		AddActorsToScene_PhysX_AssumesLocked(PSyncActors, PAsyncActors, PDynamicActors);
+			AddActorsToScene_PhysX_AssumesLocked(PSyncActors, PAsyncActors, PDynamicActors);
+		}
 	}
-}
 #endif
 
 #if WITH_BOX2D
 
-void FInitBodiesHelper::InitBodies_Box2D() const
-{
-	const int32 NumBodies = Bodies.Num();
-	bool bLocalComponentAwake = bComponentAwake;
-	if (UBodySetup2D* BodySetup2D = Cast<UBodySetup2D>(BodySetup))
+	void InitBodies_Box2D() const
 	{
-		if (b2World* BoxWorld = FPhysicsIntegration2D::FindAssociatedWorld(PrimitiveComp->GetWorld()))
+		const int32 NumBodies = Bodies.Num();
+		bool bLocalComponentAwake = bComponentAwake;
+		if (UBodySetup2D* BodySetup2D = Cast<UBodySetup2D>(BodySetup))
 		{
-			for (int32 BodyIdx = 0; BodyIdx < NumBodies; ++BodyIdx)
+			if (b2World* BoxWorld = FPhysicsIntegration2D::FindAssociatedWorld(PrimitiveComp->GetWorld()))
 			{
-				FBodyInstance* Instance = Bodies[BodyIdx];
-				const FTransform& Transform = Transforms[BodyIdx];
-
-				const b2Vec2 Scale2D = FPhysicsIntegration2D::ConvertUnrealVectorToBox(Instance->Scale3D);
-				Instance->OwnerComponent = PrimitiveComp;
-
-				// Create the body definition
-				b2BodyDef BodyDefinition;
-				if (bStatic)
-				{
-					BodyDefinition.type = b2_staticBody;
-				}
-				else
+				for (int32 BodyIdx = 0; BodyIdx < NumBodies; ++BodyIdx)
 				{
-					BodyDefinition.type = Instance->ShouldInstanceSimulatingPhysics() ? b2_dynamicBody : b2_kinematicBody;
-				}
+					FBodyInstance* Instance = Bodies[BodyIdx];
+					const FTransform& Transform = Transforms[BodyIdx];
 
-				if (SkelMeshComp)
-				{
-					bLocalComponentAwake = Instance->bStartAwake && SkelMeshComp->BodyInstance.bStartAwake;
-				}
+					const b2Vec2 Scale2D = FPhysicsIntegration2D::ConvertUnrealVectorToBox(Instance->Scale3D);
+					Instance->OwnerComponent = PrimitiveComp;
 
-				if (Instance->bStartAwake || bLocalComponentAwake)
-				{
-					BodyDefinition.awake = true;
-				}
-				else
-				{
-					BodyDefinition.awake = false;
-				}
+					// Create the body definition
+					b2BodyDef BodyDefinition;
+					if (bStatic)
+					{
+						BodyDefinition.type = b2_staticBody;
+					}
+					else
+					{
+						BodyDefinition.type = Instance->ShouldInstanceSimulatingPhysics() ? b2_dynamicBody : b2_kinematicBody;
+					}
 
-				if (Instance->ShouldInstanceSimulatingPhysics())
-				{
-					BodyDefinition.linearVelocity = FPhysicsIntegration2D::ConvertUnrealVectorToBox(InitialLinVel);
-				}
+					if (SkelMeshComp)
+					{
+						bLocalComponentAwake = Instance->bStartAwake && SkelMeshComp->BodyInstance.bStartAwake;
+					}
 
-				// Create the body
-				Instance->BodyInstancePtr = BoxWorld->CreateBody(&BodyDefinition);
-				Instance->BodyInstancePtr->SetUserData(Instance);
+					if (Instance->bStartAwake || bLocalComponentAwake)
+					{
+						BodyDefinition.awake = true;
+					}
+					else
+					{
+						BodyDefinition.awake = false;
+					}
 
-				// Circles
-				for (const FCircleElement2D& Circle : BodySetup2D->AggGeom2D.CircleElements)
-				{
-					b2CircleShape CircleShape;
-					CircleShape.m_radius = Circle.Radius * Instance->Scale3D.Size() / UnrealUnitsPerMeter;
-					CircleShape.m_p.x = Circle.Center.X * Scale2D.x;
-					CircleShape.m_p.y = Circle.Center.Y * Scale2D.y;
+					if (Instance->ShouldInstanceSimulatingPhysics())
+					{
+						BodyDefinition.linearVelocity = FPhysicsIntegration2D::ConvertUnrealVectorToBox(InitialLinVel);
+					}
 
-					b2FixtureDef FixtureDef;
-					FixtureDef.shape = &CircleShape;
+					// Create the body
+					Instance->BodyInstancePtr = BoxWorld->CreateBody(&BodyDefinition);
+					Instance->BodyInstancePtr->SetUserData(Instance);
 
-					Instance->BodyInstancePtr->CreateFixture(&FixtureDef);
-				}
+					// Circles
+					for (const FCircleElement2D& Circle : BodySetup2D->AggGeom2D.CircleElements)
+					{
+						b2CircleShape CircleShape;
+						CircleShape.m_radius = Circle.Radius * Instance->Scale3D.Size() / UnrealUnitsPerMeter;
+						CircleShape.m_p.x = Circle.Center.X * Scale2D.x;
+						CircleShape.m_p.y = Circle.Center.Y * Scale2D.y;
 
-				// Boxes
-				for (const FBoxElement2D& Box : BodySetup2D->AggGeom2D.BoxElements)
-				{
-					const b2Vec2 HalfBoxSize(Box.Width * 0.5f * Scale2D.x, Box.Height * 0.5f * Scale2D.y);
-					const b2Vec2 BoxCenter(Box.Center.X * Scale2D.x, Box.Center.Y * Scale2D.y);
+						b2FixtureDef FixtureDef;
+						FixtureDef.shape = &CircleShape;
 
-					b2PolygonShape DynamicBox;
-					DynamicBox.SetAsBox(HalfBoxSize.x, HalfBoxSize.y, BoxCenter, FMath::DegreesToRadians(Box.Angle));
+						Instance->BodyInstancePtr->CreateFixture(&FixtureDef);
+					}
 
-					b2FixtureDef FixtureDef;
-					FixtureDef.shape = &DynamicBox;
+					// Boxes
+					for (const FBoxElement2D& Box : BodySetup2D->AggGeom2D.BoxElements)
+					{
+						const b2Vec2 HalfBoxSize(Box.Width * 0.5f * Scale2D.x, Box.Height * 0.5f * Scale2D.y);
+						const b2Vec2 BoxCenter(Box.Center.X * Scale2D.x, Box.Center.Y * Scale2D.y);
 
-					Instance->BodyInstancePtr->CreateFixture(&FixtureDef);
-				}
+						b2PolygonShape DynamicBox;
+						DynamicBox.SetAsBox(HalfBoxSize.x, HalfBoxSize.y, BoxCenter, FMath::DegreesToRadians(Box.Angle));
 
-				// Convex hulls
-				for (const FConvexElement2D& Convex : BodySetup2D->AggGeom2D.ConvexElements)
-				{
-					const int32 NumVerts = Convex.VertexData.Num();
+						b2FixtureDef FixtureDef;
+						FixtureDef.shape = &DynamicBox;
+
+						Instance->BodyInstancePtr->CreateFixture(&FixtureDef);
+					}
 
-					if (NumVerts <= b2_maxPolygonVertices)
+					// Convex hulls
+					for (const FConvexElement2D& Convex : BodySetup2D->AggGeom2D.ConvexElements)
 					{
-						TArray<b2Vec2, TInlineAllocator<b2_maxPolygonVertices>> Verts;
+						const int32 NumVerts = Convex.VertexData.Num();
 
-						for (int32 VertexIndex = 0; VertexIndex < Convex.VertexData.Num(); ++VertexIndex)
+						if (NumVerts <= b2_maxPolygonVertices)
 						{
-							const FVector2D SourceVert = Convex.VertexData[VertexIndex];
-							new (Verts)b2Vec2(SourceVert.X * Scale2D.x, SourceVert.Y * Scale2D.y);
-						}
+							TArray<b2Vec2, TInlineAllocator<b2_maxPolygonVertices>> Verts;
 
-						b2PolygonShape ConvexPoly;
-						ConvexPoly.Set(Verts.GetData(), Verts.Num());
+							for (int32 VertexIndex = 0; VertexIndex < Convex.VertexData.Num(); ++VertexIndex)
+							{
+								const FVector2D SourceVert = Convex.VertexData[VertexIndex];
+								new (Verts)b2Vec2(SourceVert.X * Scale2D.x, SourceVert.Y * Scale2D.y);
+							}
 
-						b2FixtureDef FixtureDef;
-						FixtureDef.shape = &ConvexPoly;
+							b2PolygonShape ConvexPoly;
+							ConvexPoly.Set(Verts.GetData(), Verts.Num());
 
-						Instance->BodyInstancePtr->CreateFixture(&FixtureDef);
-					}
-					else
-					{
-						UE_LOG(LogPhysics, Warning, TEXT("Too many vertices in a 2D convex body")); //@TODO: Create a better error message that indicates the asset
+							b2FixtureDef FixtureDef;
+							FixtureDef.shape = &ConvexPoly;
+
+							Instance->BodyInstancePtr->CreateFixture(&FixtureDef);
+						}
+						else
+						{
+							UE_LOG(LogPhysics, Warning, TEXT("Too many vertices in a 2D convex body")); //@TODO: Create a better error message that indicates the asset
+						}
 					}
-				}
 
-				// Make sure it contained at least one shape
-				if (Instance->BodyInstancePtr->GetFixtureList() == nullptr)
-				{
-					if (DebugName.Len())
+					// Make sure it contained at least one shape
+					if (Instance->BodyInstancePtr->GetFixtureList() == nullptr)
 					{
-						UE_LOG(LogPhysics, Log, TEXT("InitBody: failed - no shapes: %s"), *DebugName);
-					}
+						if (DebugName.Len())
+						{
+							UE_LOG(LogPhysics, Log, TEXT("InitBody: failed - no shapes: %s"), *DebugName);
+						}
 
-					Instance->BodyInstancePtr->SetUserData(nullptr);
-					BoxWorld->DestroyBody(Instance->BodyInstancePtr);
-					Instance->BodyInstancePtr = nullptr;
+						Instance->BodyInstancePtr->SetUserData(nullptr);
+						BoxWorld->DestroyBody(Instance->BodyInstancePtr);
+						Instance->BodyInstancePtr = nullptr;
 
-					//clear Owner and Setup info as well to properly clean up the BodyInstance.
-					Instance->OwnerComponent = NULL;
-					Instance->BodySetup = NULL;
+						//clear Owner and Setup info as well to properly clean up the BodyInstance.
+						Instance->OwnerComponent = NULL;
+						Instance->BodySetup = NULL;
 
-					return;
-				}
-				else
-				{
-					// Position the body
-					Instance->SetBodyTransform(Transform, /*bTeleport=*/ true);
-				}
+						return;
+					}
+					else
+					{
+						// Position the body
+						Instance->SetBodyTransform(Transform, /*bTeleport=*/ true);
+					}
 
-				// Apply correct physical materials to shape we created.
-				Instance->UpdatePhysicalMaterials();
+					// Apply correct physical materials to shape we created.
+					Instance->UpdatePhysicalMaterials();
 
-				// Set the filter data on the shapes (call this after setting BodyData because it uses that pointer)
-				Instance->UpdatePhysicsFilterData();
+					// Set the filter data on the shapes (call this after setting BodyData because it uses that pointer)
+					Instance->UpdatePhysicsFilterData();
 
-				if (!bStatic)
-				{
-					// Compute mass (call this after setting BodyData because it uses that pointer)
-					Instance->UpdateMassProperties();
+					if (!bStatic)
+					{
+						// Compute mass (call this after setting BodyData because it uses that pointer)
+						Instance->UpdateMassProperties();
 
-					// Update damping
-					Instance->UpdateDampingProperties();
+						// Update damping
+						Instance->UpdateDampingProperties();
 
-					Instance->SetMaxAngularVelocity(Instance->MaxAngularVelocity, false);
+						Instance->SetMaxAngularVelocity(Instance->MaxAngularVelocity, false);
 
-					Instance->SetMaxDepenetrationVelocity(Instance->bOverrideMaxDepenetrationVelocity ? Instance->MaxDepenetrationVelocity : UPhysicsSettings::Get()->MaxDepenetrationVelocity);
+						Instance->SetMaxDepenetrationVelocity(Instance->bOverrideMaxDepenetrationVelocity ? Instance->MaxDepenetrationVelocity : UPhysicsSettings::Get()->MaxDepenetrationVelocity);
 
-					//@TODO: BOX2D: Determine if sleep threshold and solver settings can be configured per-body or not
+						//@TODO: BOX2D: Determine if sleep threshold and solver settings can be configured per-body or not
 #if 0
-					// Set the parameters for determining when to put the object to sleep.
-					float SleepEnergyThresh = PNewDynamic->getSleepThreshold();
-					if (SleepFamily == ESleepFamily::Sensitive)
-					{
-						SleepEnergyThresh /= 20.f;
-					}
-					PNewDynamic->setSleepThreshold(SleepEnergyThresh);
-					// set solver iteration count 
-					int32 PositionIterCount = FMath::Clamp(PositionSolverIterationCount, 1, 255);
-					int32 VelocityIterCount = FMath::Clamp(VelocitySolverIterationCount, 1, 255);
-					PNewDynamic->setSolverIterationCounts(PositionIterCount, VelocityIterCount);
+						// Set the parameters for determining when to put the object to sleep.
+						float SleepEnergyThresh = PNewDynamic->getSleepThreshold();
+						if (SleepFamily == ESleepFamily::Sensitive)
+						{
+							SleepEnergyThresh /= 20.f;
+						}
+						PNewDynamic->setSleepThreshold(SleepEnergyThresh);
+						// set solver iteration count 
+						int32 PositionIterCount = FMath::Clamp(PositionSolverIterationCount, 1, 255);
+						int32 VelocityIterCount = FMath::Clamp(VelocitySolverIterationCount, 1, 255);
+						PNewDynamic->setSolverIterationCounts(PositionIterCount, VelocityIterCount);
 #endif
+					}
 				}
 			}
+			return;
 		}
-		return;
 	}
-}
 
 #endif
 
+};
+
+void FBodyInstance::InitBodies(TArray<FBodyInstance*>& Bodies, TArray<FTransform>& Transforms, class UBodySetup* BodySetup, class UPrimitiveComponent* PrimitiveComp, class FPhysScene* InRBScene, PhysXAggregateType InAggregate /*= NULL*/, bool bDefer /*= false*/, UPhysicsSerializer* PhysicsSerializer )
+{
+	SCOPE_CYCLE_COUNTER(STAT_InitBodies);
+
+	check(BodySetup);
+	check(InRBScene);
+	check(Bodies.Num() > 0);
+
+	FInitBodiesHelper<false> InitBodiesHelper(Bodies, Transforms, BodySetup, PrimitiveComp, InRBScene, InAggregate, bDefer, PhysicsSerializer);
+	
+#if WITH_PHYSX
+	InitBodiesHelper.InitBodies_PhysX();
+#endif
+
+#if WITH_BOX2D
+	InitBodiesHelper.InitBodies_Box2D();
+#endif
+}
+
+
+
 
 #endif // UE_WITH_PHYSICS
 
@@ -4614,115 +4605,23 @@ void FBodyInstance::GetFilterData_AssumesLocked(FShapeData& ShapeData, bool bFor
 	}
 }
 
-void FBodyInstance::InitStaticBodies(TArray<FBodyInstance*>& Bodies, TArray<FTransform>& Transforms, class UBodySetup* BodySetup, class UPrimitiveComponent* PrimitiveComp, class FPhysScene* InRBScene)
+void FBodyInstance::InitStaticBodies(TArray<FBodyInstance*>& Bodies, TArray<FTransform>& Transforms, class UBodySetup* BodySetup, class UPrimitiveComponent* PrimitiveComp, class FPhysScene* InRBScene, UPhysicsSerializer* PhysicsSerializer)
 {
 	SCOPE_CYCLE_COUNTER(STAT_StaticInitBodies);
 
-	const int32 NumBodies = Bodies.Num();
-	AActor* OwningActor = PrimitiveComp ? PrimitiveComp->GetOwner() : nullptr;
-
-	UPhysicalMaterial* SimpleMat;
-	TArray<UPhysicalMaterial*> ComplexMats;
-	PxMaterial* PhysxSimpleMat = nullptr;
-
-	PxScene* PSceneSync = InRBScene->GetPhysXScene(PST_Sync);
-	PxScene* PSceneAsync = InRBScene->HasAsyncScene() ? InRBScene->GetPhysXScene(PST_Async) : nullptr;
-
-	TArray<PxActor*> NewSyncActors;
-	TArray<PxActor*> NewAsyncActors;
-
-	if(PSceneAsync)
-	{
-		NewAsyncActors.Reserve(NumBodies);
-	}
-	NewSyncActors.Reserve(NumBodies);
-
-	for(int32 BodyIdx = 0 ; BodyIdx < NumBodies ; ++BodyIdx)
-	{
-		FBodyInstance* Instance = Bodies[BodyIdx];
-		FTransform& Transform = Transforms[BodyIdx];
-		FShapeData ShapeData;
-
-		Instance->OwnerComponent = PrimitiveComp;
-		Instance->BodySetup = BodySetup;
-		Instance->Scale3D = Transform.GetScale3D();
-
-		// Handle autowelding here to avoid extra work
-		if(Instance->bAutoWeld)
-		{
-			UPrimitiveComponent * ParentPrimComponent = PrimitiveComp ? Cast<UPrimitiveComponent>(PrimitiveComp->AttachParent) : NULL;
-
-			//if we have a parent we will now do the weld and exit any further initialization
-			if(ParentPrimComponent && PrimitiveComp->GetWorld() && PrimitiveComp->GetWorld()->IsGameWorld())
-			{
-				if(PrimitiveComp->WeldToImplementation(ParentPrimComponent, PrimitiveComp->AttachSocketName, false))	//welded new simulated body so initialization is done
-				{
-					return;
-				}
-			}
-		}
-
-		PxRigidActor* PNewActorSync = nullptr;
-		PxRigidActor* PNewActorAsync = nullptr;
-		PxTransform PTransform = U2PTransform(Transform);
-		bool bInitFail = false;
-
-		// Don't process if we've already got a body
-		if(Instance->GetPxRigidActor_AssumesLocked())
-		{
-			Instance->OwnerComponent = nullptr;
-			Instance->BodySetup = nullptr;
-
-			continue;
-		}
-
-		Instance->PhysxUserData = FPhysxUserData(Instance);
-		PNewActorSync = GPhysXSDK->createRigidStatic(PTransform);
-		PNewActorSync->userData = &Instance->PhysxUserData;
-		check(FPhysxUserData::Get<FBodyInstance>(PNewActorSync->userData) == Instance && FPhysxUserData::Get<FBodyInstance>(PNewActorSync->userData)->OwnerComponent != NULL);
-		NewSyncActors.Add(PNewActorSync);
-
-		if(PSceneAsync)
-		{
-			PNewActorAsync = GPhysXSDK->createRigidStatic(PTransform);
-			PNewActorAsync->userData = &Instance->PhysxUserData;
-			check(FPhysxUserData::Get<FBodyInstance>(PNewActorAsync->userData) == Instance && FPhysxUserData::Get<FBodyInstance>(PNewActorAsync->userData)->OwnerComponent != NULL);
-			NewAsyncActors.Add(PNewActorAsync);
-		}
-
-		Instance->RigidActorSync = PNewActorSync;
-		Instance->RigidActorAsync = PNewActorAsync;
-		Instance->SceneIndexSync = InRBScene->PhysXSceneIndex[PST_Sync];
-		Instance->SceneIndexAsync = PSceneAsync ? InRBScene->PhysXSceneIndex[PST_Async] : 0;
-
-		// Get filter data
-		PxFilterData SimFilter;
-		PxFilterData QuerySimpleFilter;
-		PxFilterData QueryComplexFilter;
-		TEnumAsByte<ECollisionEnabled::Type> UseCollisionEnabled;
-		Instance->GetFilterData_AssumesLocked(ShapeData);
-		Instance->GetShapeFlags_AssumesLocked(ShapeData, Instance->CollisionEnabled, BodySetup->CollisionTraceFlag == CTF_UseComplexAsSimple);
+	check(BodySetup);
+	check(InRBScene);
+	check(Bodies.Num() > 0);
 
-		Instance->GetComplexPhysicalMaterials(ComplexMats);
-		SimpleMat = Instance->GetSimplePhysicalMaterial();
-		PhysxSimpleMat = SimpleMat->GetPhysXMaterial();
+	FInitBodiesHelper<true> InitBodiesHelper(Bodies, Transforms, BodySetup, PrimitiveComp, InRBScene, nullptr, true, PhysicsSerializer);
 
-		if(PNewActorSync)
-		{
-			BodySetup->AddShapesToRigidActor_AssumesLocked(Instance, PNewActorSync, PST_Sync, Instance->Scale3D, PhysxSimpleMat, ComplexMats, ShapeData, FTransform::Identity);
-		}
-		
-		if(PNewActorAsync)
-		{
-			BodySetup->AddShapesToRigidActor_AssumesLocked(Instance, PNewActorAsync, PST_Async, Instance->Scale3D, PhysxSimpleMat, ComplexMats, ShapeData, FTransform::Identity);
-		}
-	}
+#if WITH_PHYSX
+	InitBodiesHelper.InitBodies_PhysX();
+#endif
 
-	InRBScene->DeferAddActors(Bodies, NewSyncActors, PST_Sync);
-	if(PSceneAsync)
-	{
-		InRBScene->DeferAddActors(Bodies, NewAsyncActors, PST_Async);
-	}
+#if WITH_BOX2D
+	InitBodiesHelper.InitBodies_Box2D();
+#endif
 }
 
 void FBodyInstance::SetShapeFlags_AssumesLocked(TEnumAsByte<ECollisionEnabled::Type> UseCollisionEnabled, PxShape* PInShape, EPhysicsSceneType SceneType, const bool bUseComplexAsSimple)
