RPC's save checksum separately now. This allows us to still replicate an actor, even if the rpc signature changed. We use checksums in the net field cache to get past the issue where indexes can be off when we add/remove rpc's.

diff --git a/Engine/Source/Runtime/CoreUObject/Private/UObject/CoreNet.cpp b/Engine/Source/Runtime/CoreUObject/Private/UObject/CoreNet.cpp
--- a/Engine/Source/Runtime/CoreUObject/Private/UObject/CoreNet.cpp
+++ b/Engine/Source/Runtime/CoreUObject/Private/UObject/CoreNet.cpp
@@ -16,12 +16,152 @@ DEFINE_STAT(STAT_NetSerializeFastArray_BuildMap);
 -----------------------------------------------------------------------------*/
 
 FClassNetCache::FClassNetCache()
-{}
-FClassNetCache::FClassNetCache( UClass* InClass )
-: Class( InClass )
-{}
+{
+}
+
+FClassNetCache::FClassNetCache( const UClass* InClass ) : Class( InClass )
+{
+}
+
+void FClassNetCacheMgr::SortProperties( TArray< UProperty* >& Properties ) const
+{
+	// Sort NetProperties so that their ClassReps are sorted by memory offset
+	struct FCompareUFieldOffsets
+	{
+		FORCEINLINE bool operator()( UProperty & A, UProperty & B ) const
+		{
+			// Ensure stable sort
+			if ( A.GetOffset_ForGC() == B.GetOffset_ForGC() )
+			{
+				return A.GetName() < B.GetName();
+			}
 
-FClassNetCache * FClassNetCacheMgr::GetClassNetCache( UClass* Class )
+			return A.GetOffset_ForGC() < B.GetOffset_ForGC();
+		}
+	};
+
+	Sort( Properties.GetData(), Properties.Num(), FCompareUFieldOffsets() );
+}
+
+uint32 FClassNetCacheMgr::SortedStructFieldsChecksum( const UStruct* Struct, uint32 Checksum ) const
+{
+	// Generate a list that we can sort, to make sure we process these deterministically
+	TArray< UProperty * > Fields;
+
+	for ( TFieldIterator< UProperty > It( Struct ); It; ++It )
+	{
+		if ( It->PropertyFlags & CPF_RepSkip )
+		{
+			continue;
+		}
+
+		Fields.Add( *It );
+	}
+
+	// Sort them
+	SortProperties( Fields );
+
+	// Evolve the checksum on the sorted list
+	for ( auto Field : Fields )
+	{
+		Checksum = GetPropertyChecksum( Field, Checksum );
+	}
+
+	return Checksum;
+}
+
+uint32 FClassNetCacheMgr::GetPropertyChecksum( const UProperty* Property, uint32 Checksum ) const
+{
+	if ( bDebugChecksum )
+	{
+		UE_LOG( LogCoreNet, Warning, TEXT( "%s%s [%s] [%u] [%u]" ), FCString::Spc( 2 * DebugChecksumIndent ), *Property->GetName().ToLower(), *Property->GetClass()->GetName().ToLower(), Property->ArrayDim, Checksum );
+	}
+
+	// Evolve checksum on name
+	Checksum = FCrc::StrCrc32( *Property->GetName().ToLower(), Checksum );
+
+	// Evolve checksum on class name
+	Checksum = FCrc::StrCrc32( *Property->GetClass()->GetName().ToLower(), Checksum );
+
+	// Evolve checksum on array dim (to detect when static arrays change size)
+	// We use a string to be endian agnostic
+	Checksum = FCrc::StrCrc32( *FString::Printf( TEXT( "%u" ), Property->ArrayDim ), Checksum );
+
+	const UArrayProperty* ArrayProperty = Cast< UArrayProperty >( Property );
+
+	// Evolve checksum on array inner
+	if ( ArrayProperty != NULL )
+	{
+		return GetPropertyChecksum( ArrayProperty->Inner, Checksum );
+	}
+
+	const UStructProperty* StructProperty = Cast< UStructProperty >( Property );
+
+	// Evolve checksum on property struct fields
+	if ( StructProperty != NULL )
+	{
+		if ( bDebugChecksum )
+		{
+			UE_LOG( LogCoreNet, Warning, TEXT( "%s [%s] [%u]" ), FCString::Spc( 2 * DebugChecksumIndent ), *StructProperty->Struct->GetName().ToLower(), Checksum );
+		}
+
+		// Evolve checksum on struct name
+		Checksum = FCrc::StrCrc32( *StructProperty->Struct->GetName().ToLower(), Checksum );
+
+		const_cast< FClassNetCacheMgr* >( this )->DebugChecksumIndent++;
+
+		Checksum = SortedStructFieldsChecksum( StructProperty->Struct, Checksum );
+
+		const_cast< FClassNetCacheMgr* >( this )->DebugChecksumIndent--;
+	}
+
+	return Checksum;
+}
+
+uint32 FClassNetCacheMgr::GetFunctionChecksum( const UFunction* Function, uint32 Checksum ) const
+{
+	// Evolve checksum on function name
+	Checksum = FCrc::StrCrc32( *Function->GetName().ToLower(), Checksum );
+
+	// Evolve the checksum on function flags
+	Checksum = FCrc::StrCrc32( *FString::Printf( TEXT( "%u" ), Function->FunctionFlags ), Checksum );
+
+	TArray< UProperty * > Parms;
+
+	for ( TFieldIterator< UProperty > It( Function ); It && ( It->PropertyFlags & ( CPF_Parm | CPF_ReturnParm ) ) == CPF_Parm; ++It )
+	{
+		Parms.Add( *It );
+	}
+
+	// Sort parameters by offset/name
+	SortProperties( Parms );
+
+	// Evolve checksum on sorted function parameters
+	for ( UProperty* Parm : Parms )
+	{
+		Checksum = GetPropertyChecksum( Parm, Checksum );
+	}
+
+	return Checksum;
+}
+
+uint32 FClassNetCacheMgr::GetFieldChecksum( const UField* Field, uint32 Checksum ) const
+{
+	if ( Cast< UProperty >( Field ) != NULL )
+	{
+		return GetPropertyChecksum( (UProperty*)Field, Checksum );
+	}
+	else if ( Cast< UFunction >( Field ) != NULL )
+	{
+		return GetFunctionChecksum( (UFunction*)Field, Checksum );
+	}
+
+	UE_LOG( LogCoreNet, Warning, TEXT( "GetFieldChecksum: Unknown field: %s" ), *Field->GetName() );
+
+	return Checksum;
+}
+
+const FClassNetCache * FClassNetCacheMgr::GetClassNetCache( const UClass* Class )
 {
 	FClassNetCache * Result = ClassFieldIndices.FindRef( Class );
 
@@ -30,28 +170,62 @@ FClassNetCache * FClassNetCacheMgr::GetClassNetCache( UClass* Class )
 		Result					= ClassFieldIndices.Add( Class, new FClassNetCache( Class ) );
 		Result->Super			= NULL;
 		Result->FieldsBase		= 0;
+		Result->ClassChecksum	= 0;
 
 		if ( Class->GetSuperClass() )
 		{
-			Result->Super		= GetClassNetCache(Class->GetSuperClass());
-			Result->FieldsBase	= Result->Super->GetMaxIndex();
+			Result->Super			= GetClassNetCache( Class->GetSuperClass() );
+			Result->FieldsBase		= Result->Super->GetMaxIndex();
+			Result->ClassChecksum	= Result->Super->ClassChecksum;
 		}
 
 		Result->Fields.Empty( Class->NetFields.Num() );
 
+		TArray< UProperty* > Properties;
+
 		for( int32 i = 0; i < Class->NetFields.Num(); i++ )
 		{
-			// Add sandboxed items to net cache.  
+			// Add each net field to cache, and assign index/checksum
 			UField * Field = Class->NetFields[i];
-			int32 ThisIndex	= Result->GetMaxIndex();
-			new ( Result->Fields )FFieldNetCache( Field, ThisIndex );
+
+			UProperty* Property = Cast< UProperty >( Field );
+			
+			if ( Property != NULL )
+			{
+				Properties.Add( Property );
+			}
+
+			// Get individual checksum
+			const uint32 Checksum = GetFieldChecksum( Field, 0 );
+
+			// Get index
+			const int32 ThisIndex = Result->GetMaxIndex();
+
+			// Add to cached fields on this class
+			Result->Fields.Add( FFieldNetCache( Field, ThisIndex, Checksum ) );
 		}
 
 		Result->Fields.Shrink();
 
+		// Add fields to the appropriate hash maps
 		for ( TArray< FFieldNetCache >::TIterator It( Result->Fields ); It; ++It )
 		{
 			Result->FieldMap.Add( It->Field, &*It );
+
+			if ( Result->FieldChecksumMap.Contains( It->FieldChecksum ) )
+			{
+				UE_LOG( LogCoreNet, Error, TEXT ( "Duplicate checksum: %s, %u" ), *It->Field->GetName(), It->FieldChecksum );
+			}
+
+			Result->FieldChecksumMap.Add( It->FieldChecksum, &*It );
+		}
+
+		// Initialize class checksum (just use properties for this)
+		SortProperties( Properties );
+
+		for ( auto Property : Properties )
+		{
+			Result->ClassChecksum = GetPropertyChecksum( Property, Result->ClassChecksum );
 		}
 	}
 	return Result;
diff --git a/Engine/Source/Runtime/CoreUObject/Public/UObject/CoreNet.h b/Engine/Source/Runtime/CoreUObject/Public/UObject/CoreNet.h
--- a/Engine/Source/Runtime/CoreUObject/Public/UObject/CoreNet.h
+++ b/Engine/Source/Runtime/CoreUObject/Public/UObject/CoreNet.h
@@ -15,15 +15,15 @@ class COREUOBJECT_API FFieldNetCache
 {
 public:
 	UField* Field;
-	int32 FieldNetIndex;
+	int32	FieldNetIndex;
+	uint32	FieldChecksum;
 	FFieldNetCache()
 	{}
-	FFieldNetCache( UField* InField, int32 InFieldNetIndex )
-	: Field(InField), FieldNetIndex(InFieldNetIndex)
+	FFieldNetCache( UField* InField, int32 InFieldNetIndex, uint32 InFieldChecksum )
+		: Field(InField), FieldNetIndex(InFieldNetIndex), FieldChecksum(InFieldChecksum)
 	{}
 };
 
-
 //
 // Information about a class, cached for network coordination.
 //
@@ -32,52 +32,86 @@ class COREUOBJECT_API FClassNetCache
 	friend class FClassNetCacheMgr;
 public:
 	FClassNetCache();
-	FClassNetCache( UClass* Class );
-	int32 GetMaxIndex()
+	FClassNetCache( const UClass* Class );
+
+	int32 GetMaxIndex() const
+	{
+		return FieldsBase + Fields.Num();
+	}
+
+	const FFieldNetCache* GetFromField( const UObject* Field ) const
 	{
-		return FieldsBase+Fields.Num();
+		FFieldNetCache* Result = NULL;
+
+		for ( const FClassNetCache* C= this; C; C = C->Super )
+		{
+			if ( ( Result = C->FieldMap.FindRef( Field ) ) != NULL )
+			{
+				break;
+			}
+		}
+		return Result;
 	}
-	FFieldNetCache* GetFromField( UObject* Field )
+
+	const FFieldNetCache* GetFromChecksum( const uint32 Checksum ) const
 	{
-		FFieldNetCache* Result=NULL;
-		for( FClassNetCache* C=this; C; C=C->Super )
+		FFieldNetCache* Result = NULL;
+
+		for ( const FClassNetCache* C = this; C; C = C->Super )
 		{
-			if( (Result=C->FieldMap.FindRef(Field))!=NULL )
+			if ( ( Result = C->FieldChecksumMap.FindRef( Checksum ) ) != NULL )
 			{
 				break;
 			}
 		}
 		return Result;
 	}
-	FFieldNetCache* GetFromIndex( int32 Index )
+
+	const FFieldNetCache* GetFromIndex( const int32 Index ) const
 	{
-		for( FClassNetCache* C=this; C; C=C->Super )
+		for ( const FClassNetCache* C = this; C; C = C->Super )
 		{
-			if( Index>=C->FieldsBase && Index<C->FieldsBase+C->Fields.Num() )
+			if ( Index >= C->FieldsBase && Index < C->FieldsBase + C->Fields.Num() )
 			{
 				return &C->Fields[Index-C->FieldsBase];
 			}
 		}
 		return NULL;
 	}
+
+	uint32 GetClassChecksum() const { return ClassChecksum; }
+
 private:
 	int32								FieldsBase;
-	FClassNetCache *					Super;
-	TWeakObjectPtr< UClass >			Class;
+	const FClassNetCache*				Super;
+	TWeakObjectPtr< const UClass >		Class;
+	uint32								ClassChecksum;
 	TArray< FFieldNetCache >			Fields;
-	TMap< UObject *, FFieldNetCache * > FieldMap;
+	TMap< UObject*, FFieldNetCache* >	FieldMap;
+	TMap< uint32, FFieldNetCache* >		FieldChecksumMap;
 };
 
 
 class COREUOBJECT_API FClassNetCacheMgr
 {
 public:
+	FClassNetCacheMgr() : bDebugChecksum( false ), DebugChecksumIndent( 0 ) { }
+
 	/** get the cached field to index mappings for the given class */
-	FClassNetCache *	GetClassNetCache( UClass * Class );
-	void				ClearClassNetCache();
+	const FClassNetCache*	GetClassNetCache( const UClass* Class );
+	void					ClearClassNetCache();
+
+	void				SortProperties( TArray< UProperty* >& Properties ) const;
+	uint32				SortedStructFieldsChecksum( const UStruct* Struct, uint32 Checksum ) const;
+	uint32				GetPropertyChecksum( const UProperty* Property, uint32 Checksum ) const;
+	uint32				GetFunctionChecksum( const UFunction* Function, uint32 Checksum ) const;
+	uint32				GetFieldChecksum( const UField* Field, uint32 Checksum ) const;
+
+	bool				bDebugChecksum;
+	int					DebugChecksumIndent;
 
 private:
-	TMap< TWeakObjectPtr< UClass >, FClassNetCache * > ClassFieldIndices;
+	TMap< TWeakObjectPtr< const UClass >, FClassNetCache* > ClassFieldIndices;
 };
 
 
diff --git a/Engine/Source/Runtime/Engine/Classes/Engine/ActorChannel.h b/Engine/Source/Runtime/Engine/Classes/Engine/ActorChannel.h
--- a/Engine/Source/Runtime/Engine/Classes/Engine/ActorChannel.h
+++ b/Engine/Source/Runtime/Engine/Classes/Engine/ActorChannel.h
@@ -148,7 +148,7 @@ public:
 	void BeginContentBlockForSubObjectDelete( FOutBunch & Bunch, FNetworkGUID & GuidToDelete );
 
 	/** Writes the footer for a content block of proeprties / RPCs for the given object (either the actor a subobject of the actor) */
-	void EndContentBlock( UObject *Obj, FOutBunch &Bunch, FClassNetCache* ClassCache = nullptr );
+	void EndContentBlock( UObject *Obj, FOutBunch &Bunch, const FClassNetCache* ClassCache = nullptr );
 
 	/** Reads the header of the content block and instantiates the subobject if necessary */
 	UObject* ReadContentBlockHeader(FInBunch& Bunch, bool& bObjectDeleted);
diff --git a/Engine/Source/Runtime/Engine/Classes/Engine/PackageMapClient.h b/Engine/Source/Runtime/Engine/Classes/Engine/PackageMapClient.h
--- a/Engine/Source/Runtime/Engine/Classes/Engine/PackageMapClient.h
+++ b/Engine/Source/Runtime/Engine/Classes/Engine/PackageMapClient.h
@@ -82,8 +82,6 @@ public:
 	TMap< TWeakObjectPtr< UObject >, FNetworkGUID >	NetGUIDLookup;
 	int32											UniqueNetIDs[2];
 
-	TMap< TWeakObjectPtr< UClass >, uint32 >		ClassChecksumMap;
-
 	bool											IsExportingNetGUIDBunch;
 
 	UNetDriver *									Driver;
diff --git a/Engine/Source/Runtime/Engine/Private/DataChannel.cpp b/Engine/Source/Runtime/Engine/Private/DataChannel.cpp
--- a/Engine/Source/Runtime/Engine/Private/DataChannel.cpp
+++ b/Engine/Source/Runtime/Engine/Private/DataChannel.cpp
@@ -1954,6 +1954,12 @@ void UActorChannel::ProcessBunch( FInBunch & Bunch )
 		if ( !Replicator->ReceivedBunch( Bunch, RepFlags, bHasUnmapped ) )
 		{
 			UE_LOG( LogNet, Error, TEXT( "UActorChannel::ProcessBunch: Replicator.ReceivedBunch failed.  Closing connection. RepObj: %s, Channel: %i"), RepObj ? *RepObj->GetFullName() : TEXT( "NULL" ), ChIndex  );
+	
+			if ( Connection->InternalAck )
+			{
+				break;
+			}
+
 			Connection->Close();
 			return;
 		}
@@ -2361,7 +2367,7 @@ void UActorChannel::BeginContentBlockForSubObjectDelete( FOutBunch & Bunch, FNet
 	NETWORK_PROFILER(GNetworkProfiler.TrackBeginContentBlock(nullptr, Bunch.GetNumBits() - NumStartingBits));
 }
 
-void UActorChannel::EndContentBlock( UObject *Obj, FOutBunch &Bunch, FClassNetCache* ClassCache )
+void UActorChannel::EndContentBlock( UObject *Obj, FOutBunch &Bunch, const FClassNetCache* ClassCache )
 {
 	check(Obj);
 
@@ -2372,8 +2378,17 @@ void UActorChannel::EndContentBlock( UObject *Obj, FOutBunch &Bunch, FClassNetCa
 		ClassCache = Connection->Driver->NetCache->GetClassNetCache( Obj->GetClass() );
 	}
 
-	// Write max int to signify done
-	Bunch.WriteIntWrapped(ClassCache->GetMaxIndex(), ClassCache->GetMaxIndex()+1);
+	if ( Connection->InternalAck )
+	{
+		// Write out 0 checksum to signify done
+		uint32 Checksum = 0;
+		Bunch << Checksum;
+	}
+	else
+	{
+		// Write max int to signify done
+		Bunch.WriteIntWrapped(ClassCache->GetMaxIndex(), ClassCache->GetMaxIndex()+1);
+	}
 
 	NETWORK_PROFILER(GNetworkProfiler.TrackEndContentBlock(Obj, Bunch.GetNumBits() - NumStartingBits));
 }
diff --git a/Engine/Source/Runtime/Engine/Private/DataReplication.cpp b/Engine/Source/Runtime/Engine/Private/DataReplication.cpp
--- a/Engine/Source/Runtime/Engine/Private/DataReplication.cpp
+++ b/Engine/Source/Runtime/Engine/Private/DataReplication.cpp
@@ -390,7 +390,71 @@ void FObjectReplicator::ReceivedNak( int32 NakPacketId )
 	}
 }
 
-bool FObjectReplicator::ReceivedBunch( FInBunch &Bunch, const FReplicationFlags& RepFlags, bool & bOutHasUnmapped )
+const FFieldNetCache* FObjectReplicator::ReadField( const FClassNetCache* ClassCache, FInBunch& Bunch ) const
+{
+	if ( Connection->InternalAck )
+	{
+		// Replays use checksum rather than index
+		uint32 Checksum = 0;
+		Bunch << Checksum;
+
+		if ( Bunch.IsError() )
+		{
+			UE_LOG( LogNet, Error, TEXT( "ReadField: Error reading checksum: %s" ), *GetObject()->GetFullName() );
+			return NULL;
+		}
+
+		if ( Checksum == 0 )
+		{
+			return NULL;	// We're done
+		}
+
+		const FFieldNetCache * FieldNetCache = ClassCache->GetFromChecksum( Checksum );
+
+		if ( FieldNetCache == NULL )
+		{
+			UE_LOG( LogNet, Error, TEXT( "ReadField: GetFromChecksum failed: %s" ), *GetObject()->GetFullName() );
+			Bunch.SetError();
+			return NULL;
+		}
+
+		return FieldNetCache;
+	}
+
+	const int32 RepIndex = Bunch.ReadInt( ClassCache->GetMaxIndex() + 1 );
+
+	if ( Bunch.IsError() )
+	{
+		UE_LOG( LogNet, Error, TEXT( "ReadField: Error reading RepIndex: %s" ), *GetObject()->GetFullName() );
+		return NULL;
+	}
+
+	if ( RepIndex == ClassCache->GetMaxIndex() )
+	{
+		return NULL;	// We're done
+	}
+
+	if ( RepIndex > ClassCache->GetMaxIndex() )
+	{
+		// We shouldn't be receiving this bunch of this object has no properties or RPC functions to process
+		UE_LOG( LogNet, Error, TEXT( "ReadField: RepIndex too large: %s" ), *GetObject()->GetFullName() );
+		Bunch.SetError();
+		return NULL;
+	}
+	
+	const FFieldNetCache * FieldNetCache = ClassCache->GetFromIndex( RepIndex );
+
+	if ( FieldNetCache == NULL )
+	{
+		UE_LOG( LogNet, Error, TEXT( "ReadField: GetFromIndex failed: %s" ), *GetObject()->GetFullName() );
+		Bunch.SetError();
+		return NULL;
+	}
+
+	return FieldNetCache;
+}
+
+bool FObjectReplicator::ReceivedBunch( FInBunch& Bunch, const FReplicationFlags& RepFlags, bool& bOutHasUnmapped )
 {
 	UObject* Object = GetObject();
 
@@ -404,7 +468,7 @@ bool FObjectReplicator::ReceivedBunch( FInBunch &Bunch, const FReplicationFlags&
 
 	const bool bIsServer = ( OwningChannel->Connection->Driver->ServerConnection == NULL );
 
-	FClassNetCache * ClassCache = OwningChannel->Connection->Driver->NetCache->GetClassNetCache( ObjectClass );
+	const FClassNetCache * ClassCache = OwningChannel->Connection->Driver->NetCache->GetClassNetCache( ObjectClass );
 
 	if ( ClassCache == NULL )
 	{
@@ -414,42 +478,26 @@ bool FObjectReplicator::ReceivedBunch( FInBunch &Bunch, const FReplicationFlags&
 
 	bool bThisBunchReplicatedProperties = false;
 
-	// First RepIndex.
-	int32 RepIndex = Bunch.ReadInt( ClassCache->GetMaxIndex() + 1 );
+	// Read first field
+	const FFieldNetCache * FieldCache = ReadField( ClassCache, Bunch );
 
 	if ( Bunch.IsError() )
 	{
-		UE_LOG( LogNet, Error, TEXT( "ReceivedBunch: Error reading bunch 1: %s" ), *Object->GetFullName() );
+		UE_LOG( LogNet, Error, TEXT( "ReceivedBunch: Error reading field 1: %s" ), *Object->GetFullName() );
 		return false;
 	}
 
-	if ( RepIndex == ClassCache->GetMaxIndex() )
+	if ( FieldCache == NULL )
 	{
 		// There are no actual replicated properties or functions in this bunch. That is ok - we may have gotten this
-		// actor/subobject because we want the client to spawn one (but we arent actually replicating properties on it)
+		// actor/sub-object because we want the client to spawn one (but we aren't actually replicating properties on it)
 		return true;
 	}
 
-	if ( RepIndex > ClassCache->GetMaxIndex() )
-	{
-		// We shouldn't be receiving this bunch of this object has no properties or RPC functions to process
-		UE_LOG( LogNet, Error, TEXT( "ReceivedBunch: RepIndex too large: %s" ), *Object->GetFullName() );
-		return false;
-	}
-
-	FFieldNetCache * FieldCache = ClassCache->GetFromIndex( RepIndex );
-
-	if ( FieldCache == NULL )
-	{
-		UE_LOG( LogNet, Error, TEXT( "ReceivedBunch: FieldCache == NULL: %s" ), *Object->GetFullName() );
-		return false;
-	}
-
 	while ( FieldCache )
 	{
 		// Receive properties from the net.
-		UProperty * ReplicatedProp	= NULL;
-		int32		LastIndex		= 0;
+		UProperty* ReplicatedProp = NULL;
 
 		while ( FieldCache && ( ReplicatedProp = Cast< UProperty >( FieldCache->Field ) ) != NULL )
 		{
@@ -476,7 +524,7 @@ bool FObjectReplicator::ReceivedBunch( FInBunch &Bunch, const FReplicationFlags&
 				static IConsoleVariable* CVar = IConsoleManager::Get().FindConsoleVariable(TEXT("net.Replication.DebugProperty"));
 				if (CVar && !CVar->GetString().IsEmpty() && ReplicatedProp->GetName().Contains(CVar->GetString()) )
 				{
-					UE_LOG(LogNet, Log, TEXT("Replicating Property[%d] %s on %s"), RepIndex, *ReplicatedProp->GetName(), *Object->GetName());
+					UE_LOG(LogNet, Log, TEXT("Replicating Property[%d] %s on %s"), ReplicatedProp->RepIndex, *ReplicatedProp->GetName(), *Object->GetName());
 					DebugProperty = true;
 				}
 			}
@@ -493,16 +541,14 @@ bool FObjectReplicator::ReceivedBunch( FInBunch &Bunch, const FReplicationFlags&
 			else
 			{
 				// Receive array index.
-				int32 Element = 0;
+				uint32 Element = 0;
 				if ( ReplicatedProp->ArrayDim != 1 )
 				{
-					// Serialize index as delta from previous index to increase chance we'll only use 1 byte
-					uint32 idx;
-					Bunch.SerializeIntPacked( idx );
-					Element = static_cast< int32 >( idx ) + LastIndex;
-					LastIndex = Element;
+					check( ReplicatedProp->ArrayDim >= 2 );
+
+					Bunch.SerializeIntPacked( Element );
 
-					if ( Element >= ReplicatedProp->ArrayDim )
+					if ( Element >= (uint32)ReplicatedProp->ArrayDim )
 					{
 						UE_LOG( LogNet, Error, TEXT( "Element index too large %s in %s" ), *ReplicatedProp->GetName(), *Object->GetFullName() );
 						return false;
@@ -569,36 +615,14 @@ bool FObjectReplicator::ReceivedBunch( FInBunch &Bunch, const FReplicationFlags&
 				QueuePropertyRepNotify( Object, ReplicatedProp, Element, MetaData );
 			}	
 			
-			// Next.
-			RepIndex = Bunch.ReadInt( ClassCache->GetMaxIndex() + 1 );
+			// Read next field
+			FieldCache = ReadField( ClassCache, Bunch );
 
 			if ( Bunch.IsError() )
 			{
-				UE_LOG( LogNet, Error, TEXT( "ReceivedBunch: Error reading bunch 2: %s" ), *Object->GetFullName() );
-				return false;
-			}
-
-			if ( RepIndex > ClassCache->GetMaxIndex() )
-			{
-				UE_LOG( LogNet, Error, TEXT( "ReceivedBunch: RepIndex too large: %s" ), *Object->GetFullName() );
+				UE_LOG( LogNet, Error, TEXT( "ReceivedBunch: Error reading field 2: %s" ), *Object->GetFullName() );
 				return false;
 			}
-			
-			if ( RepIndex == ClassCache->GetMaxIndex() )
-			{
-				// We're done
-				FieldCache = NULL;
-			}
-			else
-			{
-				FieldCache = ClassCache->GetFromIndex( RepIndex );
-
-				if ( FieldCache == NULL )
-				{
-					UE_LOG( LogNet, Error, TEXT( "ReceivedBunch: FieldCache == NULL: %s" ), *Object->GetFullName() );
-					return false;
-				}
-			}
 		}
 
 		// Handle function calls.
@@ -688,39 +712,17 @@ bool FObjectReplicator::ReceivedBunch( FInBunch &Bunch, const FReplicationFlags&
 			}
 
 			// Next.
-			RepIndex = Bunch.ReadInt( ClassCache->GetMaxIndex() + 1 );
+			FieldCache = ReadField( ClassCache, Bunch );
 
 			if ( Bunch.IsError() )
 			{
-				UE_LOG( LogNet, Error, TEXT( "ReceivedBunch: Error reading bunch 2: %s" ), *Object->GetFullName() );
+				UE_LOG( LogNet, Error, TEXT( "ReceivedBunch: Error reading field 3: %s" ), *Object->GetFullName() );
 				return false;
 			}
-
-			if ( RepIndex > ClassCache->GetMaxIndex() )
-			{
-				UE_LOG( LogNet, Error, TEXT( "ReceivedBunch: RepIndex too large: %s" ), *Object->GetFullName() );
-				return false;
-			}
-
-			if ( RepIndex == ClassCache->GetMaxIndex() )
-			{
-				// We're done
-				FieldCache = NULL;
-			}
-			else
-			{
-				FieldCache = ClassCache->GetFromIndex( RepIndex );
-
-				if ( FieldCache == NULL )
-				{
-					UE_LOG( LogNet, Error, TEXT( "ReceivedBunch: FieldCache == NULL: %s" ), *Object->GetFullName() );
-					return false;
-				}
-			}
 		}
 		else if ( FieldCache )
 		{
-			UE_LOG( LogNet, Error, TEXT( "ReceivedBunch: Invalid replicated field %i in %s" ), RepIndex, *Object->GetFullName() );
+			UE_LOG( LogNet, Error, TEXT( "ReceivedBunch: Invalid replicated field %i in %s" ), FieldCache->FieldNetIndex, *Object->GetFullName() );
 			return false;
 		}
 	}
@@ -783,7 +785,7 @@ static FORCEINLINE FPropertyRetirement ** UpdateAckedRetirements( FPropertyRetir
 	return Rec;
 }
 
-void FObjectReplicator::ReplicateCustomDeltaProperties( FOutBunch & Bunch, FReplicationFlags RepFlags, int32& LastIndex, bool & bContentBlockWritten )
+void FObjectReplicator::ReplicateCustomDeltaProperties( FOutBunch & Bunch, FReplicationFlags RepFlags, bool & bContentBlockWritten )
 {
 	if ( LifetimeCustomDeltaProperties.Num() == 0 )
 	{
@@ -876,7 +878,7 @@ void FObjectReplicator::ReplicateCustomDeltaProperties( FOutBunch & Bunch, FRepl
 		OldState = NewState; 
 
 		// Write header, and data to send to the actual bunch
-		RepLayout->WritePropertyHeader( Object, ObjectClass, OwningChannel, It, Bunch, Index, LastIndex, bContentBlockWritten );
+		RepLayout->WritePropertyHeader( Object, ObjectClass, OwningChannel, It, Bunch, Index, bContentBlockWritten );
 
 		const int NumStartingBits = Bunch.GetNumBits();
 
@@ -908,14 +910,13 @@ bool FObjectReplicator::ReplicateProperties( FOutBunch & Bunch, FReplicationFlag
 
 	const int32 StartingBitNum = Bunch.GetNumBits();
 
-	bool	bContentBlockWritten	= false;
-	int32	LastIndex				= 0;
+	bool bContentBlockWritten = false;
 
 	// Replicate all the custom delta properties (fast arrays, etc)
-	ReplicateCustomDeltaProperties( Bunch, RepFlags, LastIndex, bContentBlockWritten );
+	ReplicateCustomDeltaProperties( Bunch, RepFlags, bContentBlockWritten );
 
 	// Replicate properties in the layout
-	RepLayout->ReplicateProperties( RepState, (uint8*)Object, ObjectClass, OwningChannel, Bunch, RepFlags, LastIndex, bContentBlockWritten );
+	RepLayout->ReplicateProperties( RepState, (uint8*)Object, ObjectClass, OwningChannel, Bunch, RepFlags, bContentBlockWritten );
 
 	// LastUpdateEmpty - this is done before dequeing the multicasted unreliable functions on purpose as they should not prevent
 	// an actor channel from going dormant.
diff --git a/Engine/Source/Runtime/Engine/Private/DemoNetDriver.cpp b/Engine/Source/Runtime/Engine/Private/DemoNetDriver.cpp
--- a/Engine/Source/Runtime/Engine/Private/DemoNetDriver.cpp
+++ b/Engine/Source/Runtime/Engine/Private/DemoNetDriver.cpp
@@ -124,8 +124,35 @@ struct FNetworkDemoHeader
 	friend FArchive& operator << ( FArchive& Ar, FNetworkDemoHeader& Header )
 	{
 		Ar << Header.Magic;
+
+		// Check magic value
+		if ( Header.Magic != NETWORK_DEMO_MAGIC )
+		{
+			UE_LOG( LogDemo, Error, TEXT( "Header.Magic != NETWORK_DEMO_MAGIC" ) );
+			Ar.SetError();
+			return Ar;
+		}
+
 		Ar << Header.Version;
+
+		// Check version
+		if ( Header.Version != NETWORK_DEMO_VERSION )
+		{
+			UE_LOG( LogDemo, Error, TEXT( "Header.Version != NETWORK_DEMO_VERSION" ) );
+			Ar.SetError();
+			return Ar;
+		}
+
+		// Check internal version
 		Ar << Header.InternalProtocolVersion;
+
+		if ( Header.InternalProtocolVersion != FNetworkVersion::InternalProtocolVersion )
+		{
+			UE_LOG( LogDemo, Error, TEXT( "Header.InternalProtocolVersion != FNetworkVersion::InternalProtocolVersion" ) );
+			Ar.SetError();
+			return Ar;
+		}
+
 		Ar << Header.EngineNetVersion;
 		Ar << Header.LevelName;
 
@@ -226,8 +253,7 @@ bool UDemoNetDriver::InitConnectInternal( FString& Error )
 
 	(*FileAr) << DemoHeader;
 
-	// Check magic value
-	if ( DemoHeader.Magic != NETWORK_DEMO_MAGIC )
+	if ( FileAr->IsError() )
 	{
 		Error = FString( TEXT( "Demo file is corrupt" ) );
 		UE_LOG( LogDemo, Error, TEXT( "UDemoNetDriver::InitConnect: %s" ), *Error );
@@ -235,23 +261,6 @@ bool UDemoNetDriver::InitConnectInternal( FString& Error )
 		return false;
 	}
 
-	// Check version
-	if ( DemoHeader.Version != NETWORK_DEMO_VERSION )
-	{
-		Error = FString( TEXT( "Demo file version is incorrect" ) );
-		UE_LOG( LogDemo, Error, TEXT( "UDemoNetDriver::InitConnect: %s" ), *Error );
-		GameInstance->HandleDemoPlaybackFailure( EDemoPlayFailure::InvalidVersion, Error );
-		return false;
-	}
-
-	if ( DemoHeader.InternalProtocolVersion != FNetworkVersion::InternalProtocolVersion )
-	{
-		Error = FString( TEXT( "Engine internal network version is incorrect" ) );
-		UE_LOG( LogDemo, Error, TEXT( "UDemoNetDriver::InitConnect: %s" ), *Error );
-		GameInstance->HandleDemoPlaybackFailure( EDemoPlayFailure::InvalidVersion, Error );
-		return false;
-	}
-
 	// Create fake control channel
 	ServerConnection->CreateChannel( CHTYPE_Control, 1 );
 
diff --git a/Engine/Source/Runtime/Engine/Private/NetworkDriver.cpp b/Engine/Source/Runtime/Engine/Private/NetworkDriver.cpp
--- a/Engine/Source/Runtime/Engine/Private/NetworkDriver.cpp
+++ b/Engine/Source/Runtime/Engine/Private/NetworkDriver.cpp
@@ -711,15 +711,16 @@ void UNetDriver::InternalProcessRemoteFunction
 	UObject* TargetObj = SubObject ? SubObject : Actor;
 
 	// Make sure this function exists for both parties.
-	FClassNetCache* ClassCache = NetCache->GetClassNetCache( TargetObj->GetClass() );
+	const FClassNetCache* ClassCache = NetCache->GetClassNetCache( TargetObj->GetClass() );
 	if (!ClassCache)
 	{
 		DEBUG_REMOTEFUNCTION(TEXT("ClassNetCache empty, not calling %s::%s"), *Actor->GetName(), *Function->GetName());
 		return;
 	}
 		
-	FFieldNetCache* FieldCache = ClassCache->GetFromField( Function );
-	if (!FieldCache)
+	const FFieldNetCache* FieldCache = ClassCache->GetFromField( Function );
+
+	if ( !FieldCache )
 	{
 		DEBUG_REMOTEFUNCTION(TEXT("FieldCache empty, not calling %s::%s"), *Actor->GetName(), *Function->GetName());
 		return;
@@ -802,8 +803,16 @@ void UNetDriver::InternalProcessRemoteFunction
 	const int NumStartingHeaderBits = Bunch.GetNumBits();
 
 	//UE_LOG(LogScript, Log, TEXT("   Call %s"),Function->GetFullName());
-	check( FieldCache->FieldNetIndex <= ClassCache->GetMaxIndex() );
-	Bunch.WriteIntWrapped(FieldCache->FieldNetIndex, ClassCache->GetMaxIndex()+1);
+	if ( Connection->InternalAck )
+	{
+		uint32 Checksum = FieldCache->FieldChecksum;
+		Bunch << Checksum;
+	}
+	else
+	{
+		check( FieldCache->FieldNetIndex <= ClassCache->GetMaxIndex() );
+		Bunch.WriteIntWrapped(FieldCache->FieldNetIndex, ClassCache->GetMaxIndex()+1);
+	}
 
 	const int HeaderBits = Bunch.GetNumBits() - NumStartingHeaderBits;
 
@@ -1252,9 +1261,9 @@ bool UNetDriver::HandleNetDumpServerRPCCommand( const TCHAR* Cmd, FOutputDevice&
 
 			if ( Function != NULL && Function->FunctionFlags & FUNC_NetServer )
 			{
-				FClassNetCache * ClassCache = NetCache->GetClassNetCache( *ClassIt );
+				const FClassNetCache * ClassCache = NetCache->GetClassNetCache( *ClassIt );
 
-				FFieldNetCache * FieldCache = ClassCache->GetFromField( Function );
+				const FFieldNetCache * FieldCache = ClassCache->GetFromField( Function );
 
 				TArray< UProperty * > Parms;
 
diff --git a/Engine/Source/Runtime/Engine/Private/PackageMapClient.cpp b/Engine/Source/Runtime/Engine/Private/PackageMapClient.cpp
--- a/Engine/Source/Runtime/Engine/Private/PackageMapClient.cpp
+++ b/Engine/Source/Runtime/Engine/Private/PackageMapClient.cpp
@@ -1347,8 +1347,6 @@ void FNetGUIDCache::CleanReferences()
 		checkf( ObjectLookup.FindRef( It.Value() ).Object == It.Key(), TEXT("Failed to validate NetGUIDLookup map in UPackageMap. GUID '%s' was not in the ObjectLookup map with with object '%s'."), *It.Value().ToString(), *It.Key().Get()->GetPathName());
 	}
 
-	ClassChecksumMap.Empty();
-
 	FArchiveCountMemGUID CountBytesAr;
 
 	ObjectLookup.CountBytes( CountBytesAr );
@@ -2079,197 +2077,9 @@ void FNetGUIDCache::GenerateFullNetGUIDPath_r( const FNetworkGUID& NetGUID, FStr
 	}
 }
 
-static uint32 GetPropertyChecksum( const UProperty* Property, uint32 Checksum );
-
-static void SortProperties( TArray< UProperty* >& Properties )
-{
-	// Sort NetProperties so that their ClassReps are sorted by memory offset
-	struct FCompareUFieldOffsets
-	{
-		FORCEINLINE bool operator()( UProperty & A, UProperty & B ) const
-		{
-			// Ensure stable sort
-			if ( A.GetOffset_ForGC() == B.GetOffset_ForGC() )
-			{
-				return A.GetName() < B.GetName();
-			}
-
-			return A.GetOffset_ForGC() < B.GetOffset_ForGC();
-		}
-	};
-
-	Sort( Properties.GetData(), Properties.Num(), FCompareUFieldOffsets() );
-}
-
-static uint32 SortedStructFieldsChecksum( const UStruct* Struct, uint32 Checksum )
-{
-	// Generate a list that we can sort, to make sure we process these deterministically
-	TArray< UProperty * > Fields;
-
-	for ( TFieldIterator< UProperty > It( Struct ); It; ++It )
-	{
-		if ( It->PropertyFlags & CPF_RepSkip )
-		{
-			continue;
-		}
-
-		Fields.Add( *It );
-	}
-
-	// Sort them
-	SortProperties( Fields );
-
-	// Evolve the checksum on the sorted list
-	for ( auto Field : Fields )
-	{
-		Checksum = GetPropertyChecksum( Field, Checksum );
-	}
-
-	return Checksum;
-}
-
-static bool GDebugChecksum		= false;
-static int GDebugChecksumIndent = 0;
-
-static uint32 GetPropertyChecksum( const UProperty* Property, uint32 Checksum )
-{
-	if ( GDebugChecksum )
-	{
-		UE_LOG( LogNetPackageMap, Warning, TEXT( "%s%s [%s] [%u] [%u]" ), FCString::Spc( 2 * GDebugChecksumIndent ), *Property->GetName().ToLower(), *Property->GetClass()->GetName().ToLower(), Property->ArrayDim, Checksum );
-	}
-
-	// Evolve checksum on name
-	Checksum = FCrc::StrCrc32( *Property->GetName().ToLower(), Checksum );
-	
-	// Evolve checksum on class name
-	Checksum = FCrc::StrCrc32( *Property->GetClass()->GetName().ToLower(), Checksum );
-
-	// Evolve checksum on array dim (to detect when static arrays change size)
-	// We use a string to be endian agnostic
-	Checksum = FCrc::StrCrc32( *FString::Printf( TEXT( "%u" ), Property->ArrayDim ), Checksum );
-
-	const UArrayProperty* ArrayProperty = Cast< UArrayProperty >( Property );
-
-	// Evolve checksum on array inner
-	if ( ArrayProperty != NULL )
-	{
-		return GetPropertyChecksum( ArrayProperty->Inner, Checksum );
-	}
-
-	const UStructProperty* StructProperty = Cast< UStructProperty >( Property );
-
-	// Evolve checksum on property struct fields
-	if ( StructProperty != NULL )
-	{
-		if ( GDebugChecksum )
-		{
-			UE_LOG( LogNetPackageMap, Warning, TEXT( "%s [%s] [%u]" ), FCString::Spc( 2 * GDebugChecksumIndent ), *StructProperty->Struct->GetName().ToLower(), Checksum );
-		}
-
-		// Evolve checksum on struct name
-		Checksum = FCrc::StrCrc32( *StructProperty->Struct->GetName().ToLower(), Checksum );
-
-		GDebugChecksumIndent++;
-
-		Checksum = SortedStructFieldsChecksum( StructProperty->Struct, Checksum );
-
-		GDebugChecksumIndent--;
-	}
-
-	return Checksum;
-}
-
 uint32 FNetGUIDCache::GetClassNetworkChecksum( const UClass* Class )
 {
-	uint32 Checksum = GDebugChecksum ? 0 : ClassChecksumMap.FindRef( Class );
-
-	if ( Checksum != 0 )
-	{
-		return Checksum;
-	}
-
-	if ( Class->GetSuperClass() != NULL )
-	{
-		Checksum = GetClassNetworkChecksum( Class->GetSuperClass() );
-	}
-
-	GDebugChecksumIndent++;
-
-	Checksum = FCrc::StrCrc32( *Class->GetName().ToLower(), Checksum );
-
-	TArray< UProperty* > NetProperties;
-	TArray< UFunction* > MulticastFunctions;
-
-	// Grab all of the net properties from the replicated fields
-
-	for ( int32 i = 0; i < Class->NetFields.Num(); i++ )
-	{
-		UProperty* Property = Cast< UProperty >( Class->NetFields[i] );
-
-		if ( Property != NULL )
-		{
-			NetProperties.Add( Property );
-		}
-
-		UFunction* Function = Cast< UFunction >( Class->NetFields[i] );
-
-		if ( Function != NULL && ( Function->FunctionFlags & FUNC_NetMulticast ) )
-		{
-			MulticastFunctions.Add( Function );
-		}
-	}
-
-	// Sort properties by offset/name
-	SortProperties( NetProperties );
-
-	for ( int32 i = 0; i < NetProperties.Num(); i++ )
-	{
-		Checksum = GetPropertyChecksum( NetProperties[i], Checksum );
-	}
-
-	struct FCompareUFieldNames
-	{
-		FORCEINLINE bool operator()( UField& A, UField& B ) const
-		{
-			return A.GetName() < B.GetName();
-		}
-	};
-
-	// Sort multi-cast functions by name
-	Sort( MulticastFunctions.GetData(), MulticastFunctions.Num(), FCompareUFieldNames() );
-
-	// Grab all multi-cast RPC's (we don't need to worry about client RPC's, except for our spectator player controller, which we'll special case)
-	for ( UFunction* Function : MulticastFunctions )
-	{
-		// Evolve checksum on function name
-		Checksum = FCrc::StrCrc32( *Function->GetName().ToLower(), Checksum );
-
-		TArray< UProperty * > Parms;
-
-		for ( TFieldIterator< UProperty > It( Function ); It && ( It->PropertyFlags & ( CPF_Parm | CPF_ReturnParm ) ) == CPF_Parm; ++It )
-		{
-			Parms.Add( *It );
-		}
-		
-		// Sort parameters by offset/name
-		SortProperties( Parms );
-
-		// Evolve checksum on sorted function parameters
-		for ( UProperty* Parm : Parms )
-		{
-			Checksum = GetPropertyChecksum( Parm, Checksum );
-		}
-	}
-
-	// Can't be 0
-	if ( Checksum == 0 )
-	{
-		Checksum = 1;
-	}
-
-	ClassChecksumMap.Add( Class, Checksum );
-
-	return Checksum;
+	return Driver->NetCache->GetClassNetCache( Class )->GetClassChecksum();
 }
 
 uint32 FNetGUIDCache::GetNetworkChecksum( const UObject* Obj )
@@ -2279,26 +2089,10 @@ uint32 FNetGUIDCache::GetNetworkChecksum( const UObject* Obj )
 		return 0;
 	}
 
-	GDebugChecksum = false;
-	GDebugChecksumIndent = 0;
-
-#if 0
-	if ( Obj->GetName().Contains( TEXT( "Default__ShooterGameState" ) ) )
-	{
-		GDebugChecksum = true;
-
-		UE_LOG( LogNetPackageMap, Warning, TEXT( "OBJECT: %s" ), *Obj->GetName() );
-	}
-#endif
-
 	// If Obj is already a class, we can use that directly
 	const UClass* Class = Cast< UClass >( Obj );
 
-	const uint32 Checksum = ( Class == NULL ) ? GetClassNetworkChecksum( Obj->GetClass() ) : GetClassNetworkChecksum( Class );
-
-	GDebugChecksum = false;
-
-	return Checksum;
+	return ( Class != NULL ) ? GetClassNetworkChecksum( Class ) : GetClassNetworkChecksum( Obj->GetClass() );
 }
 
 bool FNetGUIDCache::ShouldIgnorePackageMismatch() const
diff --git a/Engine/Source/Runtime/Engine/Private/RepLayout.cpp b/Engine/Source/Runtime/Engine/Private/RepLayout.cpp
--- a/Engine/Source/Runtime/Engine/Private/RepLayout.cpp
+++ b/Engine/Source/Runtime/Engine/Private/RepLayout.cpp
@@ -592,7 +592,6 @@ bool FRepLayout::ReplicateProperties(
 	UActorChannel *				OwningChannel,
 	FOutBunch &					Writer, 
 	const FReplicationFlags &	RepFlags,
-	int32 &						LastIndex, 
 	bool &						bContentBlockWritten ) const
 {
 	SCOPE_CYCLE_COUNTER( STAT_NetReplicateDynamicPropTime );
@@ -750,10 +749,10 @@ bool FRepLayout::ReplicateProperties(
 #endif
 
 		// For RepLayout properties, we hijack the first non custom property, and use that to identify these properties
-		WritePropertyHeader( (UObject*)Data, ObjectClass, OwningChannel, Parents[FirstNonCustomParent].Property, Writer, 0, LastIndex, bContentBlockWritten );
+		WritePropertyHeader( (UObject*)Data, ObjectClass, OwningChannel, Parents[FirstNonCustomParent].Property, Writer, 0, bContentBlockWritten );
 
 		// Send the final merged change list
-		SendProperties( RepState, RepFlags, Data, ObjectClass, OwningChannel, Writer, Changed, LastIndex, bContentBlockWritten );
+		SendProperties( RepState, RepFlags, Data, ObjectClass, OwningChannel, Writer, Changed, bContentBlockWritten );
 
 #ifdef ENABLE_SUPER_CHECKSUMS
 		Writer.WriteBit( bIsAllAcked ? 1 : 0 );
@@ -1098,14 +1097,14 @@ void FRepLayout::WritePropertyHeader(
 	UActorChannel *		OwningChannel,
 	UProperty *			Property, 
 	FOutBunch &			Bunch, 
-	int32				ArrayIndex, 
-	int32 &				LastArrayIndex, 
+	uint32				ArrayIndex, 
 	bool &				bContentBlockWritten ) const
 {
 	UNetConnection * Connection = OwningChannel->Connection;
 
 	// Get class network info cache.
-	FClassNetCache * ClassCache = Connection->Driver->NetCache->GetClassNetCache( ObjectClass );
+	const FClassNetCache* ClassCache = Connection->Driver->NetCache->GetClassNetCache( ObjectClass );
+
 	check( ClassCache );
 
 	if ( !bContentBlockWritten )
@@ -1117,29 +1116,28 @@ void FRepLayout::WritePropertyHeader(
 	const int NumStartingBits = Bunch.GetNumBits();
 
 	// Get the network friend property index to replicate
-	// Swap Role/RemoteRole while we're at it
-	FFieldNetCache * FieldCache
-		=	Property->GetFName() == NAME_Role
-		?	ClassCache->GetFromField( Connection->Driver->RemoteRoleProperty )
-		:	Property->GetFName() == NAME_RemoteRole
-		?	ClassCache->GetFromField( Connection->Driver->RoleProperty )
-		:	ClassCache->GetFromField( Property );
+	const FFieldNetCache * FieldCache = ClassCache->GetFromField( Property );
 
 	checkSlow( FieldCache );
 
 	// Send property name and optional array index.
 	check( FieldCache->FieldNetIndex <= ClassCache->GetMaxIndex() );
 
-	Bunch.WriteIntWrapped( FieldCache->FieldNetIndex, ClassCache->GetMaxIndex() + 1 );
+	if ( Connection->InternalAck )
+	{
+		uint32 Checksum = FieldCache->FieldChecksum;
+		Bunch << Checksum;
+	}
+	else
+	{
+		Bunch.WriteIntWrapped( FieldCache->FieldNetIndex, ClassCache->GetMaxIndex() + 1 );
+	}
 
 	NET_CHECKSUM( Bunch );
 
 	if ( Property->ArrayDim != 1 )
 	{
-		// Serialize index as delta from previous index to increase chance we'll only use 1 byte
-		uint32 idx = static_cast<uint32>( ArrayIndex - LastArrayIndex );
-		Bunch.SerializeIntPacked(idx);
-		LastArrayIndex = ArrayIndex;
+		Bunch.SerializeIntPacked( ArrayIndex );
 	}
 
 	NETWORK_PROFILER(GNetworkProfiler.TrackWritePropertyHeader(Property, Bunch.GetNumBits() - NumStartingBits));
@@ -1153,7 +1151,6 @@ void FRepLayout::SendProperties(
 	UActorChannel	*			OwningChannel,
 	FOutBunch &					Writer, 
 	TArray< uint16 > &			Changed, 
-	int32 &						LastIndex, 
 	bool &						bContentBlockWritten ) const
 {
 #ifdef ENABLE_PROPERTY_CHECKSUMS
diff --git a/Engine/Source/Runtime/Engine/Private/UnrealNetwork.cpp b/Engine/Source/Runtime/Engine/Private/UnrealNetwork.cpp
--- a/Engine/Source/Runtime/Engine/Private/UnrealNetwork.cpp
+++ b/Engine/Source/Runtime/Engine/Private/UnrealNetwork.cpp
@@ -10,15 +10,16 @@ FNetworkVersion::FIsNetworkCompatibleOverride FNetworkVersion::IsNetworkCompatib
 
 enum ENetworkVersionHistory
 {
-	HISTORY_INITIAL				= 1,
-	HISTORY_INTERNAL_ACK		= 3,				// We no longer save packet/channel sequence in stream. We can derive this for 100% reliable connections.
-	HISTORY_REPLAY_CHECKSUMS	= 4,				// We now save replay checksums into stream for backwards compatibility checks
-	HISTORY_REPLAY_CHECKSUMS2	= 5,				// Saving package and network checksum to packagemap, no longer save checksum in BeginContentBlockHeader
-	HISTORY_REPLAY_CHECKSUMS3	= 6,				// Various network checksum fixes
-	HISTORY_REPLAY_CHECKSUMS4	= 7,				// Added multi-cast RPC's to network checksum
+	HISTORY_INITIAL					= 1,
+	HISTORY_INTERNAL_ACK			= 3,				// We no longer save packet/channel sequence in stream. We can derive this for 100% reliable connections.
+	HISTORY_REPLAY_CHECKSUMS		= 4,				// We now save replay checksums into stream for backwards compatibility checks
+	HISTORY_REPLAY_CHECKSUMS2		= 5,				// Saving package and network checksum to packagemap, no longer save checksum in BeginContentBlockHeader
+	HISTORY_REPLAY_CHECKSUMS3		= 6,				// Various network checksum fixes
+	HISTORY_REPLAY_CHECKSUMS4		= 7,				// Added multi-cast RPC's to network checksum
+	HISTORY_REPLAY_RPC_CHECKSUMS	= 8,				// RPC's save checksums separately now
 };
 
-const uint32 FNetworkVersion::InternalProtocolVersion = HISTORY_REPLAY_CHECKSUMS4;
+const uint32 FNetworkVersion::InternalProtocolVersion = HISTORY_REPLAY_RPC_CHECKSUMS;
 
 uint32 FNetworkVersion::GetLocalNetworkVersion()
 {
diff --git a/Engine/Source/Runtime/Engine/Public/Net/DataReplication.h b/Engine/Source/Runtime/Engine/Public/Net/DataReplication.h
--- a/Engine/Source/Runtime/Engine/Public/Net/DataReplication.h
+++ b/Engine/Source/Runtime/Engine/Public/Net/DataReplication.h
@@ -111,10 +111,12 @@ public:
 	void	Serialize(FArchive& Ar);
 
 	/** Writes dirty properties to bunch */
-	void	ReplicateCustomDeltaProperties( FOutBunch & Bunch, FReplicationFlags RepFlags, int32& LastIndex, bool & bContentBlockWritten );
+	void	ReplicateCustomDeltaProperties( FOutBunch & Bunch, FReplicationFlags RepFlags, bool & bContentBlockWritten );
 	bool	ReplicateProperties( FOutBunch & Bunch, FReplicationFlags RepFlags );
 	void	PostSendBunch(FPacketIdRange & PacketRange, uint8 bReliable);
 	
+	const FFieldNetCache* ReadField( const FClassNetCache* ClassCache, FInBunch& Bunch ) const;
+
 	bool	ReceivedBunch( FInBunch & Bunch, const FReplicationFlags& RepFlags, bool & bOutHasUnmapped );
 	void	PostReceivedBunch();
 
diff --git a/Engine/Source/Runtime/Engine/Public/Net/RepLayout.h b/Engine/Source/Runtime/Engine/Public/Net/RepLayout.h
--- a/Engine/Source/Runtime/Engine/Public/Net/RepLayout.h
+++ b/Engine/Source/Runtime/Engine/Public/Net/RepLayout.h
@@ -262,8 +262,7 @@ public:
 		UActorChannel *		OwningChannel,
 		UProperty *			Property, 
 		FOutBunch &			Bunch, 
-		int32				ArrayIndex, 
-		int32 &				LastArrayIndex, 
+		uint32				ArrayIndex, 
 		bool &				bContentBlockWritten ) const;
 
 	bool ReplicateProperties( 
@@ -273,7 +272,6 @@ public:
 		UActorChannel *				OwningChannel,
 		FOutBunch &					Writer, 
 		const FReplicationFlags &	RepFlags,
-		int32 &						LastIndex, 
 		bool &						bContentBlockWritten ) const;
 
 	void SendProperties( 
@@ -284,7 +282,6 @@ public:
 		UActorChannel *				OwningChannel,
 		FOutBunch &					Writer, 
 		TArray< uint16 >	 &		Changed, 
-		int32 &						LastIndex, 
 		bool &						bContentBlockWritten ) const;
 
 	ENGINE_API void InitFromObjectClass( UClass * InObjectClass );
