Engine - Changing the category of Blendables back to Misc.

diff --git a/Engine/Source/Runtime/Engine/Classes/Engine/Scene.h b/Engine/Source/Runtime/Engine/Classes/Engine/Scene.h
--- a/Engine/Source/Runtime/Engine/Classes/Engine/Scene.h
+++ b/Engine/Source/Runtime/Engine/Classes/Engine/Scene.h
@@ -883,7 +883,7 @@ struct FPostProcessSettings
 	 * Allows custom post process materials to be defined, using a MaterialInstance with the same Material as its parent to allow blending.
 	 * Make sure you use the "PostProcess" domain type. This can be used for any UObject object implementing the IBlendableInterface (e.g. could be used to fade weather settings).
 	 */
-	UPROPERTY(EditAnywhere, BlueprintReadWrite, Category="Custom Post Process Effects", meta=( AllowedClasses="BlendableInterface" ))
+	UPROPERTY(EditAnywhere, BlueprintReadWrite, Category="Misc", meta=( AllowedClasses="BlendableInterface", Keywords="PostProcess" ))
 	TArray<UObject*> Blendables;
 
 	// good start values for a new volume, by default no value is overriding
