Use unscaled view rect for calculating ScreenToPixel and PixelToScreen

diff --git a/Engine/Source/Runtime/Engine/Private/SceneView.cpp b/Engine/Source/Runtime/Engine/Private/SceneView.cpp
--- a/Engine/Source/Runtime/Engine/Private/SceneView.cpp
+++ b/Engine/Source/Runtime/Engine/Private/SceneView.cpp
@@ -598,8 +598,8 @@ bool FSceneView::ScreenToPixel(const FVector4& ScreenPoint,FVector2D& OutPixelLo
 		float InvW = 1.0f / ScreenPoint.W;
 		float Y = (GProjectionSignY > 0.0f) ? ScreenPoint.Y : 1.0f - ScreenPoint.Y;
 		OutPixelLocation = FVector2D(
-			ViewRect.Min.X + (0.5f + ScreenPoint.X * 0.5f * InvW) * ViewRect.Width(),
-			ViewRect.Min.Y + (0.5f - Y * 0.5f * InvW) * ViewRect.Height()
+			UnscaledViewRect.Min.X + (0.5f + ScreenPoint.X * 0.5f * InvW) * UnscaledViewRect.Width(),
+			UnscaledViewRect.Min.Y + (0.5f - Y * 0.5f * InvW) * UnscaledViewRect.Height()
 			);
 		return true;
 	}
@@ -614,8 +614,8 @@ FVector4 FSceneView::PixelToScreen(float InX,float InY,float Z) const
 	if (GProjectionSignY > 0.0f)
 	{
 		return FVector4(
-			-1.0f + InX / ViewRect.Width() * +2.0f,
-			+1.0f + InY / ViewRect.Height() * -2.0f,
+			-1.0f + InX / UnscaledViewRect.Width() * +2.0f,
+			+1.0f + InY / UnscaledViewRect.Height() * -2.0f,
 			Z,
 			1
 			);
@@ -623,8 +623,8 @@ FVector4 FSceneView::PixelToScreen(float InX,float InY,float Z) const
 	else
 	{
 		return FVector4(
-			-1.0f + InX / ViewRect.Width() * +2.0f,
-			1.0f - (+1.0f + InY / ViewRect.Height() * -2.0f),
+			-1.0f + InX / UnscaledViewRect.Width() * +2.0f,
+			1.0f - (+1.0f + InY / UnscaledViewRect.Height() * -2.0f),
 			Z,
 			1
 			);
