fixed minor potential issue

diff --git a/Engine/Source/Runtime/Engine/Private/GPUSkinVertexFactory.cpp b/Engine/Source/Runtime/Engine/Private/GPUSkinVertexFactory.cpp
--- a/Engine/Source/Runtime/Engine/Private/GPUSkinVertexFactory.cpp
+++ b/Engine/Source/Runtime/Engine/Private/GPUSkinVertexFactory.cpp
@@ -435,7 +435,7 @@ public:
 				{
 					static const auto MotionBlurDebugVar = IConsoleManager::Get().FindTConsoleVariableDataInt(TEXT("r.MotionBlurDebug"));
 
-					if (MotionBlurDebugVar->GetValueOnRenderThread())
+					if (MotionBlurDebugVar && MotionBlurDebugVar->GetValueOnRenderThread())
 					{
 						UE_LOG(LogEngine, Log, TEXT("%s"), *ShaderData.GetDebugString(VertexFactory, OldBoneDataIndex));
 					}
