Fix linux build

diff --git a/Engine/Source/Runtime/Engine/Private/PhysicsEngine/BodyInstance.cpp b/Engine/Source/Runtime/Engine/Private/PhysicsEngine/BodyInstance.cpp
--- a/Engine/Source/Runtime/Engine/Private/PhysicsEngine/BodyInstance.cpp
+++ b/Engine/Source/Runtime/Engine/Private/PhysicsEngine/BodyInstance.cpp
@@ -1444,7 +1444,7 @@ struct FInitBodiesHelper
 		TArray<PxActor*> PDynamicActors;
 
 		// No reason to defer if there's an existing aggregate.
-		const bool bCanDefer = bCompileStatic || bDefer && !InAggregate;
+		const bool bCanDefer = bCompileStatic || (bDefer && !InAggregate);
 		bool bDynamicsUseAsync = false;
 		if (CreateShapesAndActors_PhysX(PSyncActors, PAsyncActors, PDynamicActors, bCanDefer, bDynamicsUseAsync) == false)
 		{
