#ue4 - Perf: Optimized CapsuleComponent::GetCollisionShape()

diff --git a/Engine/Source/Runtime/Engine/Private/Components/CapsuleComponent.cpp b/Engine/Source/Runtime/Engine/Private/Components/CapsuleComponent.cpp
--- a/Engine/Source/Runtime/Engine/Private/Components/CapsuleComponent.cpp
+++ b/Engine/Source/Runtime/Engine/Private/Components/CapsuleComponent.cpp
@@ -179,8 +179,9 @@ bool UCapsuleComponent::IsZeroExtent() const
 
 FCollisionShape UCapsuleComponent::GetCollisionShape(float Inflation) const
 {
-	const float Radius = FMath::Max(0.f, GetScaledCapsuleRadius() + Inflation);
-	const float HalfHeight = FMath::Max(0.f, GetScaledCapsuleHalfHeight() + Inflation);
+	const float ShapeScale = GetShapeScale();
+	const float Radius = FMath::Max(0.f, (CapsuleRadius * ShapeScale) + Inflation);
+	const float HalfHeight = FMath::Max(0.f, (CapsuleHalfHeight * ShapeScale) + Inflation);
 	return FCollisionShape::MakeCapsule(Radius, HalfHeight);
 }
 
