#ue4 - Perf: Optimized CharacterMovementComponent::ApplyRepulsionForce().

- Avoids useless work in the case of no overlaps
- Uses faster Cast<> rather than IsA().

diff --git a/Engine/Source/Runtime/Engine/Private/Components/CharacterMovementComponent.cpp b/Engine/Source/Runtime/Engine/Private/Components/CharacterMovementComponent.cpp
--- a/Engine/Source/Runtime/Engine/Private/Components/CharacterMovementComponent.cpp
+++ b/Engine/Source/Runtime/Engine/Private/Components/CharacterMovementComponent.cpp
@@ -6761,89 +6761,90 @@ void UCharacterMovementComponent::ApplyRepulsionForce(float DeltaSeconds)
 {
 	if (UpdatedPrimitive && RepulsionForce > 0.0f)
 	{
-		FCollisionQueryParams QueryParams;
-		QueryParams.bReturnFaceIndex = false;
-		QueryParams.bReturnPhysicalMaterial = false;
-
-		const FCollisionShape CollisionShape = UpdatedPrimitive->GetCollisionShape();
-		const float CapsuleRadius = CollisionShape.GetCapsuleRadius();
-		const float CapsuleHalfHeight = CollisionShape.GetCapsuleHalfHeight();
-		const float RepulsionForceRadius = CapsuleRadius * 1.2f;
-		const float StopBodyDistance = 2.5f;
-
 		const TArray<FOverlapInfo>& Overlaps = UpdatedPrimitive->GetOverlapInfos();
-		const FVector MyLocation = UpdatedPrimitive->GetComponentLocation();
-
-		for (int32 i=0; i < Overlaps.Num(); i++)
+		if (Overlaps.Num() > 0)
 		{
-			const FOverlapInfo& Overlap = Overlaps[i];
+			FCollisionQueryParams QueryParams;
+			QueryParams.bReturnFaceIndex = false;
+			QueryParams.bReturnPhysicalMaterial = false;
 
-			UPrimitiveComponent* OverlapComp = Overlap.OverlapInfo.Component.Get();
-			if (!OverlapComp || OverlapComp->Mobility < EComponentMobility::Movable)
-			{ 
-				continue; 
-			}
+			float CapsuleRadius = 0.f;
+			float CapsuleHalfHeight = 0.f;
+			CharacterOwner->GetCapsuleComponent()->GetScaledCapsuleSize(CapsuleRadius, CapsuleHalfHeight);
+			const float RepulsionForceRadius = CapsuleRadius * 1.2f;
+			const float StopBodyDistance = 2.5f;
+			const FVector MyLocation = UpdatedPrimitive->GetComponentLocation();
 
-			FName BoneName = NAME_None;
-			if (Overlap.GetBodyIndex() != INDEX_NONE && OverlapComp->IsA(USkinnedMeshComponent::StaticClass()))
+			for (int32 i=0; i < Overlaps.Num(); i++)
 			{
-				BoneName = ((USkinnedMeshComponent*)OverlapComp)->GetBoneName(Overlap.GetBodyIndex());
-			}
+				const FOverlapInfo& Overlap = Overlaps[i];
 
-			// Use the body instead of the component for cases where we have multi-body overlaps enabled
-			FBodyInstance* OverlapBody = OverlapComp->GetBodyInstance(BoneName);
+				UPrimitiveComponent* OverlapComp = Overlap.OverlapInfo.Component.Get();
+				if (!OverlapComp || OverlapComp->Mobility < EComponentMobility::Movable)
+				{ 
+					continue; 
+				}
 
-			if (!OverlapBody)
-			{
-				UE_LOG(LogCharacterMovement, Warning, TEXT("%s could not find overlap body for bone %s"), *GetName(), *BoneName.ToString());
-				continue;
-			}
+				FName BoneName = NAME_None;
+				if (Overlap.GetBodyIndex() != INDEX_NONE && Cast<USkinnedMeshComponent>(OverlapComp))
+				{
+					BoneName = ((USkinnedMeshComponent*)OverlapComp)->GetBoneName(Overlap.GetBodyIndex());
+				}
 
-			// Early out if this is not a destructible and the body is not simulated
-			bool bIsCompDestructible = OverlapComp->IsA(UDestructibleComponent::StaticClass());
-			if (!bIsCompDestructible && !OverlapBody->IsInstanceSimulatingPhysics())
-			{
-				continue;
-			}
+				// Use the body instead of the component for cases where we have multi-body overlaps enabled
+				FBodyInstance* OverlapBody = OverlapComp->GetBodyInstance(BoneName);
 
-			FTransform BodyTransform = OverlapBody->GetUnrealWorldTransform();
+				if (!OverlapBody)
+				{
+					UE_LOG(LogCharacterMovement, Warning, TEXT("%s could not find overlap body for bone %s"), *GetName(), *BoneName.ToString());
+					continue;
+				}
 
-			FVector BodyVelocity = OverlapBody->GetUnrealWorldVelocity();
-			FVector BodyLocation = BodyTransform.GetLocation();
+				// Early out if this is not a destructible and the body is not simulated
+				if (!OverlapBody->IsInstanceSimulatingPhysics() && !Cast<UDestructibleComponent>(OverlapComp))
+				{
+					continue;
+				}
 
-			// Trace to get the hit location on the capsule
-			FHitResult Hit;
-			bool bHasHit = UpdatedPrimitive->LineTraceComponent(Hit, BodyLocation,
-																FVector(MyLocation.X, MyLocation.Y, BodyLocation.Z),
-																QueryParams);
+				FTransform BodyTransform = OverlapBody->GetUnrealWorldTransform();
 
-			FVector HitLoc = Hit.ImpactPoint;
-			bool bIsPenetrating = Hit.bStartPenetrating || Hit.PenetrationDepth > 2.5f;
+				FVector BodyVelocity = OverlapBody->GetUnrealWorldVelocity();
+				FVector BodyLocation = BodyTransform.GetLocation();
 
-			// If we didn't hit the capsule, we're inside the capsule
-			if(!bHasHit) 
-			{ 
-				HitLoc = BodyLocation; 
-				bIsPenetrating = true;
-			}
+				// Trace to get the hit location on the capsule
+				FHitResult Hit;
+				bool bHasHit = UpdatedPrimitive->LineTraceComponent(Hit, BodyLocation,
+																	FVector(MyLocation.X, MyLocation.Y, BodyLocation.Z),
+																	QueryParams);
 
-			const float DistanceNow = (HitLoc - BodyLocation).SizeSquared2D();
-			const float DistanceLater = (HitLoc - (BodyLocation + BodyVelocity * DeltaSeconds)).SizeSquared2D();
+				FVector HitLoc = Hit.ImpactPoint;
+				bool bIsPenetrating = Hit.bStartPenetrating || Hit.PenetrationDepth > 2.5f;
 
-			if (BodyLocation.SizeSquared() > 0.1f && bHasHit && DistanceNow < StopBodyDistance && !bIsPenetrating)
-			{
-				OverlapBody->SetLinearVelocity(FVector(0.0f, 0.0f, 0.0f), false);
-			}
-			else if (DistanceLater <= DistanceNow || bIsPenetrating)
-			{
-				FVector ForceCenter(MyLocation.X, MyLocation.Y, bHasHit ? HitLoc.Z : MyLocation.Z);
+				// If we didn't hit the capsule, we're inside the capsule
+				if(!bHasHit) 
+				{ 
+					HitLoc = BodyLocation; 
+					bIsPenetrating = true;
+				}
+
+				const float DistanceNow = (HitLoc - BodyLocation).SizeSquared2D();
+				const float DistanceLater = (HitLoc - (BodyLocation + BodyVelocity * DeltaSeconds)).SizeSquared2D();
 
-				if (!bHasHit)
+				if (BodyLocation.SizeSquared() > 0.1f && bHasHit && DistanceNow < StopBodyDistance && !bIsPenetrating)
 				{
-					ForceCenter.Z = FMath::Clamp(BodyLocation.Z, MyLocation.Z - CapsuleHalfHeight, MyLocation.Z + CapsuleHalfHeight);
+					OverlapBody->SetLinearVelocity(FVector(0.0f, 0.0f, 0.0f), false);
 				}
+				else if (DistanceLater <= DistanceNow || bIsPenetrating)
+				{
+					FVector ForceCenter(MyLocation.X, MyLocation.Y, bHasHit ? HitLoc.Z : MyLocation.Z);
 
-				OverlapBody->AddRadialForceToBody(ForceCenter, RepulsionForceRadius, RepulsionForce * Mass, ERadialImpulseFalloff::RIF_Constant);
+					if (!bHasHit)
+					{
+						ForceCenter.Z = FMath::Clamp(BodyLocation.Z, MyLocation.Z - CapsuleHalfHeight, MyLocation.Z + CapsuleHalfHeight);
+					}
+
+					OverlapBody->AddRadialForceToBody(ForceCenter, RepulsionForceRadius, RepulsionForce * Mass, ERadialImpulseFalloff::RIF_Constant);
+				}
 			}
 		}
 	}
