#ue4 - Perf: Optimize UPrimitiveComponent::MoveComponentImpl() to avoid GetOwner() call if possible.

diff --git a/Engine/Source/Runtime/Engine/Private/Components/PrimitiveComponent.cpp b/Engine/Source/Runtime/Engine/Private/Components/PrimitiveComponent.cpp
--- a/Engine/Source/Runtime/Engine/Private/Components/PrimitiveComponent.cpp
+++ b/Engine/Source/Runtime/Engine/Private/Components/PrimitiveComponent.cpp
@@ -1334,8 +1334,6 @@ bool UPrimitiveComponent::MoveComponentImpl( const FVector& Delta, const FQuat&
 		return false;
 	}
 
-	AActor* const Actor = GetOwner();
-
 	// static things can move before they are registered (e.g. immediately after streaming), but not after.
 	static const FText WarnText = LOCTEXT("InvalidMove", "move");
 	if (CheckStaticMobilityAndWarn(WarnText))
@@ -1382,7 +1380,8 @@ bool UPrimitiveComponent::MoveComponentImpl( const FVector& Delta, const FQuat&
 	TArray<FOverlapInfo> PendingOverlaps;
 	TArray<FOverlapInfo> OverlapsAtEndLocation;
 	TArray<FOverlapInfo>* OverlapsAtEndLocationPtr = NULL; // When non-null, used as optimization to avoid work in UpdateOverlaps.
-	
+	AActor* CachedActor = nullptr; // Try to defer or avoid GetOwner() call.
+
 	if ( !bSweep )
 	{
 		// not sweeping, just go directly to the new transform
@@ -1392,6 +1391,8 @@ bool UPrimitiveComponent::MoveComponentImpl( const FVector& Delta, const FQuat&
 	{
 		TArray<FHitResult> Hits;
 		FVector NewLocation = TraceStart;
+		CachedActor = CachedActor ? CachedActor : GetOwner();
+		AActor* const Actor = CachedActor;
 
 		// Perform movement collision checking if needed for this actor.
 		const bool bCollisionEnabled = IsCollisionEnabled();
@@ -1611,7 +1612,8 @@ bool UPrimitiveComponent::MoveComponentImpl( const FVector& Delta, const FQuat&
 		}
 		else
 		{
-			DispatchBlockingHit(*Actor, BlockingHit);
+			CachedActor = CachedActor ? CachedActor : GetOwner();
+			DispatchBlockingHit(*CachedActor, BlockingHit);
 		}
 	}
 
@@ -1620,9 +1622,9 @@ bool UPrimitiveComponent::MoveComponentImpl( const FVector& Delta, const FQuat&
 	const float MSec = FPlatformTime::ToMilliseconds(MoveCompTakingLongTime);
 	if( MSec > PERF_SHOW_MOVECOMPONENT_TAKING_LONG_TIME_AMOUNT )
 	{
-		if (Actor)
+		if (GetOwner())
 		{
-			UE_LOG(LogPrimitiveComponent, Log, TEXT("%10f executing MoveComponent for %s owned by %s"), MSec, *GetName(), *Actor->GetFullName() );
+			UE_LOG(LogPrimitiveComponent, Log, TEXT("%10f executing MoveComponent for %s owned by %s"), MSec, *GetName(), *GetOwner()->GetFullName() );
 		}
 		else
 		{
