Fix non-editor builds.

diff --git a/Engine/Source/Runtime/Engine/Classes/Engine/DeveloperSettings.h b/Engine/Source/Runtime/Engine/Classes/Engine/DeveloperSettings.h
--- a/Engine/Source/Runtime/Engine/Classes/Engine/DeveloperSettings.h
+++ b/Engine/Source/Runtime/Engine/Classes/Engine/DeveloperSettings.h
@@ -30,10 +30,12 @@ public:
 	/** The unique name for your section of settings, uses the class's FName. */
 	FName GetSectionName() const;
 
+#if WITH_EDITOR
 	/** Gets the section text, uses the classes DisplayName by default. */
 	virtual FText GetSectionText() const;
 	/** Gets the description for the section, uses the classes ToolTip by default. */
 	virtual FText GetSectionDescription() const;
+#endif
 
 	/** Gets a custom widget for the settings.  This is only for very custom situations. */
 	virtual TSharedPtr<SWidget> GetCustomSettingsWidget() const;
diff --git a/Engine/Source/Runtime/Engine/Private/DeveloperSettings.cpp b/Engine/Source/Runtime/Engine/Private/DeveloperSettings.cpp
--- a/Engine/Source/Runtime/Engine/Private/DeveloperSettings.cpp
+++ b/Engine/Source/Runtime/Engine/Private/DeveloperSettings.cpp
@@ -56,6 +56,7 @@ FName UDeveloperSettings::GetSectionName() const
 	return GetClass()->GetFName();
 }
 
+#if WITH_EDITOR
 FText UDeveloperSettings::GetSectionText() const
 {
 	return GetClass()->GetDisplayNameText();
@@ -65,6 +66,7 @@ FText UDeveloperSettings::GetSectionDescription() const
 {
 	return GetClass()->GetToolTipText();
 }
+#endif
 
 TSharedPtr<SWidget> UDeveloperSettings::GetCustomSettingsWidget() const
 {
