#ue4 - Perf: don't zero array memory that will get immediately overwritten by physics shapes.

diff --git a/Engine/Source/Runtime/Engine/Private/Collision/WorldCollision.cpp b/Engine/Source/Runtime/Engine/Private/Collision/WorldCollision.cpp
--- a/Engine/Source/Runtime/Engine/Private/Collision/WorldCollision.cpp
+++ b/Engine/Source/Runtime/Engine/Private/Collision/WorldCollision.cpp
@@ -506,11 +506,9 @@ bool UWorld::ComponentSweepMulti(TArray<struct FHitResult>& OutHits, class UPrim
 	ExecuteOnPxRigidActorReadOnly(&PrimComp->BodyInstance, [&] (const PxRigidActor* PRigidActor)
 	{
 		// Get all the shapes from the actor
-		TArray<PxShape*, TInlineAllocator<32>> PShapes;
-		{
-			PShapes.AddZeroed(PRigidActor->getNbShapes());
-			PRigidActor->getShapes(PShapes.GetData(), PShapes.Num());
-		}
+		TArray<PxShape*, TInlineAllocator<16>> PShapes;
+		PShapes.AddUninitialized(PRigidActor->getNbShapes());
+		PRigidActor->getShapes(PShapes.GetData(), PShapes.Num());
 
 		// calculate the test global pose of the actor
 		const PxQuat PGeomRot = U2PQuat(Quat);
diff --git a/Engine/Source/Runtime/Engine/Private/PhysicsEngine/BodyInstance.cpp b/Engine/Source/Runtime/Engine/Private/PhysicsEngine/BodyInstance.cpp
--- a/Engine/Source/Runtime/Engine/Private/PhysicsEngine/BodyInstance.cpp
+++ b/Engine/Source/Runtime/Engine/Private/PhysicsEngine/BodyInstance.cpp
@@ -332,7 +332,7 @@ int32 FBodyInstance::GetAllShapes_AssumesLocked(TArray<PxShape*>& OutShapes) con
 	if(RigidActorSync)
 	{
 		NumSyncShapes = RigidActorSync->getNbShapes();
-		OutShapes.AddZeroed(NumSyncShapes);
+		OutShapes.AddUninitialized(NumSyncShapes);
 		RigidActorSync->getShapes(OutShapes.GetData(), NumSyncShapes);
 	}
 
@@ -340,7 +340,7 @@ int32 FBodyInstance::GetAllShapes_AssumesLocked(TArray<PxShape*>& OutShapes) con
 	if( RigidActorAsync != NULL && !HasSharedShapes())
 	{
 		const int32 NumAsyncShapes = RigidActorAsync->getNbShapes();
-		OutShapes.AddZeroed(NumAsyncShapes);
+		OutShapes.AddUninitialized(NumAsyncShapes);
 		RigidActorAsync->getShapes(OutShapes.GetData() + NumSyncShapes, NumAsyncShapes);
 	}
 
@@ -3090,8 +3090,8 @@ void FBodyInstance::GetComplexPhysicalMaterials(const FBodyInstance*, TWeakObjec
 /** Util for finding the number of 'collision sim' shapes on this PxRigidActor */
 int32 GetNumSimShapes_AssumesLocked(const PxRigidBody* PRigidBody)
 {
-	TArray<PxShape*> PShapes;
-	PShapes.AddZeroed(PRigidBody->getNbShapes());
+	TArray<PxShape*, TInlineAllocator<16>> PShapes;
+	PShapes.AddUninitialized(PRigidBody->getNbShapes());
 	int32 NumShapes = PRigidBody->getShapes(PShapes.GetData(), PShapes.Num());
 
 	int32 NumSimShapes = 0;
@@ -3763,7 +3763,7 @@ bool FBodyInstance::LineTrace(struct FHitResult& OutHit, const FVector& Start, c
 
 			// Get all the shapes from the actor
 			TArray<PxShape*, TInlineAllocator<16>> PShapes;
-			PShapes.AddZeroed(RigidBody->getNbShapes());
+			PShapes.AddUninitialized(RigidBody->getNbShapes());
 			int32 NumShapes = RigidBody->getShapes(PShapes.GetData(), PShapes.Num());
 
 			// Iterate over each shape
@@ -3891,7 +3891,7 @@ bool FBodyInstance::InternalSweepPhysX(struct FHitResult& OutHit, const FVector&
 
 		// Get all the shapes from the actor
 		TArray<PxShape*, TInlineAllocator<16>> PShapes;
-		PShapes.AddZeroed(RigidBody->getNbShapes());
+		PShapes.AddUninitialized(RigidBody->getNbShapes());
 		int32 NumShapes = RigidBody->getShapes(PShapes.GetData(), PShapes.Num());
 
 		// Iterate over each shape
@@ -3955,7 +3955,7 @@ float FBodyInstance::GetDistanceToBody(const FVector& Point, FVector& OutPointOn
 
 		// Get all the shapes from the actor
 		TArray<PxShape*, TInlineAllocator<16>> PShapes;
-		PShapes.AddZeroed(RigidActor->getNbShapes());
+		PShapes.AddUninitialized(RigidActor->getNbShapes());
 		int32 NumShapes = RigidActor->getShapes(PShapes.GetData(), PShapes.Num());
 
 		const PxVec3 PPoint = U2PVector(Point);
@@ -4014,8 +4014,8 @@ bool FBodyInstance::OverlapTestForBodies(const FVector& Pos, const FQuat& Rot, c
 		PxTransform PTestGlobalPose = U2PTransform(FTransform(Rot, Pos));
 
 		// Get all the shapes from the actor
-		TArray<PxShape*> PTargetShapes;
-		PTargetShapes.AddZeroed(TargetRigidBody->getNbShapes());
+		TArray<PxShape*, TInlineAllocator<16>> PTargetShapes;
+		PTargetShapes.AddUninitialized(TargetRigidBody->getNbShapes());
 		int32 NumTargetShapes = TargetRigidBody->getShapes(PTargetShapes.GetData(), PTargetShapes.Num());
 
 		for (int32 TargetShapeIdx = 0; TargetShapeIdx < PTargetShapes.Num(); ++TargetShapeIdx)
@@ -4126,7 +4126,7 @@ bool FBodyInstance::OverlapMulti(TArray<struct FOverlapResult>& InOutOverlaps, c
 		{
 		// Get all the shapes from the actor
 			TArray<PxShape*, TInlineAllocator<16>> PShapes;
-			PShapes.AddZeroed(PRigidActor->getNbShapes());
+			PShapes.AddUninitialized(PRigidActor->getNbShapes());
 			int32 NumShapes = PRigidActor->getShapes(PShapes.GetData(), PShapes.Num());
 
 			// Iterate over each shape
@@ -4184,7 +4184,7 @@ bool FBodyInstance::OverlapPhysX_AssumesLocked(const PxGeometry& PGeom, const Px
 
 	// Get all the shapes from the actor
 	TArray<PxShape*, TInlineAllocator<16>> PShapes;
-	PShapes.AddZeroed(RigidBody->getNbShapes());
+	PShapes.AddUninitialized(RigidBody->getNbShapes());
 	int32 NumShapes = RigidBody->getShapes(PShapes.GetData(), PShapes.Num());
 
 	// Iterate over each shape
diff --git a/Engine/Source/Runtime/Landscape/Private/LandscapeCollision.cpp b/Engine/Source/Runtime/Landscape/Private/LandscapeCollision.cpp
--- a/Engine/Source/Runtime/Landscape/Private/LandscapeCollision.cpp
+++ b/Engine/Source/Runtime/Landscape/Private/LandscapeCollision.cpp
@@ -1005,7 +1005,7 @@ void ULandscapeHeightfieldCollisionComponent::UpdateHeightfieldRegion(int32 Comp
 		if (BodyInstance.RigidActorSync)
 		{
 			TArray<PxShape*, TInlineAllocator<8>> PShapes;
-			PShapes.AddZeroed(BodyInstance.RigidActorSync->getNbShapes());
+			PShapes.AddUninitialized(BodyInstance.RigidActorSync->getNbShapes());
 			int32 NumShapes = BodyInstance.RigidActorSync->getShapes(PShapes.GetData(), PShapes.Num());
 			if (NumShapes > 1)
 			{
