#ue4 - Fix earlier change (CL 2507146) causing CheckHitResultNormal() not to run outside editor builds.

diff --git a/Engine/Source/Runtime/Engine/Private/Collision/CollisionConversions.cpp b/Engine/Source/Runtime/Engine/Private/Collision/CollisionConversions.cpp
--- a/Engine/Source/Runtime/Engine/Private/Collision/CollisionConversions.cpp
+++ b/Engine/Source/Runtime/Engine/Private/Collision/CollisionConversions.cpp
@@ -31,8 +31,9 @@ DECLARE_CYCLE_STAT(TEXT("ConvertQueryHit"), STAT_ConvertQueryImpactHit, STATGROU
 DECLARE_CYCLE_STAT(TEXT("ConvertOverlapToHit"), STAT_CollisionConvertOverlapToHit, STATGROUP_Collision);
 DECLARE_CYCLE_STAT(TEXT("ConvertOverlap"), STAT_CollisionConvertOverlap, STATGROUP_Collision);
 
+#define ENABLE_CHECK_HIT_NORMAL  (!(UE_BUILD_SHIPPING || UE_BUILD_TEST))
 
-#if !(UE_BUILD_SHIPPING || UE_BUILD_TEST || !WITH_EDITOR)
+#if ENABLE_CHECK_HIT_NORMAL
 /* Validate Normal of OutResult. We're on hunt for invalid normal */
 static void CheckHitResultNormal(const FHitResult& OutResult, const TCHAR* Message, const FVector& Start=FVector::ZeroVector, const FVector& End = FVector::ZeroVector, const PxGeometry* const Geom=NULL)
 {
@@ -51,8 +52,7 @@ static void CheckHitResultNormal(const FHitResult& OutResult, const TCHAR* Messa
 		ensure(OutResult.Normal.IsNormalized());
 	}
 }
-
-#endif // !(UE_BUILD_SHIPPING || UE_BUILD_TEST || !WITH_EDITOR)
+#endif // ENABLE_CHECK_HIT_NORMAL
 
 
 static FORCEINLINE bool PxQuatIsIdentity(PxQuat const& Q)
@@ -451,9 +451,9 @@ void ConvertQueryImpactHit(const UWorld* World, const PxLocationHit& PHit, FHitR
 	OutResult.TraceEnd = EndLoc;
 
 
-#if !(UE_BUILD_SHIPPING || UE_BUILD_TEST || !WITH_EDITOR)
+#if ENABLE_CHECK_HIT_NORMAL
 	CheckHitResultNormal(OutResult, TEXT("Invalid Normal from ConvertQueryImpactHit"), StartLoc, EndLoc, Geom);
-#endif
+#endif // ENABLE_CHECK_HIT_NORMAL
 
 	if (bUsePxNormal && !Normal.IsNormalized())
 	{
