Change slot evaluate node so that it works off of its own evaluation data, stops issues with parallel evaluation

diff --git a/Engine/Source/Runtime/Engine/Classes/Animation/AnimInstance.h b/Engine/Source/Runtime/Engine/Classes/Animation/AnimInstance.h
--- a/Engine/Source/Runtime/Engine/Classes/Animation/AnimInstance.h
+++ b/Engine/Source/Runtime/Engine/Classes/Animation/AnimInstance.h
@@ -239,18 +239,14 @@ struct FSlotEvaluationPose
 	/** Pose */
 	UPROPERTY()
 	FA2Pose Pose;
-
-	/** Pointer to Montage Instance */
-	FAnimMontageInstance * MontageInstance;
 	
 	FSlotEvaluationPose()
 	{
 	}
 
-	FSlotEvaluationPose(FAnimMontageInstance * InMontageInstance, float InWeight, EAdditiveAnimationType InAdditiveType)
+	FSlotEvaluationPose(float InWeight, EAdditiveAnimationType InAdditiveType)
 		: AdditiveType(InAdditiveType)
 		, Weight(InWeight)
-		, MontageInstance(InMontageInstance)
 	{
 	}
 };
@@ -352,6 +348,20 @@ struct FMontageActiveSlotTracker
 	FMontageActiveSlotTracker() : RootMotionWeight(0.f), bIsRelevantThisTick(false), bWasRelevantOnPreviousTick(false) {}
 };
 
+struct FMontageEvaluationState
+{
+	FMontageEvaluationState(UAnimMontage* InMontage, float InWeight, float InPosition) : Montage(InMontage), MontageWeight(InWeight), MontagePosition(InPosition) {}
+
+	// The montage to evaluate
+	UAnimMontage* Montage;
+
+	// The weight to use for this montage
+	float MontageWeight;
+
+	// The position to evaluate this montage at
+	float MontagePosition;
+};
+
 UCLASS(transient, Blueprintable, hideCategories=AnimInstance, BlueprintType)
 class ENGINE_API UAnimInstance : public UObject
 {
@@ -591,6 +601,9 @@ public:
 	*/
 	TArray<struct FAnimMontageInstance*> MontageInstances;
 
+	// Cached data for montage evaluation, save us having to access MontageInstances from slot nodes as that isn't thread safe
+	TArray<FMontageEvaluationState> MontageEvaluationData;
+
 	void OnMontageInstanceStopped(FAnimMontageInstance & StoppedMontageInstance);
 
 protected:
@@ -702,6 +715,8 @@ public:
 	static void AddReferencedObjects(UObject* InThis, FReferenceCollector& Collector);
 	// End UObject Interface
 
+	// Updates the montage data used for evaluation based on the current playing montages
+	void UpdateMontageEvaluationData();
 
 	//@TODO: Better comments
 	virtual void EvaluateAnimation(struct FPoseContext& Output);
diff --git a/Engine/Source/Runtime/Engine/Private/Animation/AnimInstance.cpp b/Engine/Source/Runtime/Engine/Private/Animation/AnimInstance.cpp
--- a/Engine/Source/Runtime/Engine/Private/Animation/AnimInstance.cpp
+++ b/Engine/Source/Runtime/Engine/Private/Animation/AnimInstance.cpp
@@ -1480,6 +1480,9 @@ void UAnimInstance::GetSlotWeight(FName const & SlotNodeName, float& out_SlotNod
 
 void UAnimInstance::SlotEvaluatePose(FName SlotNodeName, const FA2Pose & SourcePose, FA2Pose & BlendedPose, float SlotNodeWeight)
 {
+	//Accessing MontageInstances from this function is not safe (as this can be called during Parallel Anim Evaluation!
+	//Any montage data you need to add should be part of MontageEvaluationData
+
 	SCOPE_CYCLE_COUNTER(STAT_AnimNativeEvaluatePoses);
 	if (SlotNodeWeight <= ZERO_ANIMWEIGHT_THRESH)
 	{
@@ -1494,32 +1497,30 @@ void UAnimInstance::SlotEvaluatePose(FName SlotNodeName, const FA2Pose & SourceP
 	// first pass we go through collect weights and valid montages. 
 	float TotalWeight = 0.f;
 	float NonAdditiveWeight = 0.f;
-	for (auto Iter = MontageInstances.CreateConstIterator(); Iter; ++Iter)
+	for (const FMontageEvaluationState& EvalState : MontageEvaluationData)
 	{
-		FAnimMontageInstance * MontageInstance = (*Iter);
-		if (MontageInstance->IsValid() && MontageInstance->Montage->IsValidSlot(SlotNodeName) && (MontageInstance->Weight > ZERO_ANIMWEIGHT_THRESH))
+		if (EvalState.Montage->IsValidSlot(SlotNodeName))
 		{
-			FAnimTrack const * const AnimTrack = MontageInstance->Montage->GetAnimationData(SlotNodeName);
+			FAnimTrack const * const AnimTrack = EvalState.Montage->GetAnimationData(SlotNodeName);
 
 			// Find out additive type for pose.
 			EAdditiveAnimationType const AdditiveAnimType = AnimTrack->IsAdditive() 
 				? (AnimTrack->IsRotationOffsetAdditive() ? AAT_RotationOffsetMeshSpace : AAT_LocalSpaceBase)
 				: AAT_None;
 
-			FSlotEvaluationPose NewPose(MontageInstance, MontageInstance->Weight, AdditiveAnimType);
+			FSlotEvaluationPose NewPose(EvalState.MontageWeight, AdditiveAnimType);
 			
 			// Bone array has to be allocated prior to calling GetPoseFromAnimTrack
 			NewPose.Pose.Bones.AddUninitialized(RequiredBones.GetNumBones());
 
 			// Extract pose from Track
-			UAnimMontage const * const MontageAsset = MontageInstance->Montage;
-			FAnimExtractContext ExtractionContext(MontageInstance->GetPosition(), MontageInstance->Montage->HasRootMotion() && RootMotionMode != ERootMotionMode::NoRootMotionExtraction);
+			FAnimExtractContext ExtractionContext(EvalState.MontagePosition, EvalState.Montage->HasRootMotion() && RootMotionMode != ERootMotionMode::NoRootMotionExtraction);
 			FAnimationRuntime::GetPoseFromAnimTrack(*AnimTrack, RequiredBones, NewPose.Pose.Bones, ExtractionContext);
 
-			TotalWeight += MontageInstance->Weight;
+			TotalWeight += EvalState.MontageWeight;
 			if (AdditiveAnimType == AAT_None)
 			{
-				NonAdditiveWeight += MontageInstance->Weight;
+				NonAdditiveWeight += EvalState.MontageWeight;
 				NonAdditivePoses.Add(NewPose);
 			}
 			else
@@ -2746,4 +2747,16 @@ void UAnimInstance::UnlockAIResources(bool bUnlockMovement, bool UnlockAILogic)
 	UE_LOG(LogAnimation, Error, TEXT("%s: UnlockAIResources is no longer supported. Please use UnlockAIResourcesWithAnimation instead."), *GetName());
 }
 
+void UAnimInstance::UpdateMontageEvaluationData()
+{
+	MontageEvaluationData.Empty(MontageInstances.Num());
+	for (FAnimMontageInstance* MontageInstance : MontageInstances)
+	{
+		if (MontageInstance->Montage && MontageInstance->Weight > ZERO_ANIMWEIGHT_THRESH)
+		{
+			MontageEvaluationData.Add(FMontageEvaluationState(MontageInstance->Montage, MontageInstance->Weight, MontageInstance->GetPosition()));
+		}
+	}
+}
+
 #undef LOCTEXT_NAMESPACE 
diff --git a/Engine/Source/Runtime/Engine/Private/Components/SkeletalMeshComponent.cpp b/Engine/Source/Runtime/Engine/Private/Components/SkeletalMeshComponent.cpp
--- a/Engine/Source/Runtime/Engine/Private/Components/SkeletalMeshComponent.cpp
+++ b/Engine/Source/Runtime/Engine/Private/Components/SkeletalMeshComponent.cpp
@@ -1026,6 +1026,10 @@ void USkeletalMeshComponent::RefreshBoneTransforms(FActorComponentTickFunction*
 		RecalcRequiredBones(PredictedLODLevel);
 	}
 
+	if (AnimScriptInstance)
+	{
+		AnimScriptInstance->UpdateMontageEvaluationData();
+	}
 	const bool bDoEvaluationRateOptimization = bEnableUpdateRateOptimizations && AnimUpdateRateParams->DoEvaluationRateOptimizations();
 
 	//Handle update rate optimization setup
