Fixed editor crash when attempting to assign a maximum texture size to a non-power-of-two sized texture.

diff --git a/Engine/Source/Runtime/Engine/Private/Texture.cpp b/Engine/Source/Runtime/Engine/Private/Texture.cpp
--- a/Engine/Source/Runtime/Engine/Private/Texture.cpp
+++ b/Engine/Source/Runtime/Engine/Private/Texture.cpp
@@ -121,6 +121,7 @@ void UTexture::PostEditChangeProperty(FPropertyChangedEvent& PropertyChangedEven
 		static const FName DeferCompressionName("DeferCompression");
 #if WITH_EDITORONLY_DATA
 		static const FName MaxTextureSizeName("MaxTextureSize");
+		static const FName PowerOfTwoModeName("PowerOfTwoMode");
 #endif // #if WITH_EDITORONLY_DATA
 
 		const FName PropertyName = PropertyThatChanged->GetFName();
@@ -141,7 +142,23 @@ void UTexture::PostEditChangeProperty(FPropertyChangedEvent& PropertyChangedEven
 			}
 			else
 			{
-				MaxTextureSize = FMath::Min<int32>(FMath::RoundUpToPowerOfTwo(MaxTextureSize), GetMaximumDimension());
+				if (!Source.IsPowerOfTwo() && PowerOfTwoMode == ETexturePowerOfTwoSetting::None)
+				{
+					// Force power-of-two padding if a max texture size is specified on a non power-of-two sized texture
+					PowerOfTwoMode = ETexturePowerOfTwoSetting::PadToPowerOfTwo;
+				}
+				else
+				{
+					MaxTextureSize = FMath::Min<int32>(FMath::RoundUpToPowerOfTwo(MaxTextureSize), GetMaximumDimension());
+				}
+			}
+		}
+		else if (PropertyName == PowerOfTwoModeName)
+		{
+			if (PowerOfTwoMode == ETexturePowerOfTwoSetting::None && MaxTextureSize > 0 && !Source.IsPowerOfTwo())
+			{
+				// Force no max texture size if the source texture is not power-of-two sized and we are requesting no power-of-two padding
+				MaxTextureSize = 0;
 			}
 		}
 #endif // #if WITH_EDITORONLY_DATA
