Protect parallel animation evaluation from calls to InitAnim during its execution

diff --git a/Engine/Source/Runtime/Engine/Classes/Components/SkeletalMeshComponent.h b/Engine/Source/Runtime/Engine/Classes/Components/SkeletalMeshComponent.h
--- a/Engine/Source/Runtime/Engine/Classes/Components/SkeletalMeshComponent.h
+++ b/Engine/Source/Runtime/Engine/Classes/Components/SkeletalMeshComponent.h
@@ -1191,11 +1191,11 @@ private:
 public:
 	// Parallel evaluation wrappers
 	void ParallelAnimationEvaluation() { PerformAnimationEvaluation(AnimEvaluationContext.SkeletalMesh, AnimEvaluationContext.AnimInstance, AnimEvaluationContext.SpaceBases, AnimEvaluationContext.LocalAtoms, AnimEvaluationContext.VertexAnims, AnimEvaluationContext.RootBoneTranslation); }
-	void CompleteParallelAnimationEvaluation()
+	void CompleteParallelAnimationEvaluation(bool bDoPostAnimEvaluation)
 	{
 		ParallelAnimationEvaluationTask.SafeRelease(); //We are done with this task now, clean up!
 
-		if ((AnimEvaluationContext.AnimInstance == AnimScriptInstance) && (AnimEvaluationContext.SkeletalMesh == SkeletalMesh) && (AnimEvaluationContext.SpaceBases.Num() == GetNumSpaceBases()))
+		if (bDoPostAnimEvaluation && (AnimEvaluationContext.AnimInstance == AnimScriptInstance) && (AnimEvaluationContext.SkeletalMesh == SkeletalMesh) && (AnimEvaluationContext.SpaceBases.Num() == GetNumSpaceBases()))
 		{
 			Exchange(AnimEvaluationContext.SpaceBases, AnimEvaluationContext.bDoInterpolation ? CachedSpaceBases : GetEditableSpaceBases() );
 			Exchange(AnimEvaluationContext.LocalAtoms, AnimEvaluationContext.bDoInterpolation ? CachedLocalAtoms : LocalAtoms);
@@ -1207,6 +1207,8 @@ public:
 		AnimEvaluationContext.Clear();
 	}
 
+	bool IsRunningParallelEvaluation(bool bBlockOnTask, bool bPerformPostAnimEvaluation);
+
 	friend class FSkeletalMeshComponentDetails;
 
 private:
diff --git a/Engine/Source/Runtime/Engine/Private/Components/SkeletalMeshComponent.cpp b/Engine/Source/Runtime/Engine/Private/Components/SkeletalMeshComponent.cpp
--- a/Engine/Source/Runtime/Engine/Private/Components/SkeletalMeshComponent.cpp
+++ b/Engine/Source/Runtime/Engine/Private/Components/SkeletalMeshComponent.cpp
@@ -92,7 +92,8 @@ public:
 		SCOPE_CYCLE_COUNTER(STAT_AnimGameThreadTime);
 		if (USkeletalMeshComponent* Comp = SkeletalMeshComponent.Get())
 		{
-			Comp->CompleteParallelAnimationEvaluation();
+			const bool bPerformPostAnimEvaluation = true;
+			Comp->CompleteParallelAnimationEvaluation(bPerformPostAnimEvaluation);
 		}
 	}
 };
@@ -239,6 +240,13 @@ void USkeletalMeshComponent::InitAnim(bool bForceReinit)
 	// I'm moving the check here
 	if ( SkeletalMesh != NULL && IsRegistered() )
 	{
+		// We may be doing parallel evaluation on the current anim instance
+		// Calling this here with true will block this init till that thread completes
+		// and it is safe to continue
+		const bool bBlockOnTask = true; // wait on evaluation task so it is safe to continue with Init
+		const bool bPerformPostAnimEvaluation = false; // Skip post evaluation, it would be wasted work
+		IsRunningParallelEvaluation(bBlockOnTask, bPerformPostAnimEvaluation);
+
 		bool bBlueprintMismatch = (AnimBlueprintGeneratedClass != NULL) && 
 			(AnimScriptInstance != NULL) && (AnimScriptInstance->GetClass() != AnimBlueprintGeneratedClass);
 
@@ -1045,16 +1053,12 @@ void USkeletalMeshComponent::RefreshBoneTransforms(FActorComponentTickFunction*
 
 	const bool bDoParallelEvaluation = AnimEvaluationContext.bDoEvaluation && TickFunction && bDoPAE;
 
-	if (IsValidRef(ParallelAnimationEvaluationTask))
+	const bool bBlockOnTask = !bDoParallelEvaluation;  // If we aren't trying to do parallel evaluation then we
+															// will need to wait on an existing task.
+
+	const bool bPerformPostAnimEvaluation = true;
+	if (IsRunningParallelEvaluation(bBlockOnTask, bPerformPostAnimEvaluation))
 	{
-		//Are already processing eval on another thread, we wait for eval thread to finish
-		if (!bDoParallelEvaluation) //we are already running parallel evaluation, so if we are going to try it again just return
-		{
-			//If we are not going to attempt parallel evaluation then wait here for the existing task to finish and then
-			//perform complete directly so that when we return all the calculations are complete
-			FTaskGraphInterface::Get().WaitUntilTaskCompletes(ParallelAnimationEvaluationTask, ENamedThreads::GameThread);
-			CompleteParallelAnimationEvaluation(); //Perform completion now
-		}
 		return;
 	}
 
@@ -2142,3 +2146,18 @@ void USkeletalMeshComponent::RefreshActiveVertexAnims()
 		ActiveVertexAnims.Empty();
 	}
 }
+
+bool USkeletalMeshComponent::IsRunningParallelEvaluation(bool bBlockOnTask, bool bPerformPostAnimEvaluation)
+{
+	if (IsValidRef(ParallelAnimationEvaluationTask)) // We are already processing eval on another thread
+	{
+		if (bBlockOnTask)
+		{
+			check(IsInGameThread()); // Only attempt this from game thread!
+			FTaskGraphInterface::Get().WaitUntilTaskCompletes(ParallelAnimationEvaluationTask, ENamedThreads::GameThread);
+			CompleteParallelAnimationEvaluation(bPerformPostAnimEvaluation); //Perform completion now
+		}
+		return true;
+	}
+	return false;
+}
