PR #1076: [GameMode] Get pawn class via GetDefaultPawnClassForController in default ChoosePlayerStart implementation (Contributed by [REDACTED_PERSON])

diff --git a/Engine/Source/Runtime/Engine/Private/GameMode.cpp b/Engine/Source/Runtime/Engine/Private/GameMode.cpp
--- a/Engine/Source/Runtime/Engine/Private/GameMode.cpp
+++ b/Engine/Source/Runtime/Engine/Private/GameMode.cpp
@@ -1496,7 +1496,8 @@ AActor* AGameMode::ChoosePlayerStart_Implementation( AController* Player )
 {
 	// Choose a player start
 	APlayerStart* FoundPlayerStart = NULL;
-	APawn* PawnToFit = DefaultPawnClass ? DefaultPawnClass->GetDefaultObject<APawn>() : NULL;
+	UClass* PawnClass = GetDefaultPawnClassForController(Player);
+	APawn* PawnToFit = PawnClass ? PawnClass->GetDefaultObject<APawn>() : nullptr;
 	TArray<APlayerStart*> UnOccupiedStartPoints;
 	TArray<APlayerStart*> OccupiedStartPoints;
 	for (TActorIterator<APlayerStart> It(GetWorld()); It; ++It)
