Refactor rotation node to reduce transform passing (compacted pose prep)

diff --git a/Engine/Source/Runtime/Engine/Classes/Animation/BoneControllers/AnimNode_RotationMultiplier.h b/Engine/Source/Runtime/Engine/Classes/Animation/BoneControllers/AnimNode_RotationMultiplier.h
--- a/Engine/Source/Runtime/Engine/Classes/Animation/BoneControllers/AnimNode_RotationMultiplier.h
+++ b/Engine/Source/Runtime/Engine/Classes/Animation/BoneControllers/AnimNode_RotationMultiplier.h
@@ -54,9 +54,9 @@ struct ENGINE_API FAnimNode_RotationMultiplier : public FAnimNode_SkeletalContro
 
 private:
 	// Extract Delta Quat of rotation around Axis of animation and reference pose for the SourceBoneIndex
-	FQuat ExtractAngle(const TArray<FTransform> & RefPoseTransforms, FA2CSPose& MeshBases, const EBoneAxis Axis,  int32 SourceBoneIndex);
+	FQuat ExtractAngle(const FTransform& RefPoseTransform, const FTransform& LocalBoneTransform, const EBoneAxis Axis);
 	// Multiply scalar value Multiplier to the delta Quat of SourceBone Index's rotation
-	void MultiplyQuatBasedOnSourceIndex(const TArray<FTransform> & RefPoseTransforms, FA2CSPose& MeshBases, const EBoneAxis Axis, int32 SourceBoneIndex, float Multiplier, const FQuat& ReferenceQuat, FQuat& OutQuat);
+	FQuat MultiplyQuatBasedOnSourceIndex(const FTransform& RefPoseTransform, const FTransform& LocalBoneTransform, const EBoneAxis Axis, float InMultiplier, const FQuat& ReferenceQuat);
 
 	// FAnimNode_SkeletalControlBase interface
 	virtual void InitializeBoneReferences(const FBoneContainer& RequiredBones) override;
diff --git a/Engine/Source/Runtime/Engine/Private/Animation/AnimNodeRotationMultiplier.cpp b/Engine/Source/Runtime/Engine/Private/Animation/AnimNodeRotationMultiplier.cpp
--- a/Engine/Source/Runtime/Engine/Private/Animation/AnimNodeRotationMultiplier.cpp
+++ b/Engine/Source/Runtime/Engine/Private/Animation/AnimNodeRotationMultiplier.cpp
@@ -37,12 +37,10 @@ FVector GetAxisVector(const EBoneAxis Axis)
 	}
 }
 
-FQuat FAnimNode_RotationMultiplier::ExtractAngle(const TArray<FTransform> & RefPoseTransforms, FA2CSPose & MeshBases, const EBoneAxis Axis,  int32 SourceBoneIndex)
+FQuat FAnimNode_RotationMultiplier::ExtractAngle(const FTransform& RefPoseTransform, const FTransform& LocalBoneTransform, const EBoneAxis Axis)
 {
-	// local bone transform
-	const FTransform& LocalBoneTransform = MeshBases.GetLocalSpaceTransform(SourceBoneIndex);
 	// local bone transform with reference rotation
-	FTransform ReferenceBoneTransform = RefPoseTransforms[SourceBoneIndex];
+	FTransform ReferenceBoneTransform = RefPoseTransform;
 	ReferenceBoneTransform.SetTranslation(LocalBoneTransform.GetTranslation());
 
 	// find delta angle between the two quaternions X Axis.
@@ -64,8 +62,8 @@ FQuat FAnimNode_RotationMultiplier::ExtractAngle(const TArray<FTransform> & RefP
 	checkSlow( DeltaQuat.IsNormalized() );
 
 #if 0 //DEBUG_TWISTBONECONTROLLER
-	UE_LOG(LogSkeletalControl, Log, TEXT("\t ExtractAngle, Bone: %s (%d)"), 
-		*SourceBone.BoneName.ToString(), SourceBoneIndex);
+	UE_LOG(LogSkeletalControl, Log, TEXT("\t ExtractAngle, Bone: %s"), 
+		*SourceBone.BoneName.ToString());
 	UE_LOG(LogSkeletalControl, Log, TEXT("\t\t Bone Quat: %s, Rot: %s, AxisX: %s"), *LocalBoneTransform.GetRotation().ToString(), *LocalBoneTransform.GetRotation().Rotator().ToString(), *LocalRotationVector.ToString() );
 	UE_LOG(LogSkeletalControl, Log, TEXT("\t\t BoneRef Quat: %s, Rot: %s, AxisX: %s"), *ReferenceBoneTransform.GetRotation().ToString(), *ReferenceBoneTransform.GetRotation().Rotator().ToString(), *ReferenceRotationVector.ToString() );
 	UE_LOG(LogSkeletalControl, Log, TEXT("\t\t LocalToRefQuat Quat: %s, Rot: %s"), *LocalToRefQuat.ToString(), *LocalToRefQuat.Rotator().ToString() );
@@ -86,10 +84,10 @@ FQuat FAnimNode_RotationMultiplier::ExtractAngle(const TArray<FTransform> & RefP
 	return DeltaQuat;
 }
 
-void FAnimNode_RotationMultiplier::MultiplyQuatBasedOnSourceIndex(const TArray<FTransform> & RefPoseTransforms, FA2CSPose& MeshBases, const EBoneAxis Axis, int32 SourceBoneIndex, float InMultiplier, const FQuat& ReferenceQuat, FQuat& OutQuat)
+FQuat FAnimNode_RotationMultiplier::MultiplyQuatBasedOnSourceIndex(const FTransform& RefPoseTransform, const FTransform& LocalBoneTransform, const EBoneAxis Axis, float InMultiplier, const FQuat& ReferenceQuat)
 {
 	// Find delta angle for source bone.
-	FQuat DeltaQuat = ExtractAngle(RefPoseTransforms, MeshBases, Axis, SourceBoneIndex);
+	FQuat DeltaQuat = ExtractAngle(RefPoseTransform, LocalBoneTransform, Axis);
 
 	// Turn to Axis and Angle
 	FVector RotationAxis;
@@ -109,7 +107,7 @@ void FAnimNode_RotationMultiplier::MultiplyQuatBasedOnSourceIndex(const TArray<F
 	RotationAngle = FMath::UnwindRadians(RotationAngle);
 
 	// New bone rotation
-	OutQuat = ReferenceQuat * FQuat(RotationAxis, RotationAngle* InMultiplier);
+	FQuat OutQuat = ReferenceQuat * FQuat(RotationAxis, RotationAngle* InMultiplier);
 	// Normalize resulting quaternion.
 	OutQuat.Normalize();
 
@@ -118,6 +116,8 @@ void FAnimNode_RotationMultiplier::MultiplyQuatBasedOnSourceIndex(const TArray<F
 	UE_LOG(LogSkeletalControl, Log, TEXT("\t NewQuat: %s, Rot: %s"), *OutQuat.ToString(), *OutQuat.Rotator().ToString() );
 	UE_LOG(LogSkeletalControl, Log, TEXT("\t RollAxis: %s, RollAngle: %f"), *RotationAxis.ToString(), RotationAngle );
 #endif
+
+	return OutQuat;
 }
 
 void FAnimNode_RotationMultiplier::EvaluateBoneTransforms(USkeletalMeshComponent* SkelComp, const FBoneContainer& RequiredBones, FA2CSPose& MeshBases, TArray<FBoneTransform>& OutBoneTransforms)
@@ -129,8 +129,7 @@ void FAnimNode_RotationMultiplier::EvaluateBoneTransforms(USkeletalMeshComponent
 		// Reference bone
 		const TArray<FTransform> & LocalRefPose = RequiredBones.GetRefPoseArray();
 		const FQuat RefQuat = LocalRefPose[TargetBone.BoneIndex].GetRotation();
-		FQuat NewQuat;
-		MultiplyQuatBasedOnSourceIndex(LocalRefPose, MeshBases, RotationAxisToRefer, SourceBone.BoneIndex, Multiplier, RefQuat, NewQuat);
+		FQuat NewQuat = MultiplyQuatBasedOnSourceIndex(LocalRefPose[SourceBone.BoneIndex], MeshBases.GetLocalSpaceTransform(SourceBone.BoneIndex), RotationAxisToRefer, Multiplier, RefQuat);
 
 		FTransform NewLocalTransform = MeshBases.GetLocalSpaceTransform(TargetBone.BoneIndex);
 		
@@ -165,4 +164,4 @@ void FAnimNode_RotationMultiplier::InitializeBoneReferences(const FBoneContainer
 {
 	SourceBone.Initialize(RequiredBones);
 	TargetBone.Initialize(RequiredBones);
-}
+}
\ No newline at end of file
