[UEDOC-1485] Comment fix.

diff --git a/Engine/Source/Runtime/Engine/Classes/Kismet/KismetArrayLibrary.h b/Engine/Source/Runtime/Engine/Classes/Kismet/KismetArrayLibrary.h
--- a/Engine/Source/Runtime/Engine/Classes/Kismet/KismetArrayLibrary.h
+++ b/Engine/Source/Runtime/Engine/Classes/Kismet/KismetArrayLibrary.h
@@ -124,7 +124,7 @@ class ENGINE_API UKismetArrayLibrary : public UBlueprintFunctionLibrary
 	 *Given an array and an index, assigns the item to that array element
 	 *
 	 *@param	TargetArray		The array to perform the operation on
-	 *@param	Index			The index to assing the item to
+	 *@param	Index			The index to assign the item to
 	 *@param	Item			The item to assign to the index of the array
 	 *@param	bSizeToFit		If true, the array will expand if Index is greater than the current size of the array
 	*/
