Fix bug in UCastNavMeshDataChunk::GatherTiles due to shadow variable

diff --git a/Engine/Source/Runtime/Engine/Private/AI/Navigation/RecastNavMeshDataChunk.cpp b/Engine/Source/Runtime/Engine/Private/AI/Navigation/RecastNavMeshDataChunk.cpp
--- a/Engine/Source/Runtime/Engine/Private/AI/Navigation/RecastNavMeshDataChunk.cpp
+++ b/Engine/Source/Runtime/Engine/Private/AI/Navigation/RecastNavMeshDataChunk.cpp
@@ -303,7 +303,7 @@ void URecastNavMeshDataChunk::GatherTiles(const FPImplRecastNavMesh* NavMeshImpl
 			uint8* RawTileCacheData = nullptr;
 			if (NavMeshImpl->NavMeshOwner->SupportsRuntimeGeneration())
 			{
-				FNavMeshTileData TileCacheData = NavMeshImpl->GetTileCacheLayer(Tile->header->x, Tile->header->y, Tile->header->layer);
+				TileCacheData = NavMeshImpl->GetTileCacheLayer(Tile->header->x, Tile->header->y, Tile->header->layer);
 				if (TileCacheData.IsValid())
 				{
 					// Make our own copy of tile cache data
