Refactoring blend code to reduce the amount of code that manually loops over transforms using required bones (compacted pose prep work)

diff --git a/Engine/Source/Runtime/Engine/Private/Animation/AnimationRuntime.cpp b/Engine/Source/Runtime/Engine/Private/Animation/AnimationRuntime.cpp
--- a/Engine/Source/Runtime/Engine/Private/Animation/AnimationRuntime.cpp
+++ b/Engine/Source/Runtime/Engine/Private/Animation/AnimationRuntime.cpp
@@ -17,6 +17,36 @@
 DEFINE_LOG_CATEGORY(LogAnimation);
 DEFINE_LOG_CATEGORY(LogRootMotion);
 
+/////////////////////////////////////////////////////////
+// Templated Transform Blend Functionality
+
+namespace ETransformBlendMode
+{
+	enum Type
+	{
+		Overwrite,
+		Accumulate
+	};
+}
+
+template<int32>
+void BlendTransform(const FTransform& Source, FTransform& Dest, const ScalarRegister BlendWeight)
+{
+	check(false); /// should never call this
+}
+
+template<>
+void BlendTransform<ETransformBlendMode::Overwrite>(const FTransform& Source, FTransform& Dest, const ScalarRegister BlendWeight)
+{
+	Dest = Source * BlendWeight;
+}
+
+template<>
+void BlendTransform<ETransformBlendMode::Accumulate>(const FTransform& Source, FTransform& Dest, const ScalarRegister BlendWeight)
+{
+	Dest.AccumulateWithShortestRotation(Source, BlendWeight);
+}
+
 //////////////////////////////////////////////////////////////////////////
 
 void FAnimationRuntime::NormalizeRotations(const FBoneContainer& RequiredBones, /*inout*/ FTransformArrayA2 & Atoms)
@@ -49,6 +79,16 @@ void FAnimationRuntime::InitializeTransform(const FBoneContainer& RequiredBones,
 	}
 }
 
+template <int32 TRANSFORM_BLEND_MODE>
+FORCEINLINE void BlendPose(const FTransformArrayA2& SourcePoses, FTransformArrayA2& ResultAtoms, const TArray<FBoneIndexType> & RequiredBoneIndices, const ScalarRegister VBlendWeight)
+{
+	for (int32 i = 0; i < RequiredBoneIndices.Num(); ++i)
+	{
+		const int32 BoneIndex = RequiredBoneIndices[i];
+		BlendTransform<TRANSFORM_BLEND_MODE>(SourcePoses[BoneIndex], ResultAtoms[BoneIndex], VBlendWeight);
+	}
+}
+
 /**
  * Blends together a set of poses, each with a given weight.
  * This function is lightweight, it does not cull out nearly zero weights or check to make sure weights sum to 1.0, the caller should take care of that if needed.
@@ -68,32 +108,14 @@ void FAnimationRuntime::BlendPosesTogether(int32 NumPoses, const FTransformArray
 	if (ensure(SourcePoses && ResultAtoms.Num() >= RequiredBoneIndices.Num()))
 	{
 		// debug purpose for now, but this can cause 0 bone transform, so we'd like to catch it 
-		float WeightSum=0.f;
-		for (int32 i = 0; i < NumPoses; ++i)
-		{
-			const ScalarRegister VBlendWeight(SourceWeights[i]);
-			const FTransformArrayA2& SourceAtoms = *SourcePoses[i];
+		float WeightSum = SourceWeights[0];
+		BlendPose<ETransformBlendMode::Overwrite>(*SourcePoses[0], ResultAtoms, RequiredBoneIndices, ScalarRegister(SourceWeights[0]));
 
+		for (int32 i = 1; i < NumPoses; ++i)
+		{
 			WeightSum += SourceWeights[i];
 
-			if (i == 0)
-			{
-				// First pose will just overwrite the destination
-				for (int32 j = 0; j < RequiredBoneIndices.Num(); ++j)
-				{
-					const int32 BoneIndex = RequiredBoneIndices[j];
-					ResultAtoms[BoneIndex] = SourceAtoms[BoneIndex] * VBlendWeight;
-				}
-			}
-			else
-			{
-				// Subsequent poses need to be blended in
-				for (int32 j = 0; j < RequiredBoneIndices.Num(); ++j)
-				{
-					const int32 BoneIndex = RequiredBoneIndices[j];
-					ResultAtoms[BoneIndex].AccumulateWithShortestRotation(SourceAtoms[BoneIndex], VBlendWeight);
-				}
-			}
+			BlendPose<ETransformBlendMode::Accumulate>(*SourcePoses[i], ResultAtoms, RequiredBoneIndices, ScalarRegister(SourceWeights[i]));
 		}
 
 		ensure (WeightSum != 0.f);
@@ -121,31 +143,13 @@ void FAnimationRuntime::BlendPosesTogether(int32 NumPoses, const TArray<FTransfo
 {
 	check(NumPoses > 0);
 
-
 	const TArray<FBoneIndexType> & RequiredBoneIndices = RequiredBones.GetBoneIndicesArray();
-	for (int32 i = 0; i < NumPoses; ++i)
-	{
-		const ScalarRegister VBlendWeight(SourceWeights[i]);
-		const FTransformArrayA2& SourceAtoms = SourcePoses[i];
 
-		if (i == 0)
-		{
-			// First pose will just overwrite the destination
-			for (int32 j = 0; j < RequiredBoneIndices.Num(); ++j)
-			{
-				const int32 BoneIndex = RequiredBoneIndices[j];
-				ResultAtoms[BoneIndex] = SourceAtoms[BoneIndex] * VBlendWeight;
-			}
-		}
-		else
-		{
-			// Subsequent poses need to be blended in
-			for (int32 j = 0; j < RequiredBoneIndices.Num(); ++j)
-			{
-				const int32 BoneIndex = RequiredBoneIndices[j];
-				ResultAtoms[BoneIndex].AccumulateWithShortestRotation(SourceAtoms[BoneIndex], VBlendWeight);
-			}
-		}
+	BlendPose<ETransformBlendMode::Overwrite>(SourcePoses[0], ResultAtoms, RequiredBoneIndices, ScalarRegister(SourceWeights[0]));
+
+	for (int32 i = 1; i < NumPoses; ++i)
+	{
+		BlendPose<ETransformBlendMode::Accumulate>(SourcePoses[i], ResultAtoms, RequiredBoneIndices, ScalarRegister(SourceWeights[i]));
 	}
 
 	// Ensure that all of the resulting rotations are normalized
@@ -155,6 +159,31 @@ void FAnimationRuntime::BlendPosesTogether(int32 NumPoses, const TArray<FTransfo
 	}
 }
 
+template <int32 TRANSFORM_BLEND_MODE>
+void BlendPosePerBone(const TArray<FBoneIndexType> &RequiredBoneIndices, const TArray<int32>& PerBoneIndices, const FBlendSampleData& BlendSampleDataCache, FTransformArrayA2 &ResultAtoms, const FTransformArrayA2& SourceAtoms)
+{
+	const ScalarRegister VBlendWeight(BlendSampleDataCache.GetWeight());
+	TArray<ScalarRegister> PerBoneBlends;
+	for (int32 i = 0; i < BlendSampleDataCache.PerBoneBlendData.Num(); ++i)
+	{
+		PerBoneBlends.Add(ScalarRegister(FMath::Clamp<float>(BlendSampleDataCache.PerBoneBlendData[i], 0.f, 1.f)));
+	}
+
+	for (int32 i = 0; i < RequiredBoneIndices.Num(); ++i)
+	{
+		const int32 BoneIndex = RequiredBoneIndices[i];
+		int32 PerBoneIndex = PerBoneIndices[i];
+		if (PerBoneIndex == INDEX_NONE || !BlendSampleDataCache.PerBoneBlendData.IsValidIndex(PerBoneIndex))
+		{
+			BlendTransform<TRANSFORM_BLEND_MODE>(SourceAtoms[BoneIndex], ResultAtoms[BoneIndex], VBlendWeight);
+		}
+		else
+		{
+			BlendTransform<TRANSFORM_BLEND_MODE>(SourceAtoms[BoneIndex], ResultAtoms[BoneIndex], PerBoneBlends[PerBoneIndex]);
+		}
+	}
+}
+
 /**
  * Blends together a set of poses, each with a given weight.
  * This function is lightweight, it does not cull out nearly zero weights or check to make sure weights sum to 1.0, the caller should take care of that if needed.
@@ -171,59 +200,19 @@ void FAnimationRuntime::BlendPosesTogetherPerBone(int32 NumPoses, const TArray<F
 	check(NumPoses > 0);
 
 	const TArray<FBoneIndexType> & RequiredBoneIndices = RequiredBones.GetBoneIndicesArray();
-	for (int32 i = 0; i < NumPoses; ++i)
+
+	TArray<int32> PerBoneIndices;
+	PerBoneIndices.AddUninitialized(ResultAtoms.Num());
+	for (int BoneIndex = 0; BoneIndex < PerBoneIndices.Num(); ++BoneIndex)
 	{
-		const ScalarRegister VBlendWeight(BlendSampleDataCache[i].GetWeight());
-		TArray<ScalarRegister> PerBoneBlends;
-		for (int32 j=0; j<BlendSampleDataCache[i].PerBoneBlendData.Num(); ++j)
-		{
-			PerBoneBlends.Add(ScalarRegister(FMath::Clamp<float>(BlendSampleDataCache[i].PerBoneBlendData[j], 0.f, 1.f)));
-		}
+		PerBoneIndices[BoneIndex] = BlendSpace->GetPerBoneInterpolationIndex(RequiredBoneIndices[BoneIndex], RequiredBones);
+	}
 
-		const FTransformArrayA2& SourceAtoms = SourcePoses[i];
+	BlendPosePerBone<ETransformBlendMode::Overwrite>(RequiredBoneIndices, PerBoneIndices, BlendSampleDataCache[0], ResultAtoms, SourcePoses[0]);
 
-		if (i == 0)
-		{
-			// First pose will just overwrite the destination
-			for (int32 j = 0; j < RequiredBoneIndices.Num(); ++j)
-			{
-				const int32 BoneIndex = RequiredBoneIndices[j];
-				int32 PerBoneIndex = BlendSpace->GetPerBoneInterpolationIndex(BoneIndex, RequiredBones);
-				if ( PerBoneIndex == INDEX_NONE )
-				{
-					ResultAtoms[BoneIndex] = SourceAtoms[BoneIndex] * VBlendWeight;
-				}
-				else if (!BlendSampleDataCache[i].PerBoneBlendData.IsValidIndex(PerBoneIndex))
-				{
-					ResultAtoms[BoneIndex] = SourceAtoms[BoneIndex] * VBlendWeight;
-				}
-				else
-				{
-					ResultAtoms[BoneIndex] = SourceAtoms[BoneIndex] * PerBoneBlends[PerBoneIndex];
-				}
-			}
-		}
-		else
-		{
-			// Subsequent poses need to be blended in
-			for (int32 j = 0; j < RequiredBoneIndices.Num(); ++j)
-			{
-				const int32 BoneIndex = RequiredBoneIndices[j];
-				int32 PerBoneIndex = BlendSpace->GetPerBoneInterpolationIndex(BoneIndex, RequiredBones);
-				if ( PerBoneIndex == INDEX_NONE )
-				{
-					ResultAtoms[BoneIndex].AccumulateWithShortestRotation(SourceAtoms[BoneIndex], VBlendWeight);
-				}
-				else if (!BlendSampleDataCache[i].PerBoneBlendData.IsValidIndex(PerBoneIndex))
-				{
-					ResultAtoms[BoneIndex].AccumulateWithShortestRotation(SourceAtoms[BoneIndex], VBlendWeight);
-				}
-				else
-				{
-					ResultAtoms[BoneIndex].AccumulateWithShortestRotation(SourceAtoms[BoneIndex], PerBoneBlends[PerBoneIndex]);
-				}
-			}
-		}
+	for (int32 i = 1; i < NumPoses; ++i)
+	{
+		BlendPosePerBone<ETransformBlendMode::Accumulate>(RequiredBoneIndices, PerBoneIndices, BlendSampleDataCache[i], ResultAtoms, SourcePoses[i]);
 	}
 
 	// Ensure that all of the resulting rotations are normalized
@@ -287,13 +276,8 @@ void FAnimationRuntime::BlendPosesTogetherPerBoneInMeshSpace(int32 NumPoses, TAr
 void FAnimationRuntime::BlendPosesAccumulate(const FTransformArrayA2& BlendPoses, const float BlendWeight, const FBoneContainer& RequiredBones, /*inout*/ FTransformArrayA2& ResultAtoms)
 {
 	const ScalarRegister VBlendWeight(BlendWeight);
-	const TArray<FBoneIndexType> & RequiredBoneIndices = RequiredBones.GetBoneIndicesArray();
-	// Subsequent poses need to be blended in
-	for (int32 j = 0; j < RequiredBoneIndices.Num(); ++j)
-	{
-		const int32 BoneIndex = RequiredBoneIndices[j];
-		ResultAtoms[BoneIndex].AccumulateWithShortestRotation(BlendPoses[BoneIndex], VBlendWeight);
-	}
+
+	BlendPose<ETransformBlendMode::Accumulate>(BlendPoses, ResultAtoms, RequiredBones.GetBoneIndicesArray(), VBlendWeight);
 }
 
 void FAnimationRuntime::LerpBoneTransforms(TArray<FTransform> & A, const TArray<FTransform> & B, float Alpha, const TArray<FBoneIndexType> & RequiredBonesArray)
