Fix ScaleByMomentOfInertia using a point transform instead of a vector transform by mistake.

diff --git a/Engine/Source/Runtime/Engine/Private/PrimitiveComponentPhysics.cpp b/Engine/Source/Runtime/Engine/Private/PrimitiveComponentPhysics.cpp
--- a/Engine/Source/Runtime/Engine/Private/PrimitiveComponentPhysics.cpp
+++ b/Engine/Source/Runtime/Engine/Private/PrimitiveComponentPhysics.cpp
@@ -475,7 +475,7 @@ FVector UPrimitiveComponent::ScaleByMomentOfInertia(FVector InputVector, FName B
 	const FVector LocalInertiaTensor = GetInertiaTensor(BoneName);
 	const FVector InputVectorLocal = ComponentToWorld.InverseTransformVectorNoScale(InputVector);
 	const FVector LocalScaled = InputVectorLocal * LocalInertiaTensor;
-	const FVector WorldScaled = ComponentToWorld.TransformFVector4NoScale(LocalScaled);
+	const FVector WorldScaled = ComponentToWorld.TransformVectorNoScale(LocalScaled);
 	return WorldScaled;
 }
 
