Improve comments for FOverlapInfo and also set FHitResult::Actor

diff --git a/Engine/Source/Runtime/Engine/Classes/Components/SceneComponent.h b/Engine/Source/Runtime/Engine/Classes/Components/SceneComponent.h
--- a/Engine/Source/Runtime/Engine/Classes/Components/SceneComponent.h
+++ b/Engine/Source/Runtime/Engine/Classes/Components/SceneComponent.h
@@ -29,7 +29,9 @@ struct FOverlapInfo
 	friend bool operator == (const FOverlapInfo& LHS, const FOverlapInfo& RHS) { return LHS.OverlapInfo.Component == RHS.OverlapInfo.Component && LHS.OverlapInfo.Item == RHS.OverlapInfo.Item; }
 	bool bFromSweep;
 
-	/** Information for both sweep and overlap queries. Different parts are valid depending on bFromSweep*/
+	/** Information for both sweep and overlap queries. Different parts are valid depending on bFromSweep.
+	  * If bFromSweep is true then FHitResult is completely valid just like a regular sweep result.
+	  * If bFromSweep is false only FHitResult::Component, FHitResult::Actor, FHitResult::Item are valid as this is really just an FOverlapResult*/
 	FHitResult OverlapInfo;
 };
 
diff --git a/Engine/Source/Runtime/Engine/Private/Components/SceneComponent.cpp b/Engine/Source/Runtime/Engine/Private/Components/SceneComponent.cpp
--- a/Engine/Source/Runtime/Engine/Private/Components/SceneComponent.cpp
+++ b/Engine/Source/Runtime/Engine/Private/Components/SceneComponent.cpp
@@ -23,6 +23,7 @@ DEFINE_LOG_CATEGORY_STATIC(LogSceneComponent, Log, All);
 FOverlapInfo::FOverlapInfo(UPrimitiveComponent* InComponent, int32 InBodyIndex)
 	: bFromSweep(false)
 {
+	OverlapInfo.Actor = InComponent ? InComponent->GetOwner() : nullptr;
 	OverlapInfo.Component = InComponent;
 	OverlapInfo.Item = InBodyIndex;
 }
