Cache UActorComponent.Owner

diff --git a/Engine/Source/Runtime/Engine/Classes/Components/ActorComponent.h b/Engine/Source/Runtime/Engine/Classes/Components/ActorComponent.h
--- a/Engine/Source/Runtime/Engine/Classes/Components/ActorComponent.h
+++ b/Engine/Source/Runtime/Engine/Classes/Components/ActorComponent.h
@@ -142,6 +142,8 @@ private:
 
 	friend class FActorComponentInstanceData;
 
+	AActor* Owner;
+
 	UPROPERTY()
 	TArray<FSimpleMemberReference> UCSModifiedProperties;
 
diff --git a/Engine/Source/Runtime/Engine/Private/Components/ActorComponent.cpp b/Engine/Source/Runtime/Engine/Private/Components/ActorComponent.cpp
--- a/Engine/Source/Runtime/Engine/Private/Components/ActorComponent.cpp
+++ b/Engine/Source/Runtime/Engine/Private/Components/ActorComponent.cpp
@@ -120,6 +120,8 @@ FGlobalComponentRecreateRenderStateContext::~FGlobalComponentRecreateRenderState
 UActorComponent::UActorComponent(const FObjectInitializer& ObjectInitializer /*= FObjectInitializer::Get()*/)
 	: Super(ObjectInitializer)
 {
+	Owner = GetTypedOuter<AActor>();
+
 	PrimaryComponentTick.TickGroup = TG_DuringPhysics;
 	PrimaryComponentTick.bStartWithTickEnabled = true;
 	PrimaryComponentTick.bCanEverTick = false;
@@ -135,7 +137,6 @@ void UActorComponent::PostInitProperties()
 {
 	Super::PostInitProperties();
 
-	AActor* Owner = GetOwner();
 	if (Owner)
 	{
 		Owner->AddOwnedComponent(this);
@@ -217,7 +218,7 @@ void UActorComponent::PostRename(UObject* OldOuter, const FName OldName)
 
 	if (OldOuter != GetOuter())
 	{
-		AActor* Owner = GetOwner();
+		Owner = GetTypedOuter<AActor>();
 		AActor* OldOwner = (OldOuter->IsA<AActor>() ? CastChecked<AActor>(OldOuter) : OldOuter->GetTypedOuter<AActor>());
 
 		if (Owner != OldOwner)
@@ -244,8 +245,6 @@ bool UActorComponent::IsCreatedByConstructionScript() const
 #if WITH_EDITOR
 void UActorComponent::CheckForErrors()
 {
-	AActor* Owner = GetOwner();
-
 	if (Owner != NULL && GetClass()->HasAnyClassFlags(CLASS_Deprecated))
 	{
 		FFormatNamedArguments Arguments;
@@ -272,26 +271,14 @@ void UActorComponent::CheckForErrors()
 
 bool UActorComponent::IsOwnerSelected() const
 {
-	AActor* Owner = GetOwner();
 	return Owner && Owner->IsSelected();
 }
 
 AActor* UActorComponent::GetOwner() const
 {
 	// walk up outer chain to find an Actor
-	UObject* Obj = GetOuter();
-	while(Obj != NULL)
-	{
-		AActor* Actor = Cast<AActor>(Obj);
-		if(Actor != NULL)
-		{
-			return Actor;
-		}
-
-		Obj = Obj->GetOuter();
-	}
-
-	return NULL;
+	checkSlow(Owner == GetTypedOuter<AActor>());
+	return Owner;
 }
 
 UWorld* UActorComponent::GetWorld() const
@@ -300,7 +287,6 @@ UWorld* UActorComponent::GetWorld() const
 	if (ComponentWorld == NULL)
 	{
 		// If we don't have a world yet, it may be because we haven't gotten registered yet, but we can try to look at our owner
-		AActor* Owner = GetOwner();
 		if (Owner && !Owner->HasAnyFlags(RF_ClassDefaultObject))
 		{
 			ComponentWorld = Owner->GetWorld();
@@ -323,7 +309,6 @@ bool UActorComponent::ComponentHasTag(FName Tag) const
 
 ENetMode UActorComponent::GetNetMode() const
 {
-	AActor *Owner = GetOwner();
 	return Owner ? Owner->GetNetMode() : NM_Standalone;
 }
 
@@ -339,10 +324,9 @@ ULevel* UActorComponent::GetComponentLevel() const
 	
 	if (ComponentLevel == NULL)
 	{
-		AActor* Actor = GetOwner();
-		if (Actor)
+		if (Owner)
 		{
-			ComponentLevel = CastChecked<ULevel>( Actor->GetOuter() );
+			ComponentLevel = CastChecked<ULevel>( Owner->GetOuter() );
 		}
 	}
 	
@@ -358,10 +342,9 @@ bool UActorComponent::ComponentIsInLevel(const ULevel *TestLevel) const
 bool UActorComponent::ComponentIsInPersistentLevel(bool bIncludeLevelStreamingPersistent) const
 {
 	ULevel* MyLevel = GetComponentLevel();
-	AActor* MyActor = GetOwner();
 	UWorld* MyWorld = GetWorld();
 
-	if(MyActor == NULL || MyLevel == NULL || MyWorld == NULL)
+	if(Owner == NULL || MyLevel == NULL || MyWorld == NULL)
 	{
 		return false;
 	}
@@ -407,7 +390,6 @@ void UActorComponent::BeginDestroy()
 	World = NULL;
 
 	// Remove from the parent's OwnedComponents list
-	AActor* Owner = GetOwner();
 	if (Owner)
 	{
 		Owner->RemoveOwnedComponent(this);
@@ -432,7 +414,6 @@ bool UActorComponent::NeedsLoadForServer() const
 
 int32 UActorComponent::GetFunctionCallspace( UFunction* Function, void* Parameters, FFrame* Stack )
 {
-	AActor* Owner = GetOwner();
 	if (Owner == NULL)
 	{
 		return FunctionCallspace::Local;
@@ -442,7 +423,6 @@ int32 UActorComponent::GetFunctionCallspace( UFunction* Function, void* Paramete
 
 bool UActorComponent::CallRemoteFunction( UFunction* Function, void* Parameters, FOutParmRec* OutParms, FFrame* Stack )
 {
-	AActor* Owner = GetOwner();
 	if (Owner == NULL)
 	{
 		return false;
@@ -466,9 +446,9 @@ bool UActorComponent::Modify( bool bAlwaysMarkDirty/*=true*/ )
 {
 	// If this is a construction script component we don't store them in the transaction buffer.  Instead, mark
 	// the Actor as modified so that we store of the transaction annotation that has the component properties stashed
-	if (IsCreatedByConstructionScript() && GetOwner())
+	if (Owner && IsCreatedByConstructionScript())
 	{
-		return GetOwner()->Modify(bAlwaysMarkDirty);
+		return Owner->Modify(bAlwaysMarkDirty);
 	}
 
 	return Super::Modify(bAlwaysMarkDirty);
@@ -515,7 +495,6 @@ void UActorComponent::PostEditUndo()
 	else
 	{
 		//Let the component be properly registered, after it was restored.
-		AActor* Owner = GetOwner();
 		if (Owner)
 		{
 			Owner->AddOwnedComponent(this);
@@ -537,7 +516,6 @@ void UActorComponent::ConsolidatedPostEditChange()
 		delete ReregisterContext;
 		EditReregisterContexts.Remove(this);
 
-		AActor* Owner = GetOwner();
 		if (Owner && !Owner->IsTemplate())
 		{
 			Owner->RerunConstructionScripts();
@@ -573,7 +551,6 @@ void UActorComponent::PostEditChangeChainProperty(FPropertyChangedChainEvent& Pr
 
 void UActorComponent::OnRegister()
 {
-	AActor* Owner = GetOwner();
 	checkf(!HasAnyFlags(RF_Unreachable), TEXT("%s"), *GetDetailedInfo());
 	checkf(!GetOuter()->IsTemplate(), TEXT("'%s' (%s)"), *GetOuter()->GetFullName(), *GetDetailedInfo());
 	checkf(!IsTemplate(), TEXT("'%s' (%s)"), *GetOuter()->GetFullName(), *GetDetailedInfo() );
@@ -672,7 +649,6 @@ FString FActorComponentTickFunction::DiagnosticMessage()
 
 bool UActorComponent::SetupActorComponentTickFunction(struct FTickFunction* TickFunction)
 {
-	AActor* Owner = GetOwner();
 	if(TickFunction->bCanEverTick && !IsTemplate() && (!Owner || !Owner->IsTemplate()))
 	{
 		ULevel* ComponentLevel = (Owner ? Owner->GetLevel() : GetWorld()->PersistentLevel);
@@ -778,9 +754,6 @@ void UActorComponent::RegisterComponentWithWorld(UWorld* InWorld)
 	// If not registered, should not have a scene
 	checkf(World == NULL, TEXT("%s"), *GetFullName());
 
-	// Follow outer chain to see if we can find an Actor
-	AActor* Owner = GetOwner();
-
 	checkSlow(Owner == NULL || Owner->OwnsComponent(this));
 
 	if ((Owner != NULL) && Owner->GetClass()->HasAnyClassFlags(CLASS_NewerVersionExists))
@@ -799,7 +772,7 @@ void UActorComponent::RegisterComponentWithWorld(UWorld* InWorld)
 
 		if(InWorld != Owner->GetWorld())
 		{
-			// The only time you should specify a scene that is not GetOwner()->GetWorld() is when you don't have an Actor
+			// The only time you should specify a scene that is not Owner->GetWorld() is when you don't have an Actor
 			UE_LOG(LogActorComponent, Log, TEXT("RegisterComponentWithWorld: (%s) Specifying a world, but an Owner Actor found, and InWorld is not GetOwner()->GetWorld()"), *GetPathName());
 		}
 	}
@@ -846,9 +819,9 @@ void UActorComponent::RegisterComponentWithWorld(UWorld* InWorld)
 
 void UActorComponent::RegisterComponent()
 {
-	if (ensure(GetOwner() && GetOwner()->GetWorld()))
+	if (ensure(Owner && Owner->GetWorld()))
 	{
-		RegisterComponentWithWorld(GetOwner()->GetWorld());
+		RegisterComponentWithWorld(Owner->GetWorld());
 	}
 }
 
@@ -871,8 +844,6 @@ void UActorComponent::UnregisterComponent()
 		IStreamingManager::Get().NotifyPrimitiveDetached( Primitive );
 	}
 
-	AActor* Owner = GetOwner(); // Get Owner while we are still registered
-
 	RegisterAllComponentTickFunctions(false);
 	ExecuteUnregisterEvents();
 
@@ -899,7 +870,6 @@ void UActorComponent::DestroyComponent(bool bPromoteChildren/*= false*/)
 	}
 
 	// Then remove from Components array, if we have an Actor
-	AActor* Owner = GetOwner();
 	if(Owner != NULL)
 	{
 		if (IsCreatedByConstructionScript())
@@ -938,8 +908,7 @@ void UActorComponent::OnComponentDestroyed()
 
 void UActorComponent::K2_DestroyComponent(UObject* Object)
 {
-	AActor* Owner = (Object != this ? GetOwner() : nullptr);
-	if (Owner == NULL || Owner == Object)
+	if (Object == this || Owner == NULL || Owner == Object)
 	{
 		DestroyComponent();
 	}
@@ -1107,7 +1076,6 @@ void UActorComponent::ConditionalTickComponent(float DeltaTime, enum ELevelTick
 	if(bRegistered && !IsPendingKill())
 	{
 		//@optimization, I imagine this is all unnecessary in a shipping game with no editor
-		AActor* Owner = GetOwner();
 		if (TickType != LEVELTICK_ViewportsOnly || 
 			(bTickInEditor && TickType == LEVELTICK_ViewportsOnly) ||
 			(Owner && Owner->ShouldTickIfViewportsOnly())
@@ -1325,7 +1293,6 @@ void UActorComponent::SetTickableWhenPaused(bool bTickableWhenPaused)
 
 bool UActorComponent::IsOwnerRunningUserConstructionScript() const
 {
-	const AActor* Owner = GetOwner();
 	return Owner != nullptr && Owner->bRunningUserConstructionScript;
 }
 
@@ -1397,7 +1364,6 @@ void UActorComponent::SetIsReplicated(bool ShouldReplicate)
 	check(GetComponentClassCanReplicate()); // Only certain component classes can replicate!
 	bReplicates = ShouldReplicate;
 
-	AActor* Owner = GetOwner();
 	if (Owner)
 	{
 		Owner->UpdateReplicatedComponent( this );
@@ -1421,12 +1387,7 @@ bool UActorComponent::GetComponentClassCanReplicate() const
 
 ENetRole UActorComponent::GetOwnerRole() const
 {
-	AActor *Actor = GetOwner();
-	if (!Actor)
-	{
-		return ROLE_None;
-	}
-	return Actor->Role;
+	return (Owner ? Owner->Role.GetValue() : ROLE_None);
 }
 
 bool UActorComponent::IsNetSimulating() const
