Hopefully fix CIS failures that aren't showing up in VS

D:\\BuildFarm\\buildmachine_++depot+UE4\\Engine\\Source\\Runtime\\Engine\\Private\\Animation\\AnimationRuntime.cpp(36) : error C2719: 'BlendWeight': formal parameter with __declspec(align('16')) won't be aligned

diff --git a/Engine/Source/Runtime/Engine/Private/Animation/AnimationRuntime.cpp b/Engine/Source/Runtime/Engine/Private/Animation/AnimationRuntime.cpp
--- a/Engine/Source/Runtime/Engine/Private/Animation/AnimationRuntime.cpp
+++ b/Engine/Source/Runtime/Engine/Private/Animation/AnimationRuntime.cpp
@@ -30,21 +30,23 @@ namespace ETransformBlendMode
 }
 
 template<int32>
-void BlendTransform(const FTransform& Source, FTransform& Dest, const ScalarRegister BlendWeight)
+void BlendTransform(const FTransform& Source, FTransform& Dest, const float BlendWeight)
 {
 	check(false); /// should never call this
 }
 
 template<>
-void BlendTransform<ETransformBlendMode::Overwrite>(const FTransform& Source, FTransform& Dest, const ScalarRegister BlendWeight)
+void BlendTransform<ETransformBlendMode::Overwrite>(const FTransform& Source, FTransform& Dest, const float BlendWeight)
 {
-	Dest = Source * BlendWeight;
+	const ScalarRegister VBlendWeight(BlendWeight);
+	Dest = Source * VBlendWeight;
 }
 
 template<>
-void BlendTransform<ETransformBlendMode::Accumulate>(const FTransform& Source, FTransform& Dest, const ScalarRegister BlendWeight)
+void BlendTransform<ETransformBlendMode::Accumulate>(const FTransform& Source, FTransform& Dest, const float BlendWeight)
 {
-	Dest.AccumulateWithShortestRotation(Source, BlendWeight);
+	const ScalarRegister VBlendWeight(BlendWeight);
+	Dest.AccumulateWithShortestRotation(Source, VBlendWeight);
 }
 
 //////////////////////////////////////////////////////////////////////////
@@ -80,12 +82,12 @@ void FAnimationRuntime::InitializeTransform(const FBoneContainer& RequiredBones,
 }
 
 template <int32 TRANSFORM_BLEND_MODE>
-FORCEINLINE void BlendPose(const FTransformArrayA2& SourcePoses, FTransformArrayA2& ResultAtoms, const TArray<FBoneIndexType> & RequiredBoneIndices, const ScalarRegister VBlendWeight)
+FORCEINLINE void BlendPose(const FTransformArrayA2& SourcePoses, FTransformArrayA2& ResultAtoms, const TArray<FBoneIndexType> & RequiredBoneIndices, const float BlendWeight)
 {
 	for (int32 i = 0; i < RequiredBoneIndices.Num(); ++i)
 	{
 		const int32 BoneIndex = RequiredBoneIndices[i];
-		BlendTransform<TRANSFORM_BLEND_MODE>(SourcePoses[BoneIndex], ResultAtoms[BoneIndex], VBlendWeight);
+		BlendTransform<TRANSFORM_BLEND_MODE>(SourcePoses[BoneIndex], ResultAtoms[BoneIndex], BlendWeight);
 	}
 }
 
@@ -109,13 +111,13 @@ void FAnimationRuntime::BlendPosesTogether(int32 NumPoses, const FTransformArray
 	{
 		// debug purpose for now, but this can cause 0 bone transform, so we'd like to catch it 
 		float WeightSum = SourceWeights[0];
-		BlendPose<ETransformBlendMode::Overwrite>(*SourcePoses[0], ResultAtoms, RequiredBoneIndices, ScalarRegister(SourceWeights[0]));
+		BlendPose<ETransformBlendMode::Overwrite>(*SourcePoses[0], ResultAtoms, RequiredBoneIndices, SourceWeights[0]);
 
 		for (int32 i = 1; i < NumPoses; ++i)
 		{
 			WeightSum += SourceWeights[i];
 
-			BlendPose<ETransformBlendMode::Accumulate>(*SourcePoses[i], ResultAtoms, RequiredBoneIndices, ScalarRegister(SourceWeights[i]));
+			BlendPose<ETransformBlendMode::Accumulate>(*SourcePoses[i], ResultAtoms, RequiredBoneIndices, SourceWeights[i]);
 		}
 
 		ensure (WeightSum != 0.f);
@@ -145,11 +147,11 @@ void FAnimationRuntime::BlendPosesTogether(int32 NumPoses, const TArray<FTransfo
 
 	const TArray<FBoneIndexType> & RequiredBoneIndices = RequiredBones.GetBoneIndicesArray();
 
-	BlendPose<ETransformBlendMode::Overwrite>(SourcePoses[0], ResultAtoms, RequiredBoneIndices, ScalarRegister(SourceWeights[0]));
+	BlendPose<ETransformBlendMode::Overwrite>(SourcePoses[0], ResultAtoms, RequiredBoneIndices, SourceWeights[0]);
 
 	for (int32 i = 1; i < NumPoses; ++i)
 	{
-		BlendPose<ETransformBlendMode::Accumulate>(SourcePoses[i], ResultAtoms, RequiredBoneIndices, ScalarRegister(SourceWeights[i]));
+		BlendPose<ETransformBlendMode::Accumulate>(SourcePoses[i], ResultAtoms, RequiredBoneIndices, SourceWeights[i]);
 	}
 
 	// Ensure that all of the resulting rotations are normalized
@@ -162,11 +164,11 @@ void FAnimationRuntime::BlendPosesTogether(int32 NumPoses, const TArray<FTransfo
 template <int32 TRANSFORM_BLEND_MODE>
 void BlendPosePerBone(const TArray<FBoneIndexType> &RequiredBoneIndices, const TArray<int32>& PerBoneIndices, const FBlendSampleData& BlendSampleDataCache, FTransformArrayA2 &ResultAtoms, const FTransformArrayA2& SourceAtoms)
 {
-	const ScalarRegister VBlendWeight(BlendSampleDataCache.GetWeight());
-	TArray<ScalarRegister> PerBoneBlends;
+	const float BlendWeight = BlendSampleDataCache.GetWeight();
+	TArray<float> PerBoneBlends;
 	for (int32 i = 0; i < BlendSampleDataCache.PerBoneBlendData.Num(); ++i)
 	{
-		PerBoneBlends.Add(ScalarRegister(FMath::Clamp<float>(BlendSampleDataCache.PerBoneBlendData[i], 0.f, 1.f)));
+		PerBoneBlends.Add(FMath::Clamp<float>(BlendSampleDataCache.PerBoneBlendData[i], 0.f, 1.f));
 	}
 
 	for (int32 i = 0; i < RequiredBoneIndices.Num(); ++i)
@@ -175,7 +177,7 @@ void BlendPosePerBone(const TArray<FBoneIndexType> &RequiredBoneIndices, const T
 		int32 PerBoneIndex = PerBoneIndices[i];
 		if (PerBoneIndex == INDEX_NONE || !BlendSampleDataCache.PerBoneBlendData.IsValidIndex(PerBoneIndex))
 		{
-			BlendTransform<TRANSFORM_BLEND_MODE>(SourceAtoms[BoneIndex], ResultAtoms[BoneIndex], VBlendWeight);
+			BlendTransform<TRANSFORM_BLEND_MODE>(SourceAtoms[BoneIndex], ResultAtoms[BoneIndex], BlendWeight);
 		}
 		else
 		{
@@ -275,9 +277,7 @@ void FAnimationRuntime::BlendPosesTogetherPerBoneInMeshSpace(int32 NumPoses, TAr
  */
 void FAnimationRuntime::BlendPosesAccumulate(const FTransformArrayA2& BlendPoses, const float BlendWeight, const FBoneContainer& RequiredBones, /*inout*/ FTransformArrayA2& ResultAtoms)
 {
-	const ScalarRegister VBlendWeight(BlendWeight);
-
-	BlendPose<ETransformBlendMode::Accumulate>(BlendPoses, ResultAtoms, RequiredBones.GetBoneIndicesArray(), VBlendWeight);
+	BlendPose<ETransformBlendMode::Accumulate>(BlendPoses, ResultAtoms, RequiredBones.GetBoneIndicesArray(), BlendWeight);
 }
 
 void FAnimationRuntime::LerpBoneTransforms(TArray<FTransform> & A, const TArray<FTransform> & B, float Alpha, const TArray<FBoneIndexType> & RequiredBonesArray)
