#ue4 - Perf: inline UActorComponent::GetOwner()

diff --git a/Engine/Source/Runtime/Engine/Classes/Components/ActorComponent.h b/Engine/Source/Runtime/Engine/Classes/Components/ActorComponent.h
--- a/Engine/Source/Runtime/Engine/Classes/Components/ActorComponent.h
+++ b/Engine/Source/Runtime/Engine/Classes/Components/ActorComponent.h
@@ -634,4 +634,13 @@ private:
 	virtual void Tick( float DeltaTime ) final { check(0); }
 
 #endif
-};
\ No newline at end of file
+};
+
+//////////////////////////////////////////////////////////////////////////
+// UActorComponent inlines
+
+FORCEINLINE_DEBUGGABLE class AActor* UActorComponent::GetOwner() const
+{
+	checkSlow(Owner == GetTypedOuter<AActor>()); // verify cached value is correct
+	return Owner;
+}
diff --git a/Engine/Source/Runtime/Engine/Private/Components/ActorComponent.cpp b/Engine/Source/Runtime/Engine/Private/Components/ActorComponent.cpp
--- a/Engine/Source/Runtime/Engine/Private/Components/ActorComponent.cpp
+++ b/Engine/Source/Runtime/Engine/Private/Components/ActorComponent.cpp
@@ -274,13 +274,6 @@ bool UActorComponent::IsOwnerSelected() const
 	return Owner && Owner->IsSelected();
 }
 
-AActor* UActorComponent::GetOwner() const
-{
-	// walk up outer chain to find an Actor
-	checkSlow(Owner == GetTypedOuter<AActor>());
-	return Owner;
-}
-
 UWorld* UActorComponent::GetWorld() const
 {
 	UWorld* ComponentWorld = World;
