#ue4 - Perf: UActorComponent::GetOwner() is fast now, so making caching logic in MoveComponent() simpler.  Partial revert of CL 2520396.

diff --git a/Engine/Source/Runtime/Engine/Private/Components/PrimitiveComponent.cpp b/Engine/Source/Runtime/Engine/Private/Components/PrimitiveComponent.cpp
--- a/Engine/Source/Runtime/Engine/Private/Components/PrimitiveComponent.cpp
+++ b/Engine/Source/Runtime/Engine/Private/Components/PrimitiveComponent.cpp
@@ -1380,7 +1380,7 @@ bool UPrimitiveComponent::MoveComponentImpl( const FVector& Delta, const FQuat&
 	TArray<FOverlapInfo> PendingOverlaps;
 	TArray<FOverlapInfo> OverlapsAtEndLocation;
 	TArray<FOverlapInfo>* OverlapsAtEndLocationPtr = NULL; // When non-null, used as optimization to avoid work in UpdateOverlaps.
-	AActor* CachedActor = nullptr; // Try to defer or avoid GetOwner() call.
+	AActor* const Actor = GetOwner();
 
 	if ( !bSweep )
 	{
@@ -1391,8 +1391,6 @@ bool UPrimitiveComponent::MoveComponentImpl( const FVector& Delta, const FQuat&
 	{
 		TArray<FHitResult> Hits;
 		FVector NewLocation = TraceStart;
-		CachedActor = CachedActor ? CachedActor : GetOwner();
-		AActor* const Actor = CachedActor;
 
 		// Perform movement collision checking if needed for this actor.
 		const bool bCollisionEnabled = IsCollisionEnabled();
@@ -1612,8 +1610,7 @@ bool UPrimitiveComponent::MoveComponentImpl( const FVector& Delta, const FQuat&
 		}
 		else
 		{
-			CachedActor = CachedActor ? CachedActor : GetOwner();
-			DispatchBlockingHit(*CachedActor, BlockingHit);
+			DispatchBlockingHit(*Actor, BlockingHit);
 		}
 	}
 
