Exposed 'DateTime-DateTime' to blueprint.

diff --git a/Engine/Source/Runtime/Engine/Classes/Kismet/KismetMathLibrary.h b/Engine/Source/Runtime/Engine/Classes/Kismet/KismetMathLibrary.h
--- a/Engine/Source/Runtime/Engine/Classes/Kismet/KismetMathLibrary.h
+++ b/Engine/Source/Runtime/Engine/Classes/Kismet/KismetMathLibrary.h
@@ -784,7 +784,11 @@ class ENGINE_API UKismetMathLibrary : public UBlueprintFunctionLibrary
 
 	/* Subtraction (A - B) */
 	UFUNCTION(BlueprintPure, meta=(DisplayName="DateTime - Timespan", CompactNodeTitle="-", Keywords="- subtract minus"), Category="Math|DateTime")
-	static FDateTime Subtract_DateTimeTimespan( FDateTime A, FTimespan B );
+	static FDateTime Subtract_DateTimeTimespan(FDateTime A, FTimespan B);
+
+	/* Subtraction (A - B) */
+	UFUNCTION(BlueprintPure, meta = (DisplayName = "DateTime - DateTime", CompactNodeTitle = "-", Keywords = "- subtract minus"), Category = "Math|DateTime")
+	static FTimespan Subtract_DateTimeDateTime(FDateTime A, FDateTime B);
 
 	/* Returns true if the values are equal (A == B) */
 	UFUNCTION(BlueprintPure, meta=(DisplayName="Equal (DateTime)", CompactNodeTitle="==", Keywords="== equal"), Category="Math|DateTime")
diff --git a/Engine/Source/Runtime/Engine/Private/KismetMathLibrary.cpp b/Engine/Source/Runtime/Engine/Private/KismetMathLibrary.cpp
--- a/Engine/Source/Runtime/Engine/Private/KismetMathLibrary.cpp
+++ b/Engine/Source/Runtime/Engine/Private/KismetMathLibrary.cpp
@@ -1179,6 +1179,11 @@ FDateTime UKismetMathLibrary::Subtract_DateTimeTimespan( FDateTime A, FTimespan
 }
 
 
+FTimespan UKismetMathLibrary::Subtract_DateTimeDateTime(FDateTime A, FDateTime B)
+{
+	return A - B;
+}
+
 bool UKismetMathLibrary::EqualEqual_DateTimeDateTime( FDateTime A, FDateTime B )
 {
 	return A == B;
