Fixed crash clicking on columns in the visualize RT window

diff --git a/Engine/Source/Runtime/Engine/Private/UserInterface/VisualizeRT.cpp b/Engine/Source/Runtime/Engine/Private/UserInterface/VisualizeRT.cpp
--- a/Engine/Source/Runtime/Engine/Private/UserInterface/VisualizeRT.cpp
+++ b/Engine/Source/Runtime/Engine/Private/UserInterface/VisualizeRT.cpp
@@ -113,9 +113,10 @@ struct FReferenceToRenderer
 	}
 };
 
-struct FVisualizeRTWidget : public SCompoundWidget, public FReferenceToRenderer
+class SVisualizeRTWidget : public SCompoundWidget, public FReferenceToRenderer
 {
-	SLATE_BEGIN_ARGS(FVisualizeRTWidget)
+public:
+	SLATE_BEGIN_ARGS(SVisualizeRTWidget)
 	{
 	}
 
@@ -137,8 +138,8 @@ struct FVisualizeRTWidget : public SCompoundWidget, public FReferenceToRenderer
 				SNew( SListView< TSharedPtr<FRTInfo> > )
 				//.ItemHeight(24)
 				.ListItemsSource( &ListItems )
-				.OnGenerateRow( this, &FVisualizeRTWidget::OnGenerateWidgetForList )
-				.OnSelectionChanged(this, &FVisualizeRTWidget::OnSelectionChanged)
+				.OnGenerateRow( this, &SVisualizeRTWidget::OnGenerateWidgetForList )
+				.OnSelectionChanged(this, &SVisualizeRTWidget::OnSelectionChanged)
 				.HeaderRow
 				(
 					SNew(SHeaderRow)
@@ -193,7 +194,7 @@ struct FVisualizeRTWidget : public SCompoundWidget, public FReferenceToRenderer
 					SNew( SHorizontalBox )
 					+SHorizontalBox::Slot().Padding(24)
 					[
-						SNew(SImage).Image(this, &FVisualizeRTWidget::GetRTImage)
+						SNew(SImage).Image(this, &SVisualizeRTWidget::GetRTImage)
 					]
 				]
 			]
@@ -237,11 +238,14 @@ struct FVisualizeRTWidget : public SCompoundWidget, public FReferenceToRenderer
 
 	void OnSelectionChanged(TSharedPtr<FRTInfo> Selection, ESelectInfo::Type SelectInfo)
 	{
-		Selected = Selection.Get();
-		if (Selection->Number != TEXT("-"))
+		if( Selection.IsValid() )
 		{
-			FString Cmd = Selected->Number;
-			RendererModule.ExecVisualizeTextureCmd(Cmd);
+			Selected = Selection.Get();
+			if(Selection->Number != TEXT("-"))
+			{
+				FString Cmd = Selected->Number;
+				RendererModule.ExecVisualizeTextureCmd(Cmd);
+			}
 		}
 	}
 
@@ -336,7 +340,7 @@ bool HandleVisualizeRT()
 		.SupportsMinimize(false)
 		.SizingRule(ESizingRule::UserSized);
 	Window = FSlateApplication::Get().AddWindow( Window.ToSharedRef() );
-	Window->SetContent(SNew(FVisualizeRTWidget));
+	Window->SetContent(SNew(SVisualizeRTWidget));
 
 	return true;
 }
