Tweaked data layouts of a UActorComponent, USceneComponent, FTickFunction

UActorComponent went from 216 bytes with 11 bytes padding to 200 with 3 bytes padding.
USceneComponent went from 496 bytes with 29 bytes padding to 464 with 13 bytes padding.
FTickFunction went from 72 bytes with 10 bytes padding to 64 with 2 bytes padding.

diff --git a/Engine/Source/Runtime/Engine/Classes/Components/ActorComponent.h b/Engine/Source/Runtime/Engine/Classes/Components/ActorComponent.h
--- a/Engine/Source/Runtime/Engine/Classes/Components/ActorComponent.h
+++ b/Engine/Source/Runtime/Engine/Classes/Components/ActorComponent.h
@@ -142,8 +142,13 @@ private:
 
 	friend class FActorComponentInstanceData;
 
+public:
+	UPROPERTY()
+	EComponentCreationMethod CreationMethod;
+
 	AActor* Owner;
 
+private:
 	UPROPERTY()
 	TArray<FSimpleMemberReference> UCSModifiedProperties;
 
@@ -160,9 +165,6 @@ public:
 	bool HasBeenInitialized() const { return bHasBeenInitialized; }
 	bool HasBegunPlay() const { return bHasBegunPlay; }
 
-	UPROPERTY()
-	EComponentCreationMethod CreationMethod;
-
 	bool IsCreatedByConstructionScript() const;
 
 	UFUNCTION()
diff --git a/Engine/Source/Runtime/Engine/Classes/Components/SceneComponent.h b/Engine/Source/Runtime/Engine/Classes/Components/SceneComponent.h
--- a/Engine/Source/Runtime/Engine/Classes/Components/SceneComponent.h
+++ b/Engine/Source/Runtime/Engine/Classes/Components/SceneComponent.h
@@ -97,8 +97,17 @@
 
 	virtual void GetLifetimeReplicatedProps(TArray<FLifetimeProperty>& OutLifetimeProps) const override;
 
-	/** Current transform of this component, relative to the world */
-	FTransform ComponentToWorld;
+	/** What we are currently attached to. If valid, RelativeLocation etc. are used relative to this object */
+	UPROPERTY(Replicated)
+	class USceneComponent* AttachParent;
 
+	/** List of child SceneComponents that are attached to us. */
+	UPROPERTY(transient)
+	TArray< USceneComponent* > AttachChildren;
+
+	/** Optional socket name on AttachParent that we are attached to. */
+	UPROPERTY(Replicated)
+	FName AttachSocketName;
+
 	/** if true, will call GetCustomLocation instead or returning the location part of ComponentToWorld */
 	UPROPERTY()
@@ -158,35 +167,19 @@
 	// Transient flag that temporarily disables UpdateOverlaps within DetachFromParent().
 	uint32 bDisableDetachmentUpdateOverlaps:1;
 
-#if WITH_EDITORONLY_DATA
-protected:
-	/** Editor only component used to display the sprite so as to be able to see the location of the Audio Component  */
-	class UBillboardComponent* SpriteComponent;
+	/** Physics Volume in which this SceneComponent is located **/
+	UPROPERTY(transient)
+	class APhysicsVolume* PhysicsVolume;
 
 public:
-	UPROPERTY()
-	uint32 bVisualizeComponent:1;
-#endif
+	/** Current transform of this component, relative to the world */
+	FTransform ComponentToWorld;
 
-public:
-	/** How often this component is allowed to move, used to make various optimizations. Only safe to set in constructor, use SetMobility() during runtime. */
-	UPROPERTY(Category=Mobility, EditAnywhere, BlueprintReadOnly)
-	TEnumAsByte<EComponentMobility::Type> Mobility;
-
 	/** Current bounds of this component */
 	FBoxSphereBounds Bounds;
 
-	/** What we are currently attached to. If valid, RelativeLocation etc. are used relative to this object */
-	UPROPERTY(Replicated)
-	class USceneComponent* AttachParent;
+	UPROPERTY()
+	FVector RelativeTranslation_DEPRECATED;
 
-	/** Optional socket name on AttachParent that we are attached to. */
-	UPROPERTY(Replicated)
-	FName AttachSocketName;
-
-	/** List of child SceneComponents that are attached to us. */
-	UPROPERTY(transient)
-	TArray< USceneComponent* > AttachChildren;
-
 	/** Location of this component relative to its parent */
 	UPROPERTY(EditAnywhere, BlueprintReadOnly, ReplicatedUsing=OnRep_Transform, Category = Transform)
@@ -196,8 +189,9 @@
 	UPROPERTY(EditAnywhere, BlueprintReadOnly, ReplicatedUsing=OnRep_Transform, Category=Transform)
 	FRotator RelativeRotation;
 
-	UPROPERTY()
-	FVector RelativeTranslation_DEPRECATED;
+	/** How often this component is allowed to move, used to make various optimizations. Only safe to set in constructor, use SetMobility() during runtime. */
+	UPROPERTY(Category = Mobility, EditAnywhere, BlueprintReadOnly)
+	TEnumAsByte<EComponentMobility::Type> Mobility;
 
 	/** If detail mode is >= system detail mode, primitive won't be rendered. */
 	UPROPERTY(EditAnywhere, AdvancedDisplay, BlueprintReadOnly, Category=LOD)
@@ -207,8 +201,8 @@
 
 	bool NetUpdateTransform;
 
-	USceneComponent *NetOldAttachParent;
 	FName NetOldAttachSocketName;
+	USceneComponent *NetOldAttachParent;
 
 	UFUNCTION()
 	void OnRep_Transform();
@@ -235,14 +229,7 @@
 	UPROPERTY()
 	FVector ComponentVelocity;
 
-protected:
-
-	/** Physics Volume in which this SceneComponent is located **/
-	UPROPERTY(transient)
-	class APhysicsVolume* PhysicsVolume;
-
 public:
-
 	/** Set the location of this component relative to its parent */
 	UFUNCTION(BlueprintCallable, Category="Utilities|Transformation", meta=(DisplayName="SetRelativeLocation"))
 	void K2_SetRelativeLocation(FVector NewLocation, bool bSweep, FHitResult& SweepHitResult);
@@ -611,6 +598,16 @@
 
 	friend class FScopedMovementUpdate;
 
+#if WITH_EDITORONLY_DATA
+protected:
+	/** Editor only component used to display the sprite so as to be able to see the location of the Audio Component  */
+	class UBillboardComponent* SpriteComponent;
+
+public:
+	UPROPERTY()
+	uint32 bVisualizeComponent : 1;
+#endif
+
 public:
 
 	/** Called when AttachParent changes, to allow the scene to update its attachment state. */
diff --git a/Engine/Source/Runtime/Engine/Classes/Engine/EngineBaseTypes.h b/Engine/Source/Runtime/Engine/Classes/Engine/EngineBaseTypes.h
--- a/Engine/Source/Runtime/Engine/Classes/Engine/EngineBaseTypes.h
+++ b/Engine/Source/Runtime/Engine/Classes/Engine/EngineBaseTypes.h
@@ -141,6 +141,11 @@ public:
 	UPROPERTY()
 	TEnumAsByte<enum ETickingGroup> TickGroup;
 
+protected:
+	/** Internal data that indicates the tick group we actually executed in (it may have been delayed due to prerequisites) **/
+	TEnumAsByte<enum ETickingGroup> ActualTickGroup;
+
+public:
 	/** Bool indicating that this function should execute even if the game is paused. Pause ticks are very limited in capabilities. **/
 	UPROPERTY(EditDefaultsOnly, Category="Tick", AdvancedDisplay)
 	uint32 bTickEvenWhenPaused:1;
@@ -176,10 +181,6 @@ private:
 	/** Internal data to track if we have finshed visiting this tick function yet this frame **/
 	int32 TickQueuedGFrameCounter;
 
-protected:
-	/** Internal data that indicates the tick group we actually executed in (it may have been delayed due to prerequisites) **/
-	TEnumAsByte<enum ETickingGroup> ActualTickGroup;
-
 private:
 	/** Completion handle for the task that will run this tick. Caution, this is no reset to nullptr until an unspecified future time **/
 	FGraphEventRef CompletionHandle;
