Tweaked AActor data layout to remove 8 bytes of unnecessary padding.

diff --git a/Engine/Source/Runtime/Engine/Classes/GameFramework/Actor.h b/Engine/Source/Runtime/Engine/Classes/GameFramework/Actor.h
--- a/Engine/Source/Runtime/Engine/Classes/GameFramework/Actor.h
+++ b/Engine/Source/Runtime/Engine/Classes/GameFramework/Actor.h
@@ -336,19 +336,19 @@ public:
 	UPROPERTY()
 	uint32 bRelevantForNetworkReplays:1;
 
+	/** The time this actor was created, relative to World->GetTimeSeconds().
+	* @see UWorld::GetTimeSeconds()
+	*/
+	float CreationTime;
+
 	/** Pawn responsible for damage caused by this actor. */
 	UPROPERTY(BlueprintReadWrite, replicatedUsing=OnRep_Instigator, meta=(ExposeOnSpawn=true), Category=Actor)
 	class APawn* Instigator;
 
-	/** Called on clients when Instigator is replicated. */
+	/** Called on clients when Instigatolr is replicated. */
 	UFUNCTION()
 	virtual void OnRep_Instigator();
 
-	/** The time this actor was created, relative to World->GetTimeSeconds().
-	 * @see UWorld::GetTimeSeconds()
-	 */
-	float CreationTime;
-
 	/** Array of Actors whose Owner is this actor */
 	UPROPERTY(transient)
 	TArray<AActor*> Children;
