Editor - Continuing to itterate on editor settings autodiscovery.

diff --git a/Engine/Source/Developer/SettingsEditor/Private/SettingsEditorModule.cpp b/Engine/Source/Developer/SettingsEditor/Private/SettingsEditorModule.cpp
--- a/Engine/Source/Developer/SettingsEditor/Private/SettingsEditorModule.cpp
+++ b/Engine/Source/Developer/SettingsEditor/Private/SettingsEditorModule.cpp
@@ -177,7 +177,7 @@ private:
 
 	void UpdateSettings(bool bForce = false)
 	{
-		if ( (AnyActiveSettingsEditor() && bAreSettingsStale) || bForce )
+		if ( ( AnyActiveSettingsEditor() || bForce ) && bAreSettingsStale )
 		{
 			bAreSettingsStale = false;
 
diff --git a/Engine/Source/Editor/EditorSettingsViewer/Private/EditorSettingsViewerModule.cpp b/Engine/Source/Editor/EditorSettingsViewer/Private/EditorSettingsViewerModule.cpp
--- a/Engine/Source/Editor/EditorSettingsViewer/Private/EditorSettingsViewerModule.cpp
+++ b/Engine/Source/Editor/EditorSettingsViewer/Private/EditorSettingsViewerModule.cpp
@@ -173,13 +173,6 @@ protected:
 			LOCTEXT("LevelEditorViewportSettingsDescription", "Configure the look and feel of the Level Editor view ports."),
 			GetMutableDefault<ULevelEditorViewportSettings>()
 		);
-
-		// miscellaneous settings
-		SettingsModule.RegisterSettings("Editor", "LevelEditor", "Misc",
-			LOCTEXT("LevelEditorMiscSettingsName", "Miscellaneous"),
-			LOCTEXT("LevelEditorMiscSettingsDescription", "Configure miscellaneous settings for the Level Editor."),
-			GetMutableDefault<ULevelEditorMiscSettings>()
-		);
 	}
 
 	/**
diff --git a/Engine/Source/Editor/UnrealEd/Classes/Settings/LevelEditorMiscSettings.h b/Engine/Source/Editor/UnrealEd/Classes/Settings/LevelEditorMiscSettings.h
--- a/Engine/Source/Editor/UnrealEd/Classes/Settings/LevelEditorMiscSettings.h
+++ b/Engine/Source/Editor/UnrealEd/Classes/Settings/LevelEditorMiscSettings.h
@@ -2,15 +2,15 @@
 
 #pragma once
 
+#include "Engine/DeveloperSettings.h"
 #include "LevelEditorMiscSettings.generated.h"
 
 
 /**
- * Implements the Level Editor's miscellaneous settings.
+ * Configure miscellaneous settings for the Level Editor.
  */
-UCLASS(config=EditorPerProjectUserSettings)
-class UNREALED_API ULevelEditorMiscSettings
-	: public UObject
+UCLASS(config=EditorPerProjectUserSettings, meta=( DisplayName="Miscellaneous" ))
+class UNREALED_API ULevelEditorMiscSettings : public UDeveloperSettings
 {
 	GENERATED_UCLASS_BODY()
 
diff --git a/Engine/Source/Editor/UnrealEd/Private/Commandlets/CookCommandlet.cpp b/Engine/Source/Editor/UnrealEd/Private/Commandlets/CookCommandlet.cpp
--- a/Engine/Source/Editor/UnrealEd/Private/Commandlets/CookCommandlet.cpp
+++ b/Engine/Source/Editor/UnrealEd/Private/Commandlets/CookCommandlet.cpp
@@ -31,6 +31,7 @@ DEFINE_LOG_CATEGORY_STATIC(LogCookCommandlet, Log, All);
 UCookerSettings::UCookerSettings(const FObjectInitializer& ObjectInitializer)
 	: Super(ObjectInitializer)
 {
+	SectionName = TEXT("Cooker");
 	DefaultPVRTCQuality = 1;
 	DefaultASTCQualityBySize = 3;
 	DefaultASTCQualityBySpeed = 3;
@@ -350,8 +351,8 @@ bool UCookCommandlet::SaveCookedPackage( UPackage* Package, uint32 SaveFlags, bo
 		Filename = SandboxFile->ConvertToAbsolutePathForExternalAppForWrite(*Filename);
 
 		uint32 OriginalPackageFlags = Package->PackageFlags;
- 		UWorld* World = NULL;
- 		EObjectFlags Flags = RF_NoFlags;
+		UWorld* World = NULL;
+		EObjectFlags Flags = RF_NoFlags;
 		bool bPackageFullyLoaded = false;
 
 		if (bCompressed)
diff --git a/Engine/Source/Editor/UnrealEd/Private/Settings/SettingsClasses.cpp b/Engine/Source/Editor/UnrealEd/Private/Settings/SettingsClasses.cpp
--- a/Engine/Source/Editor/UnrealEd/Private/Settings/SettingsClasses.cpp
+++ b/Engine/Source/Editor/UnrealEd/Private/Settings/SettingsClasses.cpp
@@ -297,7 +297,8 @@ bool FAutoReimportDirectoryConfig::ParseSourceDirectoryAndMountPoint(FString& So
 
 UEditorMiscSettings::UEditorMiscSettings( const FObjectInitializer& ObjectInitializer )
 	: Super(ObjectInitializer)
-{ }
+{
+}
 
 
 /* ULevelEditorMiscSettings interface
@@ -307,6 +308,8 @@ ULevelEditorMiscSettings::ULevelEditorMiscSettings( const FObjectInitializer& Ob
 	: Super(ObjectInitializer)
 {
 	bAutoApplyLightingEnable = true;
+	SectionName = TEXT("Misc");
+	CategoryName = TEXT("LevelEditor");
 }
 
 
diff --git a/Engine/Source/Runtime/Engine/Classes/Engine/DeveloperSettings.h b/Engine/Source/Runtime/Engine/Classes/Engine/DeveloperSettings.h
--- a/Engine/Source/Runtime/Engine/Classes/Engine/DeveloperSettings.h
+++ b/Engine/Source/Runtime/Engine/Classes/Engine/DeveloperSettings.h
@@ -4,14 +4,6 @@
 
 #include "DeveloperSettings.generated.h"
 
-/** The different types of settings container that can be used. */
-UENUM()
-enum class EDeveloperSettingsContainer : uint8
-{
-	Project,
-	Editor
-};
-
 /**
  * The base class of any auto discovered settings object.
  */
@@ -23,14 +15,12 @@ class ENGINE_API UDeveloperSettings : public UObject
 public:
 	UDeveloperSettings(const FObjectInitializer& ObjectInitializer);
 
-	/** Gets the container type, Project or Editor */
-	EDeveloperSettingsContainer GetContainerType() const { return Container; }
 	/** Gets the settings container name for the settings, either Project or Editor */
-	FName GetContainerName() const;
+	virtual FName GetContainerName() const;
 	/** Gets the category for the settings, some high level grouping like, Editor, Engine, Game...etc. */
-	FName GetCategoryName() const;
+	virtual FName GetCategoryName() const;
 	/** The unique name for your section of settings, uses the class's FName. */
-	FName GetSectionName() const;
+	virtual FName GetSectionName() const;
 
 #if WITH_EDITOR
 	/** Gets the section text, uses the classes DisplayName by default. */
@@ -42,7 +32,7 @@ public:
 	/** Gets a custom widget for the settings.  This is only for very custom situations. */
 	virtual TSharedPtr<class SWidget> GetCustomSettingsWidget() const;
 
-private:
+protected:
 	/**
 	 * The category name to use, overrides the one detected by looking at the config=... class metadata.
 	 * Arbitrary category names are not supported, this must map to an existing category we support in the settings
@@ -50,6 +40,6 @@ private:
 	 */
 	FName CategoryName;
 
-	/** The container for these settings. */
-	EDeveloperSettingsContainer Container;
+	/** The Section name, is the short name for the settings.  If not filled in, will be the FName of the class. */
+	FName SectionName;
 };
diff --git a/Engine/Source/Runtime/Engine/Private/Animation/AnimationSettings.cpp b/Engine/Source/Runtime/Engine/Private/Animation/AnimationSettings.cpp
--- a/Engine/Source/Runtime/Engine/Private/Animation/AnimationSettings.cpp
+++ b/Engine/Source/Runtime/Engine/Private/Animation/AnimationSettings.cpp
@@ -22,6 +22,8 @@ UAnimationSettings::UAnimationSettings(const FObjectInitializer& ObjectInitializ
 	, bTryIntervalKeyRemoval(true)
 
 {
+	SectionName = TEXT("Animation");
+
 	KeyEndEffectorsMatchNameArray.Add(TEXT("IK"));
 	KeyEndEffectorsMatchNameArray.Add(TEXT("eye"));
 	KeyEndEffectorsMatchNameArray.Add(TEXT("weapon"));
diff --git a/Engine/Source/Runtime/Engine/Private/AudioSettings.cpp b/Engine/Source/Runtime/Engine/Private/AudioSettings.cpp
--- a/Engine/Source/Runtime/Engine/Private/AudioSettings.cpp
+++ b/Engine/Source/Runtime/Engine/Private/AudioSettings.cpp
@@ -11,4 +11,5 @@
 UAudioSettings::UAudioSettings(const FObjectInitializer& ObjectInitializer)
 : Super(ObjectInitializer)
 {
-}
\ No newline at end of file
+	SectionName = TEXT("Audio");
+}
diff --git a/Engine/Source/Runtime/Engine/Private/Collision/CollisionProfile.cpp b/Engine/Source/Runtime/Engine/Private/Collision/CollisionProfile.cpp
--- a/Engine/Source/Runtime/Engine/Private/Collision/CollisionProfile.cpp
+++ b/Engine/Source/Runtime/Engine/Private/Collision/CollisionProfile.cpp
@@ -43,6 +43,7 @@ ENGINE_API const FName UCollisionProfile::DefaultProjectile_ProfileName = FName(
 UCollisionProfile::UCollisionProfile(const FObjectInitializer& ObjectInitializer)
 	: Super(ObjectInitializer)
 {
+	SectionName = TEXT("Collision");
 }
 
 UCollisionProfile* UCollisionProfile::Get()
diff --git a/Engine/Source/Runtime/Engine/Private/DeveloperSettings.cpp b/Engine/Source/Runtime/Engine/Private/DeveloperSettings.cpp
--- a/Engine/Source/Runtime/Engine/Private/DeveloperSettings.cpp
+++ b/Engine/Source/Runtime/Engine/Private/DeveloperSettings.cpp
@@ -7,28 +7,33 @@ UDeveloperSettings::UDeveloperSettings(const FObjectInitializer& ObjectInitializ
 	: UObject(ObjectInitializer)
 {
 	CategoryName = NAME_None;
-	Container = EDeveloperSettingsContainer::Project;
+	SectionName = NAME_None;
 }
 
 FName UDeveloperSettings::GetContainerName() const
 {
-	FName ProjectName(TEXT("Project"));
-	FName EditorName(TEXT("Editor"));
+	static const FName ProjectName(TEXT("Project"));
+	static const FName EditorName(TEXT("Editor"));
 
-	switch ( Container )
+	static const FName EditorUserSettingsName(TEXT("EditorUserSettings"));
+	static const FName EditorPerProjectUserSettingsName(TEXT("EditorPerProjectUserSettings"));
+
+	FName ConfigName = GetClass()->ClassConfigName;
+
+	if ( ConfigName == EditorUserSettingsName || ConfigName == EditorPerProjectUserSettingsName )
 	{
-	case EDeveloperSettingsContainer::Project:
-		return ProjectName;
-	case EDeveloperSettingsContainer::Editor:
 		return EditorName;
-	default:
-		check(false);
-		return NAME_None;
 	}
+	
+	return ProjectName;
 }
 
 FName UDeveloperSettings::GetCategoryName() const
 {
+	static const FName GeneralName(TEXT("General"));
+	static const FName EditorUserSettingsName(TEXT("EditorUserSettings"));
+	static const FName EditorPerProjectUserSettingsName(TEXT("EditorPerProjectUserSettings"));
+
 	if ( CategoryName != NAME_None )
 	{
 		return CategoryName;
@@ -39,6 +44,10 @@ FName UDeveloperSettings::GetCategoryName() const
 	{
 		return NAME_Engine;
 	}
+	else if ( ConfigName == EditorUserSettingsName || ConfigName == EditorPerProjectUserSettingsName )
+	{
+		return GeneralName;
+	}
 	else if ( ConfigName == NAME_Editor || ConfigName == NAME_EditorSettings || ConfigName == NAME_EditorLayout || ConfigName == NAME_EditorKeyBindings )
 	{
 		return NAME_Editor;
@@ -53,6 +62,11 @@ FName UDeveloperSettings::GetCategoryName() const
 
 FName UDeveloperSettings::GetSectionName() const
 {
+	if ( SectionName != NAME_None )
+	{
+		return SectionName;
+	}
+
 	return GetClass()->GetFName();
 }
 
diff --git a/Engine/Source/Runtime/Engine/Private/NetworkSettings.cpp b/Engine/Source/Runtime/Engine/Private/NetworkSettings.cpp
--- a/Engine/Source/Runtime/Engine/Private/NetworkSettings.cpp
+++ b/Engine/Source/Runtime/Engine/Private/NetworkSettings.cpp
@@ -6,6 +6,7 @@
 UNetworkSettings::UNetworkSettings(const FObjectInitializer& ObjectInitializer)
 	: Super(ObjectInitializer)
 {
+	SectionName = TEXT("Network");
 }
 
 static FName NetworkConsoleVariableFName(TEXT("ConsoleVariable"));
diff --git a/Engine/Source/Runtime/Engine/Private/PhysicsEngine/PhysicsSettings.cpp b/Engine/Source/Runtime/Engine/Private/PhysicsEngine/PhysicsSettings.cpp
--- a/Engine/Source/Runtime/Engine/Private/PhysicsEngine/PhysicsSettings.cpp
+++ b/Engine/Source/Runtime/Engine/Private/PhysicsEngine/PhysicsSettings.cpp
@@ -29,6 +29,7 @@ UPhysicsSettings::UPhysicsSettings(const FObjectInitializer& ObjectInitializer)
 	, AsyncSceneSmoothingFactor(0.99f)
 	, InitialAverageFrameRate(1.f / 60.f)
 {
+	SectionName = TEXT("Physics");
 }
 
 void UPhysicsSettings::PostInitProperties()
diff --git a/Engine/Source/Runtime/Engine/Private/RenderingSettings.cpp b/Engine/Source/Runtime/Engine/Private/RenderingSettings.cpp
--- a/Engine/Source/Runtime/Engine/Private/RenderingSettings.cpp
+++ b/Engine/Source/Runtime/Engine/Private/RenderingSettings.cpp
@@ -6,6 +6,7 @@
 URendererSettings::URendererSettings(const FObjectInitializer& ObjectInitializer)
 	: Super(ObjectInitializer)
 {
+	SectionName = TEXT("Rendering");
 	TranslucentSortAxis = FVector(0.0f, -1.0f, 0.0f);
 }
 
diff --git a/Engine/Source/Runtime/Engine/Private/UserInterfaceSettings.cpp b/Engine/Source/Runtime/Engine/Private/UserInterfaceSettings.cpp
--- a/Engine/Source/Runtime/Engine/Private/UserInterfaceSettings.cpp
+++ b/Engine/Source/Runtime/Engine/Private/UserInterfaceSettings.cpp
@@ -15,6 +15,7 @@ UUserInterfaceSettings::UUserInterfaceSettings(const FObjectInitializer& ObjectI
 	, RenderFocusRule(ERenderFocusRule::NavigationOnly)
 	, ApplicationScale(1)
 {
+	SectionName = TEXT("UI");
 }
 
 void UUserInterfaceSettings::PostInitProperties()
