nicer comment

diff --git a/Engine/Source/Runtime/Engine/Public/SceneManagement.h b/Engine/Source/Runtime/Engine/Public/SceneManagement.h
--- a/Engine/Source/Runtime/Engine/Public/SceneManagement.h
+++ b/Engine/Source/Runtime/Engine/Public/SceneManagement.h
@@ -134,7 +134,11 @@ public:
 
 	/** Resets pool for GetReusableMID() */
 	virtual void OnStartPostProcessing(FSceneView& CurrentView) = 0;
-	/** Allows MIDs being created and released during view rendering without the overhead of creating and relasing objects */
+	/**
+	 * Allows MIDs being created and released during view rendering without the overhead of creating and releasing objects
+	 * As MID are not allowed to be parent of MID this gets fixed up by parenting it to the next Material or MIC
+	 * @param ParentMaterial can be Material, MIC or MID, must not be 0
+	 */
 	virtual UMaterialInstanceDynamic* GetReusableMID(class UMaterialInterface* ParentMaterial) = 0;
 	/** Returns the temporal LOD struct from the viewstate */
 	virtual FTemporalLODState& GetTemporalLODState() = 0;
