Mac compile fix

diff --git a/Engine/Source/Runtime/Engine/Private/Slate/SceneViewport.cpp b/Engine/Source/Runtime/Engine/Private/Slate/SceneViewport.cpp
--- a/Engine/Source/Runtime/Engine/Private/Slate/SceneViewport.cpp
+++ b/Engine/Source/Runtime/Engine/Private/Slate/SceneViewport.cpp
@@ -1356,7 +1356,6 @@ void FSceneViewport::InitDynamicRHI()
 	}
 
 	TSharedPtr<FSlateRenderer> Renderer = FSlateApplication::Get().GetRenderer();
-	FWidgetPath WidgetPath;
 	if( bUseSeparateRenderTarget )
 	{
 		uint32 TexSizeX = SizeX, TexSizeY = SizeY;
