Fix (probably) mac non-unity

diff --git a/Engine/Source/Runtime/Engine/Classes/PhysicsEngine/PhysicsSettings.h b/Engine/Source/Runtime/Engine/Classes/PhysicsEngine/PhysicsSettings.h
--- a/Engine/Source/Runtime/Engine/Classes/PhysicsEngine/PhysicsSettings.h
+++ b/Engine/Source/Runtime/Engine/Classes/PhysicsEngine/PhysicsSettings.h
@@ -6,6 +6,7 @@
 
 #pragma once
 
+#include "Engine/DeveloperSettings.h"
 #include "PhysicsSettings.generated.h"
 
 
