Fix corrupted editor viewport.  Also a crash and editor perf regression

diff --git a/Engine/Source/Runtime/Engine/Private/Slate/SceneViewport.cpp b/Engine/Source/Runtime/Engine/Private/Slate/SceneViewport.cpp
--- a/Engine/Source/Runtime/Engine/Private/Slate/SceneViewport.cpp
+++ b/Engine/Source/Runtime/Engine/Private/Slate/SceneViewport.cpp
@@ -130,10 +130,7 @@ void FSceneViewport::Destroy()
 {
 	ViewportClient = NULL;
 	
-	UpdateViewportRHI( true, 0, 0, EWindowMode::Windowed );
-
-	delete RenderThreadSlateTexture;
-	RenderThreadSlateTexture = nullptr;
+	UpdateViewportRHI( true, 0, 0, EWindowMode::Windowed );	
 }
 
 int32 FSceneViewport::GetMouseX() const
@@ -880,7 +877,7 @@ void FSceneViewport::OnViewportClosed()
 FSlateShaderResource* FSceneViewport::GetViewportRenderTargetTexture() const
 { 
 	check(IsInGameThread());
-	return BufferedSlateHandles[CurrentBufferedTargetIndex];
+	return (BufferedSlateHandles.Num() != 0) ? BufferedSlateHandles[CurrentBufferedTargetIndex] : nullptr;
 }
 
 void FSceneViewport::ResizeFrame(uint32 NewSizeX, uint32 NewSizeY, EWindowMode::Type NewWindowMode, int32 InPosX, int32 InPosY)
@@ -1102,7 +1099,7 @@ FSlateShaderResource* FSceneViewport::GetViewportRenderTargetTexture()
 	{
 		return RenderThreadSlateTexture;
 	}
-	return BufferedSlateHandles[CurrentBufferedTargetIndex];
+	return (BufferedSlateHandles.Num() != 0) ? BufferedSlateHandles[CurrentBufferedTargetIndex] : nullptr;
 }
 
 void FSceneViewport::SetRenderTargetTextureRenderThread(FTexture2DRHIRef& RT)
@@ -1157,13 +1154,17 @@ void FSceneViewport::UpdateViewportRHI(bool bDestroyed, uint32 NewSizeX, uint32
 		else
 		{
 			// Enqueue a render command to delete the handle.  It must be deleted on the render thread after the resource is released
-			ENQUEUE_UNIQUE_RENDER_COMMAND_ONEPARAMETER(DeleteSlateRenderTarget, TArray<FSlateRenderTargetRHI*>&, BufferedSlateHandles, BufferedSlateHandles,
+			ENQUEUE_UNIQUE_RENDER_COMMAND_TWOPARAMETER(DeleteSlateRenderTarget, TArray<FSlateRenderTargetRHI*>&, BufferedSlateHandles, BufferedSlateHandles,
+																				FSlateRenderTargetRHI*&, RenderThreadSlateTexture, RenderThreadSlateTexture,
 			{
 				for (int32 i = 0; i < BufferedSlateHandles.Num(); ++i)
 				{
 					delete BufferedSlateHandles[i];
-				}
-				BufferedSlateHandles.Empty();				
+					BufferedSlateHandles[i] = nullptr;
+
+					delete RenderThreadSlateTexture;
+					RenderThreadSlateTexture = nullptr;
+				}						
 			});
 
 		}
@@ -1358,18 +1359,17 @@ void FSceneViewport::InitDynamicRHI()
 	TSharedPtr<FSlateRenderer> Renderer = FSlateApplication::Get().GetRenderer();
 	if( bUseSeparateRenderTarget )
 	{
+		NumBufferedFrames = 1;
 		uint32 TexSizeX = SizeX, TexSizeY = SizeY;
 		if (GEngine->IsStereoscopic3D(this))
 		{
 			GEngine->StereoRenderingDevice->CalculateRenderTargetSize(TexSizeX, TexSizeY);
+			//todo, get from HMD
+			NumBufferedFrames = 3;
 		}
 				
 		check(BufferedSlateHandles.Num() == BufferedRenderTargetsRHI.Num() && BufferedSlateHandles.Num() == BufferedShaderResourceTexturesRHI.Num());
 
-		//todo, get from HMD
-		NumBufferedFrames = 3;
-
-
 		if (BufferedSlateHandles.Num() != NumBufferedFrames)
 		{
 			//clear existing entries
@@ -1454,7 +1454,11 @@ void FSceneViewport::ReleaseDynamicRHI()
 		{
 			BufferedSlateHandles[i]->ReleaseDynamicRHI();
 		}
-	}	
+	}
+	if (RenderThreadSlateTexture)
+	{
+		RenderThreadSlateTexture->ReleaseDynamicRHI();
+	}
 }
 
 void FSceneViewport::SetPreCaptureMousePosFromSlateCursor()
