Fix fast TArray unmapped property tracking issue

diff --git a/Engine/Source/Runtime/Engine/Private/DataReplication.cpp b/Engine/Source/Runtime/Engine/Private/DataReplication.cpp
--- a/Engine/Source/Runtime/Engine/Private/DataReplication.cpp
+++ b/Engine/Source/Runtime/Engine/Private/DataReplication.cpp
@@ -531,12 +531,18 @@ bool FObjectReplicator::ReceivedBunch( FInBunch& Bunch, const FReplicationFlags&
 #endif
 			if ( !Retirement[ ReplicatedProp->RepIndex ].CustomDelta )
 			{
+				bool bLocalHasUnmapped = false;
 				// We hijack a non custom delta property to signify we are using FRepLayout to read the entire property block
-				if ( !RepLayout->ReceiveProperties( ObjectClass, RepState, (void*)Object, Bunch, bOutHasUnmapped ) )
+				if ( !RepLayout->ReceiveProperties( ObjectClass, RepState, (void*)Object, Bunch, bLocalHasUnmapped ) )
 				{
 					UE_LOG( LogNet, Error, TEXT( "ReceiveProperties FAILED %s in %s" ), *ReplicatedProp->GetName(), *Object->GetFullName() );
 					return false;
 				}
+
+				if ( bLocalHasUnmapped )
+				{
+					bOutHasUnmapped = true;
+				}
 			}
 			else
 			{
@@ -606,6 +612,7 @@ bool FObjectReplicator::ReceivedBunch( FInBunch& Bunch, const FReplicationFlags&
 				if ( Parms.bOutHasMoreUnmapped )
 				{
 					UnmappedCustomProperties.Add( DataOffset, StructProperty );
+					bOutHasUnmapped = true;
 				}
 
 				// Successfully received it.
