#HLOD
- Added min number of actors to merge

diff --git a/Engine/Source/Editor/DetailCustomizations/Private/HierarchicalSimplificationCustomizations.cpp b/Engine/Source/Editor/DetailCustomizations/Private/HierarchicalSimplificationCustomizations.cpp
--- a/Engine/Source/Editor/DetailCustomizations/Private/HierarchicalSimplificationCustomizations.cpp
+++ b/Engine/Source/Editor/DetailCustomizations/Private/HierarchicalSimplificationCustomizations.cpp
@@ -12,6 +12,15 @@ TSharedRef<IPropertyTypeCustomization> FHierarchicalSimplificationCustomizations
 
 void FHierarchicalSimplificationCustomizations::CustomizeHeader(TSharedRef<IPropertyHandle> StructPropertyHandle, FDetailWidgetRow& HeaderRow, IPropertyTypeCustomizationUtils& StructCustomizationUtils)
 {
+	HeaderRow.
+	NameContent()
+	[
+		StructPropertyHandle->CreatePropertyNameWidget()
+	]
+	.ValueContent()
+	[
+		StructPropertyHandle->CreatePropertyValueWidget()
+	];
 }
 
 void FHierarchicalSimplificationCustomizations::CustomizeChildren( TSharedRef<IPropertyHandle> StructPropertyHandle, class IDetailChildrenBuilder& ChildBuilder, IPropertyTypeCustomizationUtils& StructCustomizationUtils )
diff --git a/Engine/Source/Editor/UnrealEd/Public/HierarchicalLOD.h b/Engine/Source/Editor/UnrealEd/Public/HierarchicalLOD.h
--- a/Engine/Source/Editor/UnrealEd/Public/HierarchicalLOD.h
+++ b/Engine/Source/Editor/UnrealEd/Public/HierarchicalLOD.h
@@ -34,8 +34,6 @@ struct FLODCluster
 	void Invalidate() { bValid = false; }
 	// return true if valid
 	bool IsValid() const {	return bValid; }
-	// should build to actor list
-	bool ShouldBuild() const { return IsValid() && Actors.Num() > 1; }
 
 	// return cost of the cluster, lower is better
 	const float GetCost() const
@@ -106,7 +104,7 @@ private:
 	void InitializeClusters(class ULevel* InLevel, const int32 LODIdx, float CullCost);
 
 	// merge clusters
-	void MergeClustersAndBuildActors(class ULevel* InLevel, const int32 LODIdx, float HighestCost);
+	void MergeClustersAndBuildActors(class ULevel* InLevel, const int32 LODIdx, float HighestCost, int32 MinNumActors);
 
 	// find minmal spanning tree of clusters
 	void FindMST();
diff --git a/Engine/Source/Runtime/Engine/Classes/GameFramework/WorldSettings.h b/Engine/Source/Runtime/Engine/Classes/GameFramework/WorldSettings.h
--- a/Engine/Source/Runtime/Engine/Classes/GameFramework/WorldSettings.h
+++ b/Engine/Source/Runtime/Engine/Classes/GameFramework/WorldSettings.h
@@ -197,7 +197,11 @@ struct ENGINE_API FHierarchicalSimplification
 {
 	GENERATED_USTRUCT_BODY()
 
-	/** If this is true, it will simplify mesh. However it is slower. If false, it will just merge actors but not simplify */
+	/** If this is true, it will simplify mesh but it is slower. 
+	 * If false, it will just merge actors but not simplify using the lower LOD if exists. 
+	 * For example if you build LOD 1, it will use LOD 1 of the mesh to merge actors if exists.  
+	 * Also it will flattening material so that it can reduce drawcalls. 
+	 */
 	UPROPERTY(Category=FHierarchicalSimplification, EditAnywhere)
 	bool bSimplifyMesh;
 
@@ -205,14 +209,17 @@ struct ENGINE_API FHierarchicalSimplification
 	UPROPERTY(Category=FHierarchicalSimplification, EditAnywhere, AdvancedDisplay, meta=(UIMin=10.f, ClampMin=10.f))
 	float DrawDistance;
 
-	/** Desired Bounding Box Size for clustering */
+	/** Desired Bounding Box Size for clustering - this is not guaranteed but used to calculate filling factor */
 	UPROPERTY(EditAnywhere, Category=FHierarchicalSimplification, AdvancedDisplay, meta=(UIMin=10.f, ClampMin=10.f))
 	float DesiredBoundSize;
 
-	/** Desired Filling Percentage for clustering */
+	/** Desired Filling Percentage for clustering - this is not guaranteed but used to calculate filling factor  */
 	UPROPERTY(EditAnywhere, Category=FHierarchicalSimplification, AdvancedDisplay, meta=(ClampMin = "0", ClampMax = "100", UIMin = "0", UIMax = "100"))
 	float DesiredFillingPercentage;
 
+	UPROPERTY(EditAnywhere, Category=FHierarchicalSimplification, AdvancedDisplay, meta=(ClampMin = "1", UIMin = "1"))
+	int32 MinNumberOfActorsToBuild;
+
 	/** Simplification Setting if bSimplifyMesh is true */
 	UPROPERTY(Category=FHierarchicalSimplification, EditAnywhere, AdvancedDisplay, meta=(editcondition = "bSimplifyMesh"))
 	FMeshProxySettings Setting;
@@ -222,6 +229,7 @@ struct ENGINE_API FHierarchicalSimplification
 		, DrawDistance(5000)
 		, DesiredBoundSize(5000)
 		, DesiredFillingPercentage(50)
+		, MinNumberOfActorsToBuild(2)
 	{
 		Setting.ScreenSize = 200;
 	}
