changing the logs in spawnactor from logs to warnings and adding a log for when the spawned actor is pending kill and will return NULL. Checking this in after a chat with James Golding. That when ever spawn actor fails to spawn an actor we should have a warning and not fail silently.

diff --git a/Engine/Source/Runtime/Engine/Private/LevelActor.cpp b/Engine/Source/Runtime/Engine/Private/LevelActor.cpp
--- a/Engine/Source/Runtime/Engine/Private/LevelActor.cpp
+++ b/Engine/Source/Runtime/Engine/Private/LevelActor.cpp
@@ -295,7 +295,7 @@ AActor* UWorld::SpawnActor( UClass* Class, FVector const* Location, FRotator con
 	// See if we can spawn on ded.server/client only etc (check NeedsLoadForClient & NeedsLoadForServer)
 	if(!CanCreateInCurrentContext(Template))
 	{
-		UE_LOG(LogSpawn, Log, TEXT("Unable to spawn class '%s' due to client/server context."), *Class->GetName() );
+		UE_LOG(LogSpawn, Warning, TEXT("Unable to spawn class '%s' due to client/server context."), *Class->GetName() );
 		return NULL;
 	}
 
@@ -308,7 +308,7 @@ AActor* UWorld::SpawnActor( UClass* Class, FVector const* Location, FRotator con
 		// Try to find a spawn position
 		if ( !FindTeleportSpot(Template, NewLocation, NewRotation) )
 		{
-			UE_LOG(LogSpawn, Log, TEXT("SpawnActor failed because of collision at the spawn location [%s] for [%s]"), *NewLocation.ToString(), *Class->GetName() );
+			UE_LOG(LogSpawn, Warning, TEXT("SpawnActor failed because of collision at the spawn location [%s] for [%s]"), *NewLocation.ToString(), *Class->GetName() );
 			return NULL;
 		}
 	}
@@ -347,6 +347,7 @@ AActor* UWorld::SpawnActor( UClass* Class, FVector const* Location, FRotator con
 
 	if (Actor->IsPendingKill() && !SpawnParameters.bNoFail)
 	{
+		UE_LOG(LogSpawn, Warning, TEXT("SpawnActor failed because the spawned actor IsPendingKill"));
 		return NULL;
 	}
 	Actor->CheckDefaultSubobjects();
