Fixed SetViewTarget(PlayerState) not working when that PlayerState's Owner does not currently control a pawn

Merging 2521981

//depot/UE4-UT/...

to //depot/UE4/...

diff --git a/Engine/Source/Runtime/Engine/Private/PlayerCameraManager.cpp b/Engine/Source/Runtime/Engine/Private/PlayerCameraManager.cpp
--- a/Engine/Source/Runtime/Engine/Private/PlayerCameraManager.cpp
+++ b/Engine/Source/Runtime/Engine/Private/PlayerCameraManager.cpp
@@ -1389,7 +1389,7 @@ void FTViewTarget::CheckViewTarget(APlayerController* OwningController)
 					}
 					else
 					{
-						PlayerState = NULL;
+						Target = PlayerState; // this will cause it to update to the next Pawn possessed by the player being viewed
 					}
 				}
 				else
