Added some simple versioning to the engine-provided savegame serialization.

diff --git a/Engine/Source/Runtime/Engine/Private/GameplayStatics.cpp b/Engine/Source/Runtime/Engine/Private/GameplayStatics.cpp
--- a/Engine/Source/Runtime/Engine/Private/GameplayStatics.cpp
+++ b/Engine/Source/Runtime/Engine/Private/GameplayStatics.cpp
@@ -20,6 +20,9 @@
 #include "Components/DecalComponent.h"
 #include "LandscapeProxy.h"
 
+static const int UE4_SAVEGAME_FILE_TYPE_TAG = 0x53415647;		// "sAvG"
+static const int UE4_SAVEGAME_FILE_VERSION = 1;
+
 //////////////////////////////////////////////////////////////////////////
 // UGameplayStatics
 
@@ -1100,6 +1103,7 @@ USaveGame* UGameplayStatics::CreateSaveGameObjectFromBlueprint(UBlueprint* SaveG
 	return OutSaveGameObject;
 }
 
+
 bool UGameplayStatics::SaveGameToSlot(USaveGame* SaveGameObject, const FString& SlotName, const int32 UserIndex)
 {
 	bool bSuccess = false;
@@ -1111,7 +1115,21 @@ bool UGameplayStatics::SaveGameToSlot(USaveGame* SaveGameObject, const FString&
 		TArray<uint8> ObjectBytes;
 		FMemoryWriter MemoryWriter(ObjectBytes, true);
 
-		// First write the class name so we know what class to load to
+		// write file type tag. identifies this file type and indicates it's using proper versioning
+		// since older UE4 versions did not version this data.
+		int32 FileTypeTag = UE4_SAVEGAME_FILE_TYPE_TAG;
+		MemoryWriter << FileTypeTag;
+
+		// Write version for this file format
+		int32 SavegameFileVersion = UE4_SAVEGAME_FILE_VERSION;
+		MemoryWriter << SavegameFileVersion;
+
+		// Write out engine and UE4 version information
+		MemoryWriter << GPackageFileUE4Version;
+		FEngineVersion SavedEngineVersion = GEngineVersion;
+		MemoryWriter << SavedEngineVersion;
+
+		// Write the class name so we know what class to load to
 		FString SaveGameClassName = SaveGameObject->GetClass()->GetName();
 		MemoryWriter << SaveGameClassName;
 
@@ -1165,6 +1183,38 @@ USaveGame* UGameplayStatics::LoadGameFromSlot(const FString& SlotName, const int
 		{
 			FMemoryReader MemoryReader(ObjectBytes, true);
 
+			int32 FileTypeTag;
+			MemoryReader << FileTypeTag;
+
+			int32 SavegameFileVersion;
+			if (FileTypeTag != UE4_SAVEGAME_FILE_TYPE_TAG)
+			{
+				// this is an old saved game, back up the file pointer to the beginning and assume version 1
+				MemoryReader.Seek(0);
+				SavegameFileVersion = 1;
+
+				// Note for 4.8 and beyond: if you get a crash loading a pre-4.8 version of your savegame file and 
+				// you don't want to delete it, try uncommenting these lines and changing them to use the version 
+				// information from your previous build. Then load and resave your savegame file.
+				//MemoryReader.SetUE4Ver(MyPreviousUE4Version);				// @see GPackageFileUE4Version
+				//MemoryReader.SetEngineVer(MyPreviousEngineVersion);		// @see GEngineVersion
+			}
+			else
+			{
+				// Read version for this file format
+				MemoryReader << SavegameFileVersion;
+
+				// Read engine and UE4 version information
+				int32 SavedUE4Version;
+				MemoryReader << SavedUE4Version;
+
+				FEngineVersion SavedEngineVersion;
+				MemoryReader << SavedEngineVersion;
+
+				MemoryReader.SetUE4Ver(SavedUE4Version);
+				MemoryReader.SetEngineVer(SavedEngineVersion);
+			}
+			
 			// Get the class name
 			FString SaveGameClassName;
 			MemoryReader << SaveGameClassName;
