Fix for destruction crashes when userData becomes stale. This has been seen in a few places:
addForce, recreation of state, etc...

UE-11729

diff --git a/Engine/Source/Runtime/Engine/Classes/Components/DestructibleComponent.h b/Engine/Source/Runtime/Engine/Classes/Components/DestructibleComponent.h
--- a/Engine/Source/Runtime/Engine/Classes/Components/DestructibleComponent.h
+++ b/Engine/Source/Runtime/Engine/Classes/Components/DestructibleComponent.h
@@ -22,14 +22,10 @@ namespace physx
 /** Mapping info for destructible chunk user data. */
 struct FDestructibleChunkInfo
 {
-	/** Index of this chunk info */
-	int32 Index;
 	/** Index of the chunk this data belongs to*/
 	int32 ChunkIndex;
 	/** Component owning this chunk info*/
 	TWeakObjectPtr<class UDestructibleComponent> OwningComponent;
-	/** Physx actor */
-	physx::PxRigidDynamic* Actor;
 };
 #endif // WITH_PHYSX 
 
@@ -75,7 +71,7 @@ class ENGINE_API UDestructibleComponent : public USkinnedMeshComponent
 
 #if WITH_PHYSX
 	/** Per chunk info */
-	TIndirectArray<FDestructibleChunkInfo> ChunkInfos;
+	TArray<FDestructibleChunkInfo> ChunkInfos;
 #endif // WITH_PHYSX 
 
 #if WITH_EDITOR
@@ -164,6 +160,9 @@ public:
 	
 	/** Resets the BodyInstance to the state that is defined in PrevState. */
 	void ResetFakeBodyInstance(FFakeBodyInstanceState& PrevState);
+
+	/** Setup a pair of PxShape and ChunkIndex */
+	void Pair( int32 ChunkIndex, physx::PxShape* PShape );
 #endif // WITH_APEX
 
 	/** This method makes a chunk (fractured piece) visible or invisible.
@@ -230,10 +229,15 @@ private:
 	/** User data wrapper for this component passed to physx */
 	FPhysxUserData PhysxUserData;
 
+	void SetCollisionResponseForShape(physx::PxShape* Shape, int32 ChunkIdx);
+	void SetCollisionResponseForActor(physx::PxRigidDynamic* Actor, int32 ChunkIdx, const FCollisionResponseContainer* ResponseOverride = NULL);
+	void SetCollisionResponseForAllActors(const FCollisionResponseContainer& ResponseOverride);
+
+
+public:
 	/** User data wrapper for the chunks passed to physx */
-	TIndirectArray<FPhysxUserData> PhysxChunkUserData;
+	TArray<FPhysxUserData> PhysxChunkUserData;
 	bool IsChunkLarge(int32 ChunkIdx) const;
-	void SetCollisionResponseForActor(physx::PxRigidDynamic* Actor, int32 ChunkIdx);
 #endif
 };
 
diff --git a/Engine/Source/Runtime/Engine/Private/Collision/CollisionConversions.cpp b/Engine/Source/Runtime/Engine/Private/Collision/CollisionConversions.cpp
--- a/Engine/Source/Runtime/Engine/Private/Collision/CollisionConversions.cpp
+++ b/Engine/Source/Runtime/Engine/Private/Collision/CollisionConversions.cpp
@@ -386,7 +386,7 @@ static void SetHitResultFromShapeAndFaceIndex(const PxShape* PShape,  const PxRi
 	// For destructibles give the ChunkInfo-Index as Item
 	if (bReturnBody && ChunkInfo)
 	{
-		OutResult.Item = ChunkInfo->Index;
+		OutResult.Item = ChunkInfo->ChunkIndex;
 
 		UDestructibleComponent* DMComp = Cast<UDestructibleComponent>(OwningComponent);
 		OutResult.BoneName = DMComp->GetBoneName(UDestructibleComponent::ChunkIdxToBoneIdx(ChunkInfo->ChunkIndex));
diff --git a/Engine/Source/Runtime/Engine/Private/Components/DestructibleComponent.cpp b/Engine/Source/Runtime/Engine/Private/Components/DestructibleComponent.cpp
--- a/Engine/Source/Runtime/Engine/Private/Components/DestructibleComponent.cpp
+++ b/Engine/Source/Runtime/Engine/Private/Components/DestructibleComponent.cpp
@@ -351,6 +351,11 @@ void UDestructibleComponent::CreatePhysicsState()
 	BodyInstance.SceneIndexAsync = SceneType == PST_Async ? PhysScene->PhysXSceneIndex[PST_Async] : 0;
 	check(ApexScene);
 
+	ChunkInfos.Reset(ChunkCount);
+	ChunkInfos.AddZeroed(ChunkCount);
+	PhysxChunkUserData.Reset(ChunkCount);
+	PhysxChunkUserData.AddZeroed(ChunkCount);
+
 	// Create an APEX NxDestructibleActor from the Destructible asset and actor descriptor
 	ApexDestructibleActor = static_cast<NxDestructibleActor*>(TheDestructibleMesh->ApexDestructibleAsset->createApexActor(*ActorParams, *ApexScene));
 	check(ApexDestructibleActor);
@@ -359,9 +364,6 @@ void UDestructibleComponent::CreatePhysicsState()
 	PhysxUserData = FPhysxUserData(this);
 	ApexDestructibleActor->userData = &PhysxUserData;
 
-	ChunkInfos.Empty();
-	PhysxChunkUserData.Empty();
-
 	// Cache cooked collision data
 	// BRGTODO : cook in asset
 	ApexDestructibleActor->cacheModuleData();
@@ -385,7 +387,7 @@ void UDestructibleComponent::CreatePhysicsState()
 		{
 			ApexDestructibleActor->setChunkPhysXActorAwakeState(0, false);
 		}
-	}
+		}
 
 	UpdateBounds();
 #endif	// #if WITH_APEX
@@ -400,11 +402,7 @@ void UDestructibleComponent::DestroyPhysicsState()
 		{
 			if (FPhysScene * PhysScene = World->GetPhysicsScene())
 			{
-				const uint32 SceneType = BodyInstance.UseAsyncScene(PhysScene) ? PST_Async : PST_Sync;
-				PxScene * PScene = PhysScene->GetPhysXScene(SceneType);
-				SCOPED_SCENE_WRITE_LOCK(PScene);
-				ApexDestructibleActor->release();
-				//Deferring here is difficult because a call to CreatePhysicsState will re-use the same chunk info buffer for new chunks. Note this may be possible, but need to fix crash now
+				PhysScene->DeferredCommandHandler.DeferredRelease(ApexDestructibleActor);
 			}
 		}
 		
@@ -472,14 +470,8 @@ void UDestructibleComponent::AddImpulseAtLocation( FVector Impulse, FVector Posi
 			ChunkIdx = BoneIdxToChunkIdx(GetBoneIndex(BoneName));
 		}
 
-		//ApexDestructibleActor->applyDamage(1.0f, Impulse.Size(), U2PVector(Position), U2PVector(Impulse.SafeNormal()), ChunkIdx);
-		for (int32 i=0; i < ChunkInfos.Num(); ++i)
-		{
-			if (ChunkInfos[i].ChunkIndex == ChunkIdx && ChunkInfos[i].Actor)
-			{
-				ChunkInfos[i].Actor->addForce(U2PVector(Impulse), PxForceMode::eIMPULSE);
-			}
-		}
+		PxRigidDynamic* Actor = ApexDestructibleActor->getChunkPhysXActor(ChunkIdx);
+		Actor->addForce(U2PVector(Impulse), PxForceMode::eIMPULSE);
 	}
 #endif
 }
@@ -493,14 +485,8 @@ void UDestructibleComponent::AddForce( FVector Force, FName BoneName /*= NAME_No
 		ChunkIdx = BoneIdxToChunkIdx(GetBoneIndex(BoneName));
 	}
 
-	//ApexDestructibleActor->applyDamage(1.0f, Impulse.Size(), U2PVector(Position), U2PVector(Impulse.SafeNormal()), ChunkIdx);
-	for (int32 i=0; i < ChunkInfos.Num(); ++i)
-	{
-		if (ChunkInfos[i].ChunkIndex == ChunkIdx && ChunkInfos[i].Actor)
-		{
-			ChunkInfos[i].Actor->addForce(U2PVector(Force), bAccelChange ? PxForceMode::eACCELERATION : PxForceMode::eFORCE);
-		}
-	}
+	PxRigidDynamic* Actor = ApexDestructibleActor->getChunkPhysXActor(ChunkIdx);
+	Actor->addForce(U2PVector(Force), bAccelChange ? PxForceMode::eACCELERATION : PxForceMode::eFORCE);
 #endif
 }
 
@@ -515,14 +501,8 @@ void UDestructibleComponent::AddForceAtLocation( FVector Force, FVector Location
 			ChunkIdx = BoneIdxToChunkIdx(GetBoneIndex(BoneName));
 		}
 
-		//ApexDestructibleActor->applyDamage(1.0f, Impulse.Size(), U2PVector(Location), U2PVector(Impulse.SafeNormal()), ChunkIdx);
-		for (int32 i=0; i < ChunkInfos.Num(); ++i)
-		{
-			if (ChunkInfos[i].ChunkIndex == ChunkIdx && ChunkInfos[i].Actor)
-			{
-				ChunkInfos[i].Actor->addForce(U2PVector(Force), PxForceMode::eFORCE);
-			}
-		}
+		PxRigidDynamic* Actor = ApexDestructibleActor->getChunkPhysXActor(ChunkIdx);
+		Actor->addForce(U2PVector(Force), PxForceMode::eFORCE);
 	}
 #endif
 }
@@ -754,8 +734,11 @@ void UDestructibleComponent::RefreshBoneTransforms(FActorComponentTickFunction*
 void UDestructibleComponent::SetDestructibleMesh(class UDestructibleMesh* NewMesh)
 {
 #if WITH_APEX
-	ChunkInfos.Empty();
-	PhysxChunkUserData.Empty();
+	uint32 ChunkCount = NewMesh->ApexDestructibleAsset->getChunkCount();
+	ChunkInfos.Reset(ChunkCount);
+	ChunkInfos.AddZeroed(ChunkCount);
+	PhysxChunkUserData.Reset(ChunkCount);
+	PhysxChunkUserData.AddZeroed(ChunkCount);
 #endif // WITH_APEX
 
 	Super::SetSkeletalMesh( NewMesh );
@@ -817,104 +800,59 @@ FTransform UDestructibleComponent::GetSocketTransform(FName InSocketName, ERelat
 	return ST;
 }
 
-void UDestructibleComponent::SetChunkVisible( int32 ChunkIndex, bool bVisible )
+void UDestructibleComponent::Pair( int32 ChunkIndex, PxShape* PShape)
 {
-	// Bone 0 is a dummy root bone
-	const int32 BoneIndex = ChunkIdxToBoneIdx(ChunkIndex);
-	bool bClearActorFromChunkInfo = false;
+	FDestructibleChunkInfo* CI;
+	FPhysxUserData* UserData;
 
-	if( bVisible )
-	{
-		UnHideBone(BoneIndex);
-#if WITH_APEX
-		PxRigidDynamic* PActor = ApexDestructibleActor != NULL ? ApexDestructibleActor->getChunkPhysXActor(ChunkIndex) : NULL;
+	check(ChunkIndex < ChunkInfos.Num());
 
-		UDestructibleMesh* DMesh = GetDestructibleMesh();
+	CI = &ChunkInfos[ChunkIndex];
 
-		if (PActor != NULL)
-		{
-			//When we unhide a chunk, it might be the first time - in this case we need to update the transforms (only active transforms get updated after this)
-			//TODO: We can probably move this code somewhere so that it's only done the first time a chunk is created in case chunks are hidden unhidden again and again? (Not sure if this ever happens)
-			const physx::PxMat44 ChunkPoseRT = ApexDestructibleActor->getChunkPose(ChunkIndex);	// Unscaled
-			const physx::PxTransform Transform(ChunkPoseRT);
-			SetChunkWorldRT(ChunkIndex, P2UQuat(Transform.q), P2UVector(Transform.p));
-
-			// If actor has already a chunk info and userdata, we just make sure it is valid and update the 
-			// physx actor if needed. We do NOT do this for FormExtended structures, as in this case, the shapes/actors
-			// are moved to the 1st structure object internally by APEX.
-			if(PActor->userData != NULL && !DMesh->DefaultDestructibleParameters.Flags.bFormExtendedStructures)
-			{
-				FDestructibleChunkInfo* CI = FPhysxUserData::Get<FDestructibleChunkInfo>(PActor->userData);
-				checkf(CI, TEXT("If a chunk actor has user data and it is not a DestructibleChunkInfo, something is messed up."));
-				//check(CI->OwningComponent == this);
-
-				if (CI->ChunkIndex != ChunkIndex)
-				{
-					// grab the old actor and clear its user data, as we steal the ChunkInfo here
-					if (CI->Actor && CI->Actor != PActor)
-					{
-						CI->Actor->userData = NULL;
-					}
 					CI->ChunkIndex = ChunkIndex;
-				}
-
-				CI->OwningComponent = this;
-				CI->Actor = PActor;
-			}
-			else if (PActor->userData == NULL)
-			{
-				// Setup the user data to have a proper chunk - actor mapping 
-				FDestructibleChunkInfo* CI = new FDestructibleChunkInfo();
-				int32 InfoIndex = ChunkInfos.Add(CI);
-				CI->Index = InfoIndex;
-				CI->ChunkIndex = ChunkIndex;
 				CI->OwningComponent = this;
-				CI->Actor = PActor;
 
-				PActor->setActorFlag(PxActorFlag::eDISABLE_GRAVITY, !BodyInstance.bEnableGravity);
+	UserData = &PhysxChunkUserData[ChunkIndex];
+	UserData->Set(UserData, CI);
 
-				FPhysxUserData * UserData = new FPhysxUserData(CI);
-				int32 UserDataIdx = PhysxChunkUserData.Add(UserData);
-				check(InfoIndex == UserDataIdx);
+	PShape->userData = UserData;
 
-				PActor->userData = UserData;
+	PShape->getActor()->userData = UserData;
 
+	PShape->getActor()->setActorFlag(PxActorFlag::eDISABLE_GRAVITY, !BodyInstance.bEnableGravity);
 
-
-				// Set collision response to non-root chunks
-				if (GetDestructibleMesh()->ApexDestructibleAsset->getChunkParentIndex(ChunkIndex) >= 0)
-				{
-					SetCollisionResponseForActor(PActor, ChunkIndex);
-				}
-			}
-		}
-		else
-		{
-			bClearActorFromChunkInfo = true;
-		}
-#endif // WITH_APEX
-	}
-	else
+	// Set collision response to non-root chunks
+	if (GetDestructibleMesh()->ApexDestructibleAsset->getChunkParentIndex(ChunkIndex) >= 0)
 	{
-		HideBone(BoneIndex, PBO_None);
-		bClearActorFromChunkInfo = true;
+		SetCollisionResponseForShape(PShape, ChunkIndex);
 	}
+}
 
-#if WITH_APEX
-	if (bClearActorFromChunkInfo)
+void UDestructibleComponent::SetChunkVisible( int32 ChunkIndex, bool bVisible )
+{
+	// Bone 0 is a dummy root bone
+	const int32 BoneIndex = ChunkIdxToBoneIdx(ChunkIndex);
+
+	if( bVisible )
 	{
-		// Make sure we clear the physx actor pointer of the chunk info as it might (and probably will) be
-		// invalid from now on
-		for (int32 i=0; i < ChunkInfos.Num(); ++i)
+		UnHideBone(BoneIndex);
+
+		if (NULL != ApexDestructibleActor)
 		{
-			if (ChunkInfos[i].ChunkIndex == ChunkIndex)
+			physx::PxShape** PShapes;
+			const physx::PxU32 PShapeCount = ApexDestructibleActor->getChunkPhysXShapes(PShapes, ChunkIndex);
+			if (PShapeCount > 0)
 			{
-				ChunkInfos[i].Actor = NULL;
-				break;
+				const physx::PxMat44 ChunkPoseRT = ApexDestructibleActor->getChunkPose(ChunkIndex);	// Unscaled
+				const physx::PxTransform Transform(ChunkPoseRT);
+				SetChunkWorldRT(ChunkIndex, P2UQuat(Transform.q), P2UVector(Transform.p));
 			}
 		}
 	}
-#endif // WITH_APEX
+	else
+	{
+		HideBone(BoneIndex, PBO_None);
+	}
 
 	// Mark the transform as dirty, so the bounds are updated and sent to the render thread
 	MarkRenderTransformDirty();
@@ -1282,13 +1220,13 @@ void UDestructibleComponent::SetEnableGravity(bool bGravityEnabled)
 	Super::SetEnableGravity(bGravityEnabled);
 	
 #if WITH_APEX
-	for (FDestructibleChunkInfo& ChunkInfo : ChunkInfos)
-	 {
-		physx::PxRigidDynamic* Actor = ChunkInfo.Actor;
-		if (Actor)
+	uint32 ChunkCount = ApexDestructibleActor->getNumVisibleChunks();
+	const uint16* ChunkIndices = ApexDestructibleActor->getVisibleChunks();
+	for(uint32 c = 0; c < ChunkCount; c++)
 		{
-			Actor->setActorFlag(PxActorFlag::eDISABLE_GRAVITY, !bGravityEnabled);
-		}
+		PxActor* PActor = ApexDestructibleActor->getChunkPhysXActor(ChunkIndices[c]);
+		check(PActor);
+		PActor->setActorFlag(PxActorFlag::eDISABLE_GRAVITY, !bGravityEnabled);
 	}
 #endif //WITH_APEX
 }
@@ -1326,8 +1264,13 @@ bool UDestructibleComponent::IsChunkLarge(int32 ChunkIdx) const
 #endif // WITH_APEX
 }
 
-void UDestructibleComponent::SetCollisionResponseForActor(PxRigidDynamic* Actor, int32 ChunkIdx)
+void UDestructibleComponent::SetCollisionResponseForActor(PxRigidDynamic* Actor, int32 ChunkIdx, const FCollisionResponseContainer* ResponseOverride /*= NULL*/)
 {
+	if (ApexDestructibleActor == NULL)
+	{
+		return;
+	}
+
 	// Get collision channel and response
 	PxFilterData PQueryFilterData, PSimFilterData;
 	uint8 MoveChannel = GetCollisionObjectType();
@@ -1336,12 +1279,12 @@ void UDestructibleComponent::SetCollisionResponseForActor(PxRigidDynamic* Actor,
 		UDestructibleMesh* TheDestructibleMesh = GetDestructibleMesh();
 		AActor* Owner = GetOwner();
 		bool bLargeChunk = IsChunkLarge(ChunkIdx);
-		const FCollisionResponse& ColResponse = bLargeChunk ? LargeChunkCollisionResponse : SmallChunkCollisionResponse;
+		const FCollisionResponseContainer& UseResponse = ResponseOverride == NULL ? (bLargeChunk ? LargeChunkCollisionResponse.GetResponseContainer() : SmallChunkCollisionResponse.GetResponseContainer()) : *ResponseOverride;
 #if WITH_APEX
 		physx::PxU32 SupportDepth = TheDestructibleMesh->ApexDestructibleAsset->getChunkDepth(ChunkIdx);
 
 		const bool bEnableImpactDamage = IsImpactDamageEnabled(TheDestructibleMesh, SupportDepth);
-		CreateShapeFilterData(MoveChannel, GetUniqueID(), ColResponse.GetResponseContainer(), 0, ChunkIdxToBoneIdx(ChunkIdx), PQueryFilterData, PSimFilterData, BodyInstance.bUseCCD, bEnableImpactDamage, false);
+		CreateShapeFilterData(MoveChannel, GetUniqueID(), UseResponse, 0, ChunkIdxToBoneIdx(ChunkIdx), PQueryFilterData, PSimFilterData, BodyInstance.bUseCCD, bEnableImpactDamage, false);
 		
 		PQueryFilterData.word3 |= EPDF_SimpleCollision | EPDF_ComplexCollision;
 
@@ -1366,6 +1309,74 @@ void UDestructibleComponent::SetCollisionResponseForActor(PxRigidDynamic* Actor,
 	}
 }
 
+void UDestructibleComponent::SetCollisionResponseForAllActors(const FCollisionResponseContainer& ResponseOverride)
+{
+	if (ApexDestructibleActor == NULL)
+	{
+		return;
+	}
+
+	PxRigidDynamic** PActorBuffer = NULL;
+	PxU32 PActorCount = 0;
+	if (ApexDestructibleActor->acquirePhysXActorBuffer(PActorBuffer, PActorCount))
+	{
+		PxScene* LockedScene = NULL;
+
+		while (PActorCount--)
+		{
+			PxRigidDynamic* PActor = *PActorBuffer++;
+			if (PActor != NULL)
+			{
+				FDestructibleChunkInfo* ChunkInfo = FPhysxUserData::Get<FDestructibleChunkInfo>(PActor->userData);
+				if (ChunkInfo != NULL)
+				{
+					if (!LockedScene)
+					{
+						LockedScene = PActor->getScene();
+						LockedScene->lockWrite();
+						LockedScene->lockRead();
+					}
+					SetCollisionResponseForActor(PActor, ChunkInfo->ChunkIndex, &ResponseOverride);	// ChunkIndex is the last chunk made visible.  But SetCollisionResponseForActor already doesn't respect per-chunk collision properties.
+				}
+			}
+		}
+
+		if (LockedScene)
+		{
+			LockedScene->unlockRead();
+			LockedScene->unlockWrite();
+			LockedScene = NULL;
+		}
+
+		ApexDestructibleActor->releasePhysXActorBuffer();
+	}
+}
+
+void UDestructibleComponent::SetCollisionResponseForShape(PxShape* Shape, int32 ChunkIdx)
+{
+	// Get collision channel and response
+	PxFilterData PQueryFilterData, PSimFilterData;
+	uint8 MoveChannel = GetCollisionObjectType();
+	if (IsCollisionEnabled())
+	{
+		AActor* Owner = GetOwner();
+		bool bLargeChunk = IsChunkLarge(ChunkIdx);
+		const FCollisionResponse& ColResponse = bLargeChunk ? LargeChunkCollisionResponse : SmallChunkCollisionResponse;
+		//TODO: we currently assume chunks will not have impact damage as it's very expensive. Should look into exposing this a bit more
+		CreateShapeFilterData(MoveChannel, (Owner ? Owner->GetUniqueID() : 0), ColResponse.GetResponseContainer(), 0, ChunkIdxToBoneIdx(ChunkIdx), PQueryFilterData, PSimFilterData, BodyInstance.bUseCCD, false, false);
+
+		PQueryFilterData.word3 |= EPDF_SimpleCollision | EPDF_ComplexCollision;
+
+		SCOPED_SCENE_WRITE_LOCK(Shape->getActor()->getScene());
+
+		Shape->setQueryFilterData(PQueryFilterData);
+		Shape->setSimulationFilterData(PSimFilterData);
+		Shape->setFlag(PxShapeFlag::eSCENE_QUERY_SHAPE, true);
+		Shape->setFlag(PxShapeFlag::eSIMULATION_SHAPE, true);
+		Shape->setFlag(PxShapeFlag::eVISUALIZATION, true);
+	}
+}
+
 void UDestructibleComponent::SetMaterial(int32 ElementIndex, UMaterialInterface* Material)
 {
 	// Mesh component handles render side materials
diff --git a/Engine/Source/Runtime/Engine/Private/PhysicsEngine/PhysScene.cpp b/Engine/Source/Runtime/Engine/Private/PhysicsEngine/PhysScene.cpp
--- a/Engine/Source/Runtime/Engine/Private/PhysicsEngine/PhysScene.cpp
+++ b/Engine/Source/Runtime/Engine/Private/PhysicsEngine/PhysScene.cpp
@@ -130,6 +130,7 @@ FPhysScene::FPhysScene()
 	check(SyncPhysXScene);
 	check(GApexModuleDestructible);
 	GApexModuleDestructible->setWorldSupportPhysXScene(*ApexScene, SyncPhysXScene);
+	GApexModuleDestructible->setDamageApplicationRaycastFlags(NxDestructibleActorRaycastFlags::AllChunks, *ApexScene);
 #endif
 }
 
diff --git a/Engine/Source/Runtime/Engine/Private/PhysicsEngine/PhysXSupport.cpp b/Engine/Source/Runtime/Engine/Private/PhysicsEngine/PhysXSupport.cpp
--- a/Engine/Source/Runtime/Engine/Private/PhysicsEngine/PhysXSupport.cpp
+++ b/Engine/Source/Runtime/Engine/Private/PhysicsEngine/PhysXSupport.cpp
@@ -769,6 +769,10 @@ bool FApexChunkReport::releaseOnNoChunksVisible(const NxDestructibleActor* destr
 ///////// FApexPhysX3Interface //////////////////////////////////
 void FApexPhysX3Interface::setContactReportFlags(physx::PxShape* PShape, physx::PxPairFlags PFlags, NxDestructibleActor* actor, PxU16 actorChunkIndex)
 {
+	UDestructibleComponent* DestructibleComponent = Cast<UDestructibleComponent>(FPhysxUserData::Get<UPrimitiveComponent>(PShape->userData));
+	check(DestructibleComponent);
+
+	DestructibleComponent->Pair(actorChunkIndex, PShape);
 }
 
 physx::PxPairFlags FApexPhysX3Interface::getContactReportFlags(const physx::PxShape* PShape) const
