Made NavModifierVolume react to runtime changes to AreaClass #UE4

UE-14285

diff --git a/Engine/Source/Runtime/Engine/Classes/AI/Navigation/NavModifierVolume.h b/Engine/Source/Runtime/Engine/Classes/AI/Navigation/NavModifierVolume.h
--- a/Engine/Source/Runtime/Engine/Classes/AI/Navigation/NavModifierVolume.h
+++ b/Engine/Source/Runtime/Engine/Classes/AI/Navigation/NavModifierVolume.h
@@ -5,16 +5,25 @@
 #include "NavRelevantInterface.h"
 #include "NavModifierVolume.generated.h"
 
+class UNavArea;
+
 /** 
  *	Allows applying selected AreaClass to navmesh, using Volume's shape
  */
 UCLASS(hidecategories=(Navigation))
 class ANavModifierVolume : public AVolume, public INavRelevantInterface
 {
-	GENERATED_UCLASS_BODY()
+	GENERATED_BODY()
+
+protected:
+	UPROPERTY(EditAnywhere, BlueprintReadOnly, Category=Default)
+	TSubclassOf<UNavArea> AreaClass;
+
+public:
+	ANavModifierVolume(const FObjectInitializer& ObjectInitializer);
 
-	UPROPERTY(EditAnywhere, BlueprintReadWrite, Category=Default)
-	TSubclassOf<class UNavArea> AreaClass;
+	UFUNCTION(BlueprintCallable, Category = "AI|Navigation")
+	void SetAreaClass(TSubclassOf<UNavArea> NewAreaClass = nullptr);
 
 	virtual void GetNavigationData(FNavigationRelevantData& Data) const override;
 	virtual FBox GetNavigationBounds() const override;
diff --git a/Engine/Source/Runtime/Engine/Private/AI/Navigation/NavModifierVolume.cpp b/Engine/Source/Runtime/Engine/Private/AI/Navigation/NavModifierVolume.cpp
--- a/Engine/Source/Runtime/Engine/Private/AI/Navigation/NavModifierVolume.cpp
+++ b/Engine/Source/Runtime/Engine/Private/AI/Navigation/NavModifierVolume.cpp
@@ -9,6 +9,7 @@
 //----------------------------------------------------------------------//
 ANavModifierVolume::ANavModifierVolume(const FObjectInitializer& ObjectInitializer)
 	: Super(ObjectInitializer)
+	, AreaClass(UNavArea_Null::StaticClass())
 {
 }
 
@@ -26,6 +27,20 @@ FBox ANavModifierVolume::GetNavigationBounds() const
 	return GetComponentsBoundingBox();
 }
 
+void ANavModifierVolume::SetAreaClass(TSubclassOf<UNavArea> NewAreaClass)
+{
+	if (NewAreaClass != AreaClass)
+	{
+		AreaClass = NewAreaClass;
+
+		UNavigationSystem* NavSys = UNavigationSystem::GetCurrent(GetWorld());
+		if (NavSys != nullptr)
+		{
+			NavSys->UpdateNavOctree(this);
+		}
+	}
+}
+
 #if WITH_EDITOR
 
 void ANavModifierVolume::PostEditUndo()
