[UE-14306] When searching for global delegate signatures (for bound event nodes), do not accept like-named delegate signatures (that are not native globals).

diff --git a/Engine/Source/Runtime/Engine/Classes/Engine/MemberReference.h b/Engine/Source/Runtime/Engine/Classes/Engine/MemberReference.h
--- a/Engine/Source/Runtime/Engine/Classes/Engine/MemberReference.h
+++ b/Engine/Source/Runtime/Engine/Classes/Engine/MemberReference.h
@@ -318,12 +318,23 @@ public:
 			}
 			else
 			{
-				// Native delegate signatures can have null target scope.
+				// Native delegate signatures can have null target scope (if they're defined outside of a class).
                 // Needed because Blueprint pin types are referencing delegate signatures as
-                // class members, which doesn't have to be true.
+                // class members, which isn't necessarily true (globally defined signatures).
 				if (TFieldType::StaticClass() == UFunction::StaticClass())
 				{
-					ReturnField = (TFieldType*)FindDelegateSignature(MemberName);
+					TFieldType* GlobalFunction = (TFieldType*)FindDelegateSignature(MemberName);
+					// we only want native delegate signatures, even more
+					if (GlobalFunction->GetOwnerClass() == nullptr)
+					{
+						// @TODO: Could this still be the wrong function? I imagine there is a way 
+						//        that we can define like named delegates without the compiler yelling.
+						ReturnField = GlobalFunction;
+					}
+					// @TODO: The delegate we want to find could exist, but FindDelegateSignature()
+					//        could have returned the first one it found (there could be a naming 
+					//        conflict since FindDelegateSignature() searches ANY_PACKAGE)
+					//        How do we make sure it only finds functions with UPackage outers?
 				}
 			}
 		}
