UE4 - Remove redundant operations while switching feature level

diff --git a/Engine/Source/Runtime/Engine/Private/World.cpp b/Engine/Source/Runtime/Engine/Private/World.cpp
--- a/Engine/Source/Runtime/Engine/Private/World.cpp
+++ b/Engine/Source/Runtime/Engine/Private/World.cpp
@@ -5496,8 +5496,6 @@ void UWorld::ChangeFeatureLevel(ERHIFeatureLevel::Type InFeatureLevel)
         FScopedSlowTask SlowTask(100.f, NSLOCTEXT("Engine", "ChangingPreviewRenderingLevelMessage", "Changing Preview Rendering Level"));
         SlowTask.MakeDialog();
         {
-            FlushRenderingCommands();
-
             SlowTask.EnterProgressFrame(10.0f);
             // Give all scene components the opportunity to prepare for pending feature level change.
             for (TObjectIterator<USceneComponent> It; It; ++It)
@@ -5511,6 +5509,7 @@ void UWorld::ChangeFeatureLevel(ERHIFeatureLevel::Type InFeatureLevel)
 
             SlowTask.EnterProgressFrame(10.0f);
             FGlobalComponentReregisterContext RecreateComponents;
+            FlushRenderingCommands();
 
             // Decrement refcount on old feature level
             UMaterialInterface::SetGlobalRequiredFeatureLevel(InFeatureLevel, true);
@@ -5555,10 +5554,6 @@ void UWorld::ChangeFeatureLevel(ERHIFeatureLevel::Type InFeatureLevel)
 
             SlowTask.EnterProgressFrame(10.0f);
             TriggerStreamingDataRebuild();
-
-            SlowTask.EnterProgressFrame(10.0f);
-            FOutputDeviceNull Ar;
-            RecompileShaders(TEXT("CHANGED"), Ar);
         }
 	}
 }
