Fix non APEX build

diff --git a/Engine/Source/Runtime/Engine/Private/Components/DestructibleComponent.cpp b/Engine/Source/Runtime/Engine/Private/Components/DestructibleComponent.cpp
--- a/Engine/Source/Runtime/Engine/Private/Components/DestructibleComponent.cpp
+++ b/Engine/Source/Runtime/Engine/Private/Components/DestructibleComponent.cpp
@@ -800,6 +800,7 @@ FTransform UDestructibleComponent::GetSocketTransform(FName InSocketName, ERelat
 	return ST;
 }
 
+#if WITH_APEX
 void UDestructibleComponent::Pair( int32 ChunkIndex, PxShape* PShape)
 {
 	FDestructibleChunkInfo* CI;
@@ -827,6 +828,7 @@ void UDestructibleComponent::Pair( int32 ChunkIndex, PxShape* PShape)
 		SetCollisionResponseForShape(PShape, ChunkIndex);
 	}
 }
+#endif
 
 void UDestructibleComponent::SetChunkVisible( int32 ChunkIndex, bool bVisible )
 {
