More non-apex build fixes

diff --git a/Engine/Source/Runtime/Engine/Private/Components/DestructibleComponent.cpp b/Engine/Source/Runtime/Engine/Private/Components/DestructibleComponent.cpp
--- a/Engine/Source/Runtime/Engine/Private/Components/DestructibleComponent.cpp
+++ b/Engine/Source/Runtime/Engine/Private/Components/DestructibleComponent.cpp
@@ -832,6 +832,7 @@ void UDestructibleComponent::Pair( int32 ChunkIndex, PxShape* PShape)
 
 void UDestructibleComponent::SetChunkVisible( int32 ChunkIndex, bool bVisible )
 {
+#if WITH_APEX
 	// Bone 0 is a dummy root bone
 	const int32 BoneIndex = ChunkIdxToBoneIdx(ChunkIndex);
 
@@ -861,6 +862,7 @@ void UDestructibleComponent::SetChunkVisible( int32 ChunkIndex, bool bVisible )
 
 	// New bone positions need to be sent to render thread
 	MarkRenderDynamicDataDirty();
+#endif
 }
 
 #if WITH_APEX
@@ -1268,6 +1270,7 @@ bool UDestructibleComponent::IsChunkLarge(int32 ChunkIdx) const
 
 void UDestructibleComponent::SetCollisionResponseForActor(PxRigidDynamic* Actor, int32 ChunkIdx, const FCollisionResponseContainer* ResponseOverride /*= NULL*/)
 {
+#if WITH_APEX
 	if (ApexDestructibleActor == NULL)
 	{
 		return;
@@ -1282,7 +1285,7 @@ void UDestructibleComponent::SetCollisionResponseForActor(PxRigidDynamic* Actor,
 		AActor* Owner = GetOwner();
 		bool bLargeChunk = IsChunkLarge(ChunkIdx);
 		const FCollisionResponseContainer& UseResponse = ResponseOverride == NULL ? (bLargeChunk ? LargeChunkCollisionResponse.GetResponseContainer() : SmallChunkCollisionResponse.GetResponseContainer()) : *ResponseOverride;
-#if WITH_APEX
+
 		physx::PxU32 SupportDepth = TheDestructibleMesh->ApexDestructibleAsset->getChunkDepth(ChunkIdx);
 
 		const bool bEnableImpactDamage = IsImpactDamageEnabled(TheDestructibleMesh, SupportDepth);
@@ -1307,12 +1310,13 @@ void UDestructibleComponent::SetCollisionResponseForActor(PxRigidDynamic* Actor,
 			Shape->setFlag(PxShapeFlag::eSIMULATION_SHAPE, true); 
 			Shape->setFlag(PxShapeFlag::eVISUALIZATION, true);
 		}
-#endif
 	}
+#endif
 }
 
 void UDestructibleComponent::SetCollisionResponseForAllActors(const FCollisionResponseContainer& ResponseOverride)
 {
+#if WITH_APEX
 	if (ApexDestructibleActor == NULL)
 	{
 		return;
@@ -1352,6 +1356,7 @@ void UDestructibleComponent::SetCollisionResponseForAllActors(const FCollisionRe
 
 		ApexDestructibleActor->releasePhysXActorBuffer();
 	}
+#endif
 }
 
 void UDestructibleComponent::SetCollisionResponseForShape(PxShape* Shape, int32 ChunkIdx)
