Fix for subtle bug introduced in CL1724453 (probably).
When resetting the parent material of a material instance to the default, the EPropertyChangeType fed in is Unspecified. This caused the RecacheMaterialInstanceUniformExpressions() call to be skipped leading to a crash in some cases.
Reset of a property should come in with the ValueSet change type and the change in 1724453 would be fine.
Editor team are going to fix this.

diff --git a/Engine/Source/Runtime/Engine/Private/Materials/MaterialInstance.cpp b/Engine/Source/Runtime/Engine/Private/Materials/MaterialInstance.cpp
--- a/Engine/Source/Runtime/Engine/Private/Materials/MaterialInstance.cpp
+++ b/Engine/Source/Runtime/Engine/Private/Materials/MaterialInstance.cpp
@@ -2062,7 +2062,7 @@ void UMaterialInstance::PostEditChangeProperty(FPropertyChangedEvent& PropertyCh
 
 	UpdateStaticPermutation(StaticParameters);
 
-	if(PropertyChangedEvent.ChangeType == EPropertyChangeType::ValueSet)
+	if (PropertyChangedEvent.ChangeType == EPropertyChangeType::ValueSet || PropertyChangedEvent.ChangeType == EPropertyChangeType::Unspecified)
 	{
 		RecacheMaterialInstanceUniformExpressions(this);
 	}
