Big revamp to SplineComponent (continuing work-in-progress).

Summary of changes:
- FInterpCurve now natively understands looping. A input key can be set, at which point the spline will loop back to its start point.  Adjusted FInterpCurve methods to take looped curves into account.
- Fixed bugs in FInterpCurve regarding returning default values (T(ForceInit) instead of T() otherwise FVectors and FQuats are not initialized correctly).
- SplineComponent now holds three separate splines (location, rotation and scale).
- The SplineReparamTable is no longer transient, but is built and saved as part of the asset.
- Changed the SplineComponent API. All methods now take a ESplineCoordinateSpace value which specifies whether the values are intended for use in local or world space. The old methods have been deprecated but not removed (for now).
- bAllowSplineEditingPerInstance is deprecated. A new member (bSplineHasBeenEdited) has superseded it, and is set automatically when editing the spline with the visualizer. Added "Reset to Default" option in the visualizer, which restores the archetype defaults.
- SplineComponentVisualizer determines if spline points have been initialized by the user construction script, and, if so, will disallow editing. This is to conform with the general idea that the UCS always overrides property values.
- Fixed a number of bugs in the visualizer, and added new facilities to allow the scale and rotation to be edited (still work-in-progress).

diff --git a/Engine/Config/BaseEngine.ini b/Engine/Config/BaseEngine.ini
--- a/Engine/Config/BaseEngine.ini
+++ b/Engine/Config/BaseEngine.ini
@@ -1167,6 +1167,9 @@ bUseStreamingPause=false
 
 +ActiveClassRedirects=(OldClassName="ESlateCheckBoxState", NewClassName="ECheckBoxState")
 
+; SplineComponent API overhaul
++K2FieldRedirects=(OldFieldName="SplineComponent.GetNumSplinePoints",NewFieldName="SplineComponent.GetNumberOfSplinePoints")
+
 
 ; VER_UE4_INVERSE_SQUARED_LIGHTS_DEFAULT
 +K2FieldRedirects=(OldFieldName="LightComponentBase.Brightness",NewFieldName="LightComponentBase.Intensity")
diff --git a/Engine/Source/Editor/ComponentVisualizers/Private/SplineComponentVisualizer.cpp b/Engine/Source/Editor/ComponentVisualizers/Private/SplineComponentVisualizer.cpp
--- a/Engine/Source/Editor/ComponentVisualizers/Private/SplineComponentVisualizer.cpp
+++ b/Engine/Source/Editor/ComponentVisualizers/Private/SplineComponentVisualizer.cpp
@@ -36,6 +36,8 @@ public:
 		UI_COMMAND(SetKeyToCurve, "Curve", "Set spline point to Curve type", EUserInterfaceActionType::RadioButton, FInputChord());
 		UI_COMMAND(SetKeyToLinear, "Linear", "Set spline point to Linear type", EUserInterfaceActionType::RadioButton, FInputChord());
 		UI_COMMAND(SetKeyToConstant, "Constant", "Set spline point to Constant type", EUserInterfaceActionType::RadioButton, FInputChord());
+		UI_COMMAND(VisualizeRollAndScale, "Visualize Roll and Scale", "Whether the visualization should show roll and scale on this spline.", EUserInterfaceActionType::ToggleButton, FInputChord());
+		UI_COMMAND(ResetToDefault, "Reset to Default", "Reset this spline to its archetype default.", EUserInterfaceActionType::Button, FInputChord());
 	}
 
 public:
@@ -62,6 +64,12 @@ public:
 
 	/** Set spline key to Constant type */
 	TSharedPtr<FUICommandInfo> SetKeyToConstant;
+
+	/** Whether the visualization should show roll and scale */
+	TSharedPtr<FUICommandInfo> VisualizeRollAndScale;
+
+	/** Reset this spline to its default */
+	TSharedPtr<FUICommandInfo> ResetToDefault;
 };
 
 
@@ -125,6 +133,17 @@ void FSplineComponentVisualizer::OnRegister()
 		FExecuteAction::CreateSP(this, &FSplineComponentVisualizer::OnSetKeyType, CIM_Constant),
 		FCanExecuteAction(),
 		FIsActionChecked::CreateSP(this, &FSplineComponentVisualizer::IsKeyTypeSet, CIM_Constant));
+
+	SplineComponentVisualizerActions->MapAction(
+		Commands.VisualizeRollAndScale,
+		FExecuteAction::CreateSP(this, &FSplineComponentVisualizer::OnSetVisualizeRollAndScale),
+		FCanExecuteAction(),
+		FIsActionChecked::CreateSP(this, &FSplineComponentVisualizer::IsVisualizingRollAndScale));
+
+	SplineComponentVisualizerActions->MapAction(
+		Commands.ResetToDefault,
+		FExecuteAction::CreateSP(this, &FSplineComponentVisualizer::OnResetToDefault),
+		FCanExecuteAction::CreateSP(this, &FSplineComponentVisualizer::CanResetToDefault));
 }
 
 FSplineComponentVisualizer::~FSplineComponentVisualizer()
@@ -136,15 +155,14 @@ void FSplineComponentVisualizer::DrawVisualization(const UActorComponent* Compon
 {
 	if (const USplineComponent* SplineComp = Cast<const USplineComponent>(Component))
 	{
-		if (!SplineComp->bAllowSplineEditingPerInstance)
-		{
-			return;
-		}
+		const USplineComponent* EditedSplineComp = GetEditedSplineComponent();
 
-		USplineComponent* EditedSplineComp = GetEditedSplineComponent();
+		const USplineComponent* Archetype = CastChecked<USplineComponent>(SplineComp->GetArchetype());
+		const bool bIsSplineEditable = SplineComp->bSplineHasBeenEdited || SplineComp->SplineInfo.Points == Archetype->SplineInfo.Points;
 
-		const FColor NormalColor = SplineComp->EditorUnselectedSplineSegmentColor;
-		const FColor SelectedColor = SplineComp->EditorSelectedSplineSegmentColor;
+		const FColor ReadOnlyColor = FColor(255, 0, 255, 255);
+		const FColor NormalColor = bIsSplineEditable ? SplineComp->EditorUnselectedSplineSegmentColor : ReadOnlyColor;
+		const FColor SelectedColor = bIsSplineEditable ? SplineComp->EditorSelectedSplineSegmentColor : ReadOnlyColor;
 		const float GrabHandleSize = 12.0f;
 		const float TangentHandleSize = 10.0f;
 
@@ -157,34 +175,76 @@ void FSplineComponentVisualizer::DrawVisualization(const UActorComponent* Compon
 			{
 				if (SplineInfo.Points[SelectedKey].IsCurveKey())
 				{
-					const FVector KeyPos = SplineComp->ComponentToWorld.TransformPosition(SplineInfo.Points[SelectedKey].OutVal);
-					const FVector TangentWorldDirection = SplineComp->ComponentToWorld.TransformVector(SplineInfo.Points[SelectedKey].LeaveTangent);
+					const FVector Location = SplineComp->GetLocationAtSplinePoint(SelectedKey, ESplineCoordinateSpace::World);
+					const FVector Tangent = SplineComp->GetTangentAtSplinePoint(SelectedKey, ESplineCoordinateSpace::World);
 
 					PDI->SetHitProxy(NULL);
-					DrawDashedLine(PDI, KeyPos, KeyPos + TangentWorldDirection, SelectedColor, 5, SDPG_Foreground);
-					DrawDashedLine(PDI, KeyPos, KeyPos - TangentWorldDirection, SelectedColor, 5, SDPG_Foreground);
-					PDI->SetHitProxy(new HSplineTangentHandleProxy(Component, SelectedKey, false));
-					PDI->DrawPoint(KeyPos + TangentWorldDirection, SelectedColor, TangentHandleSize, SDPG_Foreground);
-					PDI->SetHitProxy(new HSplineTangentHandleProxy(Component, SelectedKey, true));
-					PDI->DrawPoint(KeyPos - TangentWorldDirection, SelectedColor, TangentHandleSize, SDPG_Foreground);
+					DrawDashedLine(PDI, Location, Location + Tangent, SelectedColor, 5, SDPG_Foreground);
+					DrawDashedLine(PDI, Location, Location - Tangent, SelectedColor, 5, SDPG_Foreground);
+
+					if (bIsSplineEditable)
+					{
+						PDI->SetHitProxy(new HSplineTangentHandleProxy(Component, SelectedKey, false));
+					}
+					PDI->DrawPoint(Location + Tangent, SelectedColor, TangentHandleSize, SDPG_Foreground);
+
+					if (bIsSplineEditable)
+					{
+						PDI->SetHitProxy(new HSplineTangentHandleProxy(Component, SelectedKey, true));
+					}
+					PDI->DrawPoint(Location - Tangent, SelectedColor, TangentHandleSize, SDPG_Foreground);
+
 					PDI->SetHitProxy(NULL);
 				}
 			}
 		}
 
+		const bool bShouldVisualizeScale = SplineComp->bShouldVisualizeScale;
+		const float DefaultScale = SplineComp->ScaleVisualizationWidth;
+
 		FVector OldKeyPos(0);
-		float OldKeyTime = 0.f;
-		for (int32 KeyIdx = 0; KeyIdx < SplineInfo.Points.Num(); KeyIdx++)
+		FVector OldKeyRightVector(0);
+		FVector OldKeyScale(0);
+
+		const int32 NumPoints = SplineInfo.Points.Num();
+		const int32 NumSegments = SplineComp->IsClosedLoop() ? NumPoints : NumPoints - 1;
+		for (int32 KeyIdx = 0; KeyIdx < NumSegments + 1; KeyIdx++)
 		{
-			const float NewKeyTime = SplineInfo.Points[KeyIdx].InVal;
-			const FVector NewKeyPos = SplineComp->ComponentToWorld.TransformPosition(SplineInfo.Points[KeyIdx].OutVal);
+			const FVector NewKeyPos = SplineComp->GetLocationAtSplinePoint(KeyIdx, ESplineCoordinateSpace::World);
+			const FVector NewKeyRightVector = SplineComp->GetRightVectorAtSplinePoint(KeyIdx, ESplineCoordinateSpace::World);
+			const FVector NewKeyUpVector = SplineComp->GetUpVectorAtSplinePoint(KeyIdx, ESplineCoordinateSpace::World);
+			const FVector NewKeyScale = SplineComp->GetScaleAtSplinePoint(KeyIdx) * DefaultScale;
 
 			const FColor KeyColor = (SplineComp == EditedSplineComp && SelectedKeys.Contains(KeyIdx)) ? SelectedColor : NormalColor;
 
-			// Draw the keypoint
-			if (!SplineComp->IsClosedLoop() || KeyIdx < SplineInfo.Points.Num() - 1)
+			// Draw the keypoint and up/right vectors
+			if (KeyIdx < NumPoints)
 			{
-				PDI->SetHitProxy(new HSplineKeyProxy(Component, KeyIdx));
+				if (bShouldVisualizeScale)
+				{
+					PDI->SetHitProxy(NULL);
+
+					PDI->DrawLine(NewKeyPos, NewKeyPos - NewKeyRightVector * NewKeyScale.Y, KeyColor, SDPG_Foreground);
+					PDI->DrawLine(NewKeyPos, NewKeyPos + NewKeyRightVector * NewKeyScale.Y, KeyColor, SDPG_Foreground);
+					PDI->DrawLine(NewKeyPos, NewKeyPos + NewKeyUpVector * NewKeyScale.Z, KeyColor, SDPG_Foreground);
+
+					const int32 ArcPoints = 20;
+					FVector OldArcPos = NewKeyPos + NewKeyRightVector * NewKeyScale.Y;
+					for (int32 ArcIndex = 1; ArcIndex <= ArcPoints; ArcIndex++)
+					{
+						float Sin;
+						float Cos;
+						FMath::SinCos(&Sin, &Cos, ArcIndex * PI / ArcPoints);
+						const FVector NewArcPos = NewKeyPos + Cos * NewKeyRightVector * NewKeyScale.Y + Sin * NewKeyUpVector * NewKeyScale.Z;
+						PDI->DrawLine(OldArcPos, NewArcPos, KeyColor, SDPG_Foreground);
+						OldArcPos = NewArcPos;
+					}
+				}
+
+				if (bIsSplineEditable)
+				{
+					PDI->SetHitProxy(new HSplineKeyProxy(Component, KeyIdx));
+				}
 				PDI->DrawPoint(NewKeyPos, KeyColor, GrabHandleSize, SDPG_Foreground);
 				PDI->SetHitProxy(NULL);
 			}
@@ -193,7 +253,10 @@ void FSplineComponentVisualizer::DrawVisualization(const UActorComponent* Compon
 			if (KeyIdx > 0)
 			{
 				const FColor LineColor = (SplineComp == EditedSplineComp && SelectedKeys.Contains(KeyIdx - 1)) ? SelectedColor : NormalColor;
-				PDI->SetHitProxy(new HSplineSegmentProxy(Component, KeyIdx - 1));
+				if (bIsSplineEditable)
+				{
+					PDI->SetHitProxy(new HSplineSegmentProxy(Component, KeyIdx - 1));
+				}
 
 				// For constant interpolation - don't draw ticks - just draw dotted line.
 				if (SplineInfo.Points[KeyIdx - 1].InterpMode == CIM_Constant)
@@ -202,31 +265,40 @@ void FSplineComponentVisualizer::DrawVisualization(const UActorComponent* Compon
 				}
 				else
 				{
-					const int32 NumSteps = FMath::CeilToInt((NewKeyTime - OldKeyTime) * 32.0f);
-					const float DrawSubstep = (NewKeyTime - OldKeyTime) / NumSteps;
-
 					// Find position on first keyframe.
-					float OldTime = OldKeyTime;
 					FVector OldPos = OldKeyPos;
+					FVector OldRightVector = OldKeyRightVector;
+					FVector OldScale = OldKeyScale;
 
 					// Then draw a line for each substep.
-					for (int32 StepIdx = 1; StepIdx<NumSteps + 1; StepIdx++)
+					const int32 NumSteps = 20;
+
+					for (int32 StepIdx = 1; StepIdx <= NumSteps; StepIdx++)
 					{
-						const float NewTime = OldKeyTime + StepIdx*DrawSubstep;
-						const FVector NewPos = SplineComp->ComponentToWorld.TransformPosition( SplineInfo.Eval(NewTime, FVector::ZeroVector) );
+						const float Key = (KeyIdx - 1) + (StepIdx / static_cast<float>(NumSteps));
+						const FVector NewPos = SplineComp->GetLocationAtSplineInputKey(Key, ESplineCoordinateSpace::World);
+						const FVector NewRightVector = SplineComp->GetRightVectorAtSplineInputKey(Key, ESplineCoordinateSpace::World);
+						const FVector NewScale = SplineComp->GetScaleAtSplineInputKey(Key) * DefaultScale;
 
 						PDI->DrawLine(OldPos, NewPos, LineColor, SDPG_Foreground);
+						if (bShouldVisualizeScale)
+						{
+							PDI->DrawLine(OldPos - OldRightVector * OldScale.Y, NewPos - NewRightVector * NewScale.Y, LineColor, SDPG_Foreground);
+							PDI->DrawLine(OldPos + OldRightVector * OldScale.Y, NewPos + NewRightVector * NewScale.Y, LineColor, SDPG_Foreground);
+						}
 
-						OldTime = NewTime;
 						OldPos = NewPos;
+						OldRightVector = NewRightVector;
+						OldScale = NewScale;
 					}
 				}
 
 				PDI->SetHitProxy(NULL);
 			}
 
-			OldKeyTime = NewKeyTime;
 			OldKeyPos = NewKeyPos;
+			OldKeyRightVector = NewKeyRightVector;
+			OldKeyScale = NewKeyScale;
 		}
 	}
 }
@@ -277,8 +349,6 @@ void FSplineComponentVisualizer::ChangeSelectionState(int32 Index, bool bIsCtrlH
 
 bool FSplineComponentVisualizer::VisProxyHandleClick(FLevelEditorViewportClient* InViewportClient, HComponentVisProxy* VisProxy, const FViewportClick& Click)
 {
-	bool bEditing = false;
-
 	if(VisProxy && VisProxy->Component.IsValid())
 	{
 		const USplineComponent* SplineComp = CastChecked<const USplineComponent>(VisProxy->Component.Get());
@@ -303,7 +373,9 @@ bool FSplineComponentVisualizer::VisProxyHandleClick(FLevelEditorViewportClient*
 				SelectedTangentHandle = INDEX_NONE;
 				SelectedTangentHandleType = ESelectedTangentHandle::None;
 
-				bEditing = true;
+				CachedRotation = SplineComp->GetQuaternionAtSplinePoint(LastKeyIndexSelected, ESplineCoordinateSpace::World);
+
+				return true;
 			}
 			else if (VisProxy->IsA(HSplineSegmentProxy::StaticGetType()))
 			{
@@ -319,8 +391,10 @@ bool FSplineComponentVisualizer::VisProxyHandleClick(FLevelEditorViewportClient*
 				SelectedTangentHandle = INDEX_NONE;
 				SelectedTangentHandleType = ESelectedTangentHandle::None;
 
+				CachedRotation = SplineComp->GetQuaternionAtSplinePoint(LastKeyIndexSelected, ESplineCoordinateSpace::World);
+
 				float SubsegmentStartKey = static_cast<float>(SelectedSegmentIndex);
-				FVector SubsegmentStart = SplineComp->ComponentToWorld.TransformPosition(SplineComp->SplineInfo.Eval(SubsegmentStartKey, FVector::ZeroVector));
+				FVector SubsegmentStart = SplineComp->GetLocationAtSplineInputKey(SubsegmentStartKey, ESplineCoordinateSpace::World);
 
 				float ClosestDistance = TNumericLimits<float>::Max();
 				FVector BestLocation = SubsegmentStart;
@@ -328,7 +402,7 @@ bool FSplineComponentVisualizer::VisProxyHandleClick(FLevelEditorViewportClient*
 				for (int32 Step = 1; Step < NumSubdivisions; Step++)
 				{
 					const float SubsegmentEndKey = SelectedSegmentIndex + Step / static_cast<float>(NumSubdivisions);
-					const FVector SubsegmentEnd = SplineComp->ComponentToWorld.TransformPosition(SplineComp->SplineInfo.Eval(SubsegmentEndKey, FVector::ZeroVector));
+					const FVector SubsegmentEnd = SplineComp->GetLocationAtSplineInputKey(SubsegmentEndKey, ESplineCoordinateSpace::World);
 
 					FVector SplineClosest;
 					FVector RayClosest;
@@ -347,7 +421,7 @@ bool FSplineComponentVisualizer::VisProxyHandleClick(FLevelEditorViewportClient*
 
 				SelectedSplinePosition = BestLocation;
 
-				bEditing = true;
+				return true;
 			}
 			else if (VisProxy->IsA(HSplineTangentHandleProxy::StaticGetType()))
 			{
@@ -360,7 +434,9 @@ bool FSplineComponentVisualizer::VisProxyHandleClick(FLevelEditorViewportClient*
 				SelectedTangentHandle = KeyProxy->KeyIndex;
 				SelectedTangentHandleType = KeyProxy->bArriveTangent ? ESelectedTangentHandle::Arrive : ESelectedTangentHandle::Leave;
 
-				bEditing = true;
+				CachedRotation = SplineComp->GetQuaternionAtSplinePoint(SelectedTangentHandle, ESplineCoordinateSpace::World);
+
+				return true;
 			}
 		}
 		else
@@ -369,7 +445,7 @@ bool FSplineComponentVisualizer::VisProxyHandleClick(FLevelEditorViewportClient*
 		}
 	}
 
-	return bEditing;
+	return false;
 }
 
 USplineComponent* FSplineComponentVisualizer::GetEditedSplineComponent() const
@@ -418,32 +494,13 @@ bool FSplineComponentVisualizer::GetWidgetLocation(const FEditorViewportClient*
 
 bool FSplineComponentVisualizer::GetCustomInputCoordinateSystem(const FEditorViewportClient* ViewportClient, FMatrix& OutMatrix) const
 {
-	if (ViewportClient->GetWidgetCoordSystemSpace() == COORD_Local)
+	if (ViewportClient->GetWidgetCoordSystemSpace() == COORD_Local || ViewportClient->GetWidgetMode() == FWidget::WM_Rotate)
 	{
 		USplineComponent* SplineComp = GetEditedSplineComponent();
 		if (SplineComp != nullptr)
 		{
-			// First look at selected tangent handle for coordinate system
-			int32 Index = SelectedTangentHandle;
-			if (Index == INDEX_NONE)
-			{
-				// If not set, fall back to last key index selected
-				Index = LastKeyIndexSelected;
-			}
-
-			if (Index != INDEX_NONE)
-			{
-				check(Index < SplineComp->SplineInfo.Points.Num());
-				check(SelectedTangentHandle != INDEX_NONE || SelectedKeys.Contains(Index));
-
-				const auto& Point = SplineComp->SplineInfo.Points[Index];
-				const FVector Tangent = Point.ArriveTangent.IsNearlyZero() ? FVector(1.0f, 0.0f, 0.0f) : Point.ArriveTangent.GetSafeNormal();
-				const FVector Bitangent = (Tangent.Z == 1.0f) ? FVector(1.0f, 0.0f, 0.0f) : FVector(-Tangent.Y, Tangent.X, 0.0f).GetSafeNormal();
-				const FVector Normal = FVector::CrossProduct(Tangent, Bitangent);
-
-				OutMatrix = FMatrix(Tangent, Bitangent, Normal, FVector::ZeroVector) * FQuatRotationTranslationMatrix(SplineComp->ComponentToWorld.GetRotation(), FVector::ZeroVector);
-				return true;
-			}
+			OutMatrix = FRotationMatrix::Make(CachedRotation);
+			return true;
 		}
 	}
 
@@ -456,9 +513,11 @@ bool FSplineComponentVisualizer::HandleInputDelta(FEditorViewportClient* Viewpor
 	USplineComponent* SplineComp = GetEditedSplineComponent();
 	if (SplineComp != nullptr)
 	{
-		const int32 NumPoints = SplineComp->SplineInfo.Points.Num();
+		FInterpCurveVector& SplineInfo = SplineComp->SplineInfo;
+		FInterpCurveQuat& SplineRotInfo = SplineComp->SplineRotInfo;
+		FInterpCurveVector& SplineScaleInfo = SplineComp->SplineScaleInfo;
 
-		bool bChangedFirstPoint = false;
+		const int32 NumPoints = SplineInfo.Points.Num();
 
 		if (SelectedTangentHandle != INDEX_NONE)
 		{
@@ -466,22 +525,19 @@ bool FSplineComponentVisualizer::HandleInputDelta(FEditorViewportClient* Viewpor
 
 			check(SelectedTangentHandle < NumPoints);
 
-			FInterpCurvePoint<FVector>& EditedPoint = SplineComp->SplineInfo.Points[SelectedTangentHandle];
-
 			if (!DeltaTranslate.IsZero())
 			{
 				check(SelectedTangentHandleType != ESelectedTangentHandle::None);
 
+				SplineComp->Modify();
+
+				FInterpCurvePoint<FVector>& EditedPoint = SplineInfo.Points[SelectedTangentHandle];
 				const FVector Delta = (SelectedTangentHandleType == ESelectedTangentHandle::Leave) ? DeltaTranslate : -DeltaTranslate;
-				const FVector NewTangent = EditedPoint.LeaveTangent + SplineComp->ComponentToWorld.InverseTransformVector(Delta);
-				EditedPoint.LeaveTangent = NewTangent;
-				EditedPoint.ArriveTangent = NewTangent;
-				EditedPoint.InterpMode = CIM_CurveUser;
+				const FVector Tangent = EditedPoint.LeaveTangent + SplineComp->ComponentToWorld.InverseTransformVector(Delta);
 
-				if (SelectedTangentHandle == 0)
-				{
-					bChangedFirstPoint = true;
-				}
+				EditedPoint.LeaveTangent = Tangent;
+				EditedPoint.ArriveTangent = Tangent;
+				EditedPoint.InterpMode = CIM_CurveUser;
 			}
 		}
 		else
@@ -503,7 +559,9 @@ bool FSplineComponentVisualizer::HandleInputDelta(FEditorViewportClient* Viewpor
 
 			for (int32 SelectedKeyIndex : SelectedKeys)
 			{
-				FInterpCurvePoint<FVector>& EditedPoint = SplineComp->SplineInfo.Points[SelectedKeyIndex];
+				FInterpCurvePoint<FVector>& EditedPoint = SplineInfo.Points[SelectedKeyIndex];
+				FInterpCurvePoint<FQuat>& EditedRotPoint = SplineRotInfo.Points[SelectedKeyIndex];
+				FInterpCurvePoint<FVector>& EditedScalePoint = SplineScaleInfo.Points[SelectedKeyIndex];
 
 				if (!DeltaTranslate.IsZero())
 				{
@@ -521,53 +579,41 @@ bool FSplineComponentVisualizer::HandleInputDelta(FEditorViewportClient* Viewpor
 					EditedPoint.InterpMode = CIM_CurveUser;
 
 					// Rotate tangent according to delta rotation
-					const FVector NewTangent = DeltaRotate.RotateVector(EditedPoint.LeaveTangent);
+					FVector NewTangent = SplineComp->ComponentToWorld.GetRotation().RotateVector(EditedPoint.LeaveTangent); // convert local-space tangent vector to world-space
+					NewTangent = DeltaRotate.RotateVector(NewTangent); // apply world-space delta rotation to world-space tangent
+					NewTangent = SplineComp->ComponentToWorld.GetRotation().Inverse().RotateVector(NewTangent); // convert world-space tangent vector back into local-space
 					EditedPoint.LeaveTangent = NewTangent;
 					EditedPoint.ArriveTangent = NewTangent;
+
+					// Rotate spline rotation according to delta rotation
+					FQuat NewRot = SplineComp->ComponentToWorld.GetRotation() * EditedRotPoint.OutVal; // convert local-space rotation to world-space
+					NewRot = DeltaRotate.Quaternion() * NewRot; // apply world-space rotation
+					NewRot = SplineComp->ComponentToWorld.GetRotation().Inverse() * NewRot; // convert world-space rotation to local-space
+					EditedRotPoint.OutVal = NewRot;
 				}
 
-				if (!DeltaScale.IsZero())
+				if (DeltaScale.X != 0.0f)
 				{
 					// Set point tangent as user controlled
 					EditedPoint.InterpMode = CIM_CurveUser;
 
-					// Break tangent into direction and length so we can change its scale (the 'tension')
-					// independently of its direction.
-					FVector Direction;
-					float Length;
-					EditedPoint.LeaveTangent.ToDirectionAndLength(Direction, Length);
-
-					// Figure out which component has changed, and use it
-					float DeltaScaleValue = (DeltaScale.X != 0.0f) ? DeltaScale.X : ((DeltaScale.Y != 0.0f) ? DeltaScale.Y : DeltaScale.Z);
-
-					// Change scale, avoiding singularity by never allowing a scale of 0, hence preserving direction.
-					Length += DeltaScaleValue * 10.0f;
-					if (Length == 0.0f)
-					{
-						Length = SMALL_NUMBER;
-					}
-
-					const FVector NewTangent = Direction * Length;
+					const FVector NewTangent = EditedPoint.LeaveTangent * (1.0f + DeltaScale.X);
 					EditedPoint.LeaveTangent = NewTangent;
 					EditedPoint.ArriveTangent = NewTangent;
 				}
 
-				if (SelectedKeyIndex == 0)
+				if (DeltaScale.Y != 0.0f)
 				{
-					bChangedFirstPoint = true;
+					// Scale in Y adjusts the scale spline
+					EditedScalePoint.OutVal.Y *= (1.0f + DeltaScale.Y);
 				}
-			}
-		}
 
-		// Duplicate first point to last point if it is a closed loop
-		if (SplineComp->IsClosedLoop() && bChangedFirstPoint && NumPoints > 0)
-		{
-			const FInterpCurvePoint<FVector>& EditedPoint = SplineComp->SplineInfo.Points[0];
-			FInterpCurvePoint<FVector>& DuplicatedPoint = SplineComp->SplineInfo.Points[NumPoints - 1];
-			DuplicatedPoint.OutVal = EditedPoint.OutVal;
-			DuplicatedPoint.InterpMode = EditedPoint.InterpMode;
-			DuplicatedPoint.LeaveTangent = EditedPoint.LeaveTangent;
-			DuplicatedPoint.ArriveTangent = EditedPoint.ArriveTangent;
+				if (DeltaScale.Z != 0.0f)
+				{
+					// Scale in Z adjusts the scale spline
+					EditedScalePoint.OutVal.Z *= (1.0f + DeltaScale.Z);
+				}
+			}
 		}
 
 		NotifyComponentModified();
@@ -583,6 +629,20 @@ bool FSplineComponentVisualizer::HandleInputKey(FEditorViewportClient* ViewportC
 
 	if (Key == EKeys::LeftMouseButton && Event == IE_Released)
 	{
+		USplineComponent* SplineComp = GetEditedSplineComponent();
+		if (SplineComp != nullptr)
+		{
+			// Recache widget rotation
+			int32 Index = SelectedTangentHandle;
+			if (Index == INDEX_NONE)
+			{
+				// If not set, fall back to last key index selected
+				Index = LastKeyIndexSelected;
+			}
+
+			CachedRotation = SplineComp->GetQuaternionAtSplinePoint(Index, ESplineCoordinateSpace::World);
+		}
+
 		// Reset duplication flag on LMB release
 		bAllowDuplication = true;
 	}
@@ -622,10 +682,6 @@ void FSplineComponentVisualizer::OnDuplicateKey()
 		Owner->Modify();
 	}
 
-	// Remember whether it's a closed loop for later, and turn off the loop while we manipulate the points
-	const bool bWasLoop = SplineComp->IsClosedLoop();
-	SplineComp->SetClosedLoop(false);
-
 	// Get a sorted list of all the selected indices, highest to lowest
 	TArray<int32> SelectedKeysSorted;
 	for (int32 SelectedKeyIndex : SelectedKeys)
@@ -635,12 +691,17 @@ void FSplineComponentVisualizer::OnDuplicateKey()
 	SelectedKeysSorted.Sort([](int32 A, int32 B) { return A > B; });
 
 	// Insert duplicates into the list, highest index first, so that the lower indices remain the same
-	TArray<FInterpCurvePoint<FVector>>& Points = SplineComp->SplineInfo.Points;
+	FInterpCurveVector& SplineInfo = SplineComp->SplineInfo;
+	FInterpCurveQuat& SplineRotInfo = SplineComp->SplineRotInfo;
+	FInterpCurveVector& SplineScaleInfo = SplineComp->SplineScaleInfo;
+
 	for (int32 SelectedKeyIndex : SelectedKeysSorted)
 	{
-		FInterpCurvePoint<FVector> KeyToDupe = Points[SelectedKeyIndex];
-		KeyToDupe.InterpMode = CIM_CurveAuto;
-		Points.Insert(KeyToDupe, SelectedKeyIndex);
+		// Insert duplicates into arrays.
+		// It's necessary to take a copy because copying existing array items by reference isn't allowed (the array may reallocate)
+		SplineInfo.Points.Insert(FInterpCurvePoint<FVector>(SplineInfo.Points[SelectedKeyIndex]), SelectedKeyIndex);
+		SplineRotInfo.Points.Insert(FInterpCurvePoint<FQuat>(SplineRotInfo.Points[SelectedKeyIndex]), SelectedKeyIndex);
+		SplineScaleInfo.Points.Insert(FInterpCurvePoint<FVector>(SplineScaleInfo.Points[SelectedKeyIndex]), SelectedKeyIndex);
 	}
 
 	// Repopulate the selected keys
@@ -662,22 +723,24 @@ void FSplineComponentVisualizer::OnDuplicateKey()
 	SelectedTangentHandle = INDEX_NONE;
 	SelectedTangentHandleType = ESelectedTangentHandle::None;
 
-	// Update Input value for all keys
-	SplineComp->RefreshSplineInputs();
-
-	// Restore old closed loop state
-	SplineComp->SetClosedLoop(bWasLoop);
-
 	NotifyComponentModified();
+
+//	CachedRotation = SplineComp->GetQuaternionAtSplinePoint(LastKeyIndexSelected, ESplineCoordinateSpace::World);
 }
 
 
 bool FSplineComponentVisualizer::CanAddKey() const
 {
 	USplineComponent* SplineComp = GetEditedSplineComponent();
-	return (SplineComp &&
-		SelectedSegmentIndex != INDEX_NONE &&
-		SelectedSegmentIndex < SplineComp->SplineInfo.Points.Num() - 1);
+	if (SplineComp == nullptr)
+	{
+		return false;
+	}
+
+	const int32 NumPoints = SplineComp->SplineInfo.Points.Num();
+	const int32 NumSegments = SplineComp->IsClosedLoop() ? NumPoints : NumPoints - 1;
+
+	return (SelectedSegmentIndex != INDEX_NONE && SelectedSegmentIndex < NumSegments);
 }
 
 
@@ -698,27 +761,25 @@ void FSplineComponentVisualizer::OnAddKey()
 		Owner->Modify();
 	}
 
-	// Remember whether it's a closed loop for later, and turn off the loop while we manipulate the points
-	const bool bWasLoop = SplineComp->IsClosedLoop();
-	SplineComp->SetClosedLoop(false);
+	FInterpCurveVector& SplineInfo = SplineComp->SplineInfo;
+	FInterpCurveQuat& SplineRotInfo = SplineComp->SplineRotInfo;
+	FInterpCurveVector& SplineScaleInfo = SplineComp->SplineScaleInfo;
 
-	FInterpCurvePoint<FVector> NewKey;
-	NewKey.InVal = SelectedSegmentIndex;
-	NewKey.OutVal = SplineComp->ComponentToWorld.InverseTransformPosition(SelectedSplinePosition);
-	NewKey.InterpMode = CIM_CurveAuto;
-	int32 NewKeyIndex = SplineComp->SplineInfo.Points.Insert(NewKey, SelectedSegmentIndex + 1);
+	FInterpCurvePoint<FVector> NewPoint(0.0f, SplineComp->ComponentToWorld.InverseTransformPosition(SelectedSplinePosition), FVector::ZeroVector, FVector::ZeroVector, CIM_CurveAuto);
+	FInterpCurvePoint<FQuat> NewRotPoint(0.0f, FQuat::Identity, FQuat::Identity, FQuat::Identity, CIM_CurveAuto);
+	FInterpCurvePoint<FVector> NewScalePoint(0.0f, FVector(1.0f), FVector::ZeroVector, FVector::ZeroVector, CIM_CurveAuto);
+
+	SplineInfo.Points.Insert(NewPoint, SelectedSegmentIndex + 1);
+	SplineRotInfo.Points.Insert(NewRotPoint, SelectedSegmentIndex + 1);
+	SplineScaleInfo.Points.Insert(NewScalePoint, SelectedSegmentIndex + 1);
 
 	// Set selection to 'next' key
 	ChangeSelectionState(SelectedSegmentIndex + 1, false);
 	SelectedSegmentIndex = INDEX_NONE;
 
-	// Update Input value for all keys
-	SplineComp->RefreshSplineInputs();
-
-	// Restore old closed loop state
-	SplineComp->SetClosedLoop(bWasLoop);
-
 	NotifyComponentModified();
+
+	CachedRotation = SplineComp->GetQuaternionAtSplinePoint(LastKeyIndexSelected, ESplineCoordinateSpace::World);
 }
 
 
@@ -737,10 +798,6 @@ void FSplineComponentVisualizer::OnDeleteKey()
 		Owner->Modify();
 	}
 
-	// Remember whether it's a closed loop for later, and turn off the loop while we manipulate the points
-	const bool bWasLoop = SplineComp->IsClosedLoop();
-	SplineComp->SetClosedLoop(false);
-
 	// Get a sorted list of all the selected indices, highest to lowest
 	TArray<int32> SelectedKeysSorted;
 	for (int32 SelectedKeyIndex : SelectedKeys)
@@ -750,24 +807,26 @@ void FSplineComponentVisualizer::OnDeleteKey()
 	SelectedKeysSorted.Sort([](int32 A, int32 B) { return A > B; });
 
 	// Delete selected keys from list, highest index first
-	TArray<FInterpCurvePoint<FVector>>& Points = SplineComp->SplineInfo.Points;
+	FInterpCurveVector& SplineInfo = SplineComp->SplineInfo;
+	FInterpCurveQuat& SplineRotInfo = SplineComp->SplineRotInfo;
+	FInterpCurveVector& SplineScaleInfo = SplineComp->SplineScaleInfo;
+
 	for (int32 SelectedKeyIndex : SelectedKeysSorted)
 	{
-		Points.RemoveAt(SelectedKeyIndex);
+		SplineInfo.Points.RemoveAt(SelectedKeyIndex);
+		SplineRotInfo.Points.RemoveAt(SelectedKeyIndex);
+		SplineScaleInfo.Points.RemoveAt(SelectedKeyIndex);
 	}
 
-	// Unselect everything
-	ChangeSelectionState(INDEX_NONE, false);
+	// Select first key
+	ChangeSelectionState(0, false);
 	SelectedSegmentIndex = INDEX_NONE;
 	SelectedTangentHandle = INDEX_NONE;
 	SelectedTangentHandleType = ESelectedTangentHandle::None;
 
-	SplineComp->RefreshSplineInputs(); // update input value for each key
-
-	// Restore old closed loop state
-	SplineComp->SetClosedLoop(bWasLoop);
-
 	NotifyComponentModified();
+
+	CachedRotation = SplineComp->GetQuaternionAtSplinePoint(LastKeyIndexSelected, ESplineCoordinateSpace::World);
 }
 
 
@@ -803,9 +862,6 @@ void FSplineComponentVisualizer::OnResetToAutomaticTangent(EInterpCurveMode Mode
 			Owner->Modify();
 		}
 
-		const bool bWasLoop = SplineComp->IsClosedLoop();
-		SplineComp->SetClosedLoop(false);
-
 		for (int32 SelectedKeyIndex : SelectedKeys)
 		{
 			auto& Point = SplineComp->SplineInfo.Points[SelectedKeyIndex];
@@ -815,8 +871,9 @@ void FSplineComponentVisualizer::OnResetToAutomaticTangent(EInterpCurveMode Mode
 			}
 		}
 
-		SplineComp->SetClosedLoop(bWasLoop);
 		NotifyComponentModified();
+
+		CachedRotation = SplineComp->GetQuaternionAtSplinePoint(LastKeyIndexSelected, ESplineCoordinateSpace::World);
 	}
 }
 
@@ -853,16 +910,14 @@ void FSplineComponentVisualizer::OnSetKeyType(EInterpCurveMode Mode)
 			Owner->Modify();
 		}
 
-		const bool bWasLoop = SplineComp->IsClosedLoop();
-		SplineComp->SetClosedLoop(false);
-
 		for (int32 SelectedKeyIndex : SelectedKeys)
 		{
 			SplineComp->SplineInfo.Points[SelectedKeyIndex].InterpMode = Mode;
 		}
 
-		SplineComp->SetClosedLoop(bWasLoop);
 		NotifyComponentModified();
+
+		CachedRotation = SplineComp->GetQuaternionAtSplinePoint(LastKeyIndexSelected, ESplineCoordinateSpace::World);
 	}
 }
 
@@ -888,11 +943,66 @@ bool FSplineComponentVisualizer::IsKeyTypeSet(EInterpCurveMode Mode) const
 }
 
 
+void FSplineComponentVisualizer::OnSetVisualizeRollAndScale()
+{
+	USplineComponent* SplineComp = GetEditedSplineComponent();
+	check(SplineComp != nullptr);
+	SplineComp->bShouldVisualizeScale = !SplineComp->bShouldVisualizeScale;
+	GEditor->RedrawLevelEditingViewports(true);
+}
+
+
+bool FSplineComponentVisualizer::IsVisualizingRollAndScale() const
+{
+	USplineComponent* SplineComp = GetEditedSplineComponent();
+	check(SplineComp != nullptr);
+	return SplineComp->bShouldVisualizeScale;
+}
+
+
+void FSplineComponentVisualizer::OnResetToDefault()
+{
+	USplineComponent* SplineComp = GetEditedSplineComponent();
+	check(SplineComp != nullptr);
+
+	const FScopedTransaction Transaction(LOCTEXT("ResetToDefault", "Reset to Default"));
+
+	SplineComp->Modify();
+	if (SplineOwningActor.IsValid())
+	{
+		SplineOwningActor.Get()->Modify();
+	}
+
+	SplineComp->bSplineHasBeenEdited = false;
+
+	// Select first key
+	ChangeSelectionState(0, false);
+	SelectedSegmentIndex = INDEX_NONE;
+	SelectedTangentHandle = INDEX_NONE;
+	SelectedTangentHandleType = ESelectedTangentHandle::None;
+
+	if (SplineOwningActor.IsValid())
+	{
+		SplineOwningActor.Get()->PostEditMove(false);
+	}
+
+	GEditor->RedrawLevelEditingViewports(true);
+}
+
+
+bool FSplineComponentVisualizer::CanResetToDefault() const
+{
+	USplineComponent* SplineComp = GetEditedSplineComponent();
+	check(SplineComp != nullptr);
+	return SplineComp->bSplineHasBeenEdited;
+}
+
+
 TSharedPtr<SWidget> FSplineComponentVisualizer::GenerateContextMenu() const
 {
 	FMenuBuilder MenuBuilder(true, SplineComponentVisualizerActions);
 	{
-		MenuBuilder.BeginSection("SplineKeyEdit");
+		MenuBuilder.BeginSection("SplinePointEdit", LOCTEXT("SplinePoint", "Spline Point"));
 		{
 			if (SelectedSegmentIndex != INDEX_NONE)
 			{
@@ -928,6 +1038,18 @@ TSharedPtr<SWidget> FSplineComponentVisualizer::GenerateContextMenu() const
 			}
 		}
 		MenuBuilder.EndSection();
+
+		MenuBuilder.BeginSection("Spline", LOCTEXT("Spline", "Spline"));
+		{
+			MenuBuilder.AddMenuEntry(FSplineComponentVisualizerCommands::Get().ResetToDefault);
+		}
+		MenuBuilder.EndSection();
+
+		MenuBuilder.BeginSection("Visualization", LOCTEXT("Visualization", "Visualization"));
+		{
+			MenuBuilder.AddMenuEntry(FSplineComponentVisualizerCommands::Get().VisualizeRollAndScale);
+		}
+		MenuBuilder.EndSection();
 	}
 
 	TSharedPtr<SWidget> MenuWidget = MenuBuilder.MakeWidget();
@@ -954,6 +1076,7 @@ void FSplineComponentVisualizer::NotifyComponentModified()
 {
 	USplineComponent* SplineComp = GetEditedSplineComponent();
 
+	SplineComp->bSplineHasBeenEdited = true;
 	SplineComp->UpdateSpline();
 
 	// Notify that the spline info has been modified
@@ -964,7 +1087,7 @@ void FSplineComponentVisualizer::NotifyComponentModified()
 	// Notify of change so any CS is re-run
 	if (SplineOwningActor.IsValid())
 	{
-		SplineOwningActor.Get()->PostEditMove(true);
+		SplineOwningActor.Get()->PostEditMove(false);
 	}
 
 	GEditor->RedrawLevelEditingViewports(true);
diff --git a/Engine/Source/Editor/ComponentVisualizers/Private/SplineComponentVisualizer.h b/Engine/Source/Editor/ComponentVisualizers/Private/SplineComponentVisualizer.h
--- a/Engine/Source/Editor/ComponentVisualizers/Private/SplineComponentVisualizer.h
+++ b/Engine/Source/Editor/ComponentVisualizers/Private/SplineComponentVisualizer.h
@@ -100,6 +100,12 @@ private:
 	void OnSetKeyType(EInterpCurveMode Mode);
 	bool IsKeyTypeSet(EInterpCurveMode Mode) const;
 
+	void OnSetVisualizeRollAndScale();
+	bool IsVisualizingRollAndScale() const;
+
+	void OnResetToDefault();
+	bool CanResetToDefault() const;
+
 	/** Generate the submenu containing the available point types */
 	void GenerateSplinePointTypeSubMenu(FMenuBuilder& MenuBuilder) const;
 
@@ -143,6 +149,9 @@ private:
 	/** Position on spline we have selected */
 	FVector SelectedSplinePosition;
 
+	/** Cached rotation for this point */
+	FQuat CachedRotation;
+
 	/** Whether we currently allow duplication when dragging */
 	bool bAllowDuplication;
 };
diff --git a/Engine/Source/Editor/UnrealEd/Private/ComponentVisualizer.cpp b/Engine/Source/Editor/UnrealEd/Private/ComponentVisualizer.cpp
--- a/Engine/Source/Editor/UnrealEd/Private/ComponentVisualizer.cpp
+++ b/Engine/Source/Editor/UnrealEd/Private/ComponentVisualizer.cpp
@@ -8,7 +8,6 @@ IMPLEMENT_HIT_PROXY(HComponentVisProxy, HHitProxy);
 
 FName FComponentVisualizer::GetComponentPropertyName(const UActorComponent* Component)
 {
-	FName PropertyName = NAME_None;
 	if (Component)
 	{
 		const AActor* CompOwner = Component->GetOwner();
@@ -23,13 +22,15 @@ FName FComponentVisualizer::GetComponentPropertyName(const UActorComponent* Comp
 				UObject* Object = ObjectProp->GetObjectPropertyValue(ObjectProp->ContainerPtrToValuePtr<void>(CompOwner));
 				if (Object == Component)
 				{
-					// It does! Remember this name
-					PropertyName = ObjectProp->GetFName();
+					// It does! Return this name
+					return ObjectProp->GetFName();
 				}
 			}
 		}
 	}
-	return PropertyName;
+
+	// Didn't find actor property referencing this component
+	return NAME_None;
 }
 
 UActorComponent* FComponentVisualizer::GetComponentFromPropertyName(const AActor* CompOwner, FName PropertyName)
diff --git a/Engine/Source/Runtime/Core/Public/Math/InterpCurve.h b/Engine/Source/Runtime/Core/Public/Math/InterpCurve.h
--- a/Engine/Source/Runtime/Core/Public/Math/InterpCurve.h
+++ b/Engine/Source/Runtime/Core/Public/Math/InterpCurve.h
@@ -17,6 +17,12 @@ public:
 	/** Holds the collection of interpolation points. */
 	TArray<FInterpCurvePoint<T>> Points;
 
+	/** Specify whether the curve is looped or not */
+	bool bIsLooped;
+
+	/** Specify the offset from the last point's input key corresponding to the loop point */
+	float LoopKeyOffset;
+
 public:
 
 	/** Default constructor (no initialization). */
@@ -47,21 +53,27 @@ public:
 	/** Clears all keypoints from InterpCurve. */
 	void Reset();
 
+	/** Set loop key for curve */
+	void SetLoopKey( float InLoopKey );
+
+	/** Clear loop key for curve */
+	void ClearLoopKey();
+
 	/** 
 	 *	Evaluate the output for an arbitary input value. 
 	 *	For inputs outside the range of the keys, the first/last key value is assumed.
 	 */
-	T Eval( const float InVal, const T& Default ) const;
+	T Eval( const float InVal, const T& Default = T(ForceInit) ) const;
 
 	/** 
 	 *	Evaluate the derivative at a point on the curve.
 	 */
-	T EvalDerivative( const float InVal, const T& Default ) const;
+	T EvalDerivative( const float InVal, const T& Default = T(ForceInit) ) const;
 
 	/** 
 	 *	Evaluate the second derivative at a point on the curve.
 	 */
-	T EvalSecondDerivative( const float InVal, const T& Default ) const;
+	T EvalSecondDerivative( const float InVal, const T& Default = T(ForceInit) ) const;
 
 	/** 
 	 * Find the nearest point on spline to the given point.
@@ -84,7 +96,7 @@ public:
 	void AutoSetTangents(float Tension = 0.f, bool bStationaryEndpoints = true);
 
 	/** Calculate the min/max out value that can be returned by this InterpCurve. */
-	void CalcBounds(T& OutMin, T& OutMax, const T& Default) const;
+	void CalcBounds(T& OutMin, T& OutMax, const T& Default = T(ForceInit)) const;
 
 public:
 
@@ -102,12 +114,12 @@ public:
 		//   as inline struct properties in UnClass.cpp!
 
 		Ar << Curve.Points;
+		Ar << Curve.bIsLooped;
+		Ar << Curve.LoopKeyOffset;
 
 		return Ar;
 	}
 
-private:
-
 	/**
 	 * Finds the lower index of the two points whose input values bound the supplied input value.
 	 */
@@ -157,10 +169,36 @@ FORCEINLINE void FInterpCurve<T>::Reset()
 }
 
 
+template <class T>
+FORCEINLINE void FInterpCurve<T>::SetLoopKey(float InLoopKey)
+{
+	const float LastInKey = Points.Last().InVal;
+	if (InLoopKey > LastInKey)
+	{
+		// Calculate loop key offset from the input key of the final point
+		bIsLooped = true;
+		LoopKeyOffset = InLoopKey - LastInKey;
+	}
+	else
+	{
+		// Specified a loop key lower than the final point; turn off looping.
+		bIsLooped = false;
+	}
+}
+
+
+template <class T>
+FORCEINLINE void FInterpCurve<T>::ClearLoopKey()
+{
+	bIsLooped = false;
+}
+
+
 template< class T >
 int32 FInterpCurve<T>::GetPointIndexForInputValue(const float InValue) const
 {
 	const int32 NumPoints = Points.Num();
+	const int32 LastPoint = NumPoints - 1;
 
 	check(NumPoints > 0);
 
@@ -169,9 +207,9 @@ int32 FInterpCurve<T>::GetPointIndexForInputValue(const float InValue) const
 		return -1;
 	}
 
-	if (InValue > Points[NumPoints - 1].InVal)
+	if (InValue >= Points[LastPoint].InVal)
 	{
-		return NumPoints;
+		return LastPoint;
 	}
 
 	int32 MinIndex = 0;
@@ -195,13 +233,14 @@ int32 FInterpCurve<T>::GetPointIndexForInputValue(const float InValue) const
 }
 
 
-template< class T > 
-T FInterpCurve<T>::Eval( const float InVal, const T& Default ) const
+template< class T >
+T FInterpCurve<T>::Eval(const float InVal, const T& Default) const
 {
 	const int32 NumPoints = Points.Num();
+	const int32 LastPoint = NumPoints - 1;
 
 	// If no point in curve, return the Default value we passed in.
-	if( NumPoints == 0 )
+	if (NumPoints == 0)
 	{
 		return Default;
 	}
@@ -216,20 +255,33 @@ T FInterpCurve<T>::Eval( const float InVal, const T& Default ) const
 	}
 
 	// If on or beyond the last point, return its value.
-	if (Index >= NumPoints - 1)
+	if (Index == LastPoint)
 	{
-		return Points[NumPoints - 1].OutVal;
+		if (!bIsLooped)
+		{
+			return Points[LastPoint].OutVal;
+		}
+		else if (InVal >= Points[LastPoint].InVal + LoopKeyOffset)
+		{
+			// Looped spline: last point is the same as the first point
+			return Points[0].OutVal;
+		}
 	}
 
 	// Somewhere within curve range - interpolate.
-	check(Index >= 0 && Index < NumPoints - 1);
+	check(Index >= 0 && ((bIsLooped && Index < NumPoints) || (!bIsLooped && Index < LastPoint)));
+	const bool bLoopSegment = (bIsLooped && Index == LastPoint);
+	const int32 NextIndex = bLoopSegment ? 0 : Index + 1;
+
 	const auto& PrevPoint = Points[Index];
-	const auto& NextPoint = Points[Index + 1];
-	const float Diff = NextPoint.InVal - PrevPoint.InVal;
+	const auto& NextPoint = Points[NextIndex];
+
+	const float Diff = bLoopSegment ? LoopKeyOffset : NextPoint.InVal - PrevPoint.InVal;
 
 	if (Diff > 0.0f && PrevPoint.InterpMode != CIM_Constant)
 	{
 		const float Alpha = (InVal - PrevPoint.InVal) / Diff;
+		check(Alpha >= 0.0f && Alpha < 1.0f);
 
 		if (PrevPoint.InterpMode == CIM_Linear)
 		{
@@ -247,10 +299,11 @@ T FInterpCurve<T>::Eval( const float InVal, const T& Default ) const
 }
 
 
-template< class T > 
-T FInterpCurve<T>::EvalDerivative( const float InVal, const T& Default ) const
+template< class T >
+T FInterpCurve<T>::EvalDerivative(const float InVal, const T& Default) const
 {
 	const int32 NumPoints = Points.Num();
+	const int32 LastPoint = NumPoints - 1;
 
 	// If no point in curve, return the Default value we passed in.
 	if (NumPoints == 0)
@@ -268,16 +321,28 @@ T FInterpCurve<T>::EvalDerivative( const float InVal, const T& Default ) const
 	}
 
 	// If on or beyond the last point, return its tangent value.
-	if (Index >= NumPoints - 1)
+	if (Index == LastPoint)
 	{
-		return Points[NumPoints - 1].ArriveTangent;
+		if (!bIsLooped)
+		{
+			return Points[LastPoint].ArriveTangent;
+		}
+		else if (InVal >= Points[LastPoint].InVal + LoopKeyOffset)
+		{
+			// Looped spline: last point is the same as the first point
+			return Points[0].ArriveTangent;
+		}
 	}
 
 	// Somewhere within curve range - interpolate.
-	check(Index >= 0 && Index < NumPoints - 1);
+	check(Index >= 0 && ((bIsLooped && Index < NumPoints) || (!bIsLooped && Index < LastPoint)));
+	const bool bLoopSegment = (bIsLooped && Index == LastPoint);
+	const int32 NextIndex = bLoopSegment ? 0 : Index + 1;
+
 	const auto& PrevPoint = Points[Index];
-	const auto& NextPoint = Points[Index + 1];
-	const float Diff = NextPoint.InVal - PrevPoint.InVal;
+	const auto& NextPoint = Points[NextIndex];
+
+	const float Diff = bLoopSegment ? LoopKeyOffset : NextPoint.InVal - PrevPoint.InVal;
 
 	if (Diff > 0.0f && PrevPoint.InterpMode != CIM_Constant)
 	{
@@ -288,6 +353,7 @@ T FInterpCurve<T>::EvalDerivative( const float InVal, const T& Default ) const
 		else
 		{
 			const float Alpha = (InVal - PrevPoint.InVal) / Diff;
+			check(Alpha >= 0.0f && Alpha < 1.0f);
 
 			return FMath::CubicInterpDerivative(PrevPoint.OutVal, PrevPoint.LeaveTangent * Diff, NextPoint.OutVal, NextPoint.ArriveTangent * Diff, Alpha) / Diff;
 		}
@@ -295,15 +361,16 @@ T FInterpCurve<T>::EvalDerivative( const float InVal, const T& Default ) const
 	else
 	{
 		// Derivative of a constant is zero
-		return T();
+		return T(ForceInit);
 	}
 }
 
 
-template< class T > 
-T FInterpCurve<T>::EvalSecondDerivative( const float InVal, const T& Default ) const
+template< class T >
+T FInterpCurve<T>::EvalSecondDerivative(const float InVal, const T& Default) const
 {
 	const int32 NumPoints = Points.Num();
+	const int32 LastPoint = NumPoints - 1;
 
 	// If no point in curve, return the Default value we passed in.
 	if (NumPoints == 0)
@@ -317,31 +384,39 @@ T FInterpCurve<T>::EvalSecondDerivative( const float InVal, const T& Default ) c
 	// If before the first point, return 0
 	if (Index == -1)
 	{
-		return T();
+		return T(ForceInit);
 	}
 
 	// If on or beyond the last point, return 0
-	if (Index >= NumPoints - 1)
+	if (Index == LastPoint)
 	{
-		return T();
+		if (!bIsLooped || (InVal >= Points[LastPoint].InVal + LoopKeyOffset))
+		{
+			return T(ForceInit);
+		}
 	}
 
 	// Somewhere within curve range - interpolate.
-	check(Index >= 0 && Index < NumPoints - 1);
+	check(Index >= 0 && ((bIsLooped && Index < NumPoints) || (!bIsLooped && Index < LastPoint)));
+	const bool bLoopSegment = (bIsLooped && Index == LastPoint);
+	const int32 NextIndex = bLoopSegment ? 0 : Index + 1;
+
 	const auto& PrevPoint = Points[Index];
-	const auto& NextPoint = Points[Index + 1];
-	const float Diff = NextPoint.InVal - PrevPoint.InVal;
+	const auto& NextPoint = Points[NextIndex];
+
+	const float Diff = bLoopSegment ? LoopKeyOffset : NextPoint.InVal - PrevPoint.InVal;
 
 	if (Diff > 0.0f && PrevPoint.InterpMode != CIM_Constant)
 	{
 		if (PrevPoint.InterpMode == CIM_Linear)
 		{
 			// No change in tangent, return 0.
-			return T();
+			return T(ForceInit);
 		}
 		else
 		{
 			const float Alpha = (InVal - PrevPoint.InVal) / Diff;
+			check(Alpha >= 0.0f && Alpha < 1.0f);
 
 			return FMath::CubicInterpSecondDerivative(PrevPoint.OutVal, PrevPoint.LeaveTangent * Diff, NextPoint.OutVal, NextPoint.ArriveTangent * Diff, Alpha) / (Diff * Diff);
 		}
@@ -349,24 +424,26 @@ T FInterpCurve<T>::EvalSecondDerivative( const float InVal, const T& Default ) c
 	else
 	{
 		// Second derivative of a constant is zero
-		return T();
+		return T(ForceInit);
 	}
 }
 
 
-template< class T > 
-float FInterpCurve<T>::InaccurateFindNearest( const T &PointInSpace, float& OutDistanceSq ) const
+template< class T >
+float FInterpCurve<T>::InaccurateFindNearest(const T &PointInSpace, float& OutDistanceSq) const
 {
 	const int32 NumPoints = Points.Num();
-	if(NumPoints > 1)
+	const int32 NumSegments = bIsLooped ? NumPoints : NumPoints - 1;
+
+	if (NumPoints > 1)
 	{
 		float BestDistanceSq;
 		float BestResult = InaccurateFindNearestOnSegment(PointInSpace, 0, BestDistanceSq);
-		for(int32 segment = 1; segment < NumPoints - 1; ++segment)
+		for (int32 Segment = 1; Segment < NumSegments; ++Segment)
 		{
 			float LocalDistanceSq;
-			float LocalResult = InaccurateFindNearestOnSegment(PointInSpace, segment, LocalDistanceSq);
-			if(LocalDistanceSq < BestDistanceSq)
+			float LocalResult = InaccurateFindNearestOnSegment(PointInSpace, Segment, LocalDistanceSq);
+			if (LocalDistanceSq < BestDistanceSq)
 			{
 				BestDistanceSq = LocalDistanceSq;
 				BestResult = LocalResult;
@@ -375,42 +452,49 @@ float FInterpCurve<T>::InaccurateFindNearest( const T &PointInSpace, float& OutD
 		OutDistanceSq = BestDistanceSq;
 		return BestResult;
 	}
-	if( 1 == NumPoints )
+
+	if (NumPoints == 1)
 	{
 		OutDistanceSq = (PointInSpace - Points[0].OutVal).SizeSquared();
 		return Points[0].InVal;
 	}
+
 	return 0.0f;
 }
 
 
-template< class T > 
-float FInterpCurve<T>::InaccurateFindNearestOnSegment( const T &PointInSpace, int32 PtIdx, float& OutSquaredDistance ) const
+template< class T >
+float FInterpCurve<T>::InaccurateFindNearestOnSegment(const T& PointInSpace, int32 PtIdx, float& OutSquaredDistance) const
 {
-	if( CIM_Constant == Points[PtIdx].InterpMode )
+	const int32 NumPoints = Points.Num();
+	const int32 LastPoint = NumPoints - 1;
+	const int32 NextPtIdx = (bIsLooped && PtIdx == LastPoint) ? 0 : PtIdx + 1;
+	check(PtIdx >= 0 && ((bIsLooped && PtIdx < NumPoints) || (!bIsLooped && PtIdx < LastPoint)));
+
+	if (CIM_Constant == Points[PtIdx].InterpMode)
 	{
 		const float Distance1 = (Points[PtIdx].OutVal - PointInSpace).SizeSquared();
-		const float Distance2 = (Points[PtIdx+1].OutVal - PointInSpace).SizeSquared();
-		if(Distance1 < Distance2)
+		const float Distance2 = (Points[NextPtIdx].OutVal - PointInSpace).SizeSquared();
+		if (Distance1 < Distance2)
 		{
 			OutSquaredDistance = Distance1;
 			return Points[PtIdx].InVal;
 		}
 		OutSquaredDistance = Distance2;
-		return Points[PtIdx+1].InVal;
+		return Points[NextPtIdx].InVal;
 	}
 
-	const float Diff = Points[PtIdx+1].InVal - Points[PtIdx].InVal;
-	if(CIM_Linear == Points[PtIdx].InterpMode )
+	const float Diff = Points[NextPtIdx].InVal - Points[PtIdx].InVal;
+	if (CIM_Linear == Points[PtIdx].InterpMode)
 	{
 		// like in function: FMath::ClosestPointOnLine
-		const float A = (Points[PtIdx].OutVal-PointInSpace) | (Points[PtIdx+1].OutVal - Points[PtIdx].OutVal);
-		const float B = (Points[PtIdx+1].OutVal - Points[PtIdx].OutVal).SizeSquared();
-		const float V = FMath::Clamp(-A/B, 0.f, 1.f);
-		OutSquaredDistance = (FMath::Lerp( Points[PtIdx].OutVal, Points[PtIdx+1].OutVal, V ) - PointInSpace).SizeSquared();
+		const float A = (Points[PtIdx].OutVal - PointInSpace) | (Points[NextPtIdx].OutVal - Points[PtIdx].OutVal);
+		const float B = (Points[NextPtIdx].OutVal - Points[PtIdx].OutVal).SizeSquared();
+		const float V = FMath::Clamp(-A / B, 0.f, 1.f);
+		OutSquaredDistance = (FMath::Lerp(Points[PtIdx].OutVal, Points[NextPtIdx].OutVal, V) - PointInSpace).SizeSquared();
 		return V * Diff + Points[PtIdx].InVal;
 	}
-		
+
 	{
 		const int32 PointsChecked = 3;
 		const int32 IterationNum = 3;
@@ -418,45 +502,43 @@ float FInterpCurve<T>::InaccurateFindNearestOnSegment( const T &PointInSpace, in
 
 		// Newton's methods is repeated 3 times, starting with t = 0, 0.5, 1.
 		float ValuesT[PointsChecked];
-		ValuesT[0] = 0.0f; 
-		ValuesT[1] = 0.5f; 
+		ValuesT[0] = 0.0f;
+		ValuesT[1] = 0.5f;
 		ValuesT[2] = 1.0f;
 
 		T InitialPoints[PointsChecked];
 		InitialPoints[0] = Points[PtIdx].OutVal;
-		InitialPoints[1] = 
-			FMath::CubicInterp( Points[PtIdx].OutVal, Points[PtIdx].LeaveTangent * Diff, Points[PtIdx+1].OutVal, Points[PtIdx+1].ArriveTangent * Diff, ValuesT[1] );
-		InitialPoints[2] = Points[PtIdx+1].OutVal;
+		InitialPoints[1] = FMath::CubicInterp(Points[PtIdx].OutVal, Points[PtIdx].LeaveTangent * Diff, Points[NextPtIdx].OutVal, Points[NextPtIdx].ArriveTangent * Diff, ValuesT[1]);
+		InitialPoints[2] = Points[NextPtIdx].OutVal;
 
 		float DistancesSq[PointsChecked];
 
-		for(int32 point = 0; point < PointsChecked; ++point)
+		for (int32 point = 0; point < PointsChecked; ++point)
 		{
 			//Algorithm explanation: http://permalink.gmane.org/gmane.games.devel.sweng/8285
 			T FoundPoint = InitialPoints[point];
 			float LastMove = 1.0f;
-			for(int32 iter = 0; iter < IterationNum; ++iter)
-			{	
-				const T LastBestTangent = FMath::CubicInterpDerivative( Points[PtIdx].OutVal, Points[PtIdx].LeaveTangent * Diff, Points[PtIdx+1].OutVal, Points[PtIdx+1].ArriveTangent * Diff, ValuesT[point]);
+			for (int32 iter = 0; iter < IterationNum; ++iter)
+			{
+				const T LastBestTangent = FMath::CubicInterpDerivative(Points[PtIdx].OutVal, Points[PtIdx].LeaveTangent * Diff, Points[NextPtIdx].OutVal, Points[NextPtIdx].ArriveTangent * Diff, ValuesT[point]);
 				const T Delta = (PointInSpace - FoundPoint);
-				float Move = (LastBestTangent | Delta)/LastBestTangent.SizeSquared();
+				float Move = (LastBestTangent | Delta) / LastBestTangent.SizeSquared();
 				Move = FMath::Clamp(Move, -LastMove*Scale, LastMove*Scale);
 				ValuesT[point] += Move;
 				ValuesT[point] = FMath::Clamp(ValuesT[point], 0.0f, 1.0f);
 				LastMove = FMath::Abs(Move);
-				FoundPoint = 
-					FMath::CubicInterp( Points[PtIdx].OutVal, Points[PtIdx].LeaveTangent * Diff, Points[PtIdx+1].OutVal, Points[PtIdx+1].ArriveTangent * Diff, ValuesT[point] );
+				FoundPoint = FMath::CubicInterp(Points[PtIdx].OutVal, Points[PtIdx].LeaveTangent * Diff, Points[NextPtIdx].OutVal, Points[NextPtIdx].ArriveTangent * Diff, ValuesT[point]);
 			}
-			DistancesSq[point] = (FoundPoint-PointInSpace).SizeSquared();
+			DistancesSq[point] = (FoundPoint - PointInSpace).SizeSquared();
 			ValuesT[point] = ValuesT[point] * Diff + Points[PtIdx].InVal;
 		}
 
-		if(DistancesSq[0] <= DistancesSq[1] && DistancesSq[0] <= DistancesSq[2])
+		if (DistancesSq[0] <= DistancesSq[1] && DistancesSq[0] <= DistancesSq[2])
 		{
 			OutSquaredDistance = DistancesSq[0];
 			return ValuesT[0];
 		}
-		if(DistancesSq[1] <= DistancesSq[2])
+		if (DistancesSq[1] <= DistancesSq[2])
 		{
 			OutSquaredDistance = DistancesSq[1];
 			return ValuesT[1];
@@ -467,91 +549,85 @@ float FInterpCurve<T>::InaccurateFindNearestOnSegment( const T &PointInSpace, in
 }
 
 
-template< class T > 
+template< class T >
 void FInterpCurve<T>::AutoSetTangents(float Tension, bool bStationaryEndpoints)
 {
+	const int32 NumPoints = Points.Num();
+	const int32 LastPoint = NumPoints - 1;
+
 	// Iterate over all points in this InterpCurve
-	for(int32 PointIndex=0; PointIndex<Points.Num(); PointIndex++)
+	for (int32 PointIndex = 0; PointIndex < NumPoints; PointIndex++)
 	{
-		auto& ThisPoint = Points[PointIndex];
+		const int32 PrevIndex = (PointIndex == 0) ? (bIsLooped ? LastPoint : 0) : PointIndex - 1;
+		const int32 NextIndex = (PointIndex == LastPoint) ? (bIsLooped ? 0 : LastPoint) : PointIndex + 1;
 
-		T ArriveTangent = ThisPoint.ArriveTangent;
-		T LeaveTangent = ThisPoint.LeaveTangent;
+		auto& ThisPoint = Points[PointIndex];
+		const auto& PrevPoint = Points[PrevIndex];
+		const auto& NextPoint = Points[NextIndex];
 
-		if(PointIndex == 0 && bStationaryEndpoints)
+		if (ThisPoint.InterpMode == CIM_CurveAuto || ThisPoint.InterpMode == CIM_CurveAutoClamped)
 		{
-			if(PointIndex < Points.Num()-1) // Start point
+			if (bStationaryEndpoints && (PointIndex == 0 || (PointIndex == LastPoint && !bIsLooped)))
 			{
-				// If first section is not a curve, or is a curve and first point has manual tangent setting.
-				if (ThisPoint.InterpMode == CIM_CurveAuto || ThisPoint.InterpMode == CIM_CurveAutoClamped)
-				{
-					FMemory::Memset( &LeaveTangent, 0, sizeof(T) );
-				}
+				// Start and end points get zero tangents if bStationaryEndpoints is true
+				ThisPoint.ArriveTangent = T(ForceInit);
+				ThisPoint.LeaveTangent = T(ForceInit);
 			}
-			else // Only point
+			else if (PrevPoint.IsCurveKey())
 			{
-				FMemory::Memset( &LeaveTangent, 0, sizeof(T) );
+				const bool bWantClamping = (ThisPoint.InterpMode == CIM_CurveAutoClamped);
+				T Tangent;
+
+				const float PrevTime = (bIsLooped && PointIndex == 0) ? ThisPoint.InVal - LoopKeyOffset : PrevPoint.InVal;
+				const float NextTime = (bIsLooped && PointIndex == LastPoint) ? ThisPoint.InVal + LoopKeyOffset : NextPoint.InVal;
+
+				ComputeCurveTangent(
+					PrevTime,			// Previous time
+					PrevPoint.OutVal,	// Previous point
+					ThisPoint.InVal,	// Current time
+					ThisPoint.OutVal,	// Current point
+					NextTime,			// Next time
+					NextPoint.OutVal,	// Next point
+					Tension,			// Tension
+					bWantClamping,		// Want clamping?
+					Tangent);			// Out
+
+				ThisPoint.ArriveTangent = Tangent;
+				ThisPoint.LeaveTangent = Tangent;
 			}
-		}
-		else
-		{
-			if(PointIndex < Points.Num()-1 || !bStationaryEndpoints) // Inner point
+			else
 			{
-				const auto& PrevPoint = Points[FMath::Max(0, PointIndex - 1)];
-				const auto& NextPoint = Points[FMath::Min(Points.Num() - 1, PointIndex + 1)];
-
-				if( ThisPoint.InterpMode == CIM_CurveAuto || ThisPoint.InterpMode == CIM_CurveAutoClamped )
-				{
-					if( PrevPoint.IsCurveKey() && ThisPoint.IsCurveKey() )
-					{
-						const bool bWantClamping = ( ThisPoint.InterpMode == CIM_CurveAutoClamped );
-
-						ComputeCurveTangent(
-							PrevPoint.InVal,	// Previous time
-							PrevPoint.OutVal,	// Previous point
-							ThisPoint.InVal,	// Current time
-							ThisPoint.OutVal,	// Current point
-							NextPoint.InVal,	// Next time
-							NextPoint.OutVal,	// Next point
-							Tension,			// Tension
-							bWantClamping,		// Want clamping?
-							ArriveTangent );	// Out
-
-						// In 'auto' mode, arrive and leave tangents are always the same
-						LeaveTangent = ArriveTangent;
-					}
-					else if( PrevPoint.InterpMode == CIM_Constant || ThisPoint.InterpMode == CIM_Constant )
-					{
-						FMemory::Memset( &ArriveTangent, 0, sizeof(T) );
-						FMemory::Memset( &LeaveTangent, 0, sizeof(T) );
-					}
-				}
-			}
-			else // End point
-			{
-				// If last section is not a curve, or is a curve and final point has manual tangent setting.
-				if (ThisPoint.InterpMode == CIM_CurveAuto || ThisPoint.InterpMode == CIM_CurveAutoClamped)
-				{
-					FMemory::Memset( &ArriveTangent, 0, sizeof(T) );
-				}
+				// Following on from a line or constant; set curve tangent equal to that so there are no discontinuities
+				ThisPoint.ArriveTangent = PrevPoint.ArriveTangent;
+				ThisPoint.LeaveTangent = PrevPoint.LeaveTangent;
 			}
 		}
-
-		ThisPoint.ArriveTangent = ArriveTangent;
-		ThisPoint.LeaveTangent = LeaveTangent;
+		else if (ThisPoint.InterpMode == CIM_Linear)
+		{
+			T Tangent = NextPoint.OutVal - ThisPoint.OutVal;
+			ThisPoint.ArriveTangent = Tangent;
+			ThisPoint.LeaveTangent = Tangent;
+		}
+		else if (ThisPoint.InterpMode == CIM_Constant)
+		{
+			ThisPoint.ArriveTangent = T(ForceInit);
+			ThisPoint.LeaveTangent = T(ForceInit);
+		}
 	}
 }
 
 
-template< class T > 
+template< class T >
 void FInterpCurve<T>::CalcBounds(T& OutMin, T& OutMax, const T& Default) const
 {
-	if(Points.Num() == 0)
+	const int32 NumPoints = Points.Num();
+
+	if (NumPoints == 0)
 	{
 		OutMin = Default;
 		OutMax = Default;
 	}
-	else if(Points.Num() == 1)
+	else if (NumPoints == 1)
 	{
 		OutMin = Points[0].OutVal;
 		OutMax = Points[0].OutVal;
@@ -561,14 +637,18 @@ void FInterpCurve<T>::CalcBounds(T& OutMin, T& OutMax, const T& Default) const
 		OutMin = Points[0].OutVal;
 		OutMax = Points[0].OutVal;
 
-		for(int32 i=1; i<Points.Num(); i++)
+		const int32 NumSegments = bIsLooped ? NumPoints : NumPoints - 1;
+
+		for (int32 Index = 0; Index < NumSegments; Index++)
 		{
-			CurveFindIntervalBounds( Points[i-1], Points[i], OutMin, OutMax, 0.f );
+			const int32 NextIndex = (Index == NumPoints - 1) ? 0 : Index + 1;
+			CurveFindIntervalBounds(Points[Index], Points[NextIndex], OutMin, OutMax, 0.0f);
 		}
 	}
 }
 
 
+
 /* Common type definitions
  *****************************************************************************/
 
diff --git a/Engine/Source/Runtime/Core/Public/Math/InterpCurvePoint.h b/Engine/Source/Runtime/Core/Public/Math/InterpCurvePoint.h
--- a/Engine/Source/Runtime/Core/Public/Math/InterpCurvePoint.h
+++ b/Engine/Source/Runtime/Core/Public/Math/InterpCurvePoint.h
@@ -106,6 +106,26 @@ public:
 		Ar << Point.InterpMode;
 		return Ar;
 	}
+
+	/**
+	 * Compare equality of two Curve Points
+	 */
+	friend bool operator==( const FInterpCurvePoint& Point1, const FInterpCurvePoint& Point2 )
+	{
+		return (Point1.InVal == Point2.InVal &&
+				Point1.OutVal == Point2.OutVal &&
+				Point1.ArriveTangent == Point2.ArriveTangent &&
+				Point1.LeaveTangent == Point2.LeaveTangent &&
+				Point1.InterpMode == Point2.InterpMode);
+	}
+
+	/**
+	 * Compare inequality of two Curve Points
+	 */
+	friend bool operator!=(const FInterpCurvePoint& Point1, const FInterpCurvePoint& Point2)
+	{
+		return !(Point1 == Point2);
+	}
 };
 
 
diff --git a/Engine/Source/Runtime/Core/Public/Math/Quat.h b/Engine/Source/Runtime/Core/Public/Math/Quat.h
--- a/Engine/Source/Runtime/Core/Public/Math/Quat.h
+++ b/Engine/Source/Runtime/Core/Public/Math/Quat.h
@@ -195,13 +195,21 @@ public:
 	FORCEINLINE FQuat operator*=( const float Scale );
 
 	/**
-	 * Get the result of scaling this rotator.
+	 * Get the result of scaling this quaternion.
 	 *
 	 * @param Scale The scaling factor.
 	 * @return The result of scaling.
 	 */
 	FORCEINLINE FQuat operator*( const float Scale ) const;
 	
+	/**
+	 * Divide this quaternion by scale.
+	 *
+	 * @param Scale What to divide by.
+	 * @return a reference to this after scaling.
+	 */
+	FORCEINLINE FQuat operator/=( const float Scale );
+
 	/**
 	 * Divide this quaternion by scale.
 	 *
@@ -729,9 +737,24 @@ FORCEINLINE FQuat FQuat::operator*( const float Scale ) const
 }
 
 
-FORCEINLINE FQuat FQuat::operator/( const float Scale ) const
+FORCEINLINE FQuat FQuat::operator/=(const float Scale)
+{
+	const float Recip = 1.0f / Scale;
+	X *= Recip;
+	Y *= Recip;
+	Z *= Recip;
+	W *= Recip;
+
+	DiagnosticCheckNaN();
+
+	return *this;
+}
+
+
+FORCEINLINE FQuat FQuat::operator/(const float Scale) const
 {
-	return FQuat(X / Scale, Y / Scale, Z / Scale, W / Scale);
+	const float Recip = 1.0f / Scale;
+	return FQuat(X * Recip, Y * Recip, Z * Recip, W * Recip);
 }
 
 
diff --git a/Engine/Source/Runtime/Core/Public/UObject/ObjectVersion.h b/Engine/Source/Runtime/Core/Public/UObject/ObjectVersion.h
--- a/Engine/Source/Runtime/Core/Public/UObject/ObjectVersion.h
+++ b/Engine/Source/Runtime/Core/Public/UObject/ObjectVersion.h
@@ -498,6 +498,8 @@ enum EUnrealEngineObjectUE4Version
 	VER_UE4_BODYINSTANCE_BINARY_SERIALIZATION,
 	// Added fastcall data serialization directly in UFunction
 	VER_UE4_SERIALIZE_BLUEPRINT_EVENTGRAPH_FASTCALLS_IN_UFUNCTION,
+	// Changes to USplineComponent and FInterpCurve
+	VER_UE4_SPLINE_COMPONENT_WITH_ROTATION_AND_SCALE,
 
 	// -----<new versions can be added before this line>-------------------------------------------------
 	// - this needs to be the last line (see note below)
diff --git a/Engine/Source/Runtime/CoreUObject/Classes/Object.h b/Engine/Source/Runtime/CoreUObject/Classes/Object.h
--- a/Engine/Source/Runtime/CoreUObject/Classes/Object.h
+++ b/Engine/Source/Runtime/CoreUObject/Classes/Object.h
@@ -601,6 +601,11 @@ struct FInterpCurveFloat
 	UPROPERTY(EditAnywhere, BlueprintReadWrite, Category=InterpCurveFloat)
 	TArray<FInterpCurvePointFloat> Points;
 
+	UPROPERTY(EditAnywhere, BlueprintReadWrite, Category=InterpCurveFloat)
+	bool bIsLooped;
+
+	UPROPERTY(EditAnywhere, BlueprintReadWrite, Category=InterpCurveFloat)
+	float LoopKeyOffset;
 };
 
 
@@ -633,6 +638,11 @@ struct FInterpCurveVector2D
 	UPROPERTY(EditAnywhere, BlueprintReadWrite, Category=InterpCurveVector2D)
 	TArray<FInterpCurvePointVector2D> Points;
 
+	UPROPERTY(EditAnywhere, BlueprintReadWrite, Category=InterpCurveVector2D)
+	bool bIsLooped;
+
+	UPROPERTY(EditAnywhere, BlueprintReadWrite, Category=InterpCurveVector2D)
+	float LoopKeyOffset;
 };
 
 
@@ -665,6 +675,48 @@ struct FInterpCurveVector
 	UPROPERTY(EditAnywhere, BlueprintReadWrite, Category=InterpCurveVector)
 	TArray<FInterpCurvePointVector> Points;
 
+	UPROPERTY(EditAnywhere, BlueprintReadWrite, Category=InterpCurveVector)
+	bool bIsLooped;
+
+	UPROPERTY(EditAnywhere, BlueprintReadWrite, Category=InterpCurveVector)
+	float LoopKeyOffset;
+};
+
+
+
+USTRUCT(noexport)
+struct FInterpCurvePointQuat
+{
+	UPROPERTY(EditAnywhere, BlueprintReadWrite, Category=InterpCurvePointQuat)
+	float InVal;
+
+	UPROPERTY(EditAnywhere, BlueprintReadWrite, Category=InterpCurvePointQuat)
+	FQuat OutVal;
+
+	UPROPERTY(EditAnywhere, BlueprintReadWrite, Category=InterpCurvePointQuat)
+	FQuat ArriveTangent;
+
+	UPROPERTY(EditAnywhere, BlueprintReadWrite, Category=InterpCurvePointQuat)
+	FQuat LeaveTangent;
+
+	UPROPERTY(EditAnywhere, BlueprintReadWrite, Category=InterpCurvePointQuat)
+	TEnumAsByte<enum EInterpCurveMode> InterpMode;
+
+};
+
+
+
+USTRUCT(noexport)
+struct FInterpCurveQuat
+{
+	UPROPERTY(EditAnywhere, BlueprintReadWrite, Category=InterpCurveQuat)
+	TArray<FInterpCurvePointQuat> Points;
+
+	UPROPERTY(EditAnywhere, BlueprintReadWrite, Category=InterpCurveQuat)
+	bool bIsLooped;
+
+	UPROPERTY(EditAnywhere, BlueprintReadWrite, Category=InterpCurveQuat)
+	float LoopKeyOffset;
 };
 
 
@@ -697,6 +749,11 @@ struct FInterpCurveTwoVectors
 	UPROPERTY(EditAnywhere, BlueprintReadWrite, Category=InterpCurveTwoVectors)
 	TArray<FInterpCurvePointTwoVectors> Points;
 
+	UPROPERTY(EditAnywhere, BlueprintReadWrite, Category=InterpCurveTwoVectors)
+	bool bIsLooped;
+
+	UPROPERTY(EditAnywhere, BlueprintReadWrite, Category=InterpCurveTwoVectors)
+	float LoopKeyOffset;
 };
 
 
@@ -729,6 +786,11 @@ struct FInterpCurveLinearColor
 	UPROPERTY(EditAnywhere, BlueprintReadWrite, Category=InterpCurveLinearColor)
 	TArray<FInterpCurvePointLinearColor> Points;
 
+	UPROPERTY(EditAnywhere, BlueprintReadWrite, Category=InterpCurveLinearColor)
+	bool bIsLooped;
+
+	UPROPERTY(EditAnywhere, BlueprintReadWrite, Category=InterpCurveLinearColor)
+	float LoopKeyOffset;
 };
 
 
diff --git a/Engine/Source/Runtime/Engine/Classes/Components/SplineComponent.h b/Engine/Source/Runtime/Engine/Classes/Components/SplineComponent.h
--- a/Engine/Source/Runtime/Engine/Classes/Components/SplineComponent.h
+++ b/Engine/Source/Runtime/Engine/Classes/Components/SplineComponent.h
@@ -15,7 +15,19 @@ namespace ESplinePointType
 		Linear,
 		Curve,
 		Constant,
-		CurveClamped
+		CurveClamped,
+		CurveCustomTangent
+	};
+}
+
+/** Types of coordinate space accepted by the functions */
+UENUM()
+namespace ESplineCoordinateSpace
+{
+	enum Type
+	{
+		Local,
+		World
 	};
 }
 
@@ -29,17 +41,24 @@ class ENGINE_API USplineComponent : public UPrimitiveComponent
 {
 	GENERATED_UCLASS_BODY()
 
-	/** Actual data for spline. Locations and tangents are in world space. */
-	UPROPERTY()
+	/** Spline built from position data. */
+	UPROPERTY()//EditAnywhere, Category = Points)
 	FInterpCurveVector SplineInfo;
 
+	/** Spline built from rotation data. */
+	UPROPERTY()
+	FInterpCurveQuat SplineRotInfo;
+
+	/** Spline built from scale data. */
+	UPROPERTY()
+	FInterpCurveVector SplineScaleInfo;
+
 	/** Input, distance along curve, output, parameter that puts you there. */
-	UPROPERTY(Transient, TextExportTransient)
+	UPROPERTY()
 	FInterpCurveFloat SplineReparamTable;
 
-	/** If true, spline keys may be edited per instance in the level viewport. Otherwise, the spline should be initialized in the construction script. */
-	UPROPERTY(EditDefaultsOnly, Category = Spline)
-	bool bAllowSplineEditingPerInstance;
+	UPROPERTY()
+	bool bAllowSplineEditingPerInstance_DEPRECATED;
 
 	/** Number of steps per spline segment to place in the reparameterization table */
 	UPROPERTY(EditAnywhere, AdvancedDisplay, Category = Spline, meta=(ClampMin=4, UIMin=4, ClampMax=100, UIMax=100))
@@ -53,6 +72,10 @@ class ENGINE_API USplineComponent : public UPrimitiveComponent
 	UPROPERTY(EditAnywhere, BlueprintReadWrite, AdvancedDisplay, Category = Spline, meta=(EditCondition="!bClosedLoop"))
 	bool bStationaryEndpoints;
 
+	/** Whether the spline has been edited from its default by the spline component visualizer */
+	UPROPERTY(VisibleAnywhere, Category = Spline)
+	bool bSplineHasBeenEdited;
+
 private:
 	/**
 	 * Whether the spline is to be considered as a closed loop.
@@ -62,6 +85,10 @@ private:
 	bool bClosedLoop;
 
 public:
+	/** Default up vector in local space to be used when calculating transforms along the spline */
+	UPROPERTY(EditAnywhere, BlueprintReadWrite, Category = Spline)
+	FVector DefaultUpVector;
+
 #if WITH_EDITORONLY_DATA
 	/** Color of an unselected spline component segment in the editor */
 	UPROPERTY(EditAnywhere, Category = Editor)
@@ -70,6 +97,14 @@ public:
 	/** Color of a selected spline component segment in the editor */
 	UPROPERTY(EditAnywhere, Category = Editor)
 	FLinearColor EditorSelectedSplineSegmentColor;
+
+	/** Whether scale visualization should be displayed */
+	UPROPERTY(EditAnywhere, Category = Editor)
+	bool bShouldVisualizeScale;
+
+	/** Width of spline in editor for use with scale visualization */
+	UPROPERTY(EditAnywhere, Category = Editor)
+	float ScaleVisualizationWidth;
 #endif
 
 	// Begin UActorComponent interface.
@@ -77,11 +112,44 @@ public:
 	virtual FName GetComponentInstanceDataType() const override;
 	// End UActorComponent interface.
 
-	void ApplyComponentInstanceData(class FSplineInstanceData* ComponentInstanceData);
+	void ApplyComponentInstanceData(class FSplineInstanceData* ComponentInstanceData, const bool bPostUCS);
 
 	/** Update the spline tangents and SplineReparamTable */
 	void UpdateSpline();
 
+	/** Get location along spline at the provided input key value */
+	FVector GetLocationAtSplineInputKey(float InKey, ESplineCoordinateSpace::Type CoordinateSpace) const;
+
+	/** Get tangent along spline at the provided input key value */
+	FVector GetTangentAtSplineInputKey(float InKey, ESplineCoordinateSpace::Type CoordinateSpace) const;
+
+	/** Get unit direction along spline at the provided input key value */
+	FVector GetDirectionAtSplineInputKey(float InKey, ESplineCoordinateSpace::Type CoordinateSpace) const;
+
+	/** Get rotator corresponding to rotation along spline at the provided input key value */
+	FRotator GetRotationAtSplineInputKey(float InKey, ESplineCoordinateSpace::Type CoordinateSpace) const;
+
+	/** Get quaternion corresponding to rotation along spline at the provided input key value */
+	FQuat GetQuaternionAtSplineInputKey(float InKey, ESplineCoordinateSpace::Type CoordinateSpace) const;
+
+	/** Get up vector at the provided input key value */
+	FVector GetUpVectorAtSplineInputKey(float InKey, ESplineCoordinateSpace::Type CoordinateSpace) const;
+
+	/** Get up vector at the provided input key value */
+	FVector GetRightVectorAtSplineInputKey(float InKey, ESplineCoordinateSpace::Type CoordinateSpace) const;
+
+	/** Get transform at the provided input key value */
+	FTransform GetTransformAtSplineInputKey(float InKey, ESplineCoordinateSpace::Type CoordinateSpace, bool bUseScale = false) const;
+
+	/** Get roll in degrees at the provided input key value */
+	float GetRollAtSplineInputKey(float InKey, ESplineCoordinateSpace::Type CoordinateSpace) const;
+
+	/** Get scale at the provided input key value */
+	FVector GetScaleAtSplineInputKey(float InKey) const;
+
+	DEPRECATED(4.9, "This has been replaced by GetNumberOfSplinePoints(), as it is a clearer name for Blueprint users.")
+	int32 GetNumSplinePoints() const { return GetNumberOfSplinePoints(); }
+
 	/** Specify unselected spline component segment color in the editor */
 	UFUNCTION(BlueprintCallable, Category = Editor)
 	void SetUnselectedSplineSegmentColor(const FLinearColor& SegmentColor);
@@ -99,29 +167,45 @@ public:
 	bool IsClosedLoop() const;
 
 	/** Clears all the points in the spline */
-	UFUNCTION(BlueprintCallable, Category=Spline)
+	UFUNCTION(BlueprintCallable, Category = Spline)
 	void ClearSplinePoints();
 
+	/** Adds a point to the spline */
+	UFUNCTION(BlueprintCallable, Category = Spline)
+	void AddSplinePoint(const FVector& Position, ESplineCoordinateSpace::Type CoordinateSpace);
+
 	/** Adds a world space point to the spline */
-	UFUNCTION(BlueprintCallable, Category=Spline)
+	UFUNCTION(BlueprintCallable, Category = Spline, meta = (DeprecatedFunction, DeprecationMessage = "Please use AddSplinePoint, specifying SplineCoordinateSpace::World"))
 	void AddSplineWorldPoint(const FVector& Position);
 
 	/** Adds a local space point to the spline */
-	UFUNCTION(BlueprintCallable, Category = Spline)
+	UFUNCTION(BlueprintCallable, Category = Spline, meta = (DeprecatedFunction, DeprecationMessage = "Please use AddSplinePoint, specifying SplineCoordinateSpace::Local"))
 	void AddSplineLocalPoint(const FVector& Position);
 
+	/** Sets the spline to an array of points */
+	UFUNCTION(BlueprintCallable, Category = Spline)
+	void SetSplinePoints(const TArray<FVector>& Points, ESplineCoordinateSpace::Type CoordinateSpace);
+
 	/** Sets the spline to an array of world space points */
-	UFUNCTION(BlueprintCallable, Category=Spline)
+	UFUNCTION(BlueprintCallable, Category = Spline, meta = (DeprecatedFunction, DeprecationMessage = "Please use SetSplinePoints, specifying SplineCoordinateSpace::World"))
 	void SetSplineWorldPoints(const TArray<FVector>& Points);
 
 	/** Sets the spline to an array of local space points */
-	UFUNCTION(BlueprintCallable, Category = Spline)
+	UFUNCTION(BlueprintCallable, Category = Spline, meta = (DeprecatedFunction, DeprecationMessage = "Please use SetSplinePoints, specifying SplineCoordinateSpace::Local"))
 	void SetSplineLocalPoints(const TArray<FVector>& Points);
 
-	/** Move an existing point to a new world location */
+	/** Move an existing point to a new location */
 	UFUNCTION(BlueprintCallable, Category = Spline)
+	void SetLocationAtSplinePoint(int32 PointIndex, const FVector& InLocation, ESplineCoordinateSpace::Type CoordinateSpace);
+
+	/** Move an existing point to a new world location */
+	UFUNCTION(BlueprintCallable, Category = Spline, meta = (DeprecatedFunction, DeprecationMessage = "Please use SetLocationAtSplinePoint, specifying SplineCoordinateSpace::World"))
 	void SetWorldLocationAtSplinePoint(int32 PointIndex, const FVector& InLocation);
 
+	/** Specify the tangent at a given spline point */
+	UFUNCTION(BlueprintCallable, Category = Spline)
+	void SetTangentAtSplinePoint(int32 PointIndex, const FVector& InTangent, ESplineCoordinateSpace::Type CoordinateSpace);
+
 	/** Get the type of a spline point */
 	UFUNCTION(BlueprintCallable, Category = Spline)
 	ESplinePointType::Type GetSplinePointType(int32 PointIndex) const;
@@ -131,15 +215,59 @@ public:
 	void SetSplinePointType(int32 PointIndex, ESplinePointType::Type Type);
 
 	/** Get the number of points that make up this spline */
-	UFUNCTION(BlueprintCallable, Category=Spline)
-	int32 GetNumSplinePoints() const;
+	UFUNCTION(BlueprintCallable, Category = Spline)
+	int32 GetNumberOfSplinePoints() const;
+
+	/** Get the location at spline point */
+	UFUNCTION(BlueprintCallable, Category = Spline)
+	FVector GetLocationAtSplinePoint(int32 PointIndex, ESplineCoordinateSpace::Type CoordinateSpace) const;
 
 	/** Get the world location at spline point */
-	UFUNCTION(BlueprintCallable, Category=Spline)
+	UFUNCTION(BlueprintCallable, Category = Spline, meta = (DeprecatedFunction, DeprecationMessage = "Please use GetLocationAtSplinePoint, specifying SplineCoordinateSpace::World"))
 	FVector GetWorldLocationAtSplinePoint(int32 PointIndex) const;
 
-	/** Get local location and tangent at a spline point */
+	/** Get the location at spline point */
+	UFUNCTION(BlueprintCallable, Category = Spline)
+	FVector GetDirectionAtSplinePoint(int32 PointIndex, ESplineCoordinateSpace::Type CoordinateSpace) const;
+
+	/** Get the tangent at spline point */
+	UFUNCTION(BlueprintCallable, Category = Spline)
+	FVector GetTangentAtSplinePoint(int32 PointIndex, ESplineCoordinateSpace::Type CoordinateSpace) const;
+
+	/** Get the rotation at spline point as a quaternion */
+	UFUNCTION(BlueprintCallable, Category = Spline)
+	FQuat GetQuaternionAtSplinePoint(int32 PointIndex, ESplineCoordinateSpace::Type CoordinateSpace) const;
+
+	/** Get the rotation at spline point as a rotator */
+	UFUNCTION(BlueprintCallable, Category = Spline)
+	FRotator GetRotationAtSplinePoint(int32 PointIndex, ESplineCoordinateSpace::Type CoordinateSpace) const;
+
+	/** Get the up vector at spline point */
+	UFUNCTION(BlueprintCallable, Category = Spline)
+	FVector GetUpVectorAtSplinePoint(int32 PointIndex, ESplineCoordinateSpace::Type CoordinateSpace) const;
+
+	/** Get the right vector at spline point */
+	UFUNCTION(BlueprintCallable, Category = Spline)
+	FVector GetRightVectorAtSplinePoint(int32 PointIndex, ESplineCoordinateSpace::Type CoordinateSpace) const;
+
+	/** Get the amount of roll at spline point, in degrees */
+	UFUNCTION(BlueprintCallable, Category = Spline)
+	float GetRollAtSplinePoint(int32 PointIndex, ESplineCoordinateSpace::Type CoordinateSpace) const;
+
+	/** Get the scale at spline point */
+	UFUNCTION(BlueprintCallable, Category = Spline)
+	FVector GetScaleAtSplinePoint(int32 PointIndex) const;
+
+	/** Get the transform at spline point */
+	UFUNCTION(BlueprintCallable, Category = Spline)
+	FTransform GetTransformAtSplinePoint(int32 PointIndex, ESplineCoordinateSpace::Type CoordinateSpace, bool bUseScale = false) const;
+
+	/** Get location and tangent at a spline point */
 	UFUNCTION(BlueprintCallable, Category=Spline)
+	void GetLocationAndTangentAtSplinePoint(int32 PointIndex, FVector& Location, FVector& Tangent, ESplineCoordinateSpace::Type CoordinateSpace) const;
+
+	/** Get local location and tangent at a spline point */
+	UFUNCTION(BlueprintCallable, Category = Spline, meta = (DeprecatedFunction, DeprecationMessage = "Please use GetLocationAndTangentAtSplinePoint, specifying SplineCoordinateSpace::Local"))
 	void GetLocalLocationAndTangentAtSplinePoint(int32 PointIndex, FVector& LocalLocation, FVector& LocalTangent) const;
 
 	/** Get the distance along the spline at the spline point */
@@ -150,49 +278,134 @@ public:
 	UFUNCTION(BlueprintCallable, Category=Spline) 
 	float GetSplineLength() const;
 
+	/** Sets the default up vector used by this spline */
+	UFUNCTION(BlueprintCallable, Category=Spline)
+	void SetDefaultUpVector(const FVector& UpVector, ESplineCoordinateSpace::Type CoordinateSpace);
+
+	/** Gets the default up vector used by this spline */
+	UFUNCTION(BlueprintCallable, Category = Spline)
+	FVector GetDefaultUpVector(ESplineCoordinateSpace::Type CoordinateSpace) const;
+
 	/** Given a distance along the length of this spline, return the corresponding input key at that point */
 	UFUNCTION(BlueprintCallable, Category=Spline)
 	float GetInputKeyAtDistanceAlongSpline(float Distance) const;
 
 	/** Given a distance along the length of this spline, return the point in space where this puts you */
 	UFUNCTION(BlueprintCallable, Category=Spline)
+	FVector GetLocationAtDistanceAlongSpline(float Distance, ESplineCoordinateSpace::Type CoordinateSpace) const;
+
+	/** Given a distance along the length of this spline, return the point in world space where this puts you */
+	UFUNCTION(BlueprintCallable, Category = Spline, meta = (DeprecatedFunction, DeprecationMessage = "Please use GetLocationAtDistanceAlongSpline, specifying SplineCoordinateSpace::World"))
 	FVector GetWorldLocationAtDistanceAlongSpline(float Distance) const;
 	
 	/** Given a distance along the length of this spline, return a unit direction vector of the spline tangent there. */
 	UFUNCTION(BlueprintCallable, Category=Spline)
+	FVector GetDirectionAtDistanceAlongSpline(float Distance, ESplineCoordinateSpace::Type CoordinateSpace) const;
+
+	/** Given a distance along the length of this spline, return a unit direction vector of the spline tangent there, in world space. */
+	UFUNCTION(BlueprintCallable, Category = Spline, meta = (DeprecatedFunction, DeprecationMessage = "Please use GetDirectionAtDistanceAlongSpline, specifying SplineCoordinateSpace::World"))
 	FVector GetWorldDirectionAtDistanceAlongSpline(float Distance) const;
 
 	/** Given a distance along the length of this spline, return the tangent vector of the spline there. */
 	UFUNCTION(BlueprintCallable, Category = Spline)
+	FVector GetTangentAtDistanceAlongSpline(float Distance, ESplineCoordinateSpace::Type CoordinateSpace) const;
+
+	/** Given a distance along the length of this spline, return the tangent vector of the spline there, in world space. */
+	UFUNCTION(BlueprintCallable, Category = Spline, meta = (DeprecatedFunction, DeprecationMessage = "Please use GetTangentAtDistanceAlongSpline, specifying SplineCoordinateSpace::World"))
 	FVector GetWorldTangentAtDistanceAlongSpline(float Distance) const;
 
-	/** Given a distance along the length of this spline, return a rotation corresponding to the spline's position and direction there. */
+	/** Given a distance along the length of this spline, return a quaternion corresponding to the spline's rotation there. */
 	UFUNCTION(BlueprintCallable, Category = Spline)
+	FQuat GetQuaternionAtDistanceAlongSpline(float Distance, ESplineCoordinateSpace::Type CoordinateSpace) const;
+
+	/** Given a distance along the length of this spline, return a rotation corresponding to the spline's rotation there. */
+	UFUNCTION(BlueprintCallable, Category = Spline)
+	FRotator GetRotationAtDistanceAlongSpline(float Distance, ESplineCoordinateSpace::Type CoordinateSpace) const;
+
+	/** Given a distance along the length of this spline, return a rotation corresponding to the spline's rotation there, in world space. */
+	UFUNCTION(BlueprintCallable, Category = Spline, meta = (DeprecatedFunction, DeprecationMessage = "Please use GetRotationAtDistanceAlongSpline, specifying SplineCoordinateSpace::World"))
 	FRotator GetWorldRotationAtDistanceAlongSpline(float Distance) const;
 
+	/** Given a distance along the length of this spline, return a unit direction vector corresponding to the spline's up vector there. */
+	UFUNCTION(BlueprintCallable, Category = Spline)
+	FVector GetUpVectorAtDistanceAlongSpline(float Distance, ESplineCoordinateSpace::Type CoordinateSpace) const;
+
+	/** Given a distance along the length of this spline, return a unit direction vector corresponding to the spline's right vector there. */
+	UFUNCTION(BlueprintCallable, Category = Spline)
+	FVector GetRightVectorAtDistanceAlongSpline(float Distance, ESplineCoordinateSpace::Type CoordinateSpace) const;
+
+	/** Given a distance along the length of this spline, return the spline's roll there, in degrees. */
+	UFUNCTION(BlueprintCallable, Category = Spline)
+	float GetRollAtDistanceAlongSpline(float Distance, ESplineCoordinateSpace::Type CoordinateSpace) const;
+
+	/** Given a distance along the length of this spline, return the spline's scale there. */
+	UFUNCTION(BlueprintCallable, Category = Spline)
+	FVector GetScaleAtDistanceAlongSpline(float Distance) const;
+
+	/** Given a distance along the length of this spline, return an FTransform corresponding to that point on the spline. */
+	UFUNCTION(BlueprintCallable, Category = Spline)
+	FTransform GetTransformAtDistanceAlongSpline(float Distance, ESplineCoordinateSpace::Type CoordinateSpace, bool bUseScale = false) const;
+
 	/** Given a time from 0 to the spline duration, return the point in space where this puts you */
 	UFUNCTION(BlueprintCallable, Category=Spline)
+	FVector GetLocationAtTime(float Time, ESplineCoordinateSpace::Type CoordinateSpace, bool bUseConstantVelocity = false) const;
+
+	/** Given a time from 0 to the spline duration, return the point in space where this puts you */
+	UFUNCTION(BlueprintCallable, Category = Spline, meta = (DeprecatedFunction, DeprecationMessage = "Please use GetLocationAtTime, specifying SplineCoordinateSpace::World"))
 	FVector GetWorldLocationAtTime(float Time, bool bUseConstantVelocity = false) const;
 
 	/** Given a time from 0 to the spline duration, return a unit direction vector of the spline tangent there. */
 	UFUNCTION(BlueprintCallable, Category=Spline)
+	FVector GetDirectionAtTime(float Time, ESplineCoordinateSpace::Type CoordinateSpace, bool bUseConstantVelocity = false) const;
+
+	/** Given a time from 0 to the spline duration, return a unit direction vector of the spline tangent there. */
+	UFUNCTION(BlueprintCallable, Category = Spline, meta = (DeprecatedFunction, DeprecationMessage = "Please use GetDirectionAtTime, specifying SplineCoordinateSpace::World"))
 	FVector GetWorldDirectionAtTime(float Time, bool bUseConstantVelocity = false) const;
 
+	/** Given a time from 0 to the spline duration, return the spline's tangent there. */
+	UFUNCTION(BlueprintCallable, Category=Spline)
+	FVector GetTangentAtTime(float Time, ESplineCoordinateSpace::Type CoordinateSpace, bool bUseConstantVelocity = false) const;
+
+	/** Given a time from 0 to the spline duration, return a quaternion corresponding to the spline's rotation there. */
+	UFUNCTION(BlueprintCallable, Category=Spline)
+	FQuat GetQuaternionAtTime(float Time, ESplineCoordinateSpace::Type CoordinateSpace, bool bUseConstantVelocity = false) const;
+
 	/** Given a time from 0 to the spline duration, return a rotation corresponding to the spline's position and direction there. */
 	UFUNCTION(BlueprintCallable, Category=Spline)
+	FRotator GetRotationAtTime(float Time, ESplineCoordinateSpace::Type CoordinateSpace, bool bUseConstantVelocity = false) const;
+
+	/** Given a time from 0 to the spline duration, return a rotation corresponding to the spline's position and direction there. */
+	UFUNCTION(BlueprintCallable, Category = Spline, meta = (DeprecatedFunction, DeprecationMessage = "Please use GetRotationAtTime, specifying SplineCoordinateSpace::World"))
 	FRotator GetWorldRotationAtTime(float Time, bool bUseConstantVelocity = false) const;
 
+	/** Given a time from 0 to the spline duration, return the spline's up vector there. */
+	UFUNCTION(BlueprintCallable, Category=Spline)
+	FVector GetUpVectorAtTime(float Time, ESplineCoordinateSpace::Type CoordinateSpace, bool bUseConstantVelocity = false) const;
+
+	/** Given a time from 0 to the spline duration, return the spline's right vector there. */
+	UFUNCTION(BlueprintCallable, Category=Spline)
+	FVector GetRightVectorAtTime(float Time, ESplineCoordinateSpace::Type CoordinateSpace, bool bUseConstantVelocity = false) const;
+
+	/** Given a time from 0 to the spline duration, return the spline's transform at the corresponding position. */
+	UFUNCTION(BlueprintCallable, Category=Spline)
+	FTransform GetTransformAtTime(float Time, ESplineCoordinateSpace::Type CoordinateSpace, bool bUseConstantVelocity = false, bool bUseScale = false) const;
+
+	/** Given a time from 0 to the spline duration, return the spline's roll there, in degrees. */
+	UFUNCTION(BlueprintCallable, Category=Spline)
+	float GetRollAtTime(float Time, ESplineCoordinateSpace::Type CoordinateSpace, bool bUseConstantVelocity = false) const;
+
+	/** Given a time from 0 to the spline duration, return the spline's scale there. */
+	UFUNCTION(BlueprintCallable, Category=Spline)
+	FVector GetScaleAtTime(float Time, bool bUseConstantVelocity = false) const;
+
 	// UObject interface
 	virtual void PostLoad() override;
 	virtual void PostEditImport() override;
 #if WITH_EDITOR
-	virtual void PreEditChange(UProperty* PropertyAboutToChange) override;
-	virtual void PostEditChangeProperty(FPropertyChangedEvent& PropertyChangedEvent) override;
-	// End of UObject interface
+	virtual void Serialize(FArchive& Ar) override;
+	virtual void PostEditChangeChainProperty(FPropertyChangedChainEvent& PropertyChangedEvent) override;
 #endif
-
-	/** Walk through keys and set time for each one */
-	void RefreshSplineInputs();
+	// End of UObject interface
 
 private:
 	/** Returns the length of the specified spline segment up to the parametric value given */
@@ -200,16 +413,21 @@ private:
 
 	/** Returns the parametric value t which would result in a spline segment of the given length between S(0)...S(t) */
 	float GetSegmentParamFromLength(const int32 Index, const float Length, const float SegmentLength) const;
-
-	/** Updates the spline points to reflect a closed loop or not */
-	void UpdateLoopEndpoint(bool bInClosedLoop);
-
-	/** Adds an additional point to the end of the spline, whose OutValue and tangents coincides with the first point */
-	void AddLoopEndpoint();
-
-	/** Removes the final point from the spline, assuming it to be a duplicate of the first point */
-	void RemoveLoopEndpoint();
 };
 
 
-
+// Deprecated method definitions
+inline void USplineComponent::AddSplineWorldPoint(const FVector& Position) { AddSplinePoint(Position, ESplineCoordinateSpace::World); }
+inline void USplineComponent::AddSplineLocalPoint(const FVector& Position) { AddSplinePoint(Position, ESplineCoordinateSpace::Local); }
+inline void USplineComponent::SetSplineWorldPoints(const TArray<FVector>& Points) { SetSplinePoints(Points, ESplineCoordinateSpace::World); }
+inline void USplineComponent::SetSplineLocalPoints(const TArray<FVector>& Points) { SetSplinePoints(Points, ESplineCoordinateSpace::Local); }
+inline void USplineComponent::SetWorldLocationAtSplinePoint(int32 PointIndex, const FVector& InLocation) { SetLocationAtSplinePoint(PointIndex, InLocation, ESplineCoordinateSpace::World); }
+inline FVector USplineComponent::GetWorldLocationAtSplinePoint(int32 PointIndex) const { return GetLocationAtSplinePoint(PointIndex, ESplineCoordinateSpace::World); }
+inline void USplineComponent::GetLocalLocationAndTangentAtSplinePoint(int32 PointIndex, FVector& LocalLocation, FVector& LocalTangent) const { GetLocationAndTangentAtSplinePoint(PointIndex, LocalLocation, LocalTangent, ESplineCoordinateSpace::Local); }
+inline FVector USplineComponent::GetWorldLocationAtDistanceAlongSpline(float Distance) const { return GetLocationAtDistanceAlongSpline(Distance, ESplineCoordinateSpace::World); }
+inline FVector USplineComponent::GetWorldDirectionAtDistanceAlongSpline(float Distance) const { return GetDirectionAtDistanceAlongSpline(Distance, ESplineCoordinateSpace::World); }
+inline FVector USplineComponent::GetWorldTangentAtDistanceAlongSpline(float Distance) const { return GetTangentAtDistanceAlongSpline(Distance, ESplineCoordinateSpace::World); }
+inline FRotator USplineComponent::GetWorldRotationAtDistanceAlongSpline(float Distance) const { return GetRotationAtDistanceAlongSpline(Distance, ESplineCoordinateSpace::World); }
+inline FVector USplineComponent::GetWorldLocationAtTime(float Time, bool bUseConstantVelocity) const { return GetLocationAtTime(Time, ESplineCoordinateSpace::World, bUseConstantVelocity); }
+inline FVector USplineComponent::GetWorldDirectionAtTime(float Time, bool bUseConstantVelocity) const { return GetDirectionAtTime(Time, ESplineCoordinateSpace::World, bUseConstantVelocity); }
+inline FRotator USplineComponent::GetWorldRotationAtTime(float Time, bool bUseConstantVelocity) const { return GetRotationAtTime(Time, ESplineCoordinateSpace::World, bUseConstantVelocity); }
diff --git a/Engine/Source/Runtime/Engine/Private/Components/SplineComponent.cpp b/Engine/Source/Runtime/Engine/Private/Components/SplineComponent.cpp
--- a/Engine/Source/Runtime/Engine/Private/Components/SplineComponent.cpp
+++ b/Engine/Source/Runtime/Engine/Private/Components/SplineComponent.cpp
@@ -11,50 +11,127 @@
 
 USplineComponent::USplineComponent(const FObjectInitializer& ObjectInitializer)
 	: Super(ObjectInitializer)
-	, bAllowSplineEditingPerInstance(true)
+	, bAllowSplineEditingPerInstance_DEPRECATED(true)
 	, ReparamStepsPerSegment(10)
 	, Duration(1.0f)
 	, bStationaryEndpoints(false)
+	, bSplineHasBeenEdited(false)
 	, bClosedLoop(false)
+	, DefaultUpVector(FVector::UpVector)
 #if WITH_EDITORONLY_DATA
 	, EditorUnselectedSplineSegmentColor(FLinearColor(1.0f, 1.0f, 1.0f))
 	, EditorSelectedSplineSegmentColor(FLinearColor(1.0f, 0.0f, 0.0f))
+	, bShouldVisualizeScale(false)
+	, ScaleVisualizationWidth(30.0f)
 #endif
 {
 	SplineInfo.Points.Reset(10);
+	SplineRotInfo.Points.Reset(10);
+	SplineScaleInfo.Points.Reset(10);
 
-	// Add 2 keys by default
-	int32 PointIndex = SplineInfo.AddPoint(0.f, FVector(0,0,0));
-	SplineInfo.Points[PointIndex].InterpMode = CIM_CurveAuto;
+	SplineInfo.Points.Emplace(0.0f, FVector(0, 0, 0), FVector::ZeroVector, FVector::ZeroVector, CIM_CurveAuto);
+	SplineRotInfo.Points.Emplace(0.0f, FQuat::Identity, FQuat::Identity, FQuat::Identity, CIM_CurveAuto);
+	SplineScaleInfo.Points.Emplace(0.0f, FVector(1.0f), FVector::ZeroVector, FVector::ZeroVector, CIM_CurveAuto);
 
-	PointIndex = SplineInfo.AddPoint(1.f, FVector(100,0,0));
-	SplineInfo.Points[PointIndex].InterpMode = CIM_CurveAuto;
+	SplineInfo.Points.Emplace(1.0f, FVector(100, 0, 0), FVector::ZeroVector, FVector::ZeroVector, CIM_CurveAuto);
+	SplineRotInfo.Points.Emplace(1.0f, FQuat::Identity, FQuat::Identity, FQuat::Identity, CIM_CurveAuto);
+	SplineScaleInfo.Points.Emplace(1.0f, FVector(1.0f), FVector::ZeroVector, FVector::ZeroVector, CIM_CurveAuto);
 
 	UpdateSpline();
 }
 
 
+static EInterpCurveMode ConvertSplinePointTypeToInterpCurveMode(ESplinePointType::Type SplinePointType)
+{
+	switch (SplinePointType)
+	{
+		case ESplinePointType::Linear:				return CIM_Linear;
+		case ESplinePointType::Curve:				return CIM_CurveAuto;
+		case ESplinePointType::Constant:			return CIM_Constant;
+		case ESplinePointType::CurveCustomTangent:	return CIM_CurveUser;
+		case ESplinePointType::CurveClamped:		return CIM_CurveAutoClamped;
+
+		default:									return CIM_Unknown;
+	}
+}
+
+static ESplinePointType::Type ConvertInterpCurveModeToSplinePointType(EInterpCurveMode InterpCurveMode)
+{
+	switch (InterpCurveMode)
+	{
+		case CIM_Linear:			return ESplinePointType::Linear;
+		case CIM_CurveAuto:			return ESplinePointType::Curve;
+		case CIM_Constant:			return ESplinePointType::Constant;
+		case CIM_CurveUser:			return ESplinePointType::CurveCustomTangent;
+		case CIM_CurveAutoClamped:	return ESplinePointType::CurveClamped;
+
+		default:					return ESplinePointType::Constant;
+	}
+}
+
+
+#if WITH_EDITOR
+void USplineComponent::Serialize(FArchive& Ar)
+{
+	Super::Serialize(Ar);
+
+	// Support old resources which don't have the rotation and scale splines present
+	const int32 ArchiveUE4Version = Ar.UE4Ver();
+	if (ArchiveUE4Version < VER_UE4_SPLINE_COMPONENT_WITH_ROTATION_AND_SCALE)
+	{
+		int32 NumPoints = SplineInfo.Points.Num();
+
+		// The start point is no longer cloned as the endpoint when the spline is looped, so remove the extra endpoint if present
+		const bool bHasExtraEndpoint = bClosedLoop && (SplineInfo.Points[0].OutVal == SplineInfo.Points[NumPoints - 1].OutVal);
+
+		if (bHasExtraEndpoint)
+		{
+			SplineInfo.Points.RemoveAt(NumPoints - 1, 1, false);
+			NumPoints--;
+		}
+
+		// Fill the other two splines with some defaults
+		SplineRotInfo.Points.Reset(NumPoints);
+		SplineScaleInfo.Points.Reset(NumPoints);
+		for (int32 Count = 0; Count < NumPoints; Count++)
+		{
+			SplineRotInfo.Points.Emplace(0.0f, FQuat::Identity, FQuat::Identity, FQuat::Identity, CIM_CurveAuto);
+			SplineScaleInfo.Points.Emplace(0.0f, FVector(1.0f), FVector::ZeroVector, FVector::ZeroVector, CIM_CurveAuto);
+		}
+
+		USplineComponent* Archetype = CastChecked<USplineComponent>(GetArchetype());
+		bSplineHasBeenEdited = (SplineInfo.Points != Archetype->SplineInfo.Points);
+		UpdateSpline();
+	}
+}
+#endif
+
+
 void USplineComponent::PostLoad()
 {
 	Super::PostLoad();
-	UpdateSpline();
 }
 
 
 void USplineComponent::PostEditImport()
 {
 	Super::PostEditImport();
-	UpdateSpline();
 }
 
 
 void USplineComponent::UpdateSpline()
 {
 	const int32 NumPoints = SplineInfo.Points.Num();
-	check(!bClosedLoop || NumPoints == 0 || (NumPoints >= 2 && SplineInfo.Points[0].OutVal == SplineInfo.Points[NumPoints - 1].OutVal));
+	check(SplineRotInfo.Points.Num() == NumPoints && SplineScaleInfo.Points.Num() == NumPoints);
 
-	// Automatically set the tangents on any CurveAuto keys
-	SplineInfo.AutoSetTangents(0.0f, bStationaryEndpoints);
+	// Ensure input keys are ascending in steps of 1.0
+	for (int32 Index = 0; Index < NumPoints; Index++)
+	{
+		const float InVal = static_cast<float>(Index);
+		SplineInfo.Points[Index].InVal = InVal;
+		SplineRotInfo.Points[Index].InVal = InVal;
+		SplineScaleInfo.Points[Index].InVal = InVal;
+	}
 
 	// Nothing else to do if less than 2 points
 	if (NumPoints < 2)
@@ -62,33 +139,28 @@ void USplineComponent::UpdateSpline()
 		return;
 	}
 
-	// Adjust auto tangents for first and last keys to take into account the looping
+	// Ensure splines' looping status matches with that of the spline component
 	if (bClosedLoop)
 	{
-		auto& FirstPoint = SplineInfo.Points[0];
-		auto& LastPoint = SplineInfo.Points[NumPoints - 1];
-		const auto& SecondPoint = SplineInfo.Points[1];
-		const auto& PenultimatePoint = SplineInfo.Points[NumPoints - 2];
-
-		if (FirstPoint.InterpMode == CIM_CurveAuto || FirstPoint.InterpMode == CIM_CurveAutoClamped)
-		{
-			FVector Tangent;
-			ComputeCurveTangent(
-				PenultimatePoint.InVal - LastPoint.InVal, PenultimatePoint.OutVal,
-				FirstPoint.InVal, FirstPoint.OutVal,
-				SecondPoint.InVal, SecondPoint.OutVal,
-				0.0f,
-				FirstPoint.InterpMode == CIM_CurveAutoClamped,
-				Tangent);
-
-			FirstPoint.LeaveTangent = Tangent;
-			FirstPoint.ArriveTangent = Tangent;
-			LastPoint.LeaveTangent = Tangent;
-			LastPoint.ArriveTangent = Tangent;
-		}
+		const float LoopKey = static_cast<float>(NumPoints);
+		SplineInfo.SetLoopKey(LoopKey);
+		SplineRotInfo.SetLoopKey(LoopKey);
+		SplineScaleInfo.SetLoopKey(LoopKey);
+	}
+	else
+	{
+		SplineInfo.ClearLoopKey();
+		SplineRotInfo.ClearLoopKey();
+		SplineScaleInfo.ClearLoopKey();
 	}
 
-	const int32 NumSegments = NumPoints - 1;
+	// Automatically set the tangents on any CurveAuto keys
+	SplineInfo.AutoSetTangents(0.0f, bStationaryEndpoints);
+	SplineRotInfo.AutoSetTangents(0.0f, bStationaryEndpoints);
+	SplineScaleInfo.AutoSetTangents(0.0f, bStationaryEndpoints);
+
+	// Now initialize the spline reparam table
+	const int32 NumSegments = bClosedLoop ? NumPoints : NumPoints - 1;
 
 	// Start by clearing it
 	SplineReparamTable.Points.Reset(NumSegments * ReparamStepsPerSegment + 1);
@@ -99,17 +171,22 @@ void USplineComponent::UpdateSpline()
 		{
 			const float Param = static_cast<float>(Step) / ReparamStepsPerSegment;
 			const float SegmentLength = (Step == 0) ? 0.0f : GetSegmentLength(SegmentIndex, Param);
-			SplineReparamTable.AddPoint(SegmentLength + AccumulatedLength, SegmentIndex + Param);
+
+			SplineReparamTable.Points.Emplace(SegmentLength + AccumulatedLength, SegmentIndex + Param, 0.0f, 0.0f, CIM_Linear);
 		}
 		AccumulatedLength += GetSegmentLength(SegmentIndex, 1.0f);
 	}
-	SplineReparamTable.AddPoint(AccumulatedLength, static_cast<float>(NumSegments));
+
+	SplineReparamTable.Points.Emplace(AccumulatedLength, static_cast<float>(NumSegments), 0.0f, 0.0f, CIM_Linear);
 }
 
 
 float USplineComponent::GetSegmentLength(const int32 Index, const float Param) const
 {
-	check(Index < SplineInfo.Points.Num() - 1);
+	const int32 NumPoints = SplineInfo.Points.Num();
+	const int32 LastPoint = NumPoints - 1;
+
+	check(Index >= 0 && ((bClosedLoop && Index < NumPoints) || (!bClosedLoop && Index < LastPoint)));
 	check(Param >= 0.0f && Param <= 1.0f);
 
 	// Evaluate the length of a Hermite spline segment.
@@ -134,8 +211,8 @@ float USplineComponent::GetSegmentLength(const int32 Index, const float Param) c
 	};
 
 	const auto& StartPoint = SplineInfo.Points[Index];
-	const auto& EndPoint = SplineInfo.Points[Index + 1];
-	check(static_cast<int32>(EndPoint.InVal) - static_cast<int32>(StartPoint.InVal) == 1);
+	const auto& EndPoint = SplineInfo.Points[Index == LastPoint ? 0 : Index + 1];
+	check(Index == LastPoint || (static_cast<int32>(EndPoint.InVal) - static_cast<int32>(StartPoint.InVal) == 1));
 
 	const auto& P0 = StartPoint.OutVal;
 	const auto& T0 = StartPoint.LeaveTangent;
@@ -187,9 +264,12 @@ float USplineComponent::GetSegmentParamFromLength(const int32 Index, const float
 	//        f(t) is the function whose root we wish to find = L(t) - d,
 	//        (df/dt) = d(L(t))/dt = |dP/dt|
 
-	check(Index < SplineInfo.Points.Num() - 1);
+	const int32 NumPoints = SplineInfo.Points.Num();
+	const int32 LastPoint = NumPoints - 1;
+
+	check(Index >= 0 && ((bClosedLoop && Index < NumPoints) || (!bClosedLoop && Index < LastPoint)));
 	check(Length >= 0.0f && Length <= SegmentLength);
-	check(static_cast<int32>(SplineInfo.Points[Index + 1].InVal) - static_cast<int32>(SplineInfo.Points[Index].InVal) == 1);
+	check(Index == LastPoint || (static_cast<int32>(SplineInfo.Points[Index + 1].InVal) - static_cast<int32>(SplineInfo.Points[Index].InVal) == 1));
 
 	float Param = Length / SegmentLength;  // initial estimate for t
 
@@ -208,59 +288,138 @@ float USplineComponent::GetSegmentParamFromLength(const int32 Index, const float
 }
 
 
-void USplineComponent::SetClosedLoop(bool bInClosedLoop)
+FVector USplineComponent::GetLocationAtSplineInputKey(float InKey, ESplineCoordinateSpace::Type CoordinateSpace) const
 {
-	UpdateLoopEndpoint(bInClosedLoop);
-	UpdateSpline();
+	FVector Location = SplineInfo.Eval(InKey, FVector::ZeroVector);
+
+	if (CoordinateSpace == ESplineCoordinateSpace::World)
+	{
+		Location = ComponentToWorld.TransformPosition(Location);
+	}
+
+	return Location;
 }
 
 
-void USplineComponent::UpdateLoopEndpoint(bool bInClosedLoop)
+FVector USplineComponent::GetTangentAtSplineInputKey(float InKey, ESplineCoordinateSpace::Type CoordinateSpace) const
 {
-	if (bClosedLoop != bInClosedLoop)
+	FVector Tangent = SplineInfo.EvalDerivative(InKey, FVector::ZeroVector);
+
+	if (CoordinateSpace == ESplineCoordinateSpace::World)
 	{
-		bClosedLoop = bInClosedLoop;
+		Tangent = ComponentToWorld.TransformVector(Tangent);
+	}
 
-		if (bClosedLoop)
-		{
-			AddLoopEndpoint();
-		}
-		else
-		{
-			RemoveLoopEndpoint();
-		}
+	return Tangent;
+}
+
+
+FVector USplineComponent::GetDirectionAtSplineInputKey(float InKey, ESplineCoordinateSpace::Type CoordinateSpace) const
+{
+	FVector Direction = SplineInfo.EvalDerivative(InKey, FVector::ZeroVector).GetSafeNormal();
+
+	if (CoordinateSpace == ESplineCoordinateSpace::World)
+	{
+		Direction = ComponentToWorld.TransformVectorNoScale(Direction);
 	}
+
+	return Direction;
 }
 
 
-bool USplineComponent::IsClosedLoop() const
+FRotator USplineComponent::GetRotationAtSplineInputKey(float InKey, ESplineCoordinateSpace::Type CoordinateSpace) const
 {
-	return bClosedLoop;
+	return GetQuaternionAtSplineInputKey(InKey, CoordinateSpace).Rotator();
 }
 
 
-void USplineComponent::RemoveLoopEndpoint()
+FQuat USplineComponent::GetQuaternionAtSplineInputKey(float InKey, ESplineCoordinateSpace::Type CoordinateSpace) const
 {
-	const int32 NumPoints = SplineInfo.Points.Num();
-	check(!bClosedLoop);
-	check(NumPoints == 0 || (NumPoints >= 2 && SplineInfo.Points[0].OutVal == SplineInfo.Points[NumPoints - 1].OutVal));
-	if (NumPoints > 0)
+	FQuat Quat = SplineRotInfo.Eval(InKey, FQuat::Identity);
+	Quat.Normalize();
+
+	const FVector Direction = SplineInfo.EvalDerivative(InKey, FVector::ZeroVector).GetSafeNormal();
+	const FVector UpVector = Quat.RotateVector(DefaultUpVector);
+
+	FQuat Rot = (FRotationMatrix::MakeFromXZ(Direction, UpVector)).ToQuat();
+
+	if (CoordinateSpace == ESplineCoordinateSpace::World)
 	{
-		SplineInfo.Points.RemoveAt(NumPoints - 1, 1, false);
+		Rot = ComponentToWorld.GetRotation() * Rot;
 	}
+
+	return Rot;
 }
 
 
-void USplineComponent::AddLoopEndpoint()
+FVector USplineComponent::GetUpVectorAtSplineInputKey(float InKey, ESplineCoordinateSpace::Type CoordinateSpace) const
 {
-	check(bClosedLoop);
-	const int32 NumPoints = SplineInfo.Points.Num();
-	if (NumPoints > 0)
+	const FQuat Quat = GetQuaternionAtSplineInputKey(InKey, ESplineCoordinateSpace::Local);
+	FVector UpVector = Quat.RotateVector(FVector::UpVector);
+
+	if (CoordinateSpace == ESplineCoordinateSpace::World)
+	{
+		UpVector = ComponentToWorld.TransformVectorNoScale(UpVector);
+	}
+
+	return UpVector;
+}
+
+
+FVector USplineComponent::GetRightVectorAtSplineInputKey(float InKey, ESplineCoordinateSpace::Type CoordinateSpace) const
+{
+	const FQuat Quat = GetQuaternionAtSplineInputKey(InKey, ESplineCoordinateSpace::Local);
+	FVector RightVector = Quat.RotateVector(FVector::RightVector);
+
+	if (CoordinateSpace == ESplineCoordinateSpace::World)
+	{
+		RightVector = ComponentToWorld.TransformVectorNoScale(RightVector);
+	}
+
+	return RightVector;
+}
+
+
+FTransform USplineComponent::GetTransformAtSplineInputKey(float InKey, ESplineCoordinateSpace::Type CoordinateSpace, bool bUseScale) const
+{
+	const FVector Location(GetLocationAtSplineInputKey(InKey, ESplineCoordinateSpace::Local));
+	const FQuat Rotation(GetQuaternionAtSplineInputKey(InKey, ESplineCoordinateSpace::Local));
+	const FVector Scale = bUseScale ? GetScaleAtSplineInputKey(InKey) : FVector(1.0f);
+
+	FTransform Transform(Rotation, Location, Scale);
+
+	if (CoordinateSpace == ESplineCoordinateSpace::World)
 	{
-		FInterpCurvePoint<FVector> EndPoint(SplineInfo.Points[0]);
-		EndPoint.InVal = static_cast<float>(NumPoints);
-		SplineInfo.Points.Add(EndPoint);
+		Transform = Transform * ComponentToWorld;
 	}
+
+	return Transform;
+}
+
+
+float USplineComponent::GetRollAtSplineInputKey(float InKey, ESplineCoordinateSpace::Type CoordinateSpace) const
+{
+	return GetRotationAtSplineInputKey(InKey, CoordinateSpace).Roll;
+}
+
+
+FVector USplineComponent::GetScaleAtSplineInputKey(float InKey) const
+{
+	const FVector Scale = SplineScaleInfo.Eval(InKey, FVector(1.0f));
+	return Scale;
+}
+
+
+void USplineComponent::SetClosedLoop(bool bInClosedLoop)
+{
+	bClosedLoop = bInClosedLoop;
+	UpdateSpline();
+}
+
+
+bool USplineComponent::IsClosedLoop() const
+{
+	return bClosedLoop;
 }
 
 
@@ -282,106 +441,80 @@ void USplineComponent::SetSelectedSplineSegmentColor(const FLinearColor& Color)
 
 void USplineComponent::ClearSplinePoints()
 {
-	SplineInfo.Reset();
-	SplineReparamTable.Reset();
-}
-
-
-void USplineComponent::AddSplineWorldPoint(const FVector& Position)
-{
-	// If it's a closed loop, remove the endpoint before adding a new point
-	const bool bWasLoop = IsClosedLoop();
-	UpdateLoopEndpoint(false);
-
-	float InputKey = static_cast<float>(SplineInfo.Points.Num());
-	const int32 PointIndex = SplineInfo.AddPoint(InputKey, ComponentToWorld.InverseTransformPosition(Position));
-	SplineInfo.Points[PointIndex].InterpMode = CIM_CurveAuto;
-
-	// Then re-close the spline if required
-	UpdateLoopEndpoint(bWasLoop);
-
+	SplineInfo.Points.Reset();
+	SplineRotInfo.Points.Reset();
+	SplineScaleInfo.Points.Reset();
 	UpdateSpline();
 }
 
 
-void USplineComponent::AddSplineLocalPoint(const FVector& Position)
+void USplineComponent::AddSplinePoint(const FVector& Position, ESplineCoordinateSpace::Type CoordinateSpace)
 {
-	// If it's a closed loop, remove the endpoint before adding a new point
-	const bool bWasLoop = IsClosedLoop();
-	UpdateLoopEndpoint(false);
+	const FVector TransformedPosition = (CoordinateSpace == ESplineCoordinateSpace::World) ?
+		ComponentToWorld.InverseTransformPosition(Position) : Position;
 
-	float InputKey = static_cast<float>(SplineInfo.Points.Num());
-	const int32 PointIndex = SplineInfo.AddPoint(InputKey, Position);
-	SplineInfo.Points[PointIndex].InterpMode = CIM_CurveAuto;
+	const float InKey = static_cast<float>(SplineInfo.Points.Num());
 
-	// Then re-close the spline if required
-	UpdateLoopEndpoint(bWasLoop);
+	SplineInfo.Points.Emplace(InKey, TransformedPosition, FVector::ZeroVector, FVector::ZeroVector, CIM_CurveAuto);
+	SplineRotInfo.Points.Emplace(InKey, FQuat::Identity, FQuat::Identity, FQuat::Identity, CIM_CurveAuto);
+	SplineScaleInfo.Points.Emplace(InKey, FVector(1.0f), FVector::ZeroVector, FVector::ZeroVector, CIM_CurveAuto);
 
 	UpdateSpline();
 }
 
 
-void USplineComponent::SetSplineWorldPoints(const TArray<FVector>& Points)
+void USplineComponent::SetSplinePoints(const TArray<FVector>& Points, ESplineCoordinateSpace::Type CoordinateSpace)
 {
-	// If it's a closed loop, mark it as not closed before setting a new array of points
-	const bool bWasLoop = IsClosedLoop();
-	UpdateLoopEndpoint(false);
+	const int32 NumPoints = Points.Num();
+	SplineInfo.Points.Reset(NumPoints);
+	SplineRotInfo.Points.Reset(NumPoints);
+	SplineScaleInfo.Points.Reset(NumPoints);
 
-	SplineInfo.Points.Reset(bWasLoop ? Points.Num() + 1 : Points.Num());
 	float InputKey = 0.0f;
 	for (const auto& Point : Points)
 	{
-		int32 PointIndex = SplineInfo.AddPoint(InputKey, ComponentToWorld.InverseTransformPosition(Point));
-		SplineInfo.Points[PointIndex].InterpMode = CIM_CurveAuto;
+		const FVector TransformedPoint = (CoordinateSpace == ESplineCoordinateSpace::World) ?
+			ComponentToWorld.InverseTransformPosition(Point) : Point;
+
+		SplineInfo.Points.Emplace(InputKey, TransformedPoint, FVector::ZeroVector, FVector::ZeroVector, CIM_CurveAuto);
+		SplineRotInfo.Points.Emplace(InputKey, FQuat::Identity, FQuat::Identity, FQuat::Identity, CIM_CurveAuto);
+		SplineScaleInfo.Points.Emplace(InputKey, FVector(1.0f), FVector::ZeroVector, FVector::ZeroVector, CIM_CurveAuto);
+
 		InputKey += 1.0f;
 	}
 
-	UpdateLoopEndpoint(bWasLoop);
-
 	UpdateSpline();
 }
 
 
-void USplineComponent::SetSplineLocalPoints(const TArray<FVector>& Points)
+void USplineComponent::SetLocationAtSplinePoint(int32 PointIndex, const FVector& InLocation, ESplineCoordinateSpace::Type CoordinateSpace)
 {
-	const bool bWasLoop = IsClosedLoop();
-	UpdateLoopEndpoint(false);
+	const int32 NumPoints = SplineInfo.Points.Num();
 
-	SplineInfo.Points.Reset(bWasLoop ? Points.Num() + 1 : Points.Num());
-	float InputKey = 0.0f;
-	for (const auto& Point : Points)
+	if ((PointIndex >= 0) && (PointIndex < NumPoints))
 	{
-		int32 PointIndex = SplineInfo.AddPoint(InputKey, Point);
-		SplineInfo.Points[PointIndex].InterpMode = CIM_CurveAuto;
-		InputKey += 1.0f;
-	}
+		const FVector TransformedLocation = (CoordinateSpace == ESplineCoordinateSpace::World) ?
+			ComponentToWorld.InverseTransformPosition(InLocation) : InLocation;
 
-	UpdateLoopEndpoint(bWasLoop);
+		SplineInfo.Points[PointIndex].OutVal = TransformedLocation;
 
-	UpdateSpline();
+		UpdateSpline();
+	}
 }
 
 
-void USplineComponent::SetWorldLocationAtSplinePoint(int32 PointIndex, const FVector& InLocation)
+void USplineComponent::SetTangentAtSplinePoint(int32 PointIndex, const FVector& InTangent, ESplineCoordinateSpace::Type CoordinateSpace)
 {
 	const int32 NumPoints = SplineInfo.Points.Num();
 
 	if ((PointIndex >= 0) && (PointIndex < NumPoints))
 	{
-		SplineInfo.Points[PointIndex].OutVal = ComponentToWorld.InverseTransformPosition(InLocation);
+		const FVector TransformedTangent = (CoordinateSpace == ESplineCoordinateSpace::World) ?
+			ComponentToWorld.InverseTransformPosition(InTangent) : InTangent;
 
-		if (IsClosedLoop())
-		{
-			// In a closed loop, the first and last points are tied, so update one with the other
-			if (PointIndex == 0)
-			{
-				SplineInfo.Points[NumPoints - 1].OutVal = InLocation;
-			}
-			else if (PointIndex == NumPoints - 1)
-			{
-				SplineInfo.Points[0].OutVal = InLocation;
-			}
-		}
+		SplineInfo.Points[PointIndex].LeaveTangent = TransformedTangent;
+		SplineInfo.Points[PointIndex].ArriveTangent = TransformedTangent;
+		SplineInfo.Points[PointIndex].InterpMode = CIM_CurveUser;
 
 		UpdateSpline();
 	}
@@ -392,12 +525,7 @@ ESplinePointType::Type USplineComponent::GetSplinePointType(int32 PointIndex) co
 {
 	if ((PointIndex >= 0) && (PointIndex < SplineInfo.Points.Num()))
 	{
-		switch (SplineInfo.Points[PointIndex].InterpMode)
-		{
-			case CIM_CurveAuto:			return ESplinePointType::Curve;
-			case CIM_CurveAutoClamped:	return ESplinePointType::CurveClamped;
-			case CIM_Linear:			return ESplinePointType::Linear;
-		}
+		return ConvertInterpCurveModeToSplinePointType(SplineInfo.Points[PointIndex].InterpMode);
 	}
 
 	return ESplinePointType::Constant;
@@ -406,68 +534,95 @@ ESplinePointType::Type USplineComponent::GetSplinePointType(int32 PointIndex) co
 
 void USplineComponent::SetSplinePointType(int32 PointIndex, ESplinePointType::Type Type)
 {
-	EInterpCurveMode InterpMode;
-	switch (Type)
+	if ((PointIndex >= 0) && (PointIndex < SplineInfo.Points.Num()))
 	{
-		case ESplinePointType::Curve:			InterpMode = CIM_CurveAuto; break;
-		case ESplinePointType::CurveClamped:	InterpMode = CIM_CurveAutoClamped; break;
-		case ESplinePointType::Linear:			InterpMode = CIM_Linear; break;
-		default:								InterpMode = CIM_Constant;
+		SplineInfo.Points[PointIndex].InterpMode = ConvertSplinePointTypeToInterpCurveMode(Type);
+		UpdateSpline();
 	}
+}
 
+
+int32 USplineComponent::GetNumberOfSplinePoints() const
+{
+	// No longer returns an imaginary extra endpoint if closed loop is set
 	const int32 NumPoints = SplineInfo.Points.Num();
+	return NumPoints;
+}
 
-	if ((PointIndex >= 0) && (PointIndex < NumPoints))
-	{
-		SplineInfo.Points[PointIndex].InterpMode = InterpMode;
 
-		if (IsClosedLoop())
-		{
-			// In a closed loop, the first and last points are tied, so update one with the other
-			if (PointIndex == 0)
-			{
-				SplineInfo.Points[NumPoints - 1].InterpMode = InterpMode;
-			}
-			else if (PointIndex == NumPoints - 1)
-			{
-				SplineInfo.Points[0].InterpMode = InterpMode;
-			}
-		}
+FVector USplineComponent::GetLocationAtSplinePoint(int32 PointIndex, ESplineCoordinateSpace::Type CoordinateSpace) const
+{
+	return GetLocationAtSplineInputKey(static_cast<float>(PointIndex), CoordinateSpace);
+}
 
-		UpdateSpline();
-	}
+
+FVector USplineComponent::GetDirectionAtSplinePoint(int32 PointIndex, ESplineCoordinateSpace::Type CoordinateSpace) const
+{
+	return GetLocationAtSplineInputKey(static_cast<float>(PointIndex), CoordinateSpace);
 }
 
 
-int32 USplineComponent::GetNumSplinePoints() const
+FVector USplineComponent::GetTangentAtSplinePoint(int32 PointIndex, ESplineCoordinateSpace::Type CoordinateSpace) const
 {
-	return SplineInfo.Points.Num();
+	return GetTangentAtSplineInputKey(static_cast<float>(PointIndex), CoordinateSpace);
 }
 
-FVector USplineComponent::GetWorldLocationAtSplinePoint(int32 PointIndex) const
+
+FQuat USplineComponent::GetQuaternionAtSplinePoint(int32 PointIndex, ESplineCoordinateSpace::Type CoordinateSpace) const
 {
-	FVector LocalLocation(0,0,0);
-	if ((PointIndex >= 0) && (PointIndex < SplineInfo.Points.Num()))
-	{
-		LocalLocation = SplineInfo.Points[PointIndex].OutVal;
-	}
-	return ComponentToWorld.TransformPosition(LocalLocation);
+	return GetQuaternionAtSplineInputKey(static_cast<float>(PointIndex), CoordinateSpace);
 }
 
-void USplineComponent::GetLocalLocationAndTangentAtSplinePoint(int32 PointIndex, FVector& LocalLocation, FVector& LocalTangent) const
+
+FRotator USplineComponent::GetRotationAtSplinePoint(int32 PointIndex, ESplineCoordinateSpace::Type CoordinateSpace) const
 {
-	LocalTangent = FVector(0, 0, 0);
-	LocalLocation = FVector(0, 0, 0);
-	if ((PointIndex >= 0) && (PointIndex < SplineInfo.Points.Num()))
-	{
-		LocalTangent = SplineInfo.Points[PointIndex].LeaveTangent;
-		LocalLocation = SplineInfo.Points[PointIndex].OutVal;
-	}
+	return GetRotationAtSplineInputKey(static_cast<float>(PointIndex), CoordinateSpace);
+}
+
+
+FVector USplineComponent::GetUpVectorAtSplinePoint(int32 PointIndex, ESplineCoordinateSpace::Type CoordinateSpace) const
+{
+	return GetUpVectorAtSplineInputKey(static_cast<float>(PointIndex), CoordinateSpace);
+}
+
+
+FVector USplineComponent::GetRightVectorAtSplinePoint(int32 PointIndex, ESplineCoordinateSpace::Type CoordinateSpace) const
+{
+	return GetRightVectorAtSplineInputKey(static_cast<float>(PointIndex), CoordinateSpace);
+}
+
+
+float USplineComponent::GetRollAtSplinePoint(int32 PointIndex, ESplineCoordinateSpace::Type CoordinateSpace) const
+{
+	return GetRollAtSplineInputKey(static_cast<float>(PointIndex), CoordinateSpace);
+}
+
+
+FVector USplineComponent::GetScaleAtSplinePoint(int32 PointIndex) const
+{
+	return GetScaleAtSplineInputKey(static_cast<float>(PointIndex));
+}
+
+
+FTransform USplineComponent::GetTransformAtSplinePoint(int32 PointIndex, ESplineCoordinateSpace::Type CoordinateSpace, bool bUseScale) const
+{
+	return GetTransformAtSplineInputKey(static_cast<float>(PointIndex), CoordinateSpace, bUseScale);
+}
+
+
+void USplineComponent::GetLocationAndTangentAtSplinePoint(int32 PointIndex, FVector& Location, FVector& Tangent, ESplineCoordinateSpace::Type CoordinateSpace) const
+{
+	Location = GetLocationAtSplineInputKey(static_cast<float>(PointIndex), CoordinateSpace);
+	Tangent = GetTangentAtSplineInputKey(static_cast<float>(PointIndex), CoordinateSpace);
 }
 
+
 float USplineComponent::GetDistanceAlongSplineAtSplinePoint(int32 PointIndex) const
 {
-	if ((PointIndex >= 0) && (PointIndex < SplineInfo.Points.Num()))
+	const int32 NumPoints = SplineInfo.Points.Num();
+	const int32 NumSegments = bClosedLoop ? NumPoints : NumPoints - 1;
+
+	if ((PointIndex >= 0) && (PointIndex < NumSegments + 1))
 	{
 		return SplineReparamTable.Points[PointIndex * ReparamStepsPerSegment].InVal;
 	}
@@ -475,6 +630,7 @@ float USplineComponent::GetDistanceAlongSplineAtSplinePoint(int32 PointIndex) co
 	return 0.0f;
 }
 
+
 float USplineComponent::GetSplineLength() const
 {
 	const int32 NumPoints = SplineReparamTable.Points.Num();
@@ -482,10 +638,38 @@ float USplineComponent::GetSplineLength() const
 	// This is given by the input of the last entry in the remap table
 	if (NumPoints > 0)
 	{
-		return SplineReparamTable.Points[NumPoints - 1].InVal;
+		return SplineReparamTable.Points.Last().InVal;
 	}
 	
-	return 0.f;
+	return 0.0f;
+}
+
+
+void USplineComponent::SetDefaultUpVector(const FVector& UpVector, ESplineCoordinateSpace::Type CoordinateSpace)
+{
+	if (CoordinateSpace == ESplineCoordinateSpace::World)
+	{
+		DefaultUpVector = ComponentToWorld.InverseTransformVector(UpVector);
+	}
+	else
+	{
+		DefaultUpVector = UpVector;
+	}
+
+	UpdateSpline();
+}
+
+
+FVector USplineComponent::GetDefaultUpVector(ESplineCoordinateSpace::Type CoordinateSpace) const
+{
+	if (CoordinateSpace == ESplineCoordinateSpace::World)
+	{
+		return ComponentToWorld.TransformVector(DefaultUpVector);
+	}
+	else
+	{
+		return DefaultUpVector;
+	}
 }
 
 
@@ -499,42 +683,102 @@ float USplineComponent::GetInputKeyAtDistanceAlongSpline(float Distance) const
 	}
 
 	const float TimeMultiplier = Duration / (NumPoints - 1.0f);
-	return SplineReparamTable.Eval(Distance, 0.f) * TimeMultiplier;
+	return SplineReparamTable.Eval(Distance, 0.0f) * TimeMultiplier;
+}
+
+
+FVector USplineComponent::GetLocationAtDistanceAlongSpline(float Distance, ESplineCoordinateSpace::Type CoordinateSpace) const
+{
+	const float Param = SplineReparamTable.Eval(Distance, 0.0f);
+	return GetLocationAtSplineInputKey(Param, CoordinateSpace);
+}
+
+
+FVector USplineComponent::GetTangentAtDistanceAlongSpline(float Distance, ESplineCoordinateSpace::Type CoordinateSpace) const
+{
+	const float Param = SplineReparamTable.Eval(Distance, 0.0f);
+	return GetTangentAtSplineInputKey(Param, CoordinateSpace);
+}
+
+
+FVector USplineComponent::GetDirectionAtDistanceAlongSpline(float Distance, ESplineCoordinateSpace::Type CoordinateSpace) const
+{
+	const float Param = SplineReparamTable.Eval(Distance, 0.0f);
+	return GetDirectionAtSplineInputKey(Param, CoordinateSpace);
+}
+
+
+FQuat USplineComponent::GetQuaternionAtDistanceAlongSpline(float Distance, ESplineCoordinateSpace::Type CoordinateSpace) const
+{
+	const float Param = SplineReparamTable.Eval(Distance, 0.0f);
+	return GetQuaternionAtSplineInputKey(Param, CoordinateSpace);
+}
+
+
+FRotator USplineComponent::GetRotationAtDistanceAlongSpline(float Distance, ESplineCoordinateSpace::Type CoordinateSpace) const
+{
+	const float Param = SplineReparamTable.Eval(Distance, 0.0f);
+	return GetRotationAtSplineInputKey(Param, CoordinateSpace);
 }
 
 
-FVector USplineComponent::GetWorldLocationAtDistanceAlongSpline(float Distance) const
+FVector USplineComponent::GetUpVectorAtDistanceAlongSpline(float Distance, ESplineCoordinateSpace::Type CoordinateSpace) const
 {
-	const float Param = SplineReparamTable.Eval(Distance, 0.f);
-	const FVector Location = SplineInfo.Eval(Param, FVector::ZeroVector);
-	return ComponentToWorld.TransformPosition(Location);
+	const float Param = SplineReparamTable.Eval(Distance, 0.0f);
+	return GetUpVectorAtSplineInputKey(Param, CoordinateSpace);
 }
 
 
-FVector USplineComponent::GetWorldTangentAtDistanceAlongSpline(float Distance) const
+FVector USplineComponent::GetRightVectorAtDistanceAlongSpline(float Distance, ESplineCoordinateSpace::Type CoordinateSpace) const
 {
-	const float Param = SplineReparamTable.Eval(Distance, 0.f);
-	const FVector Tangent = SplineInfo.EvalDerivative(Param, FVector::ZeroVector);
-	return ComponentToWorld.TransformVector(Tangent);
+	const float Param = SplineReparamTable.Eval(Distance, 0.0f);
+	return GetRightVectorAtSplineInputKey(Param, CoordinateSpace);
 }
 
 
-FVector USplineComponent::GetWorldDirectionAtDistanceAlongSpline(float Distance) const
+float USplineComponent::GetRollAtDistanceAlongSpline(float Distance, ESplineCoordinateSpace::Type CoordinateSpace) const
 {
-	const float Param = SplineReparamTable.Eval(Distance, 0.f);
-	const FVector Tangent = SplineInfo.EvalDerivative(Param, FVector::ZeroVector).GetSafeNormal();
-	return ComponentToWorld.TransformVectorNoScale(Tangent);
+	const float Param = SplineReparamTable.Eval(Distance, 0.0f);
+	return GetRollAtSplineInputKey(Param, CoordinateSpace);
 }
 
 
-FRotator USplineComponent::GetWorldRotationAtDistanceAlongSpline(float Distance) const
+FVector USplineComponent::GetScaleAtDistanceAlongSpline(float Distance) const
 {
-	const FVector Dir = GetWorldDirectionAtDistanceAlongSpline(Distance);
-	return Dir.Rotation();
+	const float Param = SplineReparamTable.Eval(Distance, 0.0f);
+	return GetScaleAtSplineInputKey(Param);
 }
 
 
-FVector USplineComponent::GetWorldLocationAtTime(float Time, bool bUseConstantVelocity) const
+FTransform USplineComponent::GetTransformAtDistanceAlongSpline(float Distance, ESplineCoordinateSpace::Type CoordinateSpace, bool bUseScale) const
+{
+	const float Param = SplineReparamTable.Eval(Distance, 0.0f);
+	return GetTransformAtSplineInputKey(Param, CoordinateSpace, bUseScale);
+}
+
+
+FVector USplineComponent::GetLocationAtTime(float Time, ESplineCoordinateSpace::Type CoordinateSpace, bool bUseConstantVelocity) const
+{
+	if (Duration == 0.0f)
+	{
+		return FVector::ZeroVector;
+	}
+
+	if (bUseConstantVelocity)
+	{
+		return GetLocationAtDistanceAlongSpline(Time / Duration * GetSplineLength(), CoordinateSpace);
+	}
+	else
+	{
+		const int32 NumPoints = SplineInfo.Points.Num();
+		const int32 NumSegments = bClosedLoop ? NumPoints : NumPoints - 1;
+		const float TimeMultiplier = NumSegments / Duration;
+		return GetLocationAtSplineInputKey(Time * TimeMultiplier, CoordinateSpace);
+	}
+}
+
+
+FVector USplineComponent::GetDirectionAtTime(float Time, ESplineCoordinateSpace::Type CoordinateSpace, bool bUseConstantVelocity) const
 {
 	if (Duration == 0.0f)
 	{
@@ -543,16 +787,82 @@ FVector USplineComponent::GetWorldLocationAtTime(float Time, bool bUseConstantVe
 
 	if (bUseConstantVelocity)
 	{
-		return GetWorldLocationAtDistanceAlongSpline(Time / Duration * GetSplineLength());
+		return GetDirectionAtDistanceAlongSpline(Time / Duration * GetSplineLength(), CoordinateSpace);
+	}
+	else
+	{
+		const int32 NumPoints = SplineInfo.Points.Num();
+		const int32 NumSegments = bClosedLoop ? NumPoints : NumPoints - 1;
+		const float TimeMultiplier = NumSegments / Duration;
+		return GetDirectionAtSplineInputKey(Time * TimeMultiplier, CoordinateSpace);
+	}
+}
+
+
+FVector USplineComponent::GetTangentAtTime(float Time, ESplineCoordinateSpace::Type CoordinateSpace, bool bUseConstantVelocity) const
+{
+	if (Duration == 0.0f)
+	{
+		return FVector::ZeroVector;
+	}
+
+	if (bUseConstantVelocity)
+	{
+		return GetTangentAtDistanceAlongSpline(Time / Duration * GetSplineLength(), CoordinateSpace);
+	}
+	else
+	{
+		const int32 NumPoints = SplineInfo.Points.Num();
+		const int32 NumSegments = bClosedLoop ? NumPoints : NumPoints - 1;
+		const float TimeMultiplier = NumSegments / Duration;
+		return GetTangentAtSplineInputKey(Time * TimeMultiplier, CoordinateSpace);
+	}
+}
+
+
+FRotator USplineComponent::GetRotationAtTime(float Time, ESplineCoordinateSpace::Type CoordinateSpace, bool bUseConstantVelocity) const
+{
+	if (Duration == 0.0f)
+	{
+		return FRotator::ZeroRotator;
+	}
+
+	if (bUseConstantVelocity)
+	{
+		return GetRotationAtDistanceAlongSpline(Time / Duration * GetSplineLength(), CoordinateSpace);
+	}
+	else
+	{
+		const int32 NumPoints = SplineInfo.Points.Num();
+		const int32 NumSegments = bClosedLoop ? NumPoints : NumPoints - 1;
+		const float TimeMultiplier = NumSegments / Duration;
+		return GetRotationAtSplineInputKey(Time * TimeMultiplier, CoordinateSpace);
+	}
+}
+
+
+FQuat USplineComponent::GetQuaternionAtTime(float Time, ESplineCoordinateSpace::Type CoordinateSpace, bool bUseConstantVelocity) const
+{
+	if (Duration == 0.0f)
+	{
+		return FQuat::Identity;
 	}
 
-	const float TimeMultiplier = (SplineInfo.Points.Num() - 1.0f) / Duration;
-	const FVector Location = SplineInfo.Eval(Time * TimeMultiplier, FVector::ZeroVector);
-	return ComponentToWorld.TransformPosition(Location);
+	if (bUseConstantVelocity)
+	{
+		return GetQuaternionAtDistanceAlongSpline(Time / Duration * GetSplineLength(), CoordinateSpace);
+	}
+	else
+	{
+		const int32 NumPoints = SplineInfo.Points.Num();
+		const int32 NumSegments = bClosedLoop ? NumPoints : NumPoints - 1;
+		const float TimeMultiplier = NumSegments / Duration;
+		return GetQuaternionAtSplineInputKey(Time * TimeMultiplier, CoordinateSpace);
+	}
 }
 
 
-FVector USplineComponent::GetWorldDirectionAtTime(float Time, bool bUseConstantVelocity) const
+FVector USplineComponent::GetUpVectorAtTime(float Time, ESplineCoordinateSpace::Type CoordinateSpace, bool bUseConstantVelocity) const
 {
 	if (Duration == 0.0f)
 	{
@@ -561,27 +871,98 @@ FVector USplineComponent::GetWorldDirectionAtTime(float Time, bool bUseConstantV
 
 	if (bUseConstantVelocity)
 	{
-		return GetWorldDirectionAtDistanceAlongSpline(Time / Duration * GetSplineLength());
+		return GetUpVectorAtDistanceAlongSpline(Time / Duration * GetSplineLength(), CoordinateSpace);
+	}
+	else
+	{
+		const int32 NumPoints = SplineInfo.Points.Num();
+		const int32 NumSegments = bClosedLoop ? NumPoints : NumPoints - 1;
+		const float TimeMultiplier = NumSegments / Duration;
+		return GetUpVectorAtSplineInputKey(Time * TimeMultiplier, CoordinateSpace);
+	}
+}
+
+
+FVector USplineComponent::GetRightVectorAtTime(float Time, ESplineCoordinateSpace::Type CoordinateSpace, bool bUseConstantVelocity) const
+{
+	if (Duration == 0.0f)
+	{
+		return FVector::ZeroVector;
 	}
 
-	const float TimeMultiplier = (SplineInfo.Points.Num() - 1.0f) / Duration;
-	const FVector Tangent = SplineInfo.EvalDerivative(Time * TimeMultiplier, FVector::ZeroVector).GetSafeNormal();
-	return ComponentToWorld.TransformVectorNoScale(Tangent);
+	if (bUseConstantVelocity)
+	{
+		return GetRightVectorAtDistanceAlongSpline(Time / Duration * GetSplineLength(), CoordinateSpace);
+	}
+	else
+	{
+		const int32 NumPoints = SplineInfo.Points.Num();
+		const int32 NumSegments = bClosedLoop ? NumPoints : NumPoints - 1;
+		const float TimeMultiplier = NumSegments / Duration;
+		return GetRightVectorAtSplineInputKey(Time * TimeMultiplier, CoordinateSpace);
+	}
 }
 
 
-FRotator USplineComponent::GetWorldRotationAtTime(float Time, bool bUseConstantVelocity) const
+float USplineComponent::GetRollAtTime(float Time, ESplineCoordinateSpace::Type CoordinateSpace, bool bUseConstantVelocity) const
 {
-	FVector WorldDir = GetWorldDirectionAtTime(Time, bUseConstantVelocity);
-	return WorldDir.Rotation();
+	if (Duration == 0.0f)
+	{
+		return 0.0f;
+	}
+
+	if (bUseConstantVelocity)
+	{
+		return GetRollAtDistanceAlongSpline(Time / Duration * GetSplineLength(), CoordinateSpace);
+	}
+	else
+	{
+		const int32 NumPoints = SplineInfo.Points.Num();
+		const int32 NumSegments = bClosedLoop ? NumPoints : NumPoints - 1;
+		const float TimeMultiplier = NumSegments / Duration;
+		return GetRollAtSplineInputKey(Time * TimeMultiplier, CoordinateSpace);
+	}
 }
 
 
-void USplineComponent::RefreshSplineInputs()
+FTransform USplineComponent::GetTransformAtTime(float Time, ESplineCoordinateSpace::Type CoordinateSpace, bool bUseConstantVelocity, bool bUseScale) const
 {
-	for(int32 KeyIdx=0; KeyIdx<SplineInfo.Points.Num(); KeyIdx++)
+	if (Duration == 0.0f)
 	{
-		SplineInfo.Points[KeyIdx].InVal = KeyIdx;
+		return FTransform::Identity;
+	}
+
+	if (bUseConstantVelocity)
+	{
+		return GetTransformAtDistanceAlongSpline(Time / Duration * GetSplineLength(), CoordinateSpace, bUseScale);
+	}
+	else
+	{
+		const int32 NumPoints = SplineInfo.Points.Num();
+		const int32 NumSegments = bClosedLoop ? NumPoints : NumPoints - 1;
+		const float TimeMultiplier = NumSegments / Duration;
+		return GetTransformAtSplineInputKey(Time * TimeMultiplier, CoordinateSpace, bUseScale);
+	}
+}
+
+
+FVector USplineComponent::GetScaleAtTime(float Time, bool bUseConstantVelocity) const
+{
+	if (Duration == 0.0f)
+	{
+		return FVector(1.0f);
+	}
+
+	if (bUseConstantVelocity)
+	{
+		return GetScaleAtDistanceAlongSpline(Time / Duration * GetSplineLength());
+	}
+	else
+	{
+		const int32 NumPoints = SplineInfo.Points.Num();
+		const int32 NumSegments = bClosedLoop ? NumPoints : NumPoints - 1;
+		const float TimeMultiplier = NumSegments / Duration;
+		return GetScaleAtSplineInputKey(Time * TimeMultiplier);
 	}
 }
 
@@ -598,102 +979,73 @@ public:
 	virtual void ApplyToComponent(UActorComponent* Component, const ECacheApplyPhase CacheApplyPhase) override
 	{
 		FSceneComponentInstanceData::ApplyToComponent(Component, CacheApplyPhase);
-		CastChecked<USplineComponent>(Component)->ApplyComponentInstanceData(this);
+		CastChecked<USplineComponent>(Component)->ApplyComponentInstanceData(this, (CacheApplyPhase == ECacheApplyPhase::PostUserConstructionScript));
 	}
 
 	FInterpCurveVector SplineInfo;
-	bool bClosedLoop;
+	FInterpCurveQuat SplineRotInfo;
+	FInterpCurveVector SplineScaleInfo;
+	bool bSplineHasBeenEdited;
 };
 
+
 FName USplineComponent::GetComponentInstanceDataType() const
 {
 	static const FName SplineInstanceDataTypeName(TEXT("SplineInstanceData"));
 	return SplineInstanceDataTypeName;
 }
 
+
 FActorComponentInstanceData* USplineComponent::GetComponentInstanceData() const
 {
-	FActorComponentInstanceData* InstanceData = nullptr;
-	if (bAllowSplineEditingPerInstance)
+	FSplineInstanceData* SplineInstanceData = new FSplineInstanceData(this);
+	if (bSplineHasBeenEdited)
 	{
-		FSplineInstanceData* SplineInstanceData = new FSplineInstanceData(this);
 		SplineInstanceData->SplineInfo = SplineInfo;
-		SplineInstanceData->bClosedLoop = bClosedLoop;
-
-		InstanceData = SplineInstanceData;
-	}
-	else
-	{
-		InstanceData = Super::GetComponentInstanceData();
+		SplineInstanceData->SplineRotInfo = SplineRotInfo;
+		SplineInstanceData->SplineScaleInfo = SplineScaleInfo;
 	}
-	return InstanceData;
-}
-
-void USplineComponent::ApplyComponentInstanceData(FSplineInstanceData* SplineInstanceData)
-{
-	if (SplineInstanceData)
-	{
-		if (bAllowSplineEditingPerInstance)
-		{
-			SplineInfo = SplineInstanceData->SplineInfo;
-
-			// If the construction script changed bClosedLoop, amend the applied points accordingly
-			if (SplineInstanceData->bClosedLoop != bClosedLoop)
-			{
-				if (bClosedLoop)
-				{
-					AddLoopEndpoint();
-				}
-				else
-				{
-					RemoveLoopEndpoint();
-				}
-			}
+	SplineInstanceData->bSplineHasBeenEdited = bSplineHasBeenEdited;
 
-			UpdateSpline();
-		}
-	}
+	return SplineInstanceData;
 }
 
-#if WITH_EDITOR
-void USplineComponent::PreEditChange(UProperty* PropertyAboutToChange)
+
+void USplineComponent::ApplyComponentInstanceData(FSplineInstanceData* SplineInstanceData, const bool bPostUCS)
 {
-	Super::PreEditChange(PropertyAboutToChange);
+	check(SplineInstanceData);
 
-	if (PropertyAboutToChange != nullptr && PropertyAboutToChange->GetFName() == GET_MEMBER_NAME_CHECKED(USplineComponent, bClosedLoop))
+	if (SplineInstanceData->bSplineHasBeenEdited)
 	{
-		if (bClosedLoop)
-		{
-			// We need to detect changes in state in order to correctly open or close the loop.
-			// Since we don't know what state is about to be set, we ensure that the spline is marked as not closed here, so that either:
-			//  a) Nothing needs to be done if the property is set to false, or
-			//  b) The spline is already set as not looped, so that setting the property to true will work correctly in PostEditChangeProperty.
-			bClosedLoop = false;
-			RemoveLoopEndpoint();
-		}
+		SplineInfo = SplineInstanceData->SplineInfo;
+		SplineRotInfo = SplineInstanceData->SplineRotInfo;
+		SplineScaleInfo = SplineInstanceData->SplineScaleInfo;
 	}
+
+	bSplineHasBeenEdited = SplineInstanceData->bSplineHasBeenEdited;
+
+	UpdateSpline();
 }
 
-void USplineComponent::PostEditChangeProperty(FPropertyChangedEvent& PropertyChangedEvent)
+
+#if WITH_EDITOR
+void USplineComponent::PostEditChangeChainProperty(FPropertyChangedChainEvent& PropertyChangedEvent)
 {
 	if (PropertyChangedEvent.Property != nullptr)
 	{
-		const FName PropertyName(PropertyChangedEvent.Property->GetFName());
-		if (PropertyName == GET_MEMBER_NAME_CHECKED(USplineComponent, ReparamStepsPerSegment) ||
-			PropertyName == GET_MEMBER_NAME_CHECKED(USplineComponent, bStationaryEndpoints))
-		{
-			UpdateSpline();
-		}
+		static const FName ReparamStepsPerSegmentName = GET_MEMBER_NAME_CHECKED(USplineComponent, ReparamStepsPerSegment);
+		static const FName StationaryEndpointsName = GET_MEMBER_NAME_CHECKED(USplineComponent, bStationaryEndpoints);
+		static const FName DefaultUpVectorName = GET_MEMBER_NAME_CHECKED(USplineComponent, DefaultUpVector);
+		static const FName ClosedLoopName = GET_MEMBER_NAME_CHECKED(USplineComponent, bClosedLoop);
+		static const FName SplineInfoName = GET_MEMBER_NAME_CHECKED(USplineComponent, SplineInfo);
+		static const FName SplineHasBeenEditedName = GET_MEMBER_NAME_CHECKED(USplineComponent, bSplineHasBeenEdited);
 
-		if (PropertyName == GET_MEMBER_NAME_CHECKED(USplineComponent, bClosedLoop))
+		const FName PropertyName(PropertyChangedEvent.Property->GetFName());
+		if (PropertyName == ReparamStepsPerSegmentName ||
+			PropertyName == StationaryEndpointsName ||
+			PropertyName == DefaultUpVectorName ||
+			PropertyName == ClosedLoopName)
 		{
-			if (bClosedLoop)
-			{
-				// Spline is guaranteed to be non-looping when we get here (due to PreEditChange).
-				// Now just force the loop endpoint to be added if bClosedLoop == true.
-				AddLoopEndpoint();
-			}
-
 			UpdateSpline();
 		}
 	}
