Fix non-unity error

diff --git a/Engine/Source/Runtime/Engine/Private/AI/Navigation/NavModifierVolume.cpp b/Engine/Source/Runtime/Engine/Private/AI/Navigation/NavModifierVolume.cpp
--- a/Engine/Source/Runtime/Engine/Private/AI/Navigation/NavModifierVolume.cpp
+++ b/Engine/Source/Runtime/Engine/Private/AI/Navigation/NavModifierVolume.cpp
@@ -3,6 +3,7 @@
 #include "EnginePrivate.h"
 #include "AI/NavigationModifier.h"
 #include "AI/Navigation/NavModifierVolume.h"
+#include "AI/Navigation/NavAreas/NavArea_Null.h"
 
 //----------------------------------------------------------------------//
 // ANavModifierVolume
