Fixed server logspam when client sends an RPC referencing an object that the server destroys before the RPC arrives
Merging 2523621

//depot/UE4-UT/...

to //depot/UE4/...

diff --git a/Engine/Source/Runtime/Engine/Private/PackageMapClient.cpp b/Engine/Source/Runtime/Engine/Private/PackageMapClient.cpp
--- a/Engine/Source/Runtime/Engine/Private/PackageMapClient.cpp
+++ b/Engine/Source/Runtime/Engine/Private/PackageMapClient.cpp
@@ -1916,14 +1916,14 @@ UObject* FNetGUIDCache::GetObjectFromNetGUID( const FNetworkGUID& NetGUID, const
 
 bool FNetGUIDCache::ShouldIgnoreWhenMissing( const FNetworkGUID& NetGUID ) const
 {
-	if ( IsNetGUIDAuthority() )
+	if (NetGUID.IsDynamic())
 	{
-		return false;		// Server never ignores when missing, always warns
+		return true;		// Ignore missing dynamic guids (even on server because client may send RPC on/with object it doesn't know server destroyed)
 	}
 
-	if ( NetGUID.IsDynamic() )
+	if ( IsNetGUIDAuthority() )
 	{
-		return true;		// Ignore missing dynamic guids
+		return false;		// Server never ignores when missing, always warns
 	}
 
 	const FNetGuidCacheObject* CacheObject = ObjectLookup.Find( NetGUID );
