UE4 - log intentional hitches caused by CVarCauseHitches

diff --git a/Engine/Source/Runtime/Engine/Private/UnrealEngine.cpp b/Engine/Source/Runtime/Engine/Private/UnrealEngine.cpp
--- a/Engine/Source/Runtime/Engine/Private/UnrealEngine.cpp
+++ b/Engine/Source/Runtime/Engine/Private/UnrealEngine.cpp
@@ -6165,6 +6165,7 @@ float UEngine::GetMaxTickRate(float DeltaTime, bool bAllowFrameRateSmoothing) co
 		if (RunningHitchTimer > 1.f)
 		{
 			// hitch!
+			UE_LOG(LogEngine, Display, TEXT("Hitching by request!"));
 			FPlatformProcess::Sleep(0.2f);
 			RunningHitchTimer = 0.f;
 		}
