Fix for UE-7846

Mesh InitialOrientation is applied serparately in ParticleSystemRender and so should not be added to the Rotation value.

diff --git a/Engine/Source/Runtime/Engine/Private/Particles/ParticleEmitterInstances.cpp b/Engine/Source/Runtime/Engine/Private/Particles/ParticleEmitterInstances.cpp
--- a/Engine/Source/Runtime/Engine/Private/Particles/ParticleEmitterInstances.cpp
+++ b/Engine/Source/Runtime/Engine/Private/Particles/ParticleEmitterInstances.cpp
@@ -2932,14 +2932,14 @@ void FParticleMeshEmitterInstance::Tick(float DeltaTime, bool bSuppressSpawning)
 
 				FQuat Rotation	= FQuat::FindBetween(OldDirection, NewDirection);
 				FVector Euler	= Rotation.Euler();
-				PayloadData->Rotation = PayloadData->InitialOrientation + PayloadData->InitRotation + Euler;
+				PayloadData->Rotation = PayloadData->InitRotation + Euler;
 				PayloadData->Rotation += PayloadData->CurContinuousRotation;
 			}
 			else // not PSA_Velocity or PSA_AwayfromCenter, so rotation is not reset every tick
 			{
 				if ((Particle.Flags & STATE_Particle_FreezeRotation) == 0)
 				{
-					PayloadData->Rotation = PayloadData->InitialOrientation + PayloadData->InitRotation + PayloadData->CurContinuousRotation;
+					PayloadData->Rotation = PayloadData->InitRotation + PayloadData->CurContinuousRotation;
 				}
 			}
 
@@ -3172,7 +3172,6 @@ void FParticleMeshEmitterInstance::PostSpawn(FBaseParticle* Particle, float Inte
 
 	FVector InitialOrient = MeshTypeData->RollPitchYawRange.GetValue(SpawnTime, 0, 0, &MeshTypeData->RandomStream);
 	PayloadData->InitialOrientation = InitialOrient;
-	PayloadData->Rotation += InitialOrient;
 }
 
 bool FParticleMeshEmitterInstance::IsDynamicDataRequired(UParticleLODLevel* CurrentLODLevel)
