UE4 - support for dithered LOD transitions on static meshes (that don't move) -- also fixed parallel render thread updates for spline meshes

diff --git a/Engine/Shaders/LocalVertexFactory.usf b/Engine/Shaders/LocalVertexFactory.usf
--- a/Engine/Shaders/LocalVertexFactory.usf
+++ b/Engine/Shaders/LocalVertexFactory.usf
@@ -42,9 +42,6 @@ float4 InstancingFadeOutParams;
 	#ifndef USE_INSTANCING_EMULATED
 		#define USE_INSTANCING_EMULATED (0)
 	#endif
-	#if USE_DITHERED_LOD_TRANSITION
-		float NonInstancedDitherLODFactor;
-	#endif
 #endif	// USE_INSTANCING
 
 
diff --git a/Engine/Shaders/MaterialTemplate.usf b/Engine/Shaders/MaterialTemplate.usf
--- a/Engine/Shaders/MaterialTemplate.usf
+++ b/Engine/Shaders/MaterialTemplate.usf
@@ -1433,13 +1433,20 @@ void ClipLODTransition(FMaterialPixelParameters Parameters)
 	ClipLODTransition(Parameters, Parameters.PerInstanceParams.w);
 }
 #elif USE_DITHERED_LOD_TRANSITION
+float NonInstancedDitherLODFactor;
 void ClipLODTransition(FMaterialPixelParameters Parameters)
 {
-	ClipLODTransition(Parameters, NonInstancedDitherLODFactor);
+	if (NonInstancedDitherLODFactor != 0.0)
+	{
+		ClipLODTransition(Parameters, NonInstancedDitherLODFactor);
+	}
 }
 void ClipLODTransition(float2 SVPosition)
 {
-	ClipLODTransition(SVPosition, NonInstancedDitherLODFactor);
+	if (NonInstancedDitherLODFactor != 0.0)
+	{
+		ClipLODTransition(SVPosition, NonInstancedDitherLODFactor);
+	}
 }
 #else
 void ClipLODTransition(FMaterialPixelParameters Parameters)
diff --git a/Engine/Source/Editor/MaterialEditor/Private/MaterialEditor.h b/Engine/Source/Editor/MaterialEditor/Private/MaterialEditor.h
--- a/Engine/Source/Editor/MaterialEditor/Private/MaterialEditor.h
+++ b/Engine/Source/Editor/MaterialEditor/Private/MaterialEditor.h
@@ -113,6 +113,7 @@ public:
 	virtual int32 GetMaterialDomain() const override { return MD_Surface; }
 	virtual FString GetMaterialUsageDescription() const override { return FString::Printf(TEXT("FMatExpressionPreview %s"), Expression.IsValid() ? *Expression->GetName() : TEXT("NULL")); }
 	virtual bool IsTwoSided() const override { return false; }
+	virtual bool IsDitheredLODTransition() const override { return false; }
 	virtual bool IsLightFunction() const override { return false; }
 	virtual bool IsUsedWithDeferredDecal() const override { return false; }
 	virtual bool IsSpecialEngineMaterial() const override { return false; }
diff --git a/Engine/Source/Editor/MaterialEditor/Private/MaterialEditorInstanceDetailCustomization.cpp b/Engine/Source/Editor/MaterialEditor/Private/MaterialEditorInstanceDetailCustomization.cpp
--- a/Engine/Source/Editor/MaterialEditor/Private/MaterialEditorInstanceDetailCustomization.cpp
+++ b/Engine/Source/Editor/MaterialEditor/Private/MaterialEditorInstanceDetailCustomization.cpp
@@ -455,12 +455,14 @@ void FMaterialInstanceParameterDetails::CreateBasePropertyOverrideWidgets(IDetai
 	TAttribute<bool> IsOverrideBlendModeEnabled = TAttribute<bool>::Create(TAttribute<bool>::FGetter::CreateSP(this, &FMaterialInstanceParameterDetails::OverrideBlendModeEnabled));
 	TAttribute<bool> IsOverrideShadingModelEnabled = TAttribute<bool>::Create(TAttribute<bool>::FGetter::CreateSP(this, &FMaterialInstanceParameterDetails::OverrideShadingModelEnabled));
 	TAttribute<bool> IsOverrideTwoSidedEnabled = TAttribute<bool>::Create(TAttribute<bool>::FGetter::CreateSP(this, &FMaterialInstanceParameterDetails::OverrideTwoSidedEnabled));
+	TAttribute<bool> IsOverrideDitheredLODTransitionEnabled = TAttribute<bool>::Create(TAttribute<bool>::FGetter::CreateSP(this, &FMaterialInstanceParameterDetails::OverrideDitheredLODTransitionEnabled));
 
 	TSharedRef<IPropertyHandle> BasePropertyOverridePropery = DetailLayout.GetProperty("BasePropertyOverrides");
 	TSharedPtr<IPropertyHandle> OpacityClipMaskValueProperty = BasePropertyOverridePropery->GetChildHandle("OpacityMaskClipValue");
 	TSharedPtr<IPropertyHandle> BlendModeProperty = BasePropertyOverridePropery->GetChildHandle("BlendMode");
 	TSharedPtr<IPropertyHandle> ShadingModelProperty = BasePropertyOverridePropery->GetChildHandle("ShadingModel");
 	TSharedPtr<IPropertyHandle> TwoSidedProperty = BasePropertyOverridePropery->GetChildHandle("TwoSided");
+	TSharedPtr<IPropertyHandle> DitheredLODTransitionProperty = BasePropertyOverridePropery->GetChildHandle("DitheredLODTransition");
 
 	IDetailPropertyRow& OpacityClipMaskValuePropertyRow = BasePropertyOverrideGroup.AddPropertyRow(OpacityClipMaskValueProperty.ToSharedRef());
 	OpacityClipMaskValuePropertyRow
@@ -485,6 +487,12 @@ void FMaterialInstanceParameterDetails::CreateBasePropertyOverrideWidgets(IDetai
 		.DisplayName(TwoSidedProperty->GetPropertyDisplayName())
 		.ToolTip(TwoSidedProperty->GetToolTipText())
 		.EditCondition(IsOverrideTwoSidedEnabled, FOnBooleanValueChanged::CreateSP(this, &FMaterialInstanceParameterDetails::OnOverrideTwoSidedChanged));
+
+	IDetailPropertyRow& DitheredLODTransitionPropertyRow = BasePropertyOverrideGroup.AddPropertyRow(DitheredLODTransitionProperty.ToSharedRef());
+	DitheredLODTransitionPropertyRow
+		.DisplayName(DitheredLODTransitionProperty->GetPropertyDisplayName())
+		.ToolTip(DitheredLODTransitionProperty->GetToolTipText())
+		.EditCondition(IsOverrideDitheredLODTransitionEnabled, FOnBooleanValueChanged::CreateSP(this, &FMaterialInstanceParameterDetails::OnOverrideDitheredLODTransitionChanged));
 }
 
 bool FMaterialInstanceParameterDetails::OverrideOpacityClipMaskValueEnabled() const
@@ -507,6 +515,11 @@ bool FMaterialInstanceParameterDetails::OverrideTwoSidedEnabled() const
 	return MaterialEditorInstance->BasePropertyOverrides.bOverride_TwoSided;
 }
 
+bool FMaterialInstanceParameterDetails::OverrideDitheredLODTransitionEnabled() const
+{
+	return MaterialEditorInstance->BasePropertyOverrides.bOverride_DitheredLODTransition;
+}
+
 void FMaterialInstanceParameterDetails::OnOverrideOpacityClipMaskValueChanged(bool NewValue)
 {
 	MaterialEditorInstance->BasePropertyOverrides.bOverride_OpacityMaskClipValue = NewValue;
@@ -534,5 +547,12 @@ void FMaterialInstanceParameterDetails::OnOverrideTwoSidedChanged(bool NewValue)
 	MaterialEditorInstance->PostEditChange();
 	FEditorSupportDelegates::RedrawAllViewports.Broadcast();
 }
+
+void FMaterialInstanceParameterDetails::OnOverrideDitheredLODTransitionChanged(bool NewValue)
+{
+	MaterialEditorInstance->BasePropertyOverrides.bOverride_DitheredLODTransition = NewValue;
+	MaterialEditorInstance->PostEditChange();
+	FEditorSupportDelegates::RedrawAllViewports.Broadcast();
+}
 #undef LOCTEXT_NAMESPACE
 
diff --git a/Engine/Source/Editor/MaterialEditor/Private/MaterialEditorInstanceDetailCustomization.h b/Engine/Source/Editor/MaterialEditor/Private/MaterialEditorInstanceDetailCustomization.h
--- a/Engine/Source/Editor/MaterialEditor/Private/MaterialEditorInstanceDetailCustomization.h
+++ b/Engine/Source/Editor/MaterialEditor/Private/MaterialEditorInstanceDetailCustomization.h
@@ -74,10 +74,12 @@ private:
 	bool OverrideBlendModeEnabled() const;
 	bool OverrideShadingModelEnabled() const;
 	bool OverrideTwoSidedEnabled() const;
+	bool OverrideDitheredLODTransitionEnabled() const;
 	void OnOverrideOpacityClipMaskValueChanged(bool NewValue);
 	void OnOverrideBlendModeChanged(bool NewValue);
 	void OnOverrideShadingModelChanged(bool NewValue);
 	void OnOverrideTwoSidedChanged(bool NewValue);
+	void OnOverrideDitheredLODTransitionChanged(bool NewValue);
 
 
 private:
diff --git a/Engine/Source/Editor/UnrealEd/Private/EditorExporters.cpp b/Engine/Source/Editor/UnrealEd/Private/EditorExporters.cpp
--- a/Engine/Source/Editor/UnrealEd/Private/EditorExporters.cpp
+++ b/Engine/Source/Editor/UnrealEd/Private/EditorExporters.cpp
@@ -901,6 +901,14 @@ public:
 		}
 		return false;
 	}
+	virtual bool IsDitheredLODTransition() const  override
+	{ 
+		if (MaterialInterface)
+		{
+			return MaterialInterface->IsDitheredLODTransition();
+		}
+		return false;
+	}
 	virtual bool IsLightFunction() const override
 	{
 		if (Material)
diff --git a/Engine/Source/Editor/UnrealEd/Private/Lightmass/LightmassRender.cpp b/Engine/Source/Editor/UnrealEd/Private/Lightmass/LightmassRender.cpp
--- a/Engine/Source/Editor/UnrealEd/Private/Lightmass/LightmassRender.cpp
+++ b/Engine/Source/Editor/UnrealEd/Private/Lightmass/LightmassRender.cpp
@@ -413,6 +413,14 @@ public:
 		}
 		return false;
 	}
+	virtual bool IsDitheredLODTransition() const override
+	{
+		if (MaterialInterface)
+		{
+			return MaterialInterface->IsDitheredLODTransition();
+		}
+		return false;
+	}
 	virtual bool IsLightFunction() const override
 	{
 		if (Material)
diff --git a/Engine/Source/Runtime/Core/Public/UObject/ObjectVersion.h b/Engine/Source/Runtime/Core/Public/UObject/ObjectVersion.h
--- a/Engine/Source/Runtime/Core/Public/UObject/ObjectVersion.h
+++ b/Engine/Source/Runtime/Core/Public/UObject/ObjectVersion.h
@@ -500,6 +500,8 @@ enum EUnrealEngineObjectUE4Version
 	VER_UE4_SERIALIZE_BLUEPRINT_EVENTGRAPH_FASTCALLS_IN_UFUNCTION,
 	// Changes to USplineComponent and FInterpCurve
 	VER_UE4_SPLINE_COMPONENT_WITH_ROTATION_AND_SCALE,
+	// Material Instances overriding base material LOD transitions
+	VER_UE4_MATERIAL_INSTANCE_BASE_PROPERTY_OVERRIDES_DITHERED_LOD_TRANSITION,
 
 	// -----<new versions can be added before this line>-------------------------------------------------
 	// - this needs to be the last line (see note below)
diff --git a/Engine/Source/Runtime/Engine/Classes/Engine/StaticMesh.h b/Engine/Source/Runtime/Engine/Classes/Engine/StaticMesh.h
--- a/Engine/Source/Runtime/Engine/Classes/Engine/StaticMesh.h
+++ b/Engine/Source/Runtime/Engine/Classes/Engine/StaticMesh.h
@@ -6,7 +6,7 @@
 #include "StaticMesh.generated.h"
 
 /** The maximum number of static mesh LODs allowed. */
-#define MAX_STATIC_MESH_LODS 4
+#define MAX_STATIC_MESH_LODS 8
 
 // Forward declarations
 class UFoliageType_InstancedStaticMesh;
diff --git a/Engine/Source/Runtime/Engine/Classes/Materials/Material.h b/Engine/Source/Runtime/Engine/Classes/Materials/Material.h
--- a/Engine/Source/Runtime/Engine/Classes/Materials/Material.h
+++ b/Engine/Source/Runtime/Engine/Classes/Materials/Material.h
@@ -405,6 +405,10 @@ public:
 	UPROPERTY(EditAnywhere, Category=Material)
 	uint32 TwoSided:1;
 
+	/** Indicates that the material should be rendered without backface culling and the normal should be flipped for backfaces. */
+	UPROPERTY(EditAnywhere, Category=Material)
+	uint32 DitheredLODTransition:1;
+
 	/** Number of customized UV inputs to display.  Unconnected customized UV inputs will just pass through the vertex UVs. */
 	UPROPERTY(EditAnywhere, Category=Material, AdvancedDisplay)
 	int32 NumCustomizedUVs;
@@ -781,6 +785,7 @@ public:
 	ENGINE_API virtual EBlendMode GetBlendMode(bool bIsGameThread = IsInGameThread()) const override;
 	ENGINE_API virtual EMaterialShadingModel GetShadingModel(bool bIsGameThread = IsInGameThread()) const override;
 	ENGINE_API virtual bool IsTwoSided(bool bIsGameThread = IsInGameThread()) const override;
+	ENGINE_API virtual bool IsDitheredLODTransition(bool bIsGameThread = IsInGameThread()) const override;
 	ENGINE_API virtual bool IsMasked(bool bIsGameThread = IsInGameThread()) const override;
 	ENGINE_API virtual USubsurfaceProfile* GetSubsurfaceProfile_Internal() const override;
 
diff --git a/Engine/Source/Runtime/Engine/Classes/Materials/MaterialInstance.h b/Engine/Source/Runtime/Engine/Classes/Materials/MaterialInstance.h
--- a/Engine/Source/Runtime/Engine/Classes/Materials/MaterialInstance.h
+++ b/Engine/Source/Runtime/Engine/Classes/Materials/MaterialInstance.h
@@ -281,12 +281,14 @@ public:
 	ENGINE_API virtual EBlendMode GetBlendMode(bool bIsGameThread = IsInGameThread()) const override;
 	ENGINE_API virtual EMaterialShadingModel GetShadingModel(bool bIsGameThread = IsInGameThread()) const override;
 	ENGINE_API virtual bool IsTwoSided(bool bIsGameThread = IsInGameThread()) const override;
+	ENGINE_API virtual bool IsDitheredLODTransition(bool bIsGameThread = IsInGameThread()) const override;
 	ENGINE_API virtual bool IsMasked(bool bIsGameThread = IsInGameThread()) const override;
 
 	ENGINE_API float RenderThread_GetOpacityMaskClipValue() const;
 	ENGINE_API EBlendMode RenderThread_GetBlendMode() const;
 	ENGINE_API EMaterialShadingModel RenderThread_GetShadingModel() const;
 	ENGINE_API bool RenderThread_IsTwoSided() const;
+	ENGINE_API bool RenderThread_IsDitheredLODTransition() const;
 	ENGINE_API bool RenderThread_IsMasked() const;
 	
 	ENGINE_API virtual USubsurfaceProfile* GetSubsurfaceProfile_Internal() const override;
diff --git a/Engine/Source/Runtime/Engine/Classes/Materials/MaterialInstanceBasePropertyOverrides.h b/Engine/Source/Runtime/Engine/Classes/Materials/MaterialInstanceBasePropertyOverrides.h
--- a/Engine/Source/Runtime/Engine/Classes/Materials/MaterialInstanceBasePropertyOverrides.h
+++ b/Engine/Source/Runtime/Engine/Classes/Materials/MaterialInstanceBasePropertyOverrides.h
@@ -23,6 +23,10 @@ struct ENGINE_API FMaterialInstanceBasePropertyOverrides
 	UPROPERTY(EditAnywhere, Category = Material)
 	bool bOverride_ShadingModel;
 
+	/** Enables override of the two sided property. */
+	UPROPERTY(EditAnywhere, Category = Material)
+	bool bOverride_DitheredLODTransition;
+
 	/** Enables override of the two sided property. */
 	UPROPERTY(EditAnywhere, Category = Material)
 	bool bOverride_TwoSided;
@@ -43,6 +47,10 @@ struct ENGINE_API FMaterialInstanceBasePropertyOverrides
 	UPROPERTY(EditAnywhere, Category = Material, meta = (editcondition = "bOverride_TwoSided"))
 	uint32 TwoSided : 1;
 
+	/** If the material is two sided or not. */
+	UPROPERTY(EditAnywhere, Category = Material, meta = (editcondition = "bOverride_DitheredLODTransition"))
+	uint32 DitheredLODTransition : 1;
+
 	FMaterialInstanceBasePropertyOverrides();
 
 	bool operator==(const FMaterialInstanceBasePropertyOverrides& Other)const;
diff --git a/Engine/Source/Runtime/Engine/Classes/Materials/MaterialInterface.h b/Engine/Source/Runtime/Engine/Classes/Materials/MaterialInterface.h
--- a/Engine/Source/Runtime/Engine/Classes/Materials/MaterialInterface.h
+++ b/Engine/Source/Runtime/Engine/Classes/Materials/MaterialInterface.h
@@ -540,6 +540,7 @@ public:
 	ENGINE_API virtual EBlendMode GetBlendMode(bool bIsGameThread = IsInGameThread()) const;
 	ENGINE_API virtual EMaterialShadingModel GetShadingModel(bool bIsGameThread = IsInGameThread()) const;
 	ENGINE_API virtual bool IsTwoSided(bool bIsGameThread = IsInGameThread()) const;
+	ENGINE_API virtual bool IsDitheredLODTransition(bool bIsGameThread = IsInGameThread()) const;
 	ENGINE_API virtual bool IsMasked(bool bIsGameThread = IsInGameThread()) const;
 
 	ENGINE_API virtual USubsurfaceProfile* GetSubsurfaceProfile_Internal() const;
diff --git a/Engine/Source/Runtime/Engine/Private/InstancedStaticMesh.cpp b/Engine/Source/Runtime/Engine/Private/InstancedStaticMesh.cpp
--- a/Engine/Source/Runtime/Engine/Private/InstancedStaticMesh.cpp
+++ b/Engine/Source/Runtime/Engine/Private/InstancedStaticMesh.cpp
@@ -618,9 +618,9 @@ void FInstancedStaticMeshSceneProxy::SetupInstancedMeshBatch(int32 LODIndex, int
 	}
 }
 
-bool FInstancedStaticMeshSceneProxy::GetShadowMeshElement(int32 LODIndex, int32 BatchIndex, uint8 InDepthPriorityGroup, FMeshBatch& OutMeshBatch) const
+bool FInstancedStaticMeshSceneProxy::GetShadowMeshElement(int32 LODIndex, int32 BatchIndex, uint8 InDepthPriorityGroup, FMeshBatch& OutMeshBatch, bool bDitheredLODTransition) const
 {
-	if (LODIndex < InstancedRenderData.VertexFactories.Num() && FStaticMeshSceneProxy::GetShadowMeshElement(LODIndex, BatchIndex, InDepthPriorityGroup, OutMeshBatch))
+	if (LODIndex < InstancedRenderData.VertexFactories.Num() && FStaticMeshSceneProxy::GetShadowMeshElement(LODIndex, BatchIndex, InDepthPriorityGroup, OutMeshBatch, bDitheredLODTransition))
 	{
 		SetupInstancedMeshBatch(LODIndex, BatchIndex, OutMeshBatch);
 		return true;
diff --git a/Engine/Source/Runtime/Engine/Private/InstancedStaticMesh.h b/Engine/Source/Runtime/Engine/Private/InstancedStaticMesh.h
--- a/Engine/Source/Runtime/Engine/Private/InstancedStaticMesh.h
+++ b/Engine/Source/Runtime/Engine/Private/InstancedStaticMesh.h
@@ -550,7 +550,7 @@ public:
 	virtual int32 GetNumMeshBatches() const override;
 
 	/** Sets up a shadow FMeshBatch for a specific LOD. */
-	virtual bool GetShadowMeshElement(int32 LODIndex, int32 BatchIndex, uint8 InDepthPriorityGroup, FMeshBatch& OutMeshBatch) const override;
+	virtual bool GetShadowMeshElement(int32 LODIndex, int32 BatchIndex, uint8 InDepthPriorityGroup, FMeshBatch& OutMeshBatch, bool bDitheredLODTransition) const override;
 
 	/** Sets up a FMeshBatch for a specific LOD and element. */
 	virtual bool GetMeshElement(int32 LODIndex, int32 BatchIndex, int32 ElementIndex, uint8 InDepthPriorityGroup, const bool bUseSelectedMaterial, const bool bUseHoveredMaterial, FMeshBatch& OutMeshBatch) const override;
diff --git a/Engine/Source/Runtime/Engine/Private/Materials/Material.cpp b/Engine/Source/Runtime/Engine/Private/Materials/Material.cpp
--- a/Engine/Source/Runtime/Engine/Private/Materials/Material.cpp
+++ b/Engine/Source/Runtime/Engine/Private/Materials/Material.cpp
@@ -3755,6 +3755,12 @@ bool UMaterial::IsTwoSided(bool bIsInGameThread) const
 	return TwoSided != 0;
 }
 
+bool UMaterial::IsDitheredLODTransition(bool bIsInGameThread) const
+{
+	return DitheredLODTransition != 0;
+}
+
+
 bool UMaterial::IsMasked(bool bIsInGameThread) const
 {
 	return GetBlendMode() == BLEND_Masked;
diff --git a/Engine/Source/Runtime/Engine/Private/Materials/MaterialInstance.cpp b/Engine/Source/Runtime/Engine/Private/Materials/MaterialInstance.cpp
--- a/Engine/Source/Runtime/Engine/Private/Materials/MaterialInstance.cpp
+++ b/Engine/Source/Runtime/Engine/Private/Materials/MaterialInstance.cpp
@@ -81,6 +81,7 @@ FMaterialInstanceResource::FMaterialInstanceResource(UMaterialInstance* InOwner,
 	, BlendMode(BLEND_Opaque)
 	, ShadingModel(MSM_DefaultLit)
 	, TwoSided(false)
+	, DitheredLODTransition(false)
 {
 }
 
@@ -261,15 +262,18 @@ void FMaterialInstanceResource::GameThread_UpdateOverridableBaseProperties(const
 		*OpacityMaskClipValue = NewOpacityMaskClipValue;
 		*BlendMode = NewBlendMode;
 	});
-	ENQUEUE_UNIQUE_RENDER_COMMAND_FOURPARAMETER(
+	ENQUEUE_UNIQUE_RENDER_COMMAND_SIXPARAMETER(
 		UpdateOverrideablBaseProperties1,
 		EMaterialShadingModel*, ShadingModel, &ShadingModel,
 		EMaterialShadingModel, NewShadingModel, MaterialInterface->GetShadingModel(),
 		bool*, TwoSided, &TwoSided,
 		bool, NewTwoSided, MaterialInterface->IsTwoSided(),
+		bool*, DitheredLODTransition, &DitheredLODTransition,
+		bool, NewDitheredLODTransition, MaterialInterface->IsDitheredLODTransition(),
 		{
 		*ShadingModel = NewShadingModel;
 		*TwoSided = NewTwoSided;
+		*DitheredLODTransition = NewDitheredLODTransition;
 	});
 }
 
@@ -1649,6 +1653,14 @@ void UMaterialInstance::Serialize(FArchive& Ar)
 					Ar << bTwoSided;
 					BasePropertyOverrides.TwoSided = bTwoSided;
 
+					if( Ar.UE4Ver() >= VER_UE4_MATERIAL_INSTANCE_BASE_PROPERTY_OVERRIDES_DITHERED_LOD_TRANSITION )
+					{
+						Ar	<< BasePropertyOverrides.bOverride_DitheredLODTransition;
+
+						bool bDitheredLODTransition;
+						Ar << bDitheredLODTransition;
+						BasePropertyOverrides.DitheredLODTransition = bDitheredLODTransition;
+					}
 					// unrelated but closest change to bug
 					if( Ar.UE4Ver() < VER_UE4_STATIC_SHADOW_DEPTH_MAPS )
 					{
@@ -2458,6 +2470,14 @@ void UMaterialInstance::GetBasePropertyOverridesHash(FSHAHash& OutHash)const
 		Hash.Update((uint8*)&bUsedIsTwoSided, sizeof(bUsedIsTwoSided));
 		bHasOverrides = true;
 	}
+	bool bUsedIsDitheredLODTransition = IsDitheredLODTransition(true);
+	if (bUsedIsDitheredLODTransition != Mat->IsDitheredLODTransition(true))
+	{
+		const FString HashString = TEXT("bOverride_DitheredLODTransition");
+		Hash.UpdateWithString(*HashString, HashString.Len());
+		Hash.Update((uint8*)&bUsedIsDitheredLODTransition, sizeof(bUsedIsDitheredLODTransition));
+		bHasOverrides = true;
+	}
 
  	if (bHasOverrides)
  	{
@@ -2474,7 +2494,9 @@ bool UMaterialInstance::HasOverridenBaseProperties()const
 		(FMath::Abs(GetOpacityMaskClipValue(true) - Parent->GetOpacityMaskClipValue(true)) > SMALL_NUMBER) ||
 		(GetBlendMode(true) != Parent->GetBlendMode(true)) ||
 		(GetShadingModel(true) != Parent->GetShadingModel(true)) ||
-		(IsTwoSided(true) != Parent->IsTwoSided(true)))
+		(IsTwoSided(true) != Parent->IsTwoSided(true)) ||
+		(IsDitheredLODTransition(true) != Parent->IsDitheredLODTransition(true))
+		)
 		)
 	{
 		return true;
@@ -2565,12 +2587,34 @@ bool UMaterialInstance::IsTwoSided(bool bIsInGameThread) const
 	return RenderThread_IsTwoSided();
 }
 
+bool UMaterialInstance::IsDitheredLODTransition(bool bIsInGameThread) const
+{
+	if (bIsInGameThread)
+	{
+		if (BasePropertyOverrides.bOverride_DitheredLODTransition)
+		{
+			return BasePropertyOverrides.DitheredLODTransition != 0;
+		}
+		// go up the chain if possible
+		return Parent ? Parent->IsDitheredLODTransition(true) : false;
+	}
+
+	//Get the value mirrored in the render proxy.
+	return RenderThread_IsDitheredLODTransition();
+}
+
 bool UMaterialInstance::RenderThread_IsTwoSided() const
 {
 	FMaterialInstanceResource* Proxy = ((FMaterialInstanceResource*)GetRenderProxy(0));
 	return Proxy ? Proxy->IsTwoSided() : false;
 }
 
+bool UMaterialInstance::RenderThread_IsDitheredLODTransition() const
+{
+	FMaterialInstanceResource* Proxy = ((FMaterialInstanceResource*)GetRenderProxy(0));
+	return Proxy ? Proxy->IsDitheredLODTransition() : false;
+}
+
 bool UMaterialInstance::IsMasked(bool bIsInGameThread) const
 {
 	if (bIsInGameThread)
diff --git a/Engine/Source/Runtime/Engine/Private/Materials/MaterialInstanceSupport.h b/Engine/Source/Runtime/Engine/Private/Materials/MaterialInstanceSupport.h
--- a/Engine/Source/Runtime/Engine/Private/Materials/MaterialInstanceSupport.h
+++ b/Engine/Source/Runtime/Engine/Private/Materials/MaterialInstanceSupport.h
@@ -159,6 +159,7 @@ public:
 	EBlendMode GetBlendMode()const{ return BlendMode; }
 	EMaterialShadingModel GetShadingModel()const{ return ShadingModel; }
 	bool IsTwoSided()const{ return TwoSided; }
+	bool IsDitheredLODTransition()const{ return DitheredLODTransition; }
 
 private:
 	/**
@@ -194,6 +195,7 @@ private:
 	EBlendMode BlendMode;
 	EMaterialShadingModel ShadingModel;
 	bool TwoSided;
+	bool DitheredLODTransition;
 };
 
 template <> FORCEINLINE TArray<FMaterialInstanceResource::TNamedParameter<float> >& FMaterialInstanceResource::GetValueArray() { return ScalarParameterArray; }
diff --git a/Engine/Source/Runtime/Engine/Private/Materials/MaterialInterface.cpp b/Engine/Source/Runtime/Engine/Private/Materials/MaterialInterface.cpp
--- a/Engine/Source/Runtime/Engine/Private/Materials/MaterialInterface.cpp
+++ b/Engine/Source/Runtime/Engine/Private/Materials/MaterialInterface.cpp
@@ -250,6 +250,11 @@ bool UMaterialInterface::IsTwoSided(bool bIsInGameThread) const
 	return false;
 }
 
+bool UMaterialInterface::IsDitheredLODTransition(bool bIsInGameThread) const
+{
+	return false;
+}
+
 bool UMaterialInterface::IsMasked(bool bIsInGameThread) const
 {
 	return false;
diff --git a/Engine/Source/Runtime/Engine/Private/Materials/MaterialShared.cpp b/Engine/Source/Runtime/Engine/Private/Materials/MaterialShared.cpp
--- a/Engine/Source/Runtime/Engine/Private/Materials/MaterialShared.cpp
+++ b/Engine/Source/Runtime/Engine/Private/Materials/MaterialShared.cpp
@@ -844,6 +844,11 @@ bool FMaterialResource::IsTwoSided() const
 	return MaterialInstance ? MaterialInstance->IsTwoSided() : Material->IsTwoSided();
 }
 
+bool FMaterialResource::IsDitheredLODTransition() const 
+{
+	return MaterialInstance ? MaterialInstance->IsDitheredLODTransition() : Material->IsDitheredLODTransition();
+}
+
 bool FMaterialResource::IsMasked() const 
 {
 	return MaterialInstance ? MaterialInstance->IsMasked() : Material->IsMasked();
@@ -1139,13 +1144,13 @@ void FMaterial::SetupMaterialEnvironment(
 
 	switch(GetBlendMode())
 	{
-	case BLEND_Opaque: OutEnvironment.SetDefine(TEXT("MATERIALBLENDING_SOLID"),TEXT("1")); break;
+	case BLEND_Opaque:
 	case BLEND_Masked:
 		{
 			// Only set MATERIALBLENDING_MASKED if the material is truly masked
 			//@todo - this may cause mismatches with what the shader compiles and what the renderer thinks the shader needs
 			// For example IsTranslucentBlendMode doesn't check IsMasked
-			if(IsMasked())
+			if(!WritesEveryPixel())
 			{
 				OutEnvironment.SetDefine(TEXT("MATERIALBLENDING_MASKED"),TEXT("1"));
 			}
@@ -1186,6 +1191,7 @@ void FMaterial::SetupMaterialEnvironment(
 		OutEnvironment.SetDefine(TEXT("MATERIALDECALRESPONSEMASK"), MaterialDecalResponseMask);
 	}
 
+	OutEnvironment.SetDefine(TEXT("USE_DITHERED_LOD_TRANSITION"), IsDitheredLODTransition() ? TEXT("1") : TEXT("0"));
 	OutEnvironment.SetDefine(TEXT("MATERIAL_TWOSIDED"), IsTwoSided() ? TEXT("1") : TEXT("0"));
 	OutEnvironment.SetDefine(TEXT("MATERIAL_TANGENTSPACENORMAL"), IsTangentSpaceNormal() ? TEXT("1") : TEXT("0"));
 	OutEnvironment.SetDefine(TEXT("GENERATE_SPHERICAL_PARTICLE_NORMALS"),ShouldGenerateSphericalParticleNormals() ? TEXT("1") : TEXT("0"));
@@ -2327,10 +2333,12 @@ FMaterialInstanceBasePropertyOverrides::FMaterialInstanceBasePropertyOverrides()
 	,bOverride_BlendMode(false)
 	,bOverride_ShadingModel(false)
 	,bOverride_TwoSided(false)
+	,bOverride_DitheredLODTransition(false)
 	,OpacityMaskClipValue(.333333f)
 	,BlendMode(BLEND_Opaque)
 	,ShadingModel(MSM_DefaultLit)
 	,TwoSided(0)
+	,DitheredLODTransition(0)
 {
 
 }
@@ -2341,10 +2349,12 @@ bool FMaterialInstanceBasePropertyOverrides::operator==(const FMaterialInstanceB
 			bOverride_BlendMode == Other.bOverride_BlendMode &&
 			bOverride_ShadingModel == Other.bOverride_ShadingModel &&
 			bOverride_TwoSided == Other.bOverride_TwoSided &&
+			bOverride_DitheredLODTransition == Other.bOverride_DitheredLODTransition &&
 			OpacityMaskClipValue == Other.OpacityMaskClipValue &&
 			BlendMode == Other.BlendMode &&
 			ShadingModel == Other.ShadingModel &&
-			TwoSided == Other.TwoSided;
+			TwoSided == Other.TwoSided &&
+			DitheredLODTransition == Other.DitheredLODTransition;
 }
 
 bool FMaterialInstanceBasePropertyOverrides::operator!=(const FMaterialInstanceBasePropertyOverrides& Other)const
diff --git a/Engine/Source/Runtime/Engine/Private/SceneManagement.cpp b/Engine/Source/Runtime/Engine/Private/SceneManagement.cpp
--- a/Engine/Source/Runtime/Engine/Private/SceneManagement.cpp
+++ b/Engine/Source/Runtime/Engine/Private/SceneManagement.cpp
@@ -277,6 +277,38 @@ float ComputeBoundsScreenSize( const FVector4& Origin, const float SphereRadius,
 	return FMath::Clamp(ScreenArea / View.ViewRect.Area(), 0.0f, 1.0f);
 }
 
+float ComputeTemporalLODBoundsScreenSize( const FVector& Origin, const float SphereRadius, const FSceneView& View, int32 SampleIndex )
+{
+	// This is radial LOD, not the view parallel computation used in ComputeBoundsScreenSize
+	const float Divisor =  (Origin - View.GetTemporalLODOrigin(SampleIndex)).Size();
+
+	// Get projection multiple accounting for view scaling.
+	const float ScreenMultiple = FMath::Max(View.ViewRect.Width() / 2.0f * View.ViewMatrices.ProjMatrix.M[0][0],
+		View.ViewRect.Height() / 2.0f * View.ViewMatrices.ProjMatrix.M[1][1]);
+
+	const float ScreenRadius = ScreenMultiple * SphereRadius / FMath::Max(Divisor, 1.0f);
+	const float ScreenArea = PI * ScreenRadius * ScreenRadius;
+	return FMath::Clamp(ScreenArea / View.ViewRect.Area(), 0.0f, 1.0f);
+}
+
+int8 ComputeTemporalStaticMeshLOD( const FStaticMeshRenderData* RenderData, const FVector4& Origin, const float SphereRadius, const FSceneView& View, int32 MinLOD, float FactorScale, int32 SampleIndex )
+{
+	const int32 NumLODs = MAX_STATIC_MESH_LODS;
+
+	const float ScreenSize = ComputeTemporalLODBoundsScreenSize(Origin, SphereRadius, View, SampleIndex) * FactorScale;
+
+	// Walk backwards and return the first matching LOD
+	for(int32 LODIndex = NumLODs - 1 ; LODIndex >= 0 ; --LODIndex)
+	{
+		if(RenderData->ScreenSize[LODIndex] > ScreenSize)
+		{
+			return FMath::Max(LODIndex, MinLOD);
+		}
+	}
+
+	return MinLOD;
+}
+
 int8 ComputeStaticMeshLOD( const FStaticMeshRenderData* RenderData, const FVector4& Origin, const float SphereRadius, const FSceneView& View, int32 MinLOD, float FactorScale )
 {
 	const int32 NumLODs = MAX_STATIC_MESH_LODS;
@@ -295,9 +327,9 @@ int8 ComputeStaticMeshLOD( const FStaticMeshRenderData* RenderData, const FVecto
 	return MinLOD;
 }
 
-int8 ComputeLODForMeshes( const TIndirectArray<class FStaticMesh>& StaticMeshes, const FSceneView& View, const FVector4& Origin, float SphereRadius, int32 ForcedLODLevel, float ScreenSizeScale )
+FLODMask ComputeLODForMeshes( const TIndirectArray<class FStaticMesh>& StaticMeshes, const FSceneView& View, const FVector4& Origin, float SphereRadius, int32 ForcedLODLevel, float ScreenSizeScale )
 {
-	int8 LODToRender = 0;
+	FLODMask LODToRender;
 
 	// Handle forced LOD level first
 	if(ForcedLODLevel >= 0)
@@ -308,36 +340,68 @@ int8 ComputeLODForMeshes( const TIndirectArray<class FStaticMesh>& StaticMeshes,
 			const FStaticMesh&  Mesh = StaticMeshes[MeshIndex];
 			MaxLOD = FMath::Max(MaxLOD, Mesh.LODIndex);
 		}
-		LODToRender = FMath::Clamp<int8>(ForcedLODLevel, 0, MaxLOD);
+		LODToRender.SetLOD(FMath::Clamp<int8>(ForcedLODLevel, 0, MaxLOD));
 	}
 	else if (View.Family->EngineShowFlags.LOD)
 	{
-		int32 MinLODFound = INT_MAX;
-		bool bFoundLOD = false;
 		int32 NumMeshes = StaticMeshes.Num();
 
-		const float ScreenSize = ComputeBoundsScreenSize(Origin, SphereRadius, View);
-
-		for(int32 MeshIndex = NumMeshes-1 ; MeshIndex >= 0 ; --MeshIndex)
+		if (NumMeshes && StaticMeshes[0].bDitheredLODTransition)
 		{
-			const FStaticMesh& Mesh = StaticMeshes[MeshIndex];
-
-			float MeshScreenSize = Mesh.ScreenSize * ScreenSizeScale;
-
-			if(MeshScreenSize >= ScreenSize)
+			for (int32 SampleIndex = 0; SampleIndex < 2; SampleIndex++)
 			{
-				LODToRender = Mesh.LODIndex;
-				bFoundLOD = true;
-				break;
+				int32 MinLODFound = INT_MAX;
+				bool bFoundLOD = false;
+				const float ScreenSize = ComputeTemporalLODBoundsScreenSize(Origin, SphereRadius, View, SampleIndex);
+
+				for(int32 MeshIndex = NumMeshes-1 ; MeshIndex >= 0 ; --MeshIndex)
+				{
+					const FStaticMesh& Mesh = StaticMeshes[MeshIndex];
+
+					float MeshScreenSize = Mesh.ScreenSize * ScreenSizeScale;
+
+					if(MeshScreenSize >= ScreenSize)
+					{
+						LODToRender.SetLODSample(Mesh.LODIndex, SampleIndex);
+						bFoundLOD = true;
+						break;
+					}
+
+					MinLODFound = FMath::Min<int32>(MinLODFound, Mesh.LODIndex);
+				}
+				// If no LOD was found matching the screen size, use the lowest in the array instead of LOD 0, to handle non-zero MinLOD
+				if (!bFoundLOD)
+				{
+					LODToRender.SetLODSample(MinLODFound, SampleIndex);
+				}
 			}
-
-			MinLODFound = FMath::Min<int32>(MinLODFound, Mesh.LODIndex);
 		}
-
-		// If no LOD was found matching the screen size, use the lowest in the array instead of LOD 0, to handle non-zero MinLOD
-		if (!bFoundLOD)
+		else
 		{
-			LODToRender = MinLODFound;
+			int32 MinLODFound = INT_MAX;
+			bool bFoundLOD = false;
+			const float ScreenSize = ComputeBoundsScreenSize(Origin, SphereRadius, View);
+
+			for(int32 MeshIndex = NumMeshes-1 ; MeshIndex >= 0 ; --MeshIndex)
+			{
+				const FStaticMesh& Mesh = StaticMeshes[MeshIndex];
+
+				float MeshScreenSize = Mesh.ScreenSize * ScreenSizeScale;
+
+				if(MeshScreenSize >= ScreenSize)
+				{
+					LODToRender.SetLOD(Mesh.LODIndex);
+					bFoundLOD = true;
+					break;
+				}
+
+				MinLODFound = FMath::Min<int32>(MinLODFound, Mesh.LODIndex);
+			}
+			// If no LOD was found matching the screen size, use the lowest in the array instead of LOD 0, to handle non-zero MinLOD
+			if (!bFoundLOD)
+			{
+				LODToRender.SetLOD(MinLODFound);
+			}
 		}
 	}
 	return LODToRender;
diff --git a/Engine/Source/Runtime/Engine/Private/SplineMeshSceneProxy.cpp b/Engine/Source/Runtime/Engine/Private/SplineMeshSceneProxy.cpp
--- a/Engine/Source/Runtime/Engine/Private/SplineMeshSceneProxy.cpp
+++ b/Engine/Source/Runtime/Engine/Private/SplineMeshSceneProxy.cpp
@@ -13,7 +13,7 @@ FSplineMeshSceneProxy::FSplineMeshSceneProxy(USplineMeshComponent* InComponent)
 	{
 		for (FStaticMeshSceneProxy::FLODInfo::FSectionInfo& Section : LODInfo.Sections)
 		{
-			if (!Section.Material->CheckMaterialUsage(MATUSAGE_SplineMesh))
+			if (!Section.Material->CheckMaterialUsage_Concurrent(MATUSAGE_SplineMesh))
 			{
 				Section.Material = UMaterial::GetDefaultMaterial(MD_Surface);
 			}
@@ -62,11 +62,11 @@ FSplineMeshSceneProxy::~FSplineMeshSceneProxy()
 	LODResources.Empty();
 }
 
-bool FSplineMeshSceneProxy::GetShadowMeshElement(int32 LODIndex, int32 BatchIndex, uint8 InDepthPriorityGroup, FMeshBatch& OutMeshBatch) const
+bool FSplineMeshSceneProxy::GetShadowMeshElement(int32 LODIndex, int32 BatchIndex, uint8 InDepthPriorityGroup, FMeshBatch& OutMeshBatch, bool bDitheredLODTransition) const
 {
 	//checkf(LODIndex == 0, TEXT("Getting spline static mesh element with invalid LOD [%d]"), LODIndex);
 
-	if (FStaticMeshSceneProxy::GetShadowMeshElement(LODIndex, BatchIndex, InDepthPriorityGroup, OutMeshBatch))
+	if (FStaticMeshSceneProxy::GetShadowMeshElement(LODIndex, BatchIndex, InDepthPriorityGroup, OutMeshBatch, bDitheredLODTransition))
 	{
 		OutMeshBatch.VertexFactory = LODResources[LODIndex].VertexFactory;
 		OutMeshBatch.ReverseCulling ^= (SplineParams.StartScale.X < 0) ^ (SplineParams.StartScale.Y < 0);
diff --git a/Engine/Source/Runtime/Engine/Private/SplineMeshSceneProxy.h b/Engine/Source/Runtime/Engine/Private/SplineMeshSceneProxy.h
--- a/Engine/Source/Runtime/Engine/Private/SplineMeshSceneProxy.h
+++ b/Engine/Source/Runtime/Engine/Private/SplineMeshSceneProxy.h
@@ -154,7 +154,7 @@ public:
 	void ReleaseResources();
 
 	/** Sets up a shadow FMeshBatch for a specific LOD. */
-	virtual bool GetShadowMeshElement(int32 LODIndex, int32 BatchIndex, uint8 InDepthPriorityGroup, FMeshBatch& OutMeshBatch) const override;
+	virtual bool GetShadowMeshElement(int32 LODIndex, int32 BatchIndex, uint8 InDepthPriorityGroup, FMeshBatch& OutMeshBatch, bool bDitheredLODTransition) const override;
 
 	/** Sets up a FMeshBatch for a specific LOD and element. */
 	virtual bool GetMeshElement(int32 LODIndex, int32 BatchIndex, int32 ElementIndex, uint8 InDepthPriorityGroup, const bool bUseSelectedMaterial, const bool bUseHoveredMaterial, FMeshBatch& OutMeshBatch) const override;
diff --git a/Engine/Source/Runtime/Engine/Private/StaticMeshRender.cpp b/Engine/Source/Runtime/Engine/Private/StaticMeshRender.cpp
--- a/Engine/Source/Runtime/Engine/Private/StaticMeshRender.cpp
+++ b/Engine/Source/Runtime/Engine/Private/StaticMeshRender.cpp
@@ -6,6 +6,8 @@
 
 #include "EnginePrivate.h"
 #include "StaticMeshResources.h"
+
+
 #include "LevelUtils.h"
 #include "TessellationRendering.h"
 #include "LightMap.h"
@@ -163,7 +165,7 @@ void UStaticMeshComponent::SetLODDataCount( const uint32 MinSize, const uint32 M
 	}
 }
 
-bool FStaticMeshSceneProxy::GetShadowMeshElement(int32 LODIndex, int32 BatchIndex, uint8 InDepthPriorityGroup, FMeshBatch& OutMeshBatch) const
+bool FStaticMeshSceneProxy::GetShadowMeshElement(int32 LODIndex, int32 BatchIndex, uint8 InDepthPriorityGroup, FMeshBatch& OutMeshBatch, bool bDitheredLODTransition) const
 {
 	const FStaticMeshLODResources& LOD = RenderData->LODResources[LODIndex];
 	const FLODInfo& ProxyLODInfo = LODs[LODIndex];
@@ -190,6 +192,7 @@ bool FStaticMeshSceneProxy::GetShadowMeshElement(int32 LODIndex, int32 BatchInde
 	}
 	else
 	{
+		OutMeshBatch.bDitheredLODTransition = bDitheredLODTransition;
 		OutBatchElement.MaxScreenSize = GetScreenSize(LODIndex);
 		OutBatchElement.MinScreenSize = 0.0f;
 		if (LODIndex < MAX_STATIC_MESH_LODS - 1)
@@ -252,6 +255,9 @@ bool FStaticMeshSceneProxy::GetMeshElement(int32 LODIndex, int32 BatchIndex, int
 		}
 		else
 		{
+			// no support for stateless dithered LOD transitions for movable meshes
+			OutMeshBatch.bDitheredLODTransition = !IsMovable() && Material->IsDitheredLODTransition(false);
+
 			OutBatchElement.MaxScreenSize = GetScreenSize(LODIndex);
 			OutBatchElement.MinScreenSize = 0.0f;
 			if (LODIndex < MAX_STATIC_MESH_LODS - 1)
@@ -484,13 +490,20 @@ void FStaticMeshSceneProxy::DrawStaticElements(FStaticPrimitiveDrawInterface* PD
 					bool bSafeToUseShadowOnlyMesh = AllowShadowOnlyMesh(FeatureLevel);
 
 					bool bAnySectionCastsShadow = false;
+					bool bAnySectionUsesDitheredLODTransition = false;
+					bool bAllSectionsUseDitheredLODTransition = true;
+					bool bIsMovable = IsMovable();
 					for (int32 SectionIndex = 0; bSafeToUseShadowOnlyMesh && SectionIndex < LODModel.Sections.Num(); SectionIndex++)
 					{
 						const FMaterial* Material = ProxyLODInfo.Sections[SectionIndex].Material->GetRenderProxy(false)->GetMaterial(FeatureLevel);
+						// no support for stateless dithered LOD transitions for movable meshes
+						bAnySectionUsesDitheredLODTransition = bAnySectionUsesDitheredLODTransition || (!bIsMovable && Material->IsDitheredLODTransition());
+						bAllSectionsUseDitheredLODTransition = bAllSectionsUseDitheredLODTransition && (!bIsMovable && Material->IsDitheredLODTransition());
 						const FStaticMeshSection& Section = LODModel.Sections[SectionIndex];
 						bSafeToUseShadowOnlyMesh =
 							Section.bCastShadow
-							&& Material->GetBlendMode() == BLEND_Opaque
+							&& !(bAnySectionUsesDitheredLODTransition && !bAllSectionsUseDitheredLODTransition) // can't use a single section if they are not homogeneous
+							&& Material->WritesEveryPixel()
 							&& !Material->IsTwoSided()
 							&& !Material->MaterialModifiesMeshPosition_RenderThread();
 						bAnySectionCastsShadow |= Section.bCastShadow;
@@ -504,7 +517,7 @@ void FStaticMeshSceneProxy::DrawStaticElements(FStaticPrimitiveDrawInterface* PD
 						{
 							FMeshBatch MeshBatch;
 
-							if (GetShadowMeshElement(LODIndex, BatchIndex, PrimitiveDPG, MeshBatch))
+							if (GetShadowMeshElement(LODIndex, BatchIndex, PrimitiveDPG, MeshBatch, bAllSectionsUseDitheredLODTransition))
 							{
 								bHaveShadowOnlyMesh = true;
 								PDI->DrawMesh(MeshBatch, ScreenSize, /*bShadowOnly=*/true);
@@ -604,166 +617,184 @@ void FStaticMeshSceneProxy::GetDynamicMeshElements(const TArray<const FSceneView
 			if (VisibilityMap & (1 << ViewIndex))
 			{
 				const FSceneView* View = Views[ViewIndex];
-				const int32 LODIndex = GetLOD(View);
+				FLODMask LODMask = GetLODMask(View);
 
-				if (RenderData->LODResources.IsValidIndex(LODIndex))
+				for (int32 LODIndex = 0; LODIndex < RenderData->LODResources.Num(); LODIndex++)
 				{
-					const FStaticMeshLODResources& LODModel = RenderData->LODResources[LODIndex];
-					const FLODInfo& ProxyLODInfo = LODs[LODIndex];
-
-					if (AllowDebugViewmodes() && bIsWireframeView && !EngineShowFlags.Materials
-						// If any of the materials are mesh-modifying, we can't use the single merged mesh element of GetWireframeMeshElement()
-						&& !ProxyLODInfo.UsesMeshModifyingMaterials())
+					if (LODMask.ContainsLOD(LODIndex))
 					{
-						FLinearColor ViewWireframeColor( bLevelColorationEnabled ? LevelColor : WireframeColor );
-						if ( bPropertyColorationEnabled )
-						{
-							ViewWireframeColor = PropertyColor;
-						}
+						const FStaticMeshLODResources& LODModel = RenderData->LODResources[LODIndex];
+						const FLODInfo& ProxyLODInfo = LODs[LODIndex];
 
-						auto WireframeMaterialInstance = new FColoredMaterialRenderProxy(
-							GEngine->WireframeMaterial->GetRenderProxy(false),
-							GetSelectionColor(ViewWireframeColor,!(GIsEditor && EngineShowFlags.Selection) || bProxyIsSelected, IsHovered(), false)
-							);
+						if (AllowDebugViewmodes() && bIsWireframeView && !EngineShowFlags.Materials
+							// If any of the materials are mesh-modifying, we can't use the single merged mesh element of GetWireframeMeshElement()
+							&& !ProxyLODInfo.UsesMeshModifyingMaterials())
+						{
+							FLinearColor ViewWireframeColor( bLevelColorationEnabled ? LevelColor : WireframeColor );
+							if ( bPropertyColorationEnabled )
+							{
+								ViewWireframeColor = PropertyColor;
+							}
 
-						Collector.RegisterOneFrameMaterialProxy(WireframeMaterialInstance);
+							auto WireframeMaterialInstance = new FColoredMaterialRenderProxy(
+								GEngine->WireframeMaterial->GetRenderProxy(false),
+								GetSelectionColor(ViewWireframeColor,!(GIsEditor && EngineShowFlags.Selection) || bProxyIsSelected, IsHovered(), false)
+								);
 
-						const int32 NumBatches = GetNumMeshBatches();
+							Collector.RegisterOneFrameMaterialProxy(WireframeMaterialInstance);
 
-						for (int32 BatchIndex = 0; BatchIndex < NumBatches; BatchIndex++)
-						{
-							FMeshBatch& Mesh = Collector.AllocateMesh();
+							const int32 NumBatches = GetNumMeshBatches();
 
-							if (GetWireframeMeshElement(LODIndex, BatchIndex, WireframeMaterialInstance, SDPG_World, Mesh))
+							for (int32 BatchIndex = 0; BatchIndex < NumBatches; BatchIndex++)
 							{
-								// We implemented our own wireframe
-								Mesh.bCanApplyViewModeOverrides = false;
-								Collector.AddMesh(ViewIndex, Mesh);
-								INC_DWORD_STAT_BY(STAT_StaticMeshTriangles,Mesh.GetNumPrimitives());
+								FMeshBatch& Mesh = Collector.AllocateMesh();
+
+								if (GetWireframeMeshElement(LODIndex, BatchIndex, WireframeMaterialInstance, SDPG_World, Mesh))
+								{
+									// We implemented our own wireframe
+									Mesh.bCanApplyViewModeOverrides = false;
+									Collector.AddMesh(ViewIndex, Mesh);
+									INC_DWORD_STAT_BY(STAT_StaticMeshTriangles,Mesh.GetNumPrimitives());
+								}
 							}
 						}
-					}
-					else
-					{
-						const FLinearColor UtilColor( LevelColor );
-
-						// Draw the static mesh sections.
-						for (int32 SectionIndex = 0; SectionIndex < LODModel.Sections.Num(); SectionIndex++)
+						else
 						{
-							const int32 NumBatches = GetNumMeshBatches();
+							const FLinearColor UtilColor( LevelColor );
 
-							for (int32 BatchIndex = 0; BatchIndex < NumBatches; BatchIndex++)
+							// Draw the static mesh sections.
+							for (int32 SectionIndex = 0; SectionIndex < LODModel.Sections.Num(); SectionIndex++)
 							{
-								bool bSectionIsSelected = false;
-								FMeshBatch& MeshElement = Collector.AllocateMesh();
+								const int32 NumBatches = GetNumMeshBatches();
 
-#if WITH_EDITOR
-								if (GIsEditor)
+								for (int32 BatchIndex = 0; BatchIndex < NumBatches; BatchIndex++)
 								{
-									const FLODInfo::FSectionInfo& Section = LODs[LODIndex].Sections[SectionIndex];
+									bool bSectionIsSelected = false;
+									FMeshBatch& MeshElement = Collector.AllocateMesh();
 
-									bSectionIsSelected = Section.bSelected;
-									MeshElement.BatchHitProxyId = Section.HitProxy ? Section.HitProxy->Id : FHitProxyId();
-								}
-#endif // WITH_EDITOR
-								
-								if (GetMeshElement(LODIndex, BatchIndex, SectionIndex, SDPG_World, bSectionIsSelected, IsHovered(), MeshElement))
-								{
-									bool bAddMesh = true;
+	#if WITH_EDITOR
+									if (GIsEditor)
+									{
+										const FLODInfo::FSectionInfo& Section = LODs[LODIndex].Sections[SectionIndex];
 
-#if !(UE_BUILD_SHIPPING || UE_BUILD_TEST)
-									if (EngineShowFlags.VertexColors && AllowDebugViewmodes())
+										bSectionIsSelected = Section.bSelected;
+										MeshElement.BatchHitProxyId = Section.HitProxy ? Section.HitProxy->Id : FHitProxyId();
+									}
+	#endif // WITH_EDITOR
+								
+									if (GetMeshElement(LODIndex, BatchIndex, SectionIndex, SDPG_World, bSectionIsSelected, IsHovered(), MeshElement))
 									{
-										// Override the mesh's material with our material that draws the vertex colors
-										UMaterial* VertexColorVisualizationMaterial = NULL;
-										switch( GVertexColorViewMode )
+										bool bAddMesh = true;
+
+	#if !(UE_BUILD_SHIPPING || UE_BUILD_TEST)
+										if (EngineShowFlags.VertexColors && AllowDebugViewmodes())
 										{
-										case EVertexColorViewMode::Color:
-											VertexColorVisualizationMaterial = GEngine->VertexColorViewModeMaterial_ColorOnly;
-											break;
+											// Override the mesh's material with our material that draws the vertex colors
+											UMaterial* VertexColorVisualizationMaterial = NULL;
+											switch( GVertexColorViewMode )
+											{
+											case EVertexColorViewMode::Color:
+												VertexColorVisualizationMaterial = GEngine->VertexColorViewModeMaterial_ColorOnly;
+												break;
 
-										case EVertexColorViewMode::Alpha:
-											VertexColorVisualizationMaterial = GEngine->VertexColorViewModeMaterial_AlphaAsColor;
-											break;
+											case EVertexColorViewMode::Alpha:
+												VertexColorVisualizationMaterial = GEngine->VertexColorViewModeMaterial_AlphaAsColor;
+												break;
 
-										case EVertexColorViewMode::Red:
-											VertexColorVisualizationMaterial = GEngine->VertexColorViewModeMaterial_RedOnly;
-											break;
+											case EVertexColorViewMode::Red:
+												VertexColorVisualizationMaterial = GEngine->VertexColorViewModeMaterial_RedOnly;
+												break;
 
-										case EVertexColorViewMode::Green:
-											VertexColorVisualizationMaterial = GEngine->VertexColorViewModeMaterial_GreenOnly;
-											break;
+											case EVertexColorViewMode::Green:
+												VertexColorVisualizationMaterial = GEngine->VertexColorViewModeMaterial_GreenOnly;
+												break;
 
-										case EVertexColorViewMode::Blue:
-											VertexColorVisualizationMaterial = GEngine->VertexColorViewModeMaterial_BlueOnly;
-											break;
-										}
-										check( VertexColorVisualizationMaterial != NULL );
+											case EVertexColorViewMode::Blue:
+												VertexColorVisualizationMaterial = GEngine->VertexColorViewModeMaterial_BlueOnly;
+												break;
+											}
+											check( VertexColorVisualizationMaterial != NULL );
 
-										auto VertexColorVisualizationMaterialInstance = new FColoredMaterialRenderProxy(
-											VertexColorVisualizationMaterial->GetRenderProxy( MeshElement.MaterialRenderProxy->IsSelected(),MeshElement.MaterialRenderProxy->IsHovered() ),
-											GetSelectionColor( FLinearColor::White, bSectionIsSelected, IsHovered() )
-											);
+											auto VertexColorVisualizationMaterialInstance = new FColoredMaterialRenderProxy(
+												VertexColorVisualizationMaterial->GetRenderProxy( MeshElement.MaterialRenderProxy->IsSelected(),MeshElement.MaterialRenderProxy->IsHovered() ),
+												GetSelectionColor( FLinearColor::White, bSectionIsSelected, IsHovered() )
+												);
 
-										Collector.RegisterOneFrameMaterialProxy(VertexColorVisualizationMaterialInstance);
-										MeshElement.MaterialRenderProxy = VertexColorVisualizationMaterialInstance;
-									}
-									else if (AllowDebugViewmodes() && (bDrawComplexWireframeCollision || (bInCollisionView && bDrawComplexCollision)))
-									{
-										if (LODModel.Sections[SectionIndex].bEnableCollision)
+											Collector.RegisterOneFrameMaterialProxy(VertexColorVisualizationMaterialInstance);
+											MeshElement.MaterialRenderProxy = VertexColorVisualizationMaterialInstance;
+										}
+										else if (AllowDebugViewmodes() && (bDrawComplexWireframeCollision || (bInCollisionView && bDrawComplexCollision)))
+										{
+											if (LODModel.Sections[SectionIndex].bEnableCollision)
+											{
+												UMaterial* MaterialToUse = (bDrawComplexWireframeCollision) ? GEngine->WireframeMaterial : GEngine->ShadedLevelColorationUnlitMaterial;
+
+												// Override the mesh's material with our material that draws the collision color
+												auto CollisionMaterialInstance = new FColoredMaterialRenderProxy(
+													MaterialToUse->GetRenderProxy(bProxyIsSelected, IsHovered()),
+													WireframeColor
+													);
+
+												Collector.RegisterOneFrameMaterialProxy(CollisionMaterialInstance);
+
+												// If drawing the complex collision in wireframe, we do that in _addition_ to the normal mesh
+												if(bDrawComplexWireframeCollision)
+												{
+													FMeshBatch& CollisionElement = Collector.AllocateMesh();
+													GetMeshElement(LODIndex, BatchIndex, SectionIndex, SDPG_World, bSectionIsSelected, IsHovered(), CollisionElement);
+													CollisionElement.MaterialRenderProxy = CollisionMaterialInstance;
+													Collector.AddMesh(ViewIndex, CollisionElement);
+													INC_DWORD_STAT_BY(STAT_StaticMeshTriangles, CollisionElement.GetNumPrimitives());
+												}
+												// Otherwise just change material on existing mesh
+												else
+												{
+													MeshElement.MaterialRenderProxy = CollisionMaterialInstance;
+												}
+											}
+											else
+											{
+												bAddMesh = false;
+											}
+										}
+										else
+	#endif
+	#if WITH_EDITOR
+										if (bSectionIsSelected)
 										{
-											UMaterial* MaterialToUse = (bDrawComplexWireframeCollision) ? GEngine->WireframeMaterial : GEngine->ShadedLevelColorationUnlitMaterial;
-
 											// Override the mesh's material with our material that draws the collision color
-											auto CollisionMaterialInstance = new FColoredMaterialRenderProxy(
-												MaterialToUse->GetRenderProxy(bProxyIsSelected, IsHovered()),
-												WireframeColor
+											auto SelectedMaterialInstance = new FOverrideSelectionColorMaterialRenderProxy(
+												GEngine->ShadedLevelColorationUnlitMaterial->GetRenderProxy(bSectionIsSelected, IsHovered()),
+												GetSelectionColor(GEngine->GetSelectedMaterialColor(), bSectionIsSelected, IsHovered())
 												);
 
-											Collector.RegisterOneFrameMaterialProxy(CollisionMaterialInstance);
-
-											// If drawing the complex collision in wireframe, we do that in _addition_ to the normal mesh
-											if(bDrawComplexWireframeCollision)
+											MeshElement.MaterialRenderProxy = SelectedMaterialInstance;
+										}
+	#endif
+										if (MeshElement.bDitheredLODTransition && LODMask.IsDithered())
+										{
+											if (LODIndex == LODMask.DitheredLODIndices[0])
 											{
-												FMeshBatch& CollisionElement = Collector.AllocateMesh();
-												GetMeshElement(LODIndex, BatchIndex, SectionIndex, SDPG_World, bSectionIsSelected, IsHovered(), CollisionElement);
-												CollisionElement.MaterialRenderProxy = CollisionMaterialInstance;
-												Collector.AddMesh(ViewIndex, CollisionElement);
-												INC_DWORD_STAT_BY(STAT_StaticMeshTriangles, CollisionElement.GetNumPrimitives());
+												MeshElement.DitheredLODTransitionAlpha = View->GetTemporalLODTransition();
 											}
-											// Otherwise just change material on existing mesh
 											else
 											{
-												MeshElement.MaterialRenderProxy = CollisionMaterialInstance;
+												MeshElement.DitheredLODTransitionAlpha = View->GetTemporalLODTransition() - 1.0f;
 											}
 										}
 										else
 										{
-											bAddMesh = false;
+											MeshElement.bDitheredLODTransition = false;
 										}
-									}
-									else
-#endif
-#if WITH_EDITOR
-									if (bSectionIsSelected)
-									{
-										// Override the mesh's material with our material that draws the collision color
-										auto SelectedMaterialInstance = new FOverrideSelectionColorMaterialRenderProxy(
-											GEngine->ShadedLevelColorationUnlitMaterial->GetRenderProxy(bSectionIsSelected, IsHovered()),
-											GetSelectionColor(GEngine->GetSelectedMaterialColor(), bSectionIsSelected, IsHovered())
-											);
-
-										MeshElement.MaterialRenderProxy = SelectedMaterialInstance;
-									}
-#endif
 								
-									MeshElement.bCanApplyViewModeOverrides = true;
-									MeshElement.bUseWireframeSelectionColoring = bSectionIsSelected;
+										MeshElement.bCanApplyViewModeOverrides = true;
+										MeshElement.bUseWireframeSelectionColoring = bSectionIsSelected;
 
-									if (bAddMesh)
-									{
-										Collector.AddMesh(ViewIndex, MeshElement);
-										INC_DWORD_STAT_BY(STAT_StaticMeshTriangles,MeshElement.GetNumPrimitives());
+										if (bAddMesh)
+										{
+											Collector.AddMesh(ViewIndex, MeshElement);
+											INC_DWORD_STAT_BY(STAT_StaticMeshTriangles,MeshElement.GetNumPrimitives());
+										}
 									}
 								}
 							}
@@ -1191,6 +1222,63 @@ int32 FStaticMeshSceneProxy::GetLOD(const FSceneView* View) const
 	return ComputeStaticMeshLOD(RenderData, Bounds.Origin, Bounds.SphereRadius, *View, ClampedMinLOD);
 }
 
+FLODMask FStaticMeshSceneProxy::GetLODMask(const FSceneView* View) const
+{
+	FLODMask Result;
+	int32 CVarForcedLODLevel = GetCVarForceLOD();
+
+	//If a LOD is being forced, use that one
+	if (CVarForcedLODLevel >= 0)
+	{
+		Result.SetLOD(FMath::Clamp<int32>(CVarForcedLODLevel, 0, RenderData->LODResources.Num() - 1));
+	}
+	else if (ForcedLodModel > 0)
+	{
+		Result.SetLOD(FMath::Clamp(ForcedLodModel, 1, RenderData->LODResources.Num()) - 1);
+	}
+#if WITH_EDITOR
+	else if (View->Family && View->Family->EngineShowFlags.LOD == 0)
+	{
+		Result.SetLOD(0);
+	}
+#endif
+	else
+	{
+		const FBoxSphereBounds& Bounds = GetBounds();
+		bool bUseDithered = false;
+		if (LODs.Num())
+		{
+			// only dither if at least one section in LOD0 is dithered. Mixed dithering on sections won't work very well, but it makes an attempt
+			const FLODInfo& ProxyLODInfo = LODs[0];
+			const FStaticMeshLODResources& LODModel = RenderData->LODResources[0];
+			// Draw the static mesh elements.
+			for(int32 SectionIndex = 0; SectionIndex < LODModel.Sections.Num(); SectionIndex++)
+			{
+				UMaterialInterface* Material = ProxyLODInfo.Sections[SectionIndex].Material;
+				if (Material->IsDitheredLODTransition())
+				{
+					bUseDithered = true;
+					break;
+				}
+			}
+
+		}
+		if (bUseDithered)
+		{
+			for (int32 Sample = 0; Sample < 2; Sample++)
+			{
+				Result.SetLODSample(ComputeTemporalStaticMeshLOD(RenderData, Bounds.Origin, Bounds.SphereRadius, *View, ClampedMinLOD, 1.0f, Sample), Sample);
+			}
+		}
+		else
+		{
+			Result.SetLOD(ComputeStaticMeshLOD(RenderData, Bounds.Origin, Bounds.SphereRadius, *View, ClampedMinLOD));
+		}
+	}
+	
+	return Result;
+}
+
 FPrimitiveSceneProxy* UStaticMeshComponent::CreateSceneProxy()
 {
 	if(StaticMesh == NULL
diff --git a/Engine/Source/Runtime/Engine/Public/MaterialShared.h b/Engine/Source/Runtime/Engine/Public/MaterialShared.h
--- a/Engine/Source/Runtime/Engine/Public/MaterialShared.h
+++ b/Engine/Source/Runtime/Engine/Public/MaterialShared.h
@@ -897,6 +897,7 @@ public:
 	ENGINE_API virtual void GetShaderMapId(EShaderPlatform Platform, FMaterialShaderMapId& OutId) const;
 	virtual int32 GetMaterialDomain() const = 0; // See EMaterialDomain.
 	virtual bool IsTwoSided() const = 0;
+	virtual bool IsDitheredLODTransition() const = 0;
 	virtual bool IsTangentSpaceNormal() const { return false; }
 	virtual bool ShouldInjectEmissiveIntoLPV() const { return false; }
 	virtual bool ShouldGenerateSphericalParticleNormals() const { return false; }
@@ -1081,6 +1082,12 @@ public:
 	*/
 	ENGINE_API bool GetMaterialExpressionSource(FString& OutSource);
 
+	/* Helper function to look at both IsMasked and IsDitheredLODTransition to determine if it writes every pixel */
+	ENGINE_API bool WritesEveryPixel() const
+	{
+		return !IsMasked() && !IsDitheredLODTransition();
+	}
+
 	/** 
 	 * Adds an FMaterial to the global list.
 	 * Any FMaterials that don't belong to a UMaterialInterface need to be registered in this way to work correctly with runtime recompiling of outdated shaders.
@@ -1446,6 +1453,7 @@ public:
 	ENGINE_API virtual void GetShaderMapId(EShaderPlatform Platform, FMaterialShaderMapId& OutId) const override;
 	ENGINE_API virtual int32 GetMaterialDomain() const override;
 	ENGINE_API virtual bool IsTwoSided() const override;
+	ENGINE_API virtual bool IsDitheredLODTransition() const override;
 	ENGINE_API virtual bool IsTangentSpaceNormal() const override;
 	ENGINE_API virtual bool ShouldInjectEmissiveIntoLPV() const override;
 	ENGINE_API virtual bool ShouldGenerateSphericalParticleNormals() const override;
@@ -1524,6 +1532,7 @@ public:
 	ENGINE_API virtual const TArray<UTexture*>& GetReferencedTextures() const override;
 
 	ENGINE_API virtual bool GetAllowDevelopmentShaderCompile() const override;
+
 protected:
 	UMaterial* Material;
 	UMaterialInstance* MaterialInstance;
diff --git a/Engine/Source/Runtime/Engine/Public/MeshBatch.h b/Engine/Source/Runtime/Engine/Public/MeshBatch.h
--- a/Engine/Source/Runtime/Engine/Public/MeshBatch.h
+++ b/Engine/Source/Runtime/Engine/Public/MeshBatch.h
@@ -101,6 +101,12 @@ struct FMeshBatch
 	/** Whether the mesh batch can be selected through editor selection, aka hit proxies. */
 	uint32 bSelectable : 1;
 
+	/** Whether the mesh batch should apply dithered LOD. */
+	uint32 bDitheredLODTransition : 1;
+
+	/** If bDitheredLODTransition and this is a dynamic mesh element, then this is the alpha for dither fade (static draw lists need to derive this later as it is changes every frame) */
+	float DitheredLODTransitionAlpha;
+
 	// can be NULL
 	const FLightCacheInterface* LCI;
 
@@ -200,6 +206,8 @@ struct FMeshBatch
 	,	bUseWireframeSelectionColoring(false)
 	,	bUseSelectionOutline(true)
 	,	bSelectable(true)
+	,	bDitheredLODTransition(false)
+	,   DitheredLODTransitionAlpha(0.0f)
 	,	LCI(NULL)
 	,	DynamicVertexData(NULL)
 	,	VertexFactory(NULL)
diff --git a/Engine/Source/Runtime/Engine/Public/SceneManagement.h b/Engine/Source/Runtime/Engine/Public/SceneManagement.h
--- a/Engine/Source/Runtime/Engine/Public/SceneManagement.h
+++ b/Engine/Source/Runtime/Engine/Public/SceneManagement.h
@@ -2002,9 +2002,21 @@ float ENGINE_API ComputeBoundsScreenSize(const FVector4& Origin, const float Sph
  * @param Origin - Origin of the bounds of the mesh in world space
  * @param SphereRadius - Radius of the sphere to use to calculate screen coverage
  * @param View - The view to calculate the LOD level for
+ * @param FactorScale - multiplied by the computed screen size before computing LOD
  */
 int8 ENGINE_API ComputeStaticMeshLOD(const FStaticMeshRenderData* RenderData, const FVector4& Origin, const float SphereRadius, const FSceneView& View, int32 MinLOD, float FactorScale = 1.0f);
 
+/**
+ * Computes the LOD level for the given static meshes render data in the given view, for one of the two temporal LOD samples
+ * @param RenderData - Render data for the mesh
+ * @param Origin - Origin of the bounds of the mesh in world space
+ * @param SphereRadius - Radius of the sphere to use to calculate screen coverage
+ * @param View - The view to calculate the LOD level for
+ * @param FactorScale - multiplied by the computed screen size before computing LOD
+ * @param SampleIndex - index (0 or 1) of the temporal sample to use
+ */
+int8 ENGINE_API ComputeTemporalStaticMeshLOD( const FStaticMeshRenderData* RenderData, const FVector4& Origin, const float SphereRadius, const FSceneView& View, int32 MinLOD, float FactorScale, int32 SampleIndex );
+
 /**
  * Computes the LOD to render for the list of static meshes in the given view.
  * @param StaticMeshes - List of static meshes.
@@ -2012,7 +2024,36 @@ int8 ENGINE_API ComputeStaticMeshLOD(const FStaticMeshRenderData* RenderData, co
  * @param Origin - Origin of the bounds of the mesh in world space
  * @param SphereRadius - Radius of the sphere to use to calculate screen coverage
  */
-int8 ENGINE_API ComputeLODForMeshes(const TIndirectArray<class FStaticMesh>& StaticMeshes, const FSceneView& View, const FVector4& Origin, float SphereRadius, int32 ForcedLODLevel, float ScreenSizeScale = 1.0f);
+struct FLODMask
+{
+	int8 DitheredLODIndices[2];
+
+	FLODMask()
+	{
+		DitheredLODIndices[0] = MAX_int8;
+		DitheredLODIndices[1] = MAX_int8;
+	}
+
+	void SetLOD(int32 LODIndex)
+	{
+		DitheredLODIndices[0] = LODIndex;
+		DitheredLODIndices[1] = LODIndex;
+	}
+	void SetLODSample(int32 LODIndex, int32 SampleIndex)
+	{
+		DitheredLODIndices[SampleIndex] = (int8)LODIndex;
+	}
+	bool ContainsLOD(int32 LODIndex) const
+	{
+		return DitheredLODIndices[0] == LODIndex || DitheredLODIndices[1] == LODIndex;
+	}
+
+	bool IsDithered() const
+	{
+		return DitheredLODIndices[0] != DitheredLODIndices[1];
+	}
+};
+FLODMask ENGINE_API ComputeLODForMeshes(const TIndirectArray<class FStaticMesh>& StaticMeshes, const FSceneView& View, const FVector4& Origin, float SphereRadius, int32 ForcedLODLevel, float ScreenSizeScale = 1.0f);
 
 class FSharedSamplerState : public FRenderResource
 {
diff --git a/Engine/Source/Runtime/Engine/Public/StaticMeshResources.h b/Engine/Source/Runtime/Engine/Public/StaticMeshResources.h
--- a/Engine/Source/Runtime/Engine/Public/StaticMeshResources.h
+++ b/Engine/Source/Runtime/Engine/Public/StaticMeshResources.h
@@ -748,7 +748,7 @@ public:
 	}
 
 	/** Sets up a shadow FMeshBatch for a specific LOD. */
-	virtual bool GetShadowMeshElement(int32 LODIndex, int32 BatchIndex, uint8 InDepthPriorityGroup, FMeshBatch& OutMeshBatch) const;
+	virtual bool GetShadowMeshElement(int32 LODIndex, int32 BatchIndex, uint8 InDepthPriorityGroup, FMeshBatch& OutMeshBatch, bool bDitheredLODTransition) const;
 
 	/** Sets up a FMeshBatch for a specific LOD and element. */
 	virtual bool GetMeshElement(int32 LODIndex, int32 BatchIndex, int32 ElementIndex, uint8 InDepthPriorityGroup, const bool bUseSelectedMaterial, const bool bUseHoveredMaterial, FMeshBatch& OutMeshBatch) const;
@@ -894,6 +894,11 @@ protected:
 	 * @Param LODIndex - The LOD to get the display factor for
 	 */
 	float GetScreenSize(int32 LODIndex) const;
+
+	/**
+	 * Returns the LOD mask for a view, this is like the ordinary LOD but can return two values for dither fading
+	 */
+	FLODMask GetLODMask(const FSceneView* View) const;
 };
 
 /*-----------------------------------------------------------------------------
@@ -1159,7 +1164,6 @@ public:
 	}
 };
 
-
 #if WITH_EDITOR
 /**
  * Remaps painted vertex colors when the renderable mesh has changed.
diff --git a/Engine/Source/Runtime/Landscape/Private/LandscapeGrass.cpp b/Engine/Source/Runtime/Landscape/Private/LandscapeGrass.cpp
--- a/Engine/Source/Runtime/Landscape/Private/LandscapeGrass.cpp
+++ b/Engine/Source/Runtime/Landscape/Private/LandscapeGrass.cpp
@@ -130,9 +130,9 @@ public:
 		SetShaderValue(RHICmdList, GetVertexShader(), RenderOffsetParameter, RenderOffset);
 	}
 
-	void SetMesh(FRHICommandList& RHICmdList, const FVertexFactory* VertexFactory, const FSceneView& View, const FPrimitiveSceneProxy* Proxy, const FMeshBatchElement& BatchElement)
+	void SetMesh(FRHICommandList& RHICmdList, const FVertexFactory* VertexFactory, const FSceneView& View, const FPrimitiveSceneProxy* Proxy, const FMeshBatchElement& BatchElement, float DitheredLODTransitionValue)
 	{
-		FMeshMaterialShader::SetMesh(RHICmdList, GetVertexShader(), VertexFactory, View, Proxy, BatchElement);
+		FMeshMaterialShader::SetMesh(RHICmdList, GetVertexShader(), VertexFactory, View, Proxy, BatchElement,DitheredLODTransitionValue);
 	}
 
 	virtual bool Serialize(FArchive& Ar) override
@@ -175,9 +175,9 @@ public:
 		}
 	}
 
-	void SetMesh(FRHICommandList& RHICmdList, const FVertexFactory* VertexFactory, const FSceneView& View, const FPrimitiveSceneProxy* Proxy, const FMeshBatchElement& BatchElement)
+	void SetMesh(FRHICommandList& RHICmdList, const FVertexFactory* VertexFactory, const FSceneView& View, const FPrimitiveSceneProxy* Proxy, const FMeshBatchElement& BatchElement, float DitheredLODTransitionValue)
 	{
-		FMeshMaterialShader::SetMesh(RHICmdList, GetPixelShader(), VertexFactory, View, Proxy, BatchElement);
+		FMeshMaterialShader::SetMesh(RHICmdList, GetPixelShader(), VertexFactory, View, Proxy, BatchElement,DitheredLODTransitionValue);
 	}
 
 	virtual bool Serialize(FArchive& Ar) override
@@ -245,14 +245,15 @@ public:
 		const FMeshBatch& Mesh,
 		int32 BatchElementIndex,
 		bool bBackFace,
+		float DitheredLODTransitionValue,
 		const ElementDataType& ElementData,
 		const ContextDataType PolicyContext
 		) const
 	{
 		const FMeshBatchElement& BatchElement = Mesh.Elements[BatchElementIndex];
-		VertexShader->SetMesh(RHICmdList, VertexFactory, View, PrimitiveSceneProxy, BatchElement);
-		PixelShader->SetMesh(RHICmdList, VertexFactory, View, PrimitiveSceneProxy, BatchElement);
-		FMeshDrawingPolicy::SetMeshRenderState(RHICmdList, View, PrimitiveSceneProxy, Mesh, BatchElementIndex, bBackFace, ElementData, PolicyContext);
+		VertexShader->SetMesh(RHICmdList, VertexFactory, View, PrimitiveSceneProxy, BatchElement,DitheredLODTransitionValue);
+		PixelShader->SetMesh(RHICmdList, VertexFactory, View, PrimitiveSceneProxy, BatchElement,DitheredLODTransitionValue);
+		FMeshDrawingPolicy::SetMeshRenderState(RHICmdList, View, PrimitiveSceneProxy, Mesh, BatchElementIndex, bBackFace, DitheredLODTransitionValue, ElementData, PolicyContext);
 	}
 
 	friend int32 CompareDrawingPolicy(const FLandscapeGrassWeightDrawingPolicy& A, const FLandscapeGrassWeightDrawingPolicy& B)
@@ -345,7 +346,7 @@ public:
 					DrawingPolicy.SetSharedState(RHICmdList, View, FLandscapeGrassWeightDrawingPolicy::ContextDataType(), PassIdx, ComponentInfo.ViewOffset + FVector2D(PassOffsetX * PassIdx, 0));
 
 					// The first batch element contains the grass batch for the entire component
-					DrawingPolicy.SetMeshRenderState(RHICmdList, *View, PrimitiveSceneInfo->Proxy, Mesh, 0, false, FMeshDrawingPolicy::ElementDataType(), FLandscapeGrassWeightDrawingPolicy::ContextDataType());
+					DrawingPolicy.SetMeshRenderState(RHICmdList, *View, PrimitiveSceneInfo->Proxy, Mesh, 0, false, 0.0f, FMeshDrawingPolicy::ElementDataType(), FLandscapeGrassWeightDrawingPolicy::ContextDataType());
 					DrawingPolicy.DrawMesh(RHICmdList, Mesh, 0);
 				}
 			}
diff --git a/Engine/Source/Runtime/Renderer/Private/BasePassRendering.cpp b/Engine/Source/Runtime/Renderer/Private/BasePassRendering.cpp
--- a/Engine/Source/Runtime/Renderer/Private/BasePassRendering.cpp
+++ b/Engine/Source/Runtime/Renderer/Private/BasePassRendering.cpp
@@ -264,6 +264,7 @@ public:
 
 	const FViewInfo& View;
 	bool bBackFace;
+	float DitheredLODTransitionValue;
 	bool bPreFog;
 	FHitProxyId HitProxyId;
 
@@ -271,10 +272,12 @@ public:
 	FDrawBasePassDynamicMeshAction(
 		const FViewInfo& InView,
 		const bool bInBackFace,
+		float InDitheredLODTransitionValue,
 		const FHitProxyId InHitProxyId
 		)
 		: View(InView)
 		, bBackFace(bInBackFace)
+		, DitheredLODTransitionValue(InDitheredLODTransitionValue)
 		, HitProxyId(InHitProxyId)
 	{}
 
@@ -341,6 +344,7 @@ public:
 				Parameters.Mesh,
 				BatchElementIndex,
 				bBackFace,
+				DitheredLODTransitionValue,
 				typename TBasePassDrawingPolicy<LightMapPolicyType>::ElementDataType(LightMapElementData),
 				typename TBasePassDrawingPolicy<LightMapPolicyType>::ContextDataType()
 				);
@@ -388,6 +392,7 @@ bool FBasePassOpaqueDrawingPolicyFactory::DrawDynamicMesh(
 			FDrawBasePassDynamicMeshAction(
 				View,
 				bBackFace,
+				Mesh.DitheredLODTransitionAlpha,
 				HitProxyId
 				)
 			);
diff --git a/Engine/Source/Runtime/Renderer/Private/BasePassRendering.h b/Engine/Source/Runtime/Renderer/Private/BasePassRendering.h
--- a/Engine/Source/Runtime/Renderer/Private/BasePassRendering.h
+++ b/Engine/Source/Runtime/Renderer/Private/BasePassRendering.h
@@ -82,7 +82,7 @@ public:
 		}
 	}
 
-	void SetMesh(FRHICommandList& RHICmdList, const FVertexFactory* VertexFactory,const FSceneView& View,const FPrimitiveSceneProxy* Proxy, const FMeshBatch& Mesh, const FMeshBatchElement& BatchElement);
+	void SetMesh(FRHICommandList& RHICmdList, const FVertexFactory* VertexFactory,const FSceneView& View,const FPrimitiveSceneProxy* Proxy, const FMeshBatch& Mesh, const FMeshBatchElement& BatchElement, float DitheredLODTransitionValue);
 
 private:
 	
@@ -395,7 +395,7 @@ public:
 		EditorCompositeParams.SetParameters(RHICmdList, MaterialResource, View, bEnableEditorPrimitveDepthTest, GetPixelShader());
 	}
 
-	void SetMesh(FRHICommandList& RHICmdList, const FVertexFactory* VertexFactory,const FSceneView& View,const FPrimitiveSceneProxy* Proxy,const FMeshBatchElement& BatchElement, EBlendMode BlendMode);
+	void SetMesh(FRHICommandList& RHICmdList, const FVertexFactory* VertexFactory,const FSceneView& View,const FPrimitiveSceneProxy* Proxy,const FMeshBatchElement& BatchElement, EBlendMode BlendMode, float DitheredLODTransitionValue);
 
 	virtual bool Serialize(FArchive& Ar) override
 	{
@@ -649,6 +649,7 @@ public:
 		const FMeshBatch& Mesh,
 		int32 BatchElementIndex,
 		bool bBackFace,
+		float DitheredLODTransitionValue,
 		const ElementDataType& ElementData,
 		const ContextDataType PolicyContext
 		) const
@@ -667,12 +668,12 @@ public:
 			ElementData.LightMapElementData);
 
 		const FMeshBatchElement& BatchElement = Mesh.Elements[BatchElementIndex];
-		VertexShader->SetMesh(RHICmdList, VertexFactory,View,PrimitiveSceneProxy, Mesh,BatchElement);
+		VertexShader->SetMesh(RHICmdList, VertexFactory,View,PrimitiveSceneProxy, Mesh,BatchElement,DitheredLODTransitionValue);
 		
 		if(HullShader && DomainShader)
 		{
-			HullShader->SetMesh(RHICmdList, VertexFactory,View,PrimitiveSceneProxy,BatchElement);
-			DomainShader->SetMesh(RHICmdList, VertexFactory,View,PrimitiveSceneProxy,BatchElement);
+			HullShader->SetMesh(RHICmdList, VertexFactory,View,PrimitiveSceneProxy,BatchElement,DitheredLODTransitionValue);
+			DomainShader->SetMesh(RHICmdList, VertexFactory,View,PrimitiveSceneProxy,BatchElement,DitheredLODTransitionValue);
 		}
 
 #if !(UE_BUILD_SHIPPING || UE_BUILD_TEST)
@@ -694,10 +695,10 @@ public:
 		else
 #endif
 		{
-			PixelShader->SetMesh(RHICmdList, VertexFactory,View,PrimitiveSceneProxy,BatchElement,BlendMode);
+			PixelShader->SetMesh(RHICmdList, VertexFactory,View,PrimitiveSceneProxy,BatchElement,BlendMode,DitheredLODTransitionValue);
 		}
 
-		FMeshDrawingPolicy::SetMeshRenderState(RHICmdList, View,PrimitiveSceneProxy,Mesh,BatchElementIndex,bBackFace,FMeshDrawingPolicy::ElementDataType(),PolicyContext);
+		FMeshDrawingPolicy::SetMeshRenderState(RHICmdList, View,PrimitiveSceneProxy,Mesh,BatchElementIndex,bBackFace, DitheredLODTransitionValue,FMeshDrawingPolicy::ElementDataType(),PolicyContext);
 	}
 
 	friend int32 CompareDrawingPolicy(const TBasePassDrawingPolicy& A,const TBasePassDrawingPolicy& B)
diff --git a/Engine/Source/Runtime/Renderer/Private/BasePassRendering.inl b/Engine/Source/Runtime/Renderer/Private/BasePassRendering.inl
--- a/Engine/Source/Runtime/Renderer/Private/BasePassRendering.inl
+++ b/Engine/Source/Runtime/Renderer/Private/BasePassRendering.inl
@@ -8,10 +8,10 @@
 #pragma once
 
 template<typename LightMapPolicyType>
-inline void TBasePassVertexShaderBaseType<LightMapPolicyType>::SetMesh(FRHICommandList& RHICmdList, const FVertexFactory* VertexFactory,const FSceneView& View,const FPrimitiveSceneProxy* Proxy, const FMeshBatch& Mesh, const FMeshBatchElement& BatchElement)
+inline void TBasePassVertexShaderBaseType<LightMapPolicyType>::SetMesh(FRHICommandList& RHICmdList, const FVertexFactory* VertexFactory,const FSceneView& View,const FPrimitiveSceneProxy* Proxy, const FMeshBatch& Mesh, const FMeshBatchElement& BatchElement, float DitheredLODTransitionValue)
 {
 	FVertexShaderRHIParamRef VertexShaderRHI = GetVertexShader();
-	FMeshMaterialShader::SetMesh(RHICmdList, VertexShaderRHI, VertexFactory, View, Proxy, BatchElement);
+	FMeshMaterialShader::SetMesh(RHICmdList, VertexShaderRHI, VertexFactory, View, Proxy, BatchElement, DitheredLODTransitionValue);
 
 	const bool bHasPreviousLocalToWorldParameter = PreviousLocalToWorldParameter.IsBound();
 	const bool bHasSkipOutputVelocityParameter = SkipOutputVelocityParameter.IsBound();
@@ -47,7 +47,7 @@ inline void TBasePassVertexShaderBaseType<LightMapPolicyType>::SetMesh(FRHIComma
 }
 
 template<typename LightMapPolicyType>
-void TBasePassPixelShaderBaseType<LightMapPolicyType>::SetMesh(FRHICommandList& RHICmdList, const FVertexFactory* VertexFactory, const FSceneView& View, const FPrimitiveSceneProxy* Proxy, const FMeshBatchElement& BatchElement, EBlendMode BlendMode)
+void TBasePassPixelShaderBaseType<LightMapPolicyType>::SetMesh(FRHICommandList& RHICmdList, const FVertexFactory* VertexFactory, const FSceneView& View, const FPrimitiveSceneProxy* Proxy, const FMeshBatchElement& BatchElement, EBlendMode BlendMode, float DitheredLODTransitionValue)
 {
 	if (View.GetFeatureLevel() >= ERHIFeatureLevel::SM4
 		&& IsTranslucentBlendMode(BlendMode))
@@ -55,5 +55,5 @@ void TBasePassPixelShaderBaseType<LightMapPolicyType>::SetMesh(FRHICommandList&
 		TranslucentLightingParameters.SetMesh(RHICmdList, this, Proxy, View.GetFeatureLevel());
 	}
 
-	FMeshMaterialShader::SetMesh(RHICmdList, GetPixelShader(), VertexFactory, View, Proxy, BatchElement);
+	FMeshMaterialShader::SetMesh(RHICmdList, GetPixelShader(), VertexFactory, View, Proxy, BatchElement,DitheredLODTransitionValue);
 }
diff --git a/Engine/Source/Runtime/Renderer/Private/ConvertToUniformMesh.cpp b/Engine/Source/Runtime/Renderer/Private/ConvertToUniformMesh.cpp
--- a/Engine/Source/Runtime/Renderer/Private/ConvertToUniformMesh.cpp
+++ b/Engine/Source/Runtime/Renderer/Private/ConvertToUniformMesh.cpp
@@ -48,9 +48,9 @@ public:
 		FMeshMaterialShader::SetParameters(RHICmdList, GetVertexShader(), MaterialRenderProxy, *MaterialRenderProxy->GetMaterial(View->GetFeatureLevel()), *View, ESceneRenderTargetsMode::SetTextures);
 	}
 
-	void SetMesh(FRHICommandList& RHICmdList, const FVertexFactory* VertexFactory,const FSceneView& View,const FPrimitiveSceneProxy* Proxy,const FMeshBatchElement& BatchElement)
+	void SetMesh(FRHICommandList& RHICmdList, const FVertexFactory* VertexFactory,const FSceneView& View,const FPrimitiveSceneProxy* Proxy,const FMeshBatchElement& BatchElement, float DitheredLODTransitionValue)
 	{
-		FMeshMaterialShader::SetMesh(RHICmdList, GetVertexShader(),VertexFactory,View,Proxy,BatchElement);
+		FMeshMaterialShader::SetMesh(RHICmdList, GetVertexShader(),VertexFactory,View,Proxy,BatchElement,DitheredLODTransitionValue);
 	}
 };
 
@@ -138,9 +138,9 @@ public:
 		FMeshMaterialShader::SetParameters(RHICmdList, (FGeometryShaderRHIParamRef)GetGeometryShader(), MaterialRenderProxy, *MaterialRenderProxy->GetMaterial(View->GetFeatureLevel()), *View, ESceneRenderTargetsMode::SetTextures);
 	}
 
-	void SetMesh(FRHICommandList& RHICmdList, const FVertexFactory* VertexFactory,const FSceneView& View,const FPrimitiveSceneProxy* Proxy,const FMeshBatchElement& BatchElement)
+	void SetMesh(FRHICommandList& RHICmdList, const FVertexFactory* VertexFactory,const FSceneView& View,const FPrimitiveSceneProxy* Proxy,const FMeshBatchElement& BatchElement, float DitheredLODTransitionValue)
 	{
-		FMeshMaterialShader::SetMesh(RHICmdList, (FGeometryShaderRHIParamRef)GetGeometryShader(),VertexFactory,View,Proxy,BatchElement);
+		FMeshMaterialShader::SetMesh(RHICmdList, (FGeometryShaderRHIParamRef)GetGeometryShader(),VertexFactory,View,Proxy,BatchElement,DitheredLODTransitionValue);
 	}
 };
 
@@ -203,6 +203,7 @@ public:
 		const FMeshBatch& Mesh,
 		int32 BatchElementIndex,
 		bool bBackFace,
+		float DitheredLODTransitionValue,
 		const ElementDataType& ElementData,
 		const ContextDataType PolicyContext
 		) const;
@@ -265,6 +266,7 @@ void FConvertToUniformMeshDrawingPolicy::SetMeshRenderState(
 	const FMeshBatch& Mesh,
 	int32 BatchElementIndex,
 	bool bBackFace,
+	float DitheredLODTransitionValue,
 	const ElementDataType& ElementData,
 	const ContextDataType PolicyContext
 	) const
@@ -274,8 +276,8 @@ void FConvertToUniformMeshDrawingPolicy::SetMeshRenderState(
 	EmitMeshDrawEvents(RHICmdList, PrimitiveSceneProxy, Mesh);
 
 	// Set transforms
-	VertexShader->SetMesh(RHICmdList, VertexFactory,View,PrimitiveSceneProxy,BatchElement);
-	GeometryShader->SetMesh(RHICmdList, VertexFactory,View,PrimitiveSceneProxy,BatchElement);
+	VertexShader->SetMesh(RHICmdList, VertexFactory,View,PrimitiveSceneProxy,BatchElement, DitheredLODTransitionValue);
+	GeometryShader->SetMesh(RHICmdList, VertexFactory,View,PrimitiveSceneProxy,BatchElement, DitheredLODTransitionValue);
 	
 	// Set rasterizer state.
 	RHICmdList.SetRasterizerState(GetStaticRasterizerState<true>(
@@ -411,7 +413,7 @@ int32 FUniformMeshConverter::Convert(
 							? Mesh.Elements[BatchElementIndex].PrimitiveUniformBuffer
 							: *Mesh.Elements[BatchElementIndex].PrimitiveUniformBufferResource;
 
-						DrawingPolicy.SetMeshRenderState(RHICmdList, View,PrimitiveSceneProxy,Mesh,BatchElementIndex,false,FConvertToUniformMeshDrawingPolicy::ElementDataType(), FConvertToUniformMeshDrawingPolicy::ContextDataType());
+						DrawingPolicy.SetMeshRenderState(RHICmdList, View,PrimitiveSceneProxy,Mesh,BatchElementIndex,false,0.0f,FConvertToUniformMeshDrawingPolicy::ElementDataType(), FConvertToUniformMeshDrawingPolicy::ContextDataType());
 						DrawingPolicy.DrawMesh(RHICmdList, Mesh, BatchElementIndex);
 					}
 				}
@@ -455,7 +457,7 @@ int32 FUniformMeshConverter::Convert(
 							? StaticMesh.Elements[BatchElementIndex].PrimitiveUniformBuffer
 							: *StaticMesh.Elements[BatchElementIndex].PrimitiveUniformBufferResource;
 
-						DrawingPolicy.SetMeshRenderState(RHICmdList, View,PrimitiveSceneProxy,StaticMesh,BatchElementIndex,false,
+						DrawingPolicy.SetMeshRenderState(RHICmdList, View,PrimitiveSceneProxy,StaticMesh,BatchElementIndex,false,0.0f,
 							FConvertToUniformMeshDrawingPolicy::ElementDataType(),
 							FConvertToUniformMeshDrawingPolicy::ContextDataType()
 							);
diff --git a/Engine/Source/Runtime/Renderer/Private/CustomDepthRendering.cpp b/Engine/Source/Runtime/Renderer/Private/CustomDepthRendering.cpp
--- a/Engine/Source/Runtime/Renderer/Private/CustomDepthRendering.cpp
+++ b/Engine/Source/Runtime/Renderer/Private/CustomDepthRendering.cpp
@@ -55,6 +55,7 @@ bool FCustomDepthPrimSet::DrawPrims(FRHICommandListImmediate& RHICmdList, const
 
 						if (View.StaticMeshVisibilityMap[StaticMesh.Id])
 						{
+							float DitherValue = View.GetDitheredLODTransitionValue(StaticMesh);
 							bDirty |= FDepthDrawingPolicyFactory::DrawStaticMesh(
 								RHICmdList, 
 								View,
@@ -62,6 +63,7 @@ bool FCustomDepthPrimSet::DrawPrims(FRHICommandListImmediate& RHICmdList, const
 								StaticMesh,
 								StaticMesh.Elements.Num() == 1 ? 1 : View.StaticMeshBatchVisibility[StaticMesh.Id],
 								true,
+								DitherValue,
 								PrimitiveSceneProxy,
 								StaticMesh.BatchHitProxyId
 								);
diff --git a/Engine/Source/Runtime/Renderer/Private/DepthRendering.cpp b/Engine/Source/Runtime/Renderer/Private/DepthRendering.cpp
--- a/Engine/Source/Runtime/Renderer/Private/DepthRendering.cpp
+++ b/Engine/Source/Runtime/Renderer/Private/DepthRendering.cpp
@@ -33,7 +33,7 @@ public:
 		}
 
 		// Only compile for the default material and masked materials
-		return (Material->IsSpecialEngineMaterial() || Material->IsMasked() || Material->MaterialMayModifyMeshPosition())
+		return (Material->IsSpecialEngineMaterial() || !Material->WritesEveryPixel() || Material->MaterialMayModifyMeshPosition())
 			 && IsFeatureLevelSupported(Platform, ERHIFeatureLevel::SM4);
 	}
 
@@ -42,9 +42,9 @@ public:
 		FMeshMaterialShader::SetParameters(RHICmdList, GetVertexShader(),MaterialRenderProxy,MaterialResource,View,ESceneRenderTargetsMode::DontSet);
 	}
 
-	void SetMesh(FRHICommandList& RHICmdList, const FVertexFactory* VertexFactory,const FSceneView& View,const FPrimitiveSceneProxy* Proxy,const FMeshBatchElement& BatchElement)
+	void SetMesh(FRHICommandList& RHICmdList, const FVertexFactory* VertexFactory,const FSceneView& View,const FPrimitiveSceneProxy* Proxy,const FMeshBatchElement& BatchElement, float DitheredLODTransitionValue)
 	{
-		FMeshMaterialShader::SetMesh(RHICmdList, GetVertexShader(),VertexFactory,View,Proxy,BatchElement);
+		FMeshMaterialShader::SetMesh(RHICmdList, GetVertexShader(),VertexFactory,View,Proxy,BatchElement,DitheredLODTransitionValue);
 	}
 };
 
@@ -106,7 +106,7 @@ public:
 	static bool ShouldCache(EShaderPlatform Platform,const FMaterial* Material,const FVertexFactoryType* VertexFactoryType)
 	{
 		// Compile for materials that are masked.
-		return Material->IsMasked() && IsFeatureLevelSupported(Platform, ERHIFeatureLevel::SM4);
+		return !Material->WritesEveryPixel() && IsFeatureLevelSupported(Platform, ERHIFeatureLevel::SM4);
 	}
 
 	FDepthOnlyPS(const ShaderMetaType::CompiledShaderInitializerType& Initializer):
@@ -120,9 +120,9 @@ public:
 		FMeshMaterialShader::SetParameters(RHICmdList, GetPixelShader(),MaterialRenderProxy,MaterialResource,*View,ESceneRenderTargetsMode::DontSet);
 	}
 
-	void SetMesh(FRHICommandList& RHICmdList, const FVertexFactory* VertexFactory,const FSceneView& View,const FPrimitiveSceneProxy* Proxy,const FMeshBatchElement& BatchElement)
+	void SetMesh(FRHICommandList& RHICmdList, const FVertexFactory* VertexFactory,const FSceneView& View,const FPrimitiveSceneProxy* Proxy,const FMeshBatchElement& BatchElement, float DitheredLODTransitionValue)
 	{
-		FMeshMaterialShader::SetMesh(RHICmdList, GetPixelShader(),VertexFactory,View,Proxy,BatchElement);
+		FMeshMaterialShader::SetMesh(RHICmdList, GetPixelShader(),VertexFactory,View,Proxy,BatchElement,DitheredLODTransitionValue);
 	}
 };
 
@@ -139,7 +139,7 @@ FDepthDrawingPolicy::FDepthDrawingPolicy(
 {
 	// The primitive needs to be rendered with the material's pixel and vertex shaders if it is masked
 	bNeedsPixelShader = false;
-	if (InMaterialResource.IsMasked())
+	if (!InMaterialResource.WritesEveryPixel())
 	{
 		bNeedsPixelShader = true;
 		PixelShader = InMaterialResource.GetShader<FDepthOnlyPS>(InVertexFactory->GetType());
@@ -212,23 +212,24 @@ void FDepthDrawingPolicy::SetMeshRenderState(
 	const FMeshBatch& Mesh,
 	int32 BatchElementIndex,
 	bool bBackFace,
+	float DitheredLODTransitionValue,
 	const ElementDataType& ElementData,
 	const ContextDataType PolicyContext
 	) const
 {
 	const FMeshBatchElement& BatchElement = Mesh.Elements[BatchElementIndex];
-	VertexShader->SetMesh(RHICmdList, VertexFactory,View,PrimitiveSceneProxy,BatchElement);
+	VertexShader->SetMesh(RHICmdList, VertexFactory,View,PrimitiveSceneProxy,BatchElement,DitheredLODTransitionValue);
 	if(HullShader && DomainShader)
 	{
-		HullShader->SetMesh(RHICmdList, VertexFactory,View,PrimitiveSceneProxy,BatchElement);
-		DomainShader->SetMesh(RHICmdList, VertexFactory,View,PrimitiveSceneProxy,BatchElement);
+		HullShader->SetMesh(RHICmdList, VertexFactory,View,PrimitiveSceneProxy,BatchElement,DitheredLODTransitionValue);
+		DomainShader->SetMesh(RHICmdList, VertexFactory,View,PrimitiveSceneProxy,BatchElement,DitheredLODTransitionValue);
 	}
 
 	if (bNeedsPixelShader)
 	{
-		PixelShader->SetMesh(RHICmdList, VertexFactory,View,PrimitiveSceneProxy,BatchElement);
+		PixelShader->SetMesh(RHICmdList, VertexFactory,View,PrimitiveSceneProxy,BatchElement,DitheredLODTransitionValue);
 	}
-	FMeshDrawingPolicy::SetMeshRenderState(RHICmdList, View,PrimitiveSceneProxy,Mesh,BatchElementIndex,bBackFace,ElementData,PolicyContext);
+	FMeshDrawingPolicy::SetMeshRenderState(RHICmdList, View,PrimitiveSceneProxy,Mesh,BatchElementIndex,bBackFace, DitheredLODTransitionValue,ElementData,PolicyContext);
 }
 
 int32 CompareDrawingPolicy(const FDepthDrawingPolicy& A,const FDepthDrawingPolicy& B)
@@ -287,12 +288,13 @@ void FPositionOnlyDepthDrawingPolicy::SetMeshRenderState(
 	const FMeshBatch& Mesh,
 	int32 BatchElementIndex,
 	bool bBackFace,
+	float DitheredLODTransitionValue,
 	const ElementDataType& ElementData,
 	const ContextDataType PolicyContext
 	) const
 {
-	VertexShader->SetMesh(RHICmdList, VertexFactory,View,PrimitiveSceneProxy,Mesh.Elements[BatchElementIndex]);
-	FMeshDrawingPolicy::SetMeshRenderState(RHICmdList, View,PrimitiveSceneProxy,Mesh,BatchElementIndex,bBackFace,ElementData,PolicyContext);
+	VertexShader->SetMesh(RHICmdList, VertexFactory,View,PrimitiveSceneProxy,Mesh.Elements[BatchElementIndex],DitheredLODTransitionValue);
+	FMeshDrawingPolicy::SetMeshRenderState(RHICmdList, View,PrimitiveSceneProxy,Mesh,BatchElementIndex,bBackFace, DitheredLODTransitionValue,ElementData,PolicyContext);
 }
 
 int32 CompareDrawingPolicy(const FPositionOnlyDepthDrawingPolicy& A,const FPositionOnlyDepthDrawingPolicy& B)
@@ -311,7 +313,7 @@ void FDepthDrawingPolicyFactory::AddStaticMesh(FScene* Scene,FStaticMesh* Static
 	const EBlendMode BlendMode = Material->GetBlendMode();
 	const auto FeatureLevel = Scene->GetFeatureLevel();
 
-	if (Material->IsMasked())
+	if (!Material->WritesEveryPixel())
 	{
 		// only draw if required
 		Scene->MaskedDepthDrawList.AddMesh(
@@ -379,6 +381,7 @@ bool FDepthDrawingPolicyFactory::DrawMesh(
 	const FMeshBatch& Mesh,
 	const uint64& BatchElementMask,
 	bool bBackFace,
+	float DitheredLODTransitionValue,
 	bool bPreFog,
 	const FPrimitiveSceneProxy* PrimitiveSceneProxy,
 	FHitProxyId HitProxyId
@@ -400,7 +403,9 @@ bool FDepthDrawingPolicyFactory::DrawMesh(
 
 		if ( BlendMode == BLEND_Opaque 
 			&& Mesh.VertexFactory->SupportsPositionOnlyStream() 
-			&& !Material->MaterialModifiesMeshPosition_RenderThread())
+			&& !Material->MaterialModifiesMeshPosition_RenderThread()
+			&& Material->WritesEveryPixel()
+			)
 		{
 			//render opaque primitives that support a separate position-only vertex buffer
 			const FMaterialRenderProxy* DefaultProxy = UMaterial::GetDefaultMaterial(MD_Surface)->GetRenderProxy(false);
@@ -414,7 +419,7 @@ bool FDepthDrawingPolicyFactory::DrawMesh(
 			{
 				if(Mask & 1)
 				{
-					DrawingPolicy.SetMeshRenderState(RHICmdList, View,PrimitiveSceneProxy,Mesh,BatchElementIndex,bBackFace,FPositionOnlyDepthDrawingPolicy::ElementDataType(),FDepthDrawingPolicy::ContextDataType());
+					DrawingPolicy.SetMeshRenderState(RHICmdList, View,PrimitiveSceneProxy,Mesh,BatchElementIndex,bBackFace,DitheredLODTransitionValue,FPositionOnlyDepthDrawingPolicy::ElementDataType(),FDepthDrawingPolicy::ContextDataType());
 					DrawingPolicy.DrawMesh(RHICmdList, Mesh,BatchElementIndex);
 				}
 				Mask >>= 1;
@@ -425,7 +430,7 @@ bool FDepthDrawingPolicyFactory::DrawMesh(
 		}
 		else if (!IsTranslucentBlendMode(BlendMode))
 		{
-			const bool bMaterialMasked = Material->IsMasked();
+			const bool bMaterialMasked = !Material->WritesEveryPixel();
 
 			bool bDraw = true;
 
@@ -460,7 +465,7 @@ bool FDepthDrawingPolicyFactory::DrawMesh(
 				{
 					if(Mask & 1)
 					{
-						DrawingPolicy.SetMeshRenderState(RHICmdList, View,PrimitiveSceneProxy,Mesh,BatchElementIndex,bBackFace,FMeshDrawingPolicy::ElementDataType(),FDepthDrawingPolicy::ContextDataType());
+						DrawingPolicy.SetMeshRenderState(RHICmdList, View,PrimitiveSceneProxy,Mesh,BatchElementIndex,bBackFace,DitheredLODTransitionValue,FMeshDrawingPolicy::ElementDataType(),FDepthDrawingPolicy::ContextDataType());
 						DrawingPolicy.DrawMesh(RHICmdList, Mesh,BatchElementIndex);
 					}
 					Mask >>= 1;
@@ -493,6 +498,7 @@ bool FDepthDrawingPolicyFactory::DrawDynamicMesh(
 		Mesh,
 		Mesh.Elements.Num()==1 ? 1 : (1<<Mesh.Elements.Num())-1,	// 1 bit set for each mesh element
 		bBackFace,
+		0.0f,
 		bPreFog,
 		PrimitiveSceneProxy,
 		HitProxyId);
@@ -506,6 +512,7 @@ bool FDepthDrawingPolicyFactory::DrawStaticMesh(
 	const FStaticMesh& StaticMesh,
 	const uint64& BatchElementMask,
 	bool bPreFog,
+	float DitheredLODTransitionValue,
 	const FPrimitiveSceneProxy* PrimitiveSceneProxy,
 	FHitProxyId HitProxyId
 	)
@@ -521,6 +528,7 @@ bool FDepthDrawingPolicyFactory::DrawStaticMesh(
 		StaticMesh,
 		BatchElementMask,
 		false,
+		DitheredLODTransitionValue,
 		bPreFog,
 		PrimitiveSceneProxy,
 		HitProxyId
diff --git a/Engine/Source/Runtime/Renderer/Private/DepthRendering.h b/Engine/Source/Runtime/Renderer/Private/DepthRendering.h
--- a/Engine/Source/Runtime/Renderer/Private/DepthRendering.h
+++ b/Engine/Source/Runtime/Renderer/Private/DepthRendering.h
@@ -64,6 +64,7 @@ public:
 		const FMeshBatch& Mesh,
 		int32 BatchElementIndex,
 		bool bBackFace,
+		float DitheredLODTransitionValue,
 		const ElementDataType& ElementData,
 		const ContextDataType PolicyContext
 		) const;
@@ -117,6 +118,7 @@ public:
 		const FMeshBatch& Mesh,
 		int32 BatchElementIndex,
 		bool bBackFace,
+		float DitheredLODTransitionValue,
 		const ElementDataType& ElementData,
 		const ContextDataType PolicyContext
 		) const;
@@ -163,6 +165,7 @@ public:
 		const FStaticMesh& StaticMesh,
 		const uint64& BatchElementMask,
 		bool bPreFog,
+		float DitheredLODTransitionValue,
 		const FPrimitiveSceneProxy* PrimitiveSceneProxy,
 		FHitProxyId HitProxyId
 		);
@@ -179,6 +182,7 @@ private:
 		const FMeshBatch& Mesh,
 		const uint64& BatchElementMask,
 		bool bBackFace,
+		float DitheredLODTransitionValue,
 		bool bPreFog,
 		const FPrimitiveSceneProxy* PrimitiveSceneProxy,
 		FHitProxyId HitProxyId
diff --git a/Engine/Source/Runtime/Renderer/Private/DistortionRendering.cpp b/Engine/Source/Runtime/Renderer/Private/DistortionRendering.cpp
--- a/Engine/Source/Runtime/Renderer/Private/DistortionRendering.cpp
+++ b/Engine/Source/Runtime/Renderer/Private/DistortionRendering.cpp
@@ -131,9 +131,9 @@ public:
 		FMeshMaterialShader::SetParameters(RHICmdList, GetVertexShader(), MaterialRenderProxy, *MaterialRenderProxy->GetMaterial(View->GetFeatureLevel()), *View, ESceneRenderTargetsMode::SetTextures);
 	}
 
-	void SetMesh(FRHICommandList& RHICmdList, const FVertexFactory* VertexFactory,const FSceneView& View,const FPrimitiveSceneProxy* Proxy,const FMeshBatchElement& BatchElement)
+	void SetMesh(FRHICommandList& RHICmdList, const FVertexFactory* VertexFactory,const FSceneView& View,const FPrimitiveSceneProxy* Proxy,const FMeshBatchElement& BatchElement, float DitheredLODTransitionValue)
 	{
-		FMeshMaterialShader::SetMesh(RHICmdList, GetVertexShader(),VertexFactory,View,Proxy,BatchElement);
+		FMeshMaterialShader::SetMesh(RHICmdList, GetVertexShader(),VertexFactory,View,Proxy,BatchElement,DitheredLODTransitionValue);
 	}
 };
 
@@ -234,9 +234,9 @@ public:
 
 	}
 
-	void SetMesh(FRHICommandList& RHICmdList, const FVertexFactory* VertexFactory,const FSceneView& View,const FPrimitiveSceneProxy* Proxy,const FMeshBatchElement& BatchElement)
+	void SetMesh(FRHICommandList& RHICmdList, const FVertexFactory* VertexFactory,const FSceneView& View,const FPrimitiveSceneProxy* Proxy,const FMeshBatchElement& BatchElement, float DitheredLODTransitionValue)
 	{
-		FMeshMaterialShader::SetMesh(RHICmdList, GetPixelShader(),VertexFactory,View,Proxy,BatchElement);
+		FMeshMaterialShader::SetMesh(RHICmdList, GetPixelShader(),VertexFactory,View,Proxy,BatchElement,DitheredLODTransitionValue);
 	}
 
 	virtual bool Serialize(FArchive& Ar) override
@@ -319,6 +319,7 @@ public:
 		const FMeshBatch& Mesh,
 		int32 BatchElementIndex,
 		bool bBackFace,
+		float DitheredLODTransitionValue,
 		const ElementDataType& ElementData,
 		const ContextDataType PolicyContext
 		) const;
@@ -492,6 +493,7 @@ void TDistortionMeshDrawingPolicy<DistortMeshPolicy>::SetMeshRenderState(
 	const FMeshBatch& Mesh,
 	int32 BatchElementIndex,
 	bool bBackFace,
+	float DitheredLODTransitionValue,
 	const ElementDataType& ElementData,
 	const ContextDataType PolicyContext
 	) const
@@ -502,12 +504,12 @@ void TDistortionMeshDrawingPolicy<DistortMeshPolicy>::SetMeshRenderState(
 	EmitMeshDrawEvents(RHICmdList, PrimitiveSceneProxy, Mesh);
 
 	// Set transforms
-	VertexShader->SetMesh(RHICmdList, VertexFactory,View,PrimitiveSceneProxy,BatchElement);
+	VertexShader->SetMesh(RHICmdList, VertexFactory,View,PrimitiveSceneProxy,BatchElement,DitheredLODTransitionValue);
 
 	if(HullShader && DomainShader)
 	{
-		HullShader->SetMesh(RHICmdList, VertexFactory,View,PrimitiveSceneProxy,BatchElement);
-		DomainShader->SetMesh(RHICmdList, VertexFactory,View,PrimitiveSceneProxy,BatchElement);
+		HullShader->SetMesh(RHICmdList, VertexFactory,View,PrimitiveSceneProxy,BatchElement,DitheredLODTransitionValue);
+		DomainShader->SetMesh(RHICmdList, VertexFactory,View,PrimitiveSceneProxy,BatchElement,DitheredLODTransitionValue);
 	}
 	
 
@@ -518,7 +520,7 @@ void TDistortionMeshDrawingPolicy<DistortMeshPolicy>::SetMeshRenderState(
 	{
 		if (!bInitializeOffsets)
 		{
-			DistortPixelShader->SetMesh(RHICmdList, VertexFactory,View,PrimitiveSceneProxy,BatchElement);
+			DistortPixelShader->SetMesh(RHICmdList, VertexFactory,View,PrimitiveSceneProxy,BatchElement,DitheredLODTransitionValue);
 		}
 	}
 	
@@ -564,7 +566,7 @@ public:
 	*/
 	static bool DrawStaticMesh(
 		FRHICommandList& RHICmdList, 
-		const FSceneView* View,
+		const FViewInfo* View,
 		ContextType DrawingContext,
 		const FStaticMesh& StaticMesh,
 		uint64 BatchElementMask,
@@ -608,7 +610,7 @@ bool TDistortionMeshDrawingPolicyFactory<DistortMeshPolicy>::DrawDynamicMesh(
 		DrawingPolicy.SetSharedState(RHICmdList, &View, typename TDistortionMeshDrawingPolicy<DistortMeshPolicy>::ContextDataType());
 		for (int32 BatchElementIndex = 0; BatchElementIndex < Mesh.Elements.Num(); BatchElementIndex++)
 		{
-			DrawingPolicy.SetMeshRenderState(RHICmdList, View,PrimitiveSceneProxy,Mesh,BatchElementIndex,bBackFace,typename TDistortionMeshDrawingPolicy<DistortMeshPolicy>::ElementDataType(), typename TDistortionMeshDrawingPolicy<DistortMeshPolicy>::ContextDataType());
+			DrawingPolicy.SetMeshRenderState(RHICmdList, View,PrimitiveSceneProxy,Mesh,BatchElementIndex,bBackFace,Mesh.DitheredLODTransitionAlpha,typename TDistortionMeshDrawingPolicy<DistortMeshPolicy>::ElementDataType(), typename TDistortionMeshDrawingPolicy<DistortMeshPolicy>::ContextDataType());
 			DrawingPolicy.DrawMesh(RHICmdList, Mesh,BatchElementIndex);
 		}
 
@@ -627,7 +629,7 @@ bool TDistortionMeshDrawingPolicyFactory<DistortMeshPolicy>::DrawDynamicMesh(
 template<class DistortMeshPolicy>
 bool TDistortionMeshDrawingPolicyFactory<DistortMeshPolicy>::DrawStaticMesh(
 	FRHICommandList& RHICmdList, 
-	const FSceneView* View,
+	const FViewInfo* View,
 	ContextType bInitializeOffsets,
 	const FStaticMesh& StaticMesh,
 	uint64 BatchElementMask,
@@ -652,12 +654,13 @@ bool TDistortionMeshDrawingPolicyFactory<DistortMeshPolicy>::DrawStaticMesh(
 			);
 		RHICmdList.BuildAndSetLocalBoundShaderState(DrawingPolicy.GetBoundShaderStateInput(View->GetFeatureLevel()));
 		DrawingPolicy.SetSharedState(RHICmdList, View, typename TDistortionMeshDrawingPolicy<DistortMeshPolicy>::ContextDataType());
+		float DitherValue = View->GetDitheredLODTransitionValue(StaticMesh);
 		int32 BatchElementIndex = 0;
 		do
 		{
 			if(BatchElementMask & 1)
 			{
-				DrawingPolicy.SetMeshRenderState(RHICmdList, *View,PrimitiveSceneProxy,StaticMesh,BatchElementIndex,bBackFace,
+				DrawingPolicy.SetMeshRenderState(RHICmdList, *View,PrimitiveSceneProxy,StaticMesh,BatchElementIndex,bBackFace,DitherValue,
 					typename TDistortionMeshDrawingPolicy<DistortMeshPolicy>::ElementDataType(),
 					typename TDistortionMeshDrawingPolicy<DistortMeshPolicy>::ContextDataType()
 					);
diff --git a/Engine/Source/Runtime/Renderer/Private/DrawingPolicy.cpp b/Engine/Source/Runtime/Renderer/Private/DrawingPolicy.cpp
--- a/Engine/Source/Runtime/Renderer/Private/DrawingPolicy.cpp
+++ b/Engine/Source/Runtime/Renderer/Private/DrawingPolicy.cpp
@@ -23,11 +23,13 @@ FMeshDrawingPolicy::FMeshDrawingPolicy(
 	const FMaterial& InMaterialResource,
 	bool bInOverrideWithShaderComplexity,
 	bool bInTwoSidedOverride,
+	bool bInDitheredLODTransitionOverride,
 	bool bInWireframeOverride
 	):
 	VertexFactory(InVertexFactory),
 	MaterialRenderProxy(InMaterialRenderProxy),
 	MaterialResource(&InMaterialResource),
+	bIsDitheredLODTransitionMaterial(InMaterialResource.IsDitheredLODTransition() || bInDitheredLODTransitionOverride),
 	bIsWireframeMaterial(InMaterialResource.IsWireframe() || bInWireframeOverride),
 	//convert from signed bool to unsigned uint32
 	bOverrideWithShaderComplexity(bInOverrideWithShaderComplexity != false)
@@ -51,6 +53,7 @@ void FMeshDrawingPolicy::SetMeshRenderState(
 	const FMeshBatch& Mesh,
 	int32 BatchElementIndex,
 	bool bBackFace,
+	float DitheredLODTransitionValue,
 	const ElementDataType& ElementData,
 	const ContextDataType PolicyContext
 	) const
diff --git a/Engine/Source/Runtime/Renderer/Private/ForwardBasePassRendering.cpp b/Engine/Source/Runtime/Renderer/Private/ForwardBasePassRendering.cpp
--- a/Engine/Source/Runtime/Renderer/Private/ForwardBasePassRendering.cpp
+++ b/Engine/Source/Runtime/Renderer/Private/ForwardBasePassRendering.cpp
@@ -146,6 +146,7 @@ public:
 
 	const FViewInfo& View;
 	bool bBackFace;
+	float DitheredLODTransitionValue;
 	FHitProxyId HitProxyId;
 
 	inline bool ShouldPackAmbientSH() const
@@ -163,10 +164,12 @@ public:
 	FDrawBasePassForwardShadingDynamicMeshAction(
 		const FViewInfo& InView,
 		const bool bInBackFace,
+		float InDitheredLODTransitionValue,
 		const FHitProxyId InHitProxyId
 		)
 		: View(InView)
 		, bBackFace(bInBackFace)
+		, DitheredLODTransitionValue(InDitheredLODTransitionValue)
 		, HitProxyId(InHitProxyId)
 	{}
 
@@ -214,6 +217,7 @@ public:
 				Parameters.Mesh,
 				BatchElementIndex,
 				bBackFace,
+				DitheredLODTransitionValue,
 				typename TBasePassForForwardShadingDrawingPolicy<LightMapPolicyType>::ElementDataType(LightMapElementData),
 				typename TBasePassForForwardShadingDrawingPolicy<LightMapPolicyType>::ContextDataType()
 				);
@@ -264,6 +268,7 @@ bool FBasePassForwardOpaqueDrawingPolicyFactory::DrawDynamicMesh(
 			FDrawBasePassForwardShadingDynamicMeshAction(
 				View,
 				bBackFace,
+				Mesh.DitheredLODTransitionAlpha,
 				HitProxyId
 				)
 			);
diff --git a/Engine/Source/Runtime/Renderer/Private/ForwardBasePassRendering.h b/Engine/Source/Runtime/Renderer/Private/ForwardBasePassRendering.h
--- a/Engine/Source/Runtime/Renderer/Private/ForwardBasePassRendering.h
+++ b/Engine/Source/Runtime/Renderer/Private/ForwardBasePassRendering.h
@@ -80,9 +80,9 @@ public:
 		FMeshMaterialShader::SetParameters(RHICmdList, GetVertexShader(),MaterialRenderProxy,InMaterialResource,View,TextureMode);
 	}
 
-	void SetMesh(FRHICommandList& RHICmdList, const FVertexFactory* VertexFactory,const FSceneView& View,const FPrimitiveSceneProxy* Proxy,const FMeshBatchElement& BatchElement)
+	void SetMesh(FRHICommandList& RHICmdList, const FVertexFactory* VertexFactory,const FSceneView& View,const FPrimitiveSceneProxy* Proxy,const FMeshBatchElement& BatchElement, float DitheredLODTransitionValue)
 	{
-		FMeshMaterialShader::SetMesh(RHICmdList, GetVertexShader(),VertexFactory,View,Proxy,BatchElement);
+		FMeshMaterialShader::SetMesh(RHICmdList, GetVertexShader(),VertexFactory,View,Proxy,BatchElement,DitheredLODTransitionValue);
 	}
 
 private:
@@ -154,7 +154,7 @@ public:
 		EditorCompositeParams.SetParameters(RHICmdList, MaterialResource, View, bEnableEditorPrimitveDepthTest, GetPixelShader());
 	}
 
-	void SetMesh(FRHICommandList& RHICmdList, const FVertexFactory* VertexFactory,const FSceneView& View,const FPrimitiveSceneProxy* Proxy,const FMeshBatchElement& BatchElement)
+	void SetMesh(FRHICommandList& RHICmdList, const FVertexFactory* VertexFactory,const FSceneView& View,const FPrimitiveSceneProxy* Proxy,const FMeshBatchElement& BatchElement, float DitheredLODTransitionValue)
 	{
 		FRHIPixelShader* PixelShader = GetPixelShader();
 		if (ReflectionCubemap.IsBound())
@@ -174,7 +174,7 @@ public:
 			SetTextureParameter(RHICmdList, PixelShader, ReflectionCubemap, ReflectionSampler, ReflectionTexture);
 		}
 
-		FMeshMaterialShader::SetMesh(RHICmdList, PixelShader,VertexFactory,View,Proxy,BatchElement);		
+		FMeshMaterialShader::SetMesh(RHICmdList, PixelShader,VertexFactory,View,Proxy,BatchElement,DitheredLODTransitionValue);		
 	}
 
 	virtual bool Serialize(FArchive& Ar) override
@@ -419,6 +419,7 @@ public:
 		const FMeshBatch& Mesh,
 		int32 BatchElementIndex,
 		bool bBackFace,
+		float DitheredLODTransitionValue,
 		const ElementDataType& ElementData,
 		const ContextDataType PolicyContext
 		) const
@@ -437,7 +438,7 @@ public:
 			ElementData.LightMapElementData);
 
 		const FMeshBatchElement& BatchElement = Mesh.Elements[BatchElementIndex];
-		VertexShader->SetMesh(RHICmdList, VertexFactory,View,PrimitiveSceneProxy,BatchElement);
+		VertexShader->SetMesh(RHICmdList, VertexFactory,View,PrimitiveSceneProxy,BatchElement,DitheredLODTransitionValue);
 
 #if !(UE_BUILD_SHIPPING || UE_BUILD_TEST)
 		if (bOverrideWithShaderComplexity)
@@ -457,10 +458,10 @@ public:
 		else
 #endif
 		{
-			PixelShader->SetMesh(RHICmdList, VertexFactory,View,PrimitiveSceneProxy,BatchElement);
+			PixelShader->SetMesh(RHICmdList, VertexFactory,View,PrimitiveSceneProxy,BatchElement,DitheredLODTransitionValue);
 		}
 
-		FMeshDrawingPolicy::SetMeshRenderState(RHICmdList, View,PrimitiveSceneProxy,Mesh,BatchElementIndex,bBackFace,FMeshDrawingPolicy::ElementDataType(),PolicyContext);
+		FMeshDrawingPolicy::SetMeshRenderState(RHICmdList, View,PrimitiveSceneProxy,Mesh,BatchElementIndex,bBackFace, DitheredLODTransitionValue,FMeshDrawingPolicy::ElementDataType(),PolicyContext);
 	}
 
 	friend int32 CompareDrawingPolicy(const TBasePassForForwardShadingDrawingPolicy& A,const TBasePassForForwardShadingDrawingPolicy& B)
diff --git a/Engine/Source/Runtime/Renderer/Private/ForwardTranslucentRendering.cpp b/Engine/Source/Runtime/Renderer/Private/ForwardTranslucentRendering.cpp
--- a/Engine/Source/Runtime/Renderer/Private/ForwardTranslucentRendering.cpp
+++ b/Engine/Source/Runtime/Renderer/Private/ForwardTranslucentRendering.cpp
@@ -94,16 +94,19 @@ public:
 
 	const FViewInfo& View;
 	bool bBackFace;
+	float DitheredLODTransitionValue;
 	FHitProxyId HitProxyId;
 
 	/** Initialization constructor. */
 	FDrawTranslucentMeshForwardShadingAction(
 		const FViewInfo& InView,
 		bool bInBackFace,
+		float InDitheredLODTransitionValue,
 		FHitProxyId InHitProxyId
 		):
 		View(InView),
 		bBackFace(bInBackFace),
+		DitheredLODTransitionValue(InDitheredLODTransitionValue),
 		HitProxyId(InHitProxyId)
 	{}
 	
@@ -154,6 +157,7 @@ public:
 				Parameters.Mesh,
 				BatchElementIndex,
 				bBackFace,
+				DitheredLODTransitionValue,
 				typename TBasePassForForwardShadingDrawingPolicy<LightMapPolicyType>::ElementDataType(LightMapElementData),
 				typename TBasePassForForwardShadingDrawingPolicy<LightMapPolicyType>::ContextDataType()
 				);
@@ -208,6 +212,7 @@ bool FTranslucencyForwardShadingDrawingPolicyFactory::DrawDynamicMesh(
 			FDrawTranslucentMeshForwardShadingAction(
 				View,
 				bBackFace,
+				Mesh.DitheredLODTransitionAlpha,
 				HitProxyId
 				)
 			);
diff --git a/Engine/Source/Runtime/Renderer/Private/LightMapDensityRendering.cpp b/Engine/Source/Runtime/Renderer/Private/LightMapDensityRendering.cpp
--- a/Engine/Source/Runtime/Renderer/Private/LightMapDensityRendering.cpp
+++ b/Engine/Source/Runtime/Renderer/Private/LightMapDensityRendering.cpp
@@ -127,7 +127,7 @@ bool FLightMapDensityDrawingPolicyFactory::DrawDynamicMesh(
 				DrawingPolicy.SetSharedState(RHICmdList, &View, TLightMapDensityDrawingPolicy< TLightMapPolicy<HQ_LIGHTMAP> >::ContextDataType());
 				for (int32 BatchElementIndex = 0; BatchElementIndex < Mesh.Elements.Num(); BatchElementIndex++)
 				{
-					DrawingPolicy.SetMeshRenderState(RHICmdList, View,PrimitiveSceneProxy,Mesh,BatchElementIndex,bBackFace,
+					DrawingPolicy.SetMeshRenderState(RHICmdList, View,PrimitiveSceneProxy,Mesh,BatchElementIndex,bBackFace,Mesh.DitheredLODTransitionAlpha,
 						TLightMapDensityDrawingPolicy< TLightMapPolicy<HQ_LIGHTMAP> >::ElementDataType(LightMapInteraction),
 						TLightMapDensityDrawingPolicy< TLightMapPolicy<HQ_LIGHTMAP> >::ContextDataType()
 						);
@@ -142,7 +142,7 @@ bool FLightMapDensityDrawingPolicyFactory::DrawDynamicMesh(
 				DrawingPolicy.SetSharedState(RHICmdList, &View, TLightMapDensityDrawingPolicy< TLightMapPolicy<LQ_LIGHTMAP> >::ContextDataType());
 				for (int32 BatchElementIndex = 0; BatchElementIndex < Mesh.Elements.Num(); BatchElementIndex++)
 				{
-					DrawingPolicy.SetMeshRenderState(RHICmdList, View,PrimitiveSceneProxy,Mesh,BatchElementIndex,bBackFace,
+					DrawingPolicy.SetMeshRenderState(RHICmdList, View,PrimitiveSceneProxy,Mesh,BatchElementIndex,bBackFace,Mesh.DitheredLODTransitionAlpha,
 						TLightMapDensityDrawingPolicy< TLightMapPolicy<LQ_LIGHTMAP> >::ElementDataType(LightMapInteraction),
 						TLightMapDensityDrawingPolicy< TLightMapPolicy<LQ_LIGHTMAP> >::ContextDataType()
 						);
@@ -158,7 +158,7 @@ bool FLightMapDensityDrawingPolicyFactory::DrawDynamicMesh(
 			DrawingPolicy.SetSharedState(RHICmdList, &View, TLightMapDensityDrawingPolicy<FDummyLightMapPolicy>::ContextDataType());
 			for (int32 BatchElementIndex = 0; BatchElementIndex < Mesh.Elements.Num(); BatchElementIndex++)
 			{
-				DrawingPolicy.SetMeshRenderState(RHICmdList, View,PrimitiveSceneProxy,Mesh,BatchElementIndex,bBackFace,
+				DrawingPolicy.SetMeshRenderState(RHICmdList, View,PrimitiveSceneProxy,Mesh,BatchElementIndex,bBackFace,Mesh.DitheredLODTransitionAlpha,
 					TLightMapDensityDrawingPolicy<FDummyLightMapPolicy>::ElementDataType(LightMapInteraction),
 					TLightMapDensityDrawingPolicy<FDummyLightMapPolicy>::ContextDataType()
 					);
@@ -174,7 +174,7 @@ bool FLightMapDensityDrawingPolicyFactory::DrawDynamicMesh(
 		DrawingPolicy.SetSharedState(RHICmdList, &View, TLightMapDensityDrawingPolicy<FNoLightMapPolicy>::ContextDataType());
 		for (int32 BatchElementIndex = 0; BatchElementIndex < Mesh.Elements.Num(); BatchElementIndex++)
 		{
-			DrawingPolicy.SetMeshRenderState(RHICmdList, View,PrimitiveSceneProxy,Mesh,BatchElementIndex,bBackFace,
+			DrawingPolicy.SetMeshRenderState(RHICmdList, View,PrimitiveSceneProxy,Mesh,BatchElementIndex,bBackFace,Mesh.DitheredLODTransitionAlpha,
 				TLightMapDensityDrawingPolicy<FNoLightMapPolicy>::ElementDataType(),
 				TLightMapDensityDrawingPolicy<FNoLightMapPolicy>::ContextDataType()
 				);
diff --git a/Engine/Source/Runtime/Renderer/Private/LightMapDensityRendering.h b/Engine/Source/Runtime/Renderer/Private/LightMapDensityRendering.h
--- a/Engine/Source/Runtime/Renderer/Private/LightMapDensityRendering.h
+++ b/Engine/Source/Runtime/Renderer/Private/LightMapDensityRendering.h
@@ -54,9 +54,9 @@ public:
 		FMeshMaterialShader::SetParameters(RHICmdList, GetVertexShader(), MaterialRenderProxy, *MaterialRenderProxy->GetMaterial(View.GetFeatureLevel()), View, ESceneRenderTargetsMode::SetTextures);
 	}
 
-	void SetMesh(FRHICommandList& RHICmdList, const FVertexFactory* VertexFactory,const FSceneView& View,const FPrimitiveSceneProxy* Proxy,const FMeshBatchElement& BatchElement)
+	void SetMesh(FRHICommandList& RHICmdList, const FVertexFactory* VertexFactory,const FSceneView& View,const FPrimitiveSceneProxy* Proxy,const FMeshBatchElement& BatchElement, float DitheredLODTransitionValue)
 	{
-		FMeshMaterialShader::SetMesh(RHICmdList, GetVertexShader(),VertexFactory,View,Proxy,BatchElement);
+		FMeshMaterialShader::SetMesh(RHICmdList, GetVertexShader(),VertexFactory,View,Proxy,BatchElement,DitheredLODTransitionValue);
 	}
 };
 
@@ -178,7 +178,7 @@ public:
 
 	void SetMesh(FRHICommandList& RHICmdList, const FVertexFactory* VertexFactory,const FPrimitiveSceneProxy* PrimitiveSceneProxy,const FMeshBatchElement& BatchElement,const FSceneView& View,bool bBackFace,FVector& InBuiltLightingAndSelectedFlags,FVector2D& InLightMapResolutionScale, bool bTextureMapped)
 	{
-		FMeshMaterialShader::SetMesh(RHICmdList, GetPixelShader(),VertexFactory,View,PrimitiveSceneProxy,BatchElement);
+		FMeshMaterialShader::SetMesh(RHICmdList, GetPixelShader(),VertexFactory,View,PrimitiveSceneProxy,BatchElement,0.0f);
 
 		if (LightMapDensity.IsBound())
 		{
@@ -322,18 +322,19 @@ public:
 		const FMeshBatch& Mesh,
 		int32 BatchElementIndex,
 		bool bBackFace,
+		float DitheredLODTransitionValue,
 		const ElementDataType& ElementData,
 		const ContextDataType PolicyContext
 		) const
 	{
 		const FMeshBatchElement& BatchElement = Mesh.Elements[BatchElementIndex];
 
-		VertexShader->SetMesh(RHICmdList, VertexFactory,View,PrimitiveSceneProxy,BatchElement);
+		VertexShader->SetMesh(RHICmdList, VertexFactory,View,PrimitiveSceneProxy,BatchElement,DitheredLODTransitionValue);
 
 		if(HullShader && DomainShader)
 		{
-			HullShader->SetMesh(RHICmdList, VertexFactory,View,PrimitiveSceneProxy,BatchElement);
-			DomainShader->SetMesh(RHICmdList, VertexFactory,View,PrimitiveSceneProxy,BatchElement);
+			HullShader->SetMesh(RHICmdList, VertexFactory,View,PrimitiveSceneProxy,BatchElement,DitheredLODTransitionValue);
+			DomainShader->SetMesh(RHICmdList, VertexFactory,View,PrimitiveSceneProxy,BatchElement,DitheredLODTransitionValue);
 		}
 
 		// Set the light-map policy's mesh-specific settings.
@@ -380,7 +381,7 @@ public:
 		// Adjust for the grid texture being 2x2 repeating pattern...
 		LMResolutionScale *= 0.5f;
 		PixelShader->SetMesh(RHICmdList, VertexFactory,PrimitiveSceneProxy, BatchElement, View, bBackFace, BuiltLightingAndSelectedFlags, LMResolutionScale, bTextureMapped);	
-		FMeshDrawingPolicy::SetMeshRenderState(RHICmdList, View,PrimitiveSceneProxy,Mesh,BatchElementIndex,bBackFace,FMeshDrawingPolicy::ElementDataType(),PolicyContext);
+		FMeshDrawingPolicy::SetMeshRenderState(RHICmdList, View,PrimitiveSceneProxy,Mesh,BatchElementIndex,bBackFace, DitheredLODTransitionValue,FMeshDrawingPolicy::ElementDataType(),PolicyContext);
 	}
 
 	/** 
diff --git a/Engine/Source/Runtime/Renderer/Private/SceneHitProxyRendering.cpp b/Engine/Source/Runtime/Renderer/Private/SceneHitProxyRendering.cpp
--- a/Engine/Source/Runtime/Renderer/Private/SceneHitProxyRendering.cpp
+++ b/Engine/Source/Runtime/Renderer/Private/SceneHitProxyRendering.cpp
@@ -22,9 +22,9 @@ public:
 		FMeshMaterialShader::SetParameters(RHICmdList, GetVertexShader(), MaterialRenderProxy, *MaterialRenderProxy->GetMaterial(View.GetFeatureLevel()), View, ESceneRenderTargetsMode::SetTextures);
 	}
 
-	void SetMesh(FRHICommandList& RHICmdList, const FVertexFactory* VertexFactory,const FSceneView& View,const FPrimitiveSceneProxy* Proxy,const FMeshBatchElement& BatchElement)
+	void SetMesh(FRHICommandList& RHICmdList, const FVertexFactory* VertexFactory,const FSceneView& View,const FPrimitiveSceneProxy* Proxy,const FMeshBatchElement& BatchElement, float DitheredLODTransitionValue)
 	{
-		FMeshMaterialShader::SetMesh(RHICmdList, GetVertexShader(),VertexFactory,View,Proxy,BatchElement);
+		FMeshMaterialShader::SetMesh(RHICmdList, GetVertexShader(),VertexFactory,View,Proxy,BatchElement,DitheredLODTransitionValue);
 	}
 
 	virtual bool Serialize(FArchive& Ar) override
@@ -38,7 +38,7 @@ public:
 		// Only compile the hit proxy vertex shader on PC
 		return IsPCPlatform(Platform)
 			// and only compile for the default material or materials that are masked.
-			&& (Material->IsSpecialEngineMaterial() || Material->IsMasked() || Material->MaterialMayModifyMeshPosition() || Material->IsTwoSided());
+			&& (Material->IsSpecialEngineMaterial() || !Material->WritesEveryPixel() || Material->MaterialMayModifyMeshPosition() || Material->IsTwoSided());
 	}
 
 protected:
@@ -110,7 +110,7 @@ public:
 		// Only compile the hit proxy vertex shader on PC
 		return IsPCPlatform(Platform) 
 			// and only compile for default materials or materials that are masked.
-			&& (Material->IsSpecialEngineMaterial() || Material->IsMasked() || Material->MaterialMayModifyMeshPosition() || Material->IsTwoSided());
+			&& (Material->IsSpecialEngineMaterial() || !Material->WritesEveryPixel() || Material->MaterialMayModifyMeshPosition() || Material->IsTwoSided());
 	}
 
 	FHitProxyPS(const ShaderMetaType::CompiledShaderInitializerType& Initializer):
@@ -126,9 +126,9 @@ public:
 		FMeshMaterialShader::SetParameters(RHICmdList, GetPixelShader(), MaterialRenderProxy, *MaterialRenderProxy->GetMaterial(View.GetFeatureLevel()), View, ESceneRenderTargetsMode::SetTextures);
 	}
 
-	void SetMesh(FRHICommandList& RHICmdList, const FVertexFactory* VertexFactory,const FSceneView& View,const FPrimitiveSceneProxy* Proxy,const FMeshBatchElement& BatchElement)
+	void SetMesh(FRHICommandList& RHICmdList, const FVertexFactory* VertexFactory,const FSceneView& View,const FPrimitiveSceneProxy* Proxy,const FMeshBatchElement& BatchElement, float DitheredLODTransitionValue)
 	{
-		FMeshMaterialShader::SetMesh(RHICmdList, GetPixelShader(),VertexFactory,View,Proxy,BatchElement);
+		FMeshMaterialShader::SetMesh(RHICmdList, GetPixelShader(),VertexFactory,View,Proxy,BatchElement,DitheredLODTransitionValue);
 	}
 
 	void SetHitProxyId(FRHICommandList& RHICmdList, FHitProxyId HitProxyIdValue)
@@ -210,6 +210,7 @@ void FHitProxyDrawingPolicy::SetMeshRenderState(
 	const FMeshBatch& Mesh,
 	int32 BatchElementIndex,
 	bool bBackFace,
+	float DitheredLODTransitionValue,
 	const FHitProxyId HitProxyId,
 	const ContextDataType PolicyContext
 	) const
@@ -218,15 +219,15 @@ void FHitProxyDrawingPolicy::SetMeshRenderState(
 
 	const FMeshBatchElement& BatchElement = Mesh.Elements[BatchElementIndex];
 
-	VertexShader->SetMesh(RHICmdList, VertexFactory, View, PrimitiveSceneProxy, BatchElement);
+	VertexShader->SetMesh(RHICmdList, VertexFactory, View, PrimitiveSceneProxy, BatchElement,DitheredLODTransitionValue);
 
 	if(HullShader && DomainShader)
 	{
-		HullShader->SetMesh(RHICmdList, VertexFactory, View, PrimitiveSceneProxy, BatchElement);
-		DomainShader->SetMesh(RHICmdList, VertexFactory, View, PrimitiveSceneProxy, BatchElement);
+		HullShader->SetMesh(RHICmdList, VertexFactory, View, PrimitiveSceneProxy, BatchElement,DitheredLODTransitionValue);
+		DomainShader->SetMesh(RHICmdList, VertexFactory, View, PrimitiveSceneProxy, BatchElement,DitheredLODTransitionValue);
 	}
 
-	PixelShader->SetMesh(RHICmdList, VertexFactory, View, PrimitiveSceneProxy, BatchElement);
+	PixelShader->SetMesh(RHICmdList, VertexFactory, View, PrimitiveSceneProxy, BatchElement,DitheredLODTransitionValue);
 	// Per-instance hitproxies are supplied by the vertex factory
 	if( PrimitiveSceneProxy && PrimitiveSceneProxy->HasPerInstanceHitProxies() )
 	{
@@ -251,7 +252,7 @@ void FHitProxyDrawingPolicyFactory::AddStaticMesh(FScene* Scene,FStaticMesh* Sta
 	// Add the static mesh to the DPG's hit proxy draw list.
 	const FMaterialRenderProxy* MaterialRenderProxy = StaticMesh->MaterialRenderProxy;
 	const FMaterial* Material = MaterialRenderProxy->GetMaterial(Scene->GetFeatureLevel());
-	if (!Material->IsMasked() && !Material->IsTwoSided() && !Material->MaterialModifiesMeshPosition_RenderThread())
+	if (Material->WritesEveryPixel() && !Material->IsTwoSided() && !Material->MaterialModifiesMeshPosition_RenderThread())
 	{
 		// Default material doesn't handle masked, and doesn't have the correct bIsTwoSided setting.
 		MaterialRenderProxy = UMaterial::GetDefaultMaterial(MD_Surface)->GetRenderProxy(false);
@@ -302,7 +303,7 @@ bool FHitProxyDrawingPolicyFactory::DrawDynamicMesh(
 		if( View.bAllowTranslucentPrimitivesInHitProxy || !IsTranslucentBlendMode( Material->GetBlendMode() ) || ( HitProxy && HitProxy->AlwaysAllowsTranslucentPrimitives() ) )
 #endif
 		{
-			if (!Material->IsMasked() && !Material->IsTwoSided() && !Material->MaterialModifiesMeshPosition_RenderThread())
+			if (Material->WritesEveryPixel() && !Material->IsTwoSided() && !Material->MaterialModifiesMeshPosition_RenderThread())
 			{
 				// Default material doesn't handle masked, and doesn't have the correct bIsTwoSided setting.
 				MaterialRenderProxy = UMaterial::GetDefaultMaterial(MD_Surface)->GetRenderProxy(false);
@@ -312,7 +313,7 @@ bool FHitProxyDrawingPolicyFactory::DrawDynamicMesh(
 			DrawingPolicy.SetSharedState(RHICmdList, &View, FHitProxyDrawingPolicy::ContextDataType());
 			for (int32 BatchElementIndex = 0; BatchElementIndex < Mesh.Elements.Num(); BatchElementIndex++)
 			{
-				DrawingPolicy.SetMeshRenderState(RHICmdList, View, PrimitiveSceneProxy, Mesh, BatchElementIndex, bBackFace, HitProxyId, FHitProxyDrawingPolicy::ContextDataType());
+				DrawingPolicy.SetMeshRenderState(RHICmdList, View, PrimitiveSceneProxy, Mesh, BatchElementIndex, bBackFace, Mesh.DitheredLODTransitionAlpha, HitProxyId, FHitProxyDrawingPolicy::ContextDataType());
 				DrawingPolicy.DrawMesh(RHICmdList, Mesh,BatchElementIndex);
 			}
 			return true;
diff --git a/Engine/Source/Runtime/Renderer/Private/SceneHitProxyRendering.h b/Engine/Source/Runtime/Renderer/Private/SceneHitProxyRendering.h
--- a/Engine/Source/Runtime/Renderer/Private/SceneHitProxyRendering.h
+++ b/Engine/Source/Runtime/Renderer/Private/SceneHitProxyRendering.h
@@ -46,6 +46,7 @@ public:
 		const FMeshBatch& Mesh,
 		int32 BatchElementIndex,
 		bool bBackFace,
+		float DitheredLODTransitionValue,
 		const FHitProxyId HitProxyId,
 		const ContextDataType PolicyContext
 		) const;
diff --git a/Engine/Source/Runtime/Renderer/Private/ScenePrivate.h b/Engine/Source/Runtime/Renderer/Private/ScenePrivate.h
--- a/Engine/Source/Runtime/Renderer/Private/ScenePrivate.h
+++ b/Engine/Source/Runtime/Renderer/Private/ScenePrivate.h
@@ -313,26 +313,26 @@ class FPrimitiveFadingState
 {
 public:
 	FPrimitiveFadingState()
-		: FrameNumber(0)
+		: FadeTimeScaleBias(ForceInitToZero)
+		, FrameNumber(0)
 		, EndTime(0.0f)
-		, FadeTimeScaleBias(0.0f,0.0f)
 		, bIsVisible(false)
 		, bValid(false)
 	{
 	}
 
-	/** Frame number when last updated */
-	uint32 FrameNumber;
-	
-	/** Time when fade will be finished. */
-	float EndTime;
-	
 	/** Scale and bias to use on time to calculate fade opacity */
 	FVector2D FadeTimeScaleBias;
 
 	/** The uniform buffer for the fade parameters */
 	FDistanceCullFadeUniformBufferRef UniformBuffer;
 
+	/** Frame number when last updated */
+	uint32 FrameNumber;
+	
+	/** Time when fade will be finished. */
+	float EndTime;
+	
 	/** Currently visible? */
 	bool bIsVisible;
 
diff --git a/Engine/Source/Runtime/Renderer/Private/SceneRendering.h b/Engine/Source/Runtime/Renderer/Private/SceneRendering.h
--- a/Engine/Source/Runtime/Renderer/Private/SceneRendering.h
+++ b/Engine/Source/Runtime/Renderer/Private/SceneRendering.h
@@ -369,6 +369,12 @@ public:
 	/** A map from static mesh ID to a boolean shadow depth visibility value. */
 	FSceneBitArray StaticMeshShadowDepthMap;
 
+	/** A map from static mesh ID to a boolean dithered LOD fade out value. */
+	FSceneBitArray StaticMeshFadeOutDitheredLODMap;
+
+	/** A map from static mesh ID to a boolean dithered LOD fade in value. */
+	FSceneBitArray StaticMeshFadeInDitheredLODMap;
+
 	/** An array of batch element visibility masks, valid only for meshes
 	 set visible in either StaticMeshVisibilityMap or StaticMeshShadowDepthMap. */
 	TArray<uint64,SceneRenderingAllocator> StaticMeshBatchVisibility;
@@ -517,6 +523,23 @@ public:
 	/** Create acceleration data structure and information to do forward lighting with dynamic branching. */
 	void CreateLightGrid();
 
+	FORCEINLINE_DEBUGGABLE float GetDitheredLODTransitionValue(const FStaticMesh& Mesh) const
+	{
+		float DitherValue = 0.0f;
+		if (Mesh.bDitheredLODTransition)
+		{
+			if (StaticMeshFadeOutDitheredLODMap[Mesh.Id])
+			{
+				DitherValue = GetTemporalLODTransition();
+			}
+			else if (StaticMeshFadeInDitheredLODMap[Mesh.Id])
+			{
+				DitherValue = GetTemporalLODTransition() - 1.0f;
+			}
+		}
+		return DitherValue;
+	}
+
 private:
 
 	/** Initialization that is common to the constructors. */
diff --git a/Engine/Source/Runtime/Renderer/Private/SceneVisibility.cpp b/Engine/Source/Runtime/Renderer/Private/SceneVisibility.cpp
--- a/Engine/Source/Runtime/Renderer/Private/SceneVisibility.cpp
+++ b/Engine/Source/Runtime/Renderer/Private/SceneVisibility.cpp
@@ -704,140 +704,6 @@ static int32 OcclusionCull(FRHICommandListImmediate& RHICmdList, const FScene* S
 	return NumOccludedPrimitives;
 }
 
-/** List of relevant static primitives. Indexes into Scene->Primitives. */
-typedef TArray<int32, SceneRenderingAllocator> FRelevantStaticPrimitives;
-
-/**
- * Computes view relevance for visible primitives in the view and adds them to
- * appropriate per-view rendering lists.
- * @param Scene - The scene being rendered.
- * @param View - The view for which to compute relevance.
- * @param ViewBit - Bit mask: 1 << ViewIndex where Views(ViewIndex) == View.
- * @param OutRelevantStaticPrimitives - Upon return contains a list of relevant
- *                                      static primitives.
- *                                callback for this view will have ViewBit set.
- */
-static void ComputeRelevanceForView(
-	FRHICommandListImmediate& RHICmdList,
-	const FScene* Scene,
-	FViewInfo& View,
-	uint8 ViewBit,
-	FRelevantStaticPrimitives& OutRelevantStaticPrimitives,
-	FPrimitiveViewMasks& OutHasDynamicMeshElementsMasks,
-	FPrimitiveViewMasks& OutHasDynamicEditorMeshElementsMasks
-	)
-{
-	SCOPE_CYCLE_COUNTER(STAT_ComputeViewRelevance);
-
-	// INITVIEWS_TODO: This may be more efficient to break in to:
-	// - Compute view relevance.
-	// - Gather relevant primitives.
-	// - Update last render time.
-	// But it's hard to say since the view relevance is going to be available right now.
-
-	check(OutHasDynamicMeshElementsMasks.Num() == Scene->Primitives.Num());
-
-	float CurrentWorldTime = View.Family->CurrentWorldTime;
-	float DeltaWorldTime = View.Family->DeltaWorldTime;
-
-	for (FSceneSetBitIterator BitIt(View.PrimitiveVisibilityMap); BitIt; ++BitIt)
-	{
-		const int32 BitIndex = BitIt.GetIndex();
-		FPrimitiveSceneInfo* PrimitiveSceneInfo = Scene->Primitives[BitIndex];
-		FPrimitiveViewRelevance& ViewRelevance = View.PrimitiveViewRelevanceMap[BitIndex];
-		ViewRelevance = PrimitiveSceneInfo->Proxy->GetViewRelevance(&View);
-		ViewRelevance.bInitializedThisFrame = true;
-
-		const bool bStaticRelevance = ViewRelevance.bStaticRelevance;
-		const bool bDrawRelevance = ViewRelevance.bDrawRelevance;
-		const bool bDynamicRelevance = ViewRelevance.bDynamicRelevance;
-		const bool bShadowRelevance = ViewRelevance.bShadowRelevance;
-		const bool bEditorRelevance = ViewRelevance.bEditorPrimitiveRelevance;
-		const bool bTranslucentRelevance = ViewRelevance.HasTranslucency();		
-
-		if (bStaticRelevance && (bDrawRelevance || bShadowRelevance))
-		{
-			OutRelevantStaticPrimitives.Add(BitIt.GetIndex());
-		}
-
-		if (!bDrawRelevance)
-		{
-			View.PrimitiveVisibilityMap.AccessCorrespondingBit(BitIt) = false;
-			continue;
-		}
-
-		if (bEditorRelevance)
-		{
-			// Editor primitives are rendered after post processing and composited onto the scene
-			View.VisibleEditorPrimitives.Add(PrimitiveSceneInfo);
-
-			if (GIsEditor)
-			{
-				OutHasDynamicEditorMeshElementsMasks[BitIt.GetIndex()] |= ViewBit;
-			}
-		}
-		else if(bDynamicRelevance)
-		{
-			// Keep track of visible dynamic primitives.
-			View.VisibleDynamicPrimitives.Add(PrimitiveSceneInfo);
-			OutHasDynamicMeshElementsMasks[BitIt.GetIndex()] |= ViewBit;
-		}
-
-		if (ViewRelevance.HasTranslucency() && !bEditorRelevance && ViewRelevance.bRenderInMainPass)
-		{
-			// Add to set of dynamic translucent primitives
-			View.TranslucentPrimSet.AddScenePrimitive(PrimitiveSceneInfo, View, ViewRelevance.bNormalTranslucencyRelevance, ViewRelevance.bSeparateTranslucencyRelevance);
-
-			if (ViewRelevance.bDistortionRelevance)
-			{
-				// Add to set of dynamic distortion primitives
-				View.DistortionPrimSet.AddScenePrimitive(PrimitiveSceneInfo->Proxy);
-			}
-		}
-		
-		View.ShadingModelMaskInView |= ViewRelevance.ShadingModelMaskRelevance;
-		
-		if (ViewRelevance.bRenderCustomDepth)
-		{
-			// Add to set of dynamic distortion primitives
-			View.CustomDepthSet.AddScenePrimitive(PrimitiveSceneInfo->Proxy);
-		}
-
-		// INITVIEWS_TODO: Do this in a separate pass? There are no dependencies
-		// here except maybe ParentPrimitives. This could be done in a 
-		// low-priority background task and forgotten about.
-
-		// If the primitive's last render time is older than last frame, consider
-		// it newly visible and update its visibility change time
-		if (PrimitiveSceneInfo->LastRenderTime < CurrentWorldTime - DeltaWorldTime - DELTA)
-		{
-			PrimitiveSceneInfo->LastVisibilityChangeTime = CurrentWorldTime;
-		}
-		PrimitiveSceneInfo->LastRenderTime = CurrentWorldTime;
-
-		// If the primitive is definitely unoccluded or if in Wireframe mode and the primitive is estimated
-		// to be unoccluded, then update the primitive components's LastRenderTime 
-		// on the game thread. This signals that the primitive is visible.
-		if (View.PrimitiveDefinitelyUnoccludedMap.AccessCorrespondingBit(BitIt) || (View.Family->EngineShowFlags.Wireframe && View.PrimitiveVisibilityMap.AccessCorrespondingBit(BitIt)))
-		{
-			// Update the PrimitiveComponent's LastRenderTime.
-			*(PrimitiveSceneInfo->ComponentLastRenderTime) = CurrentWorldTime;
-		}
-
-		// Cache the nearest reflection proxy if needed
-		if (PrimitiveSceneInfo->bNeedsCachedReflectionCaptureUpdate
-			// During Forward Shading, the per-object reflection is used for everything
-			// Otherwise it is just used on translucency
-			&& (!Scene->ShouldUseDeferredRenderer() || bTranslucentRelevance))
-		{
-			PrimitiveSceneInfo->CachedReflectionCaptureProxy = Scene->FindClosestReflectionCapture(Scene->PrimitiveBounds[BitIt.GetIndex()].Origin);
-			PrimitiveSceneInfo->bNeedsCachedReflectionCaptureUpdate = false;
-		}
-
-		PrimitiveSceneInfo->ConditionalUpdateStaticMeshes(RHICmdList);
-	}
-}
-
 template<class T>
 struct FRelevancePrimSet
 {
@@ -907,6 +773,8 @@ namespace EMarkMaskBits
 		StaticMeshVisibilityMapMask = 0x2,
 		StaticMeshVelocityMapMask = 0x4,
 		StaticMeshOccluderMapMask = 0x8,
+		StaticMeshFadeOutDitheredLODMapMask = 0x10,
+		StaticMeshFadeInDitheredLODMapMask = 0x20,
 	};
 }
 
@@ -1087,7 +955,7 @@ struct FRelevancePacket
 			const FPrimitiveBounds& Bounds = Scene->PrimitiveBounds[PrimitiveIndex];
 			const FPrimitiveViewRelevance& ViewRelevance = View.PrimitiveViewRelevanceMap[PrimitiveIndex];
 
-			int8 LODToRender = ComputeLODForMeshes( PrimitiveSceneInfo->StaticMeshes, View, Bounds.Origin, Bounds.SphereRadius, ViewData.ForcedLODLevel, ViewData.LODScale);
+			FLODMask LODToRender = ComputeLODForMeshes( PrimitiveSceneInfo->StaticMeshes, View, Bounds.Origin, Bounds.SphereRadius, ViewData.ForcedLODLevel, ViewData.LODScale);
 
 			float DistanceSquared = (Bounds.Origin - ViewData.ViewOrigin).SizeSquared();
 			const float LODFactorDistanceSquared = DistanceSquared * FMath::Square(View.LODDistanceFactor * ViewData.InvLODScale);
@@ -1098,11 +966,23 @@ struct FRelevancePacket
 			for(int32 MeshIndex = 0;MeshIndex < NumStaticMeshes;MeshIndex++)
 			{
 				const FStaticMesh& StaticMesh = PrimitiveSceneInfo->StaticMeshes[MeshIndex];
-				if (StaticMesh.LODIndex == LODToRender)
+				if (LODToRender.ContainsLOD(StaticMesh.LODIndex))
 				{
 					uint8 MarkMask = 0;
 					bool bNeedsBatchVisibility = false;
 
+					if (LODToRender.IsDithered())
+					{
+						if (LODToRender.DitheredLODIndices[0] == StaticMesh.LODIndex)
+						{
+							MarkMask |= EMarkMaskBits::StaticMeshFadeOutDitheredLODMapMask;
+						}
+						if (LODToRender.DitheredLODIndices[1] == StaticMesh.LODIndex)
+						{
+							MarkMask |= EMarkMaskBits::StaticMeshFadeInDitheredLODMapMask;
+						}
+					}
+
 					if (ViewRelevance.bShadowRelevance && bDrawShadowDepth && StaticMesh.CastShadow)
 					{
 						// Mark static mesh as visible in shadows.
@@ -1337,11 +1217,19 @@ static void ComputeAndMarkRelevanceForViewParallel(
 		FTaskGraphInterface::Get().WaitUntilTaskCompletes(LastRenderThread, ENamedThreads::RenderThread_Local);
 	}
 	QUICK_SCOPE_CYCLE_COUNTER(STAT_ComputeAndMarkRelevanceForViewParallel_TransposeMeshBits);
-	check(View.StaticMeshVelocityMap.Num() == NumMesh && View.StaticMeshShadowDepthMap.Num() == NumMesh && View.StaticMeshVisibilityMap.Num() == NumMesh && View.StaticMeshOccluderMap.Num() == NumMesh);
+	check(View.StaticMeshVelocityMap.Num() == NumMesh && 
+		View.StaticMeshShadowDepthMap.Num() == NumMesh && 
+		View.StaticMeshVisibilityMap.Num() == NumMesh && 
+		View.StaticMeshOccluderMap.Num() == NumMesh &&
+		View.StaticMeshFadeOutDitheredLODMap.Num() == NumMesh && 
+		View.StaticMeshFadeInDitheredLODMap.Num() == NumMesh
+		);
 	uint32* RESTRICT StaticMeshVisibilityMap_Words = View.StaticMeshVisibilityMap.GetData();
 	uint32* RESTRICT StaticMeshVelocityMap_Words = View.StaticMeshVelocityMap.GetData();
 	uint32* RESTRICT StaticMeshShadowDepthMap_Words = View.StaticMeshShadowDepthMap.GetData();
 	uint32* RESTRICT StaticMeshOccluderMap_Words = View.StaticMeshOccluderMap.GetData();
+	uint32* RESTRICT StaticMeshFadeOutDitheredLODMap_Words = View.StaticMeshFadeOutDitheredLODMap.GetData();
+	uint32* RESTRICT StaticMeshFadeInDitheredLODMap_Words = View.StaticMeshFadeInDitheredLODMap.GetData();
 	const uint64* RESTRICT MarkMasks64 = (const uint64* RESTRICT)MarkMasks;
 	const uint8* RESTRICT MarkMasks8 = MarkMasks;
 	for (int32 BaseIndex = 0; BaseIndex < NumMesh; BaseIndex += 32)
@@ -1350,6 +1238,8 @@ static void ComputeAndMarkRelevanceForViewParallel(
 		uint32 StaticMeshVelocityMap_Word = 0;
 		uint32 StaticMeshShadowDepthMap_Word = 0;
 		uint32 StaticMeshOccluderMap_Word = 0;
+		uint32 StaticMeshFadeOutDitheredLODMap_Word = 0;
+		uint32 StaticMeshFadeInDitheredLODMap_Word = 0;
 		uint32 Mask = 1;
 		bool bAny = false;
 		for (int32 QWordIndex = 0; QWordIndex < 4; QWordIndex++)
@@ -1363,6 +1253,8 @@ static void ComputeAndMarkRelevanceForViewParallel(
 					StaticMeshVelocityMap_Word |= (MaskMask & EMarkMaskBits::StaticMeshVelocityMapMask) ? Mask : 0;
 					StaticMeshShadowDepthMap_Word |= (MaskMask & EMarkMaskBits::StaticMeshShadowDepthMapMask) ? Mask : 0;
 					StaticMeshOccluderMap_Word |= (MaskMask & EMarkMaskBits::StaticMeshOccluderMapMask) ? Mask : 0;
+					StaticMeshFadeOutDitheredLODMap_Word |= (MaskMask & EMarkMaskBits::StaticMeshFadeOutDitheredLODMapMask) ? Mask : 0;
+					StaticMeshFadeInDitheredLODMap_Word |= (MaskMask & EMarkMaskBits::StaticMeshFadeInDitheredLODMapMask) ? Mask : 0;
 				}
 				bAny = true;
 			}
@@ -1374,16 +1266,20 @@ static void ComputeAndMarkRelevanceForViewParallel(
 		}
 		if (bAny)
 		{
-			checkSlow(!*StaticMeshVisibilityMap_Words && !*StaticMeshVelocityMap_Words && !*StaticMeshShadowDepthMap_Words && !*StaticMeshOccluderMap_Words);
+			checkSlow(!*StaticMeshVisibilityMap_Words && !*StaticMeshVelocityMap_Words && !*StaticMeshShadowDepthMap_Words && !*StaticMeshOccluderMap_Words && !*StaticMeshFadeOutDitheredLODMap_Words && !*StaticMeshFadeInDitheredLODMap_Words);
 			*StaticMeshVisibilityMap_Words = StaticMeshVisibilityMap_Word;
 			*StaticMeshVelocityMap_Words = StaticMeshVelocityMap_Word;
 			*StaticMeshShadowDepthMap_Words = StaticMeshShadowDepthMap_Word;
 			*StaticMeshOccluderMap_Words = StaticMeshOccluderMap_Word;
+			*StaticMeshFadeOutDitheredLODMap_Words = StaticMeshFadeOutDitheredLODMap_Word;
+			*StaticMeshFadeInDitheredLODMap_Words = StaticMeshFadeInDitheredLODMap_Word;
 		}
 		StaticMeshVisibilityMap_Words++;
 		StaticMeshVelocityMap_Words++;
 		StaticMeshShadowDepthMap_Words++;
 		StaticMeshOccluderMap_Words++;
+		StaticMeshFadeOutDitheredLODMap_Words++;
+		StaticMeshFadeInDitheredLODMap_Words++;
 	}
 }
 
@@ -1461,82 +1357,6 @@ static void MarkAllPrimitivesForReflectionProxyUpdate(FScene* Scene)
 	}
 }
 
-/**
- * Determines which static meshes in the scene are relevant to the view based on
- * for each primitive in RelevantStaticPrimitives.
- */
-static void MarkRelevantStaticMeshesForView(
-	const FScene* Scene,
-	FViewInfo& View,
-	const FRelevantStaticPrimitives& RelevantStaticPrimitives
-	)
-{
-	SCOPE_CYCLE_COUNTER(STAT_StaticRelevance);
-
-	const FMarkRelevantStaticMeshesForViewData ViewData(View);
-
-	// using a local counter to reduce memory traffic
-	int32 NumVisibleStaticMeshElements = 0;
-
-	for (int32 StaticPrimIndex = 0, Num = RelevantStaticPrimitives.Num(); StaticPrimIndex < Num; ++StaticPrimIndex)
-	{
-		int32 PrimitiveIndex = RelevantStaticPrimitives[StaticPrimIndex];
-		const FPrimitiveSceneInfo* RESTRICT PrimitiveSceneInfo = Scene->Primitives[PrimitiveIndex];
-		const FPrimitiveBounds& Bounds = Scene->PrimitiveBounds[PrimitiveIndex];
-		const FPrimitiveViewRelevance& ViewRelevance = View.PrimitiveViewRelevanceMap[PrimitiveIndex];
-
-		int8 LODToRender = ComputeLODForMeshes( PrimitiveSceneInfo->StaticMeshes, View, Bounds.Origin, Bounds.SphereRadius, ViewData.ForcedLODLevel, ViewData.LODScale);
-
-		float DistanceSquared = (Bounds.Origin - ViewData.ViewOrigin).SizeSquared();
-		const float LODFactorDistanceSquared = DistanceSquared * FMath::Square(View.LODDistanceFactor * ViewData.InvLODScale);
-		const bool bDrawShadowDepth = FMath::Square(Bounds.SphereRadius) > ViewData.MinScreenRadiusForCSMDepthSquared * LODFactorDistanceSquared;
-		const bool bDrawDepthOnly = ViewData.bForceEarlyZPass || FMath::Square(Bounds.SphereRadius) > GMinScreenRadiusForDepthPrepass * GMinScreenRadiusForDepthPrepass * LODFactorDistanceSquared;
-		
-		// We could compute this only if required by the following code. Not sure if that is worth.
-		const bool bShouldRenderVelocity = PrimitiveSceneInfo->ShouldRenderVelocity(View);
-
-		const int32 NumStaticMeshes = PrimitiveSceneInfo->StaticMeshes.Num();
-		for(int32 MeshIndex = 0;MeshIndex < NumStaticMeshes;MeshIndex++)
-		{
-			const FStaticMesh& StaticMesh = PrimitiveSceneInfo->StaticMeshes[MeshIndex];
-			if (StaticMesh.LODIndex == LODToRender)
-			{
-				bool bNeedsBatchVisibility = false;
-
-				if (ViewRelevance.bShadowRelevance && bDrawShadowDepth && StaticMesh.CastShadow)
-				{
-					// Mark static mesh as visible in shadows.
-					View.StaticMeshShadowDepthMap[StaticMesh.Id] = true;
-					bNeedsBatchVisibility = true;
-				}
-				
-				if(ViewRelevance.bDrawRelevance && !StaticMesh.bShadowOnly && (ViewRelevance.bRenderInMainPass || ViewRelevance.bRenderCustomDepth))
-				{
-					// Mark static mesh as visible for rendering
-					View.StaticMeshVisibilityMap[StaticMesh.Id] = true;
-					View.StaticMeshVelocityMap[StaticMesh.Id] = bShouldRenderVelocity;
-					++NumVisibleStaticMeshElements;
-
-					// If the static mesh is an occluder, check whether it covers enough of the screen to be used as an occluder.
-					if(	StaticMesh.bUseAsOccluder && bDrawDepthOnly )
-					{
-						View.StaticMeshOccluderMap[StaticMesh.Id] = true;
-					}
-					bNeedsBatchVisibility = true;
-				}
-
-				// Static meshes with a single element always draw, as if the mask were 0x1.
-				if(bNeedsBatchVisibility && StaticMesh.Elements.Num() > 1)
-				{
-					View.StaticMeshBatchVisibility[StaticMesh.Id] = StaticMesh.VertexFactory->GetStaticBatchElementVisibility(View, &StaticMesh);
-				}
-			}
-		}
-	}
-
-	View.NumVisibleStaticMeshElements += NumVisibleStaticMeshElements;
-}
-
 /**
  * Helper for InitViews to detect large camera movement, in both angle and position.
  */
@@ -1852,13 +1672,6 @@ void FSceneRenderer::PreVisibilityFrameSetup(FRHICommandListImmediate& RHICmdLis
 	}
 }
 
-static TAutoConsoleVariable<int32> CVarLegacySingleThreadedRelevance(
-	TEXT("r.LegacySingleThreadedRelevance"),
-	0,  
-	TEXT("Toggles the legacy codepath for view relevance."),
-	ECVF_RenderThreadSafe
-	);
-
 void FSceneRenderer::ComputeViewVisibility(FRHICommandListImmediate& RHICmdList)
 {
 	SCOPE_CYCLE_COUNTER(STAT_ViewVisibilityTime);
@@ -1901,6 +1714,8 @@ void FSceneRenderer::ComputeViewVisibility(FRHICommandListImmediate& RHICmdList)
 		View.PrimitiveFadeUniformBuffers.AddZeroed(Scene->Primitives.Num());
 		View.StaticMeshVisibilityMap.Init(false,Scene->StaticMeshes.GetMaxIndex());
 		View.StaticMeshOccluderMap.Init(false,Scene->StaticMeshes.GetMaxIndex());
+		View.StaticMeshFadeOutDitheredLODMap.Init(false,Scene->StaticMeshes.GetMaxIndex());
+		View.StaticMeshFadeInDitheredLODMap.Init(false,Scene->StaticMeshes.GetMaxIndex());
 		View.StaticMeshVelocityMap.Init(false,Scene->StaticMeshes.GetMaxIndex());
 		View.StaticMeshShadowDepthMap.Init(false,Scene->StaticMeshes.GetMaxIndex());
 		View.StaticMeshBatchVisibility.AddZeroed(Scene->StaticMeshes.GetMaxIndex());
@@ -2069,18 +1884,7 @@ void FSceneRenderer::ComputeViewVisibility(FRHICommandListImmediate& RHICmdList)
 
 		{
 			SCOPE_CYCLE_COUNTER(STAT_ViewRelevance);
-			// Compute view relevance for all visible primitives.
-			if (!CVarLegacySingleThreadedRelevance.GetValueOnRenderThread())
-			{
-				ComputeAndMarkRelevanceForViewParallel(RHICmdList, Scene, View, ViewBit, HasDynamicMeshElementsMasks, HasDynamicEditorMeshElementsMasks);
-			}
-			else
-			{
-				// This array contains a list of relevant static primities.
-				FRelevantStaticPrimitives RelevantStaticPrimitives;
-				ComputeRelevanceForView(RHICmdList, Scene, View, ViewBit, RelevantStaticPrimitives, HasDynamicMeshElementsMasks, HasDynamicEditorMeshElementsMasks);
-				MarkRelevantStaticMeshesForView(Scene, View, RelevantStaticPrimitives);
-			}
+			ComputeAndMarkRelevanceForViewParallel(RHICmdList, Scene, View, ViewBit, HasDynamicMeshElementsMasks, HasDynamicEditorMeshElementsMasks);
 
 			if (bLODSceneTreeActive)
 			{
diff --git a/Engine/Source/Runtime/Renderer/Private/ShaderBaseClasses.cpp b/Engine/Source/Runtime/Renderer/Private/ShaderBaseClasses.cpp
--- a/Engine/Source/Runtime/Renderer/Private/ShaderBaseClasses.cpp
+++ b/Engine/Source/Runtime/Renderer/Private/ShaderBaseClasses.cpp
@@ -394,6 +394,7 @@ void FMeshMaterialShader::SetMesh(
 	const FSceneView& View,
 	const FPrimitiveSceneProxy* Proxy,
 	const FMeshBatchElement& BatchElement,
+	float DitheredLODTransitionValue,
 	uint32 DataFlags )
 {
 	// Set the mesh for the vertex factory
@@ -414,6 +415,10 @@ void FMeshMaterialShader::SetMesh(
 	{
 		SetUniformBufferParameter(RHICmdList, ShaderRHI,LODParameter,GetPrimitiveFadeUniformBufferParameter(View, Proxy));
 	}
+	if (NonInstancedDitherLODFactorParameter.IsBound())
+	{
+		SetShaderValue(RHICmdList, ShaderRHI, NonInstancedDitherLODFactorParameter, DitheredLODTransitionValue);
+	}
 }
 
 #define IMPLEMENT_MESH_MATERIAL_SHADER_SetMesh( ShaderRHIParamRef ) \
@@ -424,6 +429,7 @@ void FMeshMaterialShader::SetMesh(
 		const FSceneView& View,					\
 		const FPrimitiveSceneProxy* Proxy,		\
 		const FMeshBatchElement& BatchElement,	\
+		float DitheredLODTransitionValue,		\
 		uint32 DataFlags						\
 	);
 
@@ -438,6 +444,7 @@ bool FMeshMaterialShader::Serialize(FArchive& Ar)
 {
 	bool bShaderHasOutdatedParameters = FMaterialShader::Serialize(Ar);
 	bShaderHasOutdatedParameters |= Ar << VertexFactoryParameters;
+	Ar << NonInstancedDitherLODFactorParameter;
 	return bShaderHasOutdatedParameters;
 }
 
diff --git a/Engine/Source/Runtime/Renderer/Private/ShaderBaseClasses.h b/Engine/Source/Runtime/Renderer/Private/ShaderBaseClasses.h
--- a/Engine/Source/Runtime/Renderer/Private/ShaderBaseClasses.h
+++ b/Engine/Source/Runtime/Renderer/Private/ShaderBaseClasses.h
@@ -10,6 +10,8 @@
 
 /** The uniform shader parameters associated with a LOD fade. */
 // This was moved out of ScenePrivate.h to workaround MSVC vs clang template issue (it's used in this header file, so needs to be declared earlier)
+// Z is the dither fade value (-1 = just fading in, 0 no fade, 1 = just faded out)
+// W is unused and zero
 BEGIN_UNIFORM_BUFFER_STRUCT(FDistanceCullFadeUniformShaderParameters,)
 	DECLARE_UNIFORM_BUFFER_STRUCT_MEMBER_EX(FVector2D,FadeTimeScaleBias, EShaderPrecisionModifier::Half)
 END_UNIFORM_BUFFER_STRUCT(FDistanceCullFadeUniformShaderParameters)
@@ -56,9 +58,9 @@ public:
 		FMeshMaterialShader::SetParameters(RHICmdList, (FHullShaderRHIParamRef)GetHullShader(), MaterialRenderProxy, *MaterialRenderProxy->GetMaterial(View.GetFeatureLevel()), View, ESceneRenderTargetsMode::SetTextures);
 	}
 
-	void SetMesh(FRHICommandList& RHICmdList, const FVertexFactory* VertexFactory,const FSceneView& View,const FPrimitiveSceneProxy* Proxy,const FMeshBatchElement& BatchElement)
+	void SetMesh(FRHICommandList& RHICmdList, const FVertexFactory* VertexFactory,const FSceneView& View,const FPrimitiveSceneProxy* Proxy,const FMeshBatchElement& BatchElement, float DitheredLODTransitionValue)
 	{
-		FMeshMaterialShader::SetMesh(RHICmdList, (FHullShaderRHIParamRef)GetHullShader(),VertexFactory,View,Proxy,BatchElement);
+		FMeshMaterialShader::SetMesh(RHICmdList, (FHullShaderRHIParamRef)GetHullShader(),VertexFactory,View,Proxy,BatchElement,DitheredLODTransitionValue);
 	}
 };
 
@@ -102,9 +104,9 @@ public:
 		FMeshMaterialShader::SetParameters(RHICmdList, (FDomainShaderRHIParamRef)GetDomainShader(), MaterialRenderProxy, *MaterialRenderProxy->GetMaterial(View.GetFeatureLevel()), View, ESceneRenderTargetsMode::SetTextures);
 	}
 
-	void SetMesh(FRHICommandList& RHICmdList, const FVertexFactory* VertexFactory,const FSceneView& View,const FPrimitiveSceneProxy* Proxy,const FMeshBatchElement& BatchElement)
+	void SetMesh(FRHICommandList& RHICmdList, const FVertexFactory* VertexFactory,const FSceneView& View,const FPrimitiveSceneProxy* Proxy,const FMeshBatchElement& BatchElement, float DitheredLODTransitionValue)
 	{
-		FMeshMaterialShader::SetMesh(RHICmdList, (FDomainShaderRHIParamRef)GetDomainShader(),VertexFactory,View,Proxy,BatchElement);
+		FMeshMaterialShader::SetMesh(RHICmdList, (FDomainShaderRHIParamRef)GetDomainShader(),VertexFactory,View,Proxy,BatchElement,DitheredLODTransitionValue);
 	}
 };
 
diff --git a/Engine/Source/Runtime/Renderer/Private/ShadowRendering.cpp b/Engine/Source/Runtime/Renderer/Private/ShadowRendering.cpp
--- a/Engine/Source/Runtime/Renderer/Private/ShadowRendering.cpp
+++ b/Engine/Source/Runtime/Renderer/Private/ShadowRendering.cpp
@@ -115,9 +115,9 @@ public:
 		ShadowParameters.SetVertexShader(RHICmdList, this, View, ShadowInfo, MaterialRenderProxy);
 	}
 
-	void SetMesh(FRHICommandList& RHICmdList, const FVertexFactory* VertexFactory,const FSceneView& View,const FPrimitiveSceneProxy* Proxy,const FMeshBatchElement& BatchElement)
+	void SetMesh(FRHICommandList& RHICmdList, const FVertexFactory* VertexFactory,const FSceneView& View,const FPrimitiveSceneProxy* Proxy,const FMeshBatchElement& BatchElement, float DitheredLODTransitionValue)
 	{
-		FMeshMaterialShader::SetMesh(RHICmdList, GetVertexShader(),VertexFactory,View,Proxy,BatchElement);
+		FMeshMaterialShader::SetMesh(RHICmdList, GetVertexShader(),VertexFactory,View,Proxy,BatchElement,DitheredLODTransitionValue);
 	}
 
 private:
@@ -167,8 +167,8 @@ public:
 		else
 		{
 			return (Material->IsSpecialEngineMaterial()
-						// Masked and WPO materials need their shaders but cannot be used with a position only stream.
-						|| ((Material->IsMasked() || Material->MaterialMayModifyMeshPosition()) && !bUsePositionOnlyStream))
+					// Masked and WPO materials need their shaders but cannot be used with a position only stream.
+					|| ((!Material->WritesEveryPixel() || Material->MaterialMayModifyMeshPosition()) && !bUsePositionOnlyStream))
 					// Only compile one pass point light shaders for feature levels >= SM4
 					&& (ShaderMode != VertexShadowDepth_OnePassPointLight || IsFeatureLevelSupported(Platform, ERHIFeatureLevel::SM4))
 					// Only compile position-only shaders for vertex factories that support it.
@@ -469,9 +469,9 @@ public:
 		}
 	}
 
-	void SetMesh(FRHICommandList& RHICmdList, const FVertexFactory* VertexFactory,const FSceneView& View,const FPrimitiveSceneProxy* Proxy,const FMeshBatchElement& BatchElement)
+	void SetMesh(FRHICommandList& RHICmdList, const FVertexFactory* VertexFactory,const FSceneView& View,const FPrimitiveSceneProxy* Proxy,const FMeshBatchElement& BatchElement,float DitheredLODTransitionValue)
 	{
-		FMeshMaterialShader::SetMesh(RHICmdList, GetPixelShader(),VertexFactory,View,Proxy,BatchElement);
+		FMeshMaterialShader::SetMesh(RHICmdList, GetPixelShader(),VertexFactory,View,Proxy,BatchElement,DitheredLODTransitionValue);
 	}
 
 	virtual bool Serialize(FArchive& Ar) override
@@ -524,7 +524,7 @@ public:
 		{
 			return (Material->IsSpecialEngineMaterial()
 				// Only compile for masked or lit translucent materials
-				|| Material->IsMasked())
+				|| !Material->WritesEveryPixel())
 				&& ShaderMode == PixelShadowDepth_NonPerspectiveCorrect
 				// Don't render ShadowDepth for translucent unlit materials
 				&& (!IsTranslucentBlendMode(Material->GetBlendMode()) && Material->GetShadingModel() != MSM_Unlit)
@@ -545,7 +545,7 @@ public:
 			//Note: This logic needs to stay in sync with OverrideWithDefaultMaterialForShadowDepth!
 			return (Material->IsSpecialEngineMaterial()
 					// Only compile for masked or lit translucent materials
-					|| Material->IsMasked()
+					|| !Material->WritesEveryPixel()
 					|| (Material->MaterialMayModifyMeshPosition() && Material->IsUsedWithInstancedStaticMeshes())
 					// Perspective correct rendering needs a pixel shader and WPO materials can't be overridden with default material.
 					|| (ShaderMode == PixelShadowDepth_PerspectiveCorrect && Material->MaterialMayModifyMeshPosition()))
@@ -686,7 +686,7 @@ void OverrideWithDefaultMaterialForShadowDepth(
 	ERHIFeatureLevel::Type InFeatureLevel)
 {
 	// Override with the default material when possible.
-	if (!InOutMaterialResource->IsMasked() &&						// Don't override masked materials.
+	if (InOutMaterialResource->WritesEveryPixel() &&						// Don't override masked materials.
 		!InOutMaterialResource->MaterialModifiesMeshPosition_RenderThread() &&	// Don't override materials using world position offset.
 		!bReflectiveShadowmap)										// Don't override when rendering reflective shadow maps.
 	{
@@ -765,7 +765,7 @@ FShadowDepthDrawingPolicy<bRenderingReflectiveShadowMaps>::FShadowDepthDrawingPo
 
 	bUsePositionOnlyVS = !bRenderingReflectiveShadowMaps
 		&& VertexFactory->SupportsPositionOnlyStream()
-		&& !MaterialResource->IsMasked()
+		&& MaterialResource->WritesEveryPixel()
 		&& !MaterialResource->MaterialModifiesMeshPosition_RenderThread();
 
 	// Vertex related shaders
@@ -840,7 +840,7 @@ FShadowDepthDrawingPolicy<bRenderingReflectiveShadowMaps>::FShadowDepthDrawingPo
 	}
 
 	// Pixel shaders
-	if (!MaterialResource->IsMasked() && !bUsePerspectiveCorrectShadowDepths && !bRenderingReflectiveShadowMaps && VertexFactory->SupportsNullPixelShader())
+	if (MaterialResource->WritesEveryPixel() && !bUsePerspectiveCorrectShadowDepths && !bRenderingReflectiveShadowMaps && VertexFactory->SupportsNullPixelShader())
 	{
 		// No pixel shader necessary.
 		PixelShader = NULL;
@@ -946,6 +946,7 @@ void FShadowDepthDrawingPolicy<bRenderingReflectiveShadowMaps>::SetMeshRenderSta
 	const FMeshBatch& Mesh,
 	int32 BatchElementIndex,
 	bool bBackFace,
+	float DitheredLODTransitionValue,
 	const ElementDataType& ElementData,
 	const ContextDataType PolicyContext
 	) const
@@ -954,12 +955,12 @@ void FShadowDepthDrawingPolicy<bRenderingReflectiveShadowMaps>::SetMeshRenderSta
 
 	EmitMeshDrawEvents(RHICmdList, PrimitiveSceneProxy, Mesh);
 
-	VertexShader->SetMesh(RHICmdList, VertexFactory,View,PrimitiveSceneProxy,BatchElement);
+	VertexShader->SetMesh(RHICmdList, VertexFactory,View,PrimitiveSceneProxy,BatchElement,DitheredLODTransitionValue);
 
 	if( HullShader && DomainShader )
 	{
-		HullShader->SetMesh(RHICmdList, VertexFactory,View,PrimitiveSceneProxy,BatchElement);
-		DomainShader->SetMesh(RHICmdList, VertexFactory,View,PrimitiveSceneProxy,BatchElement);
+		HullShader->SetMesh(RHICmdList, VertexFactory,View,PrimitiveSceneProxy,BatchElement,DitheredLODTransitionValue);
+		DomainShader->SetMesh(RHICmdList, VertexFactory,View,PrimitiveSceneProxy,BatchElement,DitheredLODTransitionValue);
 	}
 	if (GeometryShader)
 	{
@@ -967,7 +968,7 @@ void FShadowDepthDrawingPolicy<bRenderingReflectiveShadowMaps>::SetMeshRenderSta
 	}
 	if (PixelShader)
 	{
-		PixelShader->SetMesh(RHICmdList, VertexFactory,View,PrimitiveSceneProxy,BatchElement);
+		PixelShader->SetMesh(RHICmdList, VertexFactory,View,PrimitiveSceneProxy,BatchElement,DitheredLODTransitionValue);
 	}
 	// Not calling FMeshDrawingPolicy::SetMeshRenderState as SetSharedState sets the rasterizer state
 }
@@ -1103,7 +1104,7 @@ bool FShadowDepthDrawingPolicyFactory::DrawDynamicMesh(
 				DrawingPolicy.SetSharedState(RHICmdList, &View, PolicyContext);
 				for (int32 BatchElementIndex = 0, Num = Mesh.Elements.Num(); BatchElementIndex < Num; BatchElementIndex++)
 				{
-					DrawingPolicy.SetMeshRenderState(RHICmdList, View,PrimitiveSceneProxy,Mesh,BatchElementIndex,bBackFace,FMeshDrawingPolicy::ElementDataType(),PolicyContext);
+					DrawingPolicy.SetMeshRenderState(RHICmdList, View,PrimitiveSceneProxy,Mesh,BatchElementIndex,bBackFace,Mesh.DitheredLODTransitionAlpha,FMeshDrawingPolicy::ElementDataType(),PolicyContext);
 					DrawingPolicy.DrawMesh(RHICmdList, Mesh,BatchElementIndex);
 				}
 			}
@@ -1125,7 +1126,7 @@ bool FShadowDepthDrawingPolicyFactory::DrawDynamicMesh(
 				DrawingPolicy.SetSharedState(RHICmdList, &View, PolicyContext);
 				for (int32 BatchElementIndex = 0; BatchElementIndex < Mesh.Elements.Num(); BatchElementIndex++)
 				{
-					DrawingPolicy.SetMeshRenderState(RHICmdList, View,PrimitiveSceneProxy,Mesh,BatchElementIndex,bBackFace,FMeshDrawingPolicy::ElementDataType(),PolicyContext);
+					DrawingPolicy.SetMeshRenderState(RHICmdList, View,PrimitiveSceneProxy,Mesh,BatchElementIndex,bBackFace,Mesh.DitheredLODTransitionAlpha,FMeshDrawingPolicy::ElementDataType(),PolicyContext);
 					DrawingPolicy.DrawMesh(RHICmdList, Mesh,BatchElementIndex);
 				}
 			}
@@ -1279,6 +1280,7 @@ void DrawMeshElements(FRHICommandList& RHICmdList, FShadowDepthDrawingPolicy<bRe
 	checkSlow(CompareDrawingPolicy(DebugPolicy, SharedDrawingPolicy) == 0);
 #endif
 
+	float DitherValue = View.GetDitheredLODTransitionValue(*Mesh);
 	// Render only those batch elements that match the current LOD
 	uint64 BatchElementMask = Mesh->Elements.Num() == 1 ? 1 : View.StaticMeshBatchVisibility[Mesh->Id];
 	int32 BatchElementIndex = 0;
@@ -1286,7 +1288,7 @@ void DrawMeshElements(FRHICommandList& RHICmdList, FShadowDepthDrawingPolicy<bRe
 	{
 		if(BatchElementMask & 1)
 		{
-			SharedDrawingPolicy.SetMeshRenderState(RHICmdList, View, Mesh->PrimitiveSceneInfo->Proxy, *Mesh, BatchElementIndex, false, FMeshDrawingPolicy::ElementDataType(),PolicyContext);
+			SharedDrawingPolicy.SetMeshRenderState(RHICmdList, View, Mesh->PrimitiveSceneInfo->Proxy, *Mesh, BatchElementIndex, false, DitherValue, FMeshDrawingPolicy::ElementDataType(),PolicyContext);
 			SharedDrawingPolicy.DrawMesh(RHICmdList, *Mesh, BatchElementIndex);
 			INC_DWORD_STAT(STAT_ShadowDynamicPathDrawCalls);
 		}
@@ -1968,6 +1970,7 @@ void FProjectedShadowInfo::RenderProjection(FRHICommandListImmediate& RHICmdList
 
 						if (View->StaticMeshVisibilityMap[StaticMesh.Id])
 						{
+							float DitherValue = View->GetDitheredLODTransitionValue(StaticMesh);
 							FDepthDrawingPolicyFactory::DrawStaticMesh(
 								RHICmdList, 
 								*View,
@@ -1975,6 +1978,7 @@ void FProjectedShadowInfo::RenderProjection(FRHICommandListImmediate& RHICmdList
 								StaticMesh,
 								StaticMesh.Elements.Num() == 1 ? 1 : View->StaticMeshBatchVisibility[StaticMesh.Id],
 								true,
+								DitherValue,
 								ReceiverPrimitiveSceneInfo->Proxy,
 								StaticMesh.BatchHitProxyId
 								);
diff --git a/Engine/Source/Runtime/Renderer/Private/ShadowRendering.h b/Engine/Source/Runtime/Renderer/Private/ShadowRendering.h
--- a/Engine/Source/Runtime/Renderer/Private/ShadowRendering.h
+++ b/Engine/Source/Runtime/Renderer/Private/ShadowRendering.h
@@ -369,6 +369,7 @@ public:
 		const FMeshBatch& Mesh,
 		int32 BatchElementIndex,
 		bool bBackFace,
+		float DitheredLODTransitionValue,
 		const ElementDataType& ElementData,
 		const ContextDataType PolicyContext
 		) const;
diff --git a/Engine/Source/Runtime/Renderer/Private/ShadowSetup.cpp b/Engine/Source/Runtime/Renderer/Private/ShadowSetup.cpp
--- a/Engine/Source/Runtime/Renderer/Private/ShadowSetup.cpp
+++ b/Engine/Source/Runtime/Renderer/Private/ShadowSetup.cpp
@@ -777,17 +777,21 @@ void FProjectedShadowInfo::AddSubjectPrimitive(FPrimitiveSceneInfo* PrimitiveSce
 						// The old LOD will continue to be drawn even though a different LOD would be chosen by distance.
 						else
 						{
-							int8 LODToRender = 0;
+							FLODMask LODToRender;
 							int32 ForcedLODLevel = (CurrentView.Family->EngineShowFlags.LOD) ? GetCVarForceLOD() : 0;
 
 							// Add the primitive's static mesh elements to the draw lists.
 							if ( bReflectiveShadowmap) 
 							{
-								LODToRender = -CHAR_MAX;
+								int8 LODToRenderScan = -CHAR_MAX;
 								// Force the lowest detail LOD Level in reflective shadow maps.
 								for (int32 Index = 0; Index < PrimitiveSceneInfo->StaticMeshes.Num(); Index++)
 								{
-									LODToRender = FMath::Max<int8>(PrimitiveSceneInfo->StaticMeshes[Index].LODIndex, LODToRender);
+									LODToRenderScan = FMath::Max<int8>(PrimitiveSceneInfo->StaticMeshes[Index].LODIndex, LODToRenderScan);
+								}
+								if (LODToRenderScan != -CHAR_MAX)
+								{
+									LODToRender.SetLOD(LODToRenderScan);
 								}
 							}
 							else
@@ -801,7 +805,7 @@ void FProjectedShadowInfo::AddSubjectPrimitive(FPrimitiveSceneInfo* PrimitiveSce
 							for (int32 MeshIndex = 0; MeshIndex < PrimitiveSceneInfo->StaticMeshes.Num(); MeshIndex++)
 							{
 								const FStaticMesh& StaticMesh = PrimitiveSceneInfo->StaticMeshes[MeshIndex];
-								if (StaticMesh.CastShadow && StaticMesh.LODIndex == LODToRender)
+								if (StaticMesh.CastShadow && LODToRender.ContainsLOD(StaticMesh.LODIndex))
 								{
 									if (bWholeSceneDirectionalShadow)
 									{
diff --git a/Engine/Source/Runtime/Renderer/Private/StaticMeshDrawList.inl b/Engine/Source/Runtime/Renderer/Private/StaticMeshDrawList.inl
--- a/Engine/Source/Runtime/Renderer/Private/StaticMeshDrawList.inl
+++ b/Engine/Source/Runtime/Renderer/Private/StaticMeshDrawList.inl
@@ -96,7 +96,7 @@ void TStaticMeshDrawList<DrawingPolicyType>::DrawElement(
 			for (uint32 BackFace = 0; BackFace < BackFaceEnd; ++BackFace)
 			{
 				INC_DWORD_STAT(STAT_StaticDrawListMeshDrawCalls);
-
+				float DitherValue = View.GetDitheredLODTransitionValue(*Element.Mesh);
 				DrawingPolicyLink->DrawingPolicy.SetMeshRenderState(
 					RHICmdList, 
 					View,
@@ -104,6 +104,7 @@ void TStaticMeshDrawList<DrawingPolicyType>::DrawElement(
 					*Element.Mesh,
 					BatchElementIndex,
 					!!BackFace,
+					DitherValue,
 					Element.PolicyData,
 					PolicyContext
 					);
diff --git a/Engine/Source/Runtime/Renderer/Private/TranslucentLighting.cpp b/Engine/Source/Runtime/Renderer/Private/TranslucentLighting.cpp
--- a/Engine/Source/Runtime/Renderer/Private/TranslucentLighting.cpp
+++ b/Engine/Source/Runtime/Renderer/Private/TranslucentLighting.cpp
@@ -184,9 +184,9 @@ public:
 		ShadowParameters.SetVertexShader(RHICmdList, this, View, ShadowInfo, MaterialRenderProxy);
 	}
 
-	void SetMesh(FRHICommandList& RHICmdList, const FVertexFactory* VertexFactory,const FSceneView& View,const FPrimitiveSceneProxy* Proxy,const FMeshBatchElement& BatchElement)
+	void SetMesh(FRHICommandList& RHICmdList, const FVertexFactory* VertexFactory,const FSceneView& View,const FPrimitiveSceneProxy* Proxy,const FMeshBatchElement& BatchElement, float DitheredLODTransitionValue)
 	{
-		FMeshMaterialShader::SetMesh(RHICmdList, GetVertexShader(),VertexFactory,View,Proxy,BatchElement);
+		FMeshMaterialShader::SetMesh(RHICmdList, GetVertexShader(),VertexFactory,View,Proxy,BatchElement,DitheredLODTransitionValue);
 	}
 
 private:
@@ -268,9 +268,9 @@ public:
 		TranslucencyProjectionParameters.Set(RHICmdList, this);
 	}
 
-	void SetMesh(FRHICommandList& RHICmdList, const FVertexFactory* VertexFactory,const FSceneView& View,const FPrimitiveSceneProxy* Proxy,const FMeshBatchElement& BatchElement)
+	void SetMesh(FRHICommandList& RHICmdList, const FVertexFactory* VertexFactory,const FSceneView& View,const FPrimitiveSceneProxy* Proxy,const FMeshBatchElement& BatchElement, float DitheredLODTransitionValue)
 	{
-		FMeshMaterialShader::SetMesh(RHICmdList, GetPixelShader(),VertexFactory,View,Proxy,BatchElement);
+		FMeshMaterialShader::SetMesh(RHICmdList, GetPixelShader(),VertexFactory,View,Proxy,BatchElement,DitheredLODTransitionValue);
 	}
 
 	virtual bool Serialize(FArchive& Ar) override
@@ -381,15 +381,16 @@ public:
 		const FMeshBatch& Mesh,
 		int32 BatchElementIndex,
 		bool bBackFace,
+		float DitheredLODTransitionValue,
 		const ElementDataType& ElementData,
 		const ContextDataType PolicyContext
 		) const
 	{
 		const FMeshBatchElement& BatchElement = Mesh.Elements[BatchElementIndex];
-		VertexShader->SetMesh(RHICmdList, VertexFactory,View,PrimitiveSceneProxy,BatchElement);
-		PixelShader->SetMesh(RHICmdList, VertexFactory,View,PrimitiveSceneProxy,BatchElement);
+		VertexShader->SetMesh(RHICmdList, VertexFactory,View,PrimitiveSceneProxy,BatchElement,DitheredLODTransitionValue);
+		PixelShader->SetMesh(RHICmdList, VertexFactory,View,PrimitiveSceneProxy,BatchElement,DitheredLODTransitionValue);
 
-		FMeshDrawingPolicy::SetMeshRenderState(RHICmdList, View,PrimitiveSceneProxy,Mesh,BatchElementIndex,bBackFace,ElementData,PolicyContext);
+		FMeshDrawingPolicy::SetMeshRenderState(RHICmdList, View,PrimitiveSceneProxy,Mesh,BatchElementIndex,bBackFace, DitheredLODTransitionValue,ElementData,PolicyContext);
 	}
 
 private:
@@ -441,7 +442,7 @@ public:
 
 				for (int32 BatchElementIndex = 0; BatchElementIndex < Mesh.Elements.Num(); BatchElementIndex++)
 				{
-					DrawingPolicy.SetMeshRenderState(RHICmdList, View,PrimitiveSceneProxy,Mesh,BatchElementIndex,bBackFace,
+					DrawingPolicy.SetMeshRenderState(RHICmdList, View,PrimitiveSceneProxy,Mesh,BatchElementIndex,bBackFace,Mesh.DitheredLODTransitionAlpha,
 						FTranslucencyShadowDepthDrawingPolicy::ElementDataType(),
 						FTranslucencyShadowDepthDrawingPolicy::ContextDataType(DrawingContext.ShadowInfo)
 						);
diff --git a/Engine/Source/Runtime/Renderer/Private/TranslucentRendering.cpp b/Engine/Source/Runtime/Renderer/Private/TranslucentRendering.cpp
--- a/Engine/Source/Runtime/Renderer/Private/TranslucentRendering.cpp
+++ b/Engine/Source/Runtime/Renderer/Private/TranslucentRendering.cpp
@@ -222,12 +222,14 @@ public:
 	const FProjectedShadowInfo* TranslucentSelfShadow;
 	FHitProxyId HitProxyId;
 	bool bBackFace;
+	float DitheredLODTransitionValue;
 	bool bUseTranslucentSelfShadowing;
 
 	/** Initialization constructor. */
 	FDrawTranslucentMeshAction(
 		const FViewInfo& InView,
 		bool bInBackFace,
+		float InDitheredLODTransitionValue,
 		FHitProxyId InHitProxyId,
 		const FProjectedShadowInfo* InTranslucentSelfShadow,
 		bool bInUseTranslucentSelfShadowing
@@ -236,6 +238,7 @@ public:
 		TranslucentSelfShadow(InTranslucentSelfShadow),
 		HitProxyId(InHitProxyId),
 		bBackFace(bInBackFace),
+		DitheredLODTransitionValue(InDitheredLODTransitionValue),
 		bUseTranslucentSelfShadowing(bInUseTranslucentSelfShadowing)
 	{}
 
@@ -303,6 +306,7 @@ public:
 					Parameters.Mesh,
 					BatchElementIndex,
 					bBackFace,
+					DitheredLODTransitionValue,
 					typename TBasePassDrawingPolicy<LightMapPolicyType>::ElementDataType(LightMapElementData),
 					typename TBasePassDrawingPolicy<LightMapPolicyType>::ContextDataType()
 					);
@@ -367,6 +371,7 @@ bool FTranslucencyDrawingPolicyFactory::DrawMesh(
 	const FMeshBatch& Mesh,
 	const uint64& BatchElementMask,
 	bool bBackFace,
+	float DitheredLODTransitionValue,
 	bool bPreFog,
 	const FPrimitiveSceneProxy* PrimitiveSceneProxy,
 	FHitProxyId HitProxyId
@@ -441,6 +446,7 @@ bool FTranslucencyDrawingPolicyFactory::DrawMesh(
 			FDrawTranslucentMeshAction(
 				View,
 				bBackFace,
+				DitheredLODTransitionValue,
 				HitProxyId,
 				DrawingContext.TranslucentSelfShadow,
 				PrimitiveSceneProxy && PrimitiveSceneProxy->CastsVolumetricTranslucentShadow()
@@ -481,6 +487,7 @@ bool FTranslucencyDrawingPolicyFactory::DrawDynamicMesh(
 		Mesh,
 		Mesh.Elements.Num()==1 ? 1 : (1<<Mesh.Elements.Num())-1,	// 1 bit set for each mesh element
 		bBackFace,
+		Mesh.DitheredLODTransitionAlpha,
 		bPreFog,
 		PrimitiveSceneProxy,
 		HitProxyId);
@@ -501,6 +508,7 @@ bool FTranslucencyDrawingPolicyFactory::DrawStaticMesh(
 	FHitProxyId HitProxyId
 	)
 {
+	float DitherValue = View.GetDitheredLODTransitionValue(StaticMesh);
 	return DrawMesh(
 		RHICmdList,
 		View,
@@ -508,6 +516,7 @@ bool FTranslucencyDrawingPolicyFactory::DrawStaticMesh(
 		StaticMesh,
 		BatchElementMask,
 		false,
+		DitherValue,
 		bPreFog,
 		PrimitiveSceneProxy,
 		HitProxyId
diff --git a/Engine/Source/Runtime/Renderer/Private/TranslucentRendering.h b/Engine/Source/Runtime/Renderer/Private/TranslucentRendering.h
--- a/Engine/Source/Runtime/Renderer/Private/TranslucentRendering.h
+++ b/Engine/Source/Runtime/Renderer/Private/TranslucentRendering.h
@@ -74,6 +74,7 @@ private:
 		const FMeshBatch& Mesh,
 		const uint64& BatchElementMask,
 		bool bBackFace,
+		float DitheredLODTransitionValue,
 		bool bPreFog,
 		const FPrimitiveSceneProxy* PrimitiveSceneProxy,
 		FHitProxyId HitProxyId
diff --git a/Engine/Source/Runtime/Renderer/Private/VelocityRendering.cpp b/Engine/Source/Runtime/Renderer/Private/VelocityRendering.cpp
--- a/Engine/Source/Runtime/Renderer/Private/VelocityRendering.cpp
+++ b/Engine/Source/Runtime/Renderer/Private/VelocityRendering.cpp
@@ -30,9 +30,9 @@ public:
 		FMeshMaterialShader::SetParameters(RHICmdList, GetVertexShader(), MaterialRenderProxy, *MaterialRenderProxy->GetMaterial(View.GetFeatureLevel()), View, ESceneRenderTargetsMode::DontSet);
 	}
 
-	void SetMesh(FRHICommandList& RHICmdList, const FVertexFactory* VertexFactory, const FMeshBatch& Mesh, int32 BatchElementIndex, const FViewInfo& View, const FPrimitiveSceneProxy* Proxy, const FMatrix& InPreviousLocalToWorld)
+	void SetMesh(FRHICommandList& RHICmdList, const FVertexFactory* VertexFactory, const FMeshBatch& Mesh, int32 BatchElementIndex, float DitheredLODTransitionValue, const FViewInfo& View, const FPrimitiveSceneProxy* Proxy, const FMatrix& InPreviousLocalToWorld)
 	{
-		FMeshMaterialShader::SetMesh(RHICmdList, GetVertexShader(), VertexFactory, View, Proxy, Mesh.Elements[BatchElementIndex]);
+		FMeshMaterialShader::SetMesh(RHICmdList, GetVertexShader(), VertexFactory, View, Proxy, Mesh.Elements[BatchElementIndex],DitheredLODTransitionValue);
 
 		SetShaderValue(RHICmdList, GetVertexShader(), PreviousLocalToWorld, InPreviousLocalToWorld.ConcatTranslation(View.PrevViewMatrices.PreViewTranslation));
 	}
@@ -45,7 +45,7 @@ public:
 	static bool ShouldCache(EShaderPlatform Platform,const FMaterial* Material,const FVertexFactoryType* VertexFactoryType)
 	{
 		//Only compile the velocity shaders for the default material or if it's masked,
-		return ((Material->IsSpecialEngineMaterial() || Material->IsMasked() 
+		return ((Material->IsSpecialEngineMaterial() || !Material->WritesEveryPixel() 
 			//or if the material is opaque and two-sided,
 			|| (Material->IsTwoSided() && !IsTranslucentBlendMode(Material->GetBlendMode()))
 			// or if the material modifies meshes
@@ -133,7 +133,7 @@ public:
 	static bool ShouldCache(EShaderPlatform Platform,const FMaterial* Material,const FVertexFactoryType* VertexFactoryType)
 	{
 		//Only compile the velocity shaders for the default material or if it's masked,
-		return ((Material->IsSpecialEngineMaterial() || Material->IsMasked() 
+		return ((Material->IsSpecialEngineMaterial() || !Material->WritesEveryPixel() 
 			//or if the material is opaque and two-sided,
 			|| (Material->IsTwoSided() && !IsTranslucentBlendMode(Material->GetBlendMode()))
 			// or if the material modifies meshes
@@ -159,11 +159,11 @@ public:
 		FMeshMaterialShader::SetParameters(RHICmdList, GetPixelShader(), MaterialRenderProxy, *MaterialRenderProxy->GetMaterial(View.GetFeatureLevel()), View, ESceneRenderTargetsMode::DontSet);
 	}
 
-	void SetMesh(FRHICommandList& RHICmdList, const FVertexFactory* VertexFactory, const FMeshBatch& Mesh,int32 BatchElementIndex,const FViewInfo& View, const FPrimitiveSceneProxy* Proxy, bool bBackFace)
+	void SetMesh(FRHICommandList& RHICmdList, const FVertexFactory* VertexFactory, const FMeshBatch& Mesh,int32 BatchElementIndex, float DitheredLODTransitionValue,const FViewInfo& View, const FPrimitiveSceneProxy* Proxy, bool bBackFace)
 	{
 		const FPixelShaderRHIParamRef ShaderRHI = GetPixelShader();
 		
-		FMeshMaterialShader::SetMesh(RHICmdList, ShaderRHI, VertexFactory, View, Proxy, Mesh.Elements[BatchElementIndex]);
+		FMeshMaterialShader::SetMesh(RHICmdList, ShaderRHI, VertexFactory, View, Proxy, Mesh.Elements[BatchElementIndex],DitheredLODTransitionValue);
 	}
 
 	virtual bool Serialize(FArchive& Ar) override
@@ -249,6 +249,7 @@ void FVelocityDrawingPolicy::SetMeshRenderState(
 	const FMeshBatch& Mesh,
 	int32 BatchElementIndex,
 	bool bBackFace,
+	float DitheredLODTransitionValue,
 	const ElementDataType& ElementData,
 	const ContextDataType PolicyContext
 	) const
@@ -264,22 +265,22 @@ void FVelocityDrawingPolicy::SetMeshRenderState(
 	// previous transform can be stored in the scene for each primitive
 	if(Scene->MotionBlurInfoData.GetPrimitiveMotionBlurInfo(PrimitiveSceneProxy->GetPrimitiveSceneInfo(), PreviousLocalToWorld))
 	{
-		VertexShader->SetMesh(RHICmdList, VertexFactory, Mesh, BatchElementIndex, View, PrimitiveSceneProxy, PreviousLocalToWorld);
+		VertexShader->SetMesh(RHICmdList, VertexFactory, Mesh, BatchElementIndex, DitheredLODTransitionValue, View, PrimitiveSceneProxy, PreviousLocalToWorld);
 	}	
 	else
 	{
 		const FMatrix& LocalToWorld = PrimitiveSceneProxy->GetLocalToWorld();		// different implementation from UE4
-		VertexShader->SetMesh(RHICmdList, VertexFactory, Mesh, BatchElementIndex, View, PrimitiveSceneProxy, LocalToWorld);
+		VertexShader->SetMesh(RHICmdList, VertexFactory, Mesh, BatchElementIndex, DitheredLODTransitionValue, View, PrimitiveSceneProxy, LocalToWorld);
 	}
 
 	if (HullShader && DomainShader)
 	{
-		DomainShader->SetMesh(RHICmdList, VertexFactory, View, PrimitiveSceneProxy, Mesh.Elements[BatchElementIndex]);
-		HullShader->SetMesh(RHICmdList, VertexFactory, View, PrimitiveSceneProxy, Mesh.Elements[BatchElementIndex]);
+		DomainShader->SetMesh(RHICmdList, VertexFactory, View, PrimitiveSceneProxy, Mesh.Elements[BatchElementIndex],DitheredLODTransitionValue);
+		HullShader->SetMesh(RHICmdList, VertexFactory, View, PrimitiveSceneProxy, Mesh.Elements[BatchElementIndex],DitheredLODTransitionValue);
 	}
 
-	PixelShader->SetMesh(RHICmdList, VertexFactory, Mesh, BatchElementIndex, View, PrimitiveSceneProxy, bBackFace);
-	FMeshDrawingPolicy::SetMeshRenderState(RHICmdList, View, PrimitiveSceneProxy, Mesh, BatchElementIndex, bBackFace, ElementData, PolicyContext);
+	PixelShader->SetMesh(RHICmdList, VertexFactory, Mesh, BatchElementIndex,DitheredLODTransitionValue, View, PrimitiveSceneProxy, bBackFace);
+	FMeshDrawingPolicy::SetMeshRenderState(RHICmdList, View, PrimitiveSceneProxy, Mesh, BatchElementIndex, bBackFace, DitheredLODTransitionValue, ElementData, PolicyContext);
 }
 
 bool FVelocityDrawingPolicy::HasVelocity(const FViewInfo& View, const FPrimitiveSceneInfo* PrimitiveSceneInfo)
@@ -410,7 +411,7 @@ void FVelocityDrawingPolicyFactory::AddStaticMesh(FScene* Scene, FStaticMesh* St
 	    EBlendMode BlendMode = Material->GetBlendMode();
 	    if(BlendMode == BLEND_Opaque || BlendMode == BLEND_Masked)
 	    {
-			if (!Material->IsMasked() && !Material->IsTwoSided() && !Material->MaterialModifiesMeshPosition_RenderThread())
+			if (Material->WritesEveryPixel() && !Material->IsTwoSided() && !Material->MaterialModifiesMeshPosition_RenderThread())
 		    {
 			    // Default material doesn't handle masked or mesh-mod, and doesn't have the correct bIsTwoSided setting.
 			    MaterialRenderProxy = UMaterial::GetDefaultMaterial(MD_Surface)->GetRenderProxy(false);
@@ -448,7 +449,7 @@ bool FVelocityDrawingPolicyFactory::DrawDynamicMesh(
 		// This should be enforced at a higher level
 		//@todo - figure out why this is failing and re-enable
 		//check(FVelocityDrawingPolicy::HasVelocity(View, PrimitiveSceneInfo));
-		if (!Material->IsMasked() && !Material->IsTwoSided() && !Material->MaterialModifiesMeshPosition_RenderThread())
+		if (Material->WritesEveryPixel() && !Material->IsTwoSided() && !Material->MaterialModifiesMeshPosition_RenderThread())
 		{
 			// Default material doesn't handle masked, and doesn't have the correct bIsTwoSided setting.
 			MaterialRenderProxy = UMaterial::GetDefaultMaterial(MD_Surface)->GetRenderProxy(false);
@@ -460,7 +461,7 @@ bool FVelocityDrawingPolicyFactory::DrawDynamicMesh(
 			DrawingPolicy.SetSharedState(RHICmdList, &View, FVelocityDrawingPolicy::ContextDataType());
 			for (int32 BatchElementIndex = 0, BatchElementCount = Mesh.Elements.Num(); BatchElementIndex < BatchElementCount; ++BatchElementIndex)
 			{
-				DrawingPolicy.SetMeshRenderState(RHICmdList, View, PrimitiveSceneProxy, Mesh, BatchElementIndex, bBackFace, FMeshDrawingPolicy::ElementDataType(), FVelocityDrawingPolicy::ContextDataType());
+				DrawingPolicy.SetMeshRenderState(RHICmdList, View, PrimitiveSceneProxy, Mesh, BatchElementIndex, bBackFace, Mesh.DitheredLODTransitionAlpha, FMeshDrawingPolicy::ElementDataType(), FVelocityDrawingPolicy::ContextDataType());
 				DrawingPolicy.DrawMesh(RHICmdList, Mesh, BatchElementIndex);
 			}
 			return true;
diff --git a/Engine/Source/Runtime/Renderer/Private/VelocityRendering.h b/Engine/Source/Runtime/Renderer/Private/VelocityRendering.h
--- a/Engine/Source/Runtime/Renderer/Private/VelocityRendering.h
+++ b/Engine/Source/Runtime/Renderer/Private/VelocityRendering.h
@@ -36,6 +36,7 @@ public:
 		const FMeshBatch& Mesh,
 		int32 BatchElementIndex,
 		bool bBackFace,
+		float DitheredLODTransitionValue,
 		const ElementDataType& ElementData, 
 		const ContextDataType PolicyContext
 		) const;
diff --git a/Engine/Source/Runtime/Renderer/Public/DrawingPolicy.h b/Engine/Source/Runtime/Renderer/Public/DrawingPolicy.h
--- a/Engine/Source/Runtime/Renderer/Public/DrawingPolicy.h
+++ b/Engine/Source/Runtime/Renderer/Public/DrawingPolicy.h
@@ -38,6 +38,7 @@ public:
 		const FMaterial& InMaterialResource,
 		bool bInOverrideWithShaderComplexity = false,
 		bool bInTwoSidedOverride = false,
+		bool bInDitheredLODTransitionOverride = false,
 		bool bInWireframeOverride = false
 		);
 
@@ -47,6 +48,7 @@ public:
 		MaterialRenderProxy = Other.MaterialRenderProxy;
 		MaterialResource = Other.MaterialResource;
 		bIsTwoSidedMaterial = Other.bIsTwoSidedMaterial;
+		bIsDitheredLODTransitionMaterial = Other.bIsDitheredLODTransitionMaterial;
 		bIsWireframeMaterial = Other.bIsWireframeMaterial;
 		bNeedsBackfacePass = Other.bNeedsBackfacePass;
 		bUsePositionOnlyVS = Other.bUsePositionOnlyVS;
@@ -65,6 +67,7 @@ public:
 			VertexFactory == OtherDrawer.VertexFactory &&
 			MaterialRenderProxy == OtherDrawer.MaterialRenderProxy &&
 			bIsTwoSidedMaterial == OtherDrawer.bIsTwoSidedMaterial && 
+			bIsDitheredLODTransitionMaterial == OtherDrawer.bIsDitheredLODTransitionMaterial && 
 			bUsePositionOnlyVS == OtherDrawer.bUsePositionOnlyVS && 
 			bIsWireframeMaterial == OtherDrawer.bIsWireframeMaterial;
 	}
@@ -80,6 +83,7 @@ public:
 		const FMeshBatch& Mesh,
 		int32 BatchElementIndex,
 		bool bBackFace,
+		float DitheredLODTransitionValue,
 		const ElementDataType& ElementData,
 		const ContextDataType PolicyContext
 		) const;
@@ -106,6 +110,7 @@ public:
 		COMPAREDRAWINGPOLICYMEMBERS(VertexFactory);
 		COMPAREDRAWINGPOLICYMEMBERS(MaterialRenderProxy);
 		COMPAREDRAWINGPOLICYMEMBERS(bIsTwoSidedMaterial);
+		COMPAREDRAWINGPOLICYMEMBERS(bIsDitheredLODTransitionMaterial);
 		return 0;
 	}
 
@@ -114,6 +119,10 @@ public:
 	{
 		return bIsTwoSidedMaterial;
 	}
+	bool IsDitheredLODTransition() const
+	{
+		return bIsDitheredLODTransitionMaterial;
+	}
 	bool IsWireframe() const
 	{
 		return bIsWireframeMaterial;
@@ -131,6 +140,7 @@ protected:
 	const FMaterialRenderProxy* MaterialRenderProxy;
 	const FMaterial* MaterialResource;
 	uint32 bIsTwoSidedMaterial : 1;
+	uint32 bIsDitheredLODTransitionMaterial : 1;
 	uint32 bIsWireframeMaterial : 1;
 	uint32 bNeedsBackfacePass : 1;
 	uint32 bUsePositionOnlyVS : 1;
diff --git a/Engine/Source/Runtime/Renderer/Public/MeshMaterialShader.h b/Engine/Source/Runtime/Renderer/Public/MeshMaterialShader.h
--- a/Engine/Source/Runtime/Renderer/Public/MeshMaterialShader.h
+++ b/Engine/Source/Runtime/Renderer/Public/MeshMaterialShader.h
@@ -18,7 +18,9 @@ public:
 	FMeshMaterialShader(const FMeshMaterialShaderType::CompiledShaderInitializerType& Initializer)
 		:	FMaterialShader(Initializer)
 		,	VertexFactoryParameters(Initializer.VertexFactoryType,Initializer.ParameterMap,(EShaderFrequency)Initializer.Target.Frequency)
-	{}
+	{
+		NonInstancedDitherLODFactorParameter.Bind(Initializer.ParameterMap, TEXT("NonInstancedDitherLODFactor"));
+	}
 
 	template< typename ShaderRHIParamRef >
 	void SetParameters(
@@ -45,6 +47,7 @@ public:
 		const FSceneView& View,
 		const FPrimitiveSceneProxy* Proxy,
 		const FMeshBatchElement& BatchElement,
+		float DitheredLODTransitionValue,
 		uint32 DataFlags = 0
 	);
 
@@ -61,4 +64,5 @@ public:
 
 private:
 	FVertexFactoryParameterRef VertexFactoryParameters;
+	FShaderParameter NonInstancedDitherLODFactorParameter;
 };
