#ue4 - Change FQuat tolerance within SceneComponent and PrimitiveComponent to allow smaller rotation changes to get through.

diff --git a/Engine/Source/Runtime/Engine/Private/Components/PrimitiveComponent.cpp b/Engine/Source/Runtime/Engine/Private/Components/PrimitiveComponent.cpp
--- a/Engine/Source/Runtime/Engine/Private/Components/PrimitiveComponent.cpp
+++ b/Engine/Source/Runtime/Engine/Private/Components/PrimitiveComponent.cpp
@@ -1362,7 +1362,7 @@ bool UPrimitiveComponent::MoveComponentImpl( const FVector& Delta, const FQuat&
 	if (DeltaSizeSq <= MinMovementDistSq)
 	{
 		// Skip if no vector or rotation.
-		if (NewRotationQuat.Equals(ComponentToWorld.GetRotation()))
+		if (NewRotationQuat.Equals(ComponentToWorld.GetRotation(), 1.e-6f))
 		{
 			// copy to optional output param
 			if (OutHit)
diff --git a/Engine/Source/Runtime/Engine/Private/Components/SceneComponent.cpp b/Engine/Source/Runtime/Engine/Private/Components/SceneComponent.cpp
--- a/Engine/Source/Runtime/Engine/Private/Components/SceneComponent.cpp
+++ b/Engine/Source/Runtime/Engine/Private/Components/SceneComponent.cpp
@@ -1906,17 +1906,17 @@ bool USceneComponent::MoveComponentImpl( const FVector& Delta, const FQuat& NewR
 	if( Delta.IsZero() )
 	{
 		// Skip if no vector or rotation.
-		if (NewRotation.Equals(ComponentToWorld.GetRotation()))
+		if (NewRotation.Equals(ComponentToWorld.GetRotation(), 1.e-6f))
 		{
 			return true;
 		}
 	}
 
-	// just teleport, sweep is supported for PrimitiveComponents.  this will update child components as well.
-	InternalSetWorldLocationAndRotation(GetComponentLocation() + Delta, NewRotation);
+	// just teleport, sweep is supported for PrimitiveComponents. This will update child components as well.
+	const bool bMoved = InternalSetWorldLocationAndRotation(GetComponentLocation() + Delta, NewRotation);
 
 	// Only update overlaps if not deferring updates within a scope
-	if (!IsDeferringMovementUpdates())
+	if (bMoved && !IsDeferringMovementUpdates())
 	{
 		// need to update overlap detection in case PrimitiveComponents are attached.
 		UpdateOverlaps();
