Stop editor crashing when duplicating montages

diff --git a/Engine/Source/Runtime/Engine/Private/Animation/AnimMontage.cpp b/Engine/Source/Runtime/Engine/Private/Animation/AnimMontage.cpp
--- a/Engine/Source/Runtime/Engine/Private/Animation/AnimMontage.cpp
+++ b/Engine/Source/Runtime/Engine/Private/Animation/AnimMontage.cpp
@@ -350,8 +350,6 @@ void UAnimMontage::PostLoad()
 		}
 	}
 
-	int32 Ver = GetLinker()->UE4Ver();
-
 	for(auto& Composite : CompositeSections)
 	{
 		if(Composite.StartTime_DEPRECATED != 0.0f)
