Fix for montages with negative play rate. Move end handling code out of "has moved" conditional so we don't end up in an infinite loop

diff --git a/Engine/Source/Runtime/Engine/Private/Animation/AnimMontage.cpp b/Engine/Source/Runtime/Engine/Private/Animation/AnimMontage.cpp
--- a/Engine/Source/Runtime/Engine/Private/Animation/AnimMontage.cpp
+++ b/Engine/Source/Runtime/Engine/Private/Animation/AnimMontage.cpp
@@ -1364,7 +1364,9 @@ void FAnimMontageInstance::Advance(float DeltaTime, struct FRootMotionMovementPa
 
 					const float PositionBeforeFiringEvents = Position;
 
-					if (FMath::Abs(ActualDeltaMove) > 0.f)
+					const bool bHaveMoved = FMath::Abs(ActualDeltaMove) > 0.f;
+
+					if (bHaveMoved)
 					{
 						// Extract Root Motion for this time slice, and accumulate it.
 						if (bExtractRootMotion && AnimInstance.IsValid())
@@ -1381,25 +1383,28 @@ void FAnimMontageInstance::Advance(float DeltaTime, struct FRootMotionMovementPa
 								OutRootMotionParams->Accumulate(RootMotion);
 							}
 						}
+					}
 
-						// If current section is last one, check to trigger a blend out.
-						if( NextSectionIndex == INDEX_NONE )
-						{
-							const float DeltaPosToEnd = bPlayingForward ? (CurrentSectionLength - PosInSection) : PosInSection;
-							const float DeltaTimeToEnd = DeltaPosToEnd / CombinedPlayRate;
+					// If current section is last one, check to trigger a blend out.
+					if( NextSectionIndex == INDEX_NONE )
+					{
+						const float DeltaPosToEnd = bPlayingForward ? (CurrentSectionLength - PosInSection) : PosInSection;
+						const float DeltaTimeToEnd = DeltaPosToEnd / FMath::Abs(CombinedPlayRate);
 
-							const bool bCustomBlendOutTriggerTime = (Montage->BlendOutTriggerTime >= 0);
-							const float DefaultBlendOutTime = Montage->BlendOutTime * DefaultBlendTimeMultiplier;
-							const float BlendOutTriggerTime = bCustomBlendOutTriggerTime ? Montage->BlendOutTriggerTime : DefaultBlendOutTime;
+						const bool bCustomBlendOutTriggerTime = (Montage->BlendOutTriggerTime >= 0);
+						const float DefaultBlendOutTime = Montage->BlendOutTime * DefaultBlendTimeMultiplier;
+						const float BlendOutTriggerTime = bCustomBlendOutTriggerTime ? Montage->BlendOutTriggerTime : DefaultBlendOutTime;
 							
-							// ... trigger blend out if within blend out time window.
-							if (DeltaTimeToEnd <= FMath::Max<float>(BlendOutTriggerTime, KINDA_SMALL_NUMBER))
-							{
-								const float BlendOutTime = bCustomBlendOutTriggerTime ? DefaultBlendOutTime : DeltaTimeToEnd;
-								Stop(BlendOutTime, false);
-							}
+						// ... trigger blend out if within blend out time window.
+						if (DeltaTimeToEnd <= FMath::Max<float>(BlendOutTriggerTime, KINDA_SMALL_NUMBER))
+						{
+							const float BlendOutTime = bCustomBlendOutTriggerTime ? DefaultBlendOutTime : DeltaTimeToEnd;
+							Stop(BlendOutTime, false);
 						}
+					}
 
+					if (bHaveMoved)
+					{
 						// Delegate has to be called last in this loop
 						// so that if this changes position, the new position will be applied in the next loop
 						// first need to have event handler to handle it
@@ -1409,8 +1414,8 @@ void FAnimMontageInstance::Advance(float DeltaTime, struct FRootMotionMovementPa
 						{
 							HandleEvents(PrevPosition, Position, BranchingPointMarker);
 						}
-
 					}
+
 					// if we reached end of section, and we were not processing a branching point, and no events has messed with out current position..
 					// .. Move to next section.
 					// (this also handles looping, the same as jumping to a different section).
