UE4 - fix crash in streaming texture manager when attempting to get data for a mesh with no render data

diff --git a/Engine/Source/Runtime/Engine/Private/Components/StaticMeshComponent.cpp b/Engine/Source/Runtime/Engine/Private/Components/StaticMeshComponent.cpp
--- a/Engine/Source/Runtime/Engine/Private/Components/StaticMeshComponent.cpp
+++ b/Engine/Source/Runtime/Engine/Private/Components/StaticMeshComponent.cpp
@@ -497,13 +497,11 @@ void UStaticMeshComponent::OnUnregister()
 
 void UStaticMeshComponent::GetStreamingTextureInfo(TArray<FStreamingTexturePrimitiveInfo>& OutStreamingTextures) const
 {
-	if ( StaticMesh && !bIgnoreInstanceForTextureStreaming )
+	if ( !bIgnoreInstanceForTextureStreaming && StaticMesh && StaticMesh->RenderData && StaticMesh->RenderData->LODResources.Num() > 0)
 	{
 		const auto FeatureLevel = GetWorld() ? GetWorld()->FeatureLevel : GMaxRHIFeatureLevel;
 
 		bool bHasValidLightmapCoordinates = ((StaticMesh->LightMapCoordinateIndex >= 0)
-			&& StaticMesh->RenderData
-			&& StaticMesh->RenderData->LODResources.Num() > 0
 			&& ((uint32)StaticMesh->LightMapCoordinateIndex < StaticMesh->RenderData->LODResources[0].VertexBuffer.GetNumTexCoords()));
 
 		// We need to come up with a compensation factor for spline deformed meshes
