Fix crash when playing an anim montage in reverse with a segment that has a negative play rate

diff --git a/Engine/Source/Runtime/Engine/Private/Animation/AnimCompositeBase.cpp b/Engine/Source/Runtime/Engine/Private/Animation/AnimCompositeBase.cpp
--- a/Engine/Source/Runtime/Engine/Private/Animation/AnimCompositeBase.cpp
+++ b/Engine/Source/Runtime/Engine/Private/Animation/AnimCompositeBase.cpp
@@ -53,13 +53,14 @@ UAnimSequenceBase * FAnimSegment::GetAnimationData(float PositionInTrack, float&
  * Note: doesn't check that position is in valid range, must do that before calling this function! */
 float FAnimSegment::ConvertTrackPosToAnimPos(const float& TrackPosition) const
 {
+	const float PlayRate = GetValidPlayRate();
 	const float AnimLength = (AnimEndTime - AnimStartTime);
-	const float AnimPositionUnWrapped = (TrackPosition - StartPos) * GetValidPlayRate();
+	const float AnimPositionUnWrapped = (TrackPosition - StartPos) * PlayRate;
 
 	// Figure out how many times animation is allowed to be looped.
 	const float LoopCount = FMath::Min(FMath::FloorToInt(FMath::Abs(AnimPositionUnWrapped) / AnimLength), FMath::Max(LoopingCount-1, 0));
 	// Position within AnimSequence
-	const float AnimPoint = (AnimPositionUnWrapped >= 0.f) ? AnimStartTime : AnimEndTime;
+	const float AnimPoint = (PlayRate >= 0.f) ? AnimStartTime : AnimEndTime;
 
 	const float AnimPosition = AnimPoint + (AnimPositionUnWrapped - float(LoopCount) * AnimLength);
 	
