Safeguarding retrieving some Blueprint asset registry tags from occuring when the ParentClass is null.

diff --git a/Engine/Source/Runtime/Engine/Private/Blueprint.cpp b/Engine/Source/Runtime/Engine/Private/Blueprint.cpp
--- a/Engine/Source/Runtime/Engine/Private/Blueprint.cpp
+++ b/Engine/Source/Runtime/Engine/Private/Blueprint.cpp
@@ -712,11 +712,14 @@ void UBlueprint::GetAssetRegistryTags(TArray<FAssetRegistryTag>& OutTags) const
 	}
 	OutTags.Add( FAssetRegistryTag("ClassFlags", FString::FromInt(ClassFlagsTagged), FAssetRegistryTag::TT_Hidden) );
 
-	OutTags.Add( FAssetRegistryTag( "IsDataOnly",
-		FBlueprintEditorUtils::IsDataOnlyBlueprint(this) ? TEXT("True") : TEXT("False"),
-		FAssetRegistryTag::TT_Alphabetical ) );
+	if ( ParentClass )
+	{
+		OutTags.Add( FAssetRegistryTag( "IsDataOnly",
+			FBlueprintEditorUtils::IsDataOnlyBlueprint(this) ? TEXT("True") : TEXT("False"),
+			FAssetRegistryTag::TT_Alphabetical ) );
 
-	OutTags.Add( FAssetRegistryTag("FiB", FFindInBlueprintSearchManager::Get().QuerySingleBlueprint((UBlueprint*)this, false), FAssetRegistryTag::TT_Hidden) );
+		OutTags.Add( FAssetRegistryTag("FiB", FFindInBlueprintSearchManager::Get().QuerySingleBlueprint((UBlueprint*)this, false), FAssetRegistryTag::TT_Hidden) );
+	}
 }
 
 FString UBlueprint::GetFriendlyName() const
