#ue4 - Perf: Tweak collision post filter to avoid cache misses.

diff --git a/Engine/Source/Runtime/Engine/Private/Collision/PhysXCollision.cpp b/Engine/Source/Runtime/Engine/Private/Collision/PhysXCollision.cpp
--- a/Engine/Source/Runtime/Engine/Private/Collision/PhysXCollision.cpp
+++ b/Engine/Source/Runtime/Engine/Private/Collision/PhysXCollision.cpp
@@ -686,13 +686,13 @@ PxSceneQueryHitType::Enum FPxQueryFilterCallbackSweep::postFilter(const PxFilter
 	PxSweepHit& SweepHit = (PxSweepHit&)hit;
 	const bool bIsOverlap = SweepHit.hadInitialOverlap();
 
-	if (DiscardInitialOverlaps && bIsOverlap)
+	if (bIsOverlap && DiscardInitialOverlaps)
 	{
 		return PxSceneQueryHitType::eNONE;
 	}
 	else
 	{
-		if (PrefilterReturnValue == PxSceneQueryHitType::eBLOCK && bIsOverlap)
+		if (bIsOverlap && PrefilterReturnValue == PxSceneQueryHitType::eBLOCK)
 		{
 			// We want to keep initial blocking overlaps and continue the sweep until a non-overlapping blocking hit.
 			// We will later report this hit as a blocking hit when we compute the hit type (using FPxQueryFilterCallback::CalcQueryHitType).
