#ue4 - Smooth rotation for network corrections of simulated proxies in character movement.

diff --git a/Engine/Source/Runtime/Engine/Classes/GameFramework/Character.h b/Engine/Source/Runtime/Engine/Classes/GameFramework/Character.h
--- a/Engine/Source/Runtime/Engine/Classes/GameFramework/Character.h
+++ b/Engine/Source/Runtime/Engine/Classes/GameFramework/Character.h
@@ -260,10 +260,14 @@ public:
 	virtual void OnRep_ReplicatedBasedMovement();
 
 protected:
-	/** Desired translation offset of mesh. */
+	/** Saved translation offset of mesh. */
 	UPROPERTY()
 	FVector BaseTranslationOffset;
 
+	/** Saved rotation offset of mesh. */
+	UPROPERTY()
+	FQuat BaseRotationOffset;
+
 	/** Event called after actor's base changes (if SetBase was requested to notify us with bNotifyPawn). */
 	virtual void BaseChange();
 
@@ -289,9 +293,12 @@ public:
 	/** Returns ReplicatedMovementMode */
 	uint8 GetReplicatedMovementMode() const { return ReplicatedMovementMode; }
 
-	/** @return Desired translation offset of mesh. */
+	/** @return Saved translation offset of mesh. */
 	const FVector& GetBaseTranslationOffset() const { return BaseTranslationOffset; }
 
+	/** @return Saved rotation offset of mesh. */
+	const FQuat& GetBaseRotationOffset() const { return BaseRotationOffset; }
+
 	// Begin INavAgentInterface Interface
 	virtual FVector GetNavAgentLocation() const override;
 	// End INavAgentInterface Interface
diff --git a/Engine/Source/Runtime/Engine/Classes/GameFramework/CharacterMovementComponent.h b/Engine/Source/Runtime/Engine/Classes/GameFramework/CharacterMovementComponent.h
--- a/Engine/Source/Runtime/Engine/Classes/GameFramework/CharacterMovementComponent.h
+++ b/Engine/Source/Runtime/Engine/Classes/GameFramework/CharacterMovementComponent.h
@@ -526,6 +526,18 @@ public:
 	UPROPERTY(Category="Character Movement (General Settings)", EditAnywhere, BlueprintReadWrite, AdvancedDisplay, meta=(ClampMin="1", ClampMax="25", UIMin="1", UIMax="25"))
 	int32 MaxSimulationIterations;
 
+	/**
+	 * How long to take to smoothly interpolate from the old pawn position on the client to the corrected one sent by the server.
+	 */
+	UPROPERTY(Category="Character Movement (General Settings)", EditDefaultsOnly, AdvancedDisplay, meta=(ClampMin="0.0", ClampMax="1.0", UIMin="0.0", UIMax="1.0"))
+	float NetworkSimulatedSmoothLocationTime;
+
+	/**
+	 * How long to take to smoothly interpolate from the old pawn rotation on the client to the corrected one sent by the server.
+	 */
+	UPROPERTY(Category="Character Movement (General Settings)", EditDefaultsOnly, AdvancedDisplay, meta=(ClampMin="0.0", ClampMax="1.0", UIMin="0.0", UIMax="1.0"))
+	float NetworkSimulatedSmoothRotationTime;
+
 	/** Used in determining if pawn is going off ledge.  If the ledge is "shorter" than this value then the pawn will be able to walk off it. **/
 	UPROPERTY(Category="Character Movement: Walking", EditAnywhere, BlueprintReadWrite, AdvancedDisplay)
 	float LedgeCheckThreshold;
@@ -1607,7 +1619,7 @@ public:
 	//--------------------------------
 	// Client hook
 	//--------------------------------
-	virtual void SmoothCorrection(const FVector& OldLocation) override;
+	virtual void SmoothCorrection(const FVector& OldLocation, const FQuat& OldRotation) override;
 
 	virtual class FNetworkPredictionData_Client* GetPredictionData_Client() const override;
 	virtual class FNetworkPredictionData_Server* GetPredictionData_Server() const override;
@@ -1961,7 +1973,7 @@ class ENGINE_API FNetworkPredictionData_Client_Character : public FNetworkPredic
 {
 public:
 
-	FNetworkPredictionData_Client_Character();
+	FNetworkPredictionData_Client_Character(const UCharacterMovementComponent& ClientMovement);
 	virtual ~FNetworkPredictionData_Client_Character();
 
 	/** Client timestamp of last time it sent a servermove() to the server.  Used for holding off on sending movement updates to save bandwidth. */
@@ -1991,6 +2003,9 @@ public:
 	/** Used for position smoothing in net games */
 	FVector MeshTranslationOffset;
 
+	/** Used for rotation smoothing in net games */
+	FQuat MeshRotationOffset;
+
 	/** Maximum location correction distance for which other pawn positions on a client will be smoothly updated */
 	float MaxSmoothNetUpdateDist;
 
@@ -1998,8 +2013,11 @@ public:
 	If it is between MaxSmoothNetUpdateDist and NoSmoothNetUpdateDist, pop to MaxSmoothNetUpdateDist away from the updated location */
 	float NoSmoothNetUpdateDist;
 
-	/** How long to take to smoothly interpolate from the old pawn position on the client to the corrected one sent by the server.  Must be > 0.0 */
+	/** How long to take to smoothly interpolate from the old pawn position on the client to the corrected one sent by the server.  Must be >= 0.0 */
 	float SmoothNetUpdateTime;
+
+	/** How long to take to smoothly interpolate from the old pawn rotation on the client to the corrected one sent by the server.  Must be >= 0.0 */
+	float SmoothNetUpdateRotationTime;
 	
 	// how long server will wait for client move update before setting position
 	// @TODO: don't duplicate between server and client data (though it's used by both)
diff --git a/Engine/Source/Runtime/Engine/Classes/Interfaces/NetworkPredictionInterface.h b/Engine/Source/Runtime/Engine/Classes/Interfaces/NetworkPredictionInterface.h
--- a/Engine/Source/Runtime/Engine/Classes/Interfaces/NetworkPredictionInterface.h
+++ b/Engine/Source/Runtime/Engine/Classes/Interfaces/NetworkPredictionInterface.h
@@ -34,7 +34,7 @@ class ENGINE_API INetworkPredictionInterface
 	//--------------------------------
 
 	/** (Client) After receiving a network update of position, allow some custom smoothing, given the old location before the correction. */
-	virtual void SmoothCorrection(const FVector& OldLocation) PURE_VIRTUAL(INetworkPredictionInterface::SmoothCorrection,);
+	virtual void SmoothCorrection(const FVector& OldLocation, const FQuat& OldRotation) PURE_VIRTUAL(INetworkPredictionInterface::SmoothCorrection,);
 
 	//--------------------------------
 	// Other
diff --git a/Engine/Source/Runtime/Engine/Private/Character.cpp b/Engine/Source/Runtime/Engine/Private/Character.cpp
--- a/Engine/Source/Runtime/Engine/Private/Character.cpp
+++ b/Engine/Source/Runtime/Engine/Private/Character.cpp
@@ -98,6 +98,8 @@ ACharacter::ACharacter(const FObjectInitializer& ObjectInitializer)
 		Mesh->bGenerateOverlapEvents = false;
 		Mesh->bCanEverAffectNavigation = false;
 	}
+
+	BaseRotationOffset = FQuat::Identity;
 }
 
 void ACharacter::PostInitializeComponents()
@@ -109,6 +111,7 @@ void ACharacter::PostInitializeComponents()
 		if (Mesh)
 		{
 			BaseTranslationOffset = Mesh->RelativeLocation;
+			BaseRotationOffset = Mesh->RelativeRotation.Quaternion();
 
 			// force animation tick after movement component updates
 			if (Mesh->PrimaryComponentTick.bCanEverTick && CharacterMovement)
@@ -934,6 +937,7 @@ void ACharacter::OnRep_ReplicatedBasedMovement()
 	{
 		// Update transform relative to movement base
 		const FVector OldLocation = GetActorLocation();
+		const FQuat OldRotation = GetActorQuat();
 		MovementBaseUtility::GetMovementBaseTransform(ReplicatedBasedMovement.MovementBase, ReplicatedBasedMovement.BoneName, CharacterMovement->OldBaseLocation, CharacterMovement->OldBaseQuat);
 		const FVector NewLocation = CharacterMovement->OldBaseLocation + ReplicatedBasedMovement.Location;
 
@@ -960,7 +964,7 @@ void ACharacter::OnRep_ReplicatedBasedMovement()
 		INetworkPredictionInterface* PredictionInterface = Cast<INetworkPredictionInterface>(GetMovementComponent());
 		if (PredictionInterface)
 		{
-			PredictionInterface->SmoothCorrection(OldLocation);
+			PredictionInterface->SmoothCorrection(OldLocation, OldRotation);
 		}
 	}
 }
@@ -1013,6 +1017,7 @@ void ACharacter::SimulatedRootMotionPositionFixup(float DeltaSeconds)
 		if( MoveIndex != INDEX_NONE )
 		{
 			const FVector OldLocation = GetActorLocation();
+			const FQuat OldRotation = GetActorQuat();
 			// Move Actor back to position of that buffered move. (server replicated position).
 			const FSimulatedRootMotionReplicatedMove& RootMotionRepMove = RootMotionRepMoves[MoveIndex];
 			if( RestoreReplicatedMove(RootMotionRepMove) )
@@ -1042,7 +1047,7 @@ void ACharacter::SimulatedRootMotionPositionFixup(float DeltaSeconds)
 							INetworkPredictionInterface* PredictionInterface = Cast<INetworkPredictionInterface>(GetMovementComponent());
 							if (PredictionInterface)
 							{
-								PredictionInterface->SmoothCorrection(OldLocation);
+								PredictionInterface->SmoothCorrection(OldLocation, OldRotation);
 							}
 						}
 					}
@@ -1190,12 +1195,13 @@ void ACharacter::PostNetReceiveLocationAndRotation()
 		if (!ReplicatedBasedMovement.HasRelativeLocation())
 		{
 			const FVector OldLocation = GetActorLocation();
+			const FQuat OldRotation = GetActorQuat();
 			UpdateSimulatedPosition(ReplicatedMovement.Location, ReplicatedMovement.Rotation);
 
 			INetworkPredictionInterface* PredictionInterface = Cast<INetworkPredictionInterface>(GetMovementComponent());
 			if (PredictionInterface)
 			{
-				PredictionInterface->SmoothCorrection(OldLocation);
+				PredictionInterface->SmoothCorrection(OldLocation, OldRotation);
 			}
 		}
 	}
diff --git a/Engine/Source/Runtime/Engine/Private/Components/CharacterMovementComponent.cpp b/Engine/Source/Runtime/Engine/Private/Components/CharacterMovementComponent.cpp
--- a/Engine/Source/Runtime/Engine/Private/Components/CharacterMovementComponent.cpp
+++ b/Engine/Source/Runtime/Engine/Private/Components/CharacterMovementComponent.cpp
@@ -202,6 +202,9 @@ UCharacterMovementComponent::UCharacterMovementComponent(const FObjectInitialize
 	MaxSimulationTimeStep = 0.05f;
 	MaxSimulationIterations = 8;
 
+	NetworkSimulatedSmoothLocationTime = 0.100f;
+	NetworkSimulatedSmoothRotationTime = 0.033f;
+
 	CrouchedSpeedMultiplier_DEPRECATED = 0.5f;
 	MaxWalkSpeedCrouched = MaxWalkSpeed * CrouchedSpeedMultiplier_DEPRECATED;
 	MaxOutOfWaterStepHeight = 40.0f;
@@ -5565,7 +5568,7 @@ float UCharacterMovementComponent::GetSimulationTimeStep(float RemainingTime, in
 }
 
 
-void UCharacterMovementComponent::SmoothCorrection(const FVector& OldLocation)
+void UCharacterMovementComponent::SmoothCorrection(const FVector& OldLocation, const FQuat& OldRotation)
 {
 	if (!HasValidData() || GetNetMode() != NM_Client)
 	{
@@ -5586,6 +5589,9 @@ void UCharacterMovementComponent::SmoothCorrection(const FVector& OldLocation)
 		{
 			ClientData->MeshTranslationOffset = ClientData->MeshTranslationOffset + OldLocation - UpdatedComponent->GetComponentLocation();	
 		}
+
+		// Take difference between where we were rotated before, and where we're going
+		ClientData->MeshRotationOffset = (UpdatedComponent->GetComponentQuat().Inverse() * OldRotation) * ClientData->MeshRotationOffset;
 	}
 }
 
@@ -5603,6 +5609,7 @@ void UCharacterMovementComponent::SmoothClientPosition(float DeltaSeconds)
 		// smooth interpolation of mesh translation to avoid popping of other client pawns unless under a low tick rate
 		if (DeltaSeconds < ClientData->SmoothNetUpdateTime)
 		{
+			// Slowly decay translation offset
 			ClientData->MeshTranslationOffset = (ClientData->MeshTranslationOffset * (1.f - DeltaSeconds / ClientData->SmoothNetUpdateTime));
 		}
 		else
@@ -5610,6 +5617,17 @@ void UCharacterMovementComponent::SmoothClientPosition(float DeltaSeconds)
 			ClientData->MeshTranslationOffset = FVector::ZeroVector;
 		}
 
+		// Smooth rotation
+		if (DeltaSeconds < ClientData->SmoothNetUpdateRotationTime)
+		{
+			// Slowly decay rotation offset
+			ClientData->MeshRotationOffset = FQuat::Slerp( ClientData->MeshRotationOffset, FQuat::Identity, DeltaSeconds / ClientData->SmoothNetUpdateRotationTime );
+		}
+		else
+		{
+			ClientData->MeshRotationOffset = FQuat::Identity;
+		}
+
 		if (IsMovingOnGround())
 		{
 			// don't smooth Z position if walking on ground
@@ -5617,7 +5635,8 @@ void UCharacterMovementComponent::SmoothClientPosition(float DeltaSeconds)
 		}
 
 		const FVector NewRelTranslation = CharacterOwner->ActorToWorld().InverseTransformVectorNoScale(ClientData->MeshTranslationOffset + CharacterOwner->GetBaseTranslationOffset());
-		CharacterOwner->GetMesh()->SetRelativeLocation(NewRelTranslation);
+		const FQuat NewRelRotation = ClientData->MeshRotationOffset * CharacterOwner->GetBaseRotationOffset();
+		CharacterOwner->GetMesh()->SetRelativeLocationAndRotation(NewRelTranslation, NewRelRotation);
 	}
 }
 
@@ -5717,7 +5736,7 @@ FNetworkPredictionData_Client* UCharacterMovementComponent::GetPredictionData_Cl
 	if (!ClientPredictionData)
 	{
 		UCharacterMovementComponent* MutableThis = const_cast<UCharacterMovementComponent*>(this);
-		MutableThis->ClientPredictionData = new FNetworkPredictionData_Client_Character();
+		MutableThis->ClientPredictionData = new FNetworkPredictionData_Client_Character(*this);
 	}
 
 	return ClientPredictionData;
@@ -6948,7 +6967,7 @@ void UCharacterMovementComponent::TickCharacterPose(float DeltaTime)
 	}
 }
 
-FNetworkPredictionData_Client_Character::FNetworkPredictionData_Client_Character()
+FNetworkPredictionData_Client_Character::FNetworkPredictionData_Client_Character(const UCharacterMovementComponent& ClientMovement)
 	: ClientUpdateTime(0.f)
 	, CurrentTimeStamp(0.f)
 	, PendingMove(NULL)
@@ -6958,15 +6977,18 @@ FNetworkPredictionData_Client_Character::FNetworkPredictionData_Client_Character
 	, bUpdatePosition(false)
 	, bSmoothNetUpdates(false)
 	, MeshTranslationOffset(ForceInitToZero)
+	, MeshRotationOffset(FQuat::Identity)
 	, MaxSmoothNetUpdateDist(0.f)
 	, NoSmoothNetUpdateDist(0.f)
 	, SmoothNetUpdateTime(0.f)
+	, SmoothNetUpdateRotationTime(0.f)
 	, MaxResponseTime(0.f)
 {
 	bSmoothNetUpdates = true;
 	MaxSmoothNetUpdateDist = 84.0;
 	NoSmoothNetUpdateDist = 128.0;
-	SmoothNetUpdateTime = 0.100;
+	SmoothNetUpdateTime = ClientMovement.NetworkSimulatedSmoothLocationTime;
+	SmoothNetUpdateRotationTime = ClientMovement.NetworkSimulatedSmoothRotationTime;	
 
 	MaxResponseTime = 0.125f;
 }
diff --git a/Engine/Source/Runtime/Engine/Private/Pawn.cpp b/Engine/Source/Runtime/Engine/Private/Pawn.cpp
--- a/Engine/Source/Runtime/Engine/Private/Pawn.cpp
+++ b/Engine/Source/Runtime/Engine/Private/Pawn.cpp
@@ -999,12 +999,13 @@ void APawn::PostNetReceiveLocationAndRotation()
 		ReplicatedMovement.Location.Z += 0.01f;
 
 		const FVector OldLocation = GetActorLocation();
+		const FQuat OldRotation = GetActorQuat();
 		SetActorLocationAndRotation(ReplicatedMovement.Location, ReplicatedMovement.Rotation, /*bSweep=*/ false);
 
 		INetworkPredictionInterface* PredictionInterface = Cast<INetworkPredictionInterface>(GetMovementComponent());
 		if (PredictionInterface)
 		{
-			PredictionInterface->SmoothCorrection(OldLocation);
+			PredictionInterface->SmoothCorrection(OldLocation, OldRotation);
 		}
 	}
 }
