#ue4 - Simplify FindGeomOpposingNormal() return value (and remove out param).

diff --git a/Engine/Source/Runtime/Engine/Private/Collision/CollisionConversions.cpp b/Engine/Source/Runtime/Engine/Private/Collision/CollisionConversions.cpp
--- a/Engine/Source/Runtime/Engine/Private/Collision/CollisionConversions.cpp
+++ b/Engine/Source/Runtime/Engine/Private/Collision/CollisionConversions.cpp
@@ -95,19 +95,19 @@ static PxVec3 TransformNormalToShapeSpace(const PxMeshScale& meshScale, const Px
 	}
 }
 
-static bool FindSimpleOpposingNormal(const PxLocationHit& PHit, const FVector& TraceDirectionDenorm, FVector& OutNormal)
+static FVector FindSimpleOpposingNormal(const PxLocationHit& PHit, const FVector& TraceDirectionDenorm, const FVector InNormal)
 {
-	// We don't compute anything special, so let calling code figure out the correct fallback.
-	return false;
+	// We don't compute anything special
+	return InNormal;
 }
 
-static bool FindBoxOpposingNormal(const PxLocationHit& PHit, const FVector& TraceDirectionDenorm, FVector& OutNormal)
+static FVector FindBoxOpposingNormal(const PxLocationHit& PHit, const FVector& TraceDirectionDenorm, const FVector InNormal)
 {
 	// We require normal info for our algorithm.
 	const bool bNormalData = (PHit.flags & PxHitFlag::eNORMAL);
 	if (!bNormalData)
 	{
-		return false;
+		return InNormal;
 	}
 
 	PxBoxGeometry PxBoxGeom;
@@ -151,15 +151,14 @@ static bool FindBoxOpposingNormal(const PxLocationHit& PHit, const FVector& Trac
 
 	// Fill in result
 	const PxVec3 WorldNormal = LocalToWorld.rotate(BestLocalNormal);
-	OutNormal = P2UVector(WorldNormal);
-	return true;
+	return P2UVector(WorldNormal);
 }
 
-static bool FindHeightFieldOpposingNormal(const PxLocationHit& PHit, const FVector& TraceDirectionDenorm, FVector& OutNormal)
+static FVector FindHeightFieldOpposingNormal(const PxLocationHit& PHit, const FVector& TraceDirectionDenorm, const FVector InNormal)
 {
 	if (IsInvalidFaceIndex(PHit.faceIndex))
 	{
-		return false;
+		return InNormal;
 	}
 
 	PxHeightFieldGeometry PHeightFieldGeom;
@@ -175,19 +174,17 @@ static bool FindHeightFieldOpposingNormal(const PxLocationHit& PHit, const FVect
 
 		PxVec3 TriNormal;
 		Tri.normal(TriNormal);
-		OutNormal = P2UVector(TriNormal);
-
-		return true;
+		return P2UVector(TriNormal);
 	}
 
-	return false;
+	return InNormal;
 }
 
-static bool FindConvexMeshOpposingNormal(const PxLocationHit& PHit, const FVector& TraceDirectionDenorm, FVector& OutNormal)
+static FVector FindConvexMeshOpposingNormal(const PxLocationHit& PHit, const FVector& TraceDirectionDenorm, const FVector InNormal)
 {
 	if (IsInvalidFaceIndex(PHit.faceIndex))
 	{
-		return false;
+		return InNormal;
 	}
 
 	PxConvexMeshGeometry PConvexMeshGeom;
@@ -210,7 +207,7 @@ static bool FindConvexMeshOpposingNormal(const PxLocationHit& PHit, const FVecto
 			// Convert to world space
 			const PxTransform PShapeWorldPose = PxShapeExt::getGlobalPose(*PHit.shape, *PHit.actor);
 			const PxVec3 PWorldPolyNormal = PShapeWorldPose.rotate(PLocalPolyNormal);
-			OutNormal = P2UVector(PWorldPolyNormal);
+			const FVector OutNormal = P2UVector(PWorldPolyNormal);
 
 #if !(UE_BUILD_SHIPPING || UE_BUILD_TEST)
 			if (!OutNormal.IsNormalized())
@@ -219,18 +216,18 @@ static bool FindConvexMeshOpposingNormal(const PxLocationHit& PHit, const FVecto
 				UE_LOG(LogPhysics, Warning, TEXT("WorldTransform \n: %s"), *P2UTransform(PShapeWorldPose).ToString());
 			}
 #endif
-			return true;
+			return OutNormal;
 		}
 	}
 
-	return false;
+	return InNormal;
 }
 
-static bool FindTriMeshOpposingNormal(const PxLocationHit& PHit, const FVector& TraceDirectionDenorm, FVector& OutNormal)
+static FVector FindTriMeshOpposingNormal(const PxLocationHit& PHit, const FVector& TraceDirectionDenorm, const FVector InNormal)
 {
 	if (IsInvalidFaceIndex(PHit.faceIndex))
 	{
-		return false;
+		return InNormal;
 	}
 
 	PxTriangleMeshGeometry PTriMeshGeom;
@@ -275,7 +272,7 @@ static bool FindTriMeshOpposingNormal(const PxLocationHit& PHit, const FVector&
 		// Convert to world space
 		const PxTransform PShapeWorldPose = PxShapeExt::getGlobalPose(*PHit.shape, *PHit.actor);
 		const PxVec3 PWorldTriNormal = PShapeWorldPose.rotate(PLocalTriNormal);
-		OutNormal = P2UVector(PWorldTriNormal);
+		FVector OutNormal = P2UVector(PWorldTriNormal);
 
 		if (PTriMeshGeom.meshFlags & PxMeshGeometryFlag::eDOUBLE_SIDED)
 		{
@@ -291,20 +288,20 @@ static bool FindTriMeshOpposingNormal(const PxLocationHit& PHit, const FVector&
 			UE_LOG(LogPhysics, Warning, TEXT("WorldTransform \n: %s"), *P2UTransform(PShapeWorldPose).ToString());
 		}
 #endif
-		return true;
+		return OutNormal;
 	}
 
-	return false;
+	return InNormal;
 }
 
 /**
- * Util to find the normal of the face that we hit.
+ * Util to find the normal of the face that we hit. Will use faceIndex from the hit if possible.
  * @param PHit - incoming hit from PhysX
  * @param TraceDirectionDenorm - direction of sweep test (not normalized)
- * @param OutNormal - normal we may recompute based on the faceIndex of the hit
- * @return true if we compute a new normal for the geometry.
+ * @param InNormal - default value in case no new normal is computed.
+ * @return New normal we compute for geometry.
  */
-static bool FindGeomOpposingNormal(PxGeometryType::Enum QueryGeomType, const PxLocationHit& PHit, const FVector& TraceDirectionDenorm, FVector& OutNormal)
+static FVector FindGeomOpposingNormal(PxGeometryType::Enum QueryGeomType, const PxLocationHit& PHit, const FVector& TraceDirectionDenorm, const FVector InNormal)
 {
 	// TODO: can we support other shapes here as well?
 	if (QueryGeomType == PxGeometryType::eCAPSULE || QueryGeomType == PxGeometryType::eSPHERE)
@@ -313,16 +310,16 @@ static bool FindGeomOpposingNormal(PxGeometryType::Enum QueryGeomType, const PxL
 		switch (GeomType)
 		{
 		case PxGeometryType::eSPHERE:
-		case PxGeometryType::eCAPSULE:		return FindSimpleOpposingNormal(PHit, TraceDirectionDenorm, OutNormal);
-		case PxGeometryType::eBOX:			return FindBoxOpposingNormal(PHit, TraceDirectionDenorm, OutNormal);
-		case PxGeometryType::eCONVEXMESH:	return FindConvexMeshOpposingNormal(PHit, TraceDirectionDenorm, OutNormal);
-		case PxGeometryType::eHEIGHTFIELD:	return FindHeightFieldOpposingNormal(PHit, TraceDirectionDenorm, OutNormal);
-		case PxGeometryType::eTRIANGLEMESH:	return FindTriMeshOpposingNormal(PHit, TraceDirectionDenorm, OutNormal);
+		case PxGeometryType::eCAPSULE:		return FindSimpleOpposingNormal(PHit, TraceDirectionDenorm, InNormal);
+		case PxGeometryType::eBOX:			return FindBoxOpposingNormal(PHit, TraceDirectionDenorm, InNormal);
+		case PxGeometryType::eCONVEXMESH:	return FindConvexMeshOpposingNormal(PHit, TraceDirectionDenorm, InNormal);
+		case PxGeometryType::eHEIGHTFIELD:	return FindHeightFieldOpposingNormal(PHit, TraceDirectionDenorm, InNormal);
+		case PxGeometryType::eTRIANGLEMESH:	return FindTriMeshOpposingNormal(PHit, TraceDirectionDenorm, InNormal);
 		default: check(false);	//unsupported geom type
 		}
 	}
 
-	return false;
+	return InNormal;
 }
 
 
@@ -464,7 +461,7 @@ void ConvertQueryImpactHit(const UWorld* World, const PxLocationHit& PHit, FHitR
 	}
 
 	const PxGeometryType::Enum SweptGeometryType = Geom ? Geom->getType() : PxGeometryType::eINVALID;
-	FindGeomOpposingNormal(SweptGeometryType, PHit, TraceStartToEnd, OutResult.ImpactNormal);
+	OutResult.ImpactNormal = FindGeomOpposingNormal(SweptGeometryType, PHit, TraceStartToEnd, Normal);
 	
 	if( PHit.shape->getGeometryType() == PxGeometryType::eHEIGHTFIELD)
 	{
