Fix crash when switching between fullscreen and windowed mode.

diff --git a/Engine/Source/Runtime/Engine/Private/Slate/SceneViewport.cpp b/Engine/Source/Runtime/Engine/Private/Slate/SceneViewport.cpp
--- a/Engine/Source/Runtime/Engine/Private/Slate/SceneViewport.cpp
+++ b/Engine/Source/Runtime/Engine/Private/Slate/SceneViewport.cpp
@@ -1421,10 +1421,13 @@ void FSceneViewport::InitDynamicRHI()
 	}
 	else
 	{
-		check(BufferedSlateHandles.Num() == 0);
-		BufferedSlateHandles.Add(nullptr);	
-		BufferedRenderTargetsRHI.Add(nullptr);
-		BufferedShaderResourceTexturesRHI.Add(nullptr);
+		check(BufferedSlateHandles.Num() == BufferedRenderTargetsRHI.Num() && BufferedSlateHandles.Num() == BufferedShaderResourceTexturesRHI.Num());
+		if (BufferedSlateHandles.Num() == 0)
+		{
+			BufferedSlateHandles.Add(nullptr);
+			BufferedRenderTargetsRHI.Add(nullptr);
+			BufferedShaderResourceTexturesRHI.Add(nullptr);
+		}		
 		NumBufferedFrames = 1;
 
 		RenderTargetTextureRHI = nullptr;		
