fixed skeletal mesh motion blur with multithreaded rendering (r.RHICmdBypass=0), fixed potential issues
fixed UE-2465 ShooterGame: Bots shimmer when game is paused

diff --git a/Engine/Source/Runtime/Engine/Private/GPUSkinVertexFactory.cpp b/Engine/Source/Runtime/Engine/Private/GPUSkinVertexFactory.cpp
--- a/Engine/Source/Runtime/Engine/Private/GPUSkinVertexFactory.cpp
+++ b/Engine/Source/Runtime/Engine/Private/GPUSkinVertexFactory.cpp
@@ -422,14 +422,14 @@ public:
 
 			bool bLocalPerBoneMotionBlur = false;
 			
-			if (FeatureLevel >= ERHIFeatureLevel::SM4 && GPrevPerBoneMotionBlur.IsAppendStarted())
+			if (FeatureLevel >= ERHIFeatureLevel::SM4 && PreviousBoneMatrices.IsBound())
 			{
 				const bool bWorldIsPaused = View.Family->bWorldIsPaused;
 
 				// we are in the velocity rendering pass
 
 				// 0xffffffff or valid index
-				uint32 OldBoneDataIndex = ShaderData.GetOldBoneData(FrameNumber);
+				uint32 OldBoneDataIndex = ShaderData.GetOldBoneData(bWorldIsPaused ? (FrameNumber - 1) : FrameNumber);
 
 #if !(UE_BUILD_SHIPPING || UE_BUILD_TEST)
 				{
@@ -477,19 +477,29 @@ public:
 				// if we haven't copied the data yet we skip the update (e.g. split screen)
 				if(ShaderData.IsOldBoneDataUpdateNeeded(FrameNumber))
 				{
-#if !(UE_BUILD_SHIPPING || UE_BUILD_TEST)
-					if (CVarVelocityTest.GetValueOnRenderThread())
+					if (bWorldIsPaused)
 					{
-						FBoneSkinning Tab[10];
-						GPrevPerBoneMotionBlur.AppendData(Tab, (FMath::Rand() % 55) + 1);
+						GPUVertexFactory->MaintainBoneDataStartIndex();
 					}
+					else
+					{
+#if !(UE_BUILD_SHIPPING || UE_BUILD_TEST)
+						if (CVarVelocityTest.GetValueOnRenderThread())
+						{
+							FBoneSkinning Tab[10];
+
+							FMemory::Memset(Tab, 0, sizeof(FBoneSkinning) * 10);
+							// We add some black gap between the elements, to ensure we pick the right data for this frame
+							// wwe create different gap size on alternating frames.
+							GPrevPerBoneMotionBlur.AppendData(Tab, (FrameNumber % 2) + 4);
+						}
 #endif // if !(UE_BUILD_SHIPPING || UE_BUILD_TEST)
 
-					// copy the bone data and tell the instance where it can pick it up next frame
-					// append data to a buffer we bind next frame to read old matrix data for motion blur
-					
-					uint32 OldBoneDataStartIndex = GPrevPerBoneMotionBlur.AppendData(ShaderData.BoneMatrices.GetData(), ShaderData.BoneMatrices.Num());
-					GPUVertexFactory->SetOldBoneDataStartIndex(FrameNumber, OldBoneDataStartIndex);
+						// copy the bone data and tell the instance where it can pick it up next frame
+						// append data to a buffer we bind next frame to read old matrix data for motion blur
+						uint32 OldBoneDataStartIndex = GPrevPerBoneMotionBlur.AppendData(ShaderData.BoneMatrices.GetData(), ShaderData.BoneMatrices.Num());
+						GPUVertexFactory->SetOldBoneDataStartIndex(FrameNumber, OldBoneDataStartIndex);
+					}
 				}
 			}
 			SetShaderValue(RHICmdList, Shader->GetVertexShader(), PerBoneMotionBlur, bLocalPerBoneMotionBlur);
diff --git a/Engine/Source/Runtime/Engine/Private/SkeletalRenderGPUSkin.cpp b/Engine/Source/Runtime/Engine/Private/SkeletalRenderGPUSkin.cpp
--- a/Engine/Source/Runtime/Engine/Private/SkeletalRenderGPUSkin.cpp
+++ b/Engine/Source/Runtime/Engine/Private/SkeletalRenderGPUSkin.cpp
@@ -1071,7 +1071,7 @@ FPreviousPerBoneMotionBlur
 -----------------------------------------------------------------------------*/
 
 FPreviousPerBoneMotionBlur::FPreviousPerBoneMotionBlur()
-	:BufferIndex(0), LockedData(0), LockedTexelCount(0), bWarningBufferSizeExceeded(false)
+	: LockedData(0), LockedTexelCount(0), bWarningBufferSizeExceeded(false)
 {
 }
 
@@ -1103,15 +1103,6 @@ void FPreviousPerBoneMotionBlur::RestoreForPausedMotionBlur()
 	{
 		UE_LOG(LogEngine, Log, TEXT("r.MotionBlurDebug: RestoreForPausedMotionBlur"));
 	}
-
-	if(BufferIndex == 0)
-	{
-		BufferIndex = PER_BONE_BUFFER_COUNT - 1;
-	}
-	else
-	{
-		--BufferIndex;
-	}
 }
 
 uint32 FPreviousPerBoneMotionBlur::GetSizeX() const
@@ -1136,22 +1127,27 @@ void FPreviousPerBoneMotionBlur::StartAppend(bool bWorldIsPaused)
 	check(!LockedData);
 	check(IsInRenderingThread());
 
-	if (CVarMotionBlurDebug.GetValueOnRenderThread())
-	{
-		UE_LOG(LogEngine, Log, TEXT("r.MotionBlurDebug: BufferIndex=%d/%d Read=%d Write=%d IsLocked=%d"),
-			BufferIndex, PER_BONE_BUFFER_COUNT, GetReadBufferIndex(), GetWriteBufferIndex(), IsAppendStarted() ? 1 : 0);
-	}
+	if(!bWorldIsPaused)
+	{		
+		InitIfNeeded();
 
-	InitIfNeeded();
+		AdvanceBufferLocation();
 
-	FBoneDataVertexBuffer& WriteTexture = PerChunkBoneMatricesTexture[GetWriteBufferIndex()];
+		FBoneDataVertexBuffer& WriteTexture = PerChunkBoneMatricesTexture[GetWriteBufferIndex()];
 
-	if(WriteTexture.IsValid())
+		if(WriteTexture.IsValid())
+		{
+			LockedData = WriteTexture.LockData();
+			check(LockedTexelPosition.GetValue() == 0); //otherwise it wasn't unlocked or it was unlocked before it was done being filled by async stuff
+			LockedTexelPosition.Set(0);
+			LockedTexelCount = WriteTexture.GetSizeX();
+		}
+	}
+
+	if (CVarMotionBlurDebug.GetValueOnRenderThread())
 	{
-		LockedData = WriteTexture.LockData();
-		check(LockedTexelPosition.GetValue() == 0); //otherwise it wasn't unlocked or it was unlocked before it was done being filled by async stuff
-		LockedTexelPosition.Set(0);
-		LockedTexelCount = WriteTexture.GetSizeX();
+		UE_LOG(LogEngine, Log, TEXT("r.MotionBlurDebug: BufferSize=%d Read=%d Write=%d IsLocked=%d"),
+			PER_BONE_BUFFER_COUNT, GetReadBufferIndex(), GetWriteBufferIndex(), IsAppendStarted() ? 1 : 0);
 	}
 }
 
@@ -1197,8 +1193,6 @@ void FPreviousPerBoneMotionBlur::EndAppend()
 		// we use float4
 		PerChunkBoneMatricesTexture[GetWriteBufferIndex()].UnlockData(LockedTexelPosition.GetValue() * 4 * sizeof(float));
 		LockedTexelPosition.Set(0);
-
-		AdvanceBufferIndex();
 	}
 
 	{
@@ -1226,14 +1220,20 @@ FString FPreviousPerBoneMotionBlur::GetDebugString() const
 	return FString::Printf(TEXT("BufferIndex=%d Pos=%d"), GPrevPerBoneMotionBlur.GetWriteBufferIndex(), GPrevPerBoneMotionBlur.LockedTexelPosition.GetValue());
 }
 
-void FPreviousPerBoneMotionBlur::AdvanceBufferIndex()
+void FPreviousPerBoneMotionBlur::AdvanceBufferLocation()
 {
-	++BufferIndex;
+	// roll the content on to the left
+	// e.g.
+	// ABCD -> BCDA
+	FBoneBufferTypeRef Temp = PerChunkBoneMatricesTexture[0].BoneBuffer;
 
-	if(BufferIndex >= PER_BONE_BUFFER_COUNT)
+	for(uint32 ToIndex = 0; ToIndex < PER_BONE_BUFFER_COUNT - 1; ++ToIndex)
 	{
-		BufferIndex = 0;
+		uint32 FromIndex  = ToIndex + 1;
+
+		PerChunkBoneMatricesTexture[ToIndex].BoneBuffer = PerChunkBoneMatricesTexture[FromIndex].BoneBuffer;
 	}
+	PerChunkBoneMatricesTexture[PER_BONE_BUFFER_COUNT - 1].BoneBuffer = Temp;
 }
 
 /** 
diff --git a/Engine/Source/Runtime/Engine/Private/SkeletalRenderGPUSkin.h b/Engine/Source/Runtime/Engine/Private/SkeletalRenderGPUSkin.h
--- a/Engine/Source/Runtime/Engine/Private/SkeletalRenderGPUSkin.h
+++ b/Engine/Source/Runtime/Engine/Private/SkeletalRenderGPUSkin.h
@@ -220,19 +220,15 @@ public:
 	/** @return 0 .. PER_BONE_BUFFER_COUNT-1 */
 	uint32 GetReadBufferIndex() const
 	{
-		return BufferIndex;
+		// we move the buffers for easier internal so it's a fixed position
+		return 0;
 	}
 
 	/** @return 0 .. PER_BONE_BUFFER_COUNT-1 */
 	uint32 GetWriteBufferIndex() const
 	{
-		uint32 ret = BufferIndex + 1;
-
-		if(ret >= PER_BONE_BUFFER_COUNT)
-		{
-			ret = 0;
-		}
-		return ret;
+		// we move the buffers for easier internal so it's a fixed position
+		return 1;
 	}
 
 
@@ -240,8 +236,6 @@ private:
 
 	/** Stores the bone information with one frame delay, required for per bone motion blur. Buffered to avoid stalls on texture Lock() */
 	FBoneDataVertexBuffer PerChunkBoneMatricesTexture[PER_BONE_BUFFER_COUNT];
-	/** cycles between the buffers to avoid stalls (when CPU would need to wait on GPU) */
-	uint32 BufferIndex;
 	/* !=0 if data is locked, does not change during the lock */
 	float* LockedData;
 	/** only valid if LockedData != 0, advances with every Append() */
@@ -253,7 +247,7 @@ private:
 
 
 	/** to cycle the internal buffer counter */
-	void AdvanceBufferIndex();
+	void AdvanceBufferLocation();
 };
 
 
diff --git a/Engine/Source/Runtime/Engine/Public/GPUSkinVertexFactory.h b/Engine/Source/Runtime/Engine/Public/GPUSkinVertexFactory.h
--- a/Engine/Source/Runtime/Engine/Public/GPUSkinVertexFactory.h
+++ b/Engine/Source/Runtime/Engine/Public/GPUSkinVertexFactory.h
@@ -344,6 +344,12 @@ public:
 			}
 		}
 
+		void MaintainBoneDataStartIndex() const
+		{
+			OldBoneFrameNumber[0]++;
+			OldBoneFrameNumber[1]++;
+		}
+
 		/** Checks if we need to update the data for this frame */
 		bool IsOldBoneDataUpdateNeeded(uint32 InFrameNumber) const
 		{
@@ -437,6 +443,11 @@ private:
 		ShaderData.SetOldBoneData(FrameNumber, Index);
 	}
 
+	void MaintainBoneDataStartIndex() const
+	{
+		ShaderData.MaintainBoneDataStartIndex();
+	}
+
 	static bool SupportsTessellationShaders() { return true; }
 
 	/** Checks if we need to update the data for this frame */
diff --git a/Engine/Source/Runtime/Renderer/Private/PostProcess/PostProcessMotionBlur.cpp b/Engine/Source/Runtime/Renderer/Private/PostProcess/PostProcessMotionBlur.cpp
--- a/Engine/Source/Runtime/Renderer/Private/PostProcess/PostProcessMotionBlur.cpp
+++ b/Engine/Source/Runtime/Renderer/Private/PostProcess/PostProcessMotionBlur.cpp
@@ -1513,12 +1513,20 @@ void FRCPassPostProcessVisualizeMotionBlur::Process(FRenderingCompositePassConte
 	Line = FString::Printf(TEXT("Visualize MotionBlur"));
 	Canvas.DrawShadowedString(X, Y += YStep, *Line, GetStatsFont(), FLinearColor(1, 1, 0));
 	
-	Line = FString::Printf(TEXT("%d"), ViewFamily.FrameNumber);
-	Canvas.DrawShadowedString(X, Y += YStep, TEXT("FrameNumber:"), GetStatsFont(), FLinearColor(1, 1, 0));
+	static const auto MotionBlurDebugVar = IConsoleManager::Get().FindTConsoleVariableDataInt(TEXT("r.MotionBlurDebug"));
+	const int32 MotionBlurDebug = MotionBlurDebugVar ? MotionBlurDebugVar->GetValueOnRenderThread() : 0;
+
+	Line = FString::Printf(TEXT("%d, %d"), ViewFamily.FrameNumber, MotionBlurDebug);
+	Canvas.DrawShadowedString(X, Y += YStep, TEXT("FrameNumber, r.MotionBlurDebug:"), GetStatsFont(), FLinearColor(1, 1, 0));
 	Canvas.DrawShadowedString(X + ColumnWidth, Y, *Line, GetStatsFont(), FLinearColor(1, 1, 0));
 
-	Line = FString::Printf(TEXT("%d"), ViewFamily.bWorldIsPaused ? 1 : 0);
-	Canvas.DrawShadowedString(X, Y += YStep, TEXT("WorldIsPaused:"), GetStatsFont(), FLinearColor(1, 1, 0));
+	static const auto VelocityTestVar = IConsoleManager::Get().FindTConsoleVariableDataInt(TEXT("r.VelocityTest"));
+	const int32 VelocityTest = VelocityTestVar ? VelocityTestVar->GetValueOnRenderThread() : 0;
+	
+	extern bool IsParallelVelocity();
+
+	Line = FString::Printf(TEXT("%d, %d, %d"), ViewFamily.bWorldIsPaused ? 1 : 0, VelocityTest, IsParallelVelocity());
+	Canvas.DrawShadowedString(X, Y += YStep, TEXT("WorldIsPaused, r.VelocityTest, Parallel:"), GetStatsFont(), FLinearColor(1, 1, 0));
 	Canvas.DrawShadowedString(X + ColumnWidth, Y, *Line, GetStatsFont(), FLinearColor(1, 1, 0));
 
 	for(uint32 BufferId = 0; BufferId < 2; ++BufferId)
diff --git a/Engine/Source/Runtime/Renderer/Private/VelocityRendering.cpp b/Engine/Source/Runtime/Renderer/Private/VelocityRendering.cpp
--- a/Engine/Source/Runtime/Renderer/Private/VelocityRendering.cpp
+++ b/Engine/Source/Runtime/Renderer/Private/VelocityRendering.cpp
@@ -17,6 +17,23 @@ static TAutoConsoleVariable<int32> CVarBasePassOutputsVelocity(
 	TEXT(" 1: Renders during the regular base pass adding an extra GBuffer, but allowing motion blur on materials with Time-based WPO."),
 	ECVF_ReadOnly | ECVF_RenderThreadSafe);
 
+static TAutoConsoleVariable<int32> CVarParallelVelocity(
+	TEXT("r.ParallelVelocity"),
+	1,  
+	TEXT("Toggles parallel velocity rendering. Parallel rendering must be enabled for this to have an effect."),
+	ECVF_RenderThreadSafe
+	);
+
+static TAutoConsoleVariable<int32> CVarRHICmdVelocityPassDeferredContexts(
+	TEXT("r.RHICmdVelocityPassDeferredContexts"),
+	1,
+	TEXT("True to use deferred contexts to parallelize velocity pass command list execution."));
+
+bool IsParallelVelocity()
+{
+	return GRHICommandList.UseParallelAlgorithms() && CVarParallelVelocity.GetValueOnRenderThread();
+}
+
 //=============================================================================
 /** Encapsulates the Velocity vertex shader. */
 class FVelocityVS : public FMeshMaterialShader
@@ -628,11 +645,6 @@ public:
 	}
 };
 
-static TAutoConsoleVariable<int32> CVarRHICmdVelocityPassDeferredContexts(
-	TEXT("r.RHICmdVelocityPassDeferredContexts"),
-	1,
-	TEXT("True to use deferred contexts to parallelize velocity pass command list execution."));
-
 void FDeferredShadingSceneRenderer::RenderVelocitiesInnerParallel(FRHICommandListImmediate& RHICmdList, TRefCountPtr<IPooledRenderTarget>& VelocityRT)
 {
 	// parallel version
@@ -690,15 +702,6 @@ void FDeferredShadingSceneRenderer::RenderVelocitiesInner(FRHICommandListImmedia
 	}
 }
 
-
-
-static TAutoConsoleVariable<int32> CVarParallelVelocity(
-	TEXT("r.ParallelVelocity"),
-	1,  
-	TEXT("Toggles parallel velocity rendering. Parallel rendering must be enabled for this to have an effect."),
-	ECVF_RenderThreadSafe
-	);
-
 bool FDeferredShadingSceneRenderer::ShouldRenderVelocities() const
 {
 	if (!GPixelFormats[PF_G16R16].Supported)
@@ -753,7 +756,7 @@ void FDeferredShadingSceneRenderer::RenderVelocities(FRHICommandListImmediate& R
 
 		BeginVelocityRendering(RHICmdList, VelocityRT, true);
 
-		if (GRHICommandList.UseParallelAlgorithms() && CVarParallelVelocity.GetValueOnRenderThread())
+		if (IsParallelVelocity())
 		{
 			RenderVelocitiesInnerParallel(RHICmdList, VelocityRT);
 		}
