#ue4 - Remove TArray from GeomSweepTest() unless it's being used by some debug code.

diff --git a/Engine/Source/Runtime/Engine/Private/Collision/PhysXCollision.cpp b/Engine/Source/Runtime/Engine/Private/Collision/PhysXCollision.cpp
--- a/Engine/Source/Runtime/Engine/Private/Collision/PhysXCollision.cpp
+++ b/Engine/Source/Runtime/Engine/Private/Collision/PhysXCollision.cpp
@@ -1220,15 +1220,22 @@ bool GeomSweepTest(const UWorld* World, const struct FCollisionShape& CollisionS
 		}
 	}
 
-	TArray<FHitResult> Hits;
 #if !(UE_BUILD_SHIPPING || UE_BUILD_TEST)
 	if((World->DebugDrawTraceTag != NAME_None) && (World->DebugDrawTraceTag == Params.TraceTag))
 	{
+		TArray<FHitResult> Hits;
 		DrawGeomSweeps(World, Start, End, PGeom, PGeomRot, Hits, DebugLineLifetime);
 	}
 #endif //!(UE_BUILD_SHIPPING || UE_BUILD_TEST)
 
-	CAPTUREGEOMSWEEP(World, Start, End, PGeomRot, ECAQueryType::Test, PGeom, TraceChannel, Params, ResponseParams, ObjectParams, Hits);
+#if ENABLE_COLLISION_ANALYZER
+	if (GCollisionAnalyzerIsRecording)
+	{
+		TArray<FHitResult> Hits;
+		CAPTUREGEOMSWEEP(World, Start, End, PGeomRot, ECAQueryType::Test, PGeom, TraceChannel, Params, ResponseParams, ObjectParams, Hits);
+	}
+#endif // ENABLE_COLLISION_ANALYZER
+
 #endif // WITH_PHYSX
 
 	//@TODO: BOX2D: Implement GeomSweepTest
