#ue4 - Minimize accumulation of errors after many composed transforms when updating SceneComponent::ComponentToWorld.

diff --git a/Engine/Source/Runtime/Engine/Private/Components/SceneComponent.cpp b/Engine/Source/Runtime/Engine/Private/Components/SceneComponent.cpp
--- a/Engine/Source/Runtime/Engine/Private/Components/SceneComponent.cpp
+++ b/Engine/Source/Runtime/Engine/Private/Components/SceneComponent.cpp
@@ -337,7 +337,8 @@ void USceneComponent::UpdateComponentToWorldWithParent(USceneComponent* Parent,
 
 	// Calculate the new ComponentToWorld transform
 	const FTransform RelativeTransform(RelativeRotationQuat, RelativeLocation, RelativeScale3D);
-	const FTransform NewTransform = CalcNewComponentToWorld(RelativeTransform, Parent);
+	FTransform NewTransform = CalcNewComponentToWorld(RelativeTransform, Parent);
+	NewTransform.NormalizeRotation(); // Minimize accumulation of errors after many composed transforms.
 
 	// If transform has changed..
 	if (!ComponentToWorld.Equals(NewTransform, SMALL_NUMBER))
