#ue4 - Reorder calls to SetHitResultFromShapeAndFaceIndex in CollisionConversions to better match memory layout of FHitResult.

diff --git a/Engine/Source/Runtime/Engine/Private/Collision/CollisionConversions.cpp b/Engine/Source/Runtime/Engine/Private/Collision/CollisionConversions.cpp
--- a/Engine/Source/Runtime/Engine/Private/Collision/CollisionConversions.cpp
+++ b/Engine/Source/Runtime/Engine/Private/Collision/CollisionConversions.cpp
@@ -418,8 +418,6 @@ void ConvertQueryImpactHit(const UWorld* World, const PxLocationHit& PHit, FHitR
 		return;
 	}
 
-	SetHitResultFromShapeAndFaceIndex(PHit.shape,  PHit.actor, PHit.faceIndex, OutResult, bReturnPhysMat);
-
 	// See if this is a 'blocking' hit
 	const PxFilterData PShapeFilter = PHit.shape->getQueryFilterData();
 	const PxSceneQueryHitType::Enum HitType = FPxQueryFilterCallback::CalcQueryHitType(QueryFilter, PShapeFilter);
@@ -463,6 +461,9 @@ void ConvertQueryImpactHit(const UWorld* World, const PxLocationHit& PHit, FHitR
 	const PxGeometryType::Enum SweptGeometryType = Geom ? Geom->getType() : PxGeometryType::eINVALID;
 	OutResult.ImpactNormal = FindGeomOpposingNormal(SweptGeometryType, PHit, TraceStartToEnd, Normal);
 	
+	// Fill in Actor, Component, material, etc.
+	SetHitResultFromShapeAndFaceIndex(PHit.shape, PHit.actor, PHit.faceIndex, OutResult, bReturnPhysMat);
+
 	if( PHit.shape->getGeometryType() == PxGeometryType::eHEIGHTFIELD)
 	{
 		// Lookup physical material for heightfields
@@ -821,8 +822,6 @@ static bool ConvertOverlappedShapeToImpactHit(const UWorld* World, const PxLocat
 	const PxRigidActor* PActor = PHit.actor;
 	const uint32 FaceIdx = PHit.faceIndex;
 
-	SetHitResultFromShapeAndFaceIndex(PShape, PActor, FaceIdx, OutResult, bReturnPhysMat);
-
 	// Time of zero because initially overlapping
 	OutResult.Time = 0.f;
 	OutResult.bStartPenetrating = true;
@@ -909,6 +908,8 @@ static bool ConvertOverlappedShapeToImpactHit(const UWorld* World, const PxLocat
 	}
 
 	OutResult.Normal = OutResult.ImpactNormal;
+	SetHitResultFromShapeAndFaceIndex(PShape, PActor, FaceIdx, OutResult, bReturnPhysMat);
+
 	return OutResult.bBlockingHit;
 }
 
