#ue4 - Clean up Character movement stats. Add new stat for AdjustFloorHeight.

diff --git a/Engine/Source/Runtime/Engine/Private/Components/CharacterMovementComponent.cpp b/Engine/Source/Runtime/Engine/Private/Components/CharacterMovementComponent.cpp
--- a/Engine/Source/Runtime/Engine/Private/Components/CharacterMovementComponent.cpp
+++ b/Engine/Source/Runtime/Engine/Private/Components/CharacterMovementComponent.cpp
@@ -34,14 +34,15 @@ DECLARE_CYCLE_STAT(TEXT("Char Movement Tick"), STAT_CharacterMovementTick, STATG
 DECLARE_CYCLE_STAT(TEXT("Char Movement Authority Time"), STAT_CharacterMovementAuthority, STATGROUP_Character);
 DECLARE_CYCLE_STAT(TEXT("Char Movement Simulated Time"), STAT_CharacterMovementSimulated, STATGROUP_Character);
 DECLARE_CYCLE_STAT(TEXT("Char Physics Interation"), STAT_CharPhysicsInteraction, STATGROUP_Character);
-DECLARE_CYCLE_STAT(TEXT("StepUp"), STAT_CharStepUp, STATGROUP_Character);
+DECLARE_CYCLE_STAT(TEXT("Char StepUp"), STAT_CharStepUp, STATGROUP_Character);
+DECLARE_CYCLE_STAT(TEXT("Char AdjustFloorHeight"), STAT_CharAdjustFloorHeight, STATGROUP_Character);
 DECLARE_CYCLE_STAT(TEXT("Char Update Acceleration"), STAT_CharUpdateAcceleration, STATGROUP_Character);
 DECLARE_CYCLE_STAT(TEXT("Char MoveUpdateDelegate"), STAT_CharMoveUpdateDelegate, STATGROUP_Character);
-DECLARE_CYCLE_STAT(TEXT("PhysWalking"), STAT_CharPhysWalking, STATGROUP_Character);
-DECLARE_CYCLE_STAT(TEXT("PhysFalling"), STAT_CharPhysFalling, STATGROUP_Character);
-DECLARE_CYCLE_STAT(TEXT("PhysNavWalking"), STAT_CharPhysNavWalking, STATGROUP_Character);
-DECLARE_CYCLE_STAT(TEXT("NavProjectPoint"), STAT_CharNavProjectPoint, STATGROUP_Character);
-DECLARE_CYCLE_STAT(TEXT("NavProjectLocation"), STAT_CharNavProjectLocation, STATGROUP_Character);
+DECLARE_CYCLE_STAT(TEXT("Char PhysWalking"), STAT_CharPhysWalking, STATGROUP_Character);
+DECLARE_CYCLE_STAT(TEXT("Char PhysFalling"), STAT_CharPhysFalling, STATGROUP_Character);
+DECLARE_CYCLE_STAT(TEXT("Char PhysNavWalking"), STAT_CharPhysNavWalking, STATGROUP_Character);
+DECLARE_CYCLE_STAT(TEXT("Char NavProjectPoint"), STAT_CharNavProjectPoint, STATGROUP_Character);
+DECLARE_CYCLE_STAT(TEXT("Char NavProjectLocation"), STAT_CharNavProjectLocation, STATGROUP_Character);
 
 // MAGIC NUMBERS
 const float MAX_STEP_SIDE_Z = 0.08f;	// maximum z value for the normal on the vertical side of steps
@@ -4053,6 +4054,8 @@ bool UCharacterMovementComponent::ShouldCatchAir(const FFindFloorResult& OldFloo
 
 void UCharacterMovementComponent::AdjustFloorHeight()
 {
+	SCOPE_CYCLE_COUNTER(STAT_CharAdjustFloorHeight);
+
 	// If we have a floor check that hasn't hit anything, don't adjust height.
 	if (!CurrentFloor.bBlockingHit)
 	{
