#ue4 - Perf: Optimize usage of weak pointers in collision conversions. Remove memcpy of FHitResult when filling in results. Avoid expensive MTD code for overlaps because this is only for blocking hits.

diff --git a/Engine/Source/Runtime/Engine/Private/Collision/CollisionConversions.cpp b/Engine/Source/Runtime/Engine/Private/Collision/CollisionConversions.cpp
--- a/Engine/Source/Runtime/Engine/Private/Collision/CollisionConversions.cpp
+++ b/Engine/Source/Runtime/Engine/Private/Collision/CollisionConversions.cpp
@@ -393,9 +393,10 @@ static void SetHitResultFromShapeAndFaceIndex(const PxShape* PShape,  const PxRi
 	{
 		OutResult.Item = BodyInst->InstanceBodyIndex;
 
-		if (BodyInst->BodySetup.IsValid())
+		const UBodySetup* BodySetup = BodyInst->BodySetup.Get();
+		if (BodySetup)
 		{
-			OutResult.BoneName = BodyInst->BodySetup->BoneName;
+			OutResult.BoneName = BodySetup->BoneName;
 		}
 	}
 	else
@@ -504,12 +505,10 @@ void ConvertRaycastResults(const UWorld* World, int32 NumHits, PxRaycastHit* Hit
 	PxTransform PStartTM(U2PVector(StartLoc));
 	for(int32 i=0; i<NumHits; i++)
 	{
-		FHitResult NewResult(ForceInit);
+		FHitResult& NewResult = OutHits[OutHits.AddDefaulted()];
 		const PxRaycastHit& PHit = Hits[i];
 
 		ConvertQueryImpactHit(World, PHit, NewResult, CheckLength, QueryFilter, StartLoc, EndLoc, NULL, PStartTM, bReturnFaceIndex, bReturnPhysMat);
-
-		OutHits.Add(NewResult);
 	}
 
 	// Sort results from first to last hit
@@ -527,10 +526,9 @@ bool AddSweepResults(const UWorld* World, int32 NumHits, const PxSweepHit* Hits,
 		checkSlow(PHit.flags & PxHitFlag::eDISTANCE);
 		if(PHit.distance <= MaxDistance)
 		{
-			FHitResult NewResult(ForceInit);
+			FHitResult& NewResult = OutHits[OutHits.AddDefaulted()];
 			ConvertQueryImpactHit(World, PHit, NewResult, CheckLength, QueryFilter, StartLoc, EndLoc, &Geom, QueryTM, false, bReturnPhysMat);
 			bHadBlockingHit |= NewResult.bBlockingHit;
-			OutHits.Add(NewResult);
 		}
 	}
 
@@ -815,26 +813,27 @@ static bool ConvertOverlappedShapeToImpactHit(const UWorld* World, const PxLocat
 {
 	SCOPE_CYCLE_COUNTER(STAT_CollisionConvertOverlapToHit);
 
-	OutResult.TraceStart = StartLoc;
-	OutResult.TraceEnd = EndLoc;
-
 	const PxShape* PShape = PHit.shape;
 	const PxRigidActor* PActor = PHit.actor;
 	const uint32 FaceIdx = PHit.faceIndex;
 
-	// Time of zero because initially overlapping
-	OutResult.Time = 0.f;
-	OutResult.bStartPenetrating = true;
-
 	// See if this is a 'blocking' hit
 	PxFilterData PShapeFilter = PShape->getQueryFilterData();
 	PxSceneQueryHitType::Enum HitType = FPxQueryFilterCallback::CalcQueryHitType(QueryFilter, PShapeFilter);
-	OutResult.bBlockingHit = (HitType == PxSceneQueryHitType::eBLOCK); 
+	const bool bBlockingHit = (HitType == PxSceneQueryHitType::eBLOCK); 
+	OutResult.bBlockingHit = bBlockingHit;
+
+	// Time of zero because initially overlapping
+	OutResult.bStartPenetrating = true;
+	OutResult.Time = 0.f;
 
 	// Return start location as 'safe location'
 	OutResult.Location = P2UVector(QueryTM.p);
 	OutResult.ImpactPoint = OutResult.Location; // @todo not really sure of a better thing to do here...
 
+	OutResult.TraceStart = StartLoc;
+	OutResult.TraceEnd = EndLoc;
+
 	const bool bFiniteNormal = PHit.normal.isFinite();
 	const bool bValidNormal = (PHit.flags & PxHitFlag::eNORMAL) && bFiniteNormal;
 
@@ -864,53 +863,64 @@ static bool ConvertOverlappedShapeToImpactHit(const UWorld* World, const PxLocat
 	}
 #endif
 
-	// Zero-distance hits are often valid hits and we can extract the hit normal.
-	// For invalid normals we can try other methods as well (get overlapping triangles).
-
-	if (PHit.distance == 0.f || !bValidNormal)
+	if (bBlockingHit)
 	{
-		const PxTransform PShapeWorldPose = PxShapeExt::getGlobalPose(*PShape, *PActor);
+		// Zero-distance hits are often valid hits and we can extract the hit normal.
+		// For invalid normals we can try other methods as well (get overlapping triangles).
+		if (PHit.distance == 0.f || !bValidNormal)
+		{
+			const PxTransform PShapeWorldPose = PxShapeExt::getGlobalPose(*PShape, *PActor);
 
-		// Try MTD with a small inflation for better accuracy, then a larger one in case the first one fails due to precision issues.
-		static const float SmallMtdInflation = 0.250f;
-		static const float LargeMtdInflation = 1.750f;
+			// Try MTD with a small inflation for better accuracy, then a larger one in case the first one fails due to precision issues.
+			static const float SmallMtdInflation = 0.250f;
+			static const float LargeMtdInflation = 1.750f;
 
-		if (ComputeInflatedMTD(SmallMtdInflation, PHit, OutResult, QueryTM, Geom, PShapeWorldPose) ||
-			ComputeInflatedMTD(LargeMtdInflation, PHit, OutResult, QueryTM, Geom, PShapeWorldPose))
-		{
-			// Success
-		}
-		else
-		{
-			static const float SmallOverlapInflation = 0.250f;
-			if (FindOverlappedTriangleNormal(World, Geom, QueryTM, PShape, PShapeWorldPose, OutResult.ImpactNormal, 0.f, false) ||
-				FindOverlappedTriangleNormal(World, Geom, QueryTM, PShape, PShapeWorldPose, OutResult.ImpactNormal, SmallOverlapInflation, false))
+			if (ComputeInflatedMTD(SmallMtdInflation, PHit, OutResult, QueryTM, Geom, PShapeWorldPose) ||
+				ComputeInflatedMTD(LargeMtdInflation, PHit, OutResult, QueryTM, Geom, PShapeWorldPose))
 			{
 				// Success
 			}
 			else
 			{
-				// MTD failed, use point distance. This is not ideal.
-				// Note: faceIndex seems to be unreliable for convex meshes in these cases, so not using FindGeomOpposingNormal() for them here.
-				PxGeometry& PGeom = PShape->getGeometry().any();
-				PxVec3 PClosestPoint;
-				const float Distance = PxGeometryQuery::pointDistance(QueryTM.p, PGeom, PShapeWorldPose, &PClosestPoint);
-
-				if (Distance < KINDA_SMALL_NUMBER)
+				static const float SmallOverlapInflation = 0.250f;
+				if (FindOverlappedTriangleNormal(World, Geom, QueryTM, PShape, PShapeWorldPose, OutResult.ImpactNormal, 0.f, false) ||
+					FindOverlappedTriangleNormal(World, Geom, QueryTM, PShape, PShapeWorldPose, OutResult.ImpactNormal, SmallOverlapInflation, false))
 				{
-					UE_LOG(LogCollision, Verbose, TEXT("Warning: ConvertOverlappedShapeToImpactHit: Query origin inside shape, giving poor MTD."));
-					PClosestPoint = PxShapeExt::getWorldBounds(*PShape, *PActor).getCenter();
+					// Success
+				}
+				else
+				{
+					// MTD failed, use point distance. This is not ideal.
+					// Note: faceIndex seems to be unreliable for convex meshes in these cases, so not using FindGeomOpposingNormal() for them here.
+					PxGeometry& PGeom = PShape->getGeometry().any();
+					PxVec3 PClosestPoint;
+					const float Distance = PxGeometryQuery::pointDistance(QueryTM.p, PGeom, PShapeWorldPose, &PClosestPoint);
+
+					if (Distance < KINDA_SMALL_NUMBER)
+					{
+						UE_LOG(LogCollision, Verbose, TEXT("Warning: ConvertOverlappedShapeToImpactHit: Query origin inside shape, giving poor MTD."));
+						PClosestPoint = PxShapeExt::getWorldBounds(*PShape, *PActor).getCenter();
+					}
+
+					OutResult.ImpactNormal = (OutResult.Location - P2UVector(PClosestPoint)).GetSafeNormal();
 				}
-
-				OutResult.ImpactNormal = (OutResult.Location - P2UVector(PClosestPoint)).GetSafeNormal();
 			}
 		}
 	}
+	else
+	{
+		// non blocking hit (overlap).
+		if (!bValidNormal)
+		{
+			OutResult.ImpactNormal = (StartLoc - EndLoc).GetSafeNormal();
+			ensure(OutResult.Normal.IsNormalized());
+		}
+	}
 
 	OutResult.Normal = OutResult.ImpactNormal;
 	SetHitResultFromShapeAndFaceIndex(PShape, PActor, FaceIdx, OutResult, bReturnPhysMat);
 
-	return OutResult.bBlockingHit;
+	return bBlockingHit;
 }
 
 
@@ -922,18 +932,34 @@ void ConvertQueryOverlap(const PxShape* PShape, const PxRigidActor* PActor, FOve
 	FDestructibleChunkInfo* ChunkInfo = FPhysxUserData::Get<FDestructibleChunkInfo>(PActor->userData);
 
 	// Grab actor/component
-	if(BodyInst && BodyInst->OwnerComponent.IsValid())
+	UPrimitiveComponent* OwnerComponent = nullptr;
+
+	// Try body instance
+	if (BodyInst)
 	{
-		OutOverlap.Actor = BodyInst->OwnerComponent->GetOwner();
-		OutOverlap.Component = BodyInst->OwnerComponent;
-		OutOverlap.ItemIndex = BodyInst->OwnerComponent->bMultiBodyOverlap ? BodyInst->InstanceBodyIndex : INDEX_NONE;
+		OwnerComponent = BodyInst->OwnerComponent.Get(); // cache weak pointer object, avoid multiple derefs below.
+		if (OwnerComponent)
+		{
+			OutOverlap.Actor = OwnerComponent->GetOwner();
+			OutOverlap.Component = BodyInst->OwnerComponent; // Copying weak pointer is faster than assigning raw pointer.
+			OutOverlap.ItemIndex = OwnerComponent->bMultiBodyOverlap ? BodyInst->InstanceBodyIndex : INDEX_NONE;
+		}
 	}
-	else if (ChunkInfo && ChunkInfo->OwningComponent.IsValid())
+	
+	if (!OwnerComponent)
 	{
-		OutOverlap.Actor = ChunkInfo->OwningComponent->GetOwner();
-		OutOverlap.Component = ChunkInfo->OwningComponent.Get();
-		OutOverlap.ItemIndex = ChunkInfo->OwningComponent->bMultiBodyOverlap ? UDestructibleComponent::ChunkIdxToBoneIdx(ChunkInfo->ChunkIndex) : INDEX_NONE;
-	}
+		// Try chunk info
+		if (ChunkInfo)
+		{
+			OwnerComponent = ChunkInfo->OwningComponent.Get(); // cache weak pointer object, avoid multiple derefs below.
+			if (OwnerComponent)
+			{
+				OutOverlap.Actor = OwnerComponent->GetOwner();
+				OutOverlap.Component = ChunkInfo->OwningComponent; // Copying weak pointer is faster than assigning raw pointer.
+				OutOverlap.ItemIndex = OwnerComponent->bMultiBodyOverlap ? UDestructibleComponent::ChunkIdxToBoneIdx(ChunkInfo->ChunkIndex) : INDEX_NONE;
+			}
+		}
+	}	
 
 	// Other info
 	OutOverlap.bBlockingHit = bBlock;
