#ue4 - Perf: Remove deep copy of hit results in UWorld::ComponentSweepMulti().

- This also ensures that sweeps with multiple shapes actually returns the correctly sorted list of hits, since all hits are now sorted together rather than appended in individually sorted chunks.
- Added note that GeomSweepMulti_PhysX does not clear the incoming OutHits array, only appends to it.
- GeomSweepMulti_PhysX is now marked DEPRECATED_FORGAME. It was never intended to be a public function that games should use.

diff --git a/Engine/Source/Runtime/Engine/Private/Collision/CollisionConversions.cpp b/Engine/Source/Runtime/Engine/Private/Collision/CollisionConversions.cpp
--- a/Engine/Source/Runtime/Engine/Private/Collision/CollisionConversions.cpp
+++ b/Engine/Source/Runtime/Engine/Private/Collision/CollisionConversions.cpp
@@ -500,7 +500,7 @@ void ConvertQueryImpactHit(const UWorld* World, const PxLocationHit& PHit, FHitR
 
 void ConvertRaycastResults(const UWorld* World, int32 NumHits, PxRaycastHit* Hits, float CheckLength, const PxFilterData& QueryFilter, TArray<FHitResult>& OutHits, const FVector& StartLoc, const FVector& EndLoc, bool bReturnFaceIndex, bool bReturnPhysMat)
 {
-	OutHits.Reserve(NumHits);
+	OutHits.Reserve(OutHits.Num() + NumHits);
 
 	PxTransform PStartTM(U2PVector(StartLoc));
 	for(int32 i=0; i<NumHits; i++)
@@ -517,7 +517,7 @@ void ConvertRaycastResults(const UWorld* World, int32 NumHits, PxRaycastHit* Hit
 
 bool AddSweepResults(const UWorld* World, int32 NumHits, const PxSweepHit* Hits, float CheckLength, const PxFilterData& QueryFilter, TArray<FHitResult>& OutHits, const FVector& StartLoc, const FVector& EndLoc, const PxGeometry& Geom, const PxTransform& QueryTM, float MaxDistance, bool bReturnPhysMat)
 {
-	OutHits.Reserve(NumHits);
+	OutHits.Reserve(OutHits.Num() + NumHits);
 	bool bHadBlockingHit = false;
 
 	for(int32 i=0; i<NumHits; i++)
@@ -1005,7 +1005,7 @@ bool ConvertOverlapResults(int32 NumOverlaps, PxOverlapHit* POverlapResults, con
 {
 	SCOPE_CYCLE_COUNTER(STAT_CollisionConvertOverlap);
 
-	OutOverlaps.Reserve(NumOverlaps);
+	OutOverlaps.Reserve(OutOverlaps.Num() + NumOverlaps);
 	bool bBlockingFound = false;
 
 	for(int32 i=0; i<NumOverlaps; i++)
diff --git a/Engine/Source/Runtime/Engine/Private/Collision/PhysXCollision.cpp b/Engine/Source/Runtime/Engine/Private/Collision/PhysXCollision.cpp
--- a/Engine/Source/Runtime/Engine/Private/Collision/PhysXCollision.cpp
+++ b/Engine/Source/Runtime/Engine/Private/Collision/PhysXCollision.cpp
@@ -1384,6 +1384,8 @@ bool GeomSweepMulti_PhysX(const UWorld* World, const PxGeometry& PGeom, const Px
 	STARTQUERYTIMER();
 	bool bBlockingHit = false;
 
+	const int32 InitialHitCount = OutHits.Num();
+
 	// Create filter data used to filter collisions
 	PxFilterData PFilter = CreateQueryFilterData(TraceChannel, Params.bTraceComplex, ResponseParams.CollisionResponse, ObjectParams, true);
 	PxSceneQueryFilterData PQueryFilterData(PFilter, PxSceneQueryFilterFlag::eSTATIC | PxSceneQueryFilterFlag::eDYNAMIC | PxSceneQueryFilterFlag::ePREFILTER | PxSceneQueryFilterFlag::ePOSTFILTER);
@@ -1520,11 +1522,20 @@ bool GeomSweepMulti_PhysX(const UWorld* World, const PxGeometry& PGeom, const Px
 #if !(UE_BUILD_SHIPPING || UE_BUILD_TEST)
 	if ((World->DebugDrawTraceTag != NAME_None) && (World->DebugDrawTraceTag == Params.TraceTag))
 	{
-		DrawGeomSweeps(World, Start, End, PGeom, PGeomRot, OutHits, DebugLineLifetime);
+		TArray<FHitResult> OnlyMyHits(OutHits);
+		OnlyMyHits.RemoveAt(0, InitialHitCount, false); // Remove whatever was there initially.
+		DrawGeomSweeps(World, Start, End, PGeom, PGeomRot, OnlyMyHits, DebugLineLifetime);
 	}
 #endif //!(UE_BUILD_SHIPPING || UE_BUILD_TEST)
 
-	CAPTUREGEOMSWEEP(World, Start, End, PGeomRot, ECAQueryType::Multi, PGeom, TraceChannel, Params, ResponseParams, ObjectParams, OutHits);
+#if ENABLE_COLLISION_ANALYZER
+	if (GCollisionAnalyzerIsRecording)
+	{
+		TArray<FHitResult> OnlyMyHits(OutHits);
+		OnlyMyHits.RemoveAt(0, InitialHitCount, false); // Remove whatever was there initially.
+		CAPTUREGEOMSWEEP(World, Start, End, PGeomRot, ECAQueryType::Multi, PGeom, TraceChannel, Params, ResponseParams, ObjectParams, OnlyMyHits);
+	}
+#endif // ENABLE_COLLISION_ANALYZER
 
 	return bBlockingHit;
 }
diff --git a/Engine/Source/Runtime/Engine/Private/Collision/PhysXCollision.h b/Engine/Source/Runtime/Engine/Private/Collision/PhysXCollision.h
--- a/Engine/Source/Runtime/Engine/Private/Collision/PhysXCollision.h
+++ b/Engine/Source/Runtime/Engine/Private/Collision/PhysXCollision.h
@@ -203,8 +203,12 @@ bool GeomSweepMulti(const UWorld* World, const struct FCollisionShape& Collision
 #endif
 
 // Note: Do not use these methods for new code, they are being phased out!
+// These functions do not empty the OutOverlaps/OutHits array; they add items to them.
 #if WITH_PHYSX
+DEPRECATED_FORGAME(4.9, "Do not access this function directly, use the generic non-PhysX functions.")
 bool GeomOverlapMulti_PhysX(const UWorld* World, const PxGeometry& PGeom, const PxTransform& PGeomPose, TArray<FOverlapResult>& OutOverlaps, ECollisionChannel TraceChannel, const struct FCollisionQueryParams& Params, const struct FCollisionResponseParams& ResponseParams, const struct FCollisionObjectQueryParams& ObjectParams);
+
+DEPRECATED_FORGAME(4.9, "Do not access this function directly, use the generic non-PhysX functions.")
 bool GeomSweepMulti_PhysX(const UWorld* World, const PxGeometry& PGeom, const PxQuat& PGeomRot, TArray<FHitResult>& OutHits, FVector Start, FVector End, ECollisionChannel TraceChannel, const struct FCollisionQueryParams& Params, const struct FCollisionResponseParams& ResponseParams, const struct FCollisionObjectQueryParams& ObjectParams = FCollisionObjectQueryParams::DefaultObjectQueryParam);
 #endif
 
diff --git a/Engine/Source/Runtime/Engine/Private/Collision/WorldCollision.cpp b/Engine/Source/Runtime/Engine/Private/Collision/WorldCollision.cpp
--- a/Engine/Source/Runtime/Engine/Private/Collision/WorldCollision.cpp
+++ b/Engine/Source/Runtime/Engine/Private/Collision/WorldCollision.cpp
@@ -489,7 +489,7 @@ bool UWorld::ComponentSweepMulti(TArray<struct FHitResult>& OutHits, class UPrim
 		return RaycastMulti(this, OutHits, Start, End, TraceChannel, Params, FCollisionResponseParams(PrimComp->GetCollisionResponseToChannels()));
 	}
 
-	OutHits.Empty();
+	OutHits.Reset();
 
 #if UE_WITH_PHYSICS
 	if (!PrimComp->BodyInstance.IsValidBodyInstance())
@@ -534,12 +534,10 @@ bool UWorld::ComponentSweepMulti(TArray<struct FHitResult>& OutHits, class UPrim
 				// consider localshape rotation for shape rotation
 				const PxQuat PShapeRot = PGeomRot * PLocalShape.q;
 
-				if (GeomSweepMulti_PhysX(this, *PGeom, PShapeRot, Hits, P2UVector(PShapeGlobalStartPose.p), P2UVector(PShapeGlobalEndPose.p), TraceChannel, Params, FCollisionResponseParams(PrimComp->GetCollisionResponseToChannels())))
+				if (GeomSweepMulti_PhysX(this, *PGeom, PShapeRot, OutHits, P2UVector(PShapeGlobalStartPose.p), P2UVector(PShapeGlobalEndPose.p), TraceChannel, Params, FCollisionResponseParams(PrimComp->GetCollisionResponseToChannels())))
 				{
 					bHaveBlockingHit = true;
 				}
-
-				OutHits.Append(Hits);
 			}
 		}
 	});
