#ue4 - Remove VERIFY_GEOMSWEEPSINGLE and VERIFY_GEOMSWEEPMULTI from PhysXCollision.

- Maintaining this code is not practical, and bloats the functions which causes potential bugs as code is updated.

diff --git a/Engine/Source/Runtime/Engine/Private/Collision/PhysXCollision.cpp b/Engine/Source/Runtime/Engine/Private/Collision/PhysXCollision.cpp
--- a/Engine/Source/Runtime/Engine/Private/Collision/PhysXCollision.cpp
+++ b/Engine/Source/Runtime/Engine/Private/Collision/PhysXCollision.cpp
@@ -472,10 +472,6 @@ private:
 #include "CollisionDebugDrawing.h"
 #include "CollisionConversions.h"
 
-
-#define VERIFY_GEOMSWEEPSINGLE 0
-#define VERIFY_GEOMSWEEPMULTI 0
-
 float DebugLineLifetime = 2.f;
 
 /**
@@ -1291,26 +1287,6 @@ bool GeomSweepSingle(const UWorld* World, const struct FCollisionShape& Collisio
 		PxSweepHit PHit;
 		bHaveBlockingHit = SyncScene->sweepSingle(PGeom, PStartTM, PDir, DeltaMag, POutputFlags, PHit, PQueryFilterData, &PQueryCallbackSweep);
 
-#if VERIFY_GEOMSWEEPSINGLE
-		// If we hit but did not start overlapping...
-		if(bHaveBlockingHit && !PHit.hadInitialOverlap())
-		{
-			const float HitTime = PHit.distance/DeltaMag;
-			FVector HitLocation = Start + (Delta/DeltaMag) * HitTime;
-			PxTransform PEndTM(U2PVector(HitLocation), PGeomRot);
-
-			PxSweepHit PTestHit;
-			bool bTestBlockingHit = SyncScene->sweepSingle(PGeom, PEndTM, -PDir, DeltaMag, POutputFlags, PTestHit, PQueryFilterData, &PQueryCallbackSweep);
-			if(bTestBlockingHit)
-			{
-				UE_LOG(LogCollision, Warning, TEXT("GeomSweepSingle : Cannot sweep back from end location."));
-
-				PxSweepHit PReproHit;
-				bool bReproBlockingHit = SyncScene->sweepSingle(PGeom, PStartTM, PDir, DeltaMag, POutputFlags, PReproHit, PQueryFilterData, &PQueryCallbackSweep);
-			}
-		}
-#endif
-
 		if (!bHaveBlockingHit)
 		{
 			// Not using anything from this scene, so unlock it.
@@ -1431,47 +1407,7 @@ bool GeomSweepMulti_PhysX(const UWorld* World, const PxGeometry& PGeom, const Px
 			MinBlockDistance = PHits[NumHits - 1].distance;
 			bBlockingHit = true;
 		}
-
-#if VERIFY_GEOMSWEEPMULTI
-		// If we hit something, test we can sweep back...
-		{
-			if (OutHits.Num() == 0)
-			{
-				FVector HitLocation = End;
-				float TestLength = DeltaMag;
-				if (bBlockingHitSync)
-				{
-					TestLength = PHits[NumHits - 1].distance;
-					FHitResult TestHit;
-					ConvertQueryImpactHit(World, PHits[NumHits - 1], TestHit, DeltaMag, PFilter, Start, End, &PGeom, PStartTM, false, Params.bReturnPhysicalMaterial);
-					HitLocation = TestHit.Location;
-				}
-
-				UE_LOG(LogCollision, Warning, TEXT("%d Start: %s End: %s"), bBlockingHitSync, *Start.ToString(), *HitLocation.ToString());
-
-				PxTransform PEndTM(U2PVector(HitLocation), PGeomRot);
-
-				PxSweepHit PTestHits[HIT_BUFFER_MAX_SYNC_QUERIES];
-				bool bTestBlockingHit;
-				PxI32 NumTestHits = SyncScene->sweepMultiple(PGeom, PEndTM, -PDir, TestLength, POutputFlags, PTestHits, HIT_BUFFER_MAX_SYNC_QUERIES, bTestBlockingHit, PQueryFilterData, &PQueryCallbackSweep);
-
-				if (bTestBlockingHit)
-				{
-					UE_LOG(LogCollision, Warning, TEXT("GeomSweepMulti : Cannot sweep back from end location."));
-
-					PxSweepHit PReproHits[HIT_BUFFER_MAX_SYNC_QUERIES];
-					bool bReproBlockingHit;
-					PxI32 NumReproHits = SyncScene->sweepMultiple(PGeom, PStartTM, PDir, DeltaMag, POutputFlags, PReproHits, HIT_BUFFER_MAX_SYNC_QUERIES, bReproBlockingHit, PQueryFilterData, &PQueryCallbackSweep);
-				}
-			}
-			else
-			{
-				UE_LOG(LogCollision, Warning, TEXT("STUCK Start: %s End: %s"), *Start.ToString(), *End.ToString());
-			}
-		}
-#endif
-
-		if (NumHits == 0)
+		else if (NumHits == 0)
 		{
 			// Not using anything from this scene, so unlock it.
 			SceneLocks.UnlockRead(SyncScene, PST_Sync);
