PR #1004: Fix for Root Motion Not Respecting World Rotation (Contributed by [REDACTED_PERSON])

diff --git a/Engine/Source/Runtime/Engine/Private/Components/SkeletalMeshComponent.cpp b/Engine/Source/Runtime/Engine/Private/Components/SkeletalMeshComponent.cpp
--- a/Engine/Source/Runtime/Engine/Private/Components/SkeletalMeshComponent.cpp
+++ b/Engine/Source/Runtime/Engine/Private/Components/SkeletalMeshComponent.cpp
@@ -1820,9 +1820,19 @@ FTransform USkeletalMeshComponent::ConvertLocalRootMotionToWorld(const FTransfor
 		UpdateComponentToWorld();
 	}
 
+	/*check(ComponentToWorld.IsRotationNormalized());
+	check(InTransform.IsRotationNormalized());
 	const FTransform NewWorldTransform = InTransform * ComponentToWorld;
+	check(NewWorldTransform.IsRotationNormalized());
 	const FVector DeltaWorldTranslation = NewWorldTransform.GetTranslation() - ComponentToWorld.GetTranslation();
 	const FQuat DeltaWorldRotation = ComponentToWorld.GetRotation().Inverse() * NewWorldTransform.GetRotation();
+	check(DeltaWorldRotation.IsNormalized());*/
+
+	const FTransform NewWorldTransform = InTransform * ComponentToWorld;
+	const FQuat NewWorldRotation = ComponentToWorld.GetRotation() * InTransform.GetRotation();
+	const FVector DeltaWorldTranslation = NewWorldTransform.GetTranslation() - ComponentToWorld.GetTranslation();
+	const FQuat DeltaWorldRotation = NewWorldRotation * ComponentToWorld.GetRotation().Inverse();
+
 
 	const FTransform DeltaWorldTransform(DeltaWorldRotation, DeltaWorldTranslation);
 
