Fixed NavLinkProxy not rebuilding navmesh when link direction gets changed #UE4

Also hidden the SmartLink category to reduce confusion

diff --git a/Engine/Source/Runtime/Engine/Classes/AI/Navigation/NavLinkProxy.h b/Engine/Source/Runtime/Engine/Classes/AI/Navigation/NavLinkProxy.h
--- a/Engine/Source/Runtime/Engine/Classes/AI/Navigation/NavLinkProxy.h
+++ b/Engine/Source/Runtime/Engine/Classes/AI/Navigation/NavLinkProxy.h
@@ -10,10 +10,13 @@
 class UNavLinkDefinition;
 class UPathFollowingComponent;
 struct FNavigationRelevantData;
+class UNavLinkCustomComponent;
+class UNavLinkRenderingComponent;
+class UBillboardComponent;
 
-DECLARE_DYNAMIC_MULTICAST_DELEGATE_TwoParams( FSmartLinkReachedSignature, class AActor*, MovingActor, const FVector&, DestinationPoint );
+DECLARE_DYNAMIC_MULTICAST_DELEGATE_TwoParams( FSmartLinkReachedSignature, AActor*, MovingActor, const FVector&, DestinationPoint );
 
-UCLASS(Blueprintable)
+UCLASS(Blueprintable, autoCollapseCategories=(SmartLink, Actor), hideCategories=(Input))
 class ENGINE_API ANavLinkProxy : public AActor, public INavLinkHostInterface, public INavRelevantInterface
 {
 	GENERATED_UCLASS_BODY()
@@ -31,7 +34,7 @@ private_subobject:
 	/** Smart link: can affect path following */
 	DEPRECATED_FORGAME(4.6, "SmartLinkComp should not be accessed directly, please use GetSmartLinkComp() function instead. SmartLinkComp will soon be private and your code will not compile.")
 	UPROPERTY(VisibleAnywhere, Category=SmartLink)
-	class UNavLinkCustomComponent* SmartLinkComp;
+	UNavLinkCustomComponent* SmartLinkComp;
 public:
 
 	/** Smart link: toggle relevancy */
@@ -43,11 +46,11 @@ private_subobject:
 	/** Editor Preview */
 	DEPRECATED_FORGAME(4.6, "EdRenderComp should not be accessed directly, please use GetEdRenderComp() function instead. EdRenderComp will soon be private and your code will not compile.")
 	UPROPERTY()
-	class UNavLinkRenderingComponent* EdRenderComp;
+	UNavLinkRenderingComponent* EdRenderComp;
 
 	DEPRECATED_FORGAME(4.6, "SpriteComponent should not be accessed directly, please use GetSpriteComponent() function instead. SpriteComponent will soon be private and your code will not compile.")
 	UPROPERTY()
-	class UBillboardComponent* SpriteComponent;
+	UBillboardComponent* SpriteComponent;
 public:
 #endif // WITH_EDITORONLY_DATA
 
@@ -100,11 +103,11 @@ protected:
 
 public:
 	/** Returns SmartLinkComp subobject **/
-	class UNavLinkCustomComponent* GetSmartLinkComp() const;
+	UNavLinkCustomComponent* GetSmartLinkComp() const;
 #if WITH_EDITORONLY_DATA
 	/** Returns EdRenderComp subobject **/
-	class UNavLinkRenderingComponent* GetEdRenderComp() const;
+	UNavLinkRenderingComponent* GetEdRenderComp() const;
 	/** Returns SpriteComponent subobject **/
-	class UBillboardComponent* GetSpriteComponent() const;
+	UBillboardComponent* GetSpriteComponent() const;
 #endif
 };
diff --git a/Engine/Source/Runtime/Engine/Private/AI/Navigation/NavLinkProxy.cpp b/Engine/Source/Runtime/Engine/Private/AI/Navigation/NavLinkProxy.cpp
--- a/Engine/Source/Runtime/Engine/Private/AI/Navigation/NavLinkProxy.cpp
+++ b/Engine/Source/Runtime/Engine/Private/AI/Navigation/NavLinkProxy.cpp
@@ -63,7 +63,8 @@ void ANavLinkProxy::PostEditChangeProperty(FPropertyChangedEvent& PropertyChange
 	}
 
 	const FName CategoryName = FObjectEditorUtils::GetCategoryFName(PropertyChangedEvent.Property);
-	if (CategoryName == TEXT("SimpleLink"))
+	const FName MemberCategoryName = FObjectEditorUtils::GetCategoryFName(PropertyChangedEvent.MemberProperty);
+	if (CategoryName == TEXT("SimpleLink") || MemberCategoryName == TEXT("SimpleLink"))
 	{
 		UNavigationSystem* NavSys = UNavigationSystem::GetCurrent(GetWorld());
 		if (NavSys)
