Added a little feature to control whether navmesh polygons are debug-drawn filled #UE4

diff --git a/Engine/Source/Runtime/Engine/Classes/AI/Navigation/RecastNavMesh.h b/Engine/Source/Runtime/Engine/Classes/AI/Navigation/RecastNavMesh.h
--- a/Engine/Source/Runtime/Engine/Classes/AI/Navigation/RecastNavMesh.h
+++ b/Engine/Source/Runtime/Engine/Classes/AI/Navigation/RecastNavMesh.h
@@ -425,6 +425,10 @@ class ENGINE_API ARecastNavMesh : public ANavigationData
 	UPROPERTY(EditAnywhere, Category=Display, config)
 	uint32 bDrawPolyEdges:1;
 
+	/** if disabled skips filling drawn navmesh polygons */
+	UPROPERTY(EditAnywhere, Category = Display)
+	uint32 bDrawFilledPolys:1;
+
 	/** should we draw border-edges */
 	UPROPERTY(EditAnywhere, Category=Display)
 	uint32 bDrawNavMeshEdges:1;
diff --git a/Engine/Source/Runtime/Engine/Private/AI/Navigation/RecastNavMesh.cpp b/Engine/Source/Runtime/Engine/Private/AI/Navigation/RecastNavMesh.cpp
--- a/Engine/Source/Runtime/Engine/Private/AI/Navigation/RecastNavMesh.cpp
+++ b/Engine/Source/Runtime/Engine/Private/AI/Navigation/RecastNavMesh.cpp
@@ -210,6 +210,7 @@ ARecastNavMesh::FNavPolyFlags ARecastNavMesh::NavLinkFlag = ARecastNavMesh::FNav
 
 ARecastNavMesh::ARecastNavMesh(const FObjectInitializer& ObjectInitializer)
 	: Super(ObjectInitializer)
+	, bDrawFilledPolys(true)
 	, bDrawNavMeshEdges(true)
 	, bDrawNavLinks(true)
 	, bDistinctlyDrawTilesBeingBuilt(true)
diff --git a/Engine/Source/Runtime/Engine/Private/Components/NavMeshRenderingComponent.cpp b/Engine/Source/Runtime/Engine/Private/Components/NavMeshRenderingComponent.cpp
--- a/Engine/Source/Runtime/Engine/Private/Components/NavMeshRenderingComponent.cpp
+++ b/Engine/Source/Runtime/Engine/Private/Components/NavMeshRenderingComponent.cpp
@@ -407,54 +407,57 @@ void UNavMeshRenderingComponent::GatherData(FNavMeshSceneProxyData& CurrentData)
 		NavMesh->FinishBatchQuery();
 
 		// Draw Mesh
-		if (NavMesh->bDrawClusters)
+		if (NavMesh->bDrawFilledPolys)
 		{
-			for (int32 Idx = 0; Idx < CurrentData.NavMeshGeometry.Clusters.Num(); ++Idx)
+			if (NavMesh->bDrawClusters)
 			{
-				const TArray<int32>& MeshIndices = CurrentData.NavMeshGeometry.Clusters[Idx].MeshIndices;
-
-				if (MeshIndices.Num() == 0)
+				for (int32 Idx = 0; Idx < CurrentData.NavMeshGeometry.Clusters.Num(); ++Idx)
 				{
-					continue;
-				}
+					const TArray<int32>& MeshIndices = CurrentData.NavMeshGeometry.Clusters[Idx].MeshIndices;
 
-				FNavMeshSceneProxyData::FDebugMeshData DebugMeshData;
-				DebugMeshData.ClusterColor = GetClusterColor(Idx);
-				for (int32 VertIdx=0; VertIdx < MeshVerts.Num(); ++VertIdx)
-				{
-					AddVertexHelper(DebugMeshData, MeshVerts[VertIdx] + CurrentData.NavMeshDrawOffset, DebugMeshData.ClusterColor);
-				}
-				for (int32 TriIdx=0; TriIdx < MeshIndices.Num(); TriIdx+=3)
-				{
-					AddTriangleHelper(DebugMeshData, MeshIndices[TriIdx], MeshIndices[TriIdx+1], MeshIndices[TriIdx+2]);
-				}
+					if (MeshIndices.Num() == 0)
+					{
+						continue;
+					}
+
+					FNavMeshSceneProxyData::FDebugMeshData DebugMeshData;
+					DebugMeshData.ClusterColor = GetClusterColor(Idx);
+					for (int32 VertIdx = 0; VertIdx < MeshVerts.Num(); ++VertIdx)
+					{
+						AddVertexHelper(DebugMeshData, MeshVerts[VertIdx] + CurrentData.NavMeshDrawOffset, DebugMeshData.ClusterColor);
+					}
+					for (int32 TriIdx = 0; TriIdx < MeshIndices.Num(); TriIdx += 3)
+					{
+						AddTriangleHelper(DebugMeshData, MeshIndices[TriIdx], MeshIndices[TriIdx + 1], MeshIndices[TriIdx + 2]);
+					}
 
-				CurrentData.MeshBuilders.Add(DebugMeshData);
+					CurrentData.MeshBuilders.Add(DebugMeshData);
+				}
 			}
-		}
-		else if (NavMesh->bDrawNavMesh)
-		{			
-			for (int32 AreaType = 0; AreaType < RECAST_MAX_AREAS; ++AreaType)
+			else if (NavMesh->bDrawNavMesh)
 			{
-				const TArray<int32>& MeshIndices = CurrentData.NavMeshGeometry.AreaIndices[AreaType];
-
-				if (MeshIndices.Num() == 0)
+				for (int32 AreaType = 0; AreaType < RECAST_MAX_AREAS; ++AreaType)
 				{
-					continue;
-				}
+					const TArray<int32>& MeshIndices = CurrentData.NavMeshGeometry.AreaIndices[AreaType];
 
-				FNavMeshSceneProxyData::FDebugMeshData DebugMeshData;
-				for (int32 VertIdx=0; VertIdx < MeshVerts.Num(); ++VertIdx)
-				{
-					AddVertexHelper(DebugMeshData, MeshVerts[VertIdx] + CurrentData.NavMeshDrawOffset, CurrentData.NavMeshColors[AreaType]);
-				}
-				for (int32 TriIdx=0; TriIdx < MeshIndices.Num(); TriIdx+=3)
-				{
-					AddTriangleHelper(DebugMeshData, MeshIndices[TriIdx], MeshIndices[TriIdx+1], MeshIndices[TriIdx+2]);
-				}
+					if (MeshIndices.Num() == 0)
+					{
+						continue;
+					}
+
+					FNavMeshSceneProxyData::FDebugMeshData DebugMeshData;
+					for (int32 VertIdx = 0; VertIdx < MeshVerts.Num(); ++VertIdx)
+					{
+						AddVertexHelper(DebugMeshData, MeshVerts[VertIdx] + CurrentData.NavMeshDrawOffset, CurrentData.NavMeshColors[AreaType]);
+					}
+					for (int32 TriIdx = 0; TriIdx < MeshIndices.Num(); TriIdx += 3)
+					{
+						AddTriangleHelper(DebugMeshData, MeshIndices[TriIdx], MeshIndices[TriIdx + 1], MeshIndices[TriIdx + 2]);
+					}
 
-				DebugMeshData.ClusterColor = CurrentData.NavMeshColors[AreaType];
-				CurrentData.MeshBuilders.Add(DebugMeshData);
+					DebugMeshData.ClusterColor = CurrentData.NavMeshColors[AreaType];
+					CurrentData.MeshBuilders.Add(DebugMeshData);
+				}
 			}
 		}
 
