Fixed a depreaction warning in SkeletalMesh

diff --git a/Engine/Source/Runtime/Engine/Private/SkeletalMesh.cpp b/Engine/Source/Runtime/Engine/Private/SkeletalMesh.cpp
--- a/Engine/Source/Runtime/Engine/Private/SkeletalMesh.cpp
+++ b/Engine/Source/Runtime/Engine/Private/SkeletalMesh.cpp
@@ -3723,7 +3723,7 @@ void ASkeletalMeshActor::PreviewSetAnimPosition(FName SlotName, int32 ChannelInd
 					}
 				}
 
-				AnimInst->PlaySlotAnimation(InAnimSequence, SlotName, 0.0f, 0.0f, 0.f, 1);
+				AnimInst->PlaySlotAnimationAsDynamicMontage(InAnimSequence, SlotName, 0.0f, 0.0f, 0.f, 1);
 				CurrentlyPlayingMontage = AnimInst->GetCurrentActiveMontage();
 			}
 
@@ -3935,7 +3935,7 @@ void SetAnimPositionInner(FName SlotName, USkeletalMeshComponent* SkeletalMeshCo
 					}
 				}
 
-				AnimInst->PlaySlotAnimation(InAnimSequence, SlotName, 0.0f, 0.0f, 0.f, 1);
+				AnimInst->PlaySlotAnimationAsDynamicMontage(InAnimSequence, SlotName, 0.0f, 0.0f, 0.f, 1);
 				CurrentlyPlayingMontage = AnimInst->GetCurrentActiveMontage();
 			}
 
