Fixed comments on linkable element types so the editor shows tooltips correctly.

diff --git a/Engine/Source/Runtime/Engine/Classes/Animation/AnimLinkableElement.h b/Engine/Source/Runtime/Engine/Classes/Animation/AnimLinkableElement.h
--- a/Engine/Source/Runtime/Engine/Classes/Animation/AnimLinkableElement.h
+++ b/Engine/Source/Runtime/Engine/Classes/Animation/AnimLinkableElement.h
@@ -14,9 +14,12 @@ namespace EAnimLinkMethod
 {
 	enum Type
 	{
-		Absolute,		// Element stays at a specific time without moving
-		Relative,		// Element moves with its segment, but not when the segment changes size
-		Proportional	// Element moves with its segment and will stay at a certain proportion through the segment
+		// Element stays at a specific time without moving
+		Absolute,
+		// Element moves with its segment, but not when the segment changes size
+		Relative,
+		// Element moves with its segment and will stay at a certain proportion through the segment
+		Proportional
 	};
 }
 
