[UE-14230] Safe guarding UPhysicsCollisionHandler::GetWorld() against calls on the class's CDO.

diff --git a/Engine/Source/Runtime/Engine/Classes/PhysicsEngine/PhysicsCollisionHandler.h b/Engine/Source/Runtime/Engine/Classes/PhysicsEngine/PhysicsCollisionHandler.h
--- a/Engine/Source/Runtime/Engine/Classes/PhysicsEngine/PhysicsCollisionHandler.h
+++ b/Engine/Source/Runtime/Engine/Classes/PhysicsEngine/PhysicsCollisionHandler.h
@@ -31,7 +31,7 @@ class ENGINE_API UPhysicsCollisionHandler : public UObject
 	/** Get the world we are handling collisions for */
 	virtual UWorld* GetWorld() const override
 	{
-		return CastChecked<UWorld>(GetOuter());
+		return !HasAnyFlags(RF_ClassDefaultObject) ? CastChecked<UWorld>(GetOuter()) : nullptr;
 	}
 
 	/** Gives game-specific ability to handle all physics collisions in one place. This is a good place to play sounds and spawn effects, as it does not require special object-specific code. */
