[UE-14230] Extending FMemberReference so it can properly represent global delegate signatures (and have a UPackage for the member's parent).

diff --git a/Engine/Config/BaseEngine.ini b/Engine/Config/BaseEngine.ini
--- a/Engine/Config/BaseEngine.ini
+++ b/Engine/Config/BaseEngine.ini
@@ -571,6 +571,9 @@ bUseStreamingPause=false
 +TaggedPropertyRedirects=(ClassName="FSM_AIProjectileThrow",OldPropertyName="MinThrowRange",NewPropertyName="MinAttackRange")
 +TaggedPropertyRedirects=(ClassName="FSM_AIProjectileThrow",OldPropertyName="NumThrowsToReachTargetAccuracy",NewPropertyName="NumAttacksToReachTargetAccuracy")
 
++TaggedPropertyRedirects=(ClassName="MemberReference",OldPropertyName="MemberParentClass",NewPropertyName="MemberParent")
++TaggedPropertyRedirects=(ClassName="SimpleMemberReference",OldPropertyName="MemberParentClass",NewPropertyName="MemberParent")
+
 ;Vim rename
 +TaggedPropertyRedirects=(ClassName="AnimInstance",OldPropertyName="VimVertexAnims",NewPropertyName="VertexAnims")
 
@@ -792,6 +795,9 @@ bUseStreamingPause=false
 +K2FieldRedirects=(OldFieldName="Actor.AddActorLocalTranslation",NewFieldName="Actor.K2_AddActorLocalOffset")
 +K2ParamRedirects=(NodeName="/Script/BlueprintGraph.K2Node_CallFunction", OldParamName="Actor.K2_AddActorLocalOffset.DeltaTranslation", NewParamName="DeltaLocation")
 
++TaggedPropertyRedirects=(ClassName="MemberReference",OldPropertyName="MemberParentClass",NewPropertyName="MemberParent")
++TaggedPropertyRedirects=(ClassName="SimpleMemberReference",OldPropertyName="MemberParentClass",NewPropertyName="MemberParent")
+
 ;
 ; Add FHitResult to swept movement, and use wrappers for K2 methods
 ;
diff --git a/Engine/Source/Editor/BlueprintGraph/Private/BlueprintActionDatabase.cpp b/Engine/Source/Editor/BlueprintGraph/Private/BlueprintActionDatabase.cpp
--- a/Engine/Source/Editor/BlueprintGraph/Private/BlueprintActionDatabase.cpp
+++ b/Engine/Source/Editor/BlueprintGraph/Private/BlueprintActionDatabase.cpp
@@ -164,7 +164,7 @@ static UBlueprintNodeSpawner* FBlueprintNodeSpawnerFactory::MakeMessageNodeSpawn
 	auto SetNodeFunctionLambda = [](UEdGraphNode* NewNode, UField const* FuncField)
 	{
 		UK2Node_Message* MessageNode = CastChecked<UK2Node_Message>(NewNode);
-		MessageNode->FunctionReference.SetExternalMember(FuncField->GetFName(), FuncField->GetOwnerClass());
+		MessageNode->FunctionReference.SetFromField<UFunction>(FuncField, /*bIsConsideredSelfContext =*/false);
 	};
 	NodeSpawner->SetNodeFieldDelegate = UBlueprintFunctionNodeSpawner::FSetNodeFieldDelegate::CreateStatic(SetNodeFunctionLambda);
 
diff --git a/Engine/Source/Editor/BlueprintGraph/Private/EdGraphSchema_K2.cpp b/Engine/Source/Editor/BlueprintGraph/Private/EdGraphSchema_K2.cpp
--- a/Engine/Source/Editor/BlueprintGraph/Private/EdGraphSchema_K2.cpp
+++ b/Engine/Source/Editor/BlueprintGraph/Private/EdGraphSchema_K2.cpp
@@ -3449,10 +3449,10 @@ bool UEdGraphSchema_K2::ArePinTypesCompatible(const FEdGraphPinType& Output, con
 			{
 				OutFunction = NULL;
 			}
-			if (!OutFunction && Output.PinSubCategoryMemberReference.MemberParentClass)
+			if (!OutFunction && Output.PinSubCategoryMemberReference.GetMemberParentClass())
 			{
-				const auto ParentClass = Output.PinSubCategoryMemberReference.MemberParentClass;
-				const auto BPOwner = Cast<UBlueprint>(ParentClass->ClassGeneratedBy);
+				const UClass* ParentClass = Output.PinSubCategoryMemberReference.GetMemberParentClass();
+				const UBlueprint* BPOwner = Cast<UBlueprint>(ParentClass->ClassGeneratedBy);
 				if (BPOwner && BPOwner->SkeletonGeneratedClass && (BPOwner->SkeletonGeneratedClass != ParentClass))
 				{
 					OutFunction = BPOwner->SkeletonGeneratedClass->FindFunctionByName(Output.PinSubCategoryMemberReference.MemberName);
@@ -3463,10 +3463,10 @@ bool UEdGraphSchema_K2::ArePinTypesCompatible(const FEdGraphPinType& Output, con
 			{
 				InFunction = NULL;
 			}
-			if (!InFunction && Input.PinSubCategoryMemberReference.MemberParentClass)
+			if (!InFunction && Input.PinSubCategoryMemberReference.GetMemberParentClass())
 			{
-				const auto ParentClass = Input.PinSubCategoryMemberReference.MemberParentClass;
-				const auto BPOwner = Cast<UBlueprint>(ParentClass->ClassGeneratedBy);
+				const UClass* ParentClass = Input.PinSubCategoryMemberReference.GetMemberParentClass();
+				const UBlueprint* BPOwner = Cast<UBlueprint>(ParentClass->ClassGeneratedBy);
 				if (BPOwner && BPOwner->SkeletonGeneratedClass && (BPOwner->SkeletonGeneratedClass != ParentClass))
 				{
 					InFunction = BPOwner->SkeletonGeneratedClass->FindFunctionByName(Input.PinSubCategoryMemberReference.MemberName);
diff --git a/Engine/Source/Editor/BlueprintGraph/Private/K2Node_ComponentBoundEvent.cpp b/Engine/Source/Editor/BlueprintGraph/Private/K2Node_ComponentBoundEvent.cpp
--- a/Engine/Source/Editor/BlueprintGraph/Private/K2Node_ComponentBoundEvent.cpp
+++ b/Engine/Source/Editor/BlueprintGraph/Private/K2Node_ComponentBoundEvent.cpp
@@ -40,9 +40,7 @@ void UK2Node_ComponentBoundEvent::InitializeComponentBoundEventParams(UObjectPro
 		DelegatePropertyName = InDelegateProperty->GetFName();
 		DelegateOwnerClass = CastChecked<UClass>(InDelegateProperty->GetOuter())->GetAuthoritativeClass();
 
-		FName ReferenceName = InDelegateProperty->SignatureFunction->GetFName();
-		UClass* ReferenceClass = InDelegateProperty->SignatureFunction->GetOwnerClass();
-		EventReference.SetExternalMember(ReferenceName, ReferenceClass);
+		EventReference.SetFromField<UFunction>(InDelegateProperty->SignatureFunction, /*bIsConsideredSelfContext =*/false);
 
 		CustomFunctionName = FName( *FString::Printf(TEXT("BndEvt__%s_%s_%s"), *InComponentProperty->GetName(), *GetName(), *EventReference.GetMemberName().ToString()) );
 		bOverrideFunction = false;
diff --git a/Engine/Source/Editor/UnrealEd/Private/Kismet2/BlueprintEditorUtils.cpp b/Engine/Source/Editor/UnrealEd/Private/Kismet2/BlueprintEditorUtils.cpp
--- a/Engine/Source/Editor/UnrealEd/Private/Kismet2/BlueprintEditorUtils.cpp
+++ b/Engine/Source/Editor/UnrealEd/Private/Kismet2/BlueprintEditorUtils.cpp
@@ -991,7 +991,7 @@ struct FRegenerationHelper
 				}
 			}
 
-			if (UClass* TypeClass = NewVar.VarType.PinSubCategoryMemberReference.MemberParentClass)
+			if (UClass* TypeClass = NewVar.VarType.PinSubCategoryMemberReference.GetMemberParentClass())
 			{
 				auto Linker = TypeClass->GetLinker();
 				if (Linker && TypeClass->HasAnyFlags(RF_NeedLoad))
diff --git a/Engine/Source/Runtime/Engine/Classes/EdGraph/EdGraphPin.h b/Engine/Source/Runtime/Engine/Classes/EdGraph/EdGraphPin.h
--- a/Engine/Source/Runtime/Engine/Classes/EdGraph/EdGraphPin.h
+++ b/Engine/Source/Runtime/Engine/Classes/EdGraph/EdGraphPin.h
@@ -10,9 +10,12 @@ struct FSimpleMemberReference
 {
 	GENERATED_USTRUCT_BODY()
 
-	/** Class that this member is defined in. */
+	/** 
+	 * Most often the Class that this member is defined in. Could be a UPackage 
+	 * if it is a native delegate signature function (declared globally).
+	 */
 	UPROPERTY()
-	TSubclassOf<class UObject> MemberParentClass;
+	UObject* MemberParent;
 
 	/** Name of the member */
 	UPROPERTY()
@@ -22,6 +25,8 @@ struct FSimpleMemberReference
 	UPROPERTY()
 	FGuid MemberGuid;
 
+	FSimpleMemberReference() : MemberParent(nullptr) {}
+
 	void Reset()
 	{
 		operator=(FSimpleMemberReference());
@@ -29,15 +34,21 @@ struct FSimpleMemberReference
 
 	bool operator==(const FSimpleMemberReference& Other) const
 	{
-		return (MemberParentClass == Other.MemberParentClass)
+		return (MemberParent == Other.MemberParent)
 			&& (MemberName == Other.MemberName)
 			&& (MemberGuid == Other.MemberGuid);
 	}
+
+	/* For backwards compatibility (when MemberParent used to exclusively be a class) */
+	UClass* GetMemberParentClass() const
+	{
+		return Cast<UClass>(MemberParent);
+	}
 };
 
 FORCEINLINE FArchive& operator<<(FArchive& Ar, FSimpleMemberReference& Data)
 {
-	Ar << Data.MemberParentClass;
+	Ar << Data.MemberParent;
 	Ar << Data.MemberName;
 	Ar << Data.MemberGuid;
 	return Ar;
diff --git a/Engine/Source/Runtime/Engine/Classes/Engine/MemberReference.h b/Engine/Source/Runtime/Engine/Classes/Engine/MemberReference.h
--- a/Engine/Source/Runtime/Engine/Classes/Engine/MemberReference.h
+++ b/Engine/Source/Runtime/Engine/Classes/Engine/MemberReference.h
@@ -50,17 +50,23 @@ struct FParamRemapInfo
 	}
 };
 
+// @TODO: this can encapsulate globally defined fields as well (like with native 
+//        delegate signatures); consider renaming to FFieldReference
 USTRUCT()
 struct FMemberReference
 {
 	GENERATED_USTRUCT_BODY()
 
 protected:
-	/** Class that this member is defined in. Should be NULL if bSelfContext is true.  */
+	/** 
+	 * Most often the Class that this member is defined in. Could be a UPackage 
+	 * if it is a native delegate signature function (declared globally). Should 
+	 * be NULL if bSelfContext is true.  
+	 */
 	UPROPERTY()
-	mutable TSubclassOf<class UObject> MemberParentClass;
+	mutable UObject* MemberParent;
 
-	/** Class that this member is defined in. Should be NULL if bSelfContext is true.  */
+	/**  */
 	UPROPERTY()
 	mutable FString MemberScope;
 
@@ -82,7 +88,7 @@ protected:
 	
 public:
 	FMemberReference()
-		: MemberParentClass(NULL)
+		: MemberParent(NULL)
 		, MemberName(NAME_None)
 		, bSelfContext(false)
 	{
@@ -92,21 +98,32 @@ public:
 	template<class TFieldType>
 	void SetFromField(const UField* InField, const bool bIsConsideredSelfContext)
 	{
-		MemberParentClass = bIsConsideredSelfContext ? NULL : InField->GetOwnerClass();
+		UClass* OwnerClass = InField->GetOwnerClass();
+		MemberParent = OwnerClass;
+
+		if (bIsConsideredSelfContext)
+		{
+			MemberParent = nullptr;
+		}
+		else if ((MemberParent == nullptr) && InField->GetName().EndsWith(HEADER_GENERATED_DELEGATE_SIGNATURE_SUFFIX))
+		{
+			MemberParent = InField->GetOutermost();
+		}
+
 		MemberName = InField->GetFName();
 		bSelfContext = bIsConsideredSelfContext;
 		bWasDeprecated = false;
 
 #if WITH_EDITOR
-		if (MemberParentClass != nullptr)
+		if (UClass* ParentAsClass = GetMemberParentClass())
 		{
-			MemberParentClass = MemberParentClass->GetAuthoritativeClass();
+			MemberParent = ParentAsClass->GetAuthoritativeClass();
 		}
 
 		MemberGuid.Invalidate();
-		if (InField->GetOwnerClass())
+		if (OwnerClass != nullptr)
 		{
-			UBlueprint::GetGuidFromClassByFieldName<TFieldType>(InField->GetOwnerClass(), InField->GetFName(), MemberGuid);
+			UBlueprint::GetGuidFromClassByFieldName<TFieldType>(OwnerClass, InField->GetFName(), MemberGuid);
 		}
 #endif
 	}
@@ -114,27 +131,30 @@ public:
 	template<class TFieldType>
 	void SetFromField(const UField* InField, UClass* SelfScope)
 	{
+		UClass* OwnerClass = InField->GetOwnerClass();
+
 		FGuid FieldGuid;
 #if WITH_EDITOR
-		if (InField->GetOwnerClass())
+		if (OwnerClass != nullptr)
 		{
-			UBlueprint::GetGuidFromClassByFieldName<TFieldType>(InField->GetOwnerClass(), InField->GetFName(), FieldGuid);
+			UBlueprint::GetGuidFromClassByFieldName<TFieldType>(OwnerClass, InField->GetFName(), FieldGuid);
 		}
 #endif
 
-		SetGivenSelfScope(InField->GetFName(), FieldGuid, InField->GetOwnerClass(), SelfScope);
+		SetGivenSelfScope(InField->GetFName(), FieldGuid, OwnerClass, SelfScope);
 	}
 
 	/** Update given a new self */
 	template<class TFieldType>
 	void RefreshGivenNewSelfScope(UClass* SelfScope)
 	{
-		if ((MemberParentClass != NULL) && (SelfScope != NULL))
+		UClass* ParentAsClass = GetMemberParentClass();
+		if ((ParentAsClass != NULL) && (SelfScope != NULL))
 		{
 #if WITH_EDITOR
-			UBlueprint::GetGuidFromClassByFieldName<TFieldType>((MemberParentClass ? *MemberParentClass : SelfScope), MemberName, MemberGuid);
+			UBlueprint::GetGuidFromClassByFieldName<TFieldType>((ParentAsClass ? ParentAsClass : SelfScope), MemberName, MemberGuid);
 #endif
-			SetGivenSelfScope(MemberName, MemberGuid, MemberParentClass, SelfScope);
+			SetGivenSelfScope(MemberName, MemberGuid, ParentAsClass, SelfScope);
 		}
 		else
 		{
@@ -145,6 +165,9 @@ public:
 	/** Set to a non-'self' member, so must include reference to class owning the member. */
 	ENGINE_API void SetExternalMember(FName InMemberName, TSubclassOf<class UObject> InMemberParentClass);
 
+	/** Set to reference a global field (intended for things like natively defined delegate signatures) */
+	ENGINE_API void SetGlobalField(FName InFieldName, UPackage* InParentPackage);
+
 	/** Set to a non-'self' delegate member, this is not self-context but is not given a parent class */
 	ENGINE_API void SetExternalDelegateMember(FName InMemberName);
 
@@ -160,7 +183,7 @@ public:
 	/** Only intended for backwards compat! */
 	ENGINE_API void SetDirect(const FName InMemberName, const FGuid InMemberGuid, TSubclassOf<class UObject> InMemberParentClass, bool bIsConsideredSelfContext);
 
-	/** Invalidate the current MemberParentClass, if this is a self scope, or the MemberScope if it is not (and set).  Intended for PostDuplication fixups */
+	/** Invalidate the current MemberParent, if this is a self scope, or the MemberScope if it is not (and set).  Intended for PostDuplication fixups */
 	ENGINE_API void InvalidateScope();
 
 	/** Get the name of this member */
@@ -174,9 +197,22 @@ public:
 		return MemberGuid;
 	}
 
-	UClass* GetMemberParentClass()
+	UClass* GetMemberParentClass() const
 	{
-		return *MemberParentClass;
+		return Cast<UClass>(MemberParent);
+	}
+
+	UPackage* GetMemberParentPackage() const
+	{
+		if (UPackage* ParentAsPackage = Cast<UPackage>(MemberParent))
+		{
+			return ParentAsPackage;
+		}
+		else if (MemberParent != nullptr)
+		{
+			return MemberParent->GetOutermost();
+		}
+		return nullptr;
 	}
 
 	/** Returns if this is a 'self' context. */
@@ -210,7 +246,7 @@ public:
 	UClass* GetMemberParentClass(UClass* SelfScope) const
 	{
 		// Local variables with a MemberScope act much the same as being SelfContext, their parent class is SelfScope.
-		return (bSelfContext || !MemberScope.IsEmpty())? SelfScope : *MemberParentClass;
+		return (bSelfContext || !MemberScope.IsEmpty())? SelfScope : GetMemberParentClass();
 	}
 
 	/** Get the scope of this member */
@@ -269,7 +305,7 @@ public:
 		else
 		{
 			// Look for remapped member
-			UClass* TargetScope = bSelfContext ? SelfScope : (UClass*)MemberParentClass;
+			UClass* TargetScope = bSelfContext ? SelfScope : GetMemberParentClass();
 #if WITH_EDITOR
 			if( TargetScope != NULL &&  !GIsSavingPackage )
 			{
@@ -280,24 +316,26 @@ public:
 			{
 				// Fix up this struct, we found a redirect
 				MemberName = ReturnField->GetFName();
-				MemberParentClass = Cast<UClass>(ReturnField->GetOuter());
+				MemberParent = Cast<UClass>(ReturnField->GetOuter());
 
 				MemberGuid.Invalidate();
 				UBlueprint::GetGuidFromClassByFieldName<TFieldType>(TargetScope, MemberName, MemberGuid);
 
-				if (MemberParentClass != nullptr)
+				if (UClass* ParentAsClass = GetMemberParentClass())
 				{
-					MemberParentClass = MemberParentClass->GetAuthoritativeClass();
+					ParentAsClass = ParentAsClass->GetAuthoritativeClass();
+					MemberParent  = ParentAsClass;
+
 					// Re-evaluate self-ness against the redirect if we were given a valid SelfScope
 					if (SelfScope != NULL)
 					{
-						SetGivenSelfScope(MemberName, MemberGuid, MemberParentClass, SelfScope);
+						SetGivenSelfScope(MemberName, MemberGuid, ParentAsClass, SelfScope);
 					}
 				}	
 			}
 			else
 #endif
-				if(TargetScope != NULL)
+			if(TargetScope != NULL)
 			{
 				// Find in target scope
 				ReturnField = FindField<TFieldType>(TargetScope, MemberName);
@@ -316,25 +354,34 @@ public:
 				}
 #endif
 			}
-			else
+			else if (UPackage* TargetPackage = GetMemberParentPackage())
 			{
-				// Native delegate signatures can have null target scope (if they're defined outside of a class).
-                // Needed because Blueprint pin types are referencing delegate signatures as
-                // class members, which isn't necessarily true (globally defined signatures).
-				if (TFieldType::StaticClass() == UFunction::StaticClass())
+				ReturnField = FindObject<TFieldType>(TargetPackage, *MemberName.ToString());
+			}
+			// For backwards compatibility: as of CL 2412156, delegate signatures 
+			// could have had a null MemberParentClass (for those natively 
+			// declared outside of a class), we used to rely on the following 
+			// FindObject<>; however this was not reliable (hence the addition 
+			// of GetMemberParentPackage(), etc.)
+			else if ((TFieldType::StaticClass() == UFunction::StaticClass()) && MemberName.ToString().EndsWith(HEADER_GENERATED_DELEGATE_SIGNATURE_SUFFIX))
+			{			
+				FString const StringName = MemberName.ToString();
+				for (TObjectIterator<UPackage> PackageIt; PackageIt && (ReturnField == nullptr); ++PackageIt)
 				{
-					TFieldType* GlobalFunction = (TFieldType*)FindDelegateSignature(MemberName);
-					// we only want native delegate signatures, even more
-					if (GlobalFunction->GetOwnerClass() == nullptr)
+					if ((PackageIt->PackageFlags & PKG_CompiledIn) == 0x00)
 					{
-						// @TODO: Could this still be the wrong function? I imagine there is a way 
-						//        that we can define like named delegates without the compiler yelling.
-						ReturnField = GlobalFunction;
+						continue;
 					}
-					// @TODO: The delegate we want to find could exist, but FindDelegateSignature()
-					//        could have returned the first one it found (there could be a naming 
-					//        conflict since FindDelegateSignature() searches ANY_PACKAGE)
-					//        How do we make sure it only finds functions with UPackage outers?
+
+					// NOTE: this could return the wrong field (if there are 
+					//       two like-named delegates defined in separate packages)
+					ReturnField = FindObject<TFieldType>(*PackageIt, *StringName);
+				}
+
+				if (ReturnField != nullptr)
+				{
+					UE_LOG(LogBlueprint, Warning, TEXT("Generic delegate signature ref (%s). Explicitly setting it to: '%s'. Make sure this is correct (there could be multiple native delegate types with this name)?"), *StringName, *ReturnField->GetPathName());
+					MemberParent = ReturnField->GetOutermost();
 				}
 			}
 		}
@@ -407,7 +454,7 @@ public:
 			TempMemberReference.SetFromField<TFieldType>(InField, false);
 
 			OutReference.MemberName = TempMemberReference.MemberName;
-			OutReference.MemberParentClass = TempMemberReference.MemberParentClass;
+			OutReference.MemberParent = TempMemberReference.MemberParent;
 			OutReference.MemberGuid = TempMemberReference.MemberGuid;
 		}
 	}
@@ -416,13 +463,16 @@ public:
 	static TFieldType* ResolveSimpleMemberReference(const FSimpleMemberReference& Reference)
 	{
 		FMemberReference TempMemberReference;
+
 		const FName Name = Reference.MemberGuid.IsValid() ? NAME_None : Reference.MemberName; // if the guid is valid don't check the name, it could be renamed 
-		TempMemberReference.SetDirect(Name, Reference.MemberGuid, Reference.MemberParentClass, false);
-		auto Result = TempMemberReference.ResolveMember<TFieldType>((UClass*)nullptr);
+		TempMemberReference.MemberName   = Name;
+		TempMemberReference.MemberGuid   = Reference.MemberGuid;
+		TempMemberReference.MemberParent = Reference.MemberParent;
 
+		auto Result = TempMemberReference.ResolveMember<TFieldType>((UClass*)nullptr);
 		if (!Result && (Name != Reference.MemberName))
 		{
-			TempMemberReference.SetDirect(Reference.MemberName, Reference.MemberGuid, Reference.MemberParentClass, false);
+			TempMemberReference.MemberName = Reference.MemberName;
 			Result = TempMemberReference.ResolveMember<TFieldType>((UClass*)nullptr);
 		}
 
diff --git a/Engine/Source/Runtime/Engine/Private/EdGraph/EdGraphPin.cpp b/Engine/Source/Runtime/Engine/Private/EdGraph/EdGraphPin.cpp
--- a/Engine/Source/Runtime/Engine/Private/EdGraph/EdGraphPin.cpp
+++ b/Engine/Source/Runtime/Engine/Private/EdGraph/EdGraphPin.cpp
@@ -53,7 +53,7 @@ bool FEdGraphPinType::Serialize(FArchive& Ar)
 			if (const UFunction* Signature = Cast<const UFunction>(PinSubCategoryObject.Get()))
 			{
 				PinSubCategoryMemberReference.MemberName = Signature->GetFName();
-				PinSubCategoryMemberReference.MemberParentClass = Signature->GetOwnerClass();
+				PinSubCategoryMemberReference.MemberParent = Signature->GetOwnerClass();
 				PinSubCategoryObject = NULL;
 			}
 			else
diff --git a/Engine/Source/Runtime/Engine/Private/MemberReference.cpp b/Engine/Source/Runtime/Engine/Private/MemberReference.cpp
--- a/Engine/Source/Runtime/Engine/Private/MemberReference.cpp
+++ b/Engine/Source/Runtime/Engine/Private/MemberReference.cpp
@@ -15,15 +15,24 @@ void FMemberReference::SetExternalMember(FName InMemberName, TSubclassOf<class U
 {
 	MemberName = InMemberName;
 #if WITH_EDITOR
-	MemberParentClass = (InMemberParentClass != nullptr) ? InMemberParentClass->GetAuthoritativeClass() : nullptr;
+	MemberParent = (InMemberParentClass != nullptr) ? InMemberParentClass->GetAuthoritativeClass() : nullptr;
 #else
-	MemberParentClass = *InMemberParentClass;
+	MemberParent = *InMemberParentClass;
 #endif
 	MemberScope.Empty();
 	bSelfContext = false;
 	bWasDeprecated = false;
 }
 
+void FMemberReference::SetGlobalField(FName InFieldName, UPackage* InParentPackage)
+{
+	MemberName = InFieldName;
+	MemberParent = InParentPackage;
+	MemberScope.Empty();
+	bSelfContext = false;
+	bWasDeprecated = false;
+}
+
 void FMemberReference::SetExternalDelegateMember(FName InMemberName)
 {
 	SetExternalMember(InMemberName, nullptr);
@@ -32,7 +41,7 @@ void FMemberReference::SetExternalDelegateMember(FName InMemberName)
 void FMemberReference::SetSelfMember(FName InMemberName)
 {
 	MemberName = InMemberName;
-	MemberParentClass = NULL;
+	MemberParent = NULL;
 	MemberScope.Empty();
 	bSelfContext = true;
 	bWasDeprecated = false;
@@ -44,7 +53,7 @@ void FMemberReference::SetDirect(const FName InMemberName, const FGuid InMemberG
 	MemberGuid = InMemberGuid;
 	bSelfContext = bIsConsideredSelfContext;
 	bWasDeprecated = false;
-	MemberParentClass = InMemberParentClass;
+	MemberParent = InMemberParentClass;
 	MemberScope.Empty();
 }
 
@@ -52,7 +61,7 @@ void FMemberReference::SetGivenSelfScope(const FName InMemberName, const FGuid I
 {
 	MemberName = InMemberName;
 	MemberGuid = InMemberGuid;
-	MemberParentClass = (InMemberParentClass != nullptr) ? InMemberParentClass->GetAuthoritativeClass() : nullptr;
+	MemberParent = (InMemberParentClass != nullptr) ? InMemberParentClass->GetAuthoritativeClass() : nullptr;
 	MemberScope.Empty();
 
 	// SelfScope should always be valid, but if it's not ensure and move on, the node will be treated as if it's not self scoped.
@@ -62,7 +71,7 @@ void FMemberReference::SetGivenSelfScope(const FName InMemberName, const FGuid I
 
 	if (bSelfContext)
 	{
-		MemberParentClass = NULL;
+		MemberParent = NULL;
 	}
 }
 
@@ -83,7 +92,7 @@ void FMemberReference::InvalidateScope()
 {
 	if( IsSelfContext() )
 	{
-		MemberParentClass = NULL;
+		MemberParent = NULL;
 	}
 	else if(IsLocalScope())
 	{
