Avoid compressing textures when autosaving

diff --git a/Engine/Source/Editor/UnrealEd/Classes/Editor/UnrealEdEngine.h b/Engine/Source/Editor/UnrealEd/Classes/Editor/UnrealEdEngine.h
--- a/Engine/Source/Editor/UnrealEd/Classes/Editor/UnrealEdEngine.h
+++ b/Engine/Source/Editor/UnrealEd/Classes/Editor/UnrealEdEngine.h
@@ -201,6 +201,7 @@ public:
 	virtual bool ShouldAbortActorDeletion() const override;
 	virtual void CloseEditor() override;
 	virtual void OnOpenMatinee() override;
+	virtual bool IsAutosaving() const override;
 	// End UEditorEngine Interface 
 	
 	// Begin FExec Interface
diff --git a/Engine/Source/Editor/UnrealEd/Private/UnrealEdEngine.cpp b/Engine/Source/Editor/UnrealEd/Private/UnrealEdEngine.cpp
--- a/Engine/Source/Editor/UnrealEd/Private/UnrealEdEngine.cpp
+++ b/Engine/Source/Editor/UnrealEd/Private/UnrealEdEngine.cpp
@@ -655,6 +655,11 @@ void UUnrealEdEngine::OnOpenMatinee()
 	OnMatineeEditorClosedDelegateHandle = GLevelEditorModeTools().OnEditorModeChanged().AddUObject( this, &UUnrealEdEngine::OnMatineeEditorClosed );
 }
 
+bool UUnrealEdEngine::IsAutosaving() const
+{
+	return PackageAutoSaver->IsAutoSaving();
+}
+
 
 void UUnrealEdEngine::ConvertMatinees()
 {
diff --git a/Engine/Source/Runtime/Engine/Classes/Engine/Engine.h b/Engine/Source/Runtime/Engine/Classes/Engine/Engine.h
--- a/Engine/Source/Runtime/Engine/Classes/Engine/Engine.h
+++ b/Engine/Source/Runtime/Engine/Classes/Engine/Engine.h
@@ -2524,7 +2524,9 @@ public:
 	/** @return true if editor analytics are enabled */
 	virtual bool AreEditorAnalyticsEnabled() const { return false; }
 	virtual void CreateStartupAnalyticsAttributes( TArray<struct FAnalyticsEventAttribute>& StartSessionAttributes ) const {}
-
+	
+	/** @return true if the engine is autosaving a package */
+	virtual bool IsAutosaving() const { return false; }
 protected:
 
 	TIndirectArray<FWorldContext>	WorldList;
diff --git a/Engine/Source/Runtime/Engine/Private/Texture.cpp b/Engine/Source/Runtime/Engine/Private/Texture.cpp
--- a/Engine/Source/Runtime/Engine/Private/Texture.cpp
+++ b/Engine/Source/Runtime/Engine/Private/Texture.cpp
@@ -366,8 +366,11 @@ void UTexture::PreSave()
 		UpdateResource();
 	}
 
-	GWarn->StatusUpdate( 0, 0, FText::Format( NSLOCTEXT("UnrealEd", "SavingPackage_CompressingSourceArt", "Compressing source art for texture:  {0}"), FText::FromString(GetName()) ) );
-	Source.Compress();
+	if( !GEngine->IsAutosaving() )
+	{
+		GWarn->StatusUpdate(0, 0, FText::Format(NSLOCTEXT("UnrealEd", "SavingPackage_CompressingSourceArt", "Compressing source art for texture:  {0}"), FText::FromString(GetName())));
+		Source.Compress();
+	}
 
 #endif // #if WITH_EDITOR
 }
