#ue4 - Add ECollisionEnabled::PhysicsOnly option to round out all options of NoCollision, QueryOnly, PhysicsOnly, QueryAndPhysics.

- Can be used for things like capes and hair that users may not want to be considered by scene queries.

diff --git a/Engine/Source/Editor/DetailCustomizations/Private/CollisionProfileDetails.cpp b/Engine/Source/Editor/DetailCustomizations/Private/CollisionProfileDetails.cpp
--- a/Engine/Source/Editor/DetailCustomizations/Private/CollisionProfileDetails.cpp
+++ b/Engine/Source/Editor/DetailCustomizations/Private/CollisionProfileDetails.cpp
@@ -623,8 +623,9 @@ void SProfileEditDialog::FillCollisionEnabledString()
 {
 	CollisionEnabledComboBoxString.Empty();
 	CollisionEnabledComboBoxString.Add(MakeShareable(new FString(TEXT("No Collision"))));
-	CollisionEnabledComboBoxString.Add(MakeShareable(new FString(TEXT("No Physics Collision"))));
-	CollisionEnabledComboBoxString.Add(MakeShareable(new FString(TEXT("Collision Enabled"))));
+	CollisionEnabledComboBoxString.Add(MakeShareable(new FString(TEXT("Query Only (No Physics Collision)"))));
+	CollisionEnabledComboBoxString.Add(MakeShareable(new FString(TEXT("Physics Only (No Query Collision)"))));
+	CollisionEnabledComboBoxString.Add(MakeShareable(new FString(TEXT("Collision Enabled (Query and Physics)"))));
 }
 
 bool	SProfileEditDialog::IsAcceptAvailable() const
@@ -740,7 +741,7 @@ void SProfileEditDialog::HandleCollisionEnabledComboBoxSelectionChanged(TSharedP
 		if(*Iter == StringItem)
 		{
 			ECollisionEnabled::Type NewCollisionEnabled = (ECollisionEnabled::Type)Iter.GetIndex();
-			check(NewCollisionEnabled >= ECollisionEnabled::NoCollision&& NewCollisionEnabled <= ECollisionEnabled::QueryAndPhysics);
+			check(NewCollisionEnabled >= ECollisionEnabled::NoCollision && NewCollisionEnabled <= ECollisionEnabled::QueryAndPhysics);
 			ProfileTemplate.CollisionEnabled = NewCollisionEnabled;
 			return;
 		}
@@ -1229,9 +1230,11 @@ FText SProfileListItem::GetCollsionEnabled() const
 	case ECollisionEnabled::NoCollision:
 		return LOCTEXT("ECollisionEnabled_NoCollision", "No Collision");
 	case ECollisionEnabled::QueryOnly:
-		return LOCTEXT("ECollisionEnabled_QueryOnly", "No Physics Collision");
+		return LOCTEXT("ECollisionEnabled_QueryOnly", "Query Only (No Physics Collision)");
+	case ECollisionEnabled::PhysicsOnly:
+		return LOCTEXT("ECollisionEnabled_PhysicsOnly", "Physics Only (No Query Collision)");
 	case ECollisionEnabled::QueryAndPhysics:
-		return LOCTEXT("ECollisionEnabled_QueryAndPhysics", "Collision Enabled");
+		return LOCTEXT("ECollisionEnabled_QueryAndPhysics", "Collision Enabled (Query and Physics)");
 	}
 
 	return LOCTEXT("ECollisionEnabled_Error", "ERROR");
diff --git a/Engine/Source/Runtime/Engine/Classes/Engine/EngineTypes.h b/Engine/Source/Runtime/Engine/Classes/Engine/EngineTypes.h
--- a/Engine/Source/Runtime/Engine/Classes/Engine/EngineTypes.h
+++ b/Engine/Source/Runtime/Engine/Classes/Engine/EngineTypes.h
@@ -820,12 +820,14 @@ namespace ECollisionEnabled
 { 
 	enum Type 
 	{ 
-		/** No collision is performed against this body */ 
+		/** No collision is enabled for this body */ 
 		NoCollision UMETA(DisplayName="No Collision"), 
-		/** This body is only used for collision raycasts, sweeps and overlaps */ 
-		QueryOnly UMETA(DisplayName="No Physics Collision"), 
-		/** This body is used for physics simulation and collision queries */ 
-		QueryAndPhysics UMETA(DisplayName="Collision Enabled") 
+		/** This body is used only for collision queries (raycasts, sweeps, and overlaps) */ 
+		QueryOnly UMETA(DisplayName="Query Only (No Physics Collision)"),
+		/** This body is used only for physics collision */
+		PhysicsOnly UMETA(DisplayName="Physics Only (No Query Collision)"),
+		/** This body interacts with all collision (Query and Physics) */ 
+		QueryAndPhysics UMETA(DisplayName="Collision Enabled (Query and Physics)") 
 	}; 
 } 
 
diff --git a/Engine/Source/Runtime/Engine/Classes/PhysicsEngine/BodyInstance.h b/Engine/Source/Runtime/Engine/Classes/PhysicsEngine/BodyInstance.h
--- a/Engine/Source/Runtime/Engine/Classes/PhysicsEngine/BodyInstance.h
+++ b/Engine/Source/Runtime/Engine/Classes/PhysicsEngine/BodyInstance.h
@@ -163,11 +163,13 @@ private:
 	UPROPERTY(EditAnywhere, Category=Custom)
 	FName CollisionProfileName;
 
-	/** Type of Collision Enabled 
+	/**
+	 * Type of collision enabled.
 	 * 
-	 *	No Collision				: No collision is performed against this neither trace or physics
-	 *	No Physics Collision		: This body is only used for collision raycasts, sweeps and overlaps 
-	 *	Collision Enabled			: This body is used for physics simulation and collision queries
+	 *	No Collision      : No collision is performed against this body.
+	 *	Query Only        : This body is used only for collision queries (raycasts, sweeps, and overlaps).
+	 *	Physics Only      : This body is used only for physics collision.
+	 *	Collision Enabled : This body interacts with all collision (Query and Physics).
 	 */
 	UPROPERTY(EditAnywhere, Category=Custom)
 	TEnumAsByte<ECollisionEnabled::Type> CollisionEnabled;
@@ -483,7 +485,7 @@ public:
 	}
 
 	/** Populate the filter data within the provided FShapeData with the correct filters for this instance
-	 *	@param ShapeData ShapeData to populate
+	 *	@param ShapeData ShapeData to populate. ShapeData.CollisionEnabled and ShapeData.FilterData will be filled in.
 	 *	@param bForceSimpleAsComplex Whether to force simple colision as complex
 	 *  Note: This function is not thread safe. Make sure to obtain the appropriate PhysX scene locks before calling this function
 	 */
@@ -522,7 +524,7 @@ public:
 	}
 
 	/** Populate the flag fields of the provided FShapeData with correct initialisation flags
-	 *	@param ShapeData ShapeData to populate
+	 *	@param ShapeData ShapeData to populate. ShapeData.FilterData will not be modified.
 	 *	@param UseCollisionEnabled Whether collision is enabled for this instance
 	 *	@param bUseComplexAsSimple Whether to use complex collision as simple
 	 *  Note: This function is not thread safe. Make sure to obtain the appropriate PhysX scene locks before calling this function
@@ -1005,7 +1007,7 @@ private:
 	/** 
 	 * Helper function to update per shape filtering info. This should interface is not very friendly and should only be used from inside FBodyInstance
 	 */
-	void UpdatePhysicsShapeFilterData(uint32 SkelMeshCompID, bool bUseComplexAsSimple, bool bUseSimpleAsComplex, bool bPhysicsStatic, TEnumAsByte<ECollisionEnabled::Type> * CollisionEnabledOverride, FCollisionResponseContainer * ResponseOverride, bool * bNotifyOverride);
+	void UpdatePhysicsShapeFilterData(uint32 SkelMeshCompID, bool bUseComplexAsSimple, bool bUseSimpleAsComplex, bool bPhysicsStatic, const TEnumAsByte<ECollisionEnabled::Type> * CollisionEnabledOverride, FCollisionResponseContainer * ResponseOverride, bool * bNotifyOverride);
 #endif 
 	/**
 	 * Invalidate Collision Profile Name
diff --git a/Engine/Source/Runtime/Engine/Private/PhysicsEngine/BodyInstance.cpp b/Engine/Source/Runtime/Engine/Private/PhysicsEngine/BodyInstance.cpp
--- a/Engine/Source/Runtime/Engine/Private/PhysicsEngine/BodyInstance.cpp
+++ b/Engine/Source/Runtime/Engine/Private/PhysicsEngine/BodyInstance.cpp
@@ -685,7 +685,7 @@ void ExecuteOnPxShapeWrite(FBodyInstance* BodyInstance, PxShape* PShape, Lambda
 }
 
 
-void FBodyInstance::UpdatePhysicsShapeFilterData(uint32 SkelMeshCompID, bool bUseComplexAsSimple, bool bUseSimpleAsComplex, bool bPhysicsStatic, TEnumAsByte<ECollisionEnabled::Type> * CollisionEnabledOverride, FCollisionResponseContainer * ResponseOverride, bool * bNotifyOverride)
+void FBodyInstance::UpdatePhysicsShapeFilterData(uint32 SkelMeshCompID, bool bUseComplexAsSimple, bool bUseSimpleAsComplex, bool bPhysicsStatic, const TEnumAsByte<ECollisionEnabled::Type> * CollisionEnabledOverride, FCollisionResponseContainer * ResponseOverride, bool * bNotifyOverride)
 {
 	ExecuteOnPhysicsReadWrite([&]
 	{
@@ -693,8 +693,8 @@ void FBodyInstance::UpdatePhysicsShapeFilterData(uint32 SkelMeshCompID, bool bUs
 		{
 
 		// Iterate over all shapes and assign filterdata
-			TArray<PxShape*> AllShapes;
-			const int32 NumSyncShapes = GetAllShapes_AssumesLocked(AllShapes);
+		TArray<PxShape*> AllShapes;
+		const int32 NumSyncShapes = GetAllShapes_AssumesLocked(AllShapes);
 
 		bool bUpdateMassProperties = false;
 		for (int32 ShapeIdx = 0; ShapeIdx < AllShapes.Num(); ShapeIdx++)
@@ -703,7 +703,7 @@ void FBodyInstance::UpdatePhysicsShapeFilterData(uint32 SkelMeshCompID, bool bUs
 			FBodyInstance* BI = FPhysxUserData::Get<FBodyInstance>(PShape->userData);
 			BI = BI ? BI : this;
 
-			const TEnumAsByte<ECollisionEnabled::Type> & UseCollisionEnabled = CollisionEnabledOverride ? *CollisionEnabledOverride : (TEnumAsByte<ECollisionEnabled::Type>)BI->GetCollisionEnabled();
+			const TEnumAsByte<ECollisionEnabled::Type> UseCollisionEnabled = CollisionEnabledOverride ? *CollisionEnabledOverride : (TEnumAsByte<ECollisionEnabled::Type>)BI->GetCollisionEnabled();
 			const FCollisionResponseContainer& UseResponse = ResponseOverride ? *ResponseOverride : BI->CollisionResponses.GetResponseContainer();
 			bool bUseNotify = bNotifyOverride ? *bNotifyOverride : BI->bNotifyRigidBodyCollision;
 
@@ -737,22 +737,24 @@ void FBodyInstance::UpdatePhysicsShapeFilterData(uint32 SkelMeshCompID, bool bUs
 			PShape->setSimulationFilterData(PSimFilterData);
 
 			// If query collision is enabled..
-			if ((UseCollisionEnabled == ECollisionEnabled::QueryAndPhysics) || (UseCollisionEnabled == ECollisionEnabled::QueryOnly))
+			if (UseCollisionEnabled != ECollisionEnabled::NoCollision)
 			{
+				const bool bQueryEnabled = ((UseCollisionEnabled == ECollisionEnabled::QueryAndPhysics) || (UseCollisionEnabled == ECollisionEnabled::QueryOnly));
+
 				// Only perform scene queries in the synchronous scene for static shapes
 				if (bPhysicsStatic)
 				{
 					bool bIsSyncShape = (ShapeIdx < NumSyncShapes);
-					PShape->setFlag(PxShapeFlag::eSCENE_QUERY_SHAPE, bIsSyncShape);
+					PShape->setFlag(PxShapeFlag::eSCENE_QUERY_SHAPE, bIsSyncShape && bQueryEnabled);
 				}
-				// If non-static, always enable scene queries
+				// If non-static, enable scene queries if requested
 				else
 				{
-					PShape->setFlag(PxShapeFlag::eSCENE_QUERY_SHAPE, true);
+					PShape->setFlag(PxShapeFlag::eSCENE_QUERY_SHAPE, bQueryEnabled);
 				}
 
 				// See if we want physics collision
-				bool bSimCollision = (UseCollisionEnabled == ECollisionEnabled::QueryAndPhysics);
+				const bool bSimCollision = ((UseCollisionEnabled == ECollisionEnabled::QueryAndPhysics) || (UseCollisionEnabled == ECollisionEnabled::PhysicsOnly));
 
 				// Triangle mesh is 'complex' geom
 				if (PShape->getGeometryType() == PxGeometryType::eTRIANGLEMESH)
@@ -780,7 +782,7 @@ void FBodyInstance::UpdatePhysicsShapeFilterData(uint32 SkelMeshCompID, bool bUs
 					PShape->setQueryFilterData(PSimpleQueryData);
 
 					// See if we currently have sim collision
-					bool bCurrentSimCollision = (PShape->getFlags() & PxShapeFlag::eSIMULATION_SHAPE);
+					const bool bCurrentSimCollision = (PShape->getFlags() & PxShapeFlag::eSIMULATION_SHAPE);
 					// Enable sim collision
 					if (bSimCollision && !bCurrentSimCollision)
 					{
@@ -795,7 +797,7 @@ void FBodyInstance::UpdatePhysicsShapeFilterData(uint32 SkelMeshCompID, bool bUs
 					}
 
 					// enable swept bounds for CCD for this shape
-						PxRigidBody* PBody = PActor->is<PxRigidBody>();
+					PxRigidBody* PBody = PActor->is<PxRigidBody>();
 					if (bSimCollision && !bPhysicsStatic && bUseCCD && PBody)
 					{
 						PBody->setRigidBodyFlag(PxRigidBodyFlag::eENABLE_CCD, true);
@@ -917,7 +919,7 @@ void FBodyInstance::UpdatePhysicsFilterData(bool bForceSimpleAsComplex)
 	const bool bUseSimpleAsComplex = bForceSimpleAsComplex || (BodySetup.Get()->CollisionTraceFlag == CTF_UseSimpleAsComplex);
 
 #if WITH_PHYSX
-	TEnumAsByte<ECollisionEnabled::Type> * CollisionEnabledOverride = bUseCollisionEnabledOverride ? &UseCollisionEnabled : NULL;
+	const TEnumAsByte<ECollisionEnabled::Type>* CollisionEnabledOverride = bUseCollisionEnabledOverride ? &UseCollisionEnabled : NULL;
 	FCollisionResponseContainer * ResponseOverride = bResponseOverride ? &UseResponse : NULL;
 	bool * bNotifyOverridePtr = bNotifyOverride ? &bUseNotifyRBCollision : NULL;
 	UpdatePhysicsShapeFilterData(SkelMeshCompID, bUseComplexAsSimple, bUseSimpleAsComplex, bPhysicsStatic, CollisionEnabledOverride, ResponseOverride, bNotifyOverridePtr);
@@ -940,7 +942,7 @@ void FBodyInstance::UpdatePhysicsFilterData(bool bForceSimpleAsComplex)
 			BoxSimFilterData.BodyIndex = InstanceBodyIndex;
 
 			// Update the body data
-			const bool bSimCollision = (UseCollisionEnabled == ECollisionEnabled::QueryAndPhysics);
+			const bool bSimCollision = ((UseCollisionEnabled == ECollisionEnabled::QueryAndPhysics) || (UseCollisionEnabled == ECollisionEnabled::PhysicsOnly));
 			BodyInstancePtr->SetBullet(bSimCollision && !bPhysicsStatic && bUseCCD);
 			BodyInstancePtr->SetActive(true);
 
@@ -1167,7 +1169,7 @@ struct FInitBodiesHelper
 
 		FShapeData ShapeData;
 		Instance->GetFilterData_AssumesLocked(ShapeData);
-		Instance->GetShapeFlags_AssumesLocked(ShapeData, Instance->CollisionEnabled, BodySetup->CollisionTraceFlag == CTF_UseComplexAsSimple);
+		Instance->GetShapeFlags_AssumesLocked(ShapeData, ShapeData.CollisionEnabled, BodySetup->CollisionTraceFlag == CTF_UseComplexAsSimple);
 
 		if (!bCompileStatic && PNewDynamic)
 		{
@@ -1847,8 +1849,8 @@ bool FBodyInstance::Weld(FBodyInstance* TheirBody, const FTransform& TheirTM)
 	PxMaterial* PSimpleMat = SimplePhysMat->GetPhysXMaterial();
 
 	FShapeData ShapeData;
-		GetFilterData_AssumesLocked(ShapeData);
-		GetShapeFlags_AssumesLocked(ShapeData, CollisionEnabled, BodySetup->CollisionTraceFlag == CTF_UseComplexAsSimple);
+	GetFilterData_AssumesLocked(ShapeData);
+	GetShapeFlags_AssumesLocked(ShapeData, ShapeData.CollisionEnabled, BodySetup->CollisionTraceFlag == CTF_UseComplexAsSimple);
 
 	//child body gets placed into the same scenes as parent body
 	if (PxRigidActor* MyBody = RigidActorSync)
@@ -4632,6 +4634,8 @@ void FBodyInstance::SetShapeFlags_AssumesLocked(TEnumAsByte<ECollisionEnabled::T
 		bool bUpdateMassProperties = false;
 		if (UseCollisionEnabled != ECollisionEnabled::NoCollision)
 		{
+			const bool bQueryEnabled = ((UseCollisionEnabled == ECollisionEnabled::QueryAndPhysics) || (UseCollisionEnabled == ECollisionEnabled::QueryOnly));
+
 			UPrimitiveComponent* OwnerComponentInst = OwnerComponent.Get();
 			AActor* Owner = OwnerComponentInst ? OwnerComponentInst->GetOwner() : NULL;
 			const bool bPhysicsStatic = !OwnerComponentInst || OwnerComponentInst->IsWorldGeometry();
@@ -4639,16 +4643,16 @@ void FBodyInstance::SetShapeFlags_AssumesLocked(TEnumAsByte<ECollisionEnabled::T
 			// Only perform scene queries in the synchronous scene for static shapes
 			if (bPhysicsStatic)
 			{
-				PShape->setFlag(PxShapeFlag::eSCENE_QUERY_SHAPE, SceneType == PST_Sync);
+				PShape->setFlag(PxShapeFlag::eSCENE_QUERY_SHAPE, (SceneType == PST_Sync) && bQueryEnabled);
 			}
-			// If non-static, always enable scene queries
+			// If non-static, enable scene queries if requested
 			else
 			{
-				PShape->setFlag(PxShapeFlag::eSCENE_QUERY_SHAPE, true);
+				PShape->setFlag(PxShapeFlag::eSCENE_QUERY_SHAPE, bQueryEnabled);
 			}
 
 			// See if we want physics collision
-			bool bSimCollision = (UseCollisionEnabled == ECollisionEnabled::QueryAndPhysics);
+			const bool bSimCollision = ((UseCollisionEnabled == ECollisionEnabled::QueryAndPhysics) || (UseCollisionEnabled == ECollisionEnabled::PhysicsOnly));
 
 			// Triangle mesh is 'complex' geom
 			if (PShape->getGeometryType() == PxGeometryType::eTRIANGLEMESH)
@@ -4715,7 +4719,7 @@ void FBodyInstance::SetShapeFlags_AssumesLocked(TEnumAsByte<ECollisionEnabled::T
 
 void FBodyInstance::GetShapeFlags_AssumesLocked(FShapeData& ShapeData, TEnumAsByte<ECollisionEnabled::Type> UseCollisionEnabled, const bool bUseComplexAsSimple /*= false*/)
 {
-	// If query collision is enabled..
+	ShapeData.CollisionEnabled = UseCollisionEnabled;
 	ShapeData.SyncShapeFlags = PxShapeFlags(0);
 	ShapeData.AsyncShapeFlags = PxShapeFlags(0);
 	ShapeData.SimpleShapeFlags = PxShapeFlags(0);
@@ -4723,16 +4727,19 @@ void FBodyInstance::GetShapeFlags_AssumesLocked(FShapeData& ShapeData, TEnumAsBy
 	ShapeData.SyncBodyFlags = PxRigidBodyFlags(0);
 
 	// Default flags
-	ShapeData.SyncShapeFlags |= PxShapeFlag::eSCENE_QUERY_SHAPE;
-	ShapeData.AsyncShapeFlags |= PxShapeFlag::eSCENE_QUERY_SHAPE;
 	ShapeData.SimpleShapeFlags |= PxShapeFlag::eVISUALIZATION;
 	ShapeData.ComplexShapeFlags |= PxShapeFlag::eVISUALIZATION;
 
-	UPrimitiveComponent* OwnerComponentInst = OwnerComponent.Get();
-	AActor* Owner = OwnerComponentInst ? OwnerComponentInst->GetOwner() : NULL;
-
-	if(ShapeData.CollisionEnabled != ECollisionEnabled::NoCollision)
+	if(UseCollisionEnabled != ECollisionEnabled::NoCollision)
 	{
+		const bool bQueryEnabled = ((UseCollisionEnabled == ECollisionEnabled::QueryAndPhysics) || (UseCollisionEnabled == ECollisionEnabled::QueryOnly));
+		if (bQueryEnabled)
+		{
+			ShapeData.SyncShapeFlags |= PxShapeFlag::eSCENE_QUERY_SHAPE;
+			ShapeData.AsyncShapeFlags |= PxShapeFlag::eSCENE_QUERY_SHAPE;
+		}
+
+		const UPrimitiveComponent* OwnerComponentInst = OwnerComponent.Get();
 		const bool bPhysicsStatic = !OwnerComponentInst || OwnerComponentInst->IsWorldGeometry();
 
 		// Only perform scene queries in the synchronous scene for static shapes
@@ -4742,7 +4749,7 @@ void FBodyInstance::GetShapeFlags_AssumesLocked(FShapeData& ShapeData, TEnumAsBy
 		}
 
 		// See if we want physics collision
-		bool bSimCollision = (UseCollisionEnabled == ECollisionEnabled::QueryAndPhysics);
+		const bool bSimCollision = ((UseCollisionEnabled == ECollisionEnabled::QueryAndPhysics) || (UseCollisionEnabled == ECollisionEnabled::PhysicsOnly));
 
 		// on dynamic objects and objects which don't use complex as simple, tri mesh not used for sim
 		if(bSimCollision && bUseComplexAsSimple)
diff --git a/Engine/Source/Runtime/Engine/Private/PhysicsEngine/PhysXSupport.h b/Engine/Source/Runtime/Engine/Private/PhysicsEngine/PhysXSupport.h
--- a/Engine/Source/Runtime/Engine/Private/PhysicsEngine/PhysXSupport.h
+++ b/Engine/Source/Runtime/Engine/Private/PhysicsEngine/PhysXSupport.h
@@ -843,7 +843,8 @@ struct FShapeFilterData
 struct FShapeData
 {
 	FShapeData()
-		: SyncShapeFlags(0)
+		: CollisionEnabled(ECollisionEnabled::NoCollision)
+		, SyncShapeFlags(0)
 		, AsyncShapeFlags(0)
 		, SimpleShapeFlags(0)
 		, ComplexShapeFlags(0)
