#ue4 - Remove unused TArray in World::ComponentSweepMulti().

diff --git a/Engine/Source/Runtime/Engine/Private/Collision/WorldCollision.cpp b/Engine/Source/Runtime/Engine/Private/Collision/WorldCollision.cpp
--- a/Engine/Source/Runtime/Engine/Private/Collision/WorldCollision.cpp
+++ b/Engine/Source/Runtime/Engine/Private/Collision/WorldCollision.cpp
@@ -525,8 +525,6 @@ bool UWorld::ComponentSweepMulti(TArray<struct FHitResult>& OutHits, class UPrim
 
 			if (PGeom != NULL)
 			{
-				TArray<struct FHitResult> Hits;
-
 				// Calc shape global pose
 				const PxTransform PLocalShape = PShape->getLocalPose();
 				const PxTransform PShapeGlobalStartPose = PGlobalStartPose.transform(PLocalShape);
