Added a function for GameViewportClient to modify the views after they've been set up

Merging 2527220

//depot/UE4-UT/...

to //depot/UE4/...

diff --git a/Engine/Source/Runtime/Engine/Classes/Engine/GameViewportClient.h b/Engine/Source/Runtime/Engine/Classes/Engine/GameViewportClient.h
--- a/Engine/Source/Runtime/Engine/Classes/Engine/GameViewportClient.h
+++ b/Engine/Source/Runtime/Engine/Classes/Engine/GameViewportClient.h
@@ -159,6 +159,10 @@ public:
 	virtual void SetIsSimulateInEditorViewport(bool bInIsSimulateInEditorViewport) override;
 	// End of FViewportClient interface.
 
+	/** make any adjustments to the views after they've been completely set up */
+	virtual void FinalizeViews(class FSceneViewFamily* ViewFamily, const TMap<ULocalPlayer*, FSceneView*>& PlayerViewMap)
+	{}
+
 	// Begin FExec interface.
 	virtual bool Exec( UWorld* InWorld, const TCHAR* Cmd,FOutputDevice& Ar) override;
 	// End of FExec interface.
diff --git a/Engine/Source/Runtime/Engine/Private/GameViewportClient.cpp b/Engine/Source/Runtime/Engine/Private/GameViewportClient.cpp
--- a/Engine/Source/Runtime/Engine/Private/GameViewportClient.cpp
+++ b/Engine/Source/Runtime/Engine/Private/GameViewportClient.cpp
@@ -996,6 +996,8 @@ void UGameViewportClient::Draw(FViewport* InViewport, FCanvas* SceneCanvas)
 		}
 	}
 
+	FinalizeViews(&ViewFamily, PlayerViewMap);
+
 	if (bReverbSettingsFound)
 	{
 		AudioDevice->SetReverbSettings( AudioVolume, ReverbSettings );
