Destructible crash when passed null mesh.

diff --git a/Engine/Source/Runtime/Engine/Private/Components/DestructibleComponent.cpp b/Engine/Source/Runtime/Engine/Private/Components/DestructibleComponent.cpp
--- a/Engine/Source/Runtime/Engine/Private/Components/DestructibleComponent.cpp
+++ b/Engine/Source/Runtime/Engine/Private/Components/DestructibleComponent.cpp
@@ -734,7 +734,7 @@ void UDestructibleComponent::RefreshBoneTransforms(FActorComponentTickFunction*
 void UDestructibleComponent::SetDestructibleMesh(class UDestructibleMesh* NewMesh)
 {
 #if WITH_APEX
-	uint32 ChunkCount = NewMesh->ApexDestructibleAsset->getChunkCount();
+	uint32 ChunkCount = NewMesh ? NewMesh->ApexDestructibleAsset->getChunkCount() : 0;
 	ChunkInfos.Reset(ChunkCount);
 	ChunkInfos.AddZeroed(ChunkCount);
 	PhysxChunkUserData.Reset(ChunkCount);
@@ -745,7 +745,7 @@ void UDestructibleComponent::SetDestructibleMesh(class UDestructibleMesh* NewMes
 
 #if WITH_EDITORONLY_DATA
 	// If the SkeletalMesh has changed, update our transient value too.
-	this->DestructibleMesh = GetDestructibleMesh();
+	DestructibleMesh = GetDestructibleMesh();
 #endif // WITH_EDITORONLY_DATA
 	
 	RecreatePhysicsState();
