#ue4 - Fix wrong var in ensure in ConvertOverlappedShapeToImpactHit().

diff --git a/Engine/Source/Runtime/Engine/Private/Collision/CollisionConversions.cpp b/Engine/Source/Runtime/Engine/Private/Collision/CollisionConversions.cpp
--- a/Engine/Source/Runtime/Engine/Private/Collision/CollisionConversions.cpp
+++ b/Engine/Source/Runtime/Engine/Private/Collision/CollisionConversions.cpp
@@ -913,7 +913,7 @@ static bool ConvertOverlappedShapeToImpactHit(const UWorld* World, const PxLocat
 		if (!bValidNormal)
 		{
 			OutResult.ImpactNormal = (StartLoc - EndLoc).GetSafeNormal();
-			ensure(OutResult.Normal.IsNormalized());
+			ensure(OutResult.ImpactNormal.IsNormalized());
 		}
 	}
 
