Add a TimerManager to GameInstance and use that instead of the one on World when the World is owned GameInstance.  This allows timers on objects that exist across level transitions to work.
Create a dummy world when GameInstance is created in standalone.  This world is not fully set up, but it will permit certain behaviors to not fail/crash when functions are called by objects being initialized (widgets, game instance) prior to the first map being loaded.

diff --git a/Engine/Source/Runtime/Engine/Classes/Engine/GameInstance.h b/Engine/Source/Runtime/Engine/Classes/Engine/GameInstance.h
--- a/Engine/Source/Runtime/Engine/Classes/Engine/GameInstance.h
+++ b/Engine/Source/Runtime/Engine/Classes/Engine/GameInstance.h
@@ -5,6 +5,9 @@
 #include "EngineBaseTypes.h"
 #include "GameInstance.generated.h"
 
+class ULocalPlayer;
+class FUniqueNetId;
+
 // 
 // 	EWelcomeScreen, 	//initial screen.  Used for platforms where we may not have a signed in user yet.
 // 	EMessageScreen, 	//message screen.  Used to display a message - EG unable to connect to game.
@@ -56,6 +59,7 @@ public:
 
 	// Begin UObject Interface
 	virtual class UWorld* GetWorld() const override;
+	virtual void FinishDestroy() override;
 	// End UObject Interface
 
 	/** virtual function to allow custom GameInstances an opportunity to set up what it needs */
@@ -166,4 +170,8 @@ public:
 	/** Overridable implementation of HandleSessionUserInviteAccepted, which does nothing but call this function */
 	virtual void OnSessionUserInviteAccepted(const bool bWasSuccess, const int32 ControllerId, TSharedPtr< FUniqueNetId > UserId, const FOnlineSessionSearchResult & InviteResult);
 
+	inline FTimerManager& GetTimerManager() const { return *TimerManager; }
+
+private:
+	FTimerManager* TimerManager;
 };
\ No newline at end of file
diff --git a/Engine/Source/Runtime/Engine/Classes/Engine/World.h b/Engine/Source/Runtime/Engine/Classes/Engine/World.h
--- a/Engine/Source/Runtime/Engine/Classes/Engine/World.h
+++ b/Engine/Source/Runtime/Engine/Classes/Engine/World.h
@@ -12,6 +12,7 @@
 #include "Engine/LatentActionManager.h"
 #include "Runtime/RHI/Public/RHIDefinitions.h"
 #include "GameFramework/Actor.h"
+#include "GameInstance.h"
 
 #if WITH_EDITOR
 	#include "Editor/UnrealEd/Public/HierarchicalLOD.h"
@@ -3120,7 +3121,7 @@ public:
 	/** Returns TimerManager instance for this world. */
 	inline FTimerManager& GetTimerManager() const
 	{
-		return *TimerManager;
+		return (OwningGameInstance ? OwningGameInstance->GetTimerManager() : *TimerManager);
 	}
 
 	/** Returns LatentActionManager instance for this world. */
diff --git a/Engine/Source/Runtime/Engine/Private/GameInstance.cpp b/Engine/Source/Runtime/Engine/Private/GameInstance.cpp
--- a/Engine/Source/Runtime/Engine/Private/GameInstance.cpp
+++ b/Engine/Source/Runtime/Engine/Private/GameInstance.cpp
@@ -20,9 +20,21 @@
 
 UGameInstance::UGameInstance(const FObjectInitializer& ObjectInitializer)
 : Super(ObjectInitializer)
+, TimerManager(new FTimerManager())
 {
 }
 
+void UGameInstance::FinishDestroy()
+{
+	if (TimerManager)
+	{
+		delete TimerManager;
+		TimerManager = nullptr;
+	}
+
+	Super::FinishDestroy();
+}
+
 UWorld* UGameInstance::GetWorld() const
 {
 	return WorldContext ? WorldContext->World() : NULL;
@@ -74,6 +86,11 @@ void UGameInstance::InitializeStandalone()
 	WorldContext = &Engine->CreateNewWorldContext(EWorldType::Game);
 	WorldContext->OwningGameInstance = this;
 
+	// In standalone create a dummy world from the beginning to avoid issues of not having a world until LoadMap gets us our real world
+	UWorld* DummyWorld = UWorld::CreateWorld(EWorldType::Game, false);
+	DummyWorld->SetGameInstance(this);
+	WorldContext->SetCurrentWorld(DummyWorld);
+
 	Init();
 }
 
diff --git a/Engine/Source/Runtime/Engine/Private/LevelTick.cpp b/Engine/Source/Runtime/Engine/Private/LevelTick.cpp
--- a/Engine/Source/Runtime/Engine/Private/LevelTick.cpp
+++ b/Engine/Source/Runtime/Engine/Private/LevelTick.cpp
@@ -1213,8 +1213,8 @@ void UWorld::Tick( ELevelTick TickType, float DeltaSeconds )
 
 		if (TickType != LEVELTICK_TimeOnly && !bIsPaused)
 		{
-			STAT(FScopeCycleCounter Context(TimerManager->GetStatId());)
-			TimerManager->Tick(DeltaSeconds);
+			STAT(FScopeCycleCounter Context(GetTimerManager().GetStatId());)
+			GetTimerManager().Tick(DeltaSeconds);
 		}
 
 		for( int32 i=0; i<FTickableGameObject::TickableObjects.Num(); i++ )
