Make PhysicsSerializer default to Instanced since it should not be shared.

UE-UE-14514

diff --git a/Engine/Source/Runtime/Engine/Public/PhysicsSerializer.h b/Engine/Source/Runtime/Engine/Public/PhysicsSerializer.h
--- a/Engine/Source/Runtime/Engine/Public/PhysicsSerializer.h
+++ b/Engine/Source/Runtime/Engine/Public/PhysicsSerializer.h
@@ -17,7 +17,7 @@ namespace physx
 
 #include "PhysicsSerializer.generated.h"
 
-UCLASS(hidecategories=Object, MinimalAPI)
+UCLASS(DefaultToInstanced, hidecategories=Object, MinimalAPI)
 class UPhysicsSerializer : public UObject
 {
 	GENERATED_UCLASS_BODY()
