#ue4 - Ensure valid transforms in SceneComponent.

diff --git a/Engine/Source/Runtime/Engine/Private/Components/SceneComponent.cpp b/Engine/Source/Runtime/Engine/Private/Components/SceneComponent.cpp
--- a/Engine/Source/Runtime/Engine/Private/Components/SceneComponent.cpp
+++ b/Engine/Source/Runtime/Engine/Private/Components/SceneComponent.cpp
@@ -338,7 +338,10 @@ void USceneComponent::UpdateComponentToWorldWithParent(USceneComponent* Parent,
 	// Calculate the new ComponentToWorld transform
 	const FTransform RelativeTransform(RelativeRotationQuat, RelativeLocation, RelativeScale3D);
 	FTransform NewTransform = CalcNewComponentToWorld(RelativeTransform, Parent);
-	NewTransform.NormalizeRotation(); // Minimize accumulation of errors after many composed transforms.
+	
+	// Minimize accumulation of errors after many composed transforms.
+	NewTransform.NormalizeRotation();
+	ensureOnce(NewTransform.IsValid());
 
 	// If transform has changed..
 	if (!ComponentToWorld.Equals(NewTransform, SMALL_NUMBER))
