UE-8753: Gets Actor Position material node working and add explicit compilation error when using object related input nodes in the deferred decal domain.

diff --git a/Engine/Shaders/MaterialTemplate.usf b/Engine/Shaders/MaterialTemplate.usf
--- a/Engine/Shaders/MaterialTemplate.usf
+++ b/Engine/Shaders/MaterialTemplate.usf
@@ -315,6 +315,23 @@ struct FMaterialVertexParameters
 	}
 #endif //DECAL_RENDERSTAGE
 
+#if DECAL_RENDERSTAGE
+	/*
+	 * Deferred decal don't have a Primitive uniform buffer, because we don't know on which primitive the decal
+	 * is being projected to. But the user may still need to get the decal's actor world position.
+	 * So instead of setting up a primitive buffer that may cost to much CPU effort to be almost never used,
+	 * we directly fetch this value from the DeferredDecal.usf specific uniform variable DecalToWorld.
+	 */
+	float3 GetActorWorldPosition()
+	{
+		return DecalToWorld[3].xyz;
+	}
+#else
+	float3 GetActorWorldPosition()
+	{
+		return Primitive.ActorWorldPosition;
+	}
+#endif //DECAL_RENDERSTAGE
 
 #if HAS_PRIMITIVE_UNIFORM_BUFFER
 
diff --git a/Engine/Source/Runtime/Engine/Private/Materials/HLSLMaterialTranslator.h b/Engine/Source/Runtime/Engine/Private/Materials/HLSLMaterialTranslator.h
--- a/Engine/Source/Runtime/Engine/Private/Materials/HLSLMaterialTranslator.h
+++ b/Engine/Source/Runtime/Engine/Private/Materials/HLSLMaterialTranslator.h
@@ -2134,7 +2134,7 @@ protected:
 
 	virtual int32 ActorWorldPosition() override
 	{
-		return AddInlinedCodeChunk(MCT_Float3,TEXT("Primitive.ActorWorldPosition"));		
+		return AddInlinedCodeChunk(MCT_Float3,TEXT("GetActorWorldPosition()"));		
 	}
 
 	virtual int32 If(int32 A,int32 B,int32 AGreaterThanB,int32 AEqualsB,int32 ALessThanB,int32 ThresholdArg) override
diff --git a/Engine/Source/Runtime/Engine/Private/Materials/MaterialExpressions.cpp b/Engine/Source/Runtime/Engine/Private/Materials/MaterialExpressions.cpp
--- a/Engine/Source/Runtime/Engine/Private/Materials/MaterialExpressions.cpp
+++ b/Engine/Source/Runtime/Engine/Private/Materials/MaterialExpressions.cpp
@@ -5763,6 +5763,11 @@ UMaterialExpressionObjectPositionWS::UMaterialExpressionObjectPositionWS(const F
 
 int32 UMaterialExpressionObjectPositionWS::Compile(class FMaterialCompiler* Compiler, int32 OutputIndex, int32 MultiplexIndex)
 {
+	if (Material->MaterialDomain == MD_DeferredDecal)
+	{
+		return CompilerError(Compiler, TEXT("Expression not available in the deferred decal material domain."));
+	}
+
 	return Compiler->ObjectWorldPosition();
 }
 
@@ -5795,6 +5800,11 @@ UMaterialExpressionObjectRadius::UMaterialExpressionObjectRadius(const FObjectIn
 
 int32 UMaterialExpressionObjectRadius::Compile(class FMaterialCompiler* Compiler, int32 OutputIndex, int32 MultiplexIndex)
 {
+	if (Material->MaterialDomain == MD_DeferredDecal)
+	{
+		return CompilerError(Compiler, TEXT("Expression not available in the deferred decal material domain."));
+	}
+
 	return Compiler->ObjectRadius();
 }
 
@@ -5827,6 +5837,11 @@ UMaterialExpressionObjectBounds::UMaterialExpressionObjectBounds(const FObjectIn
 
 int32 UMaterialExpressionObjectBounds::Compile(class FMaterialCompiler* Compiler, int32 OutputIndex, int32 MultiplexIndex)
 {
+	if (Material->MaterialDomain == MD_DeferredDecal)
+	{
+		return CompilerError(Compiler, TEXT("Expression not available in the deferred decal material domain."));
+	}
+
 	return Compiler->ObjectBounds();
 }
 
@@ -7914,6 +7929,11 @@ UMaterialExpressionObjectOrientation::UMaterialExpressionObjectOrientation(const
 
 int32 UMaterialExpressionObjectOrientation::Compile(class FMaterialCompiler* Compiler, int32 OutputIndex, int32 MultiplexIndex)
 {
+	if (Material->MaterialDomain == MD_DeferredDecal)
+	{
+		return CompilerError(Compiler, TEXT("Expression not available in the deferred decal material domain."));
+	}
+
 	return Compiler->ObjectOrientation();
 }
 
