Fix case during undo where Owner of UActorComponent could get out of sync with its Outer hierarchy
Re-private Owner

diff --git a/Engine/Source/Runtime/Engine/Classes/Components/ActorComponent.h b/Engine/Source/Runtime/Engine/Classes/Components/ActorComponent.h
--- a/Engine/Source/Runtime/Engine/Classes/Components/ActorComponent.h
+++ b/Engine/Source/Runtime/Engine/Classes/Components/ActorComponent.h
@@ -146,9 +146,9 @@ public:
 	UPROPERTY()
 	EComponentCreationMethod CreationMethod;
 
-	AActor* Owner;
-
 private:
+	mutable AActor* Owner;
+
 	UPROPERTY()
 	TArray<FSimpleMemberReference> UCSModifiedProperties;
 
@@ -555,6 +555,7 @@ public:
 	virtual void PreEditChange(UProperty* PropertyThatWillChange) override;
 	virtual void PostEditChangeProperty(FPropertyChangedEvent& PropertyChangedEvent) override;
 	virtual void PostEditChangeChainProperty( FPropertyChangedChainEvent& PropertyChangedEvent ) override;
+	virtual void PreEditUndo() override;
 	virtual void PostEditUndo() override;
 #endif // WITH_EDITOR
 	// End UObject interface.
@@ -643,6 +644,10 @@ private:
 
 FORCEINLINE_DEBUGGABLE class AActor* UActorComponent::GetOwner() const
 {
+	if (Owner == nullptr)
+	{
+		Owner = GetTypedOuter<AActor>();
+	}
 	checkSlow(Owner == GetTypedOuter<AActor>()); // verify cached value is correct
 	return Owner;
 }
diff --git a/Engine/Source/Runtime/Engine/Private/Components/ActorComponent.cpp b/Engine/Source/Runtime/Engine/Private/Components/ActorComponent.cpp
--- a/Engine/Source/Runtime/Engine/Private/Components/ActorComponent.cpp
+++ b/Engine/Source/Runtime/Engine/Private/Components/ActorComponent.cpp
@@ -471,6 +471,13 @@ void UActorComponent::PreEditChange(UProperty* PropertyThatWillChange)
 	FlushRenderingCommands();
 }
 
+void UActorComponent::PreEditUndo()
+{
+	Super::PreEditUndo();
+
+	Owner = nullptr;
+}
+
 void UActorComponent::PostEditUndo()
 {
 	// Objects marked pending kill don't call PostEditChange() from UObject::PostEditUndo(),
@@ -487,7 +494,7 @@ void UActorComponent::PostEditUndo()
 	}
 	else
 	{
-		//Let the component be properly registered, after it was restored.
+		// Let the component be properly registered, after it was restored.
 		if (Owner)
 		{
 			Owner->AddOwnedComponent(this);
