#ue4 - Perf: Remove unnecessary FRotator normalization before comparisons in CharacterMovementComponent::PhysicsRotation().

diff --git a/Engine/Source/Runtime/Engine/Private/Components/CharacterMovementComponent.cpp b/Engine/Source/Runtime/Engine/Private/Components/CharacterMovementComponent.cpp
--- a/Engine/Source/Runtime/Engine/Private/Components/CharacterMovementComponent.cpp
+++ b/Engine/Source/Runtime/Engine/Private/Components/CharacterMovementComponent.cpp
@@ -4349,7 +4349,7 @@ void UCharacterMovementComponent::PhysicsRotation(float DeltaTime)
 		DesiredRotation.Roll = 0;
 	}
 
-	if( CurrentRotation.GetDenormalized().Equals(DesiredRotation.GetDenormalized(), 0.01f) )
+	if( CurrentRotation.Equals(DesiredRotation), 0.01f )
 	{
 		return;
 	}
@@ -4357,18 +4357,18 @@ void UCharacterMovementComponent::PhysicsRotation(float DeltaTime)
 	// Accumulate a desired new rotation.
 	FRotator NewRotation = CurrentRotation;	
 
-	//YAW
-	if( DesiredRotation.Yaw != CurrentRotation.Yaw )
-	{
-		NewRotation.Yaw = FMath::FixedTurn(CurrentRotation.Yaw, DesiredRotation.Yaw, DeltaRot.Yaw);
-	}
-
 	// PITCH
 	if( DesiredRotation.Pitch != CurrentRotation.Pitch )
 	{
 		NewRotation.Pitch = FMath::FixedTurn(CurrentRotation.Pitch, DesiredRotation.Pitch, DeltaRot.Pitch);
 	}
 
+	//YAW
+	if( DesiredRotation.Yaw != CurrentRotation.Yaw )
+	{
+		NewRotation.Yaw = FMath::FixedTurn(CurrentRotation.Yaw, DesiredRotation.Yaw, DeltaRot.Yaw);
+	}
+
 	// ROLL
 	if( DesiredRotation.Roll != CurrentRotation.Roll )
 	{
@@ -4378,7 +4378,7 @@ void UCharacterMovementComponent::PhysicsRotation(float DeltaTime)
 	//UpdatedComponent->AngularVelocity = CharAngularVelocity( CurrentRotation, NewRotation, deltaTime );
 
 	// Set the new rotation.
-	if( !NewRotation.Equals(CurrentRotation.GetDenormalized(), 0.01f) )
+	if( !NewRotation.Equals(CurrentRotation, 0.01f) )
 	{
 		MoveUpdatedComponent( FVector::ZeroVector, NewRotation, true );
 	}
