Disable editing of bGenerateLandscapeGIData project setting  when bGenerateMeshDistanceFields is not checked.

diff --git a/Engine/Source/Runtime/Engine/Classes/Engine/RendererSettings.h b/Engine/Source/Runtime/Engine/Classes/Engine/RendererSettings.h
--- a/Engine/Source/Runtime/Engine/Classes/Engine/RendererSettings.h
+++ b/Engine/Source/Runtime/Engine/Classes/Engine/RendererSettings.h
@@ -151,6 +151,7 @@ class ENGINE_API URendererSettings : public UDeveloperSettings
 	uint32 bGenerateMeshDistanceFields:1;
 
 	UPROPERTY(config, EditAnywhere, Category = Lighting, meta = (
+		EditCondition = "bGenerateMeshDistanceFields",
 		ConsoleVariable = "r.GenerateLandscapeGIData", DisplayName = "Generate Landscape Real-time GI Data",
 		ToolTip = "Whether to generate a low-resolution base color texture for landscapes for rendering real-time global illumination.  This feature requires GenerateMeshDistanceFields is also enabled, and will increase mesh build times and memory usage."))
 		uint32 bGenerateLandscapeGIData : 1;
