#ue4 - Create UWorld::DefaultPhysicsVolume on demand.

- Fixes crash when dealing with physics volumes when world is not initialized (commandlets etc).

diff --git a/Engine/Source/Runtime/Engine/Classes/Engine/World.h b/Engine/Source/Runtime/Engine/Classes/Engine/World.h
--- a/Engine/Source/Runtime/Engine/Classes/Engine/World.h
+++ b/Engine/Source/Runtime/Engine/Classes/Engine/World.h
@@ -631,10 +631,13 @@ class ENGINE_API UWorld : public UObject, public FNetworkNotify
 	UPROPERTY()
 	class AParticleEventManager*				MyParticleEventManager;
 
+private:
 	/** DefaultPhysicsVolume used for whole game **/
 	UPROPERTY()
 	APhysicsVolume*								DefaultPhysicsVolume;
 
+public:
+
 	/** View locations rendered in the previous frame, if any. */
 	TArray<FVector>								ViewLocationsRenderedLastFrame;
 
@@ -2178,6 +2181,9 @@ public:
 	 */
 	APhysicsVolume* GetDefaultPhysicsVolume() const;
 
+	/** Returns true if a DefaultPhysicsVolume has been created. */
+	bool HasDefaultPhysicsVolume() const { return DefaultPhysicsVolume != nullptr; }
+
 	/** Add a physics volume to the list of those in the world. DefaultPhysicsVolume is not tracked. Used internally by APhysicsVolume. */
 	void AddPhysicsVolume(APhysicsVolume* Volume);
 
diff --git a/Engine/Source/Runtime/Engine/Private/World.cpp b/Engine/Source/Runtime/Engine/Private/World.cpp
--- a/Engine/Source/Runtime/Engine/Private/World.cpp
+++ b/Engine/Source/Runtime/Engine/Private/World.cpp
@@ -883,12 +883,8 @@ void UWorld::InitWorld(const InitializationValues IVS)
 	check(GetWorldSettings());
 
 	// initialize DefaultPhysicsVolume for the world
-	// checked
-	if (DefaultPhysicsVolume == NULL)
-	{
-		DefaultPhysicsVolume = SpawnActor<APhysicsVolume>(ADefaultPhysicsVolume::StaticClass());
-		DefaultPhysicsVolume->Priority = -1000000;
-	}
+	// Spawned on demand by this function.
+	DefaultPhysicsVolume = GetDefaultPhysicsVolume();
 
 	// update default values when world is restarted
 	DefaultPhysicsVolume->TerminalVelocity = UPhysicsSettings::Get()->DefaultTerminalVelocity;
@@ -3424,6 +3420,13 @@ void UWorld::SetPhysicsScene(FPhysScene* InScene)
 
 APhysicsVolume* UWorld::GetDefaultPhysicsVolume() const
 {
+	// Create on demand.
+	if (!DefaultPhysicsVolume)
+	{
+		UWorld* MutableThis = const_cast<UWorld*>(this);
+		MutableThis->DefaultPhysicsVolume = MutableThis->SpawnActor<APhysicsVolume>(ADefaultPhysicsVolume::StaticClass());
+		MutableThis->DefaultPhysicsVolume->Priority = -1000000;
+	}
 	return DefaultPhysicsVolume;
 }
 
