#ue4 - Fix typo in CL 2527667 that broke bUseControllerDesiredRotation in CharacterMovement.

diff --git a/Engine/Source/Runtime/Engine/Private/Components/CharacterMovementComponent.cpp b/Engine/Source/Runtime/Engine/Private/Components/CharacterMovementComponent.cpp
--- a/Engine/Source/Runtime/Engine/Private/Components/CharacterMovementComponent.cpp
+++ b/Engine/Source/Runtime/Engine/Private/Components/CharacterMovementComponent.cpp
@@ -4349,7 +4349,7 @@ void UCharacterMovementComponent::PhysicsRotation(float DeltaTime)
 		DesiredRotation.Roll = 0;
 	}
 
-	if( CurrentRotation.Equals(DesiredRotation), 0.01f )
+	if( CurrentRotation.Equals(DesiredRotation, 0.01f) )
 	{
 		return;
 	}
