Various fixes for in-game / in-editor screenshot.

Summary of changes:
- Editor now no longer relies on the delegate registered by the AutomationWorkerModule to save the image as a .png.
- Removed the OnPNGScreenshotCaptured delegate as it was only being used by the AutomationWorkerModule, and made .png the default format for saving screenshots.
- AutomationWorkerModule registers a HandlePreTestingEvent and HandlePostTestingEvent with the AutomationTestFramework, whose purpose is to hook in to the UGameViewportClient::OnScreenshotCaptured delegate, to override the default screenshot capture behaviour.
- Removed one of the FScreenshotRequest::RequestScreenshot implementations and united the funcionality to be able to specify whether the UI should be captured, and whether a unique filename suffix should be added to the supplied name.

diff --git a/Engine/Source/Runtime/AutomationWorker/Private/AutomationWorkerModule.cpp b/Engine/Source/Runtime/AutomationWorker/Private/AutomationWorkerModule.cpp
--- a/Engine/Source/Runtime/AutomationWorker/Private/AutomationWorkerModule.cpp
+++ b/Engine/Source/Runtime/AutomationWorker/Private/AutomationWorkerModule.cpp
@@ -15,8 +15,16 @@ IMPLEMENT_MODULE(FAutomationWorkerModule, AutomationWorker);
 void FAutomationWorkerModule::StartupModule()
 {
 	Initialize();
+
+	FAutomationTestFramework::GetInstance().PreTestingEvent.AddRaw(this, &FAutomationWorkerModule::HandlePreTestingEvent);
+	FAutomationTestFramework::GetInstance().PostTestingEvent.AddRaw(this, &FAutomationWorkerModule::HandlePostTestingEvent);
 }
 
+void FAutomationWorkerModule::ShutdownModule()
+{
+	FAutomationTestFramework::GetInstance().PreTestingEvent.RemoveAll(this);
+	FAutomationTestFramework::GetInstance().PostTestingEvent.RemoveAll(this);
+}
 
 bool FAutomationWorkerModule::SupportsDynamicReloading()
 {
@@ -118,15 +126,6 @@ void FAutomationWorkerModule::Initialize()
 			MessageEndpoint->Subscribe<FAutomationWorkerFindWorkers>();
 		}
 
-#if WITH_ENGINE
-		if (!GIsEditor && GEngine->GameViewport)
-		{
-			GEngine->GameViewport->OnPNGScreenshotCaptured().BindRaw(this, &FAutomationWorkerModule::HandleScreenShotCaptured);
-		}
-		//Register the editor screen shot callback
-		FAutomationTestFramework::GetInstance().OnScreenshotCaptured().BindRaw(this, &FAutomationWorkerModule::HandleScreenShotCaptured);
-#endif
-
 		bExecuteNextNetworkCommand = true;		
 	}
 	else
@@ -329,8 +328,41 @@ void FAutomationWorkerModule::HandleRequestTestsMessage( const FAutomationWorker
 }
 
 
+void FAutomationWorkerModule::HandlePreTestingEvent()
+{
+#if WITH_ENGINE
+	if (!GIsEditor && GEngine->GameViewport)
+	{
+		GEngine->GameViewport->OnScreenshotCaptured().AddRaw(this, &FAutomationWorkerModule::HandleScreenShotCaptured);
+	}
+	//Register the editor screen shot callback
+	FAutomationTestFramework::GetInstance().OnScreenshotCaptured().BindRaw(this, &FAutomationWorkerModule::HandleScreenShotCapturedWithName);
+#endif
+}
+
+
+void FAutomationWorkerModule::HandlePostTestingEvent()
+{
+#if WITH_ENGINE
+	if (!GIsEditor && GEngine->GameViewport)
+	{
+		GEngine->GameViewport->OnScreenshotCaptured().RemoveAll(this);
+	}
+	//Register the editor screen shot callback
+	FAutomationTestFramework::GetInstance().OnScreenshotCaptured().BindRaw(this, &FAutomationWorkerModule::HandleScreenShotCapturedWithName);
+#endif
+}
+
+
 #if WITH_ENGINE
-void FAutomationWorkerModule::HandleScreenShotCaptured( int32 Width, int32 Height, const TArray<FColor>& Bitmap, const FString& ScreenShotName )
+void FAutomationWorkerModule::HandleScreenShotCaptured(int32 Width, int32 Height, const TArray<FColor>& Bitmap)
+{
+	FString Filename = FScreenshotRequest::GetFilename();
+
+	HandleScreenShotCapturedWithName(Width, Height, Bitmap, Filename);
+}
+
+void FAutomationWorkerModule::HandleScreenShotCapturedWithName(int32 Width, int32 Height, const TArray<FColor>& Bitmap, const FString& ScreenShotName)
 {
 	if( FAutomationTestFramework::GetInstance().IsScreenshotAllowed() )
 	{
diff --git a/Engine/Source/Runtime/AutomationWorker/Private/AutomationWorkerModule.h b/Engine/Source/Runtime/AutomationWorker/Private/AutomationWorkerModule.h
--- a/Engine/Source/Runtime/AutomationWorker/Private/AutomationWorkerModule.h
+++ b/Engine/Source/Runtime/AutomationWorker/Private/AutomationWorkerModule.h
@@ -14,6 +14,7 @@ public:
 	// Begin IModuleInterface interface
 
 	virtual void StartupModule() override;
+	virtual void ShutdownModule() override;
 
 	virtual bool SupportsDynamicReloading() override;
 
@@ -86,9 +87,16 @@ private:
 	// Handles FAutomationWorkerRunTests messages.
 	void HandleRunTestsMessage( const FAutomationWorkerRunTests& Message, const IMessageContextRef& Context );
 
+	// Handles FAutomationTestFramework PreTestingEvents.
+	void HandlePreTestingEvent();
+
+	// Handles FAutomationTestFramework PostTestingEvents.
+	void HandlePostTestingEvent();
+
 #if WITH_ENGINE
 	/** Invoked when we have screen shot to send. */
-	void HandleScreenShotCaptured( int32 Width, int32 Height, const TArray<FColor>& Bitmap, const FString& ScreenShotName );
+	void HandleScreenShotCaptured(int32 Width, int32 Height, const TArray<FColor>& Bitmap);
+	void HandleScreenShotCapturedWithName(int32 Width, int32 Height, const TArray<FColor>& Bitmap, const FString& ScreenShotName);
 #endif
 
 private:
diff --git a/Engine/Source/Runtime/Engine/Classes/Engine/GameViewportClient.h b/Engine/Source/Runtime/Engine/Classes/Engine/GameViewportClient.h
--- a/Engine/Source/Runtime/Engine/Classes/Engine/GameViewportClient.h
+++ b/Engine/Source/Runtime/Engine/Classes/Engine/GameViewportClient.h
@@ -464,12 +464,6 @@ public:
 	/** Make sure all navigation objects have appropriate path rendering components set.  Called when EngineShowFlags.Navigation is set. */
 	virtual void VerifyPathRenderingComponents();
 
-	/** Accessor for delegate called when a png screenshot is captured */
-	FOnPNGScreenshotCaptured& OnPNGScreenshotCaptured()
-	{
-		return PNGScreenshotCapturedDelegate;
-	}
-
 	/** Accessor for delegate called when a screenshot is captured */
 	static FOnScreenshotCaptured& OnScreenshotCaptured()
 	{
@@ -738,9 +732,6 @@ private:
 	 */
 	bool SetDisplayConfiguration( const FIntPoint* Dimensions, EWindowMode::Type WindowMode);
 
-	/** Delegate called at the end of the frame when a screenshot is captured and a .png is requested */
-	FOnPNGScreenshotCaptured PNGScreenshotCapturedDelegate;
-
 	/** Delegate called at the end of the frame when a screenshot is captured */
 	static FOnScreenshotCaptured ScreenshotCapturedDelegate;
 
diff --git a/Engine/Source/Runtime/Engine/Classes/Engine/GameViewportDelegates.h b/Engine/Source/Runtime/Engine/Classes/Engine/GameViewportDelegates.h
--- a/Engine/Source/Runtime/Engine/Classes/Engine/GameViewportDelegates.h
+++ b/Engine/Source/Runtime/Engine/Classes/Engine/GameViewportDelegates.h
@@ -13,17 +13,6 @@
  */
 DECLARE_MULTICAST_DELEGATE_ThreeParams(FOnScreenshotCaptured, int32 /*Width*/, int32 /*Height*/, const TArray<FColor>& /*Colors*/);
 
-/**
- * Delegate type used by UGameViewportClient when a png screenshot has been captured
- *
- * The first parameter is the width.
- * The second parameter is the height.
- * The third parameter is the array of bitmap data.
- * The fourth parameter is the screen shot filename.
- * @see UGameViewportClient
- */
-DECLARE_DELEGATE_FourParams(FOnPNGScreenshotCaptured, int32, int32, const TArray<FColor>&, const FString&);
-
 /**
  * Delegate type used by UGameViewportClient when call is made to close a viewport
  *
diff --git a/Engine/Source/Runtime/Engine/Private/GameViewportClient.cpp b/Engine/Source/Runtime/Engine/Private/GameViewportClient.cpp
--- a/Engine/Source/Runtime/Engine/Private/GameViewportClient.cpp
+++ b/Engine/Source/Runtime/Engine/Private/GameViewportClient.cpp
@@ -1330,32 +1330,32 @@ void UGameViewportClient::ProcessScreenShots(FViewport* InViewport)
 				if (GIsDumpingMovie && ScreenShotName.IsEmpty())
 				{
 					// Request a new screenshot with a formatted name
-					FScreenshotRequest::RequestScreenshot(false);
+					const bool bShowUI = false;
+					const bool bAddFilenameSuffix = true;
+					FScreenshotRequest::RequestScreenshot(FString(), bShowUI, bAddFilenameSuffix);
 					ScreenShotName = FScreenshotRequest::GetFilename();
 				}
 
-				if (PNGScreenshotCapturedDelegate.IsBound() && FPaths::GetExtension(ScreenShotName).ToLower() == TEXT("png"))
+				GetHighResScreenshotConfig().MergeMaskIntoAlpha(Bitmap);
+				
+				FIntRect SourceRect(0, 0, GScreenshotResolutionX, GScreenshotResolutionY);
+				if (GIsHighResScreenshot)
 				{
-					PNGScreenshotCapturedDelegate.Execute(Size.X, Size.Y, Bitmap, ScreenShotName);
+					SourceRect = GetHighResScreenshotConfig().CaptureRegion;
 				}
-				else
-				{
-					// Save the contents of the array to a bitmap file.
-					bool bWriteAlpha = false;
-					FIntRect SourceRect(0, 0, GScreenshotResolutionX, GScreenshotResolutionY);
-					if (GIsHighResScreenshot)
-					{
-						bWriteAlpha = GetHighResScreenshotConfig().MergeMaskIntoAlpha(Bitmap);
-						SourceRect = GetHighResScreenshotConfig().CaptureRegion;
-					}
 
-					if (!FPaths::GetExtension(ScreenShotName).IsEmpty())
-					{
-						ScreenShotName = FPaths::GetBaseFilename(ScreenShotName, false);
-						ScreenShotName += TEXT(".bmp");
-					}
-					FFileHelper::CreateBitmap(*ScreenShotName, Size.X, Size.Y, Bitmap.GetData(), &SourceRect, &IFileManager::Get(), NULL, bWriteAlpha);
+				if (!FPaths::GetExtension(ScreenShotName).IsEmpty())
+				{
+					ScreenShotName = FPaths::GetBaseFilename(ScreenShotName, false);
+					ScreenShotName += TEXT(".png");
 				}
+
+				// Save the contents of the array to a png file.
+				TArray<uint8> CompressedBitmap;
+				FImageUtils::CompressImageArray(Size.X, Size.Y, Bitmap, CompressedBitmap);
+				FFileHelper::SaveArrayToFile(CompressedBitmap, *ScreenShotName);
+
+//				FFileHelper::CreateBitmap(*ScreenShotName, InViewport->GetSizeXY().X, InViewport->GetSizeXY().Y, Bitmap.GetData(), &SourceRect, &IFileManager::Get(), NULL, bWriteAlpha);
 			}
 		}
 
@@ -2836,10 +2836,8 @@ bool UGameViewportClient::HandleScreenshotCommand( const TCHAR* Cmd, FOutputDevi
 	if(Viewport)
 	{
 		const bool bShowUI = FParse::Command(&Cmd, TEXT("SHOWUI"));
-
-//		FScreenshotRequest::RequestScreenshot( bShowUI, TEXT("png") );
-		// PNG is disabled for now as it breaks "shot" command in game, see UE-5780 
-		FScreenshotRequest::RequestScreenshot( bShowUI, TEXT("bmp") );
+		const bool bAddFilenameSuffix = true;
+		FScreenshotRequest::RequestScreenshot( FString(), bShowUI, bAddFilenameSuffix );
 
 		GScreenMessagesRestoreState = GAreScreenMessagesEnabled;
 		GAreScreenMessagesEnabled = false;
@@ -3118,7 +3116,7 @@ bool UGameViewportClient::RequestBugScreenShot(const TCHAR* Cmd, bool bDisplayHU
 		TCHAR File[MAX_SPRINTF] = TEXT("");
 		for( int32 TestBitmapIndex = 0; TestBitmapIndex < 9; ++TestBitmapIndex )
 		{ 
-			const FString DescPlusExtension = FString::Printf( TEXT("%s%i.bmp"), Cmd, TestBitmapIndex );
+			const FString DescPlusExtension = FString::Printf( TEXT("%s%i.png"), Cmd, TestBitmapIndex );
 			const FString SSFilename = CreateProfileFilename( DescPlusExtension, false );
 
 			const FString OutputDir = FPaths::BugItDir() + FString::Printf( TEXT("%s"), Cmd) + TEXT("/");
@@ -3143,7 +3141,8 @@ bool UGameViewportClient::RequestBugScreenShot(const TCHAR* Cmd, bool bDisplayHU
 				GScreenshotBitmapIndex = TestBitmapIndex; // this is safe as the UnMisc.cpp ScreenShot code will test each number before writing a file
 
 				const bool bShowUI = true;
-				FScreenshotRequest::RequestScreenshot( File, bShowUI );
+				const bool bAddFilenameSuffix = false;
+				FScreenshotRequest::RequestScreenshot( File, bShowUI, bAddFilenameSuffix );
 				break;
 			}
 		}
diff --git a/Engine/Source/Runtime/Engine/Private/HighResScreenshot.cpp b/Engine/Source/Runtime/Engine/Private/HighResScreenshot.cpp
--- a/Engine/Source/Runtime/Engine/Private/HighResScreenshot.cpp
+++ b/Engine/Source/Runtime/Engine/Private/HighResScreenshot.cpp
@@ -105,6 +105,14 @@ bool FHighResScreenshotConfig::MergeMaskIntoAlpha(TArray<FColor>& InBitmap)
 
 		bWritten = true;
 	}
+	else
+	{
+		// Ensure that all pixels' alpha is set to 255
+		for (auto& Color : InBitmap)
+		{
+			Color.A = 255;
+		}
+	}
 
 	return bWritten;
 }
diff --git a/Engine/Source/Runtime/Engine/Private/Tests/EngineAutomationTests.cpp b/Engine/Source/Runtime/Engine/Private/Tests/EngineAutomationTests.cpp
--- a/Engine/Source/Runtime/Engine/Private/Tests/EngineAutomationTests.cpp
+++ b/Engine/Source/Runtime/Engine/Private/Tests/EngineAutomationTests.cpp
@@ -219,8 +219,9 @@ DEFINE_LATENT_AUTOMATION_COMMAND_ONE_PARAMETER(FTakeViewportScreenshotCommand, F
 
 bool FTakeViewportScreenshotCommand::Update()
 {
-	
-	FScreenshotRequest::RequestScreenshot( ScreenshotFileName, false );
+	const bool bShowUI = false;
+	const bool bAddFilenameSuffix = false;
+	FScreenshotRequest::RequestScreenshot( ScreenshotFileName, bShowUI, bAddFilenameSuffix );
 	return true;
 }
 
diff --git a/Engine/Source/Runtime/Engine/Private/UnrealClient.cpp b/Engine/Source/Runtime/Engine/Private/UnrealClient.cpp
--- a/Engine/Source/Runtime/Engine/Private/UnrealClient.cpp
+++ b/Engine/Source/Runtime/Engine/Private/UnrealClient.cpp
@@ -200,23 +200,26 @@ const FTexture2DRHIRef& FRenderTarget::GetRenderTargetTexture() const
 }
 
 
-void FScreenshotRequest::RequestScreenshot( const FString& InFilename, bool bInShowUI )
+void FScreenshotRequest::RequestScreenshot( const FString& InFilename, bool bInShowUI, bool bAddUniqueSuffix )
 {
-	Filename = InFilename;
-	CreateViewportScreenShotFilename( Filename );
-	bShowUI = bInShowUI;
-}
-
+	FString GeneratedFilename = InFilename;
+	CreateViewportScreenShotFilename(GeneratedFilename);
 
-void FScreenshotRequest::RequestScreenshot( bool bInShowUI, const FString& InExtension )
-{
-	FString NewFilename;
-	CreateViewportScreenShotFilename( NewFilename );
-	FFileHelper::GenerateNextBitmapFilename(NewFilename, InExtension, Filename);
+	if (bAddUniqueSuffix)
+	{
+		const bool bRemovePath = false;
+		GeneratedFilename = FPaths::GetBaseFilename(GeneratedFilename, bRemovePath);
+		FFileHelper::GenerateNextBitmapFilename(GeneratedFilename, TEXT("png"), Filename);
+	}
+	else
+	{
+		Filename = GeneratedFilename;
+	}
 
 	bShowUI = bInShowUI;
 }
 
+
 void FScreenshotRequest::Reset()
 {
 	Filename.Empty();
@@ -945,8 +948,9 @@ void FViewport::Draw( bool bShouldPresent /*= true */)
 		{
 			if( GIsHighResScreenshot || bTakeHighResScreenShot )
 			{
-				bool bShowUI = false;
-				FScreenshotRequest::RequestScreenshot( bShowUI, TEXT("png") );
+				const bool bShowUI = false;
+				const bool bAddFilenameSuffix = true;
+				FScreenshotRequest::RequestScreenshot( FString(), bShowUI, bAddFilenameSuffix );
 				GIsHighResScreenshot = true;
 				GScreenMessagesRestoreState = GAreScreenMessagesEnabled;
 				GAreScreenMessagesEnabled = false;
diff --git a/Engine/Source/Runtime/Engine/Private/UnrealEngine.cpp b/Engine/Source/Runtime/Engine/Private/UnrealEngine.cpp
--- a/Engine/Source/Runtime/Engine/Private/UnrealEngine.cpp
+++ b/Engine/Source/Runtime/Engine/Private/UnrealEngine.cpp
@@ -6431,7 +6431,9 @@ void UEngine::PerformanceCapture(const FString& CaptureName)
 	
 	FString ScreenshotName = FString::Printf(TEXT("%s%d.png"), *PathName, GEngineVersion.GetChangelist());
 	
-	FScreenshotRequest::RequestScreenshot( ScreenshotName, false );
+	const bool bShowUI = false;
+	const bool bAddFilenameSuffix = false;
+	FScreenshotRequest::RequestScreenshot( ScreenshotName, bShowUI, bAddFilenameSuffix );
 
 }
 
diff --git a/Engine/Source/Runtime/Engine/Public/UnrealClient.h b/Engine/Source/Runtime/Engine/Public/UnrealClient.h
--- a/Engine/Source/Runtime/Engine/Public/UnrealClient.h
+++ b/Engine/Source/Runtime/Engine/Public/UnrealClient.h
@@ -113,19 +113,11 @@ struct ENGINE_API FScreenshotRequest
 	/**
 	 * Requests a new screenshot with a specific filename
 	 *
-	 * @param InFilename	The filename to use
-	 * @param bInShowUI		Whether or not to show Slate UI
+	 * @param InFilename			The filename to use
+	 * @param bInShowUI				Whether or not to show Slate UI
+	 * @param bAddFilenameSuffix	Whether an auto-generated unique suffix should be added to the supplied filename
 	 */
-	static void RequestScreenshot( const FString& InFilename, bool bInShowUI );
-
-	/**
-	 * Requests a new screenshot.  This method will auto-generate a filename via the format "Screenshot{0}.{1} where {0} is an incremented coutner
-	 * and {1} is an optional extension (defaulting to .bmp)
-	 *
-	 * @param bInShowUI		Whether or not to show Slate UI
-	 * @param InExtension	Which file extension to use
-	 */
-	static void RequestScreenshot( bool bInShowUI, const FString& InExtension = FString(TEXT("bmp") ) );
+	static void RequestScreenshot( const FString& InFilename, bool bInShowUI, bool bAddFilenameSuffix );
 
 	/**
 	 * Resets a screenshot request
