#ue4 - Optimize CharacterMovement::PhysicsRotation() to reduce FRotator comparisons.

diff --git a/Engine/Source/Runtime/Engine/Private/Components/CharacterMovementComponent.cpp b/Engine/Source/Runtime/Engine/Private/Components/CharacterMovementComponent.cpp
--- a/Engine/Source/Runtime/Engine/Private/Components/CharacterMovementComponent.cpp
+++ b/Engine/Source/Runtime/Engine/Private/Components/CharacterMovementComponent.cpp
@@ -4315,17 +4315,17 @@ FRotator UCharacterMovementComponent::ComputeOrientToMovementRotation(const FRot
 
 void UCharacterMovementComponent::PhysicsRotation(float DeltaTime)
 {
-	if (!HasValidData() || (!CharacterOwner->Controller && !bRunPhysicsWithNoController))
+	if (!(bOrientRotationToMovement || bUseControllerDesiredRotation))
 	{
 		return;
 	}
 
-	if (!(bOrientRotationToMovement || bUseControllerDesiredRotation))
+	if (!HasValidData() || (!CharacterOwner->Controller && !bRunPhysicsWithNoController))
 	{
 		return;
 	}
 
-	const FRotator CurrentRotation = UpdatedComponent->GetComponentRotation();
+	const FRotator CurrentRotation = UpdatedComponent->GetComponentRotation(); // Normalized
 	FRotator DeltaRot = GetDeltaRotation(DeltaTime);
 	FRotator DesiredRotation = CurrentRotation;
 
@@ -4343,44 +4343,46 @@ void UCharacterMovementComponent::PhysicsRotation(float DeltaTime)
 	}
 
 	// Always remain vertical when walking or falling.
-	if( IsMovingOnGround() || IsFalling() )
+	if (IsMovingOnGround() || IsFalling())
 	{
-		DesiredRotation.Pitch = 0;
-		DesiredRotation.Roll = 0;
+		DesiredRotation.Pitch = 0.f;
+		DesiredRotation.Yaw = FRotator::NormalizeAxis(DesiredRotation.Yaw);
+		DesiredRotation.Roll = 0.f;
 	}
-
-	if( CurrentRotation.Equals(DesiredRotation, 0.01f) )
+	else
 	{
-		return;
+		DesiredRotation.Normalize();
 	}
-
+	
 	// Accumulate a desired new rotation.
-	FRotator NewRotation = CurrentRotation;	
+	const float AngleTolerance = 1e-3f;
+	bool bRotationChanged = false;
 
 	// PITCH
-	if( DesiredRotation.Pitch != CurrentRotation.Pitch )
+	if ((DeltaRot.Pitch != 0.f) && !FMath::IsNearlyEqual(CurrentRotation.Pitch, DesiredRotation.Pitch, AngleTolerance))
 	{
-		NewRotation.Pitch = FMath::FixedTurn(CurrentRotation.Pitch, DesiredRotation.Pitch, DeltaRot.Pitch);
+		DesiredRotation.Pitch = FMath::FixedTurn(CurrentRotation.Pitch, DesiredRotation.Pitch, DeltaRot.Pitch);
+		bRotationChanged = true;
 	}
 
-	//YAW
-	if( DesiredRotation.Yaw != CurrentRotation.Yaw )
+	// YAW
+	if ((DeltaRot.Yaw != 0.f) && !FMath::IsNearlyEqual(CurrentRotation.Yaw, DesiredRotation.Yaw, AngleTolerance))
 	{
-		NewRotation.Yaw = FMath::FixedTurn(CurrentRotation.Yaw, DesiredRotation.Yaw, DeltaRot.Yaw);
+		DesiredRotation.Yaw = FMath::FixedTurn(CurrentRotation.Yaw, DesiredRotation.Yaw, DeltaRot.Yaw);
+		bRotationChanged = true;
 	}
 
 	// ROLL
-	if( DesiredRotation.Roll != CurrentRotation.Roll )
+	if ((DeltaRot.Roll != 0.f) && !FMath::IsNearlyEqual(CurrentRotation.Roll, DesiredRotation.Roll, AngleTolerance))
 	{
-		NewRotation.Roll = FMath::FixedTurn(CurrentRotation.Roll, DesiredRotation.Roll, DeltaRot.Roll);
+		DesiredRotation.Roll = FMath::FixedTurn(CurrentRotation.Roll, DesiredRotation.Roll, DeltaRot.Roll);
+		bRotationChanged = true;
 	}
 
-	//UpdatedComponent->AngularVelocity = CharAngularVelocity( CurrentRotation, NewRotation, deltaTime );
-
 	// Set the new rotation.
-	if( !NewRotation.Equals(CurrentRotation, 0.01f) )
+	if (bRotationChanged)
 	{
-		MoveUpdatedComponent( FVector::ZeroVector, NewRotation, true );
+		MoveUpdatedComponent( FVector::ZeroVector, DesiredRotation, true );
 	}
 }
 
