Fixed shadowed variable build error.

diff --git a/Engine/Source/Runtime/Engine/Private/GameViewportClient.cpp b/Engine/Source/Runtime/Engine/Private/GameViewportClient.cpp
--- a/Engine/Source/Runtime/Engine/Private/GameViewportClient.cpp
+++ b/Engine/Source/Runtime/Engine/Private/GameViewportClient.cpp
@@ -1330,7 +1330,7 @@ void UGameViewportClient::ProcessScreenShots(FViewport* InViewport)
 				if (GIsDumpingMovie && ScreenShotName.IsEmpty())
 				{
 					// Request a new screenshot with a formatted name
-					const bool bShowUI = false;
+					bShowUI = false;
 					const bool bAddFilenameSuffix = true;
 					FScreenshotRequest::RequestScreenshot(FString(), bShowUI, bAddFilenameSuffix);
 					ScreenShotName = FScreenshotRequest::GetFilename();
