Fix for MeshRotationRate * Life
UE-2944

diff --git a/Engine/Source/Runtime/Engine/Private/Particles/ParticleEmitterInstances.cpp b/Engine/Source/Runtime/Engine/Private/Particles/ParticleEmitterInstances.cpp
--- a/Engine/Source/Runtime/Engine/Private/Particles/ParticleEmitterInstances.cpp
+++ b/Engine/Source/Runtime/Engine/Private/Particles/ParticleEmitterInstances.cpp
@@ -2942,14 +2942,25 @@ void FParticleMeshEmitterInstance::Tick(float DeltaTime, bool bSuppressSpawning)
 					PayloadData->Rotation = PayloadData->InitRotation + PayloadData->CurContinuousRotation;
 				}
 			}
-
-			PayloadData->CurContinuousRotation += DeltaTime * PayloadData->RotationRate;
 		}
 	}
 
 
 	// Call the standard tick
 	FParticleEmitterInstance::Tick(DeltaTime, bSuppressSpawning);
+	
+	if (MeshRotationActive)
+	{
+		//Must do this (at least) after module update other wise the reset value of RotationRate is used.
+		//Probably the other stuff before the module tick should be brought down here too and just leave the RotationRate reset before.
+		//Though for the sake of not breaking existing behavior, leave things as they are for now.
+		for (int32 i = 0; i < ActiveParticles; i++)
+		{
+			DECLARE_PARTICLE(Particle, ParticleData + ParticleStride * ParticleIndices[i]);
+			FMeshRotationPayloadData* PayloadData = (FMeshRotationPayloadData*)((uint8*)&Particle + MeshRotationOffset);
+			PayloadData->CurContinuousRotation += DeltaTime * PayloadData->RotationRate;
+		}
+	}
 
 	// Remove from the Sprite count... happens because we use the Super::Tick
 	DEC_DWORD_STAT_BY(STAT_SpriteParticles, ActiveParticles);
