UE4 - fix minor problem with spam of SetMaterialUsage blocking calls back to the main thread

diff --git a/Engine/Source/Runtime/Engine/Private/Materials/Material.cpp b/Engine/Source/Runtime/Engine/Private/Materials/Material.cpp
--- a/Engine/Source/Runtime/Engine/Private/Materials/Material.cpp
+++ b/Engine/Source/Runtime/Engine/Private/Materials/Material.cpp
@@ -940,7 +940,7 @@ bool UMaterial::CheckMaterialUsage_Concurrent(EMaterialUsage Usage, const bool b
 					Event.Trigger();
 				}
 			};
-			UE_LOG(LogMaterial, Warning, TEXT("Has to pass SMU back to game thread. This stalls the tasks graph, but since it is editor only, is not such a big deal."));
+			UE_LOG(LogMaterial, Warning, TEXT("Has to pass SMU back to game thread. This stalls the tasks graph, but since it is editor only or only happens once, is not such a big deal."));
 
 			FScopedEvent Event;
 			FCallSMU CallSMU(const_cast<UMaterial*>(this), Usage, bSkipPrim, bUsageSetSuccessfully, Event);
@@ -983,9 +983,19 @@ bool UMaterial::NeedsSetMaterialUsage_Concurrent(bool &bOutHasUsage, EMaterialUs
 	// Check that the material has been flagged for use with the given usage flag.
 	if(!GetUsageByFlag(Usage) && !bUsedAsSpecialEngineMaterial)
 	{
-		// This will be overwritten later by SetMaterialUsage, since we are saying that it needs to be called with the return value
-		bOutHasUsage = false;
-		return true;
+		uint32 UsageFlagBit = (1 << (uint32)Usage);
+		if ((UsageFlagWarnings & UsageFlagBit) == 0)
+		{
+			// This will be overwritten later by SetMaterialUsage, since we are saying that it needs to be called with the return value
+			bOutHasUsage = false;
+			return true;
+		}
+		else
+		{
+			// We have already warned about this, so we aren't going to warn or compile or set anything this time
+			bOutHasUsage = false;
+			return false;
+		}
 	}
 	return false;
 }
