Fix crash when trying to set transform, add force, etc... on a body that's been added with deferred.

UE-12129

diff --git a/Engine/Source/Runtime/Engine/Private/PhysicsEngine/BodyInstance.cpp b/Engine/Source/Runtime/Engine/Private/PhysicsEngine/BodyInstance.cpp
--- a/Engine/Source/Runtime/Engine/Private/PhysicsEngine/BodyInstance.cpp
+++ b/Engine/Source/Runtime/Engine/Private/PhysicsEngine/BodyInstance.cpp
@@ -316,6 +316,13 @@ FBodyInstance* FindParentBodyInstance(FName BodyName, USkeletalMeshComponent* Sk
 	return NULL;
 }
 
+FPhysScene* GetPhysicsScene(const FBodyInstance* BodyInstance)
+{
+	UPrimitiveComponent* OwnerComponentInst = BodyInstance->OwnerComponent.Get();
+	return OwnerComponentInst ? OwnerComponentInst->GetWorld()->GetPhysicsScene() : nullptr;
+}
+
+
 #if WITH_PHYSX
 //Determine that the shape is associated with this subbody (or root body)
 bool ShapeBoundToBody(const PxShape * PShape, const FBodyInstance* SubBody)
@@ -2592,9 +2599,10 @@ void FBodyInstance::SetBodyTransform(const FTransform& NewTransform, bool bTelep
 				// If kinematic and not teleporting, set kinematic target
 				if (!IsRigidBodyNonKinematic_AssumesLocked(PRigidDynamic) && !bTeleport)
 				{
-					const PxScene* PScene = PRigidDynamic->getScene();
-					FPhysScene* PhysScene = FPhysxUserData::Get<FPhysScene>(PScene->userData);
+					if(FPhysScene* PhysScene = GetPhysicsScene(this))
+					{
 						PhysScene->SetKinematicTarget_AssumesLocked(this, NewTransform, true);
+					}
 				}
 				// Otherwise, set global pose
 				else
@@ -3412,11 +3420,12 @@ void FBodyInstance::AddCustomPhysics(FCalculateCustomPhysics& CalculateCustomPhy
 	ExecuteOnPxRigidBodyReadOnly(this, [&](const PxRigidBody* PRigidBody)
 	{
 		if (IsRigidBodyNonKinematic_AssumesLocked(PRigidBody))
-	{
-		const PxScene* PScene = PRigidBody->getScene();
-		FPhysScene* PhysScene = FPhysxUserData::Get<FPhysScene>(PScene->userData);
-			PhysScene->AddCustomPhysics_AssumesLocked(this, CalculateCustomPhysics);
-	}
+		{
+			if(FPhysScene* PhysScene = GetPhysicsScene(this))
+			{
+				PhysScene->AddCustomPhysics_AssumesLocked(this, CalculateCustomPhysics);
+			}
+		}
 	});
 	
 #endif // WITH_PHYSX
@@ -3436,12 +3445,13 @@ void FBodyInstance::AddForce(const FVector& Force, bool bAllowSubstepping, bool
 	ExecuteOnPxRigidBodyReadWrite(this, [&](PxRigidBody* PRigidBody)
 	{
 		if (IsRigidBodyNonKinematic_AssumesLocked(PRigidBody))
-	{
-		const PxScene* PScene = PRigidBody->getScene();
-		FPhysScene* PhysScene = FPhysxUserData::Get<FPhysScene>(PScene->userData);
-			PhysScene->AddForce_AssumesLocked(this, Force, bAllowSubstepping, bAccelChange);
+		{
+			if(FPhysScene* PhysScene = GetPhysicsScene(this))
+			{
+				PhysScene->AddForce_AssumesLocked(this, Force, bAllowSubstepping, bAccelChange);
+			}
 
-	}
+		}
 	});
 #endif // WITH_PHYSX
 
@@ -3463,11 +3473,12 @@ void FBodyInstance::AddForceAtPosition(const FVector& Force, const FVector& Posi
 	ExecuteOnPxRigidBodyReadWrite(this, [&](PxRigidBody* PRigidBody)
 	{
 		if (IsRigidBodyNonKinematic_AssumesLocked(PRigidBody))
-	{
-		const PxScene* PScene = PRigidBody->getScene();
-		FPhysScene* PhysScene = FPhysxUserData::Get<FPhysScene>(PScene->userData);
-			PhysScene->AddForceAtPosition_AssumesLocked(this, Force, Position, bAllowSubstepping);
-	}
+		{
+			if(FPhysScene* PhysScene = GetPhysicsScene(this))
+			{
+				PhysScene->AddForceAtPosition_AssumesLocked(this, Force, Position, bAllowSubstepping);
+			}
+		}
 	});
 	
 #endif // WITH_PHYSX
@@ -3491,13 +3502,13 @@ void FBodyInstance::AddTorque(const FVector& Torque, bool bAllowSubstepping, boo
 	ExecuteOnPxRigidBodyReadWrite(this, [&](PxRigidBody* PRigidBody)
 	{
 		if (IsRigidBodyNonKinematic_AssumesLocked(PRigidBody))
-	{
-		const PxScene* PScene = PRigidBody->getScene();
-		FPhysScene* PhysScene = FPhysxUserData::Get<FPhysScene>(PScene->userData);
-			PhysScene->AddTorque_AssumesLocked(this, Torque, bAllowSubstepping, bAccelChange);
-	}
+		{
+			if(FPhysScene* PhysScene = GetPhysicsScene(this))
+			{
+				PhysScene->AddTorque_AssumesLocked(this, Torque, bAllowSubstepping, bAccelChange);
+			}
+		}
 	});
-	
 #endif // WITH_PHYSX
 
 #if WITH_BOX2D
@@ -3694,12 +3705,13 @@ void FBodyInstance::AddRadialForceToBody(const FVector& Origin, float Radius, fl
 	ExecuteOnPxRigidBodyReadWrite(this, [&](PxRigidBody* PRigidBody)
 	{
 		if (IsRigidBodyNonKinematic_AssumesLocked(PRigidBody))
-	{
-		const PxScene* PScene = PRigidBody->getScene();
-		FPhysScene* PhysScene = FPhysxUserData::Get<FPhysScene>(PScene->userData);
-			PhysScene->AddRadialForceToBody_AssumesLocked(this, Origin, Radius, Strength, Falloff, bAccelChange, bAllowSubstepping);
+		{
+			if(FPhysScene* PhysScene = GetPhysicsScene(this))
+			{
+				PhysScene->AddRadialForceToBody_AssumesLocked(this, Origin, Radius, Strength, Falloff, bAccelChange, bAllowSubstepping);
+			}
 
-	}
+		}
 	});
 #endif // WITH_PHYSX
 
