Remove broken/unused function in prep for compacted pose

diff --git a/Engine/Source/Runtime/Engine/Public/ReferenceSkeleton.h b/Engine/Source/Runtime/Engine/Public/ReferenceSkeleton.h
--- a/Engine/Source/Runtime/Engine/Public/ReferenceSkeleton.h
+++ b/Engine/Source/Runtime/Engine/Public/ReferenceSkeleton.h
@@ -141,40 +141,6 @@ public:
 		check( ((BoneIndex == 0) && (BoneInfo.ParentIndex == INDEX_NONE)) || ((BoneIndex > 0) && RefBoneInfo.IsValidIndex(BoneInfo.ParentIndex)) );
 	}
 
-	/** Insert a new bone
-	* BoneName must not already exist! ParentIndex must be valid. */
-	void Insert(int32 BoneIndex, const FMeshBoneInfo& BoneInfo, const FTransform& BonePose)
-	{
-		// Make sure our arrays are in sync.
-		checkSlow( (RefBoneInfo.Num() == RefBonePose.Num()) && (RefBoneInfo.Num() == NameToIndexMap.Num()) );
-
-		// Inserting a bone that already exists is illegal
-		check(FindBoneIndex(BoneInfo.Name) == INDEX_NONE);
-
-		// Parent must be valid. Either INDEX_NONE for Root, or before us.
-		check( ((BoneIndex == 0) && (BoneInfo.ParentIndex == INDEX_NONE)) || ((BoneIndex > 0) && RefBoneInfo.IsValidIndex(BoneInfo.ParentIndex)) );
-
-		// Make sure our bone transform is valid.
-		check( BonePose.GetRotation().IsNormalized() );
-
-		RefBoneInfo.Insert(BoneInfo, BoneIndex);
-		RefBonePose.Insert(BonePose, BoneIndex);
-		NameToIndexMap.Add(BoneInfo.Name, BoneIndex);
-
-		// Normalize Quaternion to be safe.
-		RefBonePose[BoneIndex].NormalizeRotation();
-
-		// Now we need to fix all the parent indices that pointed to bones after this in the array
-		// These must be after this point in the array.
-		for(int32 j=BoneIndex+1; j<GetNum(); j++)
-		{
-			if( GetParentIndex(j) >= BoneIndex )
-			{
-				RefBoneInfo[j].ParentIndex += 1;
-			}
-		}
-	}
-
 	void Empty()
 	{
 		RefBoneInfo.Empty();
