Added an API on UGameInstance to make it easier to start recording, stop recording, and start playing back replays.
Fixes UE-6803, consumers no longer have to call Exec() to get this functionality.
Handlers for the "demorec", "demostop", and "demoplay" commands now use the new API.
NullNetworkReplayStreaming now stores and reports the replay's FriendlyName.

diff --git a/Engine/Source/Runtime/Engine/Classes/Engine/DemoNetDriver.h b/Engine/Source/Runtime/Engine/Classes/Engine/DemoNetDriver.h
--- a/Engine/Source/Runtime/Engine/Classes/Engine/DemoNetDriver.h
+++ b/Engine/Source/Runtime/Engine/Classes/Engine/DemoNetDriver.h
@@ -4,6 +4,8 @@
 
 #include "DemoNetDriver.generated.h"
 
+DECLARE_LOG_CATEGORY_EXTERN( LogDemo, Log, All );
+
 /**
  * Simulated network driver for recording and playing back game sessions.
  */
diff --git a/Engine/Source/Runtime/Engine/Classes/Engine/GameInstance.h b/Engine/Source/Runtime/Engine/Classes/Engine/GameInstance.h
--- a/Engine/Source/Runtime/Engine/Classes/Engine/GameInstance.h
+++ b/Engine/Source/Runtime/Engine/Classes/Engine/GameInstance.h
@@ -172,6 +172,15 @@ public:
 
 	inline FTimerManager& GetTimerManager() const { return *TimerManager; }
 
+	/** Start recording a replay with the given custom name and friendly name. */
+	virtual void StartRecordingReplay(const FString& Name, const FString& FriendlyName);
+
+	/** Stop recording a replay if one is currently in progress */
+	virtual void StopRecordingReplay();
+
+	/** Start playing back a previously recorded replay. */
+	virtual void PlayReplay(const FString& Name);
+
 private:
 	FTimerManager* TimerManager;
 };
\ No newline at end of file
diff --git a/Engine/Source/Runtime/Engine/Private/DemoNetDriver.cpp b/Engine/Source/Runtime/Engine/Private/DemoNetDriver.cpp
--- a/Engine/Source/Runtime/Engine/Private/DemoNetDriver.cpp
+++ b/Engine/Source/Runtime/Engine/Private/DemoNetDriver.cpp
@@ -21,7 +21,7 @@
 #include "Net/NetworkProfiler.h"
 #include "Net/DataReplication.h"
 
-DEFINE_LOG_CATEGORY_STATIC( LogDemo, Log, All );
+DEFINE_LOG_CATEGORY( LogDemo );
 
 static TAutoConsoleVariable<float> CVarDemoRecordHz( TEXT( "demo.RecordHz" ), 10, TEXT( "Number of demo frames recorded per second" ) );
 static TAutoConsoleVariable<float> CVarDemoTimeDilation( TEXT( "demo.TimeDilation" ), -1.0f, TEXT( "Override time dilation during demo playback (-1 = don't override)" ) );
@@ -393,7 +393,14 @@ bool UDemoNetDriver::InitListen( FNetworkNotify* InNotify, FURL& ListenURL, bool
 	Connection->InitSendBuffer();
 	ClientConnections.Add( Connection );
 
-	ReplayStreamer->StartStreaming( DemoFilename, World->GetMapName(), true, FNetworkVersion::GetReplayVersion(), FOnStreamReadyDelegate::CreateUObject( this, &UDemoNetDriver::ReplayStreamingReady ) );
+	const TCHAR* FriendlyNameOption = ListenURL.GetOption( TEXT("DemoFriendlyName="), nullptr );
+
+	ReplayStreamer->StartStreaming(
+		DemoFilename,
+		FriendlyNameOption != nullptr ? FString(FriendlyNameOption) : World->GetMapName(),
+		true,
+		FNetworkVersion::GetReplayVersion(),
+		FOnStreamReadyDelegate::CreateUObject( this, &UDemoNetDriver::ReplayStreamingReady ) );
 
 	FArchive* FileAr = ReplayStreamer->GetHeaderArchive();
 
diff --git a/Engine/Source/Runtime/Engine/Private/GameInstance.cpp b/Engine/Source/Runtime/Engine/Private/GameInstance.cpp
--- a/Engine/Source/Runtime/Engine/Private/GameInstance.cpp
+++ b/Engine/Source/Runtime/Engine/Private/GameInstance.cpp
@@ -8,6 +8,7 @@
 #include "EnginePrivate.h"
 #include "Engine/GameInstance.h"
 #include "Engine/Engine.h"
+#include "Engine/DemoNetDriver.h"
 #include "OnlineSubsystem.h"
 #include "OnlineSessionInterface.h"
 #include "GameFramework/OnlineSession.h"
@@ -672,3 +673,114 @@ void UGameInstance::OnSessionUserInviteAccepted(const bool bWasSuccess, const in
 		}
 	}
 }
+
+void UGameInstance::StartRecordingReplay(const FString& Name, const FString& FriendlyName)
+{
+	if ( FParse::Param( FCommandLine::Get(),TEXT( "NOREPLAYS" ) ) )
+	{
+		UE_LOG( LogDemo, Warning, TEXT( "UGameInstance::StartRecordingReplay: Rejected due to -noreplays option" ) );
+		return;
+	}
+
+	UWorld* CurrentWorld = GetWorld();
+
+	if ( CurrentWorld == nullptr )
+	{
+		UE_LOG( LogDemo, Warning, TEXT( "UGameInstance::StartRecordingReplay: GetWorld() is null" ) );
+		return;
+	}
+
+	FURL DemoURL;
+	FString DemoName = Name;
+	
+	DemoName.ReplaceInline( TEXT( "%m" ), *CurrentWorld->GetMapName() );
+
+	// replace the current URL's map with a demo extension
+	DemoURL.Map = DemoName;
+	DemoURL.AddOption( *FString::Printf( TEXT( "DemoFriendlyName=%s" ), *FriendlyName ) );
+
+	CurrentWorld->DestroyDemoNetDriver();
+
+	const FName NAME_DemoNetDriver( TEXT( "DemoNetDriver" ) );
+
+	if ( !GEngine->CreateNamedNetDriver( CurrentWorld, NAME_DemoNetDriver, NAME_DemoNetDriver ) )
+	{
+		UE_LOG( LogDemo, Warning, TEXT( "RecordReplay: failed to create demo net driver!" ) );
+		return;
+	}
+
+	CurrentWorld->DemoNetDriver = Cast< UDemoNetDriver >( GEngine->FindNamedNetDriver( CurrentWorld, NAME_DemoNetDriver ) );
+
+	check( CurrentWorld->DemoNetDriver != NULL );
+
+	CurrentWorld->DemoNetDriver->SetWorld( CurrentWorld );
+
+	FString Error;
+
+	if ( !CurrentWorld->DemoNetDriver->InitListen( CurrentWorld, DemoURL, false, Error ) )
+	{
+		UE_LOG( LogDemo, Warning, TEXT( "Demo recording failed: %s" ), *Error );
+		CurrentWorld->DemoNetDriver = NULL;
+	}
+	else
+	{
+		UE_LOG(LogDemo, Log, TEXT( "Num Network Actors: %i" ), CurrentWorld->NetworkActors.Num() );
+	}
+}
+
+void UGameInstance::StopRecordingReplay()
+{
+	UWorld* CurrentWorld = GetWorld();
+
+	if ( CurrentWorld == nullptr )
+	{
+		UE_LOG( LogDemo, Warning, TEXT( "UGameInstance::StopRecordingReplay: GetWorld() is null" ) );
+		return;
+	}
+
+	CurrentWorld->DestroyDemoNetDriver();
+}
+
+void UGameInstance::PlayReplay(const FString& Name)
+{
+	UWorld* CurrentWorld = GetWorld();
+
+	if ( CurrentWorld == nullptr )
+	{
+		UE_LOG( LogDemo, Warning, TEXT( "UGameInstance::PlayReplay: GetWorld() is null" ) );
+		return;
+	}
+
+	CurrentWorld->DestroyDemoNetDriver();
+
+	FURL DemoURL;
+	UE_LOG( LogDemo, Log, TEXT( "PlayReplay: Attempting to play demo %s" ), *Name );
+
+	DemoURL.Map = Name;
+
+	const FName NAME_DemoNetDriver( TEXT( "DemoNetDriver" ) );
+
+	if ( !GEngine->CreateNamedNetDriver( CurrentWorld, NAME_DemoNetDriver, NAME_DemoNetDriver ) )
+	{
+		UE_LOG(LogDemo, Warning, TEXT( "PlayReplay: failed to create demo net driver!" ) );
+		return;
+	}
+
+	CurrentWorld->DemoNetDriver = Cast< UDemoNetDriver >( GEngine->FindNamedNetDriver( CurrentWorld, NAME_DemoNetDriver ) );
+
+	check( CurrentWorld->DemoNetDriver != NULL );
+
+	CurrentWorld->DemoNetDriver->SetWorld( CurrentWorld );
+
+	FString Error;
+
+	if ( !CurrentWorld->DemoNetDriver->InitConnect( CurrentWorld, DemoURL, Error ) )
+	{
+		UE_LOG(LogDemo, Warning, TEXT( "Demo playback failed: %s" ), *Error );
+		CurrentWorld->DestroyDemoNetDriver();
+	}
+	else
+	{
+		FCoreUObjectDelegates::PostDemoPlay.Broadcast();
+	}
+}
diff --git a/Engine/Source/Runtime/Engine/Private/GameMode.cpp b/Engine/Source/Runtime/Engine/Private/GameMode.cpp
--- a/Engine/Source/Runtime/Engine/Private/GameMode.cpp
+++ b/Engine/Source/Runtime/Engine/Private/GameMode.cpp
@@ -20,6 +20,7 @@
 #include "GameFramework/CheatManager.h"
 #include "GameFramework/GameMode.h"
 #include "Engine/ChildConnection.h"
+#include "Engine/GameInstance.h"
 
 DEFINE_LOG_CATEGORY_STATIC(LogGameMode, Log, All);
 
@@ -627,9 +628,9 @@ void AGameMode::HandleMatchHasStarted()
 		}
 	}
 
-	if (IsHandlingReplays())
+	if (IsHandlingReplays() && GetGameInstance() != nullptr)
 	{
-		GEngine->Exec(GetWorld(), TEXT("DEMOREC %m"));
+		GetGameInstance()->StartRecordingReplay(GetWorld()->GetMapName(), GetWorld()->GetMapName());
 	}
 }
 
@@ -653,9 +654,9 @@ void AGameMode::HandleMatchHasEnded()
 {
 	GameSession->HandleMatchHasEnded();
 
-	if (IsHandlingReplays())
+	if (IsHandlingReplays() && GetGameInstance() != nullptr)
 	{
-		GEngine->Exec(GetWorld(), TEXT("DEMOSTOP"));
+		GetGameInstance()->StopRecordingReplay();
 	}
 }
 
diff --git a/Engine/Source/Runtime/Engine/Private/UnrealEngine.cpp b/Engine/Source/Runtime/Engine/Private/UnrealEngine.cpp
--- a/Engine/Source/Runtime/Engine/Private/UnrealEngine.cpp
+++ b/Engine/Source/Runtime/Engine/Private/UnrealEngine.cpp
@@ -9325,12 +9325,12 @@ bool UEngine::LoadMap( FWorldContext& WorldContext, FURL URL, class UPendingNetG
 	// See if we need to record network demos
 	if ( WorldContext.World()->GetAuthGameMode() == NULL || !WorldContext.World()->GetAuthGameMode()->IsHandlingReplays() )
 	{
-		if ( URL.HasOption( TEXT( "DemoRec" ) ) )
+		if ( URL.HasOption( TEXT( "DemoRec" ) ) && WorldContext.OwningGameInstance != nullptr )
 		{
 			const TCHAR* DemoRecName = URL.GetOption( TEXT( "DemoRec=" ), NULL );
 
 			// Record the demo, optionally with the specified custom name.
-			GEngine->Exec( WorldContext.World(), *FString::Printf( TEXT("DEMOREC %s"), DemoRecName != NULL ? DemoRecName : TEXT("") ) );
+			WorldContext.OwningGameInstance->StartRecordingReplay( FString(DemoRecName), WorldContext.World()->GetMapName() );
 		}
 	}
 
diff --git a/Engine/Source/Runtime/Engine/Private/World.cpp b/Engine/Source/Runtime/Engine/Private/World.cpp
--- a/Engine/Source/Runtime/Engine/Private/World.cpp
+++ b/Engine/Source/Runtime/Engine/Private/World.cpp
@@ -2680,48 +2680,16 @@ bool UWorld::HandleLogActorCountsCommand( const TCHAR* Cmd, FOutputDevice& Ar, U
 
 bool UWorld::HandleDemoRecordCommand( const TCHAR* Cmd, FOutputDevice& Ar, UWorld* InWorld )
 {
-	if ( FParse::Param( FCommandLine::Get(),TEXT( "NOREPLAYS" ) ) )
+	if ( InWorld != nullptr && InWorld->GetGameInstance() != nullptr )
 	{
-		UE_LOG( LogWorld, Warning, TEXT( "HandleDemoRecordCommand: Rejected due to -noreplays option" ) );
-		return true;
-	}
+		FString DemoName;
 
-	FURL DemoURL;
-	FString DemoName;
-
-	FParse::Token( Cmd, DemoName, 0 );
-	
-	DemoName.ReplaceInline( TEXT( "%m" ), *GetMapName() );
+		FParse::Token( Cmd, DemoName, 0 );
 
-	// replace the current URL's map with a demo extension
-	DemoURL.Map = DemoName;
+		// The friendly name will be the map name if no name is supplied
+		const FString FriendlyName = DemoName.IsEmpty() ? InWorld->GetMapName() : DemoName;
 
-	DestroyDemoNetDriver();
-
-	const FName NAME_DemoNetDriver( TEXT( "DemoNetDriver" ) );
-
-	if ( !GEngine->CreateNamedNetDriver( this, NAME_DemoNetDriver, NAME_DemoNetDriver ) )
-	{
-		Ar.Logf( TEXT( "Failed to create demo net driver!" ) );
-		return true;
-	}
-
-	DemoNetDriver = Cast< UDemoNetDriver >( GEngine->FindNamedNetDriver( this, NAME_DemoNetDriver ) );
-
-	check( DemoNetDriver != NULL );
-
-	DemoNetDriver->SetWorld( this );
-
-	FString Error;
-
-	if ( !DemoNetDriver->InitListen( this, DemoURL, false, Error ) )
-	{
-		Ar.Logf( TEXT( "Demo recording failed: %s" ), *Error );
-		DemoNetDriver = NULL;
-	}
-	else
-	{
-		Ar.Logf( TEXT( "Num Network Actors: %i" ), InWorld->NetworkActors.Num() );
+		InWorld->GetGameInstance()->StartRecordingReplay( DemoName, FriendlyName );
 	}
 
 	return true;
@@ -2730,62 +2698,33 @@ bool UWorld::HandleDemoRecordCommand( const TCHAR* Cmd, FOutputDevice& Ar, UWorl
 bool UWorld::HandleDemoPlayCommand( const TCHAR* Cmd, FOutputDevice& Ar, UWorld* InWorld )
 {
 	FString Temp;
+	const TCHAR* ErrorString = nullptr;
 
-	if ( FParse::Token( Cmd, Temp, 0 ) )
+	if ( !FParse::Token( Cmd, Temp, 0 ) )
 	{
-		DestroyDemoNetDriver();
-
-		FURL DemoURL;
-		UE_LOG( LogWorld, Log, TEXT( "Attempting to play demo %s" ), *Temp );
-
-		DemoURL.Map = Temp;
-
-#if 0
-		UPendingNetGame * NewPendingNetGame = new UDemoPendingNetGame( FObjectInitializer(), DemoURL );
-
-		if ( !NewPendingNetGame->DemoNetDriver )
-		{
-			Ar.Logf( TEXT( "Demo playback failed: %s" ), *NewPendingNetGame->ConnectionError );
-			NewPendingNetGame = NULL;
-		}
-
-		GEngine->CancelPending( InWorld, NewPendingNetGame );
-#else
-		const FName NAME_DemoNetDriver( TEXT( "DemoNetDriver" ) );
-
-		if ( !GEngine->CreateNamedNetDriver( this, NAME_DemoNetDriver, NAME_DemoNetDriver ) )
-		{
-			Ar.Logf( TEXT( "Failed to create demo net driver!" ) );
-			return true;
-		}
-
-		DemoNetDriver = Cast< UDemoNetDriver >( GEngine->FindNamedNetDriver( this, NAME_DemoNetDriver ) );
-
-		check( DemoNetDriver != NULL );
-
-		DemoNetDriver->SetWorld( this );
+		ErrorString = TEXT( "You must specify a filename" );
+	}
+	else if ( InWorld == nullptr )
+	{
+		ErrorString = TEXT( "InWorld is null" );
+	}
+	else if ( InWorld->GetGameInstance() == nullptr )
+	{
+		ErrorString = TEXT( "InWorld->GetGameInstance() is null" );
+	}
 
-		FString Error;
+	if ( ErrorString != nullptr )
+	{
+		Ar.Log( ErrorString );
 
-		if ( !DemoNetDriver->InitConnect( this, DemoURL, Error ) )
-		{
-			Ar.Logf( TEXT( "Demo playback failed: %s" ), *Error );
-			DestroyDemoNetDriver();
-		}
-		else
+		if ( GetGameInstance() != nullptr )
 		{
-			FCoreUObjectDelegates::PostDemoPlay.Broadcast();
+			GetGameInstance()->HandleDemoPlaybackFailure( EDemoPlayFailure::Generic, FString( ErrorString ) );
 		}
-#endif
 	}
 	else
 	{
-		Ar.Log( TEXT( "You must specify a filename" ) );
-
-		if ( GetGameInstance() != nullptr )
-		{
-			GetGameInstance()->HandleDemoPlaybackFailure( EDemoPlayFailure::Generic, FString( TEXT( "You must specify a filename " ) ) );
-		}
+		InWorld->GetGameInstance()->PlayReplay( Temp );
 	}
 
 	return true;
@@ -2793,7 +2732,11 @@ bool UWorld::HandleDemoPlayCommand( const TCHAR* Cmd, FOutputDevice& Ar, UWorld*
 
 bool UWorld::HandleDemoStopCommand( const TCHAR* Cmd, FOutputDevice& Ar, UWorld* InWorld )
 {
-	DestroyDemoNetDriver();
+	if ( InWorld != nullptr && InWorld->GetGameInstance() != nullptr )
+	{
+		InWorld->GetGameInstance()->StopRecordingReplay();
+	}
+
 	return true;
 }
 
diff --git a/Engine/Source/Runtime/NetworkReplayStreaming/NullNetworkReplayStreaming/Private/NullNetworkReplayStreaming.cpp b/Engine/Source/Runtime/NetworkReplayStreaming/NullNetworkReplayStreaming/Private/NullNetworkReplayStreaming.cpp
--- a/Engine/Source/Runtime/NetworkReplayStreaming/NullNetworkReplayStreaming/Private/NullNetworkReplayStreaming.cpp
+++ b/Engine/Source/Runtime/NetworkReplayStreaming/NullNetworkReplayStreaming/Private/NullNetworkReplayStreaming.cpp
@@ -212,6 +212,7 @@ void FNullNetworkReplayStreamer::StartStreaming( const FString& CustomName, cons
 		// Set up replay info
 		ReplayInfo.NetworkVersion = ReplayVersion.NetworkVersion;
 		ReplayInfo.Changelist = ReplayVersion.Changelist;
+		ReplayInfo.FriendlyName = FriendlyName;
 	}
 
 	// Notify immediately
@@ -352,6 +353,7 @@ void FNullNetworkReplayStreamer::EnumerateStreams( const FNetworkReplayVersion&
 				Info.Timestamp = IFileManager::Get().GetTimeStamp( *FullDemoFilePath );
 				Info.bIsLive = IsNamedStreamLive( Directory );
 				Info.LengthInMS = StoredReplayInfo.LengthInMS;
+				Info.FriendlyName = StoredReplayInfo.FriendlyName;
 
 				// Live streams not supported yet
 				if (!Info.bIsLive)
diff --git a/Engine/Source/Runtime/NetworkReplayStreaming/NullNetworkReplayStreaming/Public/NullNetworkReplayStreaming.h b/Engine/Source/Runtime/NetworkReplayStreaming/NullNetworkReplayStreaming/Public/NullNetworkReplayStreaming.h
--- a/Engine/Source/Runtime/NetworkReplayStreaming/NullNetworkReplayStreaming/Public/NullNetworkReplayStreaming.h
+++ b/Engine/Source/Runtime/NetworkReplayStreaming/NullNetworkReplayStreaming/Public/NullNetworkReplayStreaming.h
@@ -17,12 +17,14 @@ public:
 	int32		LengthInMS;
 	uint32		NetworkVersion;
 	uint32		Changelist;
-	
+	FString		FriendlyName;
+
 	// FOnlineJsonSerializable
 	BEGIN_ONLINE_JSON_SERIALIZER
 		ONLINE_JSON_SERIALIZE( "LengthInMS",		LengthInMS );
 		ONLINE_JSON_SERIALIZE( "NetworkVersion",	NetworkVersion );
 		ONLINE_JSON_SERIALIZE( "Changelist",		Changelist );
+		ONLINE_JSON_SERIALIZE( "FriendlyName",		FriendlyName );
 	END_ONLINE_JSON_SERIALIZER
 };
 
