UMG - Adding multicast events to a widget now creates them on the last uber event graph you had selected.

diff --git a/Engine/Source/Editor/UMGEditor/Private/Customizations/UMGDetailCustomizations.cpp b/Engine/Source/Editor/UMGEditor/Private/Customizations/UMGDetailCustomizations.cpp
--- a/Engine/Source/Editor/UMGEditor/Private/Customizations/UMGDetailCustomizations.cpp
+++ b/Engine/Source/Editor/UMGEditor/Private/Customizations/UMGDetailCustomizations.cpp
@@ -55,12 +55,8 @@ private:
 	{
 		UBlueprint* Blueprint = Editor.Pin()->GetBlueprintObj();
 
-		UEdGraph* TargetGraph = nullptr;
-		if ( Blueprint->UbergraphPages.Num() > 0 )
-		{
-			TargetGraph = Blueprint->UbergraphPages[0]; // Just use the first graph
-		}
-
+		UEdGraph* TargetGraph = Blueprint->GetLastEditedUberGraph();
+		
 		if ( TargetGraph != nullptr )
 		{
 			Editor.Pin()->SetCurrentMode(FWidgetBlueprintApplicationModes::GraphMode);
diff --git a/Engine/Source/Runtime/Engine/Classes/Engine/Blueprint.h b/Engine/Source/Runtime/Engine/Classes/Engine/Blueprint.h
--- a/Engine/Source/Runtime/Engine/Classes/Engine/Blueprint.h
+++ b/Engine/Source/Runtime/Engine/Classes/Engine/Blueprint.h
@@ -646,6 +646,12 @@ public:
 
 	static FName GetFunctionNameFromClassByGuid(const UClass* InClass, const FGuid FunctionGuid);
 	static bool GetFunctionGuidFromClassByFieldName(const UClass* InClass, const FName FunctionName, FGuid& FunctionGuid);
+
+	/**
+	 * Gets the last edited uber graph.  If no graph was found in the last edited document set, the first
+	 * ubergraph is returned.  If there are no ubergraphs nullptr is returned.
+	 */
+	UEdGraph* GetLastEditedUberGraph() const;
 #endif
 
 	/** Find a function given its name and optionally an object property name within this Blueprint */
diff --git a/Engine/Source/Runtime/Engine/Private/Blueprint.cpp b/Engine/Source/Runtime/Engine/Private/Blueprint.cpp
--- a/Engine/Source/Runtime/Engine/Private/Blueprint.cpp
+++ b/Engine/Source/Runtime/Engine/Private/Blueprint.cpp
@@ -1277,4 +1277,31 @@ bool UBlueprint::GetFunctionGuidFromClassByFieldName(const UClass* InClass, cons
 	return FBlueprintEditorUtils::GetFunctionGuidFromClassByFieldName(InClass, FunctionName, FunctionGuid);
 }
 
+UEdGraph* UBlueprint::GetLastEditedUberGraph() const
+{
+	for ( int32 LastEditedIndex = LastEditedDocuments.Num() - 1; LastEditedIndex >= 0; LastEditedIndex-- )
+	{
+		if ( UObject* Obj = LastEditedDocuments[LastEditedIndex].EditedObject )
+		{
+			if ( UEdGraph* Graph = Cast<UEdGraph>(Obj) )
+			{
+				for ( int32 GraphIndex = 0; GraphIndex < UbergraphPages.Num(); GraphIndex++ )
+				{
+					if ( Graph == UbergraphPages[GraphIndex] )
+					{
+						return UbergraphPages[GraphIndex];
+					}
+				}
+			}
+		}
+	}
+
+	if ( UbergraphPages.Num() > 0 )
+	{
+		return UbergraphPages[0];
+	}
+
+	return nullptr;
+}
+
 #endif //WITH_EDITOR
\ No newline at end of file
