Dont flush texture streaming when resizing the viewport

diff --git a/Engine/Source/Runtime/Engine/Private/Slate/SceneViewport.cpp b/Engine/Source/Runtime/Engine/Private/Slate/SceneViewport.cpp
--- a/Engine/Source/Runtime/Engine/Private/Slate/SceneViewport.cpp
+++ b/Engine/Source/Runtime/Engine/Private/Slate/SceneViewport.cpp
@@ -1118,9 +1118,6 @@ void FSceneViewport::SetRenderTargetTextureRenderThread(FTexture2DRHIRef& RT)
 
 void FSceneViewport::UpdateViewportRHI(bool bDestroyed, uint32 NewSizeX, uint32 NewSizeY, EWindowMode::Type NewWindowMode)
 {
-	// Make sure we're not in the middle of streaming textures.
-	(*GFlushStreamingFunc)();
-
 	{
 		SCOPED_SUSPEND_RENDERING_THREAD(true);
 
