Only report spatialization error once per asset, not once per wave instance

diff --git a/Engine/Source/Runtime/Engine/Private/SoundWave.cpp b/Engine/Source/Runtime/Engine/Private/SoundWave.cpp
--- a/Engine/Source/Runtime/Engine/Private/SoundWave.cpp
+++ b/Engine/Source/Runtime/Engine/Private/SoundWave.cpp
@@ -648,21 +648,27 @@ void USoundWave::Parse( FAudioDevice* AudioDevice, const UPTRINT NodeWaveInstanc
 		// Sanity check
 		if( NumChannels > 2 && WaveInstance->bUseSpatialization && !WaveInstance->bReportedSpatializationWarning)
 		{
-			FString SoundWarningInfo = FString::Printf(TEXT("Spatialisation on stereo and multichannel sounds is not supported. SoundWave: %s"), *GetName());
-			if (ActiveSound.Sound != this)
+			static TSet<USoundWave*> ReportedSounds;
+			if (!ReportedSounds.Contains(this))
 			{
-				SoundWarningInfo += FString::Printf(TEXT(" SoundCue: %s"), *ActiveSound.Sound->GetName());
-			}
+				FString SoundWarningInfo = FString::Printf(TEXT("Spatialisation on stereo and multichannel sounds is not supported. SoundWave: %s"), *GetName());
+				if (ActiveSound.Sound != this)
+				{
+					SoundWarningInfo += FString::Printf(TEXT(" SoundCue: %s"), *ActiveSound.Sound->GetName());
+				}
 
-			if (ActiveSound.AudioComponent.IsValid())
-			{
-				// TODO - Audio Threading. This log would have to be a task back to game thread
-				AActor* SoundOwner = ActiveSound.AudioComponent->GetOwner();
-				UE_LOG(LogAudio, Warning, TEXT( "%s Actor: %s AudioComponent: %s" ), *SoundWarningInfo, (SoundOwner ? *SoundOwner->GetName() : TEXT("None")), *ActiveSound.AudioComponent->GetName() );
-			}
-			else
-			{
-				UE_LOG(LogAudio, Warning, TEXT("%s"), *SoundWarningInfo );
+				if (ActiveSound.AudioComponent.IsValid())
+				{
+					// TODO - Audio Threading. This log would have to be a task back to game thread
+					AActor* SoundOwner = ActiveSound.AudioComponent->GetOwner();
+					UE_LOG(LogAudio, Warning, TEXT( "%s Actor: %s AudioComponent: %s" ), *SoundWarningInfo, (SoundOwner ? *SoundOwner->GetName() : TEXT("None")), *ActiveSound.AudioComponent->GetName() );
+				}
+				else
+				{
+					UE_LOG(LogAudio, Warning, TEXT("%s"), *SoundWarningInfo );
+				}
+
+				ReportedSounds.Add(this);
 			}
 			WaveInstance->bReportedSpatializationWarning = true;
 		}
