Implements the Log2 material node

diff --git a/Engine/Documentation/Source/Shared/GraphNodes/Material/Material.INT.udn b/Engine/Documentation/Source/Shared/GraphNodes/Material/Material.INT.udn
--- a/Engine/Documentation/Source/Shared/GraphNodes/Material/Material.INT.udn
+++ b/Engine/Documentation/Source/Shared/GraphNodes/Material/Material.INT.udn
@@ -617,6 +617,14 @@ Engine/Rendering/Materials/ExpressionReference/Math#power
 [/VAR]
 [/EXCERPT:UMaterialExpressionPower]
 
+[EXCERPT:UMaterialExpressionLogarithm2]
+The Logarithm expression takes two inputs, compute X's logarithm base 2, and outputs the result; in other words, the number of time you should get 2 multiply itself to find X. The logarithm computation happens per channel. Both inputs must have the same number of values unless one of the values is a single float value. 
+
+[VAR:ToolTipFullLink]
+Engine/Rendering/Materials/ExpressionReference/Math#Logarithm
+[/VAR]
+[/EXCERPT:UMaterialExpressionLogarithm2]
+
 [EXCERPT:UMaterialExpressionSine]
 The Sine expression outputs the sine of the value input (in radians). Most commonly, this is used to output a continuous oscillating waveform by connecting a Time expression to its input. The output value will cycle back and forth between -1 and 1. The difference between this and the output of the Cosine expression is the output waveform is offset half the period.
 
diff --git a/Engine/Source/Runtime/Engine/Classes/Materials/MaterialExpressionLogarithm2.h b/Engine/Source/Runtime/Engine/Classes/Materials/MaterialExpressionLogarithm2.h
new file mode 100644
--- /dev/null
+++ b/Engine/Source/Runtime/Engine/Classes/Materials/MaterialExpressionLogarithm2.h
@@ -0,0 +1,24 @@
+// Copyright 1998-2015 Epic Games, Inc. All Rights Reserved.
+
+
+#pragma once
+#include "Materials/MaterialExpression.h"
+#include "MaterialExpressionLogarithm2.generated.h"
+
+UCLASS(collapsecategories, hidecategories=Object, MinimalAPI)
+class UMaterialExpressionLogarithm2 : public UMaterialExpression
+{
+	GENERATED_UCLASS_BODY()
+	
+	UPROPERTY()
+	FExpressionInput X;
+
+	// Begin UMaterialExpression Interface
+	virtual int32 Compile(class FMaterialCompiler* Compiler, int32 OutputIndex, int32 MultiplexIndex) override;
+	virtual void GetCaption(TArray<FString>& OutCaptions) const override;
+	// End UMaterialExpression Interface
+
+};
+
+
+
diff --git a/Engine/Source/Runtime/Engine/Private/Materials/HLSLMaterialTranslator.h b/Engine/Source/Runtime/Engine/Private/Materials/HLSLMaterialTranslator.h
--- a/Engine/Source/Runtime/Engine/Private/Materials/HLSLMaterialTranslator.h
+++ b/Engine/Source/Runtime/Engine/Private/Materials/HLSLMaterialTranslator.h
@@ -2902,6 +2902,21 @@ protected:
 		// use ClampedPow so artist are prevented to cause NAN creeping into the math
 		return AddCodeChunk(GetParameterType(Base),TEXT("ClampedPow(%s,%s)"),*GetParameterCode(Base),*CoerceParameter(Exponent,MCT_Float));
 	}
+	
+	virtual int32 Logarithm2(int32 X) override
+	{
+		if(X == INDEX_NONE)
+		{
+			return INDEX_NONE;
+		}
+		
+		if(GetParameterUniformExpression(X))
+		{
+			return AddUniformExpression(new FMaterialUniformExpressionLogarithm2(GetParameterUniformExpression(X)),GetParameterType(X),TEXT("log2(%s)"),*GetParameterCode(X));
+		}
+
+		return AddCodeChunk(GetParameterType(X),TEXT("log2(%s)"),*GetParameterCode(X));
+	}
 
 	virtual int32 SquareRoot(int32 X) override
 	{
diff --git a/Engine/Source/Runtime/Engine/Private/Materials/MaterialExpressions.cpp b/Engine/Source/Runtime/Engine/Private/Materials/MaterialExpressions.cpp
--- a/Engine/Source/Runtime/Engine/Private/Materials/MaterialExpressions.cpp
+++ b/Engine/Source/Runtime/Engine/Private/Materials/MaterialExpressions.cpp
@@ -58,6 +58,7 @@
 #include "Materials/MaterialExpressionLightmassReplace.h"
 #include "Materials/MaterialExpressionLightVector.h"
 #include "Materials/MaterialExpressionLinearInterpolate.h"
+#include "Materials/MaterialExpressionLogarithm2.h"
 #include "Materials/MaterialExpressionMakeMaterialAttributes.h"
 #include "Materials/MaterialExpressionMaterialFunctionCall.h"
 #include "Materials/MaterialExpressionMax.h"
@@ -4996,6 +4997,38 @@ void UMaterialExpressionPower::GetCaption(TArray<FString>& OutCaptions) const
 	OutCaptions.Add(ret);
 }
 
+UMaterialExpressionLogarithm2::UMaterialExpressionLogarithm2(const FObjectInitializer& ObjectInitializer)
+	: Super(ObjectInitializer)
+{
+	// Structure to hold one-time initialization
+	struct FConstructorStatics
+	{
+		FString NAME_Math;
+		FConstructorStatics()
+			: NAME_Math(LOCTEXT( "Math", "Math" ).ToString())
+		{
+		}
+	};
+	static FConstructorStatics ConstructorStatics;
+
+	MenuCategories.Add(ConstructorStatics.NAME_Math);
+}
+
+int32 UMaterialExpressionLogarithm2::Compile(class FMaterialCompiler* Compiler, int32 OutputIndex, int32 MultiplexIndex)
+{
+	if(!X.Expression)
+	{
+		return Compiler->Errorf(TEXT("Missing Log2 X input"));
+	}
+
+	return Compiler->Logarithm2(X.Compile(Compiler));
+}
+
+void UMaterialExpressionLogarithm2::GetCaption(TArray<FString>& OutCaptions) const
+{
+	OutCaptions.Add(TEXT("Log2"));
+}
+
 FString UMaterialExpressionSceneColor::GetInputName(int32 InputIndex) const
 {
 	if(InputIndex == 0)
diff --git a/Engine/Source/Runtime/Engine/Private/Materials/MaterialUniformExpressions.cpp b/Engine/Source/Runtime/Engine/Private/Materials/MaterialUniformExpressions.cpp
--- a/Engine/Source/Runtime/Engine/Private/Materials/MaterialUniformExpressions.cpp
+++ b/Engine/Source/Runtime/Engine/Private/Materials/MaterialUniformExpressions.cpp
@@ -508,6 +508,7 @@ IMPLEMENT_MATERIALUNIFORMEXPRESSION_TYPE(FMaterialUniformExpressionFlipBookTextu
 IMPLEMENT_MATERIALUNIFORMEXPRESSION_TYPE(FMaterialUniformExpressionSine);
 IMPLEMENT_MATERIALUNIFORMEXPRESSION_TYPE(FMaterialUniformExpressionSquareRoot);
 IMPLEMENT_MATERIALUNIFORMEXPRESSION_TYPE(FMaterialUniformExpressionLength);
+IMPLEMENT_MATERIALUNIFORMEXPRESSION_TYPE(FMaterialUniformExpressionLogarithm2);
 IMPLEMENT_MATERIALUNIFORMEXPRESSION_TYPE(FMaterialUniformExpressionFoldedMath);
 IMPLEMENT_MATERIALUNIFORMEXPRESSION_TYPE(FMaterialUniformExpressionPeriodic);
 IMPLEMENT_MATERIALUNIFORMEXPRESSION_TYPE(FMaterialUniformExpressionAppendVector);
diff --git a/Engine/Source/Runtime/Engine/Private/Materials/MaterialUniformExpressions.h b/Engine/Source/Runtime/Engine/Private/Materials/MaterialUniformExpressions.h
--- a/Engine/Source/Runtime/Engine/Private/Materials/MaterialUniformExpressions.h
+++ b/Engine/Source/Runtime/Engine/Private/Materials/MaterialUniformExpressions.h
@@ -501,6 +501,55 @@ private:
 	TRefCountPtr<FMaterialUniformExpression> X;
 };
 
+/**
+ */
+class FMaterialUniformExpressionLogarithm2: public FMaterialUniformExpression
+{
+	DECLARE_MATERIALUNIFORMEXPRESSION_TYPE(FMaterialUniformExpressionLogarithm2);
+public:
+
+	FMaterialUniformExpressionLogarithm2() {}
+	FMaterialUniformExpressionLogarithm2(FMaterialUniformExpression* InX):
+		X(InX)
+	{}
+
+	// FMaterialUniformExpression interface.
+	void Serialize(FArchive& Ar) override
+	{
+		Ar << X;
+	}
+	void GetNumberValue(const FMaterialRenderContext& Context,FLinearColor& OutValue) const override
+	{
+		FLinearColor ValueX = FLinearColor::Black;
+		X->GetNumberValue(Context,ValueX);
+		OutValue.R = FMath::Log2(ValueX.R);
+		OutValue.G = FMath::Log2(ValueX.G);
+		OutValue.B = FMath::Log2(ValueX.B);
+		OutValue.A = FMath::Log2(ValueX.A);
+	}
+	bool IsConstant() const override
+	{
+		return X->IsConstant();
+	}
+	bool IsChangingPerFrame() const override
+	{
+		return X->IsChangingPerFrame();
+	}
+	bool IsIdentical(const FMaterialUniformExpression* OtherExpression) const override
+	{
+		if (GetType() != OtherExpression->GetType())
+		{
+			return false;
+		}
+
+		auto OtherLog = static_cast<const FMaterialUniformExpressionLogarithm2 *>(OtherExpression);
+		return X->IsIdentical(OtherLog->X);
+	}
+
+private:
+	TRefCountPtr<FMaterialUniformExpression> X;
+};
+
 /**
  */
 enum EFoldedMathOperation
diff --git a/Engine/Source/Runtime/Engine/Public/MaterialCompiler.h b/Engine/Source/Runtime/Engine/Public/MaterialCompiler.h
--- a/Engine/Source/Runtime/Engine/Public/MaterialCompiler.h
+++ b/Engine/Source/Runtime/Engine/Public/MaterialCompiler.h
@@ -135,6 +135,7 @@ public:
 	virtual int32 Cross(int32 A,int32 B) = 0;
 
 	virtual int32 Power(int32 Base,int32 Exponent) = 0;
+	virtual int32 Logarithm2(int32 X) = 0;
 	virtual int32 SquareRoot(int32 X) = 0;
 	virtual int32 Length(int32 X) = 0;
 
@@ -290,6 +291,7 @@ public:
 	virtual int32 Cross(int32 A,int32 B) override { return Compiler->Cross(A,B); }
 
 	virtual int32 Power(int32 Base,int32 Exponent) override { return Compiler->Power(Base,Exponent); }
+	virtual int32 Logarithm2(int32 X) override { return Compiler->Logarithm2(X); }
 	virtual int32 SquareRoot(int32 X) override { return Compiler->SquareRoot(X); }
 	virtual int32 Length(int32 X) override { return Compiler->Length(X); }
 
