#HLOD

- Fixed so that user can't open editor with the LOD generated assets.

diff --git a/Engine/Source/Runtime/Engine/Private/LODActor.cpp b/Engine/Source/Runtime/Engine/Private/LODActor.cpp
--- a/Engine/Source/Runtime/Engine/Private/LODActor.cpp
+++ b/Engine/Source/Runtime/Engine/Private/LODActor.cpp
@@ -55,10 +55,6 @@ bool ALODActor::GetReferencedContentObjects( TArray<UObject*>& Objects ) const
 {
 	Super::GetReferencedContentObjects(Objects);
 
-	if (StaticMeshComponent && StaticMeshComponent->StaticMesh)
-	{
-		Objects.Add(StaticMeshComponent->StaticMesh);
-	}
 	return true;
 }
 
