Corrected UDeveloperSettings:::Get...Name() use of EditorUserSettings.
EditorUserSettings should no longer exist, in favor of EditorPerProjectUserSettings (which is already handled by these functions). Changed EditorUserSettings to EditorSettings to catch uses of EditorSettings classes

diff --git a/Engine/Source/Runtime/Engine/Private/DeveloperSettings.cpp b/Engine/Source/Runtime/Engine/Private/DeveloperSettings.cpp
--- a/Engine/Source/Runtime/Engine/Private/DeveloperSettings.cpp
+++ b/Engine/Source/Runtime/Engine/Private/DeveloperSettings.cpp
@@ -15,12 +15,12 @@ FName UDeveloperSettings::GetContainerName() const
 	static const FName ProjectName(TEXT("Project"));
 	static const FName EditorName(TEXT("Editor"));
 
-	static const FName EditorUserSettingsName(TEXT("EditorUserSettings"));
+	static const FName EditorSettingsName(TEXT("EditorSettings"));
 	static const FName EditorPerProjectUserSettingsName(TEXT("EditorPerProjectUserSettings"));
 
 	FName ConfigName = GetClass()->ClassConfigName;
 
-	if ( ConfigName == EditorUserSettingsName || ConfigName == EditorPerProjectUserSettingsName )
+	if ( ConfigName == EditorSettingsName || ConfigName == EditorPerProjectUserSettingsName )
 	{
 		return EditorName;
 	}
@@ -31,7 +31,7 @@ FName UDeveloperSettings::GetContainerName() const
 FName UDeveloperSettings::GetCategoryName() const
 {
 	static const FName GeneralName(TEXT("General"));
-	static const FName EditorUserSettingsName(TEXT("EditorUserSettings"));
+	static const FName EditorSettingsName(TEXT("EditorSettings"));
 	static const FName EditorPerProjectUserSettingsName(TEXT("EditorPerProjectUserSettings"));
 
 	if ( CategoryName != NAME_None )
@@ -44,7 +44,7 @@ FName UDeveloperSettings::GetCategoryName() const
 	{
 		return NAME_Engine;
 	}
-	else if ( ConfigName == EditorUserSettingsName || ConfigName == EditorPerProjectUserSettingsName )
+	else if ( ConfigName == EditorSettingsName || ConfigName == EditorPerProjectUserSettingsName )
 	{
 		return GeneralName;
 	}
