Fixed scene capture components not updating correctly in the editor or when they are part of a blueprint

diff --git a/Engine/Source/Runtime/Engine/Classes/Components/SceneCaptureComponent2D.h b/Engine/Source/Runtime/Engine/Classes/Components/SceneCaptureComponent2D.h
--- a/Engine/Source/Runtime/Engine/Classes/Components/SceneCaptureComponent2D.h
+++ b/Engine/Source/Runtime/Engine/Classes/Components/SceneCaptureComponent2D.h
@@ -45,6 +45,7 @@ private:
 
 public:
 	// Begin UActorComponent Interface
+	virtual void OnRegister() override;
 	virtual void SendRenderTransform_Concurrent() override;
 	virtual bool RequiresGameThreadEndOfFrameUpdates() const override
 	{
@@ -61,6 +62,7 @@ public:
 	// End UObject Interface
 
 	/** Render the scene to the texture */
+	UFUNCTION(BlueprintCallable,Category = "Rendering|SceneCapture")
 	ENGINE_API void UpdateContent();
 
 	ENGINE_API static void UpdateDeferredCaptures( FSceneInterface* Scene );
diff --git a/Engine/Source/Runtime/Engine/Classes/Components/SceneCaptureComponentCube.h b/Engine/Source/Runtime/Engine/Classes/Components/SceneCaptureComponentCube.h
--- a/Engine/Source/Runtime/Engine/Classes/Components/SceneCaptureComponentCube.h
+++ b/Engine/Source/Runtime/Engine/Classes/Components/SceneCaptureComponentCube.h
@@ -18,6 +18,7 @@ class USceneCaptureComponentCube : public USceneCaptureComponent
 
 public:
 	// Begin UActorComponent Interface
+	virtual void OnRegister() override;
 	virtual void SendRenderTransform_Concurrent() override;
 	virtual void TickComponent(float DeltaTime, enum ELevelTick TickType, FActorComponentTickFunction* ThisTickFunction) override;
 	// End UActorComponent Interface
@@ -29,6 +30,7 @@ public:
 	// End UObject Interface
 
 	/** Render the scene to the texture */
+	UFUNCTION(BlueprintCallable,Category = "Rendering|SceneCapture")
 	void UpdateContent();
 
 	ENGINE_API static void UpdateDeferredCaptures(FSceneInterface* Scene);
diff --git a/Engine/Source/Runtime/Engine/Private/Components/SceneCaptureComponent.cpp b/Engine/Source/Runtime/Engine/Private/Components/SceneCaptureComponent.cpp
--- a/Engine/Source/Runtime/Engine/Private/Components/SceneCaptureComponent.cpp
+++ b/Engine/Source/Runtime/Engine/Private/Components/SceneCaptureComponent.cpp
@@ -272,12 +272,25 @@ USceneCaptureComponent2D::USceneCaptureComponent2D(const FObjectInitializer& Obj
 	bAutoActivate = true;
 	PrimaryComponentTick.bCanEverTick = true;
 	PrimaryComponentTick.TickGroup = TG_DuringPhysics;
+	// Tick in the editor so that bCaptureEveryFrame preview works
+	bTickInEditor = true;
 	// previous behavior was to capture from raw scene color 
 	CaptureSource = SCS_SceneColorHDR;
 	// default to full blend weight..
 	PostProcessBlendWeight = 1.0f;
 }
 
+void USceneCaptureComponent2D::OnRegister()
+{
+	Super::OnRegister();
+
+#if WITH_EDITOR
+	// Update content on register to have at least one frames worth of good data.
+	// Without updating here this component would not work in a blueprint construction script which recreates the component after each move in the editor
+	UpdateContent();
+#endif
+}
+
 void USceneCaptureComponent2D::SendRenderTransform_Concurrent()
 {	
 	UpdateContent();
@@ -295,7 +308,7 @@ void USceneCaptureComponent2D::TickComponent(float DeltaTime, enum ELevelTick Ti
 	}
 }
 
-static TMultiMap<TWeakObjectPtr<UWorld>, USceneCaptureComponent2D*> SceneCapturesToUpdateMap;
+static TMultiMap<TWeakObjectPtr<UWorld>, TWeakObjectPtr<USceneCaptureComponent2D>> SceneCapturesToUpdateMap;
 
 void USceneCaptureComponent2D::UpdateContent()
 {
@@ -313,16 +326,19 @@ void USceneCaptureComponent2D::UpdateDeferredCaptures( FSceneInterface* Scene )
 	{
 		// Only update the scene captures assoicated with the current scene.
 		// Updating others not associated with the scene would cause invalid data to be rendered into the target
-		TArray<USceneCaptureComponent2D*> SceneCapturesToUpdate;
+		TArray< TWeakObjectPtr<USceneCaptureComponent2D> > SceneCapturesToUpdate;
 		SceneCapturesToUpdateMap.MultiFind( World, SceneCapturesToUpdate );
 		
-		for( USceneCaptureComponent2D* Component : SceneCapturesToUpdate )
+		for( TWeakObjectPtr<USceneCaptureComponent2D> Component : SceneCapturesToUpdate )
 		{
-			Scene->UpdateSceneCaptureContents( Component );
+			if( Component.IsValid() )
+			{
+				Scene->UpdateSceneCaptureContents( Component.Get() );
+			}
 		}
 		
 		// All scene captures for this world have been updated
-		SceneCapturesToUpdateMap.Remove( World) ;
+		SceneCapturesToUpdateMap.Remove( World );
 	}
 }
 
@@ -346,6 +362,17 @@ USceneCaptureComponentCube::USceneCaptureComponentCube(const FObjectInitializer&
 	bAutoActivate = true;
 	PrimaryComponentTick.bCanEverTick = true;
 	PrimaryComponentTick.TickGroup = TG_DuringPhysics;
+	bTickInEditor = true;
+}
+
+void USceneCaptureComponentCube::OnRegister()
+{
+	Super::OnRegister();
+#if WITH_EDITOR
+	// Update content on register to have at least one frames worth of good data.
+	// Without updating here this component would not work in a blueprint construction script which recreates the component after each move in the editor
+	UpdateContent();
+#endif
 }
 
 void USceneCaptureComponentCube::SendRenderTransform_Concurrent()
@@ -365,7 +392,7 @@ void USceneCaptureComponentCube::TickComponent(float DeltaTime, enum ELevelTick
 	}
 }
 
-static TMultiMap<TWeakObjectPtr<UWorld>, USceneCaptureComponentCube*> CubedSceneCapturesToUpdateMap;
+static TMultiMap<TWeakObjectPtr<UWorld>, TWeakObjectPtr<USceneCaptureComponentCube> > CubedSceneCapturesToUpdateMap;
 
 void USceneCaptureComponentCube::UpdateContent()
 {
@@ -382,14 +409,17 @@ void USceneCaptureComponentCube::UpdateDeferredCaptures( FSceneInterface* Scene
 	
 	if( World && CubedSceneCapturesToUpdateMap.Num() > 0 )
 	{
-		// Only update the scene captures assoicated with the current scene.
+		// Only update the scene captures associated with the current scene.
 		// Updating others not associated with the scene would cause invalid data to be rendered into the target
-		TArray<USceneCaptureComponentCube*> SceneCapturesToUpdate;
+		TArray< TWeakObjectPtr<USceneCaptureComponentCube> > SceneCapturesToUpdate;
 		CubedSceneCapturesToUpdateMap.MultiFind( World, SceneCapturesToUpdate );
 		
-		for( USceneCaptureComponentCube* Component : SceneCapturesToUpdate )
+		for( TWeakObjectPtr<USceneCaptureComponentCube> Component : SceneCapturesToUpdate )
 		{
-			Scene->UpdateSceneCaptureContents( Component );
+			if( Component.IsValid() )
+			{
+				Scene->UpdateSceneCaptureContents( Component.Get() );
+			}
 		}
 		
 		// All scene captures for this world have been updated
