Merging CL 2522630
Added UObject::PreDestroyFromReplication()

diff --git a/Engine/Source/Runtime/CoreUObject/Private/UObject/Obj.cpp b/Engine/Source/Runtime/CoreUObject/Private/UObject/Obj.cpp
--- a/Engine/Source/Runtime/CoreUObject/Private/UObject/Obj.cpp
+++ b/Engine/Source/Runtime/CoreUObject/Private/UObject/Obj.cpp
@@ -3757,6 +3757,12 @@ void UObject::PostNetReceive()
 
 }
 
+/** Called right before being marked for destruction due to network replication */
+void UObject::PreDestroyFromReplication()
+{
+
+}
+
 /** IsNameStableForNetworking means an object can be referred to its path name (relative to outer) over the network */
 bool UObject::IsNameStableForNetworking() const
 {
diff --git a/Engine/Source/Runtime/CoreUObject/Public/UObject/UObject.h b/Engine/Source/Runtime/CoreUObject/Public/UObject/UObject.h
--- a/Engine/Source/Runtime/CoreUObject/Public/UObject/UObject.h
+++ b/Engine/Source/Runtime/CoreUObject/Public/UObject/UObject.h
@@ -628,6 +628,9 @@ public:
 	/** Called right after receiving a bunch */
 	virtual void PostNetReceive();
 
+	/** Called right before being marked for destruction due to network replication */
+	virtual void PreDestroyFromReplication();
+
 	/**
  	 *******************************************************
 	 * Non virtual functions, not intended to be overridden
diff --git a/Engine/Source/Runtime/Engine/Private/DataChannel.cpp b/Engine/Source/Runtime/Engine/Private/DataChannel.cpp
--- a/Engine/Source/Runtime/Engine/Private/DataChannel.cpp
+++ b/Engine/Source/Runtime/Engine/Private/DataChannel.cpp
@@ -1445,6 +1445,7 @@ void UActorChannel::DestroyActorAndComponents()
 		{
 			UObject *SubObject = CreateSubObjects[i].Get();
 			Actor->OnSubobjectDestroyFromReplication(SubObject);
+			SubObject->PreDestroyFromReplication();
 			SubObject->MarkPendingKill();
 		}
 	}
@@ -1454,6 +1455,7 @@ void UActorChannel::DestroyActorAndComponents()
 	// Destroy the actor
 	if ( Actor != NULL )
 	{
+		Actor->PreDestroyFromReplication();
 		Actor->Destroy( true );
 	}
 }
@@ -2502,6 +2504,7 @@ UObject* UActorChannel::ReadContentBlockHeader(FInBunch & Bunch, bool& bObjectDe
 
 			Actor->OnSubobjectDestroyFromReplication( SubObj );
 
+			SubObj->PreDestroyFromReplication();
 			SubObj->MarkPendingKill();
 		}
 		bObjectDeleted = true;
