Merging replay spectator changes from dev branch

diff --git a/Engine/Source/Runtime/Engine/Classes/Engine/DemoNetDriver.h b/Engine/Source/Runtime/Engine/Classes/Engine/DemoNetDriver.h
--- a/Engine/Source/Runtime/Engine/Classes/Engine/DemoNetDriver.h
+++ b/Engine/Source/Runtime/Engine/Classes/Engine/DemoNetDriver.h
@@ -47,9 +47,6 @@ class ENGINE_API UDemoNetDriver : public UNetDriver
 	/** This is our spectator controller that is used to view the demo world from */
 	APlayerController* SpectatorController;
 
-	UPROPERTY(config)
-	FString DemoSpectatorClass;
-
 	/** Our network replay streamer */
 	TSharedPtr< class INetworkReplayStreamer >	ReplayStreamer;
 
diff --git a/Engine/Source/Runtime/Engine/Classes/GameFramework/GameMode.h b/Engine/Source/Runtime/Engine/Classes/GameFramework/GameMode.h
--- a/Engine/Source/Runtime/Engine/Classes/GameFramework/GameMode.h
+++ b/Engine/Source/Runtime/Engine/Classes/GameFramework/GameMode.h
@@ -229,6 +229,10 @@ public:
 	UPROPERTY(EditAnywhere, noclear, BlueprintReadOnly, Category=Classes)
 	TSubclassOf<class ASpectatorPawn> SpectatorClass;
 
+	/** The PlayerController class used when spectating a network replay. */
+	UPROPERTY(EditAnywhere, noclear, BlueprintReadOnly, Category = Classes)
+	TSubclassOf<class APlayerController> ReplaySpectatorPlayerControllerClass;
+
 	/** A PlayerState of this class will be associated with every player to replicate relevant player information to all clients. */
 	UPROPERTY(EditAnywhere, noclear, BlueprintReadOnly, Category=Classes)
 	TSubclassOf<class APlayerState> PlayerStateClass;
diff --git a/Engine/Source/Runtime/Engine/Classes/GameFramework/PlayerController.h b/Engine/Source/Runtime/Engine/Classes/GameFramework/PlayerController.h
--- a/Engine/Source/Runtime/Engine/Classes/GameFramework/PlayerController.h
+++ b/Engine/Source/Runtime/Engine/Classes/GameFramework/PlayerController.h
@@ -1457,7 +1457,8 @@ public:
 	// Spectating
 
 	/** Get the Pawn used when spectating. NULL when not spectating. */
-	FORCEINLINE class ASpectatorPawn* GetSpectatorPawn() const { return SpectatorPawn; }
+	UFUNCTION(BlueprintCallable, Category = "Pawn")
+	class ASpectatorPawn* GetSpectatorPawn() const { return SpectatorPawn; }
 
 	/** Returns the first of GetPawn() or GetSpectatorPawn() that is not NULL, or NULL otherwise. */
 	class APawn* GetPawnOrSpectator() const;
@@ -1491,7 +1492,6 @@ protected:
 
 private:
 	/** The pawn used when spectating (NULL if not spectating). */
-	UPROPERTY()
 	class ASpectatorPawn* SpectatorPawn;
 
 protected:
diff --git a/Engine/Source/Runtime/Engine/Private/DemoNetDriver.cpp b/Engine/Source/Runtime/Engine/Private/DemoNetDriver.cpp
--- a/Engine/Source/Runtime/Engine/Private/DemoNetDriver.cpp
+++ b/Engine/Source/Runtime/Engine/Private/DemoNetDriver.cpp
@@ -20,6 +20,7 @@
 #include "Net/UnrealNetwork.h"
 #include "Net/NetworkProfiler.h"
 #include "Net/DataReplication.h"
+#include "GameFramework/GameMode.h"
 
 DEFINE_LOG_CATEGORY( LogDemo );
 
@@ -1349,7 +1350,7 @@ void UDemoNetDriver::SpawnDemoRecSpectator( UNetConnection* Connection )
 {
 	check( Connection != NULL );
 
-	UClass* C = StaticLoadClass( AActor::StaticClass(), NULL, *DemoSpectatorClass, NULL, LOAD_None, NULL );
+	UClass* C = GetWorld()->GetAuthGameMode()->ReplaySpectatorPlayerControllerClass;
 
 	if ( C == NULL )
 	{
diff --git a/Engine/Source/Runtime/Engine/Private/GameMode.cpp b/Engine/Source/Runtime/Engine/Private/GameMode.cpp
--- a/Engine/Source/Runtime/Engine/Private/GameMode.cpp
+++ b/Engine/Source/Runtime/Engine/Private/GameMode.cpp
@@ -52,6 +52,7 @@ AGameMode::AGameMode(const FObjectInitializer& ObjectInitializer)
 	DefaultPawnClass = ADefaultPawn::StaticClass();
 	PlayerControllerClass = APlayerController::StaticClass();
 	SpectatorClass = ASpectatorPawn::StaticClass();
+	ReplaySpectatorPlayerControllerClass = APlayerController::StaticClass();
 	EngineMessageClass = UEngineMessage::StaticClass();
 	GameStateClass = AGameState::StaticClass();
 	CurrentID = 1;
