UE4 - tweak and turn on async render thread updates for all non-editor builds that use threading for performance

diff --git a/Engine/Source/Runtime/Engine/Classes/Components/ActorComponent.h b/Engine/Source/Runtime/Engine/Classes/Components/ActorComponent.h
--- a/Engine/Source/Runtime/Engine/Classes/Components/ActorComponent.h
+++ b/Engine/Source/Runtime/Engine/Classes/Components/ActorComponent.h
@@ -288,6 +288,12 @@ protected:
 	 */
 	virtual void OnUnregister();
 
+	/** Return true if CreateRenderState() should be called */
+	virtual bool ShouldCreateRenderState() const 
+	{
+		return false;
+	}
+
 	/** Used to create any rendering thread information for this component
 	*
 	* **Caution**, this is called concurrently on multiple threads (but never the same component concurrently)
@@ -480,6 +486,9 @@ public:
 	/** return true if this component requires end of frame updates to happen from the game thread. */
 	virtual bool RequiresGameThreadEndOfFrameUpdates() const;
 
+	/** return true if this component requires end of frame recreates to happen from the game thread. */
+	virtual bool RequiresGameThreadEndOfFrameRecreate() const;
+
 	/** Recreate the render state right away. Generally you always want to call MarkRenderStateDirty instead. 
 	*
 	* **Caution**, this is called concurrently on multiple threads (but never the same component concurrently)
diff --git a/Engine/Source/Runtime/Engine/Classes/Components/SceneComponent.h b/Engine/Source/Runtime/Engine/Classes/Components/SceneComponent.h
--- a/Engine/Source/Runtime/Engine/Classes/Components/SceneComponent.h
+++ b/Engine/Source/Runtime/Engine/Classes/Components/SceneComponent.h
@@ -511,6 +511,11 @@ public:
 
 	// Begin ActorComponent interface
 	virtual void OnRegister() override;
+	/** Return true if CreateRenderState() should be called */
+	virtual bool ShouldCreateRenderState() const override
+	{
+		return true;
+	}
 	virtual void UpdateComponentToWorld(bool bSkipPhysicsMove = false) override final;
 	virtual void DestroyComponent(bool bPromoteChildren = false) override;
 	virtual void OnComponentDestroyed() override;
diff --git a/Engine/Source/Runtime/Engine/Classes/Components/SkinnedMeshComponent.h b/Engine/Source/Runtime/Engine/Classes/Components/SkinnedMeshComponent.h
--- a/Engine/Source/Runtime/Engine/Classes/Components/SkinnedMeshComponent.h
+++ b/Engine/Source/Runtime/Engine/Classes/Components/SkinnedMeshComponent.h
@@ -433,15 +433,10 @@ protected:
 	virtual void CreateRenderState_Concurrent() override;
 	virtual void SendRenderDynamicData_Concurrent() override;
 	virtual void DestroyRenderState_Concurrent() override;
-#if 0
-	/** return true if this component requires end of frame updates to happen from the game thread. */
-	virtual bool RequiresGameThreadEndOfFrameUpdates()
+	virtual bool RequiresGameThreadEndOfFrameRecreate() const override
 	{
-		//@todo there isn't a great reason for this, it has to do with activating morph targets triggering material usage flag changes and potentially shader compilation.
-		// it seems unlikely this would cook right anyway
-		return true;
+		return false;
 	}
-#endif
 	virtual void TickComponent(float DeltaTime, enum ELevelTick TickType, FActorComponentTickFunction *ThisTickFunction) override;
 	virtual UObject const* AdditionalStatObject() const override;
 	// End UActorComponent interface
diff --git a/Engine/Source/Runtime/Engine/Private/Components/ActorComponent.cpp b/Engine/Source/Runtime/Engine/Private/Components/ActorComponent.cpp
--- a/Engine/Source/Runtime/Engine/Private/Components/ActorComponent.cpp
+++ b/Engine/Source/Runtime/Engine/Private/Components/ActorComponent.cpp
@@ -969,7 +969,7 @@ void UActorComponent::ExecuteRegisterEvents()
 		checkf(bRegistered, TEXT("Failed to route OnRegister (%s)"), *GetFullName());
 	}
 
-	if(FApp::CanEverRender() && !bRenderStateCreated && World->Scene)
+	if(FApp::CanEverRender() && !bRenderStateCreated && World->Scene && ShouldCreateRenderState())
 	{
 		CreateRenderState_Concurrent();
 		checkf(bRenderStateCreated, TEXT("Failed to route CreateRenderState_Concurrent (%s)"), *GetFullName());
@@ -1203,7 +1203,7 @@ void UActorComponent::MarkForNeededEndOfFrameRecreate()
 	if (ComponentWorld)
 	{
 		// by convention, recreates are always done on the gamethread
-		ComponentWorld->MarkActorComponentForNeededEndOfFrameUpdate(this, true);
+		ComponentWorld->MarkActorComponentForNeededEndOfFrameUpdate(this, RequiresGameThreadEndOfFrameRecreate());
 	}
 	else if (!HasAnyFlags(RF_Unreachable))
 	{
@@ -1217,6 +1217,11 @@ bool UActorComponent::RequiresGameThreadEndOfFrameUpdates() const
 	return false;
 }
 
+bool UActorComponent::RequiresGameThreadEndOfFrameRecreate() const
+{
+	return true;
+}
+
 void UActorComponent::Activate(bool bReset)
 {
 	if (bReset || ShouldActivate()==true)
diff --git a/Engine/Source/Runtime/Engine/Private/LevelTick.cpp b/Engine/Source/Runtime/Engine/Private/LevelTick.cpp
--- a/Engine/Source/Runtime/Engine/Private/LevelTick.cpp
+++ b/Engine/Source/Runtime/Engine/Private/LevelTick.cpp
@@ -698,8 +698,13 @@ void UWorld::RunTickGroup(ETickingGroup Group, bool bBlockTillComplete = true)
 
 static TAutoConsoleVariable<int32> CVarAllowAsyncRenderThreadUpdates(
 	TEXT("AllowAsyncRenderThreadUpdates"),
+	1,
+	TEXT("Used to control async renderthread updates. Also gated on FApp::ShouldUseThreadingForPerformance()."));
+
+static TAutoConsoleVariable<int32> CVarAllowAsyncRenderThreadUpdatesEditor(
+	TEXT("AllowAsyncRenderThreadUpdatesEditor"),
 	0,
-	TEXT("Used to control async renderthread updates."));
+	TEXT("Used to control async renderthread updates in the editor."));
 
 static TAutoConsoleVariable<int32> CVarCollectGarbageEveryFrame(
 	TEXT("CollectGarbageEveryFrame"),
@@ -752,7 +757,8 @@ void UWorld::MarkActorComponentForNeededEndOfFrameUpdate(UActorComponent* Compon
 	{
 		if (!bForceGameThread)
 		{
-			bool bAllowConcurrentUpdates = !!CVarAllowAsyncRenderThreadUpdates.GetValueOnGameThread();
+			bool bAllowConcurrentUpdates = FApp::ShouldUseThreadingForPerformance() && 
+				(GIsEditor ? !!CVarAllowAsyncRenderThreadUpdatesEditor.GetValueOnGameThread() : !!CVarAllowAsyncRenderThreadUpdates.GetValueOnGameThread());
 			bForceGameThread = !bAllowConcurrentUpdates;
 		}
 
diff --git a/Engine/Source/Runtime/Engine/Private/SkeletalMesh.cpp b/Engine/Source/Runtime/Engine/Private/SkeletalMesh.cpp
--- a/Engine/Source/Runtime/Engine/Private/SkeletalMesh.cpp
+++ b/Engine/Source/Runtime/Engine/Private/SkeletalMesh.cpp
@@ -4128,8 +4128,8 @@ FSkeletalMeshSceneProxy::FSkeletalMeshSceneProxy(const USkinnedMeshComponent* Co
 				check(Section.MaterialIndex == LODModel.Sections[Section.CorrespondClothSectionIndex].MaterialIndex);
 			}
 
-			if(!Material || !Material->CheckMaterialUsage(MATUSAGE_SkeletalMesh) ||
-			  (bClothSection && !Material->CheckMaterialUsage(MATUSAGE_Clothing)))
+			if(!Material || !Material->CheckMaterialUsage_Concurrent(MATUSAGE_SkeletalMesh) ||
+			  (bClothSection && !Material->CheckMaterialUsage_Concurrent(MATUSAGE_Clothing)))
 			{
 				Material = UMaterial::GetDefaultMaterial(MD_Surface);
 				MaterialRelevance |= Material->GetRelevance(FeatureLevel);
