UE-14271 Fix for max distance calculation for SoundNodeParamCrossFade

https://jira.ol.epicgames.net/browse/UE-14271

diff --git a/Engine/Source/Runtime/Engine/Classes/Sound/SoundNodeParamCrossFade.h b/Engine/Source/Runtime/Engine/Classes/Sound/SoundNodeParamCrossFade.h
--- a/Engine/Source/Runtime/Engine/Classes/Sound/SoundNodeParamCrossFade.h
+++ b/Engine/Source/Runtime/Engine/Classes/Sound/SoundNodeParamCrossFade.h
@@ -20,4 +20,6 @@ public:
 
 	virtual float GetCurrentDistance(FAudioDevice* AudioDevice, FActiveSound& ActiveSound, const FSoundParseParameters& ParseParams) const override;
 	virtual bool AllowCrossfading(FActiveSound& ActiveSound) const override;
+	virtual float MaxAudibleDistance(float CurrentMaxDistance) override;
+
 };
diff --git a/Engine/Source/Runtime/Engine/Private/SoundNodeParamCrossFade.cpp b/Engine/Source/Runtime/Engine/Private/SoundNodeParamCrossFade.cpp
--- a/Engine/Source/Runtime/Engine/Private/SoundNodeParamCrossFade.cpp
+++ b/Engine/Source/Runtime/Engine/Private/SoundNodeParamCrossFade.cpp
@@ -25,3 +25,9 @@ bool USoundNodeParamCrossFade::AllowCrossfading(FActiveSound& ActiveSound) const
 	// Always allow parameter to control crossfading, even on 2D/preview sounds
 	return true;
 }
+
+float USoundNodeParamCrossFade::MaxAudibleDistance(float CurrentMaxDistance)
+{
+	// Param-based crossfades are not a factor in max distance calculations
+	return CurrentMaxDistance;
+}
