Added deprecation warnings to UNavMovementComponent::UpdateNavAgent #UE4

UE-14653

diff --git a/Engine/Source/Runtime/Engine/Classes/GameFramework/NavMovementComponent.h b/Engine/Source/Runtime/Engine/Classes/GameFramework/NavMovementComponent.h
--- a/Engine/Source/Runtime/Engine/Classes/GameFramework/NavMovementComponent.h
+++ b/Engine/Source/Runtime/Engine/Classes/GameFramework/NavMovementComponent.h
@@ -139,6 +139,11 @@ public:
 	const FNavAgentProperties* GetNavAgentProperties() const;
 	DEPRECATED(4.7, "This function is deprecated. Please use GetNavAgentPropertiesRef instead.")
 	FNavAgentProperties* GetNavAgentProperties();
+
+	DEPRECATED(4.8, "This function is deprecated. Please use UpdateNavAgent version that's accepring a reference instead.")
+	void UpdateNavAgent(AActor* Owner);
+	DEPRECATED(4.8, "This function is deprecated. Please use UpdateNavAgent version that's accepring a reference instead.")
+	void UpdateNavAgent(UCapsuleComponent* CapsuleComponent);
 };
 
 
diff --git a/Engine/Source/Runtime/Engine/Private/Components/NavMovementComponent.cpp b/Engine/Source/Runtime/Engine/Private/Components/NavMovementComponent.cpp
--- a/Engine/Source/Runtime/Engine/Private/Components/NavMovementComponent.cpp
+++ b/Engine/Source/Runtime/Engine/Private/Components/NavMovementComponent.cpp
@@ -104,4 +104,18 @@ const FNavAgentProperties* UNavMovementComponent::GetNavAgentProperties() const
 FNavAgentProperties* UNavMovementComponent::GetNavAgentProperties()
 {
 	return &GetNavAgentPropertiesRef();
+}
+void UNavMovementComponent::UpdateNavAgent(AActor* Owner)
+{
+	if (Owner)
+	{
+		UpdateNavAgent(*Owner);
+	}
+}
+void UNavMovementComponent::UpdateNavAgent(UCapsuleComponent* CapsuleComponent)
+{
+	if (CapsuleComponent)
+	{
+		UpdateNavAgent(*CapsuleComponent);
+	}
 }
\ No newline at end of file
