#ue4 - Make sure we have a valid character before running CharacterMovementComponent::FindFloor().

- Potential fix for UE-14196. Added check() as well to get more info.

diff --git a/Engine/Source/Runtime/Engine/Private/Components/CharacterMovementComponent.cpp b/Engine/Source/Runtime/Engine/Private/Components/CharacterMovementComponent.cpp
--- a/Engine/Source/Runtime/Engine/Private/Components/CharacterMovementComponent.cpp
+++ b/Engine/Source/Runtime/Engine/Private/Components/CharacterMovementComponent.cpp
@@ -4805,12 +4805,14 @@ void UCharacterMovementComponent::ComputeFloorDist(const FVector& CapsuleLocatio
 void UCharacterMovementComponent::FindFloor(const FVector& CapsuleLocation, FFindFloorResult& OutFloorResult, bool bZeroDelta, const FHitResult* DownwardSweepResult) const
 {
 	// No collision, no floor...
-	if (!UpdatedComponent->IsCollisionEnabled())
+	if (!HasValidData() || !UpdatedComponent->IsCollisionEnabled())
 	{
 		OutFloorResult.Clear();
 		return;
 	}
 
+	check(CharacterOwner->GetCapsuleComponent());
+
 	// Increase height check slightly if walking, to prevent floor height adjustment from later invalidating the floor result.
 	const float HeightCheckAdjust = (IsMovingOnGround() ? MAX_FLOOR_DIST + KINDA_SMALL_NUMBER : -MAX_FLOOR_DIST);
 
