Fix continuously re-running construction scripts while making interactive changes in the details panel

https://jira.ol.epicgames.net/browse/UE-14577

diff --git a/Engine/Source/Runtime/Engine/Classes/Components/ActorComponent.h b/Engine/Source/Runtime/Engine/Classes/Components/ActorComponent.h
--- a/Engine/Source/Runtime/Engine/Classes/Components/ActorComponent.h
+++ b/Engine/Source/Runtime/Engine/Classes/Components/ActorComponent.h
@@ -272,7 +272,7 @@ private:
 	void ExecuteRegisterEvents();
 
 	/* Utility function for each of the PostEditChange variations to call for the same behavior */
-	void ConsolidatedPostEditChange();
+	void ConsolidatedPostEditChange(const FPropertyChangedEvent& PropertyChangedEvent);
 protected:
 
 	friend class FComponentReregisterContextBase;
diff --git a/Engine/Source/Runtime/Engine/Private/Components/ActorComponent.cpp b/Engine/Source/Runtime/Engine/Private/Components/ActorComponent.cpp
--- a/Engine/Source/Runtime/Engine/Private/Components/ActorComponent.cpp
+++ b/Engine/Source/Runtime/Engine/Private/Components/ActorComponent.cpp
@@ -488,7 +488,7 @@ void UActorComponent::PostEditUndo()
 	Super::PostEditUndo();
 }
 
-void UActorComponent::ConsolidatedPostEditChange()
+void UActorComponent::ConsolidatedPostEditChange(const FPropertyChangedEvent& PropertyChangedEvent)
 {
 	FComponentReregisterContext* ReregisterContext = EditReregisterContexts.FindRef(this);
 	if(ReregisterContext)
@@ -496,7 +496,7 @@ void UActorComponent::ConsolidatedPostEditChange()
 		delete ReregisterContext;
 		EditReregisterContexts.Remove(this);
 
-		if (Owner && !Owner->IsTemplate())
+		if ( Owner && !Owner->IsTemplate() && PropertyChangedEvent.ChangeType != EPropertyChangeType::Interactive )
 		{
 			Owner->RerunConstructionScripts();
 		}
@@ -514,14 +514,14 @@ void UActorComponent::ConsolidatedPostEditChange()
 
 void UActorComponent::PostEditChangeProperty(FPropertyChangedEvent& PropertyChangedEvent)
 {
-	ConsolidatedPostEditChange();
+	ConsolidatedPostEditChange(PropertyChangedEvent);
 
 	Super::PostEditChangeProperty(PropertyChangedEvent);
 }
 
 void UActorComponent::PostEditChangeChainProperty(FPropertyChangedChainEvent& PropertyChangedEvent)
 {
-	ConsolidatedPostEditChange();
+	ConsolidatedPostEditChange(PropertyChangedEvent);
 
 	Super::PostEditChangeChainProperty(PropertyChangedEvent);
 }
diff --git a/Engine/Source/Runtime/Slate/Private/Widgets/Colors/SColorWheel.cpp b/Engine/Source/Runtime/Slate/Private/Widgets/Colors/SColorWheel.cpp
--- a/Engine/Source/Runtime/Slate/Private/Widgets/Colors/SColorWheel.cpp
+++ b/Engine/Source/Runtime/Slate/Private/Widgets/Colors/SColorWheel.cpp
@@ -38,13 +38,14 @@ FReply SColorWheel::OnMouseButtonDown( const FGeometry& MyGeometry, const FPoint
 {
 	if (MouseEvent.GetEffectingButton() == EKeys::LeftMouseButton)
 	{
+		OnMouseCaptureBegin.ExecuteIfBound();
+
 		if (!ProcessMouseAction(MyGeometry, MouseEvent, false))
 		{
+			OnMouseCaptureEnd.ExecuteIfBound();
 			return FReply::Unhandled();
 		}
 
-		OnMouseCaptureBegin.ExecuteIfBound();
-
 		return FReply::Handled().CaptureMouse(SharedThis(this));
 	}
 	else
