Focus game viewport to ensure keyboard control on standalone client boot.

diff --git a/Engine/Source/Runtime/Engine/Private/GameEngine.cpp b/Engine/Source/Runtime/Engine/Private/GameEngine.cpp
--- a/Engine/Source/Runtime/Engine/Private/GameEngine.cpp
+++ b/Engine/Source/Runtime/Engine/Private/GameEngine.cpp
@@ -314,14 +314,22 @@ void UGameEngine::SwitchGameWindowToUseGameViewport()
 		{
 			CreateGameViewport( GameViewport );
 		}
-		GameViewportWindow.Pin()->SetContent( GameViewportWidget.ToSharedRef() );
+		
+		TSharedRef<SViewport> GameViewportWidgetRef = GameViewportWidget.ToSharedRef();
+		TSharedPtr<SWindow> GameViewportWindowPtr = GameViewportWindow.Pin();
+		
+		GameViewportWindowPtr->SetContent(GameViewportWidgetRef);
+		GameViewportWindowPtr->SlatePrepass();
+		
 		SceneViewport->ResizeFrame((uint32)GSystemResolution.ResX, (uint32)GSystemResolution.ResY, GSystemResolution.WindowMode, 0, 0);
 
 		// Move the registration of the game viewport to that messages are correctly received.
 		if (!FPlatformProperties::SupportsWindowedMode())
 		{
-			FSlateApplication::Get().RegisterGameViewport( GameViewportWidget.ToSharedRef() );
+			FSlateApplication::Get().RegisterGameViewport(GameViewportWidgetRef);
 		}
+		
+		FSlateApplication::Get().SetAllUserFocusToGameViewport();
 	}
 }
 
