Remove InitBodies. Deferred addition is only supported for static bodies

UE-12129

diff --git a/Engine/Source/Runtime/Engine/Classes/PhysicsEngine/BodyInstance.h b/Engine/Source/Runtime/Engine/Classes/PhysicsEngine/BodyInstance.h
--- a/Engine/Source/Runtime/Engine/Classes/PhysicsEngine/BodyInstance.h
+++ b/Engine/Source/Runtime/Engine/Classes/PhysicsEngine/BodyInstance.h
@@ -420,17 +420,7 @@ public:
 	 *	@param InRBScene The physics scene to place the body into
 	 *	@param InAggregate An aggregate to place the body into
 	 */
-	void InitBody(class UBodySetup* Setup, const FTransform& Transform, class UPrimitiveComponent* PrimComp, class FPhysScene* InRBScene, PhysXAggregateType InAggregate = NULL, bool bDefer = false);
-
-	/** Initialise multiple rigid bodies 
-	 *	@param Bodies List of body instances to initialise
-	 *	@param Transforms List of transforms to use for the bodies
-	 *	@param BodySetup The setup to use for the bodies
-	 *	@param PrimitiveComp The owning component
-	 *	@param InRBScene The physics scene to put the bodies into
-	 *	@param InAggregate Aggregate to place the bodies into, be aware for PhysX we can only have 128 elements in an aggregate!
-	 */
-	static void InitBodies(TArray<FBodyInstance*>& Bodies, TArray<FTransform>& Transforms, class UBodySetup* BodySetup, class UPrimitiveComponent* PrimitiveComp, class FPhysScene* InRBScene, PhysXAggregateType InAggregate = NULL, bool Defer = false, class UPhysicsSerializer* PhysicsSerializer = nullptr);
+	void InitBody(class UBodySetup* Setup, const FTransform& Transform, class UPrimitiveComponent* PrimComp, class FPhysScene* InRBScene, PhysXAggregateType InAggregate = NULL);
 
 	/** Validate a body transform, outputting debug info
 	 *	@param Transform Transform to debug
@@ -447,7 +437,7 @@ public:
 	 *	@param InRBScene
 	 *  @param PhysicsSerializer
 	 */
-	static void InitStaticBodies(TArray<FBodyInstance*>& Bodies, TArray<FTransform>& Transforms, class UBodySetup* BodySetup, class UPrimitiveComponent* PrimitiveComp, class FPhysScene* InRBScene, UPhysicsSerializer* PhysicsSerializer);
+	static void InitStaticBodies(TArray<FBodyInstance*>& Bodies, TArray<FTransform>& Transforms, class UBodySetup* BodySetup, class UPrimitiveComponent* PrimitiveComp, class FPhysScene* InRBScene, class UPhysicsSerializer* PhysicsSerializer);
 
 	/** Obtains the appropriate PhysX scene lock for READING and executes the passed in lambda. */
 	void ExecuteOnPhysicsReadOnly(TFunctionRef<void()> Func) const;
diff --git a/Engine/Source/Runtime/Engine/Private/InstancedStaticMesh.cpp b/Engine/Source/Runtime/Engine/Private/InstancedStaticMesh.cpp
--- a/Engine/Source/Runtime/Engine/Private/InstancedStaticMesh.cpp
+++ b/Engine/Source/Runtime/Engine/Private/InstancedStaticMesh.cpp
@@ -986,6 +986,8 @@ void UInstancedStaticMeshComponent::CreateAllInstanceBodies()
 
 	if (UBodySetup* BodySetup = GetBodySetup())
 	{
+		FPhysScene* PhysScene = GetWorld()->GetPhysicsScene();
+
 	    const int32 NumBodies = PerInstanceSMData.Num();
 	    InstanceBodies.SetNumUninitialized(NumBodies);
     
@@ -996,26 +998,32 @@ void UInstancedStaticMeshComponent::CreateAllInstanceBodies()
 	    {
 		    InstanceBodies[i] = new FBodyInstance;
 		    FBodyInstance* Instance = InstanceBodies[i];
-		    Transforms.Add(FTransform(PerInstanceSMData[i].Transform) * ComponentToWorld);
-    
+			const FTransform InstanceTM = FTransform(PerInstanceSMData[i].Transform) * ComponentToWorld;   
 		    Instance->CopyBodyInstancePropertiesFrom(&BodyInstance);
 		    Instance->InstanceBodyIndex = i; // Set body index 
 		    Instance->bAutoWeld = false;
+    
+		    // make sure we never enable bSimulatePhysics for ISMComps
+		    Instance->bSimulatePhysics = false;
 
+			if(Mobility == EComponentMobility::Movable)
+			{
+				Instance->InitBody(BodySetup, InstanceTM, this, PhysScene);
+			}else
+			{
+				Transforms.Add(InstanceTM);
 #if WITH_PHYSX
-			Instance->RigidActorSyncId = i+1;
+				Instance->RigidActorSyncId = i + 1;
 
-			if(GetWorld()->GetPhysicsScene()->HasAsyncScene())
-			{
-				Instance->RigidActorAsyncId = Instance->RigidActorSyncId + NumBodies;
-			}
+				if (GetWorld()->GetPhysicsScene()->HasAsyncScene())
+				{
+					Instance->RigidActorAsyncId = Instance->RigidActorSyncId + NumBodies;
+				}
 #endif
-    
-		    // make sure we never enable bSimulatePhysics for ISMComps
-		    Instance->bSimulatePhysics = false;
+			}
 	    }
 
-		if (NumBodies > 0)
+		if (NumBodies > 0 && Mobility == EComponentMobility::Static)
 		{
 			TArray<UBodySetup*> BodySetups;
 			TArray<UPhysicalMaterial*> PhysicalMaterials;
@@ -1026,15 +1034,7 @@ void UInstancedStaticMeshComponent::CreateAllInstanceBodies()
 			PhysicalMaterials.Add(FBodyInstance::GetSimplePhysicalMaterial(&BodyInstance, WeakSelfPtr, TWeakObjectPtr<UBodySetup>(BodySetup)));
 
 			PhysicsSerializer->CreatePhysicsData(BodySetups, PhysicalMaterials);
-			
-			if(Mobility == EComponentMobility::Static)
-			{
-				FBodyInstance::InitStaticBodies(InstanceBodies, Transforms, BodySetup, this, GetWorld()->GetPhysicsScene(), PhysicsSerializer);
-			}else
-			{
-				FBodyInstance::InitBodies(InstanceBodies, Transforms, BodySetup, this, GetWorld()->GetPhysicsScene(), nullptr, true, PhysicsSerializer);
-			}
-			
+			FBodyInstance::InitStaticBodies(InstanceBodies, Transforms, BodySetup, this, GetWorld()->GetPhysicsScene(), PhysicsSerializer);
 
 			//Serialize physics data for fast path cooking
 			PhysicsSerializer->SerializePhysics(InstanceBodies, BodySetups, PhysicalMaterials);
diff --git a/Engine/Source/Runtime/Engine/Private/PhysicsEngine/BodyInstance.cpp b/Engine/Source/Runtime/Engine/Private/PhysicsEngine/BodyInstance.cpp
--- a/Engine/Source/Runtime/Engine/Private/PhysicsEngine/BodyInstance.cpp
+++ b/Engine/Source/Runtime/Engine/Private/PhysicsEngine/BodyInstance.cpp
@@ -971,114 +971,17 @@
 
 #if UE_WITH_PHYSICS
 
-void FBodyInstance::InitBody(class UBodySetup* Setup, const FTransform& Transform, class UPrimitiveComponent* PrimComp, class FPhysScene* InRBScene, PhysXAggregateType InAggregate /*= NULL*/, bool bDefer /*= false*/ )
-{
-	static TArray<FBodyInstance*> Bodies;
-	static TArray<FTransform> Transforms;
 
-	if(Bodies.Num() == 1 && Transforms.Num() == 1)
-	{
-		Bodies[0] = this;
-		Transforms[0] = Transform;
-
-	}
-	else
-	{
-		Bodies.Add(this);
-		Transforms.Add(Transform);
-	}
-	check(Bodies.Num() == Transforms.Num() == 1);
-
-	InitBodies(Bodies, Transforms, Setup, PrimComp, InRBScene, InAggregate, bDefer);
-}
-
-TSharedPtr<TArray<ANSICHAR>> GetDebugDebugName(const UPrimitiveComponent* PrimitiveComp, const UBodySetup* BodySetup, FString& DebugName)
-{
-	// Setup names
-	// Make the debug name for this geometry...
-	DebugName = (TEXT(""));
-	TSharedPtr<TArray<ANSICHAR>> PhysXName;
-
-#if (WITH_EDITORONLY_DATA || UE_BUILD_DEBUG || LOOKING_FOR_PERF_ISSUES) && !(UE_BUILD_SHIPPING || UE_BUILD_TEST) && !NO_LOGGING
-	if (PrimitiveComp)
-	{
-		DebugName += FString::Printf(TEXT("Component: %s "), *PrimitiveComp->GetReadableName());
-	}
-
-	if (BodySetup->BoneName != NAME_None)
-	{
-		DebugName += FString::Printf(TEXT("Bone: %s "), *BodySetup->BoneName.ToString());
-	}
-
-	// Convert to char* for PhysX
-	PhysXName = MakeShareable(new TArray<ANSICHAR>(StringToArray<ANSICHAR>(*DebugName, DebugName.Len() + 1)));
-#endif
-
-	return PhysXName;
-}
-
-const USkeletalMeshComponent* GetSkeletalMeshComponentAndProperties(const UPrimitiveComponent* PrimitiveComp, const UBodySetup* BodySetup, float& InstanceBlendWeight, bool& bInstanceSimulatePhysics, bool& bComponentAwake)
-{
-	const USkeletalMeshComponent* SkelMeshComp = Cast<USkeletalMeshComponent>(PrimitiveComp);
-	if (SkelMeshComp)
-	{
-		if ((BodySetup->PhysicsType == PhysType_Simulated) || (BodySetup->PhysicsType == PhysType_Default))
-		{
-			bool bEnableSim = (SkelMeshComp && IsRunningDedicatedServer()) ? SkelMeshComp->bEnablePhysicsOnDedicatedServer : true;
-			bEnableSim &= ((BodySetup->PhysicsType == PhysType_Simulated) || (SkelMeshComp->BodyInstance.bSimulatePhysics));	//if unfixed enable. If default look at parent
-
-			if (bEnableSim)
-			{
-				// set simulate to true if using physics
-				bInstanceSimulatePhysics = true;
-				if (BodySetup->PhysicsType == PhysType_Simulated)
-				{
-					InstanceBlendWeight = 1.f;
-				}
-			}
-			else
-			{
-				bInstanceSimulatePhysics = false;
-				if (BodySetup->PhysicsType == PhysType_Simulated)
-				{
-					InstanceBlendWeight = 0.f;
-				}
-			}
-		}
-
-		bComponentAwake = SkelMeshComp->BodyInstance.bStartAwake;
-	}
-
-	return SkelMeshComp;
-}
-
-FVector GetInitialLinearVelocity(const AActor* OwningActor, bool& bComponentAwake)
-{
-	FVector InitialLinVel(EForceInit::ForceInitToZero);
-	if (OwningActor)
-	{
-		InitialLinVel = OwningActor->GetVelocity();
-
-		if (InitialLinVel.Size() > KINDA_SMALL_NUMBER)
-		{
-			bComponentAwake = true;
-		}
-	}
-
-	return InitialLinVel;
-}
-
 template <bool bCompileStatic>
 struct FInitBodiesHelper
 {
-	FInitBodiesHelper(const TArray<FBodyInstance*>& InBodies, const TArray<FTransform>& InTransforms, class UBodySetup* InBodySetup, class UPrimitiveComponent* InPrimitiveComp, class FPhysScene* InInRBScene, FBodyInstance::PhysXAggregateType InInAggregate = NULL, bool InDefer = false, UPhysicsSerializer* InPhysicsSerializer = nullptr)
+	FInitBodiesHelper(const TArray<FBodyInstance*>& InBodies, const TArray<FTransform>& InTransforms, class UBodySetup* InBodySetup, class UPrimitiveComponent* InPrimitiveComp, class FPhysScene* InInRBScene, FBodyInstance::PhysXAggregateType InInAggregate = NULL, UPhysicsSerializer* InPhysicsSerializer = nullptr)
 	: Bodies(InBodies)
 	, Transforms(InTransforms)
 	, BodySetup(InBodySetup)
 	, PrimitiveComp(InPrimitiveComp)
 	, PhysScene(InInRBScene)
 	, InAggregate(InInAggregate)
-	, bDefer(InDefer)
 	, PhysicsSerializer(InPhysicsSerializer)
 	, DebugName(TEXT(""))
 	, InstanceBlendWeight(-1.f)
@@ -1093,7 +996,7 @@
 
 		bStatic = bCompileStatic || PrimitiveComp == nullptr || PrimitiveComp->Mobility != EComponentMobility::Movable;
 		SkelMeshComp = bCompileStatic ? nullptr : GetSkeletalMeshComponentAndProperties(PrimitiveComp, BodySetup, InstanceBlendWeight, bInstanceSimulatePhysics, bComponentAwake);
-		
+
 		const AActor* OwningActor = PrimitiveComp ? PrimitiveComp->GetOwner() : nullptr;
 		InitialLinVel = GetInitialLinearVelocity(OwningActor, bComponentAwake);
 
@@ -1101,7 +1004,18 @@
 		PSyncScene = PhysScene->GetPhysXScene(PST_Sync);
 		PAsyncScene = PhysScene->HasAsyncScene() ? PhysScene->GetPhysXScene(PST_Async) : nullptr;
 #endif
-	
+
 	}
 
+	void InitBodies() const
+	{
+#if WITH_PHYSX
+		InitBodies_PhysX();
+#endif
+
+#if WITH_BOX2D
+		InitBodies_Box2D();
+#endif
+	}
+
 	//The arguments passed into InitBodies
@@ -1111,9 +1025,8 @@
 	class UPrimitiveComponent* PrimitiveComp;
 	class FPhysScene* PhysScene;
 	FBodyInstance::PhysXAggregateType InAggregate;
-	const bool bDefer;
 	UPhysicsSerializer* PhysicsSerializer;
-	
+
 	FString DebugName;
 	TSharedPtr<TArray<ANSICHAR>> PhysXName;
 
@@ -1452,8 +1365,8 @@
 		TArray<PxActor*> PAsyncActors;
 		TArray<PxActor*> PDynamicActors;
 
-		// No reason to defer if there's an existing aggregate.
-		const bool bCanDefer = bCompileStatic || (bDefer && !InAggregate);
+		// Only static objects qualify for deferred addition
+		const bool bCanDefer = bCompileStatic;
 		bool bDynamicsUseAsync = false;
 		if (CreateShapesAndActors_PhysX(PSyncActors, PAsyncActors, PDynamicActors, bCanDefer, bDynamicsUseAsync) == false)
 		{
@@ -1650,26 +1563,108 @@
 
 };
 
-void FBodyInstance::InitBodies(TArray<FBodyInstance*>& Bodies, TArray<FTransform>& Transforms, class UBodySetup* BodySetup, class UPrimitiveComponent* PrimitiveComp, class FPhysScene* InRBScene, PhysXAggregateType InAggregate /*= NULL*/, bool bDefer /*= false*/, UPhysicsSerializer* PhysicsSerializer )
+void FBodyInstance::InitBody(class UBodySetup* Setup, const FTransform& Transform, class UPrimitiveComponent* PrimComp, class FPhysScene* InRBScene, PhysXAggregateType InAggregate /*= NULL*/)
 {
-	SCOPE_CYCLE_COUNTER(STAT_InitBodies);
-
-	check(BodySetup);
+	SCOPE_CYCLE_COUNTER(STAT_InitBody);
+	check(Setup);
 	check(InRBScene);
-	check(Bodies.Num() > 0);
+	
+	static TArray<FBodyInstance*> Bodies;
+	static TArray<FTransform> Transforms;
 
-	FInitBodiesHelper<false> InitBodiesHelper(Bodies, Transforms, BodySetup, PrimitiveComp, InRBScene, InAggregate, bDefer, PhysicsSerializer);
+	if(Bodies.Num() == 1 && Transforms.Num() == 1)
+	{
+		Bodies[0] = this;
+		Transforms[0] = Transform;
+
+	}
+	else
+	{
+		Bodies.Add(this);
+		Transforms.Add(Transform);
+	}
 	
-#if WITH_PHYSX
-	InitBodiesHelper.InitBodies_PhysX();
+	check(Bodies.Num() == Transforms.Num() == 1);
+
+	FInitBodiesHelper<false> InitBodiesHelper(Bodies, Transforms, Setup, PrimComp, InRBScene, InAggregate);
+	InitBodiesHelper.InitBodies();
+}
+
+TSharedPtr<TArray<ANSICHAR>> GetDebugDebugName(const UPrimitiveComponent* PrimitiveComp, const UBodySetup* BodySetup, FString& DebugName)
+{
+	// Setup names
+	// Make the debug name for this geometry...
+	DebugName = (TEXT(""));
+	TSharedPtr<TArray<ANSICHAR>> PhysXName;
+
+#if (WITH_EDITORONLY_DATA || UE_BUILD_DEBUG || LOOKING_FOR_PERF_ISSUES) && !(UE_BUILD_SHIPPING || UE_BUILD_TEST) && !NO_LOGGING
+	if (PrimitiveComp)
+	{
+		DebugName += FString::Printf(TEXT("Component: %s "), *PrimitiveComp->GetReadableName());
+	}
+
+	if (BodySetup->BoneName != NAME_None)
+	{
+		DebugName += FString::Printf(TEXT("Bone: %s "), *BodySetup->BoneName.ToString());
+	}
+
+	// Convert to char* for PhysX
+	PhysXName = MakeShareable(new TArray<ANSICHAR>(StringToArray<ANSICHAR>(*DebugName, DebugName.Len() + 1)));
 #endif
 
-#if WITH_BOX2D
-	InitBodiesHelper.InitBodies_Box2D();
-#endif
+	return PhysXName;
+}
+
+const USkeletalMeshComponent* GetSkeletalMeshComponentAndProperties(const UPrimitiveComponent* PrimitiveComp, const UBodySetup* BodySetup, float& InstanceBlendWeight, bool& bInstanceSimulatePhysics, bool& bComponentAwake)
+{
+	const USkeletalMeshComponent* SkelMeshComp = Cast<USkeletalMeshComponent>(PrimitiveComp);
+	if (SkelMeshComp)
+	{
+		if ((BodySetup->PhysicsType == PhysType_Simulated) || (BodySetup->PhysicsType == PhysType_Default))
+		{
+			bool bEnableSim = (SkelMeshComp && IsRunningDedicatedServer()) ? SkelMeshComp->bEnablePhysicsOnDedicatedServer : true;
+			bEnableSim &= ((BodySetup->PhysicsType == PhysType_Simulated) || (SkelMeshComp->BodyInstance.bSimulatePhysics));	//if unfixed enable. If default look at parent
+
+			if (bEnableSim)
+			{
+				// set simulate to true if using physics
+				bInstanceSimulatePhysics = true;
+				if (BodySetup->PhysicsType == PhysType_Simulated)
+				{
+					InstanceBlendWeight = 1.f;
+				}
+			}
+			else
+			{
+				bInstanceSimulatePhysics = false;
+				if (BodySetup->PhysicsType == PhysType_Simulated)
+				{
+					InstanceBlendWeight = 0.f;
+				}
+			}
+		}
+
+		bComponentAwake = SkelMeshComp->BodyInstance.bStartAwake;
+	}
+
+	return SkelMeshComp;
+}
+
+FVector GetInitialLinearVelocity(const AActor* OwningActor, bool& bComponentAwake)
+{
+	FVector InitialLinVel(EForceInit::ForceInitToZero);
+	if (OwningActor)
+	{
+		InitialLinVel = OwningActor->GetVelocity();
+
+		if (InitialLinVel.Size() > KINDA_SMALL_NUMBER)
+		{
+			bComponentAwake = true;
+		}
+	}
+
+	return InitialLinVel;
 }
 
 
-
-
 #endif // UE_WITH_PHYSICS
@@ -4627,15 +4622,8 @@
 	check(InRBScene);
 	check(Bodies.Num() > 0);
 
-	FInitBodiesHelper<true> InitBodiesHelper(Bodies, Transforms, BodySetup, PrimitiveComp, InRBScene, nullptr, true, PhysicsSerializer);
-
-#if WITH_PHYSX
-	InitBodiesHelper.InitBodies_PhysX();
-#endif
-
-#if WITH_BOX2D
-	InitBodiesHelper.InitBodies_Box2D();
-#endif
+	FInitBodiesHelper<true> InitBodiesHelper(Bodies, Transforms, BodySetup, PrimitiveComp, InRBScene, nullptr, PhysicsSerializer);
+	InitBodiesHelper.InitBodies();
 }
 
 void FBodyInstance::SetShapeFlags_AssumesLocked(TEnumAsByte<ECollisionEnabled::Type> UseCollisionEnabled, PxShape* PInShape, EPhysicsSceneType SceneType, const bool bUseComplexAsSimple)
