Disable binary physx serialization on mac. There is a bug with modules that needs to be sorted before this can be enabled.

UE-14480, UE-14321

diff --git a/Engine/Source/Runtime/Engine/Private/PhysicsEngine/PhysicsSerializer.cpp b/Engine/Source/Runtime/Engine/Private/PhysicsEngine/PhysicsSerializer.cpp
--- a/Engine/Source/Runtime/Engine/Private/PhysicsEngine/PhysicsSerializer.cpp
+++ b/Engine/Source/Runtime/Engine/Private/PhysicsEngine/PhysicsSerializer.cpp
@@ -16,7 +16,10 @@ UPhysicsSerializer::UPhysicsSerializer(const FObjectInitializer& ObjectInitializ
 
 FByteBulkData* UPhysicsSerializer::GetBinaryData(FName Format, const TArray<FBodyInstance*>& Bodies, const TArray<class UBodySetup*>& BodySetups, const TArray<class UPhysicalMaterial*>& PhysicalMaterials)
 {
-	
+#if PLATFORM_MAC
+	return nullptr;	//This is not supported right now
+#endif
+
 	QUICK_SCOPE_CYCLE_COUNTER(STAT_GetBinaryData);
 	const bool bContainedData = BinaryFormatData.Contains(Format);
 	FByteBulkData* Result = &BinaryFormatData.GetFormat(Format);
@@ -88,6 +91,10 @@ void UPhysicsSerializer::SerializePhysics(const TArray<FBodyInstance*>& Bodies,
 
 void UPhysicsSerializer::CreatePhysicsData(const TArray<UBodySetup*>& BodySetups, const TArray<UPhysicalMaterial*>& PhysicalMaterials)
 {
+#if PLATFORM_MAC
+	return;	//This is not supported right now
+#endif
+
 #if !WITH_EDITOR
 	QUICK_SCOPE_CYCLE_COUNTER(STAT_PhysicsSerializer_CreatePhysicsData);
 
