UE4-Orion-To-UE4(2530838): Added "World" check to AActor::GetNetDriver();

diff --git a/Engine/Source/Runtime/Engine/Private/Actor.cpp b/Engine/Source/Runtime/Engine/Private/Actor.cpp
--- a/Engine/Source/Runtime/Engine/Private/Actor.cpp
+++ b/Engine/Source/Runtime/Engine/Private/Actor.cpp
@@ -3133,6 +3133,7 @@ UNetDriver* AActor::GetNetDriver() const
 	UWorld *World = GetWorld();
 	if (NetDriverName == NAME_GameNetDriver)
 	{
+		check(World);
 		return World->GetNetDriver();
 	}
 
