Fixing RecastNavmesh crash when performing hot reload.

diff --git a/Engine/Source/Runtime/Engine/Private/AI/Navigation/RecastNavMesh.cpp b/Engine/Source/Runtime/Engine/Private/AI/Navigation/RecastNavMesh.cpp
--- a/Engine/Source/Runtime/Engine/Private/AI/Navigation/RecastNavMesh.cpp
+++ b/Engine/Source/Runtime/Engine/Private/AI/Navigation/RecastNavMesh.cpp
@@ -752,7 +752,10 @@ void ARecastNavMesh::Serialize( FArchive& Ar )
 			// empty, just skip over this data
 			Ar.Seek( RecastNavMeshSizePos + RecastNavMeshSizeBytes );
 			// if it's not getting filled it's better to just remove it
-			RecastNavMeshImpl->ReleaseDetourNavMesh();
+			if (RecastNavMeshImpl)
+			{
+				RecastNavMeshImpl->ReleaseDetourNavMesh();
+			}
 		}
 	}
 	else
