Avoid issues where cached UActorComponent Owner could get out of sync during undo/redo
Ensure that ActorComponent functions called while Owner is null (during undo/redo) get the correct Owner by using the accessor instead of the variable directly

diff --git a/Engine/Source/Runtime/Engine/Classes/Components/ActorComponent.h b/Engine/Source/Runtime/Engine/Classes/Components/ActorComponent.h
--- a/Engine/Source/Runtime/Engine/Classes/Components/ActorComponent.h
+++ b/Engine/Source/Runtime/Engine/Classes/Components/ActorComponent.h
@@ -136,6 +136,11 @@ private:
 	/** Indicates that BeginPlay has been called, but EndPlay has not yet */
 	uint32 bHasBegunPlay:1;
 
+#if WITH_EDITOR
+	/** During undo/redo it isn't safe to cache owner */
+	uint32 bCanUseCachedOwner:1;
+#endif
+
 	/** Tracks whether the component has been added to one of the world's end of frame update lists */
 	uint32 MarkedForEndOfFrameUpdateState:2;
 	friend struct FMarkComponentEndOfFrameUpdateState;
@@ -653,10 +658,19 @@ private:
 
 FORCEINLINE_DEBUGGABLE class AActor* UActorComponent::GetOwner() const
 {
-	if (Owner == nullptr)
+#if WITH_EDITOR
+	// During undo/redo the cached owner is unreliable so just used GetTypedOuter
+	if (bCanUseCachedOwner)
 	{
-		Owner = GetTypedOuter<AActor>();
+		checkSlow(Owner == GetTypedOuter<AActor>()); // verify cached value is correct
+		return Owner;
 	}
+	else
+	{
+		return GetTypedOuter<AActor>();
+	}
+#else
 	checkSlow(Owner == GetTypedOuter<AActor>()); // verify cached value is correct
 	return Owner;
+#endif
 }
diff --git a/Engine/Source/Runtime/Engine/Classes/Components/LightComponent.h b/Engine/Source/Runtime/Engine/Classes/Components/LightComponent.h
--- a/Engine/Source/Runtime/Engine/Classes/Components/LightComponent.h
+++ b/Engine/Source/Runtime/Engine/Classes/Components/LightComponent.h
@@ -299,8 +299,6 @@ public:
 	virtual void Serialize(FArchive& Ar) override;
 	virtual void PostLoad() override;
 #if WITH_EDITOR
-	virtual void PreEditUndo() override;
-	virtual void PostEditUndo() override;
 	virtual bool CanEditChange(const UProperty* InProperty) const override;
 	virtual void PostEditChangeProperty(FPropertyChangedEvent& PropertyChangedEvent) override;
 	virtual void UpdateLightSpriteTexture() override;
diff --git a/Engine/Source/Runtime/Engine/Private/Components/ActorComponent.cpp b/Engine/Source/Runtime/Engine/Private/Components/ActorComponent.cpp
--- a/Engine/Source/Runtime/Engine/Private/Components/ActorComponent.cpp
+++ b/Engine/Source/Runtime/Engine/Private/Components/ActorComponent.cpp
@@ -225,50 +225,55 @@ bool UActorComponent::IsCreatedByConstructionScript() const
 #if WITH_EDITOR
 void UActorComponent::CheckForErrors()
 {
-	if (Owner != NULL && GetClass()->HasAnyClassFlags(CLASS_Deprecated))
+	if (AActor* MyOwner = GetOwner())
 	{
-		FFormatNamedArguments Arguments;
-		Arguments.Add(TEXT("ComponentName"), FText::FromString(GetName()));
-		Arguments.Add(TEXT("OwnerName"), FText::FromString(Owner->GetName()));
-		FMessageLog("MapCheck").Warning()
-			->AddToken(FUObjectToken::Create(Owner))
-			->AddToken(FTextToken::Create(FText::Format( LOCTEXT( "MapCheck_Message_DeprecatedClass", "{ComponentName}::{OwnerName} is obsolete and must be removed (Class is deprecated)" ), Arguments ) ) )
-			->AddToken(FMapErrorToken::Create(FMapErrors::DeprecatedClass));
-	}
+		if (GetClass()->HasAnyClassFlags(CLASS_Deprecated))
+		{
+			FFormatNamedArguments Arguments;
+			Arguments.Add(TEXT("ComponentName"), FText::FromString(GetName()));
+			Arguments.Add(TEXT("OwnerName"), FText::FromString(MyOwner->GetName()));
+			FMessageLog("MapCheck").Warning()
+				->AddToken(FUObjectToken::Create(MyOwner))
+				->AddToken(FTextToken::Create(FText::Format( LOCTEXT( "MapCheck_Message_DeprecatedClass", "{ComponentName}::{OwnerName} is obsolete and must be removed (Class is deprecated)" ), Arguments ) ) )
+				->AddToken(FMapErrorToken::Create(FMapErrors::DeprecatedClass));
+		}
 
-	if (Owner != NULL && GetClass()->HasAnyClassFlags(CLASS_Abstract))
-	{
-		FFormatNamedArguments Arguments;
-		Arguments.Add(TEXT("ComponentName"), FText::FromString(GetName()));
-		Arguments.Add(TEXT("OwnerName"), FText::FromString(Owner->GetName()));
-		FMessageLog("MapCheck").Warning()
-			->AddToken(FUObjectToken::Create(Owner))
-			->AddToken(FTextToken::Create(FText::Format( LOCTEXT( "MapCheck_Message_AbstractClass", "{ComponentName}::{OwnerName} is obsolete and must be removed (Class is abstract)" ), Arguments ) ))
-			->AddToken(FMapErrorToken::Create(FMapErrors::AbstractClass));
+		if (GetClass()->HasAnyClassFlags(CLASS_Abstract))
+		{
+			FFormatNamedArguments Arguments;
+			Arguments.Add(TEXT("ComponentName"), FText::FromString(GetName()));
+			Arguments.Add(TEXT("OwnerName"), FText::FromString(MyOwner->GetName()));
+			FMessageLog("MapCheck").Warning()
+				->AddToken(FUObjectToken::Create(MyOwner))
+				->AddToken(FTextToken::Create(FText::Format( LOCTEXT( "MapCheck_Message_AbstractClass", "{ComponentName}::{OwnerName} is obsolete and must be removed (Class is abstract)" ), Arguments ) ))
+				->AddToken(FMapErrorToken::Create(FMapErrors::AbstractClass));
+		}
 	}
 }
 #endif
 
 bool UActorComponent::IsOwnerSelected() const
 {
-	return Owner && Owner->IsSelected();
+	AActor* MyOwner = GetOwner();
+	return MyOwner && MyOwner->IsSelected();
 }
 
 UWorld* UActorComponent::GetWorld() const
 {
 	UWorld* ComponentWorld = World;
-	if (ComponentWorld == NULL)
+	if (ComponentWorld == nullptr)
 	{
+		AActor* MyOwner = GetOwner();
 		// If we don't have a world yet, it may be because we haven't gotten registered yet, but we can try to look at our owner
-		if (Owner && !Owner->HasAnyFlags(RF_ClassDefaultObject))
+		if (MyOwner && !MyOwner->HasAnyFlags(RF_ClassDefaultObject))
 		{
-			ComponentWorld = Owner->GetWorld();
+			ComponentWorld = MyOwner->GetWorld();
 		}
 
-		if( ComponentWorld == NULL )
+		if( ComponentWorld == nullptr )
 		{
 			// As a fallback check the outer of this component for a world. In some cases components are spawned directly in the world
-			return Cast<UWorld>( GetOuter() );
+			ComponentWorld = Cast<UWorld>(GetOuter());
 		}
 	}
 	return ComponentWorld;
@@ -282,7 +287,8 @@ bool UActorComponent::ComponentHasTag(FName Tag) const
 
 ENetMode UActorComponent::GetNetMode() const
 {
-	return Owner ? Owner->GetNetMode() : NM_Standalone;
+	AActor* MyOwner = GetOwner();
+	return MyOwner ? MyOwner->GetNetMode() : NM_Standalone;
 }
 
 FSceneInterface* UActorComponent::GetScene() const
@@ -293,17 +299,8 @@ FSceneInterface* UActorComponent::GetScene() const
 ULevel* UActorComponent::GetComponentLevel() const
 {
 	// For model components Level is outer object
-	ULevel* ComponentLevel = Cast<ULevel>( GetOuter() );
-	
-	if (ComponentLevel == NULL)
-	{
-		if (Owner)
-		{
-			ComponentLevel = CastChecked<ULevel>( Owner->GetOuter() );
-		}
-	}
-	
-	return ComponentLevel;
+	AActor* MyOwner = GetOwner();
+	return (MyOwner ? CastChecked<ULevel>(MyOwner->GetOuter()) : Cast<ULevel>( GetOuter() ) );
 }
 
 bool UActorComponent::ComponentIsInLevel(const ULevel *TestLevel) const
@@ -311,13 +308,12 @@ bool UActorComponent::ComponentIsInLevel(const ULevel *TestLevel) const
 	return (GetComponentLevel() == TestLevel);
 }
 
-
 bool UActorComponent::ComponentIsInPersistentLevel(bool bIncludeLevelStreamingPersistent) const
 {
 	ULevel* MyLevel = GetComponentLevel();
 	UWorld* MyWorld = GetWorld();
 
-	if(Owner == NULL || MyLevel == NULL || MyWorld == NULL)
+	if (MyLevel == NULL || MyWorld == NULL)
 	{
 		return false;
 	}
@@ -363,9 +359,9 @@ void UActorComponent::BeginDestroy()
 	World = NULL;
 
 	// Remove from the parent's OwnedComponents list
-	if (Owner)
+	if (AActor* MyOwner = GetOwner())
 	{
-		Owner->RemoveOwnedComponent(this);
+		MyOwner->RemoveOwnedComponent(this);
 	}
 
 	Super::BeginDestroy();
@@ -387,25 +383,20 @@ bool UActorComponent::NeedsLoadForServer() const
 
 int32 UActorComponent::GetFunctionCallspace( UFunction* Function, void* Parameters, FFrame* Stack )
 {
-	if (Owner == NULL)
-	{
-		return FunctionCallspace::Local;
-	}
-	return Owner->GetFunctionCallspace(Function, Parameters, Stack);
+	AActor* MyOwner = GetOwner();
+	return (MyOwner ? MyOwner->GetFunctionCallspace(Function, Parameters, Stack) : FunctionCallspace::Local);
 }
 
 bool UActorComponent::CallRemoteFunction( UFunction* Function, void* Parameters, FOutParmRec* OutParms, FFrame* Stack )
 {
-	if (Owner == NULL)
-	{
-		return false;
-	}
-	
-	UNetDriver* NetDriver = Owner->GetNetDriver();
-	if (NetDriver)
+	if (AActor* MyOwner = GetOwner())
 	{
-		NetDriver->ProcessRemoteFunction(Owner, Function, Parameters, OutParms, Stack, this);
-		return true;
+		UNetDriver* NetDriver = MyOwner->GetNetDriver();
+		if (NetDriver)
+		{
+			NetDriver->ProcessRemoteFunction(MyOwner, Function, Parameters, OutParms, Stack, this);
+			return true;
+		}
 	}
 
 	return false;
@@ -419,9 +410,10 @@ bool UActorComponent::Modify( bool bAlwaysMarkDirty/*=true*/ )
 {
 	// If this is a construction script component we don't store them in the transaction buffer.  Instead, mark
 	// the Actor as modified so that we store of the transaction annotation that has the component properties stashed
-	if (Owner && IsCreatedByConstructionScript())
+	AActor* MyOwner = GetOwner();
+	if (MyOwner && IsCreatedByConstructionScript())
 	{
-		return Owner->Modify(bAlwaysMarkDirty);
+		return MyOwner->Modify(bAlwaysMarkDirty);
 	}
 
 	return Super::Modify(bAlwaysMarkDirty);
@@ -456,6 +448,7 @@ void UActorComponent::PreEditUndo()
 	Super::PreEditUndo();
 
 	Owner = nullptr;
+	bCanUseCachedOwner = false;
 }
 
 void UActorComponent::PostEditUndo()
@@ -474,6 +467,9 @@ void UActorComponent::PostEditUndo()
 	}
 	else
 	{
+		Owner = GetTypedOuter<AActor>();
+		bCanUseCachedOwner = true;
+
 		// Let the component be properly registered, after it was restored.
 		if (Owner)
 		{
@@ -496,9 +492,10 @@ void UActorComponent::ConsolidatedPostEditChange(const FPropertyChangedEvent& Pr
 		delete ReregisterContext;
 		EditReregisterContexts.Remove(this);
 
-		if ( Owner && !Owner->IsTemplate() && PropertyChangedEvent.ChangeType != EPropertyChangeType::Interactive )
+		AActor* MyOwner = GetOwner();
+		if ( MyOwner && !MyOwner->IsTemplate() && PropertyChangedEvent.ChangeType != EPropertyChangeType::Interactive )
 		{
-			Owner->RerunConstructionScripts();
+			MyOwner->RerunConstructionScripts();
 		}
 	}
 
@@ -534,9 +531,9 @@ void UActorComponent::OnRegister()
 	checkf(!HasAnyFlags(RF_Unreachable), TEXT("%s"), *GetDetailedInfo());
 	checkf(!GetOuter()->IsTemplate(), TEXT("'%s' (%s)"), *GetOuter()->GetFullName(), *GetDetailedInfo());
 	checkf(!IsTemplate(), TEXT("'%s' (%s)"), *GetOuter()->GetFullName(), *GetDetailedInfo() );
-	checkf(World, TEXT("OnRegister: %s to %s"), *GetDetailedInfo(), Owner ? *Owner->GetFullName() : TEXT("*** No Owner ***") );
-	checkf(!bRegistered, TEXT("OnRegister: %s to %s"), *GetDetailedInfo(), Owner ? *Owner->GetFullName() : TEXT("*** No Owner ***") );
-	checkf(!IsPendingKill(), TEXT("OnRegister: %s to %s"), *GetDetailedInfo(), Owner ? *Owner->GetFullName() : TEXT("*** No Owner ***") );
+	checkf(World, TEXT("OnRegister: %s to %s"), *GetDetailedInfo(), GetOwner() ? *GetOwner()->GetFullName() : TEXT("*** No Owner ***") );
+	checkf(!bRegistered, TEXT("OnRegister: %s to %s"), *GetDetailedInfo(), GetOwner() ? *GetOwner()->GetFullName() : TEXT("*** No Owner ***") );
+	checkf(!IsPendingKill(), TEXT("OnRegister: %s to %s"), *GetDetailedInfo(), GetOwner() ? *GetOwner()->GetFullName() : TEXT("*** No Owner ***") );
 
 	bRegistered = true;
 
@@ -629,12 +626,16 @@ FString FActorComponentTickFunction::DiagnosticMessage()
 
 bool UActorComponent::SetupActorComponentTickFunction(struct FTickFunction* TickFunction)
 {
-	if(TickFunction->bCanEverTick && !IsTemplate() && (!Owner || !Owner->IsTemplate()))
+	if(TickFunction->bCanEverTick && !IsTemplate())
 	{
-		ULevel* ComponentLevel = (Owner ? Owner->GetLevel() : GetWorld()->PersistentLevel);
-		TickFunction->SetTickFunctionEnable(TickFunction->bStartWithTickEnabled);
-		TickFunction->RegisterTickFunction(ComponentLevel);
-		return true;
+		AActor* MyOwner = GetOwner();
+		if (!MyOwner || !MyOwner->IsTemplate())
+		{
+			ULevel* ComponentLevel = (MyOwner ? MyOwner->GetLevel() : GetWorld()->PersistentLevel);
+			TickFunction->SetTickFunctionEnable(TickFunction->bStartWithTickEnabled);
+			TickFunction->RegisterTickFunction(ComponentLevel);
+			return true;
+		}
 	}
 	return false;
 }
@@ -734,9 +735,10 @@ void UActorComponent::RegisterComponentWithWorld(UWorld* InWorld)
 	// If not registered, should not have a scene
 	checkf(World == NULL, TEXT("%s"), *GetFullName());
 
-	checkSlow(Owner == NULL || Owner->OwnsComponent(this));
+	AActor* MyOwner = GetOwner();
+	checkSlow(MyOwner == nullptr || MyOwner->OwnsComponent(this));
 
-	if ((Owner != NULL) && Owner->GetClass()->HasAnyClassFlags(CLASS_NewerVersionExists))
+	if (MyOwner && MyOwner->GetClass()->HasAnyClassFlags(CLASS_NewerVersionExists))
 	{
 		UE_LOG(LogActorComponent, Log, TEXT("RegisterComponentWithWorld: Owner belongs to a DEADCLASS"));
 		return;
@@ -744,13 +746,13 @@ void UActorComponent::RegisterComponentWithWorld(UWorld* InWorld)
 
 #if !(UE_BUILD_SHIPPING || UE_BUILD_TEST)
 	// Can only register with an Actor if we are created within one
-	if(Owner != NULL)
+	if(MyOwner)
 	{
-		checkf(!Owner->HasAnyFlags(RF_Unreachable), TEXT("%s"), *GetFullName());
+		checkf(!MyOwner->HasAnyFlags(RF_Unreachable), TEXT("%s"), *GetFullName());
 		// can happen with undo because the owner will be restored "next"
-		//checkf(!Owner->IsPendingKill(), TEXT("%s"), *GetFullName());
+		//checkf(!MyOwner->IsPendingKill(), TEXT("%s"), *GetFullName());
 
-		if(InWorld != Owner->GetWorld())
+		if(InWorld != MyOwner->GetWorld())
 		{
 			// The only time you should specify a scene that is not Owner->GetWorld() is when you don't have an Actor
 			UE_LOG(LogActorComponent, Log, TEXT("RegisterComponentWithWorld: (%s) Specifying a world, but an Owner Actor found, and InWorld is not GetOwner()->GetWorld()"), *GetPathName());
@@ -763,7 +765,7 @@ void UActorComponent::RegisterComponentWithWorld(UWorld* InWorld)
 	ExecuteRegisterEvents();
 	RegisterAllComponentTickFunctions(true);
 
-	if (Owner == nullptr || Owner->IsActorInitialized())
+	if (MyOwner == nullptr || MyOwner->IsActorInitialized())
 	{
 		if (!bHasBeenInitialized && bWantsInitializeComponent)
 		{
@@ -771,7 +773,7 @@ void UActorComponent::RegisterComponentWithWorld(UWorld* InWorld)
 		}
 	}
 
-	if (Owner && Owner->HasActorBegunPlay())
+	if (MyOwner && MyOwner->HasActorBegunPlay())
 	{
 		if (bWantsBeginPlay)
 		{
@@ -799,9 +801,10 @@ void UActorComponent::RegisterComponentWithWorld(UWorld* InWorld)
 
 void UActorComponent::RegisterComponent()
 {
-	if (ensure(Owner && Owner->GetWorld()))
+	AActor* MyOwner = GetOwner();
+	if (ensure(MyOwner && MyOwner->GetWorld()))
 	{
-		RegisterComponentWithWorld(Owner->GetWorld());
+		RegisterComponentWithWorld(MyOwner->GetWorld());
 	}
 }
 
@@ -850,20 +853,20 @@ void UActorComponent::DestroyComponent(bool bPromoteChildren/*= false*/)
 	}
 
 	// Then remove from Components array, if we have an Actor
-	if(Owner != NULL)
+	if(AActor* MyOwner = GetOwner())
 	{
 		if (IsCreatedByConstructionScript())
 		{
-			Owner->BlueprintCreatedComponents.Remove(this);
+			MyOwner->BlueprintCreatedComponents.Remove(this);
 		}
 		else
 		{
-			Owner->RemoveInstanceComponent(this);
+			MyOwner->RemoveInstanceComponent(this);
 		}
-		Owner->RemoveOwnedComponent(this);
-		if (Owner->GetRootComponent() == this)
+		MyOwner->RemoveOwnedComponent(this);
+		if (MyOwner->GetRootComponent() == this)
 		{
-			Owner->SetRootComponent(NULL);
+			MyOwner->SetRootComponent(NULL);
 		}
 	}
 
@@ -888,14 +891,15 @@ void UActorComponent::OnComponentDestroyed()
 
 void UActorComponent::K2_DestroyComponent(UObject* Object)
 {
-	if (Object == this || Owner == NULL || Owner == Object)
+	AActor* MyOwner = GetOwner();
+	if (Object == this || MyOwner == NULL || MyOwner == Object)
 	{
 		DestroyComponent();
 	}
 	else
 	{
 		// TODO: Put in Message Log
-		UE_LOG(LogActorComponent, Error, TEXT("May not destroy component %s owned by %s."), *GetFullName(), *Owner->GetFullName());
+		UE_LOG(LogActorComponent, Error, TEXT("May not destroy component %s owned by %s."), *GetFullName(), *MyOwner->GetFullName());
 	}
 }
 
@@ -1055,13 +1059,14 @@ void UActorComponent::ConditionalTickComponent(float DeltaTime, enum ELevelTick
 {
 	if(bRegistered && !IsPendingKill())
 	{
+		AActor* MyOwner = GetOwner();
 		//@optimization, I imagine this is all unnecessary in a shipping game with no editor
 		if (TickType != LEVELTICK_ViewportsOnly || 
 			(bTickInEditor && TickType == LEVELTICK_ViewportsOnly) ||
-			(Owner && Owner->ShouldTickIfViewportsOnly())
+			(MyOwner && MyOwner->ShouldTickIfViewportsOnly())
 			)
 		{
-			const float TimeDilation = (Owner ? Owner->CustomTimeDilation : 1.f);
+			const float TimeDilation = (MyOwner ? MyOwner->CustomTimeDilation : 1.f);
 			TickComponent(DeltaTime * TimeDilation, TickType, &ThisTickFunction);
 		}
 	}
@@ -1278,7 +1283,8 @@ void UActorComponent::SetTickableWhenPaused(bool bTickableWhenPaused)
 
 bool UActorComponent::IsOwnerRunningUserConstructionScript() const
 {
-	return Owner != nullptr && Owner->bRunningUserConstructionScript;
+	AActor* MyOwner = GetOwner();
+	return (MyOwner && MyOwner->bRunningUserConstructionScript);
 }
 
 void UActorComponent::AddAssetUserData(UAssetUserData* InUserData)
@@ -1349,9 +1355,9 @@ void UActorComponent::SetIsReplicated(bool ShouldReplicate)
 	check(GetComponentClassCanReplicate()); // Only certain component classes can replicate!
 	bReplicates = ShouldReplicate;
 
-	if (Owner)
+	if (AActor* MyOwner = GetOwner())
 	{
-		Owner->UpdateReplicatedComponent( this );
+		MyOwner->UpdateReplicatedComponent( this );
 	}
 }
 
@@ -1372,7 +1378,8 @@ bool UActorComponent::GetComponentClassCanReplicate() const
 
 ENetRole UActorComponent::GetOwnerRole() const
 {
-	return (Owner ? Owner->Role.GetValue() : ROLE_None);
+	AActor* MyOwner = GetOwner();
+	return (MyOwner ? MyOwner->Role.GetValue() : ROLE_None);
 }
 
 bool UActorComponent::IsNetSimulating() const
@@ -1466,4 +1473,4 @@ void UActorComponent::GetUCSModifiedProperties(TSet<const UProperty*>& ModifiedP
 }
 
 
-#undef LOCTEXT_NAMESPACE
+#undef LOCTEXT_NAMESPACE
\ No newline at end of file
diff --git a/Engine/Source/Runtime/Engine/Private/Components/LightComponent.cpp b/Engine/Source/Runtime/Engine/Private/Components/LightComponent.cpp
--- a/Engine/Source/Runtime/Engine/Private/Components/LightComponent.cpp
+++ b/Engine/Source/Runtime/Engine/Private/Components/LightComponent.cpp
@@ -450,18 +450,6 @@ void ULightComponent::PostLoad()
 }
 
 #if WITH_EDITOR
-void ULightComponent::PreEditUndo()
-{
-	// Directly call UActorComponent::PreEditChange to avoid ULightComponent::PreEditChange being called for transactions.
-	UActorComponent::PreEditChange(NULL);
-}
-
-void ULightComponent::PostEditUndo()
-{
-	// Directly call UActorComponent::PostEditChange to avoid ULightComponent::PostEditChange being called for transactions.
-	UActorComponent::PostEditChange();
-}
-
 bool ULightComponent::CanEditChange(const UProperty* InProperty) const
 {
 	if (InProperty)
