Fix PHYSX_MEMORY_STAT not working because of header size

diff --git a/Engine/Source/Runtime/Engine/Private/PhysicsEngine/PhysXSupport.h b/Engine/Source/Runtime/Engine/Private/PhysicsEngine/PhysXSupport.h
--- a/Engine/Source/Runtime/Engine/Private/PhysicsEngine/PhysXSupport.h
+++ b/Engine/Source/Runtime/Engine/Private/PhysicsEngine/PhysXSupport.h
@@ -513,7 +513,14 @@ class FPhysXAllocator : public PxAllocatorCallback
 
 public:
 	FPhysXAllocator()
-	{}
+#if PHYSX_MEMORY_STATS
+	: HeaderSize(32) //physx needs 16 byte alignment. We need to inflate the header size to account for this.
+#endif
+	{
+#if PHYSX_MEMORY_STATS
+		static_assert(sizeof(FPhysXAllocationHeader) <= 32, "FPhysXAllocationHeader size must be less than HeaderSize bytes.");
+#endif
+	}
 
 	virtual ~FPhysXAllocator() 
 	{}
@@ -521,7 +528,6 @@ public:
 	virtual void* allocate(size_t size, const char* typeName, const char* filename, int line) override
 	{
 #if PHYSX_MEMORY_STATS
-		static_assert(sizeof(FPhysXAllocationHeader) <= 16, "FPhysXAllocationHeader size must be less than 16 bytes.");
 
 		INC_DWORD_STAT_BY(STAT_MemoryPhysXTotalAllocationSize, size);
 
@@ -529,7 +535,7 @@ public:
 		FName AllocationName(*AllocationString);
 
 		// Assign header
-		FPhysXAllocationHeader* AllocationHeader = (FPhysXAllocationHeader*)FMemory::Malloc(size + 16, 16);
+		FPhysXAllocationHeader* AllocationHeader = (FPhysXAllocationHeader*)FMemory::Malloc(size + HeaderSize, 16);
 		AllocationHeader->AllocationTypeName = AllocationName;
 		AllocationHeader->AllocationSize = size;
 
@@ -544,7 +550,7 @@ public:
 			AllocationsByType.Add(AllocationName, size);
 		}
 
-		return (uint8*)AllocationHeader + 16;
+		return (uint8*)AllocationHeader + HeaderSize;
 #else
 		void* ptr = FMemory::Malloc(size, 16);
 		#if PHYSX_MEMORY_STAT_ONLY
@@ -559,7 +565,7 @@ public:
 #if PHYSX_MEMORY_STATS
 		if( ptr )
 		{
-			FPhysXAllocationHeader* AllocationHeader = (FPhysXAllocationHeader*)((uint8*)ptr - 16);
+			FPhysXAllocationHeader* AllocationHeader = (FPhysXAllocationHeader*)((uint8*)ptr - HeaderSize);
 			DEC_DWORD_STAT_BY(STAT_MemoryPhysXTotalAllocationSize, AllocationHeader->AllocationSize);
 			size_t* TotalByType = AllocationsByType.Find(AllocationHeader->AllocationTypeName);
 			*TotalByType -= AllocationHeader->AllocationSize;
@@ -592,6 +598,11 @@ public:
 		}
 	}
 #endif
+
+private:
+#if PHYSX_MEMORY_STATS
+	const uint32 HeaderSize;
+#endif
 };
 
 class FPhysXBroadcastingAllocator : public PxBroadcastingAllocator
