Fixed matinee sound preview playback so that it respects the loop section and correctly restarts.  Improved the quality of audio playback when scrubbing.

diff --git a/Engine/Source/Runtime/Engine/Private/Interpolation.cpp b/Engine/Source/Runtime/Engine/Private/Interpolation.cpp
--- a/Engine/Source/Runtime/Engine/Private/Interpolation.cpp
+++ b/Engine/Source/Runtime/Engine/Private/Interpolation.cpp
@@ -8004,6 +8004,17 @@ void UInterpTrackSound::PreviewUpdateTrack(float NewPosition, UInterpTrackInst*
 		SoundInst->PlayAudioComp->Stop();
 	}
 
+	// If the new position for the track is before the last interp position, then the playback must have looped,
+	// so force playback to restart from the new position
+	const bool bJustLooped = NewPosition < MatineeActor->InterpPosition && MatineeActor->bIsPlaying;
+	if (bJustLooped)
+	{
+		SoundInst->PlayAudioComp->Stop();
+		bPlaying = false;
+		const float Epsilon = 0.1f;
+		SoundInst->LastUpdatePosition = NewPosition - Epsilon;
+	}
+
 	// Dont play sounds unless we are preview playback (ie not scrubbing).
 	bool bJump = !( MatineeActor->bIsPlaying );
 	UpdateTrack(NewPosition, TrInst, bJump);
@@ -8030,7 +8041,7 @@ void UInterpTrackSound::PreviewUpdateTrack(float NewPosition, UInterpTrackInst*
 			Component->bIsUISound = true;
 			Component->Play(NewPosition - SoundTrackKey.Time);
 
-			const float ScrubDuration = 0.05f;
+			const float ScrubDuration = 0.1f;
 			Component->FadeOut(ScrubDuration, 1.0f);
 		}
 	}
