Fix spectator camera gamepad input speed incorrectly being time dilated

diff --git a/Engine/Source/Runtime/Engine/Private/DefaultPawn.cpp b/Engine/Source/Runtime/Engine/Private/DefaultPawn.cpp
--- a/Engine/Source/Runtime/Engine/Private/DefaultPawn.cpp
+++ b/Engine/Source/Runtime/Engine/Private/DefaultPawn.cpp
@@ -183,13 +183,13 @@ void ADefaultPawn::MoveUp_World(float Val)
 void ADefaultPawn::TurnAtRate(float Rate)
 {
 	// calculate delta for this frame from the rate information
-	AddControllerYawInput(Rate * BaseTurnRate * GetWorld()->GetDeltaSeconds());
+	AddControllerYawInput(Rate * BaseTurnRate * GetWorld()->GetDeltaSeconds() * CustomTimeDilation);
 }
 
 void ADefaultPawn::LookUpAtRate(float Rate)
 {
 	// calculate delta for this frame from the rate information
-	AddControllerPitchInput(Rate * BaseLookUpRate * GetWorld()->GetDeltaSeconds());
+	AddControllerPitchInput(Rate * BaseLookUpRate * GetWorld()->GetDeltaSeconds() * CustomTimeDilation);
 }
 
 // @TODO: DEPRECATED, remove.
diff --git a/Engine/Source/Runtime/Engine/Private/DemoNetDriver.cpp b/Engine/Source/Runtime/Engine/Private/DemoNetDriver.cpp
--- a/Engine/Source/Runtime/Engine/Private/DemoNetDriver.cpp
+++ b/Engine/Source/Runtime/Engine/Private/DemoNetDriver.cpp
@@ -516,10 +516,19 @@ void UDemoNetDriver::TickDispatch(float DeltaSeconds)
 		//	(we want to continue to fly around in real-time)
 		if ( SpectatorController != NULL )
 		{
-			SpectatorController->CustomTimeDilation = 1.0f;
+			if ( World->GetWorldSettings()->DemoPlayTimeDilation > KINDA_SMALL_NUMBER )
+			{
+				SpectatorController->CustomTimeDilation = 1.0f / World->GetWorldSettings()->DemoPlayTimeDilation;
+			}
+			else
+			{
+				SpectatorController->CustomTimeDilation = 1.0f;
+			}
 
 			if ( SpectatorController->GetSpectatorPawn() != NULL )
 			{
+				SpectatorController->GetSpectatorPawn()->CustomTimeDilation = SpectatorController->CustomTimeDilation;
+
 				// Disable collision on the spectator
 				SpectatorController->GetSpectatorPawn()->SetActorEnableCollision( false );
 					
@@ -529,7 +538,7 @@ void UDemoNetDriver::TickDispatch(float DeltaSeconds)
 
 				if ( SpectatorMovement )
 				{
-					SpectatorMovement->bIgnoreTimeDilation = true;
+					//SpectatorMovement->bIgnoreTimeDilation = true;
 					SpectatorMovement->PrimaryComponentTick.bTickEvenWhenPaused = true;
 				}
 			}
