UMG - Fixing some touch related bugs for drag/drop, and fixing some immediate ending of drag/drop when simulating touch in the editor.

diff --git a/Engine/Source/Runtime/Engine/Private/GameViewportClient.cpp b/Engine/Source/Runtime/Engine/Private/GameViewportClient.cpp
--- a/Engine/Source/Runtime/Engine/Private/GameViewportClient.cpp
+++ b/Engine/Source/Runtime/Engine/Private/GameViewportClient.cpp
@@ -508,10 +508,14 @@ void UGameViewportClient::MouseLeave(FViewport* InViewport)
 
 	if (InViewport && GetDefault<UInputSettings>()->bUseMouseForTouch)
 	{
-		FIntPoint LastViewportCursorPos;
-		InViewport->GetMousePos(LastViewportCursorPos, false);
-		FVector2D CursorPos(LastViewportCursorPos.X, LastViewportCursorPos.Y);
-		FSlateApplication::Get().SetGameIsFakingTouchEvents(false, &CursorPos);
+		// Only send the touch end event if we're not drag/dropping, as that will end the drag/drop operation.
+		if ( !FSlateApplication::Get().IsDragDropping() )
+		{
+			FIntPoint LastViewportCursorPos;
+			InViewport->GetMousePos(LastViewportCursorPos, false);
+			FVector2D CursorPos(LastViewportCursorPos.X, LastViewportCursorPos.Y);
+			FSlateApplication::Get().SetGameIsFakingTouchEvents(false, &CursorPos);
+		}
 	}
 }
 
diff --git a/Engine/Source/Runtime/UMG/Private/WidgetBlueprintLibrary.cpp b/Engine/Source/Runtime/UMG/Private/WidgetBlueprintLibrary.cpp
--- a/Engine/Source/Runtime/UMG/Private/WidgetBlueprintLibrary.cpp
+++ b/Engine/Source/Runtime/UMG/Private/WidgetBlueprintLibrary.cpp
@@ -238,7 +238,7 @@ FEventReply UWidgetBlueprintLibrary::DetectDrag(UPARAM(ref) FEventReply& Reply,
 
 FEventReply UWidgetBlueprintLibrary::DetectDragIfPressed(const FPointerEvent& PointerEvent, UWidget* WidgetDetectingDrag, FKey DragKey)
 {
-	if ( PointerEvent.GetEffectingButton() == DragKey )
+	if ( PointerEvent.GetEffectingButton() == DragKey || PointerEvent.IsTouchEvent() )
 	{
 		FEventReply Reply = UWidgetBlueprintLibrary::Handled();
 		return UWidgetBlueprintLibrary::DetectDrag(Reply, WidgetDetectingDrag, DragKey);
