#ue4 - Remove needless "TODO".

diff --git a/Engine/Source/Runtime/Engine/Private/Character.cpp b/Engine/Source/Runtime/Engine/Private/Character.cpp
--- a/Engine/Source/Runtime/Engine/Private/Character.cpp
+++ b/Engine/Source/Runtime/Engine/Private/Character.cpp
@@ -2,7 +2,6 @@
 
 /*=============================================================================
 	Character.cpp: ACharacter implementation
-	TODO: Put description here
 =============================================================================*/
 
 #include "EnginePrivate.h"
