UE-14753 Add ENGINE_API to EditorApplyScale on TriggerBox/Sphere/Capsule to allow classes to be subclassed in other modules.

diff --git a/Engine/Source/Runtime/Engine/Classes/Engine/TriggerBox.h b/Engine/Source/Runtime/Engine/Classes/Engine/TriggerBox.h
--- a/Engine/Source/Runtime/Engine/Classes/Engine/TriggerBox.h
+++ b/Engine/Source/Runtime/Engine/Classes/Engine/TriggerBox.h
@@ -14,7 +14,7 @@ class ATriggerBox : public ATriggerBase
 
 #if WITH_EDITOR
 	// Begin AActor interface.
-	virtual void EditorApplyScale(const FVector& DeltaScale, const FVector* PivotLocation, bool bAltDown, bool bShiftDown, bool bCtrlDown) override;
+	ENGINE_API virtual void EditorApplyScale(const FVector& DeltaScale, const FVector* PivotLocation, bool bAltDown, bool bShiftDown, bool bCtrlDown) override;
 	// End AActor interface.
 #endif
 };
diff --git a/Engine/Source/Runtime/Engine/Classes/Engine/TriggerCapsule.h b/Engine/Source/Runtime/Engine/Classes/Engine/TriggerCapsule.h
--- a/Engine/Source/Runtime/Engine/Classes/Engine/TriggerCapsule.h
+++ b/Engine/Source/Runtime/Engine/Classes/Engine/TriggerCapsule.h
@@ -14,7 +14,7 @@ class ATriggerCapsule : public ATriggerBase
 
 #if WITH_EDITOR
 	// Begin AActor interface.
-	virtual void EditorApplyScale(const FVector& DeltaScale, const FVector* PivotLocation, bool bAltDown, bool bShiftDown, bool bCtrlDown) override;
+	ENGINE_API virtual void EditorApplyScale(const FVector& DeltaScale, const FVector* PivotLocation, bool bAltDown, bool bShiftDown, bool bCtrlDown) override;
 	// End AActor interface.
 #endif
 };
diff --git a/Engine/Source/Runtime/Engine/Classes/Engine/TriggerSphere.h b/Engine/Source/Runtime/Engine/Classes/Engine/TriggerSphere.h
--- a/Engine/Source/Runtime/Engine/Classes/Engine/TriggerSphere.h
+++ b/Engine/Source/Runtime/Engine/Classes/Engine/TriggerSphere.h
@@ -14,7 +14,7 @@ class ATriggerSphere : public ATriggerBase
 
 #if WITH_EDITOR
 	// Begin AActor interface.
-	virtual void EditorApplyScale(const FVector& DeltaScale, const FVector* PivotLocation, bool bAltDown, bool bShiftDown, bool bCtrlDown) override;
+	ENGINE_API virtual void EditorApplyScale(const FVector& DeltaScale, const FVector* PivotLocation, bool bAltDown, bool bShiftDown, bool bCtrlDown) override;
 	// End AActor interface.
 #endif
 };
