Fix for a cooking crash where a UEnum was looked up using StaticFindObject during saving

diff --git a/Engine/Source/Runtime/Engine/Classes/Materials/MaterialInterface.h b/Engine/Source/Runtime/Engine/Classes/Materials/MaterialInterface.h
--- a/Engine/Source/Runtime/Engine/Classes/Materials/MaterialInterface.h
+++ b/Engine/Source/Runtime/Engine/Classes/Materials/MaterialInterface.h
@@ -607,6 +607,13 @@ public:
 		}
 	}
 
+	/** Access the cached uenum type information for material sampler type */
+	static UEnum* GetSamplerTypeEnum() 
+	{ 
+		check(SamplerTypeEnum); 
+		return SamplerTypeEnum; 
+	}
+
 protected:
 
 	/** Returns a bitfield indicating which feature levels should be compiled for rendering. GMaxRHIFeatureLevel is always present */
@@ -620,6 +627,10 @@ private:
 	 */
 	static void PostLoadDefaultMaterials();
 
+	/**
+	* Cached type information for the sampler type enumeration. 
+	*/
+	static UEnum* SamplerTypeEnum;
 };
 
 /** Helper function to serialize inline shader maps for the given material resources. */
diff --git a/Engine/Source/Runtime/Engine/Private/Materials/MaterialExpressions.cpp b/Engine/Source/Runtime/Engine/Private/Materials/MaterialExpressions.cpp
--- a/Engine/Source/Runtime/Engine/Private/Materials/MaterialExpressions.cpp
+++ b/Engine/Source/Runtime/Engine/Private/Materials/MaterialExpressions.cpp
@@ -1263,7 +1263,7 @@ static bool VerifySamplerType(
 		EMaterialSamplerType CorrectSamplerType = UMaterialExpressionTextureBase::GetSamplerTypeForTexture( Texture );
 		if ( SamplerType != CorrectSamplerType )
 		{
-			UEnum* SamplerTypeEnum = FindObject<UEnum>( NULL, TEXT("/Script/Engine.EMaterialSamplerType") );
+			UEnum* SamplerTypeEnum = UMaterialInterface::GetSamplerTypeEnum();
 			check( SamplerTypeEnum );
 
 			FString SamplerTypeDisplayName = SamplerTypeEnum->GetEnumText(SamplerType).ToString();
@@ -1278,7 +1278,7 @@ static bool VerifySamplerType(
 		}
 		if((SamplerType == SAMPLERTYPE_Normal || SamplerType == SAMPLERTYPE_Masks) && Texture->SRGB)
 		{
-			UEnum* SamplerTypeEnum = FindObject<UEnum>( NULL, TEXT("/Script/Engine.EMaterialSamplerType") );
+			UEnum* SamplerTypeEnum = UMaterialInterface::GetSamplerTypeEnum();
 			check( SamplerTypeEnum );
 
 			FString SamplerTypeDisplayName = SamplerTypeEnum->GetEnumText(SamplerType).ToString();
diff --git a/Engine/Source/Runtime/Engine/Private/Materials/MaterialInterface.cpp b/Engine/Source/Runtime/Engine/Private/Materials/MaterialInterface.cpp
--- a/Engine/Source/Runtime/Engine/Private/Materials/MaterialInterface.cpp
+++ b/Engine/Source/Runtime/Engine/Private/Materials/MaterialInterface.cpp
@@ -9,6 +9,10 @@
 
 //////////////////////////////////////////////////////////////////////////
 
+UEnum* UMaterialInterface::SamplerTypeEnum = nullptr;
+
+//////////////////////////////////////////////////////////////////////////
+
 /** Copies the material's relevance flags to a primitive's view relevance flags. */
 void FMaterialRelevance::SetPrimitiveViewRelevance(FPrimitiveViewRelevance& OutViewRelevance) const
 {
@@ -29,6 +33,9 @@ UMaterialInterface::UMaterialInterface(const FObjectInitializer& ObjectInitializ
 	{
 		InitDefaultMaterials();
 		AssertDefaultMaterialsExist();
+
+		SamplerTypeEnum = FindObject<UEnum>(NULL, TEXT("/Script/Engine.EMaterialSamplerType"));
+		check(SamplerTypeEnum);
 	}
 }
 
