Track checkpoints and "socket" send size in the network profiler while recording replays.

diff --git a/Engine/Source/Runtime/Engine/Classes/Engine/DemoNetConnection.h b/Engine/Source/Runtime/Engine/Classes/Engine/DemoNetConnection.h
--- a/Engine/Source/Runtime/Engine/Classes/Engine/DemoNetConnection.h
+++ b/Engine/Source/Runtime/Engine/Classes/Engine/DemoNetConnection.h
@@ -48,4 +48,7 @@ public:
 	}
 
 	TArray<FQueuedDemoPacket> QueuedDemoPackets;
+
+private:
+	void TrackSendForProfiler(const void* Data, int32 NumBytes);
 };
diff --git a/Engine/Source/Runtime/Engine/Private/DemoNetDriver.cpp b/Engine/Source/Runtime/Engine/Private/DemoNetDriver.cpp
--- a/Engine/Source/Runtime/Engine/Private/DemoNetDriver.cpp
+++ b/Engine/Source/Runtime/Engine/Private/DemoNetDriver.cpp
@@ -21,6 +21,7 @@
 #include "Net/NetworkProfiler.h"
 #include "Net/DataReplication.h"
 #include "GameFramework/GameMode.h"
+#include "SocketSubsystem.h"
 
 DEFINE_LOG_CATEGORY( LogDemo );
 
@@ -1709,6 +1710,8 @@ void UDemoNetConnection::LowLevelSend( void* Data, int32 Count )
 	
 		*CheckpointArchive << Count;
 		CheckpointArchive->Serialize( Data, Count );
+
+		TrackSendForProfiler( Data, Count );
 		return;
 	}
 
@@ -1739,8 +1742,8 @@ void UDemoNetConnection::LowLevelSend( void* Data, int32 Count )
 		*FileAr << Count;
 		FileAr->Serialize( Data, Count );
 		
-		NETWORK_PROFILER(GNetworkProfiler.FlushOutgoingBunches(this));
-
+		TrackSendForProfiler( Data, Count );
+		
 #if DEMO_CHECKSUMS == 1
 		uint32 Checksum = FCrc::MemCrc32( Data, Count, 0 );
 		*FileAr << Checksum;
@@ -1748,6 +1751,15 @@ void UDemoNetConnection::LowLevelSend( void* Data, int32 Count )
 	}
 }
 
+void UDemoNetConnection::TrackSendForProfiler(const void* Data, int32 NumBytes)
+{
+	NETWORK_PROFILER(GNetworkProfiler.FlushOutgoingBunches(this));
+
+	// Track "socket send" even though we're not technically sending to a socket, to get more accurate information in the profiler.
+	TSharedRef<FInternetAddr> FakeAddr = ISocketSubsystem::Get(PLATFORM_SOCKETSUBSYSTEM)->CreateInternetAddr();
+	NETWORK_PROFILER(GNetworkProfiler.TrackSocketSendTo(TEXT("Unreal"), Data, NumBytes, NumPacketIdBits, NumBunchBits, NumAckBits, NumPaddingBits, *FakeAddr));
+}
+
 FString UDemoNetConnection::LowLevelDescribe()
 {
 	return TEXT( "Demo recording/playback driver connection" );
