Don't create unnecessary FInternetAddrs while recording replays

diff --git a/Engine/Source/Runtime/Engine/Private/DemoNetDriver.cpp b/Engine/Source/Runtime/Engine/Private/DemoNetDriver.cpp
--- a/Engine/Source/Runtime/Engine/Private/DemoNetDriver.cpp
+++ b/Engine/Source/Runtime/Engine/Private/DemoNetDriver.cpp
@@ -21,7 +21,6 @@
 #include "Net/NetworkProfiler.h"
 #include "Net/DataReplication.h"
 #include "GameFramework/GameMode.h"
-#include "SocketSubsystem.h"
 
 DEFINE_LOG_CATEGORY( LogDemo );
 
@@ -1756,8 +1755,7 @@ void UDemoNetConnection::TrackSendForProfiler(const void* Data, int32 NumBytes)
 	NETWORK_PROFILER(GNetworkProfiler.FlushOutgoingBunches(this));
 
 	// Track "socket send" even though we're not technically sending to a socket, to get more accurate information in the profiler.
-	TSharedRef<FInternetAddr> FakeAddr = ISocketSubsystem::Get(PLATFORM_SOCKETSUBSYSTEM)->CreateInternetAddr();
-	NETWORK_PROFILER(GNetworkProfiler.TrackSocketSendTo(TEXT("Unreal"), Data, NumBytes, NumPacketIdBits, NumBunchBits, NumAckBits, NumPaddingBits, *FakeAddr));
+	NETWORK_PROFILER(GNetworkProfiler.TrackSocketSendToCore(TEXT("Unreal"), Data, NumBytes, NumPacketIdBits, NumBunchBits, NumAckBits, NumPaddingBits, 0));
 }
 
 FString UDemoNetConnection::LowLevelDescribe()
