Fix NON editor builds

diff --git a/Engine/Source/Runtime/Engine/Private/Components/PrimitiveComponent.cpp b/Engine/Source/Runtime/Engine/Private/Components/PrimitiveComponent.cpp
--- a/Engine/Source/Runtime/Engine/Private/Components/PrimitiveComponent.cpp
+++ b/Engine/Source/Runtime/Engine/Private/Components/PrimitiveComponent.cpp
@@ -2606,9 +2606,12 @@ void UPrimitiveComponent::SetCustomNavigableGeometry(const EHasCustomNavigableGe
 	bHasCustomNavigableGeometry = InType;
 }
 
+#if WITH_EDITOR
 const bool UPrimitiveComponent::ShouldGenerateAutoLOD() const
 {
 	return (Mobility != EComponentMobility::Movable);
 }
+#endif 
+
 #undef LOCTEXT_NAMESPACE
 
