Fixed incorrect component transforms on non-root rigid body components when running standalone. (UE-13917).
Essentially undid CL 2422323 and re-solved that issue (UE-7420) in a more robust way.

diff --git a/Engine/Source/Runtime/Engine/Private/Components/SceneComponent.cpp b/Engine/Source/Runtime/Engine/Private/Components/SceneComponent.cpp
--- a/Engine/Source/Runtime/Engine/Private/Components/SceneComponent.cpp
+++ b/Engine/Source/Runtime/Engine/Private/Components/SceneComponent.cpp
@@ -1162,7 +1162,7 @@ void USceneComponent::AttachTo(class USceneComponent* Parent, FName InSocketName
 			//Also physics state may not be created yet so we use bSimulatePhysics to determine if the object has any intention of being physically simulated
 			UPrimitiveComponent * PrimitiveComponent = Cast<UPrimitiveComponent>(this);
 
-			if (PrimitiveComponent && PrimitiveComponent->BodyInstance.bSimulatePhysics && !bWeldSimulatedBodies && GetWorld() && GetWorld()->IsGameWorld() && !GetWorld()->bIsRunningConstructionScript)
+			if (PrimitiveComponent && PrimitiveComponent->BodyInstance.bSimulatePhysics && !bWeldSimulatedBodies && GetWorld() && GetWorld()->IsGameWorld() /*&& !GetWorld()->bIsRunningConstructionScript*/)
 			{
 				//Since the object is physically simulated it can't be the case that it's a child of object A and being attached to object B (at runtime)
 				if (bMaintainWorldPosition == false)	//User tried to attach but physically based so detach. However, if they provided relative coordinates we should still get the correct position
diff --git a/Engine/Source/Runtime/Engine/Private/SimpleConstructionScript.cpp b/Engine/Source/Runtime/Engine/Private/SimpleConstructionScript.cpp
--- a/Engine/Source/Runtime/Engine/Private/SimpleConstructionScript.cpp
+++ b/Engine/Source/Runtime/Engine/Private/SimpleConstructionScript.cpp
@@ -441,6 +441,7 @@ void USimpleConstructionScript::ExecuteScriptOnActor(AActor* Actor, const FTrans
 {
 	if(RootNodes.Num() > 0)
 	{
+		TSet<UActorComponent*> AllComponentsCreatedBySCS;
 		TInlineComponentArray<UActorComponent*> InstancedComponents;
 		for(auto NodeIt = RootNodes.CreateIterator(); NodeIt; ++NodeIt)
 		{
@@ -498,12 +499,24 @@ void USimpleConstructionScript::ExecuteScriptOnActor(AActor* Actor, const FTrans
 					}
 				}
 
+
 				// Create the new component instance and any child components it may have
 				UActorComponent* InstancedComponent = RootNode->ExecuteNodeOnActor(Actor, ParentComponent != nullptr ? ParentComponent : RootComponent, &RootTransform, bIsDefaultTransform);
 				if(InstancedComponent != nullptr)
 				{
 					InstancedComponents.Add(InstancedComponent);
 				}
+
+				// get list of every component SCS created, in case some of them aren't in the attachment hierarchy any more (e.g. rigid bodies)
+				TInlineComponentArray<USceneComponent*> ComponentsAfterSCS;
+				Actor->GetComponents(ComponentsAfterSCS);
+				for (auto C : ComponentsAfterSCS)
+				{
+					if (Components.Contains(C) == false)
+					{
+						AllComponentsCreatedBySCS.Add(C);
+					}
+				}
 			}
 		}
 
@@ -513,6 +526,15 @@ void USimpleConstructionScript::ExecuteScriptOnActor(AActor* Actor, const FTrans
 		{
 			RegisterInstancedComponent(InstancedComponent);
 		}
+
+		// now that the instanced components in the attachment hierarchy are registered, register any other components that SCS made but aren't in the attachment hierarchy for whatever reason.
+		for (auto C : AllComponentsCreatedBySCS)
+		{
+			if (C->IsRegistered() == false)
+			{
+				C->RegisterComponent();
+			}
+		}
 	}
 	else if(Actor->GetRootComponent() == NULL) // Must have a root component at the end of SCS, so if we don't have one already (from base class), create a SceneComponent now
 	{
