#ue4 - Perf: Optimize usage of MasterPoseComponent weak pointer in SkinnedMeshComponent.

diff --git a/Engine/Source/Runtime/Engine/Private/Components/SkinnedMeshComponent.cpp b/Engine/Source/Runtime/Engine/Private/Components/SkinnedMeshComponent.cpp
--- a/Engine/Source/Runtime/Engine/Private/Components/SkinnedMeshComponent.cpp
+++ b/Engine/Source/Runtime/Engine/Private/Components/SkinnedMeshComponent.cpp
@@ -604,8 +604,9 @@ FBoxSphereBounds USkinnedMeshComponent::CalcMeshBound(const FVector& RootOffset,
 	AActor* Owner = GetOwner();
 	FVector DrawScale = LocalToWorld.GetScale3D();	
 
+	const USkinnedMeshComponent* const MasterPoseComponentInst = MasterPoseComponent.Get();
 	UPhysicsAsset * const PhysicsAsset = GetPhysicsAsset();
-	UPhysicsAsset * const MasterPhysicsAsset = (MasterPoseComponent != NULL)? MasterPoseComponent->GetPhysicsAsset() : NULL;
+	UPhysicsAsset * const MasterPhysicsAsset = (MasterPoseComponentInst != nullptr)? MasterPoseComponentInst->GetPhysicsAsset() : nullptr;
 
 	// Can only use the PhysicsAsset to calculate the bounding box if we are not non-uniformly scaling the mesh.
 	const bool bCanUsePhysicsAsset = DrawScale.IsUniform() && (SkeletalMesh != NULL)
@@ -625,16 +626,16 @@ FBoxSphereBounds USkinnedMeshComponent::CalcMeshBound(const FVector& RootOffset,
 		RootAdjustedBounds.Origin += RootOffset; // Adjust bounds by root bone translation
 		NewBounds = RootAdjustedBounds.TransformBy(LocalToWorld);
 	}
-	else if(MasterPoseComponent.IsValid() && MasterPoseComponent->SkeletalMesh && MasterPoseComponent->bComponentUseFixedSkelBounds)
+	else if(MasterPoseComponentInst && MasterPoseComponentInst->SkeletalMesh && MasterPoseComponentInst->bComponentUseFixedSkelBounds)
 	{
-		FBoxSphereBounds RootAdjustedBounds = MasterPoseComponent->SkeletalMesh->Bounds;
+		FBoxSphereBounds RootAdjustedBounds = MasterPoseComponentInst->SkeletalMesh->Bounds;
 		RootAdjustedBounds.Origin += RootOffset; // Adjust bounds by root bone translation
 		NewBounds = RootAdjustedBounds.TransformBy(LocalToWorld);
 	}
 	// Use MasterPoseComponent's PhysicsAsset if told to
-	else if (MasterPoseComponent.IsValid() && bCanUsePhysicsAsset && bUseBoundsFromMasterPoseComponent)
+	else if (MasterPoseComponentInst && bCanUsePhysicsAsset && bUseBoundsFromMasterPoseComponent)
 	{
-		NewBounds = MasterPoseComponent->Bounds;
+		NewBounds = MasterPoseComponentInst->Bounds;
 	}
 #if WITH_EDITOR
 	// For AnimSet Viewer, use 'bounds preview' physics asset if present.
@@ -649,7 +650,7 @@ FBoxSphereBounds USkinnedMeshComponent::CalcMeshBound(const FVector& RootOffset,
 		NewBounds = FBoxSphereBounds(PhysicsAsset->CalcAABB(this, LocalToWorld));
 	}
 	// Use MasterPoseComponent's PhysicsAsset, if we don't have one and it does
-	else if(MasterPoseComponent.IsValid() && bCanUsePhysicsAsset && bMasterHasPhysBodies)
+	else if(MasterPoseComponentInst && bCanUsePhysicsAsset && bMasterHasPhysBodies)
 	{
 		NewBounds = FBoxSphereBounds(MasterPhysicsAsset->CalcAABB(this, LocalToWorld));
 	}
@@ -687,7 +688,8 @@ FMatrix USkinnedMeshComponent::GetBoneMatrix(int32 BoneIdx) const
 	}
 
 	// Handle case of use a MasterPoseComponent - get bone matrix from there.
-	if(MasterPoseComponent.IsValid())
+	const USkinnedMeshComponent* const MasterPoseComponentInst = MasterPoseComponent.Get();
+	if(MasterPoseComponentInst)
 	{
 		if(BoneIdx < MasterBoneMap.Num())
 		{
@@ -695,9 +697,9 @@ FMatrix USkinnedMeshComponent::GetBoneMatrix(int32 BoneIdx) const
 
 			// If ParentBoneIndex is valid, grab matrix from MasterPoseComponent.
 			if(	ParentBoneIndex != INDEX_NONE && 
-				ParentBoneIndex < MasterPoseComponent->GetNumSpaceBases())
+				ParentBoneIndex < MasterPoseComponentInst->GetNumSpaceBases())
 			{
-				return MasterPoseComponent->GetSpaceBases()[ParentBoneIndex].ToMatrixWithScale() * ComponentToWorld.ToMatrixWithScale();
+				return MasterPoseComponentInst->GetSpaceBases()[ParentBoneIndex].ToMatrixWithScale() * ComponentToWorld.ToMatrixWithScale();
 			}
 			else
 			{
@@ -740,7 +742,8 @@ FTransform USkinnedMeshComponent::GetBoneTransform(int32 BoneIdx) const
 FTransform USkinnedMeshComponent::GetBoneTransform(int32 BoneIdx, const FTransform& LocalToWorld) const
 {
 	// Handle case of use a MasterPoseComponent - get bone matrix from there.
-	if(MasterPoseComponent.IsValid())
+	const USkinnedMeshComponent* const MasterPoseComponentInst = MasterPoseComponent.Get();
+	if(MasterPoseComponentInst)
 	{
 		if(BoneIdx < MasterBoneMap.Num())
 		{
@@ -748,9 +751,9 @@ FTransform USkinnedMeshComponent::GetBoneTransform(int32 BoneIdx, const FTransfo
 
 			// If ParentBoneIndex is valid, grab matrix from MasterPoseComponent.
 			if(	ParentBoneIndex != INDEX_NONE && 
-				ParentBoneIndex < MasterPoseComponent->GetNumSpaceBases())
+				ParentBoneIndex < MasterPoseComponentInst->GetNumSpaceBases())
 			{
-				return MasterPoseComponent->GetSpaceBases()[ParentBoneIndex] * LocalToWorld;
+				return MasterPoseComponentInst->GetSpaceBases()[ParentBoneIndex] * LocalToWorld;
 			}
 			else
 			{
@@ -1196,16 +1199,17 @@ FQuat USkinnedMeshComponent::GetBoneQuaternion(FName BoneName, EBoneSpaces::Type
 	FTransform BoneTransform;
 	if( Space == EBoneSpaces::ComponentSpace )
 	{
-		if(MasterPoseComponent.IsValid())
+		const USkinnedMeshComponent* const MasterPoseComponentInst = MasterPoseComponent.Get();
+		if(MasterPoseComponentInst)
 		{
 			if(BoneIndex < MasterBoneMap.Num())
 			{
 				int32 ParentBoneIndex = MasterBoneMap[BoneIndex];
 				// If ParentBoneIndex is valid, grab matrix from MasterPoseComponent.
 				if(	ParentBoneIndex != INDEX_NONE && 
-					ParentBoneIndex < MasterPoseComponent->GetNumSpaceBases())
+					ParentBoneIndex < MasterPoseComponentInst->GetNumSpaceBases())
 				{
-					BoneTransform = MasterPoseComponent->GetSpaceBases()[ParentBoneIndex];
+					BoneTransform = MasterPoseComponentInst->GetSpaceBases()[ParentBoneIndex];
 				}
 				else
 				{
@@ -1243,16 +1247,17 @@ FVector USkinnedMeshComponent::GetBoneLocation(FName BoneName, EBoneSpaces::Type
 
 	if( Space == EBoneSpaces::ComponentSpace )
 	{
-		if(MasterPoseComponent.IsValid())
+		const USkinnedMeshComponent* const MasterPoseComponentInst = MasterPoseComponent.Get();
+		if(MasterPoseComponentInst)
 		{
 			if(BoneIndex < MasterBoneMap.Num())
 			{
 				int32 ParentBoneIndex = MasterBoneMap[BoneIndex];
 				// If ParentBoneIndex is valid, grab transform from MasterPoseComponent.
 				if(	ParentBoneIndex != INDEX_NONE && 
-					ParentBoneIndex < MasterPoseComponent->GetNumSpaceBases())
+					ParentBoneIndex < MasterPoseComponentInst->GetNumSpaceBases())
 				{
-					return MasterPoseComponent->GetSpaceBases()[ParentBoneIndex].GetLocation();
+					return MasterPoseComponentInst->GetSpaceBases()[ParentBoneIndex].GetLocation();
 				}
 			}
 			
@@ -1496,7 +1501,8 @@ FORCEINLINE FVector USkinnedMeshComponent::GetTypedSkinnedVertexPosition(const F
 {
 	FVector SkinnedPos(0,0,0);
 
-	const USkinnedMeshComponent* BaseComponent = MasterPoseComponent.IsValid() ? MasterPoseComponent.Get() : this;
+	const USkinnedMeshComponent* const MasterPoseComponentInst = MasterPoseComponent.Get();
+	const USkinnedMeshComponent* BaseComponent = MasterPoseComponentInst ? MasterPoseComponentInst : this;
 
 	// Do soft skinning for this vertex.
 	if(bSoftVertex)
@@ -1511,7 +1517,7 @@ FORCEINLINE FVector USkinnedMeshComponent::GetTypedSkinnedVertexPosition(const F
 #endif
 		{
 			int32 BoneIndex = Chunk.BoneMap[SrcSoftVertex->InfluenceBones[InfluenceIndex]];
-			if(MasterPoseComponent.IsValid())
+			if(MasterPoseComponentInst)
 			{		
 				check(MasterBoneMap.Num() == SkeletalMesh->RefSkeleton.GetNum());
 				BoneIndex = MasterBoneMap[BoneIndex];
@@ -1538,7 +1544,7 @@ FORCEINLINE FVector USkinnedMeshComponent::GetTypedSkinnedVertexPosition(const F
 		const TGPUSkinVertexBase<false>* SrcRigidVertex = VertexBufferGPUSkin.GetVertexPtr<false>(Chunk.GetRigidVertexBufferIndex()+VertIndex);
 		const int32 RigidInfluenceIndex = SkinningTools::GetRigidInfluenceIndex();
 		int32 BoneIndex = Chunk.BoneMap[SrcRigidVertex->InfluenceBones[RigidInfluenceIndex]];
-		if(MasterPoseComponent.IsValid())
+		if(MasterPoseComponentInst)
 		{
 			check(MasterBoneMap.Num() == SkeletalMesh->RefSkeleton.GetNum());
 			BoneIndex = MasterBoneMap[BoneIndex];
