Added flag to turn off physx serialization alignment

diff --git a/Engine/Source/Runtime/Engine/Private/PhysicsEngine/PhysicsSerializer.cpp b/Engine/Source/Runtime/Engine/Private/PhysicsEngine/PhysicsSerializer.cpp
--- a/Engine/Source/Runtime/Engine/Private/PhysicsEngine/PhysicsSerializer.cpp
+++ b/Engine/Source/Runtime/Engine/Private/PhysicsEngine/PhysicsSerializer.cpp
@@ -23,7 +23,11 @@ FByteBulkData* UPhysicsSerializer::GetBinaryData(FName Format, const TArray<FBod
 	QUICK_SCOPE_CYCLE_COUNTER(STAT_GetBinaryData);
 	const bool bContainedData = BinaryFormatData.Contains(Format);
 	FByteBulkData* Result = &BinaryFormatData.GetFormat(Format);
-	Result->SetBulkDataAlignment(PHYSX_SERIALIZATION_ALIGNMENT);
+	if (!FParse::Param(FCommandLine::Get(), TEXT("NoPhysxAlignment")))
+	{
+		Result->SetBulkDataAlignment(PHYSX_SERIALIZATION_ALIGNMENT);
+	}
+	
 	if (!bContainedData)
 	{
 #if WITH_EDITOR
