Store USoundCue::SoundWave as a TAssetPtr instead of a hard reference (first step towards only loading waves for the current audio quality)

diff --git a/Engine/Source/Developer/MovieSceneTools/Private/AudioTrackEditor.cpp b/Engine/Source/Developer/MovieSceneTools/Private/AudioTrackEditor.cpp
--- a/Engine/Source/Developer/MovieSceneTools/Private/AudioTrackEditor.cpp
+++ b/Engine/Source/Developer/MovieSceneTools/Private/AudioTrackEditor.cpp
@@ -58,12 +58,11 @@ USoundWave* DeriveSoundWave(UMovieSceneAudioSection* AudioSection)
 		// @todo Sequencer - Right now for sound cues, we just use the first sound wave in the cue
 		// In the future, it would be better to properly generate the sound cue's data after forcing determinism
 		const TArray<USoundNode*>& AllNodes = SoundCue->AllNodes;
-		for (int32 i = 0; i < AllNodes.Num(); ++i)
+		for (int32 i = 0; i < AllNodes.Num() && SoundWave == nullptr; ++i)
 		{
 			if (AllNodes[i]->IsA<USoundNodeWavePlayer>())
 			{
-				SoundWave = Cast<USoundNodeWavePlayer>(AllNodes[i])->SoundWave;
-				break;
+				SoundWave = Cast<USoundNodeWavePlayer>(AllNodes[i])->SoundWave.Get();
 			}
 		}
 	}
@@ -75,7 +74,7 @@ float DeriveUnloopedDuration(UMovieSceneAudioSection* AudioSection)
 {
 	USoundWave* SoundWave = DeriveSoundWave(AudioSection);
 
-	float Duration = SoundWave->GetDuration();
+	const float Duration = (SoundWave ? SoundWave->GetDuration() : 0.f);
 	return Duration == INDEFINITELY_LOOPING_DURATION ? SoundWave->Duration : Duration;
 }
 
diff --git a/Engine/Source/Editor/SoundCueEditor/Private/SoundCueEditor.cpp b/Engine/Source/Editor/SoundCueEditor/Private/SoundCueEditor.cpp
--- a/Engine/Source/Editor/SoundCueEditor/Private/SoundCueEditor.cpp
+++ b/Engine/Source/Editor/SoundCueEditor/Private/SoundCueEditor.cpp
@@ -390,7 +390,7 @@ void FSoundCueEditor::SyncInBrowser()
 			USoundNodeWavePlayer* SelectedWave = Cast<USoundNodeWavePlayer>(SelectedNode->SoundNode);
 			if (SelectedWave && SelectedWave->SoundWave)
 			{
-				ObjectsToSync.AddUnique(SelectedWave->SoundWave);
+				ObjectsToSync.AddUnique(SelectedWave->SoundWave.Get());
 			}
 			//ObjectsToSync.AddUnique(SelectedNode->SoundNode);
 		}
diff --git a/Engine/Source/Runtime/Engine/Classes/Sound/SoundNodeWavePlayer.h b/Engine/Source/Runtime/Engine/Classes/Sound/SoundNodeWavePlayer.h
--- a/Engine/Source/Runtime/Engine/Classes/Sound/SoundNodeWavePlayer.h
+++ b/Engine/Source/Runtime/Engine/Classes/Sound/SoundNodeWavePlayer.h
@@ -14,7 +14,7 @@ class USoundNodeWavePlayer : public USoundNode
 	GENERATED_UCLASS_BODY()
 
 	UPROPERTY(EditAnywhere, Category=WavePlayer)
-	USoundWave* SoundWave;
+	TAssetPtr<USoundWave> SoundWave;
 
 	UPROPERTY(EditAnywhere, Category=WavePlayer)
 	uint32 bLooping:1;
diff --git a/Engine/Source/Runtime/Engine/Private/Audio.cpp b/Engine/Source/Runtime/Engine/Private/Audio.cpp
--- a/Engine/Source/Runtime/Engine/Private/Audio.cpp
+++ b/Engine/Source/Runtime/Engine/Private/Audio.cpp
@@ -99,7 +99,7 @@ FName FSoundBuffer::GetSoundClassName()
 			// look through them to see if this cue uses a wave this buffer is bound to, via ResourceID
 			for (int32 WaveIndex = 0; WaveIndex < WavePlayers.Num(); ++WaveIndex)
 			{
-				USoundWave* WaveNode = WavePlayers[WaveIndex]->SoundWave;
+				USoundWave* WaveNode = WavePlayers[WaveIndex]->SoundWave.Get();
 				if (WaveNode != NULL)
 				{
 					if (WaveNode->ResourceID == ResourceID)
diff --git a/Engine/Source/Runtime/Engine/Private/AudioDevice.cpp b/Engine/Source/Runtime/Engine/Private/AudioDevice.cpp
--- a/Engine/Source/Runtime/Engine/Private/AudioDevice.cpp
+++ b/Engine/Source/Runtime/Engine/Private/AudioDevice.cpp
@@ -315,7 +315,7 @@ void FAudioDevice::GetSoundClassInfo( TMap<FName, FAudioClassInfo>& AudioClassIn
 		for( int32 WaveIndex = 0; WaveIndex < WavePlayers.Num(); ++WaveIndex )
 		{
 			// Presume one class per sound node wave
-			USoundWave *SoundWave = WavePlayers[ WaveIndex ]->SoundWave;
+			USoundWave *SoundWave = WavePlayers[ WaveIndex ]->SoundWave.Get();
 			if (SoundWave && SoundCue->GetSoundClass())
 			{
 				SoundWaveClasses.Add( SoundWave, SoundCue->GetSoundClass()->GetFName() );
@@ -594,7 +594,7 @@ bool FAudioDevice::HandlePlaySoundCueCommand( const TCHAR* Cmd, FOutputDevice& A
 			Cue->RecursiveFindNode<USoundNodeWavePlayer>( Cue->FirstNode, WavePlayers );
 			for( int32 i = 0; i < WavePlayers.Num(); ++i )
 			{
-				USoundWave* SoundWave = WavePlayers[ i ]->SoundWave;
+				USoundWave* SoundWave = WavePlayers[ i ]->SoundWave.Get();
 				if (SoundWave)
 				{
 					SoundWave->LogSubtitle( Ar );
diff --git a/Engine/Source/Runtime/Engine/Private/SoundCue.cpp b/Engine/Source/Runtime/Engine/Private/SoundCue.cpp
--- a/Engine/Source/Runtime/Engine/Private/SoundCue.cpp
+++ b/Engine/Source/Runtime/Engine/Private/SoundCue.cpp
@@ -61,7 +61,8 @@ void USoundCue::PostLoad()
 {
 	Super::PostLoad();
 
-		// Game doesn't care if there are NULL graph nodes
+	// Game doesn't care if there are NULL graph nodes
+#if WITH_EDITOR
 	if (GIsEditor)
 	{
 		// Deal with SoundNode types being removed - iterate in reverse as nodes may be removed
@@ -74,6 +75,41 @@ void USoundCue::PostLoad()
 				FBlueprintEditorUtils::RemoveNode(NULL, Node, true);
 			}
 		}
+
+		// Always load all sound waves in the editor
+		for (USoundNode* SoundNode : AllNodes)
+		{
+			if (USoundNodeWavePlayer* WavePlayerNode = Cast<USoundNodeWavePlayer>(SoundNode))
+			{
+				WavePlayerNode->SoundWave.LoadSynchronous();
+			}
+		}
+	}
+	else
+#endif
+	{
+		TArray<USoundNode*> NodesToEvaluate;
+		NodesToEvaluate.Push(FirstNode);
+
+		while (NodesToEvaluate.Num() > 0)
+		{
+			if (USoundNode* SoundNode = NodesToEvaluate.Pop(false))
+			{
+				if (USoundNodeWavePlayer* WavePlayerNode = Cast<USoundNodeWavePlayer>(SoundNode))
+				{
+					WavePlayerNode->SoundWave.LoadSynchronous();
+				}
+				/* else if USoundNodeQuality
+				{
+					// Only pick the node connected for current quality
+				}
+				*/
+				else
+				{
+					NodesToEvaluate.Append(SoundNode->ChildNodes);
+				}
+			}
+		}
 	}
 }
 
@@ -205,7 +241,7 @@ SIZE_T USoundCue::GetResourceSize(EResourceSizeMode::Type Mode)
 
 		for( int32 WaveIndex = 0; WaveIndex < WavePlayers.Num(); ++WaveIndex )
 		{
-			USoundWave* SoundWave = WavePlayers[WaveIndex]->SoundWave;
+			USoundWave* SoundWave = WavePlayers[WaveIndex]->SoundWave.Get();
 			if (SoundWave)
 			{
 				ResourceSize += SoundWave->GetResourceSize(Mode);
@@ -224,7 +260,7 @@ int32 USoundCue::GetResourceSizeForFormat(FName Format)
 	int32 ResourceSize = 0;
 	for (int32 WaveIndex = 0; WaveIndex < WavePlayers.Num(); ++WaveIndex)
 	{
-		USoundWave* SoundWave = WavePlayers[WaveIndex]->SoundWave;
+		USoundWave* SoundWave = WavePlayers[WaveIndex]->SoundWave.Get();
 		if (SoundWave)
 		{
 			ResourceSize += SoundWave->GetResourceSizeForFormat(Format);
