#ue4 - Perf: Optimize weak pointer usage in USkeletalMeshComponent::CalcBounds().

diff --git a/Engine/Source/Runtime/Engine/Private/Components/SkeletalMeshComponent.cpp b/Engine/Source/Runtime/Engine/Private/Components/SkeletalMeshComponent.cpp
--- a/Engine/Source/Runtime/Engine/Private/Components/SkeletalMeshComponent.cpp
+++ b/Engine/Source/Runtime/Engine/Private/Components/SkeletalMeshComponent.cpp
@@ -1242,20 +1242,21 @@ void USkeletalMeshComponent::UpdateBounds()
 
 FBoxSphereBounds USkeletalMeshComponent::CalcBounds(const FTransform& LocalToWorld) const
 {
-	FVector RootBoneOffset;
+	FVector RootBoneOffset = RootBoneTranslation;
 
 	// if to use MasterPoseComponent's fixed skel bounds, 
 	// send MasterPoseComponent's Root Bone Translation
-	if(MasterPoseComponent.IsValid() && MasterPoseComponent->SkeletalMesh &&
-		MasterPoseComponent->bComponentUseFixedSkelBounds &&
-		MasterPoseComponent->IsA((USkeletalMeshComponent::StaticClass())))
+	if (MasterPoseComponent.IsValid())
 	{
-		USkeletalMeshComponent* BaseComponent = CastChecked<USkeletalMeshComponent>(MasterPoseComponent.Get());
-		RootBoneOffset = BaseComponent->RootBoneTranslation; // Adjust bounds by root bone translation
-	}
-	else
-	{
-		RootBoneOffset = RootBoneTranslation;
+		const USkinnedMeshComponent* const MasterPoseComponentInst = MasterPoseComponent.Get();
+		check(MasterPoseComponentInst);
+		if (MasterPoseComponentInst->SkeletalMesh &&
+			MasterPoseComponentInst->bComponentUseFixedSkelBounds &&
+			MasterPoseComponentInst->IsA((USkeletalMeshComponent::StaticClass())))
+		{
+			const USkeletalMeshComponent* BaseComponent = CastChecked<USkeletalMeshComponent>(MasterPoseComponentInst);
+			RootBoneOffset = BaseComponent->RootBoneTranslation; // Adjust bounds by root bone translation
+		}
 	}
 
 	FBoxSphereBounds NewBounds = CalcMeshBound( RootBoneOffset, bHasValidBodies, LocalToWorld );
