Fix replicated variables in actor component blueprints by actually adding them to the component's lifetime replicated properties list. Fixes UE-14857.

diff --git a/Engine/Source/Runtime/Engine/Private/Components/ActorComponent.cpp b/Engine/Source/Runtime/Engine/Private/Components/ActorComponent.cpp
--- a/Engine/Source/Runtime/Engine/Private/Components/ActorComponent.cpp
+++ b/Engine/Source/Runtime/Engine/Private/Components/ActorComponent.cpp
@@ -1389,6 +1389,12 @@ bool UActorComponent::IsNetSimulating() const
 
 void UActorComponent::GetLifetimeReplicatedProps( TArray< FLifetimeProperty > & OutLifetimeProps ) const
 {
+	UBlueprintGeneratedClass* BPClass = Cast<UBlueprintGeneratedClass>(GetClass());
+	if (BPClass != NULL)
+	{
+		BPClass->GetLifetimeBlueprintReplicationList(OutLifetimeProps);
+	}
+
 	DOREPLIFETIME( UActorComponent, bIsActive );
 	DOREPLIFETIME( UActorComponent, bReplicates );
 }
