Fix issues with nested components after rename/undo/redo
* Cached owner not being updated
* In the incorrect Actor's OwnedComponents array

diff --git a/Engine/Source/Runtime/Engine/Private/Components/ActorComponent.cpp b/Engine/Source/Runtime/Engine/Private/Components/ActorComponent.cpp
--- a/Engine/Source/Runtime/Engine/Private/Components/ActorComponent.cpp
+++ b/Engine/Source/Runtime/Engine/Private/Components/ActorComponent.cpp
@@ -211,6 +211,25 @@ void UActorComponent::PostRename(UObject* OldOuter, const FName OldName)
 			{
 				Owner->AddOwnedComponent(this);
 			}
+
+			TArray<UObject*> Children;
+			GetObjectsWithOuter(this, Children);
+
+			for (UObject* Child : Children)
+			{
+				if (UActorComponent* ChildComponent = Cast<UActorComponent>(Child))
+				{
+					ChildComponent->Owner = Owner;
+					if (OldOwner)
+					{
+						OldOwner->RemoveOwnedComponent(ChildComponent);
+					}
+					if (Owner)
+					{
+						Owner->AddOwnedComponent(ChildComponent);
+					}
+				}
+			}
 		}
 	}
 
@@ -476,6 +495,25 @@ void UActorComponent::PostEditUndo()
 			Owner->AddOwnedComponent(this);
 		}
 
+		TArray<UObject*> Children;
+		GetObjectsWithOuter(this, Children);
+
+		for (UObject* Child : Children)
+		{
+			if (UActorComponent* ChildComponent = Cast<UActorComponent>(Child))
+			{
+				if (ChildComponent->Owner)
+				{
+					ChildComponent->Owner->RemoveOwnedComponent(ChildComponent);
+				}
+				ChildComponent->Owner = Owner;
+				if (Owner)
+				{
+					Owner->AddOwnedComponent(ChildComponent);
+				}
+			}
+		}
+
 		if (GetWorld())
 		{
 			GetWorld()->UpdateActorComponentEndOfFrameUpdateState(this);
