[UE-14662] Fix for crash caused by fixup code running when some components have been removed or renamed...

diff --git a/Engine/Source/Runtime/Engine/Private/Blueprint.cpp b/Engine/Source/Runtime/Engine/Private/Blueprint.cpp
--- a/Engine/Source/Runtime/Engine/Private/Blueprint.cpp
+++ b/Engine/Source/Runtime/Engine/Private/Blueprint.cpp
@@ -1137,35 +1137,39 @@ bool UBlueprint::ChangeOwnerOfTemplates()
 		for( auto CompIt = ComponentTemplates.CreateIterator(); CompIt; ++CompIt )
 		{
 			UActorComponent* Component = (*CompIt);
-			check(Component);
-			if(Component->GetOuter() == this)
+			if (Component)
 			{
-				const bool bRenamed = Component->Rename(*Component->GetName(), BPGClass, REN_ForceNoResetLoaders|REN_DoNotDirty);
-				ensure(bRenamed);
-				bIsStillStale |= !bRenamed;
-				bMigratedOwner = true;
-			}
-			if (auto TimelineComponent = Cast<UTimelineComponent>(Component))
-			{
-				TimelineComponent->GetAllCurves(Curves);
+				if (Component->GetOuter() == this)
+				{
+					const bool bRenamed = Component->Rename(*Component->GetName(), BPGClass, REN_ForceNoResetLoaders | REN_DoNotDirty);
+					ensure(bRenamed);
+					bIsStillStale |= !bRenamed;
+					bMigratedOwner = true;
+				}
+				if (auto TimelineComponent = Cast<UTimelineComponent>(Component))
+				{
+					TimelineComponent->GetAllCurves(Curves);
+				}
 			}
 		}
 
 		for( auto CompIt = Timelines.CreateIterator(); CompIt; ++CompIt )
 		{
 			UTimelineTemplate* Template = (*CompIt);
-			check(Template);
-			if(Template->GetOuter() == this)
+			if (Template)
 			{
-				const FString OldTemplateName = Template->GetName();
-				ensure(!OldTemplateName.EndsWith(TEXT("_Template")));
-				const bool bRenamed = Template->Rename(*UTimelineTemplate::TimelineVariableNameToTemplateName(Template->GetFName()), BPGClass, REN_ForceNoResetLoaders|REN_DoNotDirty);
-				ensure(bRenamed);
-				bIsStillStale |= !bRenamed;
-				ensure(OldTemplateName == UTimelineTemplate::TimelineTemplateNameToVariableName(Template->GetFName()));
-				bMigratedOwner = true;
+				if(Template->GetOuter() == this)
+				{
+					const FString OldTemplateName = Template->GetName();
+					ensure(!OldTemplateName.EndsWith(TEXT("_Template")));
+					const bool bRenamed = Template->Rename(*UTimelineTemplate::TimelineVariableNameToTemplateName(Template->GetFName()), BPGClass, REN_ForceNoResetLoaders|REN_DoNotDirty);
+					ensure(bRenamed);
+					bIsStillStale |= !bRenamed;
+					ensure(OldTemplateName == UTimelineTemplate::TimelineTemplateNameToVariableName(Template->GetFName()));
+					bMigratedOwner = true;
+				}
+				Template->GetAllCurves(Curves);
 			}
-			Template->GetAllCurves(Curves);
 		}
 		for (auto Curve : Curves)
 		{
