Fixed Navmesh generating disconnected tiles on landscapes in active-tiles mode #UE4

diff --git a/Engine/Source/Runtime/Engine/Private/AI/Navigation/RecastNavMesh.cpp b/Engine/Source/Runtime/Engine/Private/AI/Navigation/RecastNavMesh.cpp
--- a/Engine/Source/Runtime/Engine/Private/AI/Navigation/RecastNavMesh.cpp
+++ b/Engine/Source/Runtime/Engine/Private/AI/Navigation/RecastNavMesh.cpp
@@ -1961,6 +1961,9 @@ void ARecastNavMesh::PostEditChangeProperty(FPropertyChangedEvent& PropertyChang
 				// tile's can't be too big, otherwise we'll crash while tryng to allocate
 				// memory during navmesh generation
 				TileSizeUU = FMath::Clamp(TileSizeUU, CellSize, ArbitraryMaxVoxelTileSize * CellSize);
+
+				// update config
+				FillConfig(NavDataConfig);
 			}
 
 			if (HasAnyFlags(RF_ClassDefaultObject) == false)
diff --git a/Engine/Source/Runtime/Engine/Private/AI/Navigation/RecastNavMeshGenerator.cpp b/Engine/Source/Runtime/Engine/Private/AI/Navigation/RecastNavMeshGenerator.cpp
--- a/Engine/Source/Runtime/Engine/Private/AI/Navigation/RecastNavMeshGenerator.cpp
+++ b/Engine/Source/Runtime/Engine/Private/AI/Navigation/RecastNavMeshGenerator.cpp
@@ -530,10 +530,10 @@ void ExportHeightFieldSlice(const FNavHeightfieldSamples& PrefetchedHeightfieldS
 	// calculate the actual start and number of columns we want
 	const FBox LocalBox = SliceBox.TransformBy(LocalToWorld.Inverse());
 
-	const int32 StartingRow = FMath::Max(FMath::FloorToInt(LocalBox.Min.Y), 0);
-	const int32 StartingColumn = FMath::Max(FMath::FloorToInt(LocalBox.Min.X), 0);
-	const int32 RowLimit = FMath::Min(FMath::CeilToInt(LocalBox.Max.Y), NumRows);
-	const int32 ColumnLimit = FMath::Min(FMath::CeilToInt(LocalBox.Max.X), NumCols);
+	const int32 StartingRow = FMath::Max(FMath::FloorToInt(LocalBox.Min.Y) - 1, 0);
+	const int32 StartingColumn = FMath::Max(FMath::FloorToInt(LocalBox.Min.X) - 1, 0);
+	const int32 RowLimit = FMath::Min(FMath::CeilToInt(LocalBox.Max.Y) + 1, NumRows);
+	const int32 ColumnLimit = FMath::Min(FMath::CeilToInt(LocalBox.Max.X) + 1, NumCols);
 	
 	const int32 VertexCount = (RowLimit - StartingRow)*(ColumnLimit - StartingColumn);
 
@@ -1691,7 +1691,9 @@ void FRecastTileGenerator::DoAsyncGeometryGathering()
 			if(NavRelevant)
 			{
 				NavRelevant->PrepareGeometryExportSync();
-				NavRelevant->GatherGeometrySlice(GeomExport, TileBB.ExpandBy(NavDataConfig.AgentRadius * 2));
+				// adding a small bump to avoid special case of zero-expansion when tile bounds
+				// overlap landscape's tile bounds
+				NavRelevant->GatherGeometrySlice(GeomExport, TileBB.ExpandBy(NavDataConfig.AgentRadius * 2 + TileConfig.cs));
 
 				RecastGeometryExport::CovertCoordDataToRecast(GeomExport.VertexBuffer);
 				RecastGeometryExport::StoreCollisionCache(GeomExport);
@@ -1822,7 +1824,9 @@ void FRecastTileGenerator::GatherGeometry(const FRecastNavMeshGenerator& ParentG
 					if (NavRelevant)
 					{
 						NavRelevant->PrepareGeometryExportSync();
-						NavRelevant->GatherGeometrySlice(GeomExport, TileBB.ExpandBy(OwnerNavDataConfig.AgentRadius * 2));
+						// adding a small bump to avoid special case of zero-expansion when tile bounds
+						// overlap landscape's tile bounds
+						NavRelevant->GatherGeometrySlice(GeomExport, TileBB.ExpandBy(OwnerNavDataConfig.AgentRadius * 2 + TileConfig.cs));
 
 						RecastGeometryExport::CovertCoordDataToRecast(GeomExport.VertexBuffer);
 						RecastGeometryExport::StoreCollisionCache(GeomExport);
