Fixed issues with navmesh rendering in editor.  Show flags are presistent between SIE/Editor but it's not presistent for PIE now.
It fixes UE-14894 "Navmesh debug drawing now working as expected in PIE/SIE"

diff --git a/Engine/Source/Runtime/Engine/Private/Components/NavMeshRenderingComponent.cpp b/Engine/Source/Runtime/Engine/Private/Components/NavMeshRenderingComponent.cpp
--- a/Engine/Source/Runtime/Engine/Private/Components/NavMeshRenderingComponent.cpp
+++ b/Engine/Source/Runtime/Engine/Private/Components/NavMeshRenderingComponent.cpp
@@ -41,11 +41,11 @@ bool UNavMeshRenderingComponent::IsNavigationShowFlagSet(const UWorld* World)
 	FWorldContext* WorldContext = GEngine->GetWorldContextFromWorld(World);
 
 #if WITH_EDITOR
-	if (GEditor && WorldContext && WorldContext->WorldType == EWorldType::Editor)
+	if (GEditor && WorldContext && WorldContext->WorldType != EWorldType::Game)
 	{
 		for (FEditorViewportClient* CurrentViewport : GEditor->AllViewportClients)
 		{
-			if (CurrentViewport && CurrentViewport->EngineShowFlags.Navigation)
+			if (CurrentViewport && CurrentViewport->IsVisible() && CurrentViewport->EngineShowFlags.Navigation)
 			{
 				bShowNavigation = true;
 				break;
