[UE-13154] Remove potentially unreferenced objects after duplicating the editor world for PIE/SIE.

change summary:
- modified UGameInstance::InitializeForPIE() to do a GC pass after creating the PIE world via duplication. this ensures that any potentially unreferenced objects resulting from the duplication process will be cleaned up prior to starting PIE/SIE.

diff --git a/Engine/Source/Runtime/Engine/Private/GameInstance.cpp b/Engine/Source/Runtime/Engine/Private/GameInstance.cpp
--- a/Engine/Source/Runtime/Engine/Private/GameInstance.cpp
+++ b/Engine/Source/Runtime/Engine/Private/GameInstance.cpp
@@ -130,6 +130,7 @@ bool UGameInstance::InitializePIE(bool bAnyBlueprintErrors, int32 PIEInstance)
 	// We always need to create a new PIE world unless we're using the editor world for SIE
 	UWorld* NewWorld = nullptr;
 
+	bool bNeedsGarbageCollection = false;
 	const EPlayNetMode PlayNetMode = [&PlayInSettings]{ EPlayNetMode NetMode(PIE_Standalone); return (PlayInSettings->GetPlayNetMode(NetMode) ? NetMode : PIE_Standalone); }();
 	const bool CanRunUnderOneProcess = [&PlayInSettings]{ bool RunUnderOneProcess(false); return (PlayInSettings->GetRunUnderOneProcess(RunUnderOneProcess) && RunUnderOneProcess); }();
 	if (PlayNetMode == PIE_Client)
@@ -147,6 +148,9 @@ bool UGameInstance::InitializePIE(bool bAnyBlueprintErrors, int32 PIEInstance)
 	{
 		// Standard PIE path: just duplicate the EditorWorld
 		NewWorld = EditorEngine->CreatePIEWorldByDuplication(*WorldContext, EditorEngine->EditorWorld, PIEMapName);
+
+		// Duplication can result in unreferenced objects, so indicate that we should do a GC pass after initializing the world context
+		bNeedsGarbageCollection = true;
 	}
 
 	// failed to create the world!
@@ -164,6 +168,12 @@ bool UGameInstance::InitializePIE(bool bAnyBlueprintErrors, int32 PIEInstance)
 	NewWorld->ClearFlags(RF_Standalone);
 	NewWorld->bKismetScriptError = bAnyBlueprintErrors;
 
+	// Do a GC pass if necessary to remove any potentially unreferenced objects
+	if(bNeedsGarbageCollection)
+	{
+		CollectGarbage(GARBAGE_COLLECTION_KEEPFLAGS);
+	}
+
 	Init();
 
 	return true;
