PR #1114: [Opt] check if if ignored components has enabled collision (Contributed by [REDACTED_PERSON])

diff --git a/Engine/Source/Runtime/Engine/Private/Collision/Collision.cpp b/Engine/Source/Runtime/Engine/Private/Collision/Collision.cpp
--- a/Engine/Source/Runtime/Engine/Private/Collision/Collision.cpp
+++ b/Engine/Source/Runtime/Engine/Private/Collision/Collision.cpp
@@ -172,7 +172,7 @@ void FCollisionQueryParams::AddIgnoredActor(const AActor* InIgnoreActor)
 			for (const UActorComponent* ActorComponent : InIgnoreActor->GetComponents())
 			{
 				const UPrimitiveComponent* PrimComponent = Cast<const UPrimitiveComponent>(ActorComponent);
-				if (PrimComponent)
+				if (PrimComponent && PrimComponent->BodyInstance.GetCollisionEnabled() != ECollisionEnabled::NoCollision)
 				{
 					IgnoreComponents.Add(PrimComponent->GetUniqueID());
 				}
@@ -183,7 +183,7 @@ void FCollisionQueryParams::AddIgnoredActor(const AActor* InIgnoreActor)
 			for (const UActorComponent* ActorComponent : InIgnoreActor->GetComponents())
 			{
 				const UPrimitiveComponent* PrimComponent = Cast<const UPrimitiveComponent>(ActorComponent);
-				if (PrimComponent)
+				if (PrimComponent && PrimComponent->BodyInstance.GetCollisionEnabled() != ECollisionEnabled::NoCollision)
 				{
 					IgnoreComponents.AddUnique(PrimComponent->GetUniqueID());
 				}
