[UE-11819] Clear existing component delegate bindings on Blueprint class instance property changes.

change summary:
- added UBlueprintGeneratedClass::UnbindDynamicDelegatesForProperty() to unbind instances from delegates owned by an object referenced by a specific class property
- added a new UDynamicBlueprintBinding::UnbindDynamicDelegatesForProperty() API
- added a UComponentDelegateBinding::UnbindDynamicDelegatesForProperty() override to unbind an instance from any component delegate bindings that match a specific property
- consolidated common code between UComponentDelegateBinding::BindDynamicDelegates() and UnbindDynamicDelegatesForProperty() into a new FindComponentTargetDelegate() class utility method
- modified UComponentDelegateBinding::BindDynamicDelegates() to utiliize FindComponentTargetDelegate()
- modified AActor::PreEditChange() to call UBlueprintGeneratedClass::UnbindDynamicDelegatesForProperty() when an object property value is about to be changed, so that any existing delegate bindings will be cleared on the old value

diff --git a/Engine/Source/Runtime/Engine/Classes/Engine/BlueprintGeneratedClass.h b/Engine/Source/Runtime/Engine/Classes/Engine/BlueprintGeneratedClass.h
--- a/Engine/Source/Runtime/Engine/Classes/Engine/BlueprintGeneratedClass.h
+++ b/Engine/Source/Runtime/Engine/Classes/Engine/BlueprintGeneratedClass.h
@@ -505,6 +505,9 @@ public:
 	/** Bind functions on supplied actor to delegates */
 	void BindDynamicDelegates(UObject* InInstance) const;
 
+	/** Unbind functions on supplied actor from delegates tied to a specific property */
+	void UnbindDynamicDelegatesForProperty(UObject* InInstance, const UObjectProperty* InObjectProperty);
+	
 	// Finds the desired dynamic binding object for this blueprint generated class
 	UDynamicBlueprintBinding* GetDynamicBindingObject(UClass* Class) const;
 
diff --git a/Engine/Source/Runtime/Engine/Classes/Engine/ComponentDelegateBinding.h b/Engine/Source/Runtime/Engine/Classes/Engine/ComponentDelegateBinding.h
--- a/Engine/Source/Runtime/Engine/Classes/Engine/ComponentDelegateBinding.h
+++ b/Engine/Source/Runtime/Engine/Classes/Engine/ComponentDelegateBinding.h
@@ -43,5 +43,10 @@ class ENGINE_API UComponentDelegateBinding
 
 	// Begin DynamicBlueprintBinding interface
 	virtual void BindDynamicDelegates(UObject* InInstance) const override;
+	virtual void UnbindDynamicDelegatesForProperty(UObject* InInstance, const UObjectProperty* InObjectProperty) const override;
 	// End DynamicBlueprintBinding interface
+
+private:
+	// Utility method used to find the target delegate given an instance and a binding descriptor
+	static FMulticastScriptDelegate* FindComponentTargetDelegate(const UObject* InInstance, const FBlueprintComponentDelegateBinding& InBinding, const UObjectProperty* InObjectProperty = nullptr);
 };
diff --git a/Engine/Source/Runtime/Engine/Classes/Engine/DynamicBlueprintBinding.h b/Engine/Source/Runtime/Engine/Classes/Engine/DynamicBlueprintBinding.h
--- a/Engine/Source/Runtime/Engine/Classes/Engine/DynamicBlueprintBinding.h
+++ b/Engine/Source/Runtime/Engine/Classes/Engine/DynamicBlueprintBinding.h
@@ -12,4 +12,5 @@ class ENGINE_API UDynamicBlueprintBinding
 	GENERATED_UCLASS_BODY()
 
 	virtual void BindDynamicDelegates(UObject* InInstance) const { }
+	virtual void UnbindDynamicDelegatesForProperty(UObject* InInstance, const UObjectProperty* InObjectProperty) const { }
 };
diff --git a/Engine/Source/Runtime/Engine/Private/ActorEditor.cpp b/Engine/Source/Runtime/Engine/Private/ActorEditor.cpp
--- a/Engine/Source/Runtime/Engine/Private/ActorEditor.cpp
+++ b/Engine/Source/Runtime/Engine/Private/ActorEditor.cpp
@@ -18,6 +18,13 @@ void AActor::PreEditChange(UProperty* PropertyThatWillChange)
 {
 	Super::PreEditChange(PropertyThatWillChange);
 
+	UObjectProperty* ObjProp = Cast<UObjectProperty>(PropertyThatWillChange);
+	UBlueprintGeneratedClass* BPGC = Cast<UBlueprintGeneratedClass>(GetClass());
+	if ( BPGC != nullptr && ObjProp != nullptr )
+	{
+		BPGC->UnbindDynamicDelegatesForProperty(this, ObjProp);
+	}
+
 	if ( ReregisterComponentsWhenModified() )
 	{
 		UnregisterAllComponents();
diff --git a/Engine/Source/Runtime/Engine/Private/BlueprintGeneratedClass.cpp b/Engine/Source/Runtime/Engine/Private/BlueprintGeneratedClass.cpp
--- a/Engine/Source/Runtime/Engine/Private/BlueprintGeneratedClass.cpp
+++ b/Engine/Source/Runtime/Engine/Private/BlueprintGeneratedClass.cpp
@@ -407,6 +407,17 @@ void UBlueprintGeneratedClass::BindDynamicDelegates(UObject* InInstance) const
 	}
 }
 
+void UBlueprintGeneratedClass::UnbindDynamicDelegatesForProperty(UObject* InInstance, const UObjectProperty* InObjectProperty)
+{
+	for (int32 Index = 0; Index < DynamicBindingObjects.Num(); ++Index)
+	{
+		if ( ensure(DynamicBindingObjects[Index] != NULL) )
+		{
+			DynamicBindingObjects[Index]->UnbindDynamicDelegatesForProperty(InInstance, InObjectProperty);
+		}
+	}
+}
+
 bool UBlueprintGeneratedClass::GetGeneratedClassesHierarchy(const UClass* InClass, TArray<const UBlueprintGeneratedClass*>& OutBPGClasses)
 {
 	OutBPGClasses.Empty();
diff --git a/Engine/Source/Runtime/Engine/Private/ComponentDelegateBinding.cpp b/Engine/Source/Runtime/Engine/Private/ComponentDelegateBinding.cpp
--- a/Engine/Source/Runtime/Engine/Private/ComponentDelegateBinding.cpp
+++ b/Engine/Source/Runtime/Engine/Private/ComponentDelegateBinding.cpp
@@ -9,34 +9,66 @@ UComponentDelegateBinding::UComponentDelegateBinding(const FObjectInitializer& O
 {
 }
 
+FMulticastScriptDelegate* UComponentDelegateBinding::FindComponentTargetDelegate(const UObject* InInstance, const FBlueprintComponentDelegateBinding& InBinding, const UObjectProperty* InObjectProperty)
+{
+	// Get the property that points to the component
+	const UObjectProperty* ObjProp = (InObjectProperty != nullptr && InObjectProperty->GetFName() == InBinding.ComponentPropertyName) ? InObjectProperty : FindField<UObjectProperty>(InInstance->GetClass(), InBinding.ComponentPropertyName);
+	if(ObjProp != nullptr)
+	{
+		// ..see if there is actually a component assigned
+		UObject* Component = Cast<UObject>(ObjProp->GetObjectPropertyValue_InContainer(InInstance));
+		if(Component != nullptr)
+		{
+			// If there is, find and return the delegate property on it
+			UMulticastDelegateProperty* DelegateProp = FindField<UMulticastDelegateProperty>(Component->GetClass(), InBinding.DelegatePropertyName);
+			if(DelegateProp != nullptr)
+			{
+				return DelegateProp->GetPropertyValuePtr_InContainer(Component);
+			}
+		}
+	}
+
+	return nullptr;
+}
+
 void UComponentDelegateBinding::BindDynamicDelegates(UObject* InInstance) const
 {
 	for(int32 BindIdx=0; BindIdx<ComponentDelegateBindings.Num(); BindIdx++)
 	{
 		const FBlueprintComponentDelegateBinding& Binding = ComponentDelegateBindings[BindIdx];
+
+		// Get the delegate property on the component we want to bind to
+		FMulticastScriptDelegate* TargetDelegate = FindComponentTargetDelegate(InInstance, Binding);
 		// Get the function we want to bind
 		UFunction* FunctionToBind = FindField<UFunction>(InInstance->GetClass(), Binding.FunctionNameToBind);
-		// Get the property that points to the component we want to assign to
-		UObjectProperty* ObjProp = FindField<UObjectProperty>(InInstance->GetClass(), Binding.ComponentPropertyName);
 		// If we have both of those..
-		if(ObjProp != NULL && FunctionToBind != NULL)
+		if(TargetDelegate != nullptr && FunctionToBind != nullptr)
 		{
-			// ..see if there is actually a component assigned
-			UObject* Component = Cast<UObject>(ObjProp->GetObjectPropertyValue_InContainer(InInstance));
-			if(Component != NULL)
+			// Bind function on the instance to this delegate
+			FScriptDelegate Delegate;
+			Delegate.BindUFunction(InInstance, Binding.FunctionNameToBind);
+			TargetDelegate->AddUnique(Delegate);
+		}
+	}
+}
+
+void UComponentDelegateBinding::UnbindDynamicDelegatesForProperty(UObject* InInstance, const UObjectProperty* InObjectProperty) const
+{
+	for(int32 BindIdx=0; BindIdx<ComponentDelegateBindings.Num(); BindIdx++)
+	{
+		const FBlueprintComponentDelegateBinding& Binding = ComponentDelegateBindings[BindIdx];
+		const UObjectProperty* ObjProp = FindField<UObjectProperty>(InInstance->GetClass(), Binding.ComponentPropertyName);
+		if(ObjProp == InObjectProperty)
+		{
+			// Get the delegate property on the component we want to unbind from
+			FMulticastScriptDelegate* TargetDelegate = FindComponentTargetDelegate(InInstance, Binding, ObjProp);
+			if(TargetDelegate != nullptr)
 			{
-				// If there is, find the delegate property on it
-				UMulticastDelegateProperty* DelegateProp = FindField<UMulticastDelegateProperty>(Component->GetClass(), Binding.DelegatePropertyName);
-				if(DelegateProp)
-				{
-					// Found that, finally bind function on the actor to this delegate
-					FMulticastScriptDelegate* TargetDelegate = DelegateProp->GetPropertyValuePtr_InContainer(Component);
-
-					FScriptDelegate Delegate;
-					Delegate.BindUFunction(InInstance, Binding.FunctionNameToBind);
-					TargetDelegate->AddUnique(Delegate);
-				}
+				// Unbind function on the instance from this delegate
+				TargetDelegate->Remove(InInstance, Binding.FunctionNameToBind);
 			}
+
+			break;
 		}
 	}
 }
\ No newline at end of file
