[UE-10170] Overlap events no longer fire outside of simulation mode in the Blueprint editor.

change summary:
- modified AActor::PostActorConstruction() to remove a redundant UpdateOverlaps() call
- modified AActor::PostActorConstruction() to restrict the UpdateOverlaps() call to only occur if GetWorld()->HasBegunPlay() is also true

diff --git a/Engine/Source/Runtime/Engine/Private/Actor.cpp b/Engine/Source/Runtime/Engine/Private/Actor.cpp
--- a/Engine/Source/Runtime/Engine/Private/Actor.cpp
+++ b/Engine/Source/Runtime/Engine/Private/Actor.cpp
@@ -2463,9 +2463,6 @@ void AActor::PostActorConstruction()
 		PreInitializeComponents();
 	}
 
-	// components are all there, init overlapping state
-	UpdateOverlaps();
-
 	// If this is dynamically spawned replicted actor, defer calls to BeginPlay and UpdateOverlaps until replicated properties are deserialized
 	bool deferBeginPlayAndUpdateOverlaps = (bExchangedRoles && RemoteRole == ROLE_Authority);
 
@@ -2496,7 +2493,7 @@ void AActor::PostActorConstruction()
 	}
 
 	// Components are all there and we've begun play, init overlapping state
-	if (!deferBeginPlayAndUpdateOverlaps)
+	if (GetWorld()->HasBegunPlay() && !deferBeginPlayAndUpdateOverlaps)
 	{
 		UpdateOverlaps();
 	}
