Fixed "show navigation" command for PIE in editor.

diff --git a/Engine/Source/Runtime/Engine/Private/Components/NavMeshRenderingComponent.cpp b/Engine/Source/Runtime/Engine/Private/Components/NavMeshRenderingComponent.cpp
--- a/Engine/Source/Runtime/Engine/Private/Components/NavMeshRenderingComponent.cpp
+++ b/Engine/Source/Runtime/Engine/Private/Components/NavMeshRenderingComponent.cpp
@@ -43,12 +43,17 @@ bool UNavMeshRenderingComponent::IsNavigationShowFlagSet(const UWorld* World)
 #if WITH_EDITOR
 	if (GEditor && WorldContext && WorldContext->WorldType != EWorldType::Game)
 	{
-		for (FEditorViewportClient* CurrentViewport : GEditor->AllViewportClients)
+		bShowNavigation = WorldContext->GameViewport != nullptr && WorldContext->GameViewport->EngineShowFlags.Navigation;
+		if (bShowNavigation == false)
 		{
-			if (CurrentViewport && CurrentViewport->IsVisible() && CurrentViewport->EngineShowFlags.Navigation)
+			// we have to check all vieports because we can't to distinguish between SIE and PIE at this point.
+			for (FEditorViewportClient* CurrentViewport : GEditor->AllViewportClients)
 			{
-				bShowNavigation = true;
-				break;
+				if (CurrentViewport && CurrentViewport->IsVisible() && CurrentViewport->EngineShowFlags.Navigation)
+				{
+					bShowNavigation = true;
+					break;
+				}
 			}
 		}
 	}
