#HLOD
you can't move LODActor using gizmo
- you still can move using transform widget, but that is a bigger task, so fixing this first.

https://jira.ol.epicgames.net/browse/UE-14608

diff --git a/Engine/Source/Runtime/Engine/Classes/Engine/LODActor.h b/Engine/Source/Runtime/Engine/Classes/Engine/LODActor.h
--- a/Engine/Source/Runtime/Engine/Classes/Engine/LODActor.h
+++ b/Engine/Source/Runtime/Engine/Classes/Engine/LODActor.h
@@ -49,6 +49,10 @@ protected:
 	virtual FBox GetComponentsBoundingBox(bool bNonColliding = false) const override;	
 #if WITH_EDITOR
 	virtual void PostEditChangeProperty(FPropertyChangedEvent& PropertyChangedEvent) override;
+	virtual void EditorApplyTranslation(const FVector& DeltaTranslation, bool bAltDown, bool bShiftDown, bool bCtrlDown) override;
+	virtual void EditorApplyRotation(const FRotator& DeltaRotation, bool bAltDown, bool bShiftDown, bool bCtrlDown) override;
+	virtual void EditorApplyScale(const FVector& DeltaScale, const FVector* PivotLocation, bool bAltDown, bool bShiftDown, bool bCtrlDown) override;
+	virtual void EditorApplyMirror(const FVector& MirrorScale, const FVector& PivotLocation) override;
 #endif // WITH_EDITOR	
 	virtual void PostRegisterAllComponents() override;
 	// End UObject interface.
diff --git a/Engine/Source/Runtime/Engine/Private/LODActor.cpp b/Engine/Source/Runtime/Engine/Private/LODActor.cpp
--- a/Engine/Source/Runtime/Engine/Private/LODActor.cpp
+++ b/Engine/Source/Runtime/Engine/Private/LODActor.cpp
@@ -134,6 +134,22 @@ void ALODActor::CheckForErrors()
 	}
 }
 
+void ALODActor::EditorApplyTranslation(const FVector& DeltaTranslation, bool bAltDown, bool bShiftDown, bool bCtrlDown)
+{
+}
+
+void ALODActor::EditorApplyRotation(const FRotator& DeltaRotation, bool bAltDown, bool bShiftDown, bool bCtrlDown)
+{
+}
+
+void ALODActor::EditorApplyScale(const FVector& DeltaScale, const FVector* PivotLocation, bool bAltDown, bool bShiftDown, bool bCtrlDown)
+{
+}
+
+void ALODActor::EditorApplyMirror(const FVector& MirrorScale, const FVector& PivotLocation)
+{
+}
+
 #endif // WITH_EDITOR
 
 FBox ALODActor::GetComponentsBoundingBox(bool bNonColliding) const 
