Attached Actors should maintain their world position when being detached from their parent as part of being destroyed.

diff --git a/Engine/Source/Runtime/Engine/Private/LevelActor.cpp b/Engine/Source/Runtime/Engine/Private/LevelActor.cpp
--- a/Engine/Source/Runtime/Engine/Private/LevelActor.cpp
+++ b/Engine/Source/Runtime/Engine/Private/LevelActor.cpp
@@ -502,7 +502,7 @@ bool UWorld::DestroyActor( AActor* ThisActor, bool bNetForce, bool bShouldModify
 			OldParentActor->Modify();
 		}
 
-		ThisActor->DetachRootComponentFromParent(false);
+		ThisActor->DetachRootComponentFromParent();
 
 #if WITH_EDITOR
 		if( GIsEditor )
