Fixed FExportMaterialProxy and FLightmassMaterialProxy creating their shader map ID's shader type dependencies off of the base material, instead of themselves, which would cause some shaders to not be cached due to mismatching ShouldCache results.

diff --git a/Engine/Source/Developer/MaterialUtilities/Private/MaterialUtilities.cpp b/Engine/Source/Developer/MaterialUtilities/Private/MaterialUtilities.cpp
--- a/Engine/Source/Developer/MaterialUtilities/Private/MaterialUtilities.cpp
+++ b/Engine/Source/Developer/MaterialUtilities/Private/MaterialUtilities.cpp
@@ -50,6 +50,17 @@ public:
 
 		FMaterialShaderMapId ResourceId;
 		Resource->GetShaderMapId(GMaxRHIShaderPlatform, ResourceId);
+
+		{
+			TArray<FShaderType*> ShaderTypes;
+			TArray<FVertexFactoryType*> VFTypes;
+			GetDependentShaderAndVFTypes(GMaxRHIShaderPlatform, ShaderTypes, VFTypes);
+
+			// Overwrite the shader map Id's dependencies with ones that came from the FMaterial actually being compiled (this)
+			// This is necessary as we change FMaterial attributes like GetShadingModel(), which factor into the ShouldCache functions that determine dependent shader types
+			ResourceId.SetShaderDependencies(ShaderTypes, VFTypes);
+		}
+
 		// Override with a special usage so we won't re-use the shader map used by the material for rendering
 		switch (InPropertyToCompile)
 		{
diff --git a/Engine/Source/Editor/UnrealEd/Private/Lightmass/LightmassRender.cpp b/Engine/Source/Editor/UnrealEd/Private/Lightmass/LightmassRender.cpp
--- a/Engine/Source/Editor/UnrealEd/Private/Lightmass/LightmassRender.cpp
+++ b/Engine/Source/Editor/UnrealEd/Private/Lightmass/LightmassRender.cpp
@@ -197,6 +197,17 @@ public:
 			{
 				FMaterialShaderMapId ResourceId;
 				Resource->GetShaderMapId(GMaxRHIShaderPlatform, ResourceId);
+
+				{
+					TArray<FShaderType*> ShaderTypes;
+					TArray<FVertexFactoryType*> VFTypes;
+					GetDependentShaderAndVFTypes(GMaxRHIShaderPlatform, ShaderTypes, VFTypes);
+
+					// Overwrite the shader map Id's dependencies with ones that came from the FMaterial actually being compiled (this)
+					// This is necessary as we change FMaterial attributes like GetShadingModel(), which factor into the ShouldCache functions that determine dependent shader types
+					ResourceId.SetShaderDependencies(ShaderTypes, VFTypes);
+				}
+
 				// Override with a special usage so we won't re-use the shader map used by the material for rendering
 				ResourceId.Usage = GetShaderMapUsage();
 				CacheShaders(ResourceId, GMaxRHIShaderPlatform, true);
diff --git a/Engine/Source/Runtime/Engine/Public/MaterialShared.h b/Engine/Source/Runtime/Engine/Public/MaterialShared.h
--- a/Engine/Source/Runtime/Engine/Public/MaterialShared.h
+++ b/Engine/Source/Runtime/Engine/Public/MaterialShared.h
@@ -401,7 +401,7 @@ public:
 	~FMaterialShaderMapId()
 	{ }
 
-	void SetShaderDependencies(const TArray<FShaderType*>& ShaderTypes, const TArray<FVertexFactoryType*>& VFTypes);
+	ENGINE_API void SetShaderDependencies(const TArray<FShaderType*>& ShaderTypes, const TArray<FVertexFactoryType*>& VFTypes);
 
 	void Serialize(FArchive& Ar);
 
@@ -1144,6 +1144,9 @@ protected:
 
 	/** Gets the Guid that represents this material. */
 	virtual FGuid GetMaterialId() const = 0;
+	
+	/** Produces arrays of any shader and vertex factory type that this material is dependent on. */
+	ENGINE_API void GetDependentShaderAndVFTypes(EShaderPlatform Platform, TArray<FShaderType*>& OutShaderTypes, TArray<FVertexFactoryType*>& OutVFTypes) const;
 
 private:
 
@@ -1223,9 +1226,6 @@ private:
 
 	void GetReferencedTexturesHash(FSHAHash& OutHash) const;
 
-	/** Produces arrays of any shader and vertex factory type that this material is dependent on. */
-	ENGINE_API void GetDependentShaderAndVFTypes(EShaderPlatform Platform, TArray<FShaderType*>& OutShaderTypes, TArray<FVertexFactoryType*>& OutVFTypes) const;
-
 	EMaterialQualityLevel::Type GetQualityLevelForShaderMapId() const 
 	{
 		return bQualityLevelHasDifferentNodes ? QualityLevel : EMaterialQualityLevel::Num;
