Asserts to help track down "Assertion failed: NumJobs == QueuedJobs.Num()"

diff --git a/Engine/Source/Runtime/Engine/Private/ShaderCompiler/ShaderCompiler.cpp b/Engine/Source/Runtime/Engine/Private/ShaderCompiler/ShaderCompiler.cpp
--- a/Engine/Source/Runtime/Engine/Private/ShaderCompiler/ShaderCompiler.cpp
+++ b/Engine/Source/Runtime/Engine/Private/ShaderCompiler/ShaderCompiler.cpp
@@ -160,7 +160,7 @@ static void DoReadTaskResults(const TArray<FShaderCompileJob*>& QueuedJobs, FArc
 
 	int32 NumJobs;
 	OutputFile << NumJobs;
-	check(NumJobs == QueuedJobs.Num());
+	checkf(NumJobs == QueuedJobs.Num(), TEXT("Worker returned %u jobs, %u expected"), NumJobs, QueuedJobs.Num());
 
 	for (int32 JobIndex = 0; JobIndex < NumJobs; JobIndex++)
 	{
@@ -954,6 +954,7 @@ void FShaderCompileThreadRunnable::ReadAvailableResults()
 					if (OutputFilePtr)
 					{
 						FArchive& OutputFile = *OutputFilePtr;
+						check(!CurrentWorkerInfo.bComplete);
 						DoReadTaskResults(CurrentWorkerInfo.QueuedJobs, OutputFile);
 
 						// Close the output file.
