Added boolean to make firing of Update pin on setting playback position for a timeline optional

diff --git a/Engine/Source/Runtime/Engine/Classes/Components/TimelineComponent.h b/Engine/Source/Runtime/Engine/Classes/Components/TimelineComponent.h
--- a/Engine/Source/Runtime/Engine/Classes/Components/TimelineComponent.h
+++ b/Engine/Source/Runtime/Engine/Classes/Components/TimelineComponent.h
@@ -277,7 +277,7 @@ public:
 	ENGINE_API bool IsReversing() const;
 
 	/** Jump to a position in the timeline. If bFireEvents is true, event functions will fire, otherwise will not. */
-	ENGINE_API void SetPlaybackPosition(float NewPosition, bool bFireEvents);
+	ENGINE_API void SetPlaybackPosition(float NewPosition, bool bFireEvents, bool bFireUpdate = true);
 
 	/** Get the current playback position of the Timeline */
 	ENGINE_API float GetPlaybackPosition() const;
@@ -398,9 +398,12 @@ public:
 	UFUNCTION(BlueprintCallable, Category="Components|Timeline")
 	ENGINE_API bool IsReversing() const;
 
-	/** Jump to a position in the timeline. If bFireEvents is true, event functions will fire, otherwise will not. */
-	UFUNCTION(BlueprintCallable, Category="Components|Timeline")
-	ENGINE_API void SetPlaybackPosition(float NewPosition, bool bFireEvents);
+	/** Jump to a position in the timeline. 
+	  * @param bFireEvents If true, event functions that are between current position and new playback position will fire. 
+	  * @param bFireUpdate If true, the update output exec will fire after setting the new playback position.
+	 */
+	UFUNCTION(BlueprintCallable, Category="Components|Timeline", meta=(AdvancedDisplay="bFireUpdate"))
+	ENGINE_API void SetPlaybackPosition(float NewPosition, bool bFireEvents, bool bFireUpdate = true);
 
 	/** Get the current playback position of the Timeline */
 	UFUNCTION(BlueprintCallable, Category="Components|Timeline")
diff --git a/Engine/Source/Runtime/Engine/Private/Timeline.cpp b/Engine/Source/Runtime/Engine/Private/Timeline.cpp
--- a/Engine/Source/Runtime/Engine/Private/Timeline.cpp
+++ b/Engine/Source/Runtime/Engine/Private/Timeline.cpp
@@ -129,7 +129,7 @@ void FTimeline::AddInterpLinearColor(UCurveLinearColor* LinearColorCurve, FOnTim
 	InterpLinearColors.Add(NewEntry);
 }
 
-void FTimeline::SetPlaybackPosition(float NewPosition, bool bFireEvents)
+void FTimeline::SetPlaybackPosition(float NewPosition, bool bFireEvents, bool bFireUpdate)
 {
 	float OldPosition = Position;
 	Position = NewPosition;
@@ -319,7 +319,10 @@ void FTimeline::SetPlaybackPosition(float NewPosition, bool bFireEvents)
 	}
 
 	// Execute the delegate to say that all properties are updated
-	TimelinePostUpdateFunc.ExecuteIfBound();
+	if (bFireUpdate)
+	{
+		TimelinePostUpdateFunc.ExecuteIfBound();
+	}
 }
 
 
@@ -639,9 +642,9 @@ void UTimelineComponent::AddInterpLinearColor(UCurveLinearColor* LinearColorCurv
 	TheTimeline.AddInterpLinearColor(LinearColorCurve, InterpFunc, PropertyName);
 }
 
-void UTimelineComponent::SetPlaybackPosition(float NewPosition, bool bFireEvents)
+void UTimelineComponent::SetPlaybackPosition(float NewPosition, bool bFireEvents, bool bFireUpdate)
 {
-	TheTimeline.SetPlaybackPosition(NewPosition, bFireEvents);
+	TheTimeline.SetPlaybackPosition(NewPosition, bFireEvents, bFireUpdate);
 }
 
 float UTimelineComponent::GetPlaybackPosition() const
