Fix crash when there is a null entry in the AttachChildren array

diff --git a/Engine/Source/Runtime/Engine/Private/Components/SceneComponent.cpp b/Engine/Source/Runtime/Engine/Private/Components/SceneComponent.cpp
--- a/Engine/Source/Runtime/Engine/Private/Components/SceneComponent.cpp
+++ b/Engine/Source/Runtime/Engine/Private/Components/SceneComponent.cpp
@@ -650,14 +650,16 @@ void USceneComponent::OnComponentDestroyed()
 	// Detach children before destroying
 	for (int32 Index = AttachChildren.Num()-1; Index >= 0; --Index)
 	{
-		USceneComponent* Child = AttachChildren[Index];
-		if (AttachParent)
+		if (USceneComponent* Child = AttachChildren[Index])
 		{
-			Child->AttachTo(AttachParent);
-		}
-		else
-		{
-			Child->DetachFromParent();
+			if (AttachParent)
+			{
+				Child->AttachTo(AttachParent);
+			}
+			else
+			{
+				Child->DetachFromParent();
+			}
 		}
 	}
 
