#ue4 - Actor::GetActorEnableCollision() is now const.

diff --git a/Engine/Source/Runtime/Engine/Classes/GameFramework/Actor.h b/Engine/Source/Runtime/Engine/Classes/GameFramework/Actor.h
--- a/Engine/Source/Runtime/Engine/Classes/GameFramework/Actor.h
+++ b/Engine/Source/Runtime/Engine/Classes/GameFramework/Actor.h
@@ -851,7 +851,7 @@ public:
 
 	/** Get current state of collision for the whole actor */
 	UFUNCTION(BlueprintPure, Category="Collision")
-	bool GetActorEnableCollision();
+	bool GetActorEnableCollision() const;
 
 	/** Destroy the actor */
 	UFUNCTION(BlueprintCallable, Category="Utilities", meta=(Keywords = "Delete", DisplayName = "DestroyActor"))
diff --git a/Engine/Source/Runtime/Engine/Private/Actor.cpp b/Engine/Source/Runtime/Engine/Private/Actor.cpp
--- a/Engine/Source/Runtime/Engine/Private/Actor.cpp
+++ b/Engine/Source/Runtime/Engine/Private/Actor.cpp
@@ -2991,7 +2991,7 @@ void AActor::SetActorEnableCollision(bool bNewActorEnableCollision)
 }
 
 
-bool AActor::GetActorEnableCollision()
+bool AActor::GetActorEnableCollision() const
 {
 	return bActorEnableCollision;
 }
