#ue4 - Check that Actor or Component collision is enabled when adding to ignored components list in FCollisionQueryParams.

diff --git a/Engine/Source/Runtime/Engine/Private/Collision/Collision.cpp b/Engine/Source/Runtime/Engine/Private/Collision/Collision.cpp
--- a/Engine/Source/Runtime/Engine/Private/Collision/Collision.cpp
+++ b/Engine/Source/Runtime/Engine/Private/Collision/Collision.cpp
@@ -164,7 +164,7 @@ FCollisionQueryParams::FCollisionQueryParams(FName InTraceTag, bool bInTraceComp
 
 void FCollisionQueryParams::AddIgnoredActor(const AActor* InIgnoreActor)
 {
-	if (InIgnoreActor)
+	if (InIgnoreActor && InIgnoreActor->GetActorEnableCollision())
 	{
 		if (IgnoreComponents.Num() == 0)
 		{
@@ -217,7 +217,7 @@ void FCollisionQueryParams::AddIgnoredActors(const TArray<TWeakObjectPtr<AActor>
 
 void FCollisionQueryParams::AddIgnoredComponent(const UPrimitiveComponent* InIgnoreComponent)
 {
-	if (InIgnoreComponent)
+	if (InIgnoreComponent && InIgnoreComponent->GetCollisionEnabled() != ECollisionEnabled::NoCollision)
 	{
 		IgnoreComponents.AddUnique(InIgnoreComponent->GetUniqueID());
 	}
