fixed UE-13151 Text Render component not showing in packaged builds

diff --git a/Engine/Source/Runtime/Engine/Private/Components/TextRenderComponent.cpp b/Engine/Source/Runtime/Engine/Private/Components/TextRenderComponent.cpp
--- a/Engine/Source/Runtime/Engine/Private/Components/TextRenderComponent.cpp
+++ b/Engine/Source/Runtime/Engine/Private/Components/TextRenderComponent.cpp
@@ -207,7 +207,7 @@ FVector2D ComputeTextSize(FTextIterator It, class UFont* Font,
 
 		UTexture2D* Tex = Font->Textures[Char.TextureIndex];
 
-		if(Tex && Tex->Resource)
+		if(Tex)
 		{
 			FIntPoint ImportedTextureSize = Tex->GetImportedSize();
 			FVector2D InvTextureSize(1.0f / (float)ImportedTextureSize.X, 1.0f / (float)ImportedTextureSize.Y);
@@ -621,7 +621,7 @@ bool  FTextRenderSceneProxy::BuildStringMesh( TArray<FDynamicMeshVertex>& OutVer
 
 			UTexture2D* Tex = Font->Textures[Char.TextureIndex];
 
-			if(Tex && Tex->Resource)
+			if(Tex)
 			{
 				FIntPoint ImportedTextureSize = Tex->GetImportedSize();
 				FVector2D InvTextureSize(1.0f / (float)ImportedTextureSize.X, 1.0f / (float)ImportedTextureSize.Y);
