#ue4 - Add stats for PlayerController in "stat Game" and "stat PlayerController"

diff --git a/Engine/Source/Runtime/Engine/Private/LevelTick.cpp b/Engine/Source/Runtime/Engine/Private/LevelTick.cpp
--- a/Engine/Source/Runtime/Engine/Private/LevelTick.cpp
+++ b/Engine/Source/Runtime/Engine/Private/LevelTick.cpp
@@ -60,6 +60,7 @@ DEFINE_STAT(STAT_TickTime);
 DEFINE_STAT(STAT_WorldTickTime);
 DEFINE_STAT(STAT_UpdateCameraTime);
 DEFINE_STAT(STAT_CharacterMovement);
+DEFINE_STAT(STAT_PlayerControllerTick);
 
 DEFINE_STAT(STAT_VolumeStreamingTickTime);
 DEFINE_STAT(STAT_VolumeStreamingChecks);
diff --git a/Engine/Source/Runtime/Engine/Private/PlayerController.cpp b/Engine/Source/Runtime/Engine/Private/PlayerController.cpp
--- a/Engine/Source/Runtime/Engine/Private/PlayerController.cpp
+++ b/Engine/Source/Runtime/Engine/Private/PlayerController.cpp
@@ -43,6 +43,13 @@ DEFINE_LOG_CATEGORY(LogPlayerController);
 
 #define LOCTEXT_NAMESPACE "PlayerController"
 
+DECLARE_STATS_GROUP(TEXT("PlayerController"), STATGROUP_PlayerController, STATCAT_Advanced);
+DECLARE_CYCLE_STAT(TEXT("PC Tick Actor"), STAT_PC_TickActor, STATGROUP_PlayerController);
+DECLARE_CYCLE_STAT(TEXT("  PC Tick Input"), STAT_PC_TickInput, STATGROUP_PlayerController);
+DECLARE_CYCLE_STAT(TEXT("    PC Build Input Stack"), STAT_PC_BuildInputStack, STATGROUP_PlayerController);
+DECLARE_CYCLE_STAT(TEXT("    PC Process Input Stack"), STAT_PC_ProcessInputStack, STATGROUP_PlayerController);
+
+
 const float RetryClientRestartThrottleTime = 0.5f;
 const float RetryServerAcknowledgeThrottleTime = 0.25f;
 const float RetryServerCheckSpectatorThrottleTime = 0.25f;
@@ -2357,13 +2364,19 @@ void APlayerController::BuildInputStack(TArray<UInputComponent*>& InputStack)
 
 void APlayerController::ProcessPlayerInput(const float DeltaTime, const bool bGamePaused)
 {
-	// process all input components in the stack, top down
 	TArray<UInputComponent*> InputStack;
 
-	BuildInputStack(InputStack);
+	// process all input components in the stack, top down
+	{
+		SCOPE_CYCLE_COUNTER(STAT_PC_BuildInputStack);
+		BuildInputStack(InputStack);
+	}
 
 	// process the desired components
-	PlayerInput->ProcessInputStack(InputStack, DeltaTime, bGamePaused);
+	{
+		SCOPE_CYCLE_COUNTER(STAT_PC_ProcessInputStack);
+		PlayerInput->ProcessInputStack(InputStack, DeltaTime, bGamePaused);
+	}
 }
 
 void APlayerController::PreProcessInput(const float DeltaTime, const bool bGamePaused)
@@ -3722,6 +3735,8 @@ void APlayerController::SetPlayer( UPlayer* InPlayer )
 
 void APlayerController::TickPlayerInput(const float DeltaSeconds, const bool bGamePaused)
 {
+	SCOPE_CYCLE_COUNTER(STAT_PC_TickInput);
+
 	check(PlayerInput);
 	PlayerInput->Tick(DeltaSeconds);
 
@@ -3771,6 +3786,9 @@ void APlayerController::TickPlayerInput(const float DeltaSeconds, const bool bGa
 
 void APlayerController::TickActor( float DeltaSeconds, ELevelTick TickType, FActorTickFunction& ThisTickFunction )
 {
+	SCOPE_CYCLE_COUNTER(STAT_PlayerControllerTick);
+	SCOPE_CYCLE_COUNTER(STAT_PC_TickActor);
+
 	if (TickType == LEVELTICK_PauseTick && !ShouldPerformFullTickWhenPaused())
 	{
 		if (PlayerInput)
diff --git a/Engine/Source/Runtime/Engine/Public/EngineStats.h b/Engine/Source/Runtime/Engine/Public/EngineStats.h
--- a/Engine/Source/Runtime/Engine/Public/EngineStats.h
+++ b/Engine/Source/Runtime/Engine/Public/EngineStats.h
@@ -46,6 +46,7 @@ DECLARE_CYCLE_STAT_EXTERN(TEXT("Tick Time"),STAT_TickTime,STATGROUP_Game, );
 DECLARE_CYCLE_STAT_EXTERN(TEXT("World Tick Time"),STAT_WorldTickTime,STATGROUP_Game, );
 DECLARE_CYCLE_STAT_EXTERN(TEXT("Update Camera Time"),STAT_UpdateCameraTime,STATGROUP_Game, );
 DECLARE_CYCLE_STAT_EXTERN(TEXT("Char Movement Total"), STAT_CharacterMovement, STATGROUP_Game, );
+DECLARE_CYCLE_STAT_EXTERN(TEXT("PlayerController Tick"), STAT_PlayerControllerTick, STATGROUP_Game, );
 DECLARE_CYCLE_STAT_EXTERN(TEXT("Post Tick Component Update"),STAT_PostTickComponentUpdate,STATGROUP_Game, );
 DECLARE_CYCLE_STAT_EXTERN(TEXT("Post Tick Component Wait"),STAT_PostTickComponentUpdateWait,STATGROUP_Game, );
 DECLARE_CYCLE_STAT_EXTERN(TEXT("  Recreate"),STAT_PostTickComponentRecreate,STATGROUP_Game, );
