- Proper handling for const ref user defined pin
- VER_UE4_SERIALIZE_PINTYPE_CONST
- Refactor: UEdGraphNode::CreatePin

diff --git a/Engine/Source/Editor/BlueprintGraph/Classes/K2Node_EditablePinBase.h b/Engine/Source/Editor/BlueprintGraph/Classes/K2Node_EditablePinBase.h
--- a/Engine/Source/Editor/BlueprintGraph/Classes/K2Node_EditablePinBase.h
+++ b/Engine/Source/Editor/BlueprintGraph/Classes/K2Node_EditablePinBase.h
@@ -29,13 +29,22 @@ struct FUserPinInfo
 	{
 		Ar << Info.PinName;
 
-		Ar << Info.PinType.bIsArray;
-		Ar << Info.PinType.bIsReference;
-
-		Ar << Info.PinType.PinCategory;
-		Ar << Info.PinType.PinSubCategory;
-
-		Ar << Info.PinType.PinSubCategoryObject;
+		if (Ar.UE4Ver() >= VER_UE4_SERIALIZE_PINTYPE_CONST)
+		{
+			Info.PinType.Serialize(Ar);
+			Ar << Info.DesiredPinDirection;
+		}
+		else
+		{
+			Ar << Info.PinType.bIsArray;
+			Ar << Info.PinType.bIsReference;
+
+			Ar << Info.PinType.PinCategory;
+			Ar << Info.PinType.PinSubCategory;
+
+			Ar << Info.PinType.PinSubCategoryObject;
+		}
+		
 		Ar << Info.PinDefaultValue;
 
 		return Ar;
diff --git a/Engine/Source/Editor/BlueprintGraph/Private/K2Node.cpp b/Engine/Source/Editor/BlueprintGraph/Private/K2Node.cpp
--- a/Engine/Source/Editor/BlueprintGraph/Private/K2Node.cpp
+++ b/Engine/Source/Editor/BlueprintGraph/Private/K2Node.cpp
@@ -878,7 +878,7 @@ void FOptionalPinManager::CreateVisiblePins(TArray<FOptionalPinFromProperty>& Pr
 							Args.Add(TEXT("Index"), Index);
 							const FText PinFriendlyName = FText::Format(LOCTEXT("PinFriendlyNameWithIndex", "{PinName}_{Index}"), Args);
 							const FString PinName = PinFriendlyName.ToString();
-							NewPin = TargetNode->CreatePin(Direction, PinType.PinCategory, PinType.PinSubCategory, PinType.PinSubCategoryObject.Get(), PinType.bIsArray, PinType.bIsReference, PinName);
+							NewPin = TargetNode->CreatePin(Direction, PinType, PinName);
 							NewPin->PinFriendlyName = PinFriendlyName;
 							Schema->ConstructBasicPinTooltip(*NewPin, PropertyEntry.PropertyTooltip, NewPin->PinToolTip);
 
@@ -911,7 +911,7 @@ void FOptionalPinManager::CreateVisiblePins(TArray<FOptionalPinFromProperty>& Pr
 					if (PropertyEntry.bShowPin)
 					{
 						const FString PinName = PropertyEntry.PropertyName.ToString();
-						NewPin = TargetNode->CreatePin(Direction, PinType.PinCategory, PinType.PinSubCategory, PinType.PinSubCategoryObject.Get(), PinType.bIsArray, PinType.bIsReference, PinName);
+						NewPin = TargetNode->CreatePin(Direction, PinType, PinName);
 						NewPin->PinFriendlyName = PropertyEntry.PropertyFriendlyName.IsEmpty() ? FText::FromString(PinName) : FText::FromString(PropertyEntry.PropertyFriendlyName);
 						Schema->ConstructBasicPinTooltip(*NewPin, PropertyEntry.PropertyTooltip, NewPin->PinToolTip);
 
diff --git a/Engine/Source/Editor/BlueprintGraph/Private/K2Node_Composite.cpp b/Engine/Source/Editor/BlueprintGraph/Private/K2Node_Composite.cpp
--- a/Engine/Source/Editor/BlueprintGraph/Private/K2Node_Composite.cpp
+++ b/Engine/Source/Editor/BlueprintGraph/Private/K2Node_Composite.cpp
@@ -27,14 +27,7 @@ void UK2Node_Composite::AllocateDefaultPins()
 			UEdGraphPin* PortPin = *PinIt;
 			if (PortPin->Direction == EGPD_Input)
 			{
-				UEdGraphPin* NewPin = CreatePin(
-					UEdGraphPin::GetComplementaryDirection(PortPin->Direction),
-					PortPin->PinType.PinCategory,
-					PortPin->PinType.PinSubCategory,
-					PortPin->PinType.PinSubCategoryObject.Get(),
-					PortPin->PinType.bIsArray,
-					PortPin->PinType.bIsReference,
-					PortPin->PinName);
+				UEdGraphPin* NewPin = CreatePin(UEdGraphPin::GetComplementaryDirection(PortPin->Direction), PortPin->PinType, PortPin->PinName);
 				NewPin->DefaultValue = NewPin->AutogeneratedDefaultValue = PortPin->DefaultValue;
 			}
 		}
@@ -47,14 +40,7 @@ void UK2Node_Composite::AllocateDefaultPins()
 			UEdGraphPin* PortPin = *PinIt;
 			if (PortPin->Direction == EGPD_Output)
 			{
-				UEdGraphPin* NewPin = CreatePin(
-					UEdGraphPin::GetComplementaryDirection(PortPin->Direction),
-					PortPin->PinType.PinCategory,
-					PortPin->PinType.PinSubCategory,
-					PortPin->PinType.PinSubCategoryObject.Get(),
-					PortPin->PinType.bIsArray,
-					PortPin->PinType.bIsReference,
-					PortPin->PinName);
+				UEdGraphPin* NewPin = CreatePin(UEdGraphPin::GetComplementaryDirection(PortPin->Direction), PortPin->PinType, PortPin->PinName);
 				NewPin->DefaultValue = NewPin->AutogeneratedDefaultValue = PortPin->DefaultValue;
 			}
 		}
diff --git a/Engine/Source/Editor/BlueprintGraph/Private/K2Node_CustomEvent.cpp b/Engine/Source/Editor/BlueprintGraph/Private/K2Node_CustomEvent.cpp
--- a/Engine/Source/Editor/BlueprintGraph/Private/K2Node_CustomEvent.cpp
+++ b/Engine/Source/Editor/BlueprintGraph/Private/K2Node_CustomEvent.cpp
@@ -162,14 +162,7 @@ bool UK2Node_CustomEvent::CanCreateUserDefinedPin(const FEdGraphPinType& InPinTy
 
 UEdGraphPin* UK2Node_CustomEvent::CreatePinFromUserDefinition(const TSharedPtr<FUserPinInfo> NewPinInfo)
 {
-	UEdGraphPin* NewPin = CreatePin(
-		EGPD_Output, 
-		NewPinInfo->PinType.PinCategory, 
-		NewPinInfo->PinType.PinSubCategory, 
-		NewPinInfo->PinType.PinSubCategoryObject.Get(), 
-		NewPinInfo->PinType.bIsArray, 
-		NewPinInfo->PinType.bIsReference, 
-		NewPinInfo->PinName);
+	UEdGraphPin* NewPin = CreatePin(EGPD_Output, NewPinInfo->PinType, NewPinInfo->PinName);
 	NewPin->DefaultValue = NewPin->AutogeneratedDefaultValue = NewPinInfo->PinDefaultValue;
 	return NewPin;
 }
diff --git a/Engine/Source/Editor/BlueprintGraph/Private/K2Node_DoOnceMultiInput.cpp b/Engine/Source/Editor/BlueprintGraph/Private/K2Node_DoOnceMultiInput.cpp
--- a/Engine/Source/Editor/BlueprintGraph/Private/K2Node_DoOnceMultiInput.cpp
+++ b/Engine/Source/Editor/BlueprintGraph/Private/K2Node_DoOnceMultiInput.cpp
@@ -207,28 +207,14 @@ void UK2Node_DoOnceMultiInput::AddPinsInner(int32 AdditionalPinIndex)
 	{
 		const FEdGraphPinType InputType = GetInType();
 		FText InputPinName = GetNameForPin(AdditionalPinIndex, true);
-		UEdGraphPin* InputPin = CreatePin(EGPD_Input, 
-			InputType.PinCategory, 
-			InputType.PinSubCategory, 
-			InputType.PinSubCategoryObject.Get(), 
-			InputType.bIsArray, 
-			InputType.bIsReference, 
-			InputPinName.BuildSourceString()
-		);
+		UEdGraphPin* InputPin = CreatePin(EGPD_Input, InputType, InputPinName.BuildSourceString());
 		InputPin->PinFriendlyName = InputPinName;
 	}
 
 	{
 		const FEdGraphPinType OutputType = GetOutType();
 		FText OutputPinName = GetNameForPin(AdditionalPinIndex, false);
-		UEdGraphPin* OutputPin = CreatePin(EGPD_Output,
-			OutputType.PinCategory,
-			OutputType.PinSubCategory,
-			OutputType.PinSubCategoryObject.Get(),
-			OutputType.bIsArray,
-			OutputType.bIsReference,
-			OutputPinName.BuildSourceString()
-		);
+		UEdGraphPin* OutputPin = CreatePin(EGPD_Output, OutputType, OutputPinName.BuildSourceString());
 		OutputPin->PinFriendlyName = OutputPinName;
 	}
 }
diff --git a/Engine/Source/Editor/BlueprintGraph/Private/K2Node_FunctionEntry.cpp b/Engine/Source/Editor/BlueprintGraph/Private/K2Node_FunctionEntry.cpp
--- a/Engine/Source/Editor/BlueprintGraph/Private/K2Node_FunctionEntry.cpp
+++ b/Engine/Source/Editor/BlueprintGraph/Private/K2Node_FunctionEntry.cpp
@@ -221,14 +221,7 @@ UEdGraphPin* UK2Node_FunctionEntry::CreatePinFromUserDefinition(const TSharedPtr
 		return nullptr;
 	}
 
-	UEdGraphPin* NewPin = CreatePin(
-		EGPD_Output, 
-		NewPinInfo->PinType.PinCategory, 
-		NewPinInfo->PinType.PinSubCategory, 
-		NewPinInfo->PinType.PinSubCategoryObject.Get(), 
-		NewPinInfo->PinType.bIsArray, 
-		NewPinInfo->PinType.bIsReference, 
-		NewPinInfo->PinName);
+	UEdGraphPin* NewPin = CreatePin(EGPD_Output, NewPinInfo->PinType, NewPinInfo->PinName);
 	NewPin->DefaultValue = NewPin->AutogeneratedDefaultValue = NewPinInfo->PinDefaultValue;
 	return NewPin;
 }
diff --git a/Engine/Source/Editor/BlueprintGraph/Private/K2Node_FunctionResult.cpp b/Engine/Source/Editor/BlueprintGraph/Private/K2Node_FunctionResult.cpp
--- a/Engine/Source/Editor/BlueprintGraph/Private/K2Node_FunctionResult.cpp
+++ b/Engine/Source/Editor/BlueprintGraph/Private/K2Node_FunctionResult.cpp
@@ -135,14 +135,7 @@ bool UK2Node_FunctionResult::CanCreateUserDefinedPin(const FEdGraphPinType& InPi
 
 UEdGraphPin* UK2Node_FunctionResult::CreatePinFromUserDefinition(const TSharedPtr<FUserPinInfo> NewPinInfo)
 {
-	UEdGraphPin* Pin = CreatePin(
-		EGPD_Input, 
-		NewPinInfo->PinType.PinCategory, 
-		NewPinInfo->PinType.PinSubCategory, 
-		NewPinInfo->PinType.PinSubCategoryObject.Get(), 
-		NewPinInfo->PinType.bIsArray, 
-		NewPinInfo->PinType.bIsReference, 
-		NewPinInfo->PinName);
+	UEdGraphPin* Pin = CreatePin(EGPD_Input, NewPinInfo->PinType, NewPinInfo->PinName);
 	FFillDefaultPinValueHelper::Fill(Pin);
 	return Pin;
 }
diff --git a/Engine/Source/Editor/BlueprintGraph/Private/K2Node_MacroInstance.cpp b/Engine/Source/Editor/BlueprintGraph/Private/K2Node_MacroInstance.cpp
--- a/Engine/Source/Editor/BlueprintGraph/Private/K2Node_MacroInstance.cpp
+++ b/Engine/Source/Editor/BlueprintGraph/Private/K2Node_MacroInstance.cpp
@@ -103,14 +103,7 @@ void UK2Node_MacroInstance::AllocateDefaultPins()
 						// We're not interested in any pins that have been expanded internally on the macro
 						if (PortPin->ParentPin == NULL)
 						{
-							UEdGraphPin* NewLocalPin = CreatePin(
-								UEdGraphPin::GetComplementaryDirection(PortPin->Direction),
-								PortPin->PinType.PinCategory,
-								PortPin->PinType.PinSubCategory,
-								PortPin->PinType.PinSubCategoryObject.Get(),
-								PortPin->PinType.bIsArray,
-								PortPin->PinType.bIsReference,
-								PortPin->PinName);
+							UEdGraphPin* NewLocalPin = CreatePin(UEdGraphPin::GetComplementaryDirection(PortPin->Direction), PortPin->PinType, PortPin->PinName);
 							NewLocalPin->DefaultValue = NewLocalPin->AutogeneratedDefaultValue = PortPin->DefaultValue;
 						}
 					}
diff --git a/Engine/Source/Editor/BlueprintGraph/Private/K2Node_Tunnel.cpp b/Engine/Source/Editor/BlueprintGraph/Private/K2Node_Tunnel.cpp
--- a/Engine/Source/Editor/BlueprintGraph/Private/K2Node_Tunnel.cpp
+++ b/Engine/Source/Editor/BlueprintGraph/Private/K2Node_Tunnel.cpp
@@ -189,14 +189,7 @@ UEdGraphPin* UK2Node_Tunnel::CreatePinFromUserDefinition(const TSharedPtr<FUserP
 		Direction = NewPinInfo->DesiredPinDirection;
 	}
 
-	UEdGraphPin* Result = CreatePin(
-		Direction, 
-		NewPinInfo->PinType.PinCategory, 
-		NewPinInfo->PinType.PinSubCategory, 
-		NewPinInfo->PinType.PinSubCategoryObject.Get(), 
-		NewPinInfo->PinType.bIsArray, 
-		NewPinInfo->PinType.bIsReference, 
-		NewPinInfo->PinName);
+	UEdGraphPin* Result = CreatePin(Direction, NewPinInfo->PinType, NewPinInfo->PinName);
 	Result->DefaultValue = Result->AutogeneratedDefaultValue = NewPinInfo->PinDefaultValue;
 
 	// Make sure it mirrors onto the associated node
@@ -206,14 +199,7 @@ UEdGraphPin* UK2Node_Tunnel::CreatePinFromUserDefinition(const TSharedPtr<FUserP
 		UEdGraphPin* HasPinAlready = TargetNode->FindPin(Result->PinName);
 		if (HasPinAlready == NULL)
 		{
-			TargetNode->CreatePin(
-				UEdGraphPin::GetComplementaryDirection(Direction), 
-				NewPinInfo->PinType.PinCategory, 
-				NewPinInfo->PinType.PinSubCategory, 
-				NewPinInfo->PinType.PinSubCategoryObject.Get(), 
-				NewPinInfo->PinType.bIsArray, 
-				NewPinInfo->PinType.bIsReference, 
-				NewPinInfo->PinName);
+			TargetNode->CreatePin(UEdGraphPin::GetComplementaryDirection(Direction), NewPinInfo->PinType, NewPinInfo->PinName);
 		}
 	}
 	else if (UK2Node_Tunnel* TunnelNode = Cast<UK2Node_Tunnel>(TargetNode))
diff --git a/Engine/Source/Editor/Kismet/Private/BlueprintEditor.cpp b/Engine/Source/Editor/Kismet/Private/BlueprintEditor.cpp
--- a/Engine/Source/Editor/Kismet/Private/BlueprintEditor.cpp
+++ b/Engine/Source/Editor/Kismet/Private/BlueprintEditor.cpp
@@ -5793,14 +5793,7 @@ void FBlueprintEditor::ExtractEventTemplateForFunction(class UK2Node_CustomEvent
 					{
 						FString UniquePortName = InGatewayNode->CreateUniquePinName(PortName);
 
-						RemotePortPin = InGatewayNode->CreatePin(
-							Pin->Direction,
-							Pin->PinType.PinCategory,
-							Pin->PinType.PinSubCategory,
-							Pin->PinType.PinSubCategoryObject.Get(),
-							Pin->PinType.bIsArray,
-							Pin->PinType.bIsReference,
-							UniquePortName);
+						RemotePortPin = InGatewayNode->CreatePin(Pin->Direction, Pin->PinType,UniquePortName);
 						InResultNode->CreateUserDefinedPin(UniquePortName, Pin->PinType, EGPD_Input);
 					}
 					PinLink->BreakAllPinLinks();
@@ -5948,15 +5941,7 @@ void FBlueprintEditor::CollapseNodesIntoGraph(UEdGraphNode* InGatewayNode, UK2No
 
 						if(!RemotePortPin && !LocalPortPin)
 						{
-							RemotePortPin = InGatewayNode->CreatePin(
-								LocalPin->Direction,
-								LocalPin->PinType.PinCategory,
-								LocalPin->PinType.PinSubCategory,
-								LocalPin->PinType.PinSubCategoryObject.Get(),
-								LocalPin->PinType.bIsArray,
-								LocalPin->PinType.bIsReference,
-								UniquePortName);
-
+							RemotePortPin = InGatewayNode->CreatePin(LocalPin->Direction, LocalPin->PinType, UniquePortName);
 							LocalPortPin = LocalPort->CreateUserDefinedPin(UniquePortName, LocalPin->PinType, (LocalPin->Direction == EGPD_Input)? EGPD_Output : EGPD_Input);
 						}
 					}
diff --git a/Engine/Source/Editor/Kismet/Private/SMyBlueprint.cpp b/Engine/Source/Editor/Kismet/Private/SMyBlueprint.cpp
--- a/Engine/Source/Editor/Kismet/Private/SMyBlueprint.cpp
+++ b/Engine/Source/Editor/Kismet/Private/SMyBlueprint.cpp
@@ -2541,7 +2541,8 @@ void SMyBlueprint::EnsureLastPinTypeValid()
 	const UEdGraphSchema_K2* Schema = GetDefault<UEdGraphSchema_K2>();
 	const bool bLastPinTypeValid = (Schema->PC_Struct != LastPinType.PinCategory) || LastPinType.PinSubCategoryObject.IsValid();
 	const bool bLastFunctionPinTypeValid = (Schema->PC_Struct != LastFunctionPinType.PinCategory) || LastFunctionPinType.PinSubCategoryObject.IsValid();
-	if (!bLastPinTypeValid || !bLastFunctionPinTypeValid)
+	const bool bConstType = LastPinType.bIsConst || LastFunctionPinType.bIsConst;
+	if (!bLastPinTypeValid || !bLastFunctionPinTypeValid || bConstType)
 	{
 		ResetLastPinType();
 	}
@@ -2551,10 +2552,8 @@ void SMyBlueprint::ResetLastPinType()
 {
 	const UEdGraphSchema_K2* Schema = GetDefault<UEdGraphSchema_K2>();
 
+	LastPinType.ResetToDefaults();
 	LastPinType.PinCategory = Schema->PC_Boolean;
-	LastPinType.PinSubCategory = TEXT("");
-	LastPinType.bIsArray = false;
-	LastPinType.bIsReference = false;
 	LastFunctionPinType = LastPinType;
 }
 
diff --git a/Engine/Source/Editor/KismetCompiler/Private/KismetCompiler.cpp b/Engine/Source/Editor/KismetCompiler/Private/KismetCompiler.cpp
--- a/Engine/Source/Editor/KismetCompiler/Private/KismetCompiler.cpp
+++ b/Engine/Source/Editor/KismetCompiler/Private/KismetCompiler.cpp
@@ -656,6 +656,11 @@ void FKismetCompilerContext::CreatePropertiesFromList(UStruct* Scope, UField**&
 		{
 			NewProperty->PropertyFlags |= PropertyFlags;
 
+			if (bPropertiesAreParameters && Term.Type.bIsConst)
+			{
+				NewProperty->SetPropertyFlags(CPF_ConstParm);
+			}
+
 			if (Term.bPassedByReference)
 			{
 				// special case for BlueprintImplementableEvent
@@ -2426,14 +2431,7 @@ void FKismetCompilerContext::CreateFunctionStubForEvent(UK2Node_Event* SrcEventN
 			UEdGraphPin* UGSourcePin = SrcEventNode->FindPin(SourcePin->PinName);
 			const FString MemberVariableName = ClassScopeNetNameMap.MakeValidName(UGSourcePin);
 
-			UEdGraphPin* DestPin = AssignmentNode->CreatePin(
-				EGPD_Input, 
-				SourcePin->PinType.PinCategory, 
-				SourcePin->PinType.PinSubCategory, 
-				SourcePin->PinType.PinSubCategoryObject.Get(), 
-				SourcePin->PinType.bIsArray, 
-				SourcePin->PinType.bIsReference, 
-				MemberVariableName);
+			UEdGraphPin* DestPin = AssignmentNode->CreatePin(EGPD_Input, SourcePin->PinType, MemberVariableName);
 			MessageLog.NotifyIntermediateObjectCreation(DestPin, SourcePin);
 			DestPin->MakeLinkTo(SourcePin);
 		}
diff --git a/Engine/Source/Editor/UnrealEd/Private/Commandlets/DumpBlueprintsInfoCommandlet.cpp b/Engine/Source/Editor/UnrealEd/Private/Commandlets/DumpBlueprintsInfoCommandlet.cpp
--- a/Engine/Source/Editor/UnrealEd/Private/Commandlets/DumpBlueprintsInfoCommandlet.cpp
+++ b/Engine/Source/Editor/UnrealEd/Private/Commandlets/DumpBlueprintsInfoCommandlet.cpp
@@ -2004,15 +2004,7 @@ static void DumpBlueprintInfoUtils::DumpContextualPinTypeActions(uint32 Indent,
 	FGraphContextMenuBuilder ContextMenuBuilder(Graph);
 
 	UK2Node_Composite* DummyNode = NewObject<UK2Node_Composite>(Graph);
-	UEdGraphPin* DummyPin = DummyNode->CreatePin(
-		EGPD_Input,
-		PinType.PinCategory,
-		PinType.PinSubCategory,
-		PinType.PinSubCategoryObject.Get(),
-		PinType.bIsArray,
-		PinType.bIsReference,
-		DummyNode->GetName()
-	);
+	UEdGraphPin* DummyPin = DummyNode->CreatePin(EGPD_Input, PinType, DummyNode->GetName());
 	ContextMenuBuilder.FromPin = DummyPin;
 
 	UBlueprint* Blueprint = FBlueprintEditorUtils::FindBlueprintForGraph(Graph);
diff --git a/Engine/Source/Runtime/Core/Public/UObject/ObjectVersion.h b/Engine/Source/Runtime/Core/Public/UObject/ObjectVersion.h
--- a/Engine/Source/Runtime/Core/Public/UObject/ObjectVersion.h
+++ b/Engine/Source/Runtime/Core/Public/UObject/ObjectVersion.h
@@ -506,6 +506,8 @@ enum EUnrealEngineObjectUE4Version
 	VER_UE4_SERIALIZE_LANDSCAPE_ES2_TEXTURES,
 	// Constraint motor velocity is broken into per-component
 	VER_UE4_CONSTRAINT_INSTANCE_MOTOR_FLAGS,
+	// Serialize bIsConst in FEdGraphPinType
+	VER_UE4_SERIALIZE_PINTYPE_CONST,
 
 	// -----<new versions can be added before this line>-------------------------------------------------
 	// - this needs to be the last line (see note below)
diff --git a/Engine/Source/Runtime/Engine/Classes/EdGraph/EdGraphNode.h b/Engine/Source/Runtime/Engine/Classes/EdGraph/EdGraphNode.h
--- a/Engine/Source/Runtime/Engine/Classes/EdGraph/EdGraphNode.h
+++ b/Engine/Source/Runtime/Engine/Classes/EdGraph/EdGraphNode.h
@@ -10,6 +10,7 @@ class UEdGraphNode;
 class UEdGraphSchema;
 class UEdGraphPin;
 class SGraphNode;
+struct FEdGraphPinType;
 
 /** Enum used to define which way data flows into or out of this pin */
 UENUM()
@@ -172,6 +173,9 @@ public:
 	/** Create a new pin on this node using the supplied info, and return the new pin */
 	UEdGraphPin* CreatePin(EEdGraphPinDirection Dir, const FString& PinCategory, const FString& PinSubCategory, UObject* PinSubCategoryObject, bool bIsArray, bool bIsReference, const FString& PinName, bool bIsConst = false, int32 Index = INDEX_NONE);
 
+	/** Create a new pin on this node using the supplied pin type, and return the new pin */
+	UEdGraphPin* CreatePin(EEdGraphPinDirection Dir, const FEdGraphPinType& InPinType, const FString& PinName, int32 Index = INDEX_NONE);
+
 	// Allocates a pin from the pool
 	static UEdGraphPin* AllocatePinFromPool(UEdGraphNode* OuterNode);
 
diff --git a/Engine/Source/Runtime/Engine/Classes/EdGraph/EdGraphPin.h b/Engine/Source/Runtime/Engine/Classes/EdGraph/EdGraphPin.h
--- a/Engine/Source/Runtime/Engine/Classes/EdGraph/EdGraphPin.h
+++ b/Engine/Source/Runtime/Engine/Classes/EdGraph/EdGraphPin.h
@@ -99,6 +99,7 @@ public:
 		PinSubCategoryObject = NULL;
 		bIsArray = false;
 		bIsReference = false;
+		bIsConst = false;
 		bIsWeakPointer = false;
 	}
 	FEdGraphPinType(const FString& InPinCategory, const FString& InPinSubCategory, UObject* InPinSubCategoryObject, bool bInIsArray, bool bInIsReference)
@@ -108,17 +109,19 @@ public:
 		PinSubCategoryObject = InPinSubCategoryObject;
 		bIsArray = bInIsArray;
 		bIsReference = bInIsReference;
+		bIsConst = false;
 		bIsWeakPointer = false;
 	}
 	bool operator == ( const FEdGraphPinType& Other ) const
 	{
-		return (PinCategory == Other.PinCategory) 
-			&& (PinSubCategory == Other.PinSubCategory) 
-			&& (PinSubCategoryObject == Other.PinSubCategoryObject) 
-			&& (bIsArray == Other.bIsArray) 
+		return (PinCategory == Other.PinCategory)
+			&& (PinSubCategory == Other.PinSubCategory)
+			&& (PinSubCategoryObject == Other.PinSubCategoryObject)
+			&& (bIsArray == Other.bIsArray)
 			&& (bIsReference == Other.bIsReference)
 			&& (bIsWeakPointer == Other.bIsWeakPointer)
-			&& (PinSubCategoryMemberReference == Other.PinSubCategoryMemberReference);
+			&& (PinSubCategoryMemberReference == Other.PinSubCategoryMemberReference)
+			&& (bIsConst == Other.bIsConst);
 	}
 	bool operator != ( const FEdGraphPinType& Other ) const
 	{
@@ -127,7 +130,8 @@ public:
 			|| (PinSubCategoryObject != Other.PinSubCategoryObject) 
 			|| (bIsArray != Other.bIsArray) 
 			|| (bIsReference != Other.bIsReference)
-			|| (bIsWeakPointer != Other.bIsWeakPointer);
+			|| (bIsWeakPointer != Other.bIsWeakPointer)
+			|| (bIsConst != Other.bIsConst);
 	}
 
 	void ResetToDefaults()
@@ -139,9 +143,10 @@ public:
 		bIsArray = false;
 		bIsReference = false;
 		bIsWeakPointer = false;
+		bIsConst = false;
 	}
 
-	bool Serialize(FArchive& Ar);
+	ENGINE_API bool Serialize(FArchive& Ar);
 };
 
 template<>
diff --git a/Engine/Source/Runtime/Engine/Private/EdGraph/EdGraphNode.cpp b/Engine/Source/Runtime/Engine/Private/EdGraph/EdGraphNode.cpp
--- a/Engine/Source/Runtime/Engine/Private/EdGraph/EdGraphNode.cpp
+++ b/Engine/Source/Runtime/Engine/Private/EdGraph/EdGraphNode.cpp
@@ -53,21 +53,15 @@ UEdGraphNode::UEdGraphNode(const FObjectInitializer& ObjectInitializer)
 }
 
 #if WITH_EDITOR
-UEdGraphPin* UEdGraphNode::CreatePin(EEdGraphPinDirection Dir, const FString& PinCategory, const FString& PinSubCategory, UObject* PinSubCategoryObject, bool bIsArray, bool bIsReference, const FString& PinName, bool bIsConst /*= false*/, int32 Index /*= INDEX_NONE*/)
+
+UEdGraphPin* UEdGraphNode::CreatePin(EEdGraphPinDirection Dir, const FEdGraphPinType& InPinType, const FString& PinName, int32 Index /*= INDEX_NONE*/)
 {
-#if 0
-	UEdGraphPin* NewPin = AllocatePinFromPool(this);
-#else
 	UEdGraphPin* NewPin = NewObject<UEdGraphPin>(this);
-#endif
 	NewPin->PinName = PinName;
 	NewPin->Direction = Dir;
-	NewPin->PinType.PinCategory = PinCategory;
-	NewPin->PinType.PinSubCategory = PinSubCategory;
-	NewPin->PinType.PinSubCategoryObject = PinSubCategoryObject;
-	NewPin->PinType.bIsArray = bIsArray;
-	NewPin->PinType.bIsReference = bIsReference;
-	NewPin->PinType.bIsConst = bIsConst;
+
+	NewPin->PinType = InPinType;
+
 	NewPin->SetFlags(RF_Transactional);
 
 	if (HasAnyFlags(RF_Transient))
@@ -76,7 +70,7 @@ UEdGraphPin* UEdGraphNode::CreatePin(EEdGraphPinDirection Dir, const FString& Pi
 	}
 
 	Modify(false);
-	if ( Pins.IsValidIndex( Index ) )
+	if (Pins.IsValidIndex(Index))
 	{
 		Pins.Insert(NewPin, Index);
 	}
@@ -87,6 +81,14 @@ UEdGraphPin* UEdGraphNode::CreatePin(EEdGraphPinDirection Dir, const FString& Pi
 	return NewPin;
 }
 
+UEdGraphPin* UEdGraphNode::CreatePin(EEdGraphPinDirection Dir, const FString& PinCategory, const FString& PinSubCategory, UObject* PinSubCategoryObject, bool bIsArray, bool bIsReference, const FString& PinName, bool bIsConst /*= false*/, int32 Index /*= INDEX_NONE*/)
+{
+	FEdGraphPinType PinType(PinCategory, PinSubCategory, PinSubCategoryObject, bIsArray, bIsReference);
+	PinType.bIsConst = bIsConst;
+
+	return CreatePin(Dir, PinType, PinName, Index);
+}
+
 UEdGraphPin* UEdGraphNode::FindPin(const FString& PinName) const
 {
 	for(int32 PinIdx=0; PinIdx<Pins.Num(); PinIdx++)
diff --git a/Engine/Source/Runtime/Engine/Private/EdGraph/EdGraphPin.cpp b/Engine/Source/Runtime/Engine/Private/EdGraph/EdGraphPin.cpp
--- a/Engine/Source/Runtime/Engine/Private/EdGraph/EdGraphPin.cpp
+++ b/Engine/Source/Runtime/Engine/Private/EdGraph/EdGraphPin.cpp
@@ -63,6 +63,15 @@ bool FEdGraphPinType::Serialize(FArchive& Ar)
 		}
 	}
 
+	if (Ar.UE4Ver() >= VER_UE4_SERIALIZE_PINTYPE_CONST)
+	{
+		Ar << bIsConst;
+	}
+	else if (Ar.IsLoading())
+	{
+		bIsConst = false;
+	}
+
 	return true;
 }
 
