This Adds an NOR and NAND to the kismet math library.
#UE-12354 - Users would like NAND (not and) and NOR (not or) bool nodes
#Branch UE4
#Proj Engine

diff --git a/Engine/Source/Runtime/Engine/Classes/Kismet/KismetMathLibrary.h b/Engine/Source/Runtime/Engine/Classes/Kismet/KismetMathLibrary.h
--- a/Engine/Source/Runtime/Engine/Classes/Kismet/KismetMathLibrary.h
+++ b/Engine/Source/Runtime/Engine/Classes/Kismet/KismetMathLibrary.h
@@ -84,6 +84,10 @@ class ENGINE_API UKismetMathLibrary : public UBlueprintFunctionLibrary
 	UFUNCTION(BlueprintPure, meta=(DisplayName = "AND Boolean", CompactNodeTitle = "AND", Keywords = "& and", CommutativeAssociativeBinaryOperator = "true"), Category="Math|Boolean")
 	static bool BooleanAND(bool A, bool B);
 
+	/* Returns the logical NAND of two values (A AND B) */
+	UFUNCTION(BlueprintPure, meta=(DisplayName = "NAND Boolean", CompactNodeTitle = "NAND", Keywords = "!& nand", CommutativeAssociativeBinaryOperator = "true"), Category="Math|Boolean")
+	static bool BooleanNAND(bool A, bool B);
+
 	/* Returns the logical OR of two values (A OR B) */
 	UFUNCTION(BlueprintPure, meta=(DisplayName = "OR Boolean", CompactNodeTitle = "OR", Keywords = "| or", CommutativeAssociativeBinaryOperator = "true"), Category="Math|Boolean")
 	static bool BooleanOR(bool A, bool B);
@@ -92,6 +96,10 @@ class ENGINE_API UKismetMathLibrary : public UBlueprintFunctionLibrary
 	UFUNCTION(BlueprintPure, meta=(DisplayName = "XOR Boolean", CompactNodeTitle = "XOR", Keywords = "^ xor"), Category="Math|Boolean")
 	static bool BooleanXOR(bool A, bool B);
 
+	/* Returns the logical Not OR of two values (A NOR B) */
+	UFUNCTION(BlueprintPure, meta=(DisplayName = "NOR Boolean", CompactNodeTitle = "NOR", Keywords = "!^ nor"), Category="Math|Boolean")
+	static bool BooleanNOR(bool A, bool B);
+
 	//
 	// Byte functions.
 	//
diff --git a/Engine/Source/Runtime/Engine/Private/KismetMathLibrary.cpp b/Engine/Source/Runtime/Engine/Private/KismetMathLibrary.cpp
--- a/Engine/Source/Runtime/Engine/Private/KismetMathLibrary.cpp
+++ b/Engine/Source/Runtime/Engine/Private/KismetMathLibrary.cpp
@@ -56,6 +56,10 @@ bool UKismetMathLibrary::BooleanAND(bool A, bool B)
 	return A && B;
 }
 
+bool UKismetMathLibrary::BooleanNAND(bool A, bool B)
+{
+	return !(A && B);
+}
 
 bool UKismetMathLibrary::BooleanOR(bool A, bool B)
 {
@@ -67,6 +71,11 @@ bool UKismetMathLibrary::BooleanXOR(bool A, bool B)
 	return A ^ B;
 }
 
+bool UKismetMathLibrary::BooleanNOR(bool A, bool B)
+{
+	return !(A || B);
+}
+
 uint8 UKismetMathLibrary::Multiply_ByteByte(uint8 A, uint8 B)
 {
 	return A * B;
