UE-13926 - Invalidate lighting cache when changing static mesh.

diff --git a/Engine/Source/Runtime/Engine/Private/Components/StaticMeshComponent.cpp b/Engine/Source/Runtime/Engine/Private/Components/StaticMeshComponent.cpp
--- a/Engine/Source/Runtime/Engine/Private/Components/StaticMeshComponent.cpp
+++ b/Engine/Source/Runtime/Engine/Private/Components/StaticMeshComponent.cpp
@@ -1116,6 +1116,11 @@ void UStaticMeshComponent::PostEditChangeProperty(FPropertyChangedEvent& Propert
 		{
 			GEngine->TriggerStreamingDataRebuild();
 		}
+
+		if ( PropertyThatChanged->GetName() == TEXT("StaticMesh") )
+		{
+			InvalidateLightingCache();
+		}
 	}
 
 	FBodyInstanceEditorHelpers::EnsureConsistentMobilitySimulationSettingsOnPostEditChange(this, PropertyChangedEvent);
