#ue4 - Perf: Avoid some GestureRecognizer logic when nothing has changed.

diff --git a/Engine/Source/Runtime/Engine/Private/GestureRecognizer.cpp b/Engine/Source/Runtime/Engine/Private/GestureRecognizer.cpp
--- a/Engine/Source/Runtime/Engine/Private/GestureRecognizer.cpp
+++ b/Engine/Source/Runtime/Engine/Private/GestureRecognizer.cpp
@@ -20,79 +20,83 @@ void FGestureRecognizer::DetectGestures(const FVector (&Touches)[EKeys::NUM_TOUC
 		}
 	}
 
-	// place new anchor points
-	for (int32 Index = 0; Index < ARRAY_COUNT(AnchorPoints); Index++)
+	// Further processing is only needed if there were or are active touches.
+	if (PreviousTouchCount != 0 || TouchCount != 0)
 	{
-		if (PreviousTouchCount < Index + 1 && TouchCount >= Index + 1)
+		// place new anchor points
+		for (int32 Index = 0; Index < ARRAY_COUNT(AnchorPoints); Index++)
 		{
-			AnchorPoints[Index] = FVector2D(Touches[Index].X, Touches[Index].Y);
+			if (PreviousTouchCount < Index + 1 && TouchCount >= Index + 1)
+			{
+				AnchorPoints[Index] = FVector2D(Touches[Index].X, Touches[Index].Y);
+			}
 		}
-	}
 
-	// handle different types of swipes
-	if (TouchCount >= 2)
-	{
-		float* CurrentAlpha = CurrentGestureValues.Find(EKeys::Gesture_Pinch);
-		if (CurrentAlpha == NULL)
+		// handle different types of swipes
+		if (TouchCount >= 2)
 		{
-			// remember the starting spots
-			AnchorDistance = (AnchorPoints[1] - AnchorPoints[0]).Size();
+			float* CurrentAlpha = CurrentGestureValues.Find(EKeys::Gesture_Pinch);
+			if (CurrentAlpha == NULL)
+			{
+				// remember the starting spots
+				AnchorDistance = (AnchorPoints[1] - AnchorPoints[0]).Size();
+
+				// alpha of 1 is initial pinch anchor distance
+				CurrentGestureValues.Add(EKeys::Gesture_Pinch, 1.0f);
+
+				HandleGesture(PlayerInput, EKeys::Gesture_Pinch, true, false);
+			}
+			else
+			{
+				// calculate current alpha
+				float NewDistance = (FVector2D(Touches[1].X, Touches[1].Y) - FVector2D(Touches[0].X, Touches[0].Y)).Size();
+				*CurrentAlpha = NewDistance / AnchorDistance;
+
+				HandleGesture(PlayerInput, EKeys::Gesture_Pinch, false, false);
+			}
+		}
 
-			// alpha of 1 is initial pinch anchor distance
-			CurrentGestureValues.Add(EKeys::Gesture_Pinch, 1.0f);
+		// cancel any 2 finger gestures
+		if (PreviousTouchCount >= 2 && TouchCount < 2)
+		{
+			HandleGesture(PlayerInput, EKeys::Gesture_TwoFingerSwipeLeftRight, false, true);
+			HandleGesture(PlayerInput, EKeys::Gesture_TwoFingerSwipeUpDown, false, true);
+			HandleGesture(PlayerInput, EKeys::Gesture_Pinch, false, true);
+		}
 
-			HandleGesture(PlayerInput, EKeys::Gesture_Pinch, true, false);
+		if (PreviousTouchCount == 0 && TouchCount == 1)
+		{
+			// initialize the flick
+			FlickTime = 0.0f;
 		}
-		else
+		else if (PreviousTouchCount == 1 && TouchCount == 1)
 		{
-			// calculate current alpha
-			float NewDistance = (FVector2D(Touches[1].X, Touches[1].Y) - FVector2D(Touches[0].X, Touches[0].Y)).Size();
-			*CurrentAlpha = NewDistance / AnchorDistance;
-
-			HandleGesture(PlayerInput, EKeys::Gesture_Pinch, false, false);
+			// remember the position for when we let go
+			FlickCurrent = FVector2D(Touches[0].X, Touches[0].Y);
+			FlickTime += DeltaTime;
 		}
-	}
-
-	// cancel any 2 finger gestures
-	if (PreviousTouchCount >= 2 && TouchCount < 2)
-	{
-		HandleGesture(PlayerInput, EKeys::Gesture_TwoFingerSwipeLeftRight, false, true);
-		HandleGesture(PlayerInput, EKeys::Gesture_TwoFingerSwipeUpDown, false, true);
-		HandleGesture(PlayerInput, EKeys::Gesture_Pinch, false, true);
-	}
-
-	if (PreviousTouchCount == 0 && TouchCount == 1)
-	{
-		// initialize the flick
-		FlickTime = 0.0f;
-	}
-	else if (PreviousTouchCount == 1 && TouchCount == 1)
-	{
-		// remember the position for when we let go
-		FlickCurrent = FVector2D(Touches[0].X, Touches[0].Y);
-		FlickTime += DeltaTime;
-	}
-	else if (PreviousTouchCount >= 1 && TouchCount == 0)
-	{
-		// must be a fast flick
-		if (FlickTime < 0.25f && (FlickCurrent - AnchorPoints[0]).Size() > 100.f)
+		else if (PreviousTouchCount >= 1 && TouchCount == 0)
 		{
-			// this is the angle from +X in screen space, meaning right is 0, up is 90, left is 180, down is 270
-			float Angle = FMath::Atan2(-(FlickCurrent.Y - AnchorPoints[0].Y), FlickCurrent.X - AnchorPoints[0].X) * 180.f / PI;
-			Angle = FRotator::ClampAxis(Angle);
-
-			// flicks are one-shot, so we start and end in the same frame
-			CurrentGestureValues.Add(EKeys::Gesture_Flick, Angle);
-			HandleGesture(PlayerInput, EKeys::Gesture_Flick, true, true);
+			// must be a fast flick
+			if (FlickTime < 0.25f && (FlickCurrent - AnchorPoints[0]).Size() > 100.f)
+			{
+				// this is the angle from +X in screen space, meaning right is 0, up is 90, left is 180, down is 270
+				float Angle = FMath::Atan2(-(FlickCurrent.Y - AnchorPoints[0].Y), FlickCurrent.X - AnchorPoints[0].X) * 180.f / PI;
+				Angle = FRotator::ClampAxis(Angle);
+
+				// flicks are one-shot, so we start and end in the same frame
+				CurrentGestureValues.Add(EKeys::Gesture_Flick, Angle);
+				HandleGesture(PlayerInput, EKeys::Gesture_Flick, true, true);
+			}
 		}
-	}
 
-	if (PreviousTouchCount >= 1 && TouchCount == 0)
-	{
-		// cancel other 1 finger gestures
-		HandleGesture(PlayerInput, EKeys::Gesture_SwipeLeftRight, false, true);
-		HandleGesture(PlayerInput, EKeys::Gesture_SwipeUpDown, false, true);
-	}
+		if (PreviousTouchCount >= 1 && TouchCount == 0)
+		{
+			// cancel other 1 finger gestures
+			HandleGesture(PlayerInput, EKeys::Gesture_SwipeLeftRight, false, true);
+			HandleGesture(PlayerInput, EKeys::Gesture_SwipeUpDown, false, true);
+		}
+	}	
 
 	// remember for next frame
 	PreviousTouchCount = TouchCount;
