Fixed incorrect behavior when dragging objects into the world from the content browser.  Collision is disabled at runtime for these objects but is set to QueryOnly in editor for use in tools that require collision checks.  Recent code was added that checks the collision flag (consistent at runtime) and prevents actors/components from being ignored during traces if they have no collision.  However it was not considering the the case of the editor using QueryOnly. Not ignoring the dragged actor when placing from the content browser causes the actor to be checked against itself causing a feedback effect.

https://jira.ol.epicgames.net/browse/UE-14987

diff --git a/Engine/Source/Runtime/Engine/Private/Collision/Collision.cpp b/Engine/Source/Runtime/Engine/Private/Collision/Collision.cpp
--- a/Engine/Source/Runtime/Engine/Private/Collision/Collision.cpp
+++ b/Engine/Source/Runtime/Engine/Private/Collision/Collision.cpp
@@ -164,7 +164,16 @@ FCollisionQueryParams::FCollisionQueryParams(FName InTraceTag, bool bInTraceComp
 
 void FCollisionQueryParams::AddIgnoredActor(const AActor* InIgnoreActor)
 {
-	if (InIgnoreActor && InIgnoreActor->GetActorEnableCollision())
+	// Note: This code operates differently in the editor because the actors in the editor can have their collision setting become out of sync with physx.  This happens because even when collision is disabled on an actor, in the editor we 
+	// tell phsx that we still require queries( See FBodyInstance::UpdatePhysicsFilterData).  Doing so allows us to perform editor only traces against objects with collision disabled.  Due to this, we cannot assume that the collision enabled flag here
+	// setting is correct compared to physx and so we still ignore specified components regardless of their collision setting
+#if WITH_EDITOR
+	const bool bCheckForCollision = InIgnoreActor && InIgnoreActor->GetWorld() && InIgnoreActor->GetWorld()->IsGameWorld();
+#else
+	const bool bCheckForCollision = true;
+#endif
+
+	if (InIgnoreActor)
 	{
 		if (IgnoreComponents.Num() == 0)
 		{
@@ -172,7 +181,8 @@ void FCollisionQueryParams::AddIgnoredActor(const AActor* InIgnoreActor)
 			for (const UActorComponent* ActorComponent : InIgnoreActor->GetComponents())
 			{
 				const UPrimitiveComponent* PrimComponent = Cast<const UPrimitiveComponent>(ActorComponent);
-				if (PrimComponent && PrimComponent->BodyInstance.GetCollisionEnabled() != ECollisionEnabled::NoCollision)
+				
+				if (PrimComponent && (!bCheckForCollision || PrimComponent->BodyInstance.GetCollisionEnabled() != ECollisionEnabled::NoCollision) )
 				{
 					IgnoreComponents.Add(PrimComponent->GetUniqueID());
 				}
@@ -183,7 +193,8 @@ void FCollisionQueryParams::AddIgnoredActor(const AActor* InIgnoreActor)
 			for (const UActorComponent* ActorComponent : InIgnoreActor->GetComponents())
 			{
 				const UPrimitiveComponent* PrimComponent = Cast<const UPrimitiveComponent>(ActorComponent);
-				if (PrimComponent && PrimComponent->BodyInstance.GetCollisionEnabled() != ECollisionEnabled::NoCollision)
+
+				if (PrimComponent && (!bCheckForCollision || PrimComponent->BodyInstance.GetCollisionEnabled() != ECollisionEnabled::NoCollision) )
 				{
 					IgnoreComponents.AddUnique(PrimComponent->GetUniqueID());
 				}
@@ -217,7 +228,16 @@ void FCollisionQueryParams::AddIgnoredActors(const TArray<TWeakObjectPtr<AActor>
 
 void FCollisionQueryParams::AddIgnoredComponent(const UPrimitiveComponent* InIgnoreComponent)
 {
-	if (InIgnoreComponent && InIgnoreComponent->GetCollisionEnabled() != ECollisionEnabled::NoCollision)
+	// Note: This code operates differently in the editor because the actors in the editor can have their collision setting become out of sync with physx.  This happens because even when collision is disabled on an actor, in the editor we 
+	// tell phsx that we still require queries( See FBodyInstance::UpdatePhysicsFilterData).  Doing so allows us to perform editor only traces against objects with collision disabled.  Due to this, we cannot assume that the collision enabled flag here
+	// setting is correct compared to physx and so we still ignore specified components regardless of their collision setting
+#if WITH_EDITOR
+	const bool bCheckForCollision = InIgnoreComponent && InIgnoreComponent->GetWorld() && InIgnoreComponent->GetWorld()->IsGameWorld();
+#else
+	const bool bCheckForCollision = true;
+#endif
+
+	if (InIgnoreComponent && (!bCheckForCollision || InIgnoreComponent->GetCollisionEnabled() != ECollisionEnabled::NoCollision))
 	{
 		IgnoreComponents.AddUnique(InIgnoreComponent->GetUniqueID());
 	}
