Merging fix that re-added UPROPERTY tag that was removed by accident in CL#2521346 (CL# 2521346 in main)

diff --git a/Engine/Source/Runtime/Engine/Classes/GameFramework/PlayerController.h b/Engine/Source/Runtime/Engine/Classes/GameFramework/PlayerController.h
--- a/Engine/Source/Runtime/Engine/Classes/GameFramework/PlayerController.h
+++ b/Engine/Source/Runtime/Engine/Classes/GameFramework/PlayerController.h
@@ -1492,6 +1492,7 @@ protected:
 
 private:
 	/** The pawn used when spectating (NULL if not spectating). */
+	UPROPERTY()
 	class ASpectatorPawn* SpectatorPawn;
 
 protected:
