Return GameModeClassAliases to being a config variable.

diff --git a/Engine/Source/Runtime/Engine/Classes/GameFramework/GameMode.h b/Engine/Source/Runtime/Engine/Classes/GameFramework/GameMode.h
--- a/Engine/Source/Runtime/Engine/Classes/GameFramework/GameMode.h
+++ b/Engine/Source/Runtime/Engine/Classes/GameFramework/GameMode.h
@@ -262,7 +262,7 @@ public:
 protected:
 
 	/** Handy alternate short names for GameMode classes (e.g. "DM" could be an alias for "MyProject.MyGameModeMP_DM". */
-	UPROPERTY(EditAnywhere, Category=GameMode, AdvancedDisplay)
+	UPROPERTY(config)
 	TArray<struct FGameClassShortName> GameModeClassAliases;
 
 	/** Time a playerstate will stick around in an inactive state after a player logout */
