- Fix for matinee rendering crash when stereo device is connected.

diff --git a/Engine/Source/Runtime/Engine/Private/GameEngine.cpp b/Engine/Source/Runtime/Engine/Private/GameEngine.cpp
--- a/Engine/Source/Runtime/Engine/Private/GameEngine.cpp
+++ b/Engine/Source/Runtime/Engine/Private/GameEngine.cpp
@@ -90,6 +90,7 @@ UGameEngine::UGameEngine(const FObjectInitializer& ObjectInitializer)
 void UGameEngine::CreateGameViewportWidget( UGameViewportClient* GameViewportClient )
 {
 	bool bRenderDirectlyToWindow = !GEngine->MatineeScreenshotOptions.bStartWithMatineeCapture && GIsDumpingMovie == 0; 
+	const bool bStereoAllowed = bRenderDirectlyToWindow;
 	TSharedRef<SOverlay> ViewportOverlayWidgetRef = SNew( SOverlay );
 
 	TSharedRef<SGameLayerManager> GameLayerManagerRef = SNew(SGameLayerManager)
@@ -106,7 +107,7 @@ void UGameEngine::CreateGameViewportWidget( UGameViewportClient* GameViewportCli
 			.RenderDirectlyToWindow(bRenderDirectlyToWindow)
 			//gamma handled by the scene renderer
 			.EnableGammaCorrection(false)
-			.EnableStereoRendering(true)
+			.EnableStereoRendering(bStereoAllowed)
 			[
 				GameLayerManagerRef
 			];
