It is acceptable to have loose SceneComponents that are not attached to your actor (removing fixup where we tried to enforce that they be attached).

diff --git a/Engine/Source/Runtime/Engine/Private/Actor.cpp b/Engine/Source/Runtime/Engine/Private/Actor.cpp
--- a/Engine/Source/Runtime/Engine/Private/Actor.cpp
+++ b/Engine/Source/Runtime/Engine/Private/Actor.cpp
@@ -2344,32 +2344,22 @@ static USceneComponent* FixupNativeActorComponents(AActor* Actor)
 	if ((SceneRootComponent == nullptr) && (SceneComponents.Num() > 0))
 	{
 		UE_LOG(LogActor, Warning, TEXT("%s has natively added scene component(s), but none of them were set as the actor's RootComponent - picking one arbitrarily"), *Actor->GetFullName());
-	}
-
-	// if the user forgot to set one of their native components as the root, 
-	// we arbitrarily pick one for them (otherwise the SCS could attempt to 
-	// create its own root, and nest native components under it)
-	for (USceneComponent* Component : SceneComponents)
-	{
-		if ((Component == nullptr) || 
-			(Component->AttachParent != nullptr) || 
-			(Component->CreationMethod != EComponentCreationMethod::Native) || 
-			(Component == SceneRootComponent))
-		{
-			continue;
-		}
+	
+		// if the user forgot to set one of their native components as the root, 
+		// we arbitrarily pick one for them (otherwise the SCS could attempt to 
+		// create its own root, and nest native components under it)
+		for (USceneComponent* Component : SceneComponents)
+		{
+			if ((Component == nullptr) ||
+				(Component->AttachParent != nullptr) ||
+				(Component->CreationMethod != EComponentCreationMethod::Native))
+			{
+				continue;
+			}
 
-		// if we've already picked a root component (and this was left 
-		// unattached), then attach this one to the root
-		if (SceneRootComponent != nullptr)
-		{
-			UE_LOG(LogActor, Warning, TEXT("Natively added component (%s) was left unattached from the actor's root."), *SceneRootComponent->GetReadableName());
-			Component->AttachTo(SceneRootComponent);
-		}
-		else
-		{
 			SceneRootComponent = Component;
 			Actor->SetRootComponent(Component);
+			break;
 		}
 	}
 
