Fixed: Quality settings should not affect landscape resolution UE-11542

diff --git a/Engine/Source/Runtime/Engine/Private/ContentStreaming.cpp b/Engine/Source/Runtime/Engine/Private/ContentStreaming.cpp
--- a/Engine/Source/Runtime/Engine/Private/ContentStreaming.cpp
+++ b/Engine/Source/Runtime/Engine/Private/ContentStreaming.cpp
@@ -505,6 +505,12 @@ int32 FFloatMipLevel::ComputeMip(const FStreamingTexture* StreamingTexture, floa
 			bOptimal ? StreamingTexture->MaxAllowedOptimalMips : 
 #endif
 			StreamingTexture->MaxAllowedMips;
+		// Landscape textures should not be affected by MipBias
+		if (StreamingTexture->LODGroup == TEXTUREGROUP_Terrain_Heightmap || 
+			StreamingTexture->LODGroup == TEXTUREGROUP_Terrain_Weightmap)
+		{
+			MipBias = 0.0f;
+		}
 	}
 
 	// bias the clamped mip level
