Fully DLL exposed ADecalActor so that you can derive from it

UE-15030 - Linker errors when deriving from ADecalActor in game

Also cleaned up some inline forward declarations so they're done at the top of the file instead.

diff --git a/Engine/Source/Runtime/Engine/Classes/Engine/DecalActor.h b/Engine/Source/Runtime/Engine/Classes/Engine/DecalActor.h
--- a/Engine/Source/Runtime/Engine/Classes/Engine/DecalActor.h
+++ b/Engine/Source/Runtime/Engine/Classes/Engine/DecalActor.h
@@ -5,7 +5,10 @@
 #include "DecalActor.generated.h"
 
 
+class UArrowComponent;
+class UBillboardComponent;
 class UBoxComponent;
+class UDecalComponent;
 
 /**
 * DecalActor contains a DecalComponent which can be used to render material modifications on top of existing geometry.
@@ -13,8 +16,8 @@ class UBoxComponent;
 * @see https://docs.unrealengine.com/latest/INT/Engine/Actors/DecalActor
 * @see UDecalComponent
 */
-UCLASS(hideCategories=(Collision, Attachment, Actor, Input, Replication), showCategories=("Input|MouseInput", "Input|TouchInput"), ComponentWrapperClass,MinimalAPI)
-class ADecalActor
+UCLASS(hideCategories=(Collision, Attachment, Actor, Input, Replication), showCategories=("Input|MouseInput", "Input|TouchInput"), ComponentWrapperClass)
+class ENGINE_API ADecalActor
 	: public AActor
 {
 	GENERATED_UCLASS_BODY()
@@ -23,13 +26,13 @@ private_subobject:
 	/** The decal component for this decal actor */
 	DEPRECATED_FORGAME(4.6, "Decal should not be accessed directly, please use GetDecal() function instead. Decal will soon be private and your code will not compile.")
 	UPROPERTY(Category = Decal, VisibleAnywhere, BlueprintReadOnly, meta = (ExposeFunctionCategories = "Decal,Rendering|Components|Decal", AllowPrivateAccess = "true"))
-	class UDecalComponent* Decal;
+	UDecalComponent* Decal;
 
 #if WITH_EDITORONLY_DATA
 	/* Reference to the editor only arrow visualization component */
 	DEPRECATED_FORGAME(4.6, "ArrowComponent should not be accessed directly, please use GetArrowComponent() function instead. ArrowComponent will soon be private and your code will not compile.")
 	UPROPERTY()
-	class UArrowComponent* ArrowComponent;
+	UArrowComponent* ArrowComponent;
 
 	/* Reference to the billboard component */
 	DEPRECATED_FORGAME(4.6, "SpriteComponent should not be accessed directly, please use GetSpriteComponent() function instead. SpriteComponent will soon be private and your code will not compile.")
@@ -69,13 +72,13 @@ public:
 public:
 
 	/** Returns Decal subobject **/
-	ENGINE_API class UDecalComponent* GetDecal() const;
+	UDecalComponent* GetDecal() const;
 #if WITH_EDITORONLY_DATA
 	/** Returns ArrowComponent subobject **/
-	ENGINE_API class UArrowComponent* GetArrowComponent() const;
+	UArrowComponent* GetArrowComponent() const;
 	/** Returns SpriteComponent subobject **/
-	ENGINE_API UBillboardComponent* GetSpriteComponent() const;
+	UBillboardComponent* GetSpriteComponent() const;
 	/** Returns BoxComponent subobject **/
-	ENGINE_API UBoxComponent* GetBoxComponent() const;
+	UBoxComponent* GetBoxComponent() const;
 #endif
 };
