Added SetSourceLength blueprint exposed function on point light component

diff --git a/Engine/Source/Runtime/Engine/Classes/Components/PointLightComponent.h b/Engine/Source/Runtime/Engine/Classes/Components/PointLightComponent.h
--- a/Engine/Source/Runtime/Engine/Classes/Components/PointLightComponent.h
+++ b/Engine/Source/Runtime/Engine/Classes/Components/PointLightComponent.h
@@ -65,7 +65,10 @@ class ENGINE_API UPointLightComponent : public ULightComponent
 	void SetLightFalloffExponent(float NewLightFalloffExponent);
 
 	UFUNCTION(BlueprintCallable, Category="Rendering|Lighting")
-	void SetSourceRadius(float bNewValue);
+	void SetSourceRadius(float NewValue);
+
+	UFUNCTION(BlueprintCallable, Category="Rendering|Lighting")
+	void SetSourceLength(float NewValue);
 
 protected:
 	// Begin UActorComponent Interface
diff --git a/Engine/Source/Runtime/Engine/Private/Components/PointLightComponent.cpp b/Engine/Source/Runtime/Engine/Private/Components/PointLightComponent.cpp
--- a/Engine/Source/Runtime/Engine/Private/Components/PointLightComponent.cpp
+++ b/Engine/Source/Runtime/Engine/Private/Components/PointLightComponent.cpp
@@ -145,12 +145,22 @@ void UPointLightComponent::SetLightFalloffExponent(float NewLightFalloffExponent
 	}
 }
 
-void UPointLightComponent::SetSourceRadius(float bNewValue)
+void UPointLightComponent::SetSourceRadius(float NewValue)
 {
 	if (AreDynamicDataChangesAllowed()
-		&& SourceRadius != bNewValue)
+		&& SourceRadius != NewValue)
 	{
-		SourceRadius = bNewValue;
+		SourceRadius = NewValue;
+		MarkRenderStateDirty();
+	}
+}
+
+void UPointLightComponent::SetSourceLength(float NewValue)
+{
+	if (AreDynamicDataChangesAllowed()
+		&& SourceLength != NewValue)
+	{
+		SourceLength = NewValue;
 		MarkRenderStateDirty();
 	}
 }
