Reverted SetSourceRadius parameter name change

diff --git a/Engine/Source/Runtime/Engine/Classes/Components/PointLightComponent.h b/Engine/Source/Runtime/Engine/Classes/Components/PointLightComponent.h
--- a/Engine/Source/Runtime/Engine/Classes/Components/PointLightComponent.h
+++ b/Engine/Source/Runtime/Engine/Classes/Components/PointLightComponent.h
@@ -65,7 +65,7 @@ class ENGINE_API UPointLightComponent : public ULightComponent
 	void SetLightFalloffExponent(float NewLightFalloffExponent);
 
 	UFUNCTION(BlueprintCallable, Category="Rendering|Lighting")
-	void SetSourceRadius(float NewValue);
+	void SetSourceRadius(float bNewValue);
 
 	UFUNCTION(BlueprintCallable, Category="Rendering|Lighting")
 	void SetSourceLength(float NewValue);
