Added a configurable limit to number of navmesh tiles #UE4

with a little preview of PolyRef bits map
UE-14926

diff --git a/Engine/Source/Runtime/Engine/Classes/AI/Navigation/RecastNavMesh.h b/Engine/Source/Runtime/Engine/Classes/AI/Navigation/RecastNavMesh.h
--- a/Engine/Source/Runtime/Engine/Classes/AI/Navigation/RecastNavMesh.h
+++ b/Engine/Source/Runtime/Engine/Classes/AI/Navigation/RecastNavMesh.h
@@ -538,7 +538,23 @@ class ENGINE_API ARecastNavMesh : public ANavigationData
 
 	UPROPERTY(EditAnywhere, Category = Generation, config, meta = (ClampMin = "0", UIMin = "0"), AdvancedDisplay)
 	int32 MaxSimultaneousTileGenerationJobsCount;
-	
+
+	/** Absolute hard limit to number of navmesh tiles. Be very, very careful while modifying it while
+	 *	having big maps with navmesh. A single, empty tile takes 176 bytes and empty tiles are
+	 *	allocated up front (subject to change, but that's where it's at now)
+	 *	@note TileNumberHardLimit is always rounded up to the closest power of 2 */
+	UPROPERTY(EditAnywhere, Category = Generation, config, meta = (ClampMin = "1", UIMin = "1"), AdvancedDisplay)
+	int32 TileNumberHardLimit;
+
+	UPROPERTY(VisibleAnywhere, Category = "Generation|Poly Ref Bits", AdvancedDisplay)
+	int32 PolyRefTileBits;
+
+	UPROPERTY(VisibleAnywhere, Category = Generation, AdvancedDisplay)
+	int32 PolyRefNavPolyBits;
+
+	UPROPERTY(VisibleAnywhere, Category = Generation, AdvancedDisplay)
+	int32 PolyRefSaltBits;
+
 	/** navmesh draw distance in game (always visible in editor) */
 	UPROPERTY(config)
 	float DefaultDrawDistance;
@@ -931,6 +947,7 @@ public:
 	virtual bool SupportsStreaming() const override;
 	virtual void ConditionalConstructGenerator() override;
 	bool ShouldGatherDataOnGameThread() const { return bDoFullyAsyncNavDataGathering == false; }
+	int32 GetTileNumberHardLimit() const { return TileNumberHardLimit; }
 
 	void UpdateActiveTiles(const TArray<FNavigationInvokerRaw>& InvokerLocations);
 	void RemoveTiles(const TArray<FIntPoint>& Tiles);
@@ -938,6 +955,8 @@ public:
 
 protected:
 
+	void UpdatePolyRefBitsPreview();
+
 	/** Returns query extent including adjustments for voxelization error compensation */
 	FVector GetModifiedQueryExtent(const FVector& QueryExtent) const;
 
diff --git a/Engine/Source/Runtime/Engine/Private/AI/Navigation/RecastNavMesh.cpp b/Engine/Source/Runtime/Engine/Private/AI/Navigation/RecastNavMesh.cpp
--- a/Engine/Source/Runtime/Engine/Private/AI/Navigation/RecastNavMesh.cpp
+++ b/Engine/Source/Runtime/Engine/Private/AI/Navigation/RecastNavMesh.cpp
@@ -234,6 +234,7 @@ ARecastNavMesh::ARecastNavMesh(const FObjectInitializer& ObjectInitializer)
 	LayerChunkSplits = 2;
 	MaxSimultaneousTileGenerationJobsCount = 1024;
 	bDoFullyAsyncNavDataGathering = false;
+	TileNumberHardLimit = 1 << 20;
 
 #if RECAST_ASYNC_REBUILDING
 	BatchQueryCounter = 0;
@@ -480,6 +481,7 @@ void ARecastNavMesh::PostLoad()
 	Super::PostLoad();
 	// tilesize validation. This is temporary and should get removed by 4.9
 	TileSizeUU = FMath::Clamp(TileSizeUU, CellSize, ArbitraryMaxVoxelTileSize * CellSize);
+	UpdatePolyRefBitsPreview();
 }
 
 void ARecastNavMesh::PostInitProperties()
@@ -563,6 +565,8 @@ void ARecastNavMesh::PostInitProperties()
 			AgentMaxStepHeight = DefOb->AgentMaxStepHeight;
 		}
 	}
+
+	UpdatePolyRefBitsPreview();
 }
 
 FVector ARecastNavMesh::GetModifiedQueryExtent(const FVector& QueryExtent) const
@@ -1959,6 +1963,7 @@ void ARecastNavMesh::PostEditChangeProperty(FPropertyChangedEvent& PropertyChang
 {
 	static const FName NAME_Generation = FName(TEXT("Generation"));
 	static const FName NAME_Display = FName(TEXT("Display"));
+	static const FName NAME_TileNumberHardLimit = GET_MEMBER_NAME_CHECKED(ARecastNavMesh, TileNumberHardLimit);
 
 	Super::PostEditChangeProperty(PropertyChangedEvent);
 
@@ -1985,6 +1990,11 @@ void ARecastNavMesh::PostEditChangeProperty(FPropertyChangedEvent& PropertyChang
 				// update config
 				FillConfig(NavDataConfig);
 			}
+			else if (PropName == NAME_TileNumberHardLimit)
+			{
+				TileNumberHardLimit = 1 << (FMath::CeilToInt(FMath::Log2(TileNumberHardLimit)));
+				UpdatePolyRefBitsPreview();
+			}
 
 			if (HasAnyFlags(RF_ClassDefaultObject) == false)
 			{
@@ -2082,6 +2092,14 @@ bool ARecastNavMesh::HasValidNavmesh() const
 #endif // WITH_RECAST
 }
 
+void ARecastNavMesh::UpdatePolyRefBitsPreview()
+{
+	static const int32 TotalBits = (sizeof(dtPolyRef) * 8);
+
+	FRecastNavMeshGenerator::CalcPolyRefBits(this, PolyRefTileBits, PolyRefNavPolyBits);
+	PolyRefSaltBits = TotalBits - PolyRefTileBits - PolyRefNavPolyBits;
+}
+
 //----------------------------------------------------------------------//
 // RecastNavMesh: Active Tiles 
 //----------------------------------------------------------------------//
diff --git a/Engine/Source/Runtime/Engine/Private/AI/Navigation/RecastNavMeshGenerator.cpp b/Engine/Source/Runtime/Engine/Private/AI/Navigation/RecastNavMeshGenerator.cpp
--- a/Engine/Source/Runtime/Engine/Private/AI/Navigation/RecastNavMeshGenerator.cpp
+++ b/Engine/Source/Runtime/Engine/Private/AI/Navigation/RecastNavMeshGenerator.cpp
@@ -3035,7 +3035,7 @@ FRecastNavMeshGenerator::FRecastNavMeshGenerator(ARecastNavMesh& InDestNavMesh)
 			if (!bRecreateNavmesh)
 			{
 				CalcNavMeshProperties(MaxTiles, MaxPolysPerTile);
-				if (MaxTiles > SavedNavParams->maxTiles)
+				if (FMath::Log2(MaxTiles) != FMath::Log2(SavedNavParams->maxTiles))
 				{
 					bRecreateNavmesh = true;
 				}
@@ -3211,15 +3211,26 @@ bool FRecastNavMeshGenerator::ConstructTiledNavMesh()
 	return bSuccess;
 }
 
-void FRecastNavMeshGenerator::CalcNavMeshProperties(int32& MaxTiles, int32& MaxPolys)
+void FRecastNavMeshGenerator::CalcPolyRefBits(ARecastNavMesh* NavMeshOwner, int32& MaxTileBits, int32& MaxPolyBits)
 {
-	// limit max amount of tiles
+	static const int32 TotalBits = (sizeof(dtPolyRef) * 8);
 #if USE_64BIT_ADDRESS
-	const int32 MaxTileBits = 30;
+	MaxTileBits = NavMeshOwner ? FMath::CeilToFloat(FMath::Log2(NavMeshOwner->GetTileNumberHardLimit())) : 20;
+	MaxPolyBits = FMath::Min<int32>(32, (TotalBits - DT_MIN_SALT_BITS) - MaxTileBits);
 #else
-	const int32 MaxTileBits = 14;
+	MaxTileBits = 14;
+	MaxPolyBits = (TotalBits - DT_MIN_SALT_BITS) - MaxTileBits;
 #endif//USE_64BIT_ADDRESS
+}
+
+void FRecastNavMeshGenerator::CalcNavMeshProperties(int32& MaxTiles, int32& MaxPolys)
+{
+	int32 MaxTileBits = -1;
+	int32 MaxPolyBits = -1;
 
+	// limit max amount of tiles
+	CalcPolyRefBits(DestNavMesh, MaxTileBits, MaxPolyBits);
+	
 	const int32 MaxTilesFromMask = (1 << MaxTileBits);
 	int32 MaxRequestedTiles = 0;
 	if (DestNavMesh->IsResizable())
@@ -3240,7 +3251,11 @@ void FRecastNavMeshGenerator::CalcNavMeshProperties(int32& MaxTiles, int32& MaxP
 	// Max tiles and max polys affect how the tile IDs are calculated.
 	// There are (sizeof(dtPolyRef)*8 - DT_MIN_SALT_BITS) bits available for 
 	// identifying a tile and a polygon.
-	MaxPolys = 1 << ((sizeof(dtPolyRef)* 8 - DT_MIN_SALT_BITS) - MaxTileBits);
+#if USE_64BIT_ADDRESS
+	MaxPolys = (MaxPolyBits >= 32) ? INT_MAX : (1 << MaxPolyBits);
+#else
+	MaxPolys = 1 << ((sizeof(dtPolyRef) * 8 - DT_MIN_SALT_BITS) - MaxTileBits);
+#endif // USE_64BIT_ADDRESS
 	MaxTiles = MaxRequestedTiles;
 }
 
diff --git a/Engine/Source/Runtime/Engine/Public/AI/Navigation/RecastNavMeshGenerator.h b/Engine/Source/Runtime/Engine/Public/AI/Navigation/RecastNavMeshGenerator.h
--- a/Engine/Source/Runtime/Engine/Public/AI/Navigation/RecastNavMeshGenerator.h
+++ b/Engine/Source/Runtime/Engine/Public/AI/Navigation/RecastNavMeshGenerator.h
@@ -458,6 +458,8 @@ public:
 	 *	@mote function does not validate the input parameter - it's on caller */
 	void SetMaxTileGeneratorTasks(int32 NewLimit) { MaxTileGeneratorTasks = NewLimit; }
 
+	static void CalcPolyRefBits(ARecastNavMesh* NavMeshOwner, int32& MaxTileBits, int32& MaxPolyBits);
+
 protected:
 	bool IsInActiveSet(const FIntPoint& Tile) const;
 	void RestrictBuildingToActiveTiles(bool InRestrictBuildingToActiveTiles);
