UE-15026 Supporting soloing audio device instances so audio can be heard from different PIE perspectives

diff --git a/Engine/Source/Runtime/Engine/Private/AudioDevice.cpp b/Engine/Source/Runtime/Engine/Private/AudioDevice.cpp
--- a/Engine/Source/Runtime/Engine/Private/AudioDevice.cpp
+++ b/Engine/Source/Runtime/Engine/Private/AudioDevice.cpp
@@ -732,6 +732,27 @@ bool FAudioDevice::HandleEnableHRTFForAllCommand(const TCHAR* Cmd, FOutputDevice
 	return true;
 }
 
+bool FAudioDevice::HandleSoloCommand(const TCHAR* Cmd, FOutputDevice& Ar)
+{
+	// Apply the solo to the given device
+	FAudioDeviceManager* DeviceManager = GEngine->GetAudioDeviceManager();
+	if (DeviceManager)
+	{
+		DeviceManager->SetSoloDevice(DeviceHandle);
+	}
+	return true;
+}
+
+bool FAudioDevice::HandleClearSoloCommand(const TCHAR* Cmd, FOutputDevice& Ar)
+{
+	FAudioDeviceManager* DeviceManager = GEngine->GetAudioDeviceManager();
+	if (DeviceManager)
+	{
+		DeviceManager->SetSoloDevice(INDEX_NONE);
+	}
+	return true;
+}
+
 #endif // !UE_BUILD_SHIPPING
 
 EDebugState FAudioDevice::GetMixDebugState( void )
@@ -830,6 +851,14 @@ bool FAudioDevice::Exec( UWorld* InWorld, const TCHAR* Cmd, FOutputDevice& Ar )
 	{
 		return HandleEnableHRTFForAllCommand(Cmd, Ar);
 	}
+	else if (FParse::Command(&Cmd, TEXT("Solo")))
+	{
+		return HandleSoloCommand(Cmd, Ar);
+	}
+	else if (FParse::Command(&Cmd, TEXT("ClearSolo")))
+	{
+		return HandleClearSoloCommand(Cmd, Ar);
+	}
 #endif // !UE_BUILD_SHIPPING
 
 	return false;
diff --git a/Engine/Source/Runtime/Engine/Private/AudioDeviceManager.cpp b/Engine/Source/Runtime/Engine/Private/AudioDeviceManager.cpp
--- a/Engine/Source/Runtime/Engine/Private/AudioDeviceManager.cpp
+++ b/Engine/Source/Runtime/Engine/Private/AudioDeviceManager.cpp
@@ -32,8 +32,8 @@ FAudioDeviceManager::FAudioDeviceManager()
 	, FreeIndicesSize(0)
 	, NumActiveAudioDevices(0)
 	, NumWorldsUsingMainAudioDevice(0)
-	, InFocusDeviceHandle(INDEX_NONE)
 	, NextResourceID(1)
+	, SoloDeviceHandle(INDEX_NONE)
 {
 }
 
@@ -327,22 +327,49 @@ void FAudioDeviceManager::InitSoundClasses()
 
 void FAudioDeviceManager::SetActiveDevice(uint32 InAudioDeviceHandle)
 {
-	for (FAudioDevice* AudioDevice : Devices)
+	// Only change the active device if there are no solo'd audio devices
+	if (SoloDeviceHandle == INDEX_NONE)
 	{
-		if (AudioDevice)
+		for (FAudioDevice* AudioDevice : Devices)
 		{
-			if (AudioDevice->DeviceHandle == InAudioDeviceHandle)
+			if (AudioDevice)
 			{
-				AudioDevice->bIsDeviceMuted = false;
+				if (AudioDevice->DeviceHandle == InAudioDeviceHandle)
+				{
+					AudioDevice->bIsDeviceMuted = false;
+				}
+				else
+				{
+					AudioDevice->bIsDeviceMuted = true;
+				}
 			}
-			else
+		}
+	}
+}
+
+void FAudioDeviceManager::SetSoloDevice(uint32 InAudioDeviceHandle)
+{
+	SoloDeviceHandle = InAudioDeviceHandle;
+	if (SoloDeviceHandle != INDEX_NONE)
+	{
+		for (FAudioDevice* AudioDevice : Devices)
+		{
+			if (AudioDevice)
 			{
-				AudioDevice->bIsDeviceMuted = true;
+				if (AudioDevice->DeviceHandle == InAudioDeviceHandle)
+				{
+					AudioDevice->bIsDeviceMuted = false;
+				}
+				else
+				{
+					AudioDevice->bIsDeviceMuted = true;
+				}
 			}
 		}
 	}
 }
 
+
 uint8 FAudioDeviceManager::GetNumActiveAudioDevices() const
 {
 	return NumActiveAudioDevices;
diff --git a/Engine/Source/Runtime/Engine/Public/AudioDevice.h b/Engine/Source/Runtime/Engine/Public/AudioDevice.h
--- a/Engine/Source/Runtime/Engine/Public/AudioDevice.h
+++ b/Engine/Source/Runtime/Engine/Public/AudioDevice.h
@@ -210,6 +210,8 @@ public:
 	bool HandleResetSoundStateCommand( const TCHAR* Cmd, FOutputDevice& Ar );
 	bool HandleToggleSpatializationExtensionCommand( const TCHAR* Cmd, FOutputDevice& Ar );
 	bool HandleEnableHRTFForAllCommand(const TCHAR* Cmd, FOutputDevice& Ar);
+	bool HandleSoloCommand( const TCHAR* Cmd, FOutputDevice& Ar );
+	bool HandleClearSoloCommand( const TCHAR* Cmd, FOutputDevice& Ar );
 #endif
 
 	/**
diff --git a/Engine/Source/Runtime/Engine/Public/AudioDeviceManager.h b/Engine/Source/Runtime/Engine/Public/AudioDeviceManager.h
--- a/Engine/Source/Runtime/Engine/Public/AudioDeviceManager.h
+++ b/Engine/Source/Runtime/Engine/Public/AudioDeviceManager.h
@@ -84,6 +84,9 @@ public:
 	/** Sets which audio device is the active audio device. */
 	void SetActiveDevice(uint32 InAudioDeviceHandle);
 
+	/** Sets an audio device to be solo'd */
+	void SetSoloDevice(uint32 InAudioDeviceHandle);
+
 	/** Links up the resource data indices for looking up and cleaning up. */
 	void TrackResource(USoundWave* SoundWave, FSoundBuffer* Buffer);
 
@@ -149,11 +152,11 @@ private:
 	*/
 	TArray<FAudioDevice* > Devices;
 
-	/** The audio device handle currently in focus */
-	int32 InFocusDeviceHandle;
-
 	/** Next resource ID to assign out to a wave/buffer */
 	int32 NextResourceID;
+
+	/** Which audio device is solo'd */
+	uint32 SoloDeviceHandle;
 };
 
 
