UE4-15036 Fix for audio device initializing twice when launched without the editor

diff --git a/Engine/Source/Runtime/Engine/Private/AudioDeviceManager.cpp b/Engine/Source/Runtime/Engine/Private/AudioDeviceManager.cpp
--- a/Engine/Source/Runtime/Engine/Private/AudioDeviceManager.cpp
+++ b/Engine/Source/Runtime/Engine/Private/AudioDeviceManager.cpp
@@ -65,19 +65,20 @@ FAudioDevice* FAudioDeviceManager::CreateAudioDevice(uint32& HandleOut, bool bCr
 		return nullptr;
 	}
 
-#if !WITH_EDITOR
 	// If we are running without the editor, we only need one audio device.
-	if (NumActiveAudioDevices == 1)
+	if (!GIsEditor)
 	{
-		FAudioDevice* MainAudioDevice = GEngine->GetMainAudioDevice();
-		if (MainAudioDevice)
+		if (NumActiveAudioDevices == 1)
 		{
-			HandleOut = MainAudioDevice->DeviceHandle;
-			return MainAudioDevice;
+			FAudioDevice* MainAudioDevice = GEngine->GetMainAudioDevice();
+			if (MainAudioDevice)
+			{
+				HandleOut = MainAudioDevice->DeviceHandle;
+				return MainAudioDevice;
+			}
+			return nullptr;
 		}
-		return nullptr;
 	}
-#endif
 
 	FAudioDevice* NewAudioDevice = nullptr;
 
