UE-15026 Changing the command to "SoloAudio" and "ClearSoloAudio" to make it more clear that this is an audio command.

diff --git a/Engine/Source/Runtime/Engine/Private/AudioDevice.cpp b/Engine/Source/Runtime/Engine/Private/AudioDevice.cpp
--- a/Engine/Source/Runtime/Engine/Private/AudioDevice.cpp
+++ b/Engine/Source/Runtime/Engine/Private/AudioDevice.cpp
@@ -851,11 +851,11 @@ bool FAudioDevice::Exec( UWorld* InWorld, const TCHAR* Cmd, FOutputDevice& Ar )
 	{
 		return HandleEnableHRTFForAllCommand(Cmd, Ar);
 	}
-	else if (FParse::Command(&Cmd, TEXT("Solo")))
+	else if (FParse::Command(&Cmd, TEXT("SoloAudio")))
 	{
 		return HandleSoloCommand(Cmd, Ar);
 	}
-	else if (FParse::Command(&Cmd, TEXT("ClearSolo")))
+	else if (FParse::Command(&Cmd, TEXT("ClearSoloAudio")))
 	{
 		return HandleClearSoloCommand(Cmd, Ar);
 	}
