Fix FBodyInstance::InitBody not adding instances into aggregates. Fixes crash on moving dynamic ISMC.
Also remove some checks that require a scene lock

UE-14930

diff --git a/Engine/Source/Runtime/Engine/Private/InstancedStaticMesh.cpp b/Engine/Source/Runtime/Engine/Private/InstancedStaticMesh.cpp
--- a/Engine/Source/Runtime/Engine/Private/InstancedStaticMesh.cpp
+++ b/Engine/Source/Runtime/Engine/Private/InstancedStaticMesh.cpp
@@ -974,7 +974,6 @@ void UInstancedStaticMeshComponent::InitInstanceBody(int32 InstanceIdx, FBodyIns
 	// Create physics body instance.
 	// Aggregates aren't used for static objects
 	auto* Aggregate = (Mobility == EComponentMobility::Movable) ? Aggregates[FMath::DivideAndRoundDown<int32>(InstanceIdx, AggregateMaxSize)] : nullptr;
-	check(Mobility != EComponentMobility::Movable || Aggregate->getNbActors() < Aggregate->getMaxNbActors());
 	InstanceBodyInstance->bAutoWeld = false;	//We don't support this for instanced meshes.
 	InstanceBodyInstance->InitBody(BodySetup, InstanceTransform, this, GetWorld()->GetPhysicsScene(), Aggregate);
 #endif //WITH_PHYSX
@@ -1102,7 +1101,6 @@ void UInstancedStaticMeshComponent::DestroyPhysicsState()
 	// releasing Aggregates, they shouldn't contain any Bodies now, because they are released above
 	for (auto* Aggregate : Aggregates)
 	{
-		check(!Aggregate->getNbActors());
 		Aggregate->release();
 	}
 	Aggregates.Empty();
diff --git a/Engine/Source/Runtime/Engine/Private/PhysicsEngine/BodyInstance.cpp b/Engine/Source/Runtime/Engine/Private/PhysicsEngine/BodyInstance.cpp
--- a/Engine/Source/Runtime/Engine/Private/PhysicsEngine/BodyInstance.cpp
+++ b/Engine/Source/Runtime/Engine/Private/PhysicsEngine/BodyInstance.cpp
@@ -1259,8 +1259,9 @@ struct FInitBodiesHelper
 				if (!Instance->bEnableGravity)
 				{
 					PNewDynamic->setActorFlag(PxActorFlag::eDISABLE_GRAVITY, true);
-					bDynamicsUseAsyncScene = Instance->UseAsyncScene(PhysScene);
 				}
+
+				bDynamicsUseAsyncScene = Instance->UseAsyncScene(PhysScene);
 			}
 
 			if (bCompileStatic || bCanDefer)
@@ -1313,7 +1314,7 @@ struct FInitBodiesHelper
 		return true;
 	}
 
-	void AddActorsToScene_PhysX_AssumesLocked(TArray<PxActor*>& PSyncActors, TArray<PxActor*>& PAsyncActors, TArray<PxActor*>& PDynamicActors) const
+	void AddActorsToScene_PhysX_AssumesLocked(TArray<PxActor*>& PSyncActors, TArray<PxActor*>& PAsyncActors, TArray<PxActor*>& PDynamicActors, PxScene* PDyanmicScene) const
 	{
 		SCOPE_CYCLE_COUNTER(STAT_BulkSceneAdd);
 
@@ -1321,8 +1322,7 @@ struct FInitBodiesHelper
 		// Use the aggregate if it exists, has enough slots and is in the correct scene (or no scene)
 		if (InAggregate && PDynamicActors.Num() > 0 && (InAggregate->getMaxNbActors() - InAggregate->getNbActors()) >= (uint32)PDynamicActors.Num())
 		{
-
-			if (InAggregate->getScene() == nullptr || InAggregate->getScene() == PDynamicActors[0]->getScene())
+			if (InAggregate->getScene() == nullptr || InAggregate->getScene() == PDyanmicScene)
 			{
 				for (PxActor* Actor : PDynamicActors)
 				{
@@ -1381,7 +1381,7 @@ struct FInitBodiesHelper
 			SCOPED_SCENE_WRITE_LOCK(bAddingToSyncScene ? PSyncScene : nullptr);
 			SCOPED_SCENE_WRITE_LOCK(bAddingToAsyncScene ? PAsyncScene : nullptr);
 
-			AddActorsToScene_PhysX_AssumesLocked(PSyncActors, PAsyncActors, PDynamicActors);
+			AddActorsToScene_PhysX_AssumesLocked(PSyncActors, PAsyncActors, PDynamicActors, bDynamicsUseAsync ? PAsyncScene : PSyncScene);
 		}
 	}
 #endif
