Physx binary serialization properly releases shapes

UE-15083

diff --git a/Engine/Source/Runtime/Engine/Private/PhysicsEngine/PhysicsSerializer.cpp b/Engine/Source/Runtime/Engine/Private/PhysicsEngine/PhysicsSerializer.cpp
--- a/Engine/Source/Runtime/Engine/Private/PhysicsEngine/PhysicsSerializer.cpp
+++ b/Engine/Source/Runtime/Engine/Private/PhysicsEngine/PhysicsSerializer.cpp
@@ -138,6 +138,10 @@ void UPhysicsSerializer::CreatePhysicsData(const TArray<UBodySetup*>& BodySetups
 					const PxSerialObjectId ObjectId = PCollection->getId(PObject);
 					ActorsMap.Add(ObjectId, PRigidActor);
 				}
+				else if (PxShape* PShape = PObject.is<PxShape>())
+				{
+					PShape->release();	//we do not need to hold on to this reference because our actors always have the reference directly
+				}
 			}
 		}
 
