Messed around with WITH_RECAST macro to fix odd compilation issues for iOS

diff --git a/Engine/Source/Runtime/Engine/Private/AI/Navigation/RecastNavMesh.cpp b/Engine/Source/Runtime/Engine/Private/AI/Navigation/RecastNavMesh.cpp
--- a/Engine/Source/Runtime/Engine/Private/AI/Navigation/RecastNavMesh.cpp
+++ b/Engine/Source/Runtime/Engine/Private/AI/Navigation/RecastNavMesh.cpp
@@ -16,11 +16,12 @@
 
 #if WITH_RECAST
 #include "DetourAlloc.h"
-#endif 
+#endif // WITH_RECAST
 
 #if WITH_EDITOR
 #include "UnrealEd.h"
 #endif
+
 #include "AI/Navigation/NavMeshRenderingComponent.h"
 
 #if WITH_RECAST
@@ -45,7 +46,7 @@ FNavMeshTileData::FNavData::~FNavData()
 	dtFree(RawNavData);
 #else
 	FMemory::Free(RawNavData);
-#endif
+#endif // WITH_RECAST
 }
 
 FNavMeshTileData::FNavMeshTileData(uint8* RawData, int32 RawDataSize, int32 LayerIdx, FBox LayerBounds)
@@ -90,7 +91,7 @@ void FNavMeshTileData::MakeUnique()
 		uint8* UniqueRawData = (uint8*)dtAlloc(sizeof(uint8)*DataSize, DT_ALLOC_PERM);
 #else
 		uint8* UniqueRawData = (uint8*)FMemory::Malloc(sizeof(uint8)*DataSize);
-#endif
+#endif //WITH_RECAST
 		FMemory::Memcpy(UniqueRawData, NavData->RawNavData, DataSize);
 		NavData = MakeShareable(new FNavData(UniqueRawData));
 	}
@@ -574,6 +575,14 @@ FVector ARecastNavMesh::GetModifiedQueryExtent(const FVector& QueryExtent) const
 	return FVector(QueryExtent.X, QueryExtent.Y, QueryExtent.Z + FMath::Max(0.0f, VerticalDeviationFromGroundCompensation));
 }
 
+void ARecastNavMesh::UpdatePolyRefBitsPreview()
+{
+	static const int32 TotalBits = (sizeof(dtPolyRef) * 8);
+
+	FRecastNavMeshGenerator::CalcPolyRefBits(this, PolyRefTileBits, PolyRefNavPolyBits);
+	PolyRefSaltBits = TotalBits - PolyRefTileBits - PolyRefNavPolyBits;
+}
+
 void ARecastNavMesh::OnNavAreaAdded(const UClass* NavAreaClass, int32 AgentIndex)
 {
 	Super::OnNavAreaAdded(NavAreaClass, AgentIndex);
@@ -2080,7 +2089,6 @@ void ARecastNavMesh::UpdateNavObject()
 {
 	OnNavMeshUpdate.Broadcast();
 }
-
 #endif	//WITH_RECAST
 
 bool ARecastNavMesh::HasValidNavmesh() const
@@ -2092,14 +2100,6 @@ bool ARecastNavMesh::HasValidNavmesh() const
 #endif // WITH_RECAST
 }
 
-void ARecastNavMesh::UpdatePolyRefBitsPreview()
-{
-	static const int32 TotalBits = (sizeof(dtPolyRef) * 8);
-
-	FRecastNavMeshGenerator::CalcPolyRefBits(this, PolyRefTileBits, PolyRefNavPolyBits);
-	PolyRefSaltBits = TotalBits - PolyRefTileBits - PolyRefNavPolyBits;
-}
-
 //----------------------------------------------------------------------//
 // RecastNavMesh: Active Tiles 
 //----------------------------------------------------------------------//
@@ -2226,12 +2226,11 @@ void ARecastNavMesh::RebuildTile(const TArray<FIntPoint>& Tiles)
 		}
 	}
 }
-#endif
 
 //----------------------------------------------------------------------//
 // FRecastNavMeshCachedData
 //----------------------------------------------------------------------//
-#if WITH_RECAST
+
 FRecastNavMeshCachedData FRecastNavMeshCachedData::Construct(const ARecastNavMesh* RecastNavMeshActor)
 {
 	check(RecastNavMeshActor);
