Expose BP node to query skeletal mesh sections for cloth usage

diff --git a/Engine/Source/Runtime/Engine/Classes/Engine/SkeletalMesh.h b/Engine/Source/Runtime/Engine/Classes/Engine/SkeletalMesh.h
--- a/Engine/Source/Runtime/Engine/Classes/Engine/SkeletalMesh.h
+++ b/Engine/Source/Runtime/Engine/Classes/Engine/SkeletalMesh.h
@@ -843,6 +843,17 @@ public:
 	ENGINE_API void  LoadClothCollisionVolumes(int32 AssetIndex, physx::apex::NxClothingAsset* ClothingAsset);
 	ENGINE_API bool IsMappedClothingLOD(int32 LODIndex, int32 AssetIndex);
 	ENGINE_API int32 GetClothAssetIndex(int32 LODIndex, int32 SectionIndex);
+
+	/** 
+	 * Checks whether the provided section is using APEX cloth. if bCheckCorrespondingSections is true
+	 * disabled sections will defer to correspond sections to see if they use cloth (non-cloth sections
+	 * are disabled and another section added when cloth is enabled, using this flag allows for a check
+	 * on the original section to succeed)
+	 * @param InSectionIndex Index to check
+	 * @param bCheckCorrespondingSections Whether to check corresponding sections for disabled sections
+	 */
+	UFUNCTION(BlueprintCallable, Category="Cloth")
+	ENGINE_API bool IsSectionUsingCloth(int32 InSectionIndex, bool bCheckCorrespondingSections = true) const;
 #endif// #if WITH_APEX_CLOTHING
 
 	ENGINE_API void CreateBodySetup();
diff --git a/Engine/Source/Runtime/Engine/Private/SkeletalMesh.cpp b/Engine/Source/Runtime/Engine/Private/SkeletalMesh.cpp
--- a/Engine/Source/Runtime/Engine/Private/SkeletalMesh.cpp
+++ b/Engine/Source/Runtime/Engine/Private/SkeletalMesh.cpp
@@ -3493,6 +3493,36 @@ FString USkeletalMesh::GetDesc()
 	return FString::Printf( TEXT("%d Triangles, %d Bones"), Resource->LODModels[0].GetTotalFaces(), RefSkeleton.GetNum() );
 }
 
+#if WITH_APEX_CLOTHING
+bool USkeletalMesh::IsSectionUsingCloth(int32 InSectionIndex, bool bCheckCorrespondingSections) const
+{
+	if(ImportedResource.IsValid())
+	{
+		for(FStaticLODModel& LodModel : ImportedResource->LODModels)
+		{
+			if(LodModel.Sections.IsValidIndex(InSectionIndex))
+			{
+				FSkelMeshSection* SectionToCheck = &LodModel.Sections[InSectionIndex];
+				if(SectionToCheck->bDisabled && bCheckCorrespondingSections)
+				{
+					SectionToCheck = &LodModel.Sections[SectionToCheck->CorrespondClothSectionIndex];
+				}
+
+				int32 ChunkIdx = SectionToCheck->ChunkIndex;
+				check(LodModel.Chunks.IsValidIndex(ChunkIdx)); // Built on import, should be valid at all times
+				
+				FSkelMeshChunk& Chunk = LodModel.Chunks[ChunkIdx];
+				if(Chunk.HasApexClothData())
+				{
+					return true;
+				}
+			}
+		}
+	}
+
+	return false;
+}
+#endif
 
 /*-----------------------------------------------------------------------------
 USkeletalMeshSocket
