Fix for crash importing lighting with static meshes with LODs that have been vertex painted

diff --git a/Engine/Source/Runtime/Engine/Private/StaticMeshRender.cpp b/Engine/Source/Runtime/Engine/Private/StaticMeshRender.cpp
--- a/Engine/Source/Runtime/Engine/Private/StaticMeshRender.cpp
+++ b/Engine/Source/Runtime/Engine/Private/StaticMeshRender.cpp
@@ -146,6 +146,12 @@ void UStaticMeshComponent::SetLODDataCount( const uint32 MinSize, const uint32 M
 	check(MaxSize <= MAX_STATIC_MESH_LODS);
 	if (MaxSize < (uint32)LODData.Num())
 	{
+		// FStaticMeshComponentLODInfo can't be deleted directly as it has rendering resources
+		for (int32 Index = MaxSize; Index < LODData.Num(); Index++)
+		{
+			LODData[Index].ReleaseOverrideVertexColorsAndBlock();
+		}
+
 		// call destructors
 		LODData.RemoveAt(MaxSize, LODData.Num() - MaxSize);
 	}
