#ue4 - Refactored collision check in FCollisionQueryParams::AddIgnoredActor() and AddIgnoredComponents() to a static function.

- Allow it to filter out PhysicOnly collision as well.
- Use the virtual GetCollisionEnabled() in case it's overridden. Was inconsistent before (one used bodyinstance directly, other used function).

diff --git a/Engine/Source/Runtime/Engine/Private/Collision/Collision.cpp b/Engine/Source/Runtime/Engine/Private/Collision/Collision.cpp
--- a/Engine/Source/Runtime/Engine/Private/Collision/Collision.cpp
+++ b/Engine/Source/Runtime/Engine/Private/Collision/Collision.cpp
@@ -162,6 +162,16 @@ FCollisionQueryParams::FCollisionQueryParams(FName InTraceTag, bool bInTraceComp
 	}
 }
 
+
+static FORCEINLINE_DEBUGGABLE bool IsQueryCollisionEnabled(const UPrimitiveComponent* PrimComponent)
+{
+	const ECollisionEnabled::Type CollisionEnabled = PrimComponent->GetCollisionEnabled();
+
+	return CollisionEnabled == ECollisionEnabled::QueryAndPhysics
+		|| CollisionEnabled == ECollisionEnabled::QueryOnly;
+}
+
+
 void FCollisionQueryParams::AddIgnoredActor(const AActor* InIgnoreActor)
 {
 	// Note: This code operates differently in the editor because the actors in the editor can have their collision setting become out of sync with physx.  This happens because even when collision is disabled on an actor, in the editor we 
@@ -181,8 +191,7 @@ void FCollisionQueryParams::AddIgnoredActor(const AActor* InIgnoreActor)
 			for (const UActorComponent* ActorComponent : InIgnoreActor->GetComponents())
 			{
 				const UPrimitiveComponent* PrimComponent = Cast<const UPrimitiveComponent>(ActorComponent);
-				
-				if (PrimComponent && (!bCheckForCollision || PrimComponent->BodyInstance.GetCollisionEnabled() != ECollisionEnabled::NoCollision) )
+				if (PrimComponent && (!bCheckForCollision || IsQueryCollisionEnabled(PrimComponent)) )
 				{
 					IgnoreComponents.Add(PrimComponent->GetUniqueID());
 				}
@@ -193,8 +202,7 @@ void FCollisionQueryParams::AddIgnoredActor(const AActor* InIgnoreActor)
 			for (const UActorComponent* ActorComponent : InIgnoreActor->GetComponents())
 			{
 				const UPrimitiveComponent* PrimComponent = Cast<const UPrimitiveComponent>(ActorComponent);
-
-				if (PrimComponent && (!bCheckForCollision || PrimComponent->BodyInstance.GetCollisionEnabled() != ECollisionEnabled::NoCollision) )
+				if (PrimComponent && (!bCheckForCollision || IsQueryCollisionEnabled(PrimComponent)) )
 				{
 					IgnoreComponents.AddUnique(PrimComponent->GetUniqueID());
 				}
@@ -237,7 +245,7 @@ void FCollisionQueryParams::AddIgnoredComponent(const UPrimitiveComponent* InIgn
 	const bool bCheckForCollision = true;
 #endif
 
-	if (InIgnoreComponent && (!bCheckForCollision || InIgnoreComponent->GetCollisionEnabled() != ECollisionEnabled::NoCollision))
+	if (InIgnoreComponent && (!bCheckForCollision || IsQueryCollisionEnabled(InIgnoreComponent)))
 	{
 		IgnoreComponents.AddUnique(InIgnoreComponent->GetUniqueID());
 	}
