Default implementation of USoundNode::CreateStartingConnectors will now make at least GetMinChildNodes() connections.

diff --git a/Engine/Source/Runtime/Engine/Classes/Sound/SoundNodeBranch.h b/Engine/Source/Runtime/Engine/Classes/Sound/SoundNodeBranch.h
--- a/Engine/Source/Runtime/Engine/Classes/Sound/SoundNodeBranch.h
+++ b/Engine/Source/Runtime/Engine/Classes/Sound/SoundNodeBranch.h
@@ -37,7 +37,6 @@ public:
 	{ 
 		return BranchPurpose::MAX;
 	}
-	virtual void CreateStartingConnectors( void ) override;
 
 	virtual void RemoveChildNode( int32 Index ) override
 	{
diff --git a/Engine/Source/Runtime/Engine/Private/SoundNode.cpp b/Engine/Source/Runtime/Engine/Private/SoundNode.cpp
--- a/Engine/Source/Runtime/Engine/Private/SoundNode.cpp
+++ b/Engine/Source/Runtime/Engine/Private/SoundNode.cpp
@@ -82,7 +82,12 @@ void USoundNode::GetAllNodes( TArray<USoundNode*>& SoundNodes )
 
 void USoundNode::CreateStartingConnectors()
 {
-	InsertChildNode( ChildNodes.Num() );
+	int32 ConnectorsToMake = FMath::Max(1, GetMinChildNodes());
+	while (ConnectorsToMake > 0)
+	{
+		InsertChildNode( ChildNodes.Num() );
+		--ConnectorsToMake;
+	}
 }
 
 void USoundNode::InsertChildNode( int32 Index )
diff --git a/Engine/Source/Runtime/Engine/Private/SoundNodeBranch.cpp b/Engine/Source/Runtime/Engine/Private/SoundNodeBranch.cpp
--- a/Engine/Source/Runtime/Engine/Private/SoundNodeBranch.cpp
+++ b/Engine/Source/Runtime/Engine/Private/SoundNodeBranch.cpp
@@ -33,14 +33,6 @@ void USoundNodeBranch::ParseNodes( FAudioDevice* AudioDevice, const UPTRINT Node
 	}
 }
 
-void USoundNodeBranch::CreateStartingConnectors()
-{
-	// Locality Nodes default with two connectors - Locally instigated and .
-	InsertChildNode( ChildNodes.Num() );
-	InsertChildNode( ChildNodes.Num() );
-	InsertChildNode( ChildNodes.Num() );
-}
-
 #if WITH_EDITOR
 FText USoundNodeBranch::GetInputPinName(int32 PinIndex) const
 {
