UProperty and UFunction metadata for Category can now be localized and have improved FText usage throughout the editor (Blueprints, Behavior Trees, Materials, etc).

Big conversion of FStrings and FNames to FText.

Version bump to move some MaterialFunction UProperty from a TArray<FString> to TArray<FText> (some Engine assets are older than being allowed to auto-convert the UProperty)

Auto conversion of FName to FText (and back) now supported (as well as TArrays of those types).

Searching categories by both the localized string and the source string is now supported in Blueprints.

diff --git a/Engine/Config/BaseEditor.ini b/Engine/Config/BaseEditor.ini
--- a/Engine/Config/BaseEditor.ini
+++ b/Engine/Config/BaseEditor.ini
@@ -2,6 +2,7 @@
 +LocalizationPaths=../../../Engine/Content/Localization/Editor
 +LocalizationPaths=../../../Engine/Content/Localization/EditorTutorials
 +LocalizationPaths=../../../Engine/Content/Localization/Keywords
++LocalizationPaths=../../../Engine/Content/Localization/Category
 +ToolTipLocalizationPaths=../../../Engine/Content/Localization/ToolTips
 +PropertyNameLocalizationPaths=../../../Engine/Content/Localization/PropertyNames
 ShouldLoadLocalizedPropertyNames=True
diff --git a/Engine/Config/Localization/Category.ini b/Engine/Config/Localization/Category.ini
new file mode 100644
--- /dev/null
+++ b/Engine/Config/Localization/Category.ini
@@ -0,0 +1,58 @@
+﻿[CommonSettings]
+SourcePath=./Content/Localization/Category
+DestinationPath=./Content/Localization/Category
+ManifestName=Category.manifest
+ArchiveName=Category.archive
+PortableObjectName=Category.po
+SourceCulture=en
+CulturesToGenerate=es
+CulturesToGenerate=en
+CulturesToGenerate=ja
+CulturesToGenerate=ko
+CulturesToGenerate=pt-BR
+CulturesToGenerate=zh-CN
+
+;Gather text from metadata
+[GatherTextStep0]
+CommandletClass=GatherTextFromMetaData
+IncludePathFilters=Source/Editor/*
+IncludePathFilters=Source/Runtime/*
+IncludePathFilters=Source/Developer/*
+ExcludePathFilters=Config/Localization/*
+ExcludePathFilters=Source/Developer/NoRedist/CommunityPortalServices/*
+InputKeys=Category
+OutputNamespaces=UObjectCategory
+OutputKeys="{FieldPath}"
+
+;Write Manifest
+[GatherTextStep1]
+CommandletClass=GenerateGatherManifest
+
+;Write Archives
+[GatherTextStep2]
+CommandletClass=GenerateGatherArchive
+bPurgeOldEmptyEntries=true
+
+;Import localized PO files
+[GatherTextStep3]
+CommandletClass=InternationalizationExport
+bImportLoc=true
+
+;Write Localized Text Resource
+[GatherTextStep4]
+CommandletClass=GenerateTextLocalizationResource
+ResourceName=Category.locres
+
+;Export to PO files
+[GatherTextStep5]
+CommandletClass=InternationalizationExport
+bExportLoc=true
+
+;;Write Text Localization Report
+;[GatherTextStep6]
+;CommandletClass=GenerateTextLocalizationReport
+;DestinationPath=\\epicgames.net\root\UE3\Localization\WordCounts
+;bWordCountReport=true
+;WordCountReportName=Category.csv
+;bConflictReport=true
+;ConflictReportName=CategoryConflicts.txt
\ No newline at end of file
diff --git a/Engine/Source/Developer/MovieSceneTools/Classes/K2Node_PlayMovieScene.h b/Engine/Source/Developer/MovieSceneTools/Classes/K2Node_PlayMovieScene.h
--- a/Engine/Source/Developer/MovieSceneTools/Classes/K2Node_PlayMovieScene.h
+++ b/Engine/Source/Developer/MovieSceneTools/Classes/K2Node_PlayMovieScene.h
@@ -21,7 +21,7 @@ struct FEdGraphSchemaAction_K2AddPlayMovieScene : public FEdGraphSchemaAction_K2
 	{
 	}
 	
-	FEdGraphSchemaAction_K2AddPlayMovieScene(const FString& InNodeCategory, const FText& InMenuDesc, const FString& InToolTip, const int32 InGrouping)
+	FEdGraphSchemaAction_K2AddPlayMovieScene(const FText& InNodeCategory, const FText& InMenuDesc, const FString& InToolTip, const int32 InGrouping)
 		: FEdGraphSchemaAction_K2NewNode(InNodeCategory, InMenuDesc, InToolTip, InGrouping)
 	{
 	}
diff --git a/Engine/Source/Editor/AIGraph/Classes/AIGraphSchema.h b/Engine/Source/Editor/AIGraph/Classes/AIGraphSchema.h
--- a/Engine/Source/Editor/AIGraph/Classes/AIGraphSchema.h
+++ b/Engine/Source/Editor/AIGraph/Classes/AIGraphSchema.h
@@ -20,7 +20,7 @@ struct AIGRAPH_API FAISchemaAction_NewNode : public FEdGraphSchemaAction
 		, NodeTemplate(nullptr)
 	{}
 
-	FAISchemaAction_NewNode(const FString& InNodeCategory, const FText& InMenuDesc, const FString& InToolTip, const int32 InGrouping)
+	FAISchemaAction_NewNode(const FText& InNodeCategory, const FText& InMenuDesc, const FString& InToolTip, const int32 InGrouping)
 		: FEdGraphSchemaAction(InNodeCategory, InMenuDesc, InToolTip, InGrouping)
 		, NodeTemplate(nullptr)
 	{}
@@ -61,7 +61,7 @@ struct AIGRAPH_API FAISchemaAction_NewSubNode : public FEdGraphSchemaAction
 		, ParentNode(nullptr)
 	{}
 
-	FAISchemaAction_NewSubNode(const FString& InNodeCategory, const FText& InMenuDesc, const FString& InToolTip, const int32 InGrouping)
+	FAISchemaAction_NewSubNode(const FText& InNodeCategory, const FText& InMenuDesc, const FString& InToolTip, const int32 InGrouping)
 		: FEdGraphSchemaAction(InNodeCategory, InMenuDesc, InToolTip, InGrouping)
 		, NodeTemplate(nullptr)
 		, ParentNode(nullptr)
@@ -96,7 +96,7 @@ class AIGRAPH_API UAIGraphSchema : public UEdGraphSchema
 
 protected:
 
-	static TSharedPtr<FAISchemaAction_NewNode> AddNewNodeAction(FGraphActionListBuilderBase& ContextMenuBuilder, const FString& Category, const FText& MenuDesc, const FString& Tooltip);
-	static TSharedPtr<FAISchemaAction_NewSubNode> AddNewSubNodeAction(FGraphActionListBuilderBase& ContextMenuBuilder, const FString& Category, const FText& MenuDesc, const FString& Tooltip);
+	static TSharedPtr<FAISchemaAction_NewNode> AddNewNodeAction(FGraphActionListBuilderBase& ContextMenuBuilder, const FText& Category, const FText& MenuDesc, const FString& Tooltip);
+	static TSharedPtr<FAISchemaAction_NewSubNode> AddNewSubNodeAction(FGraphActionListBuilderBase& ContextMenuBuilder, const FText& Category, const FText& MenuDesc, const FString& Tooltip);
 };
 
diff --git a/Engine/Source/Editor/AIGraph/Classes/AIGraphTypes.h b/Engine/Source/Editor/AIGraph/Classes/AIGraphTypes.h
--- a/Engine/Source/Editor/AIGraph/Classes/AIGraphTypes.h
+++ b/Engine/Source/Editor/AIGraph/Classes/AIGraphTypes.h
@@ -14,7 +14,7 @@ struct AIGRAPH_API FGraphNodeClassData
 
 	FString ToString() const;
 	FString GetClassName() const;
-	FString GetCategory() const;
+	FText GetCategory() const;
 	FString GetDisplayName() const;
 	UClass* GetClass(bool bSilent = false);
 	bool IsAbstract() const;
@@ -48,7 +48,7 @@ private:
 
 	/** User-defined category for this class */
 	UPROPERTY()
-	FString Category;
+	FText Category;
 
 	/** message for deprecated class */
 	FString DeprecatedMessage;
diff --git a/Engine/Source/Editor/AIGraph/Private/AIGraphSchema.cpp b/Engine/Source/Editor/AIGraph/Private/AIGraphSchema.cpp
--- a/Engine/Source/Editor/AIGraph/Private/AIGraphSchema.cpp
+++ b/Engine/Source/Editor/AIGraph/Private/AIGraphSchema.cpp
@@ -121,7 +121,7 @@ UAIGraphSchema::UAIGraphSchema(const FObjectInitializer& ObjectInitializer) : Su
 {
 }
 
-TSharedPtr<FAISchemaAction_NewNode> UAIGraphSchema::AddNewNodeAction(FGraphActionListBuilderBase& ContextMenuBuilder, const FString& Category, const FText& MenuDesc, const FString& Tooltip)
+TSharedPtr<FAISchemaAction_NewNode> UAIGraphSchema::AddNewNodeAction(FGraphActionListBuilderBase& ContextMenuBuilder, const FText& Category, const FText& MenuDesc, const FString& Tooltip)
 {
 	TSharedPtr<FAISchemaAction_NewNode> NewAction = TSharedPtr<FAISchemaAction_NewNode>(new FAISchemaAction_NewNode(Category, MenuDesc, Tooltip, 0));
 	ContextMenuBuilder.AddAction(NewAction);
@@ -129,7 +129,7 @@ TSharedPtr<FAISchemaAction_NewNode> UAIGraphSchema::AddNewNodeAction(FGraphActio
 	return NewAction;
 }
 
-TSharedPtr<FAISchemaAction_NewSubNode> UAIGraphSchema::AddNewSubNodeAction(FGraphActionListBuilderBase& ContextMenuBuilder, const FString& Category, const FText& MenuDesc, const FString& Tooltip)
+TSharedPtr<FAISchemaAction_NewSubNode> UAIGraphSchema::AddNewSubNodeAction(FGraphActionListBuilderBase& ContextMenuBuilder, const FText& Category, const FText& MenuDesc, const FString& Tooltip)
 {
 	TSharedPtr<FAISchemaAction_NewSubNode> NewAction = TSharedPtr<FAISchemaAction_NewSubNode>(new FAISchemaAction_NewSubNode(Category, MenuDesc, Tooltip, 0));
 	ContextMenuBuilder.AddAction(NewAction);
diff --git a/Engine/Source/Editor/AIGraph/Private/AIGraphTypes.cpp b/Engine/Source/Editor/AIGraph/Private/AIGraphTypes.cpp
--- a/Engine/Source/Editor/AIGraph/Private/AIGraphTypes.cpp
+++ b/Engine/Source/Editor/AIGraph/Private/AIGraphTypes.cpp
@@ -68,9 +68,9 @@ FString FGraphNodeClassData::GetDisplayName() const
 	return Class.IsValid() ? Class->GetMetaData(TEXT("DisplayName")) : FString();
 }
 
-FString FGraphNodeClassData::GetCategory() const
+FText FGraphNodeClassData::GetCategory() const
 {
-	return Class.IsValid() ? Class->GetMetaData(TEXT("Category")) : Category;
+	return Class.IsValid() ? Class->GetMetaDataText(TEXT("Category"), TEXT("UObjectCategory"), Class->GetFullGroupName(false)) : Category;
 }
 
 bool FGraphNodeClassData::IsAbstract() const
diff --git a/Engine/Source/Editor/AnimGraph/Classes/AnimationStateMachineSchema.h b/Engine/Source/Editor/AnimGraph/Classes/AnimationStateMachineSchema.h
--- a/Engine/Source/Editor/AnimGraph/Classes/AnimationStateMachineSchema.h
+++ b/Engine/Source/Editor/AnimGraph/Classes/AnimationStateMachineSchema.h
@@ -17,7 +17,7 @@ struct ANIMGRAPH_API FEdGraphSchemaAction_NewStateNode : public FEdGraphSchemaAc
 		, NodeTemplate(NULL)
 	{}
 
-	FEdGraphSchemaAction_NewStateNode(const FString& InNodeCategory, const FText& InMenuDesc, const FString& InToolTip, const int32 InGrouping)
+	FEdGraphSchemaAction_NewStateNode(const FText& InNodeCategory, const FText& InMenuDesc, const FString& InToolTip, const int32 InGrouping)
 		: FEdGraphSchemaAction(InNodeCategory, InMenuDesc, InToolTip, InGrouping)
 		, NodeTemplate(NULL)
 	{}
@@ -45,7 +45,7 @@ struct ANIMGRAPH_API FEdGraphSchemaAction_NewStateComment : public FEdGraphSchem
 		: FEdGraphSchemaAction()
 	{}
 
-	FEdGraphSchemaAction_NewStateComment(const FString& InNodeCategory, const FText& InMenuDesc, const FString& InToolTip, const int32 InGrouping)
+	FEdGraphSchemaAction_NewStateComment(const FText& InNodeCategory, const FText& InMenuDesc, const FString& InToolTip, const int32 InGrouping)
 		: FEdGraphSchemaAction(InNodeCategory, InMenuDesc, InToolTip, InGrouping)
 	{}
 
diff --git a/Engine/Source/Editor/AnimGraph/Private/AnimGraphNode_BlendSpaceBase.cpp b/Engine/Source/Editor/AnimGraph/Private/AnimGraphNode_BlendSpaceBase.cpp
--- a/Engine/Source/Editor/AnimGraph/Private/AnimGraphNode_BlendSpaceBase.cpp
+++ b/Engine/Source/Editor/AnimGraph/Private/AnimGraphNode_BlendSpaceBase.cpp
@@ -7,6 +7,8 @@
 #include "AnimGraphNode_RotationOffsetBlendSpace.h"
 #include "GraphEditorActions.h"
 
+#define LOCTEXT_NAMESPACE "AnimGraphNode_BlendSpaceBase"
+
 /////////////////////////////////////////////////////
 
 // Action to add a sequence player node to the graph
@@ -34,7 +36,7 @@ struct FNewBlendSpacePlayerAction : public FEdGraphSchemaAction_K2NewNode
 
 		MenuDescription = NodeTemplate->GetNodeTitle(ENodeTitleType::ListView);
 
-		Category = TEXT("Animations");
+		Category = LOCTEXT("Animation", "Animations");
 
 		// Grab extra keywords
 		Keywords = FText::FromString(BlendSpace->GetPathName());
@@ -109,3 +111,5 @@ void UAnimGraphNode_BlendSpaceBase::PostProcessPinName(const UEdGraphPin* Pin, F
 
 	Super::PostProcessPinName(Pin, DisplayName);
 }
+
+#undef LOCTEXT_NAMESPACE
\ No newline at end of file
diff --git a/Engine/Source/Editor/AnimGraph/Private/AnimGraphNode_SequencePlayer.cpp b/Engine/Source/Editor/AnimGraph/Private/AnimGraphNode_SequencePlayer.cpp
--- a/Engine/Source/Editor/AnimGraph/Private/AnimGraphNode_SequencePlayer.cpp
+++ b/Engine/Source/Editor/AnimGraph/Private/AnimGraphNode_SequencePlayer.cpp
@@ -32,7 +32,7 @@ public:
 
 		MenuDescription = Title;
 		TooltipDescription = TEXT("Evaluates an animation sequence to produce a pose");
-		Category = TEXT("Animations");
+		Category = LOCTEXT("Animation", "Animations");
 
 		// Grab extra keywords
 		Keywords = FText::FromName(InAssetInfo.ObjectPath);
diff --git a/Engine/Source/Editor/AnimGraph/Private/AnimationStateMachineSchema.cpp b/Engine/Source/Editor/AnimGraph/Private/AnimationStateMachineSchema.cpp
--- a/Engine/Source/Editor/AnimGraph/Private/AnimationStateMachineSchema.cpp
+++ b/Engine/Source/Editor/AnimGraph/Private/AnimationStateMachineSchema.cpp
@@ -27,7 +27,7 @@
 
 /////////////////////////////////////////////////////
 
-TSharedPtr<FEdGraphSchemaAction_NewStateNode> AddNewStateNodeAction(FGraphContextMenuBuilder& ContextMenuBuilder, const FString& Category, const FText& MenuDesc, const FString& Tooltip, const int32 Grouping = 0)
+TSharedPtr<FEdGraphSchemaAction_NewStateNode> AddNewStateNodeAction(FGraphContextMenuBuilder& ContextMenuBuilder, const FText& Category, const FText& MenuDesc, const FString& Tooltip, const int32 Grouping = 0)
 {
 	TSharedPtr<FEdGraphSchemaAction_NewStateNode> NewStateNode( new FEdGraphSchemaAction_NewStateNode(Category, MenuDesc, Tooltip, Grouping) );
 	ContextMenuBuilder.AddAction( NewStateNode );
@@ -266,13 +266,13 @@ void UAnimationStateMachineSchema::GetGraphContextActions(FGraphContextMenuBuild
 {
 	// Add state node
 	{
-		TSharedPtr<FEdGraphSchemaAction_NewStateNode> Action = AddNewStateNodeAction(ContextMenuBuilder, TEXT(""), LOCTEXT("AddState", "Add State..."), TEXT("A new state"));
+		TSharedPtr<FEdGraphSchemaAction_NewStateNode> Action = AddNewStateNodeAction(ContextMenuBuilder, FText::GetEmpty(), LOCTEXT("AddState", "Add State..."), TEXT("A new state"));
 		Action->NodeTemplate = NewObject<UAnimStateNode>(ContextMenuBuilder.OwnerOfTemporaries);
 	}
 
 	// Add conduit node
 	{
-		TSharedPtr<FEdGraphSchemaAction_NewStateNode> Action = AddNewStateNodeAction(ContextMenuBuilder, TEXT(""), LOCTEXT("AddConduit", "Add Conduit..."), TEXT("A new conduit state"));
+		TSharedPtr<FEdGraphSchemaAction_NewStateNode> Action = AddNewStateNodeAction(ContextMenuBuilder, FText::GetEmpty(), LOCTEXT("AddConduit", "Add Conduit..."), TEXT("A new conduit state"));
 		Action->NodeTemplate = NewObject<UAnimStateConduitNode>(ContextMenuBuilder.OwnerOfTemporaries);
 	}
 
@@ -291,7 +291,7 @@ void UAnimationStateMachineSchema::GetGraphContextActions(FGraphContextMenuBuild
 
 		if (!bHasEntry)
 		{
-			TSharedPtr<FEdGraphSchemaAction_NewStateNode> Action = AddNewStateNodeAction(ContextMenuBuilder, TEXT(""), LOCTEXT("AddEntryPoint", "Add Entry Point..."), TEXT("Define State Machine's Entry Point"));
+			TSharedPtr<FEdGraphSchemaAction_NewStateNode> Action = AddNewStateNodeAction(ContextMenuBuilder, FText::GetEmpty(), LOCTEXT("AddEntryPoint", "Add Entry Point..."), TEXT("Define State Machine's Entry Point"));
 			Action->NodeTemplate = NewObject<UAnimStateEntryNode>(ContextMenuBuilder.OwnerOfTemporaries);
 		}
 	}
@@ -304,7 +304,7 @@ void UAnimationStateMachineSchema::GetGraphContextActions(FGraphContextMenuBuild
 		const FText MenuDescription = bIsManyNodesSelected ? LOCTEXT("CreateCommentSelection", "Create Comment from Selection") : LOCTEXT("AddComment", "Add Comment...");
 		const FString ToolTip = TEXT("Create a resizeable comment box around selected nodes.");
 
-		TSharedPtr<FEdGraphSchemaAction_NewStateComment> NewComment( new FEdGraphSchemaAction_NewStateComment(TEXT(""), MenuDescription, ToolTip, 0) );
+		TSharedPtr<FEdGraphSchemaAction_NewStateComment> NewComment( new FEdGraphSchemaAction_NewStateComment(FText::GetEmpty(), MenuDescription, ToolTip, 0) );
 		ContextMenuBuilder.AddAction( NewComment );
 	}
 }
diff --git a/Engine/Source/Editor/BehaviorTreeEditor/Classes/EdGraphSchema_BehaviorTree.h b/Engine/Source/Editor/BehaviorTreeEditor/Classes/EdGraphSchema_BehaviorTree.h
--- a/Engine/Source/Editor/BehaviorTreeEditor/Classes/EdGraphSchema_BehaviorTree.h
+++ b/Engine/Source/Editor/BehaviorTreeEditor/Classes/EdGraphSchema_BehaviorTree.h
@@ -14,7 +14,7 @@ struct FBehaviorTreeSchemaAction_AutoArrange : public FEdGraphSchemaAction
 	FBehaviorTreeSchemaAction_AutoArrange() 
 		: FEdGraphSchemaAction() {}
 
-	FBehaviorTreeSchemaAction_AutoArrange(const FString& InNodeCategory, const FText& InMenuDesc, const FString& InToolTip, const int32 InGrouping)
+	FBehaviorTreeSchemaAction_AutoArrange(const FText& InNodeCategory, const FText& InMenuDesc, const FString& InToolTip, const int32 InGrouping)
 		: FEdGraphSchemaAction(InNodeCategory, InMenuDesc, InToolTip, InGrouping) 
 	{}
 
diff --git a/Engine/Source/Editor/BehaviorTreeEditor/Classes/EdGraphSchema_BehaviorTreeDecorator.h b/Engine/Source/Editor/BehaviorTreeEditor/Classes/EdGraphSchema_BehaviorTreeDecorator.h
--- a/Engine/Source/Editor/BehaviorTreeEditor/Classes/EdGraphSchema_BehaviorTreeDecorator.h
+++ b/Engine/Source/Editor/BehaviorTreeEditor/Classes/EdGraphSchema_BehaviorTreeDecorator.h
@@ -18,7 +18,7 @@ struct FDecoratorSchemaAction_NewNode : public FEdGraphSchemaAction
 		, NodeTemplate(NULL)
 	{}
 
-	FDecoratorSchemaAction_NewNode(const FString& InNodeCategory, const FText& InMenuDesc, const FString& InToolTip, const int32 InGrouping)
+	FDecoratorSchemaAction_NewNode(const FText& InNodeCategory, const FText& InMenuDesc, const FString& InToolTip, const int32 InGrouping)
 		: FEdGraphSchemaAction(InNodeCategory, InMenuDesc, InToolTip, InGrouping) 
 		, NodeTemplate(NULL)
 	{}
@@ -60,5 +60,5 @@ class UEdGraphSchema_BehaviorTreeDecorator : public UEdGraphSchema
 	virtual bool ShouldHidePinDefaultValue(UEdGraphPin* Pin) const override;
 	// End EdGraphSchema interface
 
-	static TSharedPtr<FDecoratorSchemaAction_NewNode> AddNewDecoratorAction(FGraphContextMenuBuilder& ContextMenuBuilder, const FString& Category, const FText& MenuDesc, const FString& Tooltip);
+	static TSharedPtr<FDecoratorSchemaAction_NewNode> AddNewDecoratorAction(FGraphContextMenuBuilder& ContextMenuBuilder, const FText& Category, const FText& MenuDesc, const FString& Tooltip);
 };
diff --git a/Engine/Source/Editor/BehaviorTreeEditor/Private/BehaviorTreeDecoratorGraphNode_Logic.cpp b/Engine/Source/Editor/BehaviorTreeEditor/Private/BehaviorTreeDecoratorGraphNode_Logic.cpp
--- a/Engine/Source/Editor/BehaviorTreeEditor/Private/BehaviorTreeDecoratorGraphNode_Logic.cpp
+++ b/Engine/Source/Editor/BehaviorTreeEditor/Private/BehaviorTreeDecoratorGraphNode_Logic.cpp
@@ -42,7 +42,7 @@ void UBehaviorTreeDecoratorGraphNode_Logic::GetMenuEntries(struct FGraphContextM
 	const int32 NumModes = sizeof(Modes) / sizeof(Modes[0]);
 	for (int32 i = 0; i < NumModes; i++)
 	{
-		TSharedPtr<FDecoratorSchemaAction_NewNode> AddOpAction = UEdGraphSchema_BehaviorTreeDecorator::AddNewDecoratorAction(ContextMenuBuilder, TEXT("Logic"), FText::FromString(DescribeLogicModeHelper(Modes[i])), "");
+		TSharedPtr<FDecoratorSchemaAction_NewNode> AddOpAction = UEdGraphSchema_BehaviorTreeDecorator::AddNewDecoratorAction(ContextMenuBuilder, NSLOCTEXT("BehaviorTreeGraphNode_Logic", "Logic", "Logic"), FText::FromString(DescribeLogicModeHelper(Modes[i])), "");
 
 		UBehaviorTreeDecoratorGraphNode_Logic* OpNode = NewObject<UBehaviorTreeDecoratorGraphNode_Logic>(ContextMenuBuilder.OwnerOfTemporaries);
 		OpNode->LogicMode = Modes[i];
diff --git a/Engine/Source/Editor/BehaviorTreeEditor/Private/EdGraphSchema_BehaviorTree.cpp b/Engine/Source/Editor/BehaviorTreeEditor/Private/EdGraphSchema_BehaviorTree.cpp
--- a/Engine/Source/Editor/BehaviorTreeEditor/Private/EdGraphSchema_BehaviorTree.cpp
+++ b/Engine/Source/Editor/BehaviorTreeEditor/Private/EdGraphSchema_BehaviorTree.cpp
@@ -48,7 +48,7 @@ void UEdGraphSchema_BehaviorTree::GetGraphNodeContextActions(FGraphContextMenuBu
 
 	if (SubNodeFlags == ESubNode::Decorator)
 	{
-		const FString& Category = UBehaviorTreeGraphNode_CompositeDecorator::StaticClass()->GetMetaData(TEXT("Category"));
+		const FText& Category = UBehaviorTreeGraphNode_CompositeDecorator::StaticClass()->GetMetaDataText(TEXT("Category"), TEXT("UObjectCategory"), UBehaviorTreeGraphNode_CompositeDecorator::StaticClass()->GetFullGroupName(false));
 		UEdGraph* Graph = (UEdGraph*)ContextMenuBuilder.CurrentGraph;
 		UBehaviorTreeGraphNode_CompositeDecorator* OpNode = NewObject<UBehaviorTreeGraphNode_CompositeDecorator>(Graph);
 		TSharedPtr<FAISchemaAction_NewSubNode> AddOpAction = UAIGraphSchema::AddNewSubNodeAction(ContextMenuBuilder, Category, FText::FromString(OpNode->GetNodeTypeDescription()), "");
@@ -147,7 +147,7 @@ void UEdGraphSchema_BehaviorTree::GetGraphContextActions(FGraphContextMenuBuilde
 	if (bNoParent)
 	{
 		TSharedPtr<FBehaviorTreeSchemaAction_AutoArrange> Action = TSharedPtr<FBehaviorTreeSchemaAction_AutoArrange>(
-			new FBehaviorTreeSchemaAction_AutoArrange(FString(), LOCTEXT("AutoArrange", "Auto Arrange"), FString(), 0)
+			new FBehaviorTreeSchemaAction_AutoArrange(FText::GetEmpty(), LOCTEXT("AutoArrange", "Auto Arrange"), FString(), 0)
 			);
 
 		ContextMenuBuilder.AddAction(Action);
diff --git a/Engine/Source/Editor/BehaviorTreeEditor/Private/EdGraphSchema_BehaviorTreeDecorator.cpp b/Engine/Source/Editor/BehaviorTreeEditor/Private/EdGraphSchema_BehaviorTreeDecorator.cpp
--- a/Engine/Source/Editor/BehaviorTreeEditor/Private/EdGraphSchema_BehaviorTreeDecorator.cpp
+++ b/Engine/Source/Editor/BehaviorTreeEditor/Private/EdGraphSchema_BehaviorTreeDecorator.cpp
@@ -105,7 +105,7 @@ UEdGraphSchema_BehaviorTreeDecorator::UEdGraphSchema_BehaviorTreeDecorator(const
 	PC_Boolean = TEXT("bool");
 }
 
-TSharedPtr<FDecoratorSchemaAction_NewNode> UEdGraphSchema_BehaviorTreeDecorator::AddNewDecoratorAction(FGraphContextMenuBuilder& ContextMenuBuilder, const FString& Category, const FText& MenuDesc, const FString& Tooltip)
+TSharedPtr<FDecoratorSchemaAction_NewNode> UEdGraphSchema_BehaviorTreeDecorator::AddNewDecoratorAction(FGraphContextMenuBuilder& ContextMenuBuilder, const FText& Category, const FText& MenuDesc, const FString& Tooltip)
 {
 	TSharedPtr<FDecoratorSchemaAction_NewNode> NewAction = TSharedPtr<FDecoratorSchemaAction_NewNode>(new FDecoratorSchemaAction_NewNode(Category, MenuDesc, Tooltip, 0));
 
diff --git a/Engine/Source/Editor/BlueprintGraph/Classes/EdGraphSchema_K2.h b/Engine/Source/Editor/BlueprintGraph/Classes/EdGraphSchema_K2.h
--- a/Engine/Source/Editor/BlueprintGraph/Classes/EdGraphSchema_K2.h
+++ b/Engine/Source/Editor/BlueprintGraph/Classes/EdGraphSchema_K2.h
@@ -33,7 +33,7 @@ struct BLUEPRINTGRAPH_API FEdGraphSchemaAction_K2Struct : public FEdGraphSchemaA
 		: FEdGraphSchemaAction()
 	{}
 
-	FEdGraphSchemaAction_K2Struct (const FString& InNodeCategory, const FText& InMenuDesc, const FString& InToolTip, const int32 InGrouping)
+	FEdGraphSchemaAction_K2Struct (const FText& InNodeCategory, const FText& InMenuDesc, const FString& InToolTip, const int32 InGrouping)
 		: FEdGraphSchemaAction(InNodeCategory, InMenuDesc, InToolTip, InGrouping)
 	{}
 };
@@ -255,7 +255,7 @@ class BLUEPRINTGRAPH_API UEdGraphSchema_K2 : public UEdGraphSchema
 	static const FName GN_AnimGraph;
 
 	// variable names
-	static const FName VR_DefaultCategory;
+	static const FText VR_DefaultCategory;
 
 	// action grouping values
 	static const int32 AG_LevelReference;
diff --git a/Engine/Source/Editor/BlueprintGraph/Classes/EdGraphSchema_K2_Actions.h b/Engine/Source/Editor/BlueprintGraph/Classes/EdGraphSchema_K2_Actions.h
--- a/Engine/Source/Editor/BlueprintGraph/Classes/EdGraphSchema_K2_Actions.h
+++ b/Engine/Source/Editor/BlueprintGraph/Classes/EdGraphSchema_K2_Actions.h
@@ -31,7 +31,7 @@ struct BLUEPRINTGRAPH_API FEdGraphSchemaAction_K2NewNode : public FEdGraphSchema
 		, bGotoNode(false)
 	{}
 
-	FEdGraphSchemaAction_K2NewNode(const FString& InNodeCategory, const FText& InMenuDesc, const FString& InToolTip, const int32 InGrouping)
+	FEdGraphSchemaAction_K2NewNode(const FText& InNodeCategory, const FText& InMenuDesc, const FString& InToolTip, const int32 InGrouping)
 		: FEdGraphSchemaAction(InNodeCategory, InMenuDesc, InToolTip, InGrouping) 
 		, NodeTemplate(NULL)
 		, bGotoNode(false)
@@ -78,7 +78,7 @@ struct BLUEPRINTGRAPH_API FEdGraphSchemaAction_K2ViewNode : public FEdGraphSchem
 		, NodePtr(NULL)
 	{}
 
-	FEdGraphSchemaAction_K2ViewNode(const FString& InNodeCategory, const FText& InMenuDesc, const FString& InToolTip, const int32 InGrouping)
+	FEdGraphSchemaAction_K2ViewNode(const FText& InNodeCategory, const FText& InMenuDesc, const FString& InToolTip, const int32 InGrouping)
 		: FEdGraphSchemaAction(InNodeCategory, InMenuDesc, InToolTip, InGrouping) 
 		, NodePtr(NULL)
 	{}
@@ -107,7 +107,7 @@ struct BLUEPRINTGRAPH_API FEdGraphSchemaAction_K2AssignDelegate : public FEdGrap
 		:FEdGraphSchemaAction_K2NewNode()
 	{}
 
-	FEdGraphSchemaAction_K2AssignDelegate(const FString& InNodeCategory, const FText& InMenuDesc, const FString& InToolTip, const int32 InGrouping)
+	FEdGraphSchemaAction_K2AssignDelegate(const FText& InNodeCategory, const FText& InMenuDesc, const FString& InToolTip, const int32 InGrouping)
 		: FEdGraphSchemaAction_K2NewNode(InNodeCategory, InMenuDesc, InToolTip, InGrouping) 
 	{}
 
@@ -139,7 +139,7 @@ struct BLUEPRINTGRAPH_API FEdGraphSchemaAction_EventFromFunction : public FEdGra
 		, SignatureFunction(NULL)
 	{}
 
-	FEdGraphSchemaAction_EventFromFunction(const FString& InNodeCategory, const FText& InMenuDesc, const FString& InToolTip, const int32 InGrouping)
+	FEdGraphSchemaAction_EventFromFunction(const FText& InNodeCategory, const FText& InMenuDesc, const FString& InToolTip, const int32 InGrouping)
 		: FEdGraphSchemaAction(InNodeCategory, InMenuDesc, InToolTip, InGrouping) 
 		, SignatureFunction(NULL)
 	{}
@@ -179,7 +179,7 @@ struct FEdGraphSchemaAction_K2AddComponent : public FEdGraphSchemaAction_K2NewNo
 		, ComponentAsset(NULL)
 	{}
 
-	FEdGraphSchemaAction_K2AddComponent(const FString& InNodeCategory, const FText& InMenuDesc, const FString& InToolTip, const int32 InGrouping)
+	FEdGraphSchemaAction_K2AddComponent(const FText& InNodeCategory, const FText& InMenuDesc, const FString& InToolTip, const int32 InGrouping)
 		: FEdGraphSchemaAction_K2NewNode(InNodeCategory, InMenuDesc, InToolTip, InGrouping)
 		, ComponentClass(NULL)
 		, ComponentAsset(NULL)
@@ -209,7 +209,7 @@ struct FEdGraphSchemaAction_K2AddTimeline : public FEdGraphSchemaAction_K2NewNod
 		: FEdGraphSchemaAction_K2NewNode()
 	{}
 
-	FEdGraphSchemaAction_K2AddTimeline(const FString& InNodeCategory, const FText& InMenuDesc, const FString& InToolTip, const int32 InGrouping)
+	FEdGraphSchemaAction_K2AddTimeline(const FText& InNodeCategory, const FText& InMenuDesc, const FString& InToolTip, const int32 InGrouping)
 		: FEdGraphSchemaAction_K2NewNode(InNodeCategory, InMenuDesc, InToolTip, InGrouping)
 	{}
 
@@ -236,7 +236,7 @@ struct FEdGraphSchemaAction_K2AddEvent : public FEdGraphSchemaAction_K2NewNode
 		: FEdGraphSchemaAction_K2NewNode()
 	{}
 
-	FEdGraphSchemaAction_K2AddEvent(const FString& InNodeCategory, const FText& InMenuDesc, const FString& InToolTip, const int32 InGrouping)
+	FEdGraphSchemaAction_K2AddEvent(const FText& InNodeCategory, const FText& InMenuDesc, const FString& InToolTip, const int32 InGrouping)
 		: FEdGraphSchemaAction_K2NewNode(InNodeCategory, InMenuDesc, InToolTip, InGrouping)
 	{}
 
@@ -266,7 +266,7 @@ struct FEdGraphSchemaAction_K2AddCustomEvent : public FEdGraphSchemaAction_K2New
 		: FEdGraphSchemaAction_K2NewNode()
 	{}
 
-	FEdGraphSchemaAction_K2AddCustomEvent(const FString& InNodeCategory, const FText& InMenuDesc, const FString& InToolTip, const int32 InGrouping)
+	FEdGraphSchemaAction_K2AddCustomEvent(const FText& InNodeCategory, const FText& InMenuDesc, const FString& InToolTip, const int32 InGrouping)
 		: FEdGraphSchemaAction_K2NewNode(InNodeCategory, InMenuDesc, InToolTip, InGrouping)
 	{}
 
@@ -297,7 +297,7 @@ struct FEdGraphSchemaAction_K2AddCallOnActor : public FEdGraphSchemaAction_K2New
 		: FEdGraphSchemaAction_K2NewNode()
 	{}
 
-	FEdGraphSchemaAction_K2AddCallOnActor(const FString& InNodeCategory, const FText& InMenuDesc, const FString& InToolTip, const int32 InGrouping)
+	FEdGraphSchemaAction_K2AddCallOnActor(const FText& InNodeCategory, const FText& InMenuDesc, const FString& InToolTip, const int32 InGrouping)
 		: FEdGraphSchemaAction_K2NewNode(InNodeCategory, InMenuDesc, InToolTip, InGrouping)
 	{}
 
@@ -325,7 +325,7 @@ struct FEdGraphSchemaAction_K2AddCallOnVariable : public FEdGraphSchemaAction_K2
 		: FEdGraphSchemaAction_K2NewNode()
 	{}
 
-	FEdGraphSchemaAction_K2AddCallOnVariable(const FString& InNodeCategory, const FText& InMenuDesc, const FString& InToolTip, const int32 InGrouping)
+	FEdGraphSchemaAction_K2AddCallOnVariable(const FText& InNodeCategory, const FText& InMenuDesc, const FString& InToolTip, const int32 InGrouping)
 		: FEdGraphSchemaAction_K2NewNode(InNodeCategory, InMenuDesc, InToolTip, InGrouping)
 	{}
 
@@ -351,7 +351,6 @@ struct FEdGraphSchemaAction_K2AddComment : public FEdGraphSchemaAction
 	FEdGraphSchemaAction_K2AddComment()
 		: FEdGraphSchemaAction()
 	{
-		Category = TEXT("");
 		MenuDescription = NSLOCTEXT("K2AddComment", "AddComment", "Add Comment...");
 		TooltipDescription = NSLOCTEXT("K2AddComment", "AddComment_Tooltip", "Create a resizable comment box.").ToString();
 	}
@@ -359,7 +358,6 @@ struct FEdGraphSchemaAction_K2AddComment : public FEdGraphSchemaAction
 	FEdGraphSchemaAction_K2AddComment(const FText& InDescription, const FString& InToolTip)
 		: FEdGraphSchemaAction()
 	{
-		Category = TEXT("");
 		MenuDescription = InDescription;
 		TooltipDescription = InToolTip;
 	}
@@ -386,7 +384,6 @@ struct FEdGraphSchemaAction_K2AddDocumentation : public FEdGraphSchemaAction
 	FEdGraphSchemaAction_K2AddDocumentation()
 		: FEdGraphSchemaAction()
 	{
-		Category = TEXT("");
 		MenuDescription = NSLOCTEXT("K2AddDocumentation", "AddDocumentation", "Add Documentation ...");
 		TooltipDescription = NSLOCTEXT("K2AddDocumentation", "AddDocumentation_Tooltip", "Creates a Documentation Node.").ToString();
 	}
@@ -394,7 +391,6 @@ struct FEdGraphSchemaAction_K2AddDocumentation : public FEdGraphSchemaAction
 	FEdGraphSchemaAction_K2AddDocumentation(const FString& InDescription, const FString& InToolTip)
 		: FEdGraphSchemaAction()
 	{
-		Category = TEXT("");
 		MenuDescription = FText::FromString( InDescription );
 		TooltipDescription = InToolTip;
 	}
@@ -423,7 +419,7 @@ struct BLUEPRINTGRAPH_API FEdGraphSchemaAction_K2TargetNode : public FEdGraphSch
 		: FEdGraphSchemaAction_K2NewNode()
 	{}
 
-	FEdGraphSchemaAction_K2TargetNode(const FString& InNodeCategory, const FText& InMenuDesc, const FString& InToolTip, const int32 InGrouping)
+	FEdGraphSchemaAction_K2TargetNode(const FText& InNodeCategory, const FText& InMenuDesc, const FString& InToolTip, const int32 InGrouping)
 		: FEdGraphSchemaAction_K2NewNode(InNodeCategory, InMenuDesc, InToolTip, InGrouping)
 	{}
 
@@ -450,7 +446,7 @@ struct BLUEPRINTGRAPH_API FEdGraphSchemaAction_K2PasteHere : public FEdGraphSche
 		: FEdGraphSchemaAction_K2NewNode()
 	{}
 
-	FEdGraphSchemaAction_K2PasteHere (const FString& InNodeCategory, const FText& InMenuDesc, const FString& InToolTip, const int32 InGrouping)
+	FEdGraphSchemaAction_K2PasteHere (const FText& InNodeCategory, const FText& InMenuDesc, const FString& InToolTip, const int32 InGrouping)
 		: FEdGraphSchemaAction_K2NewNode(InNodeCategory, InMenuDesc, InToolTip, InGrouping)
 	{}
 
@@ -492,7 +488,7 @@ struct BLUEPRINTGRAPH_API FEdGraphSchemaAction_K2Enum : public FEdGraphSchemaAct
 		: FEdGraphSchemaAction()
 	{}
 
-	FEdGraphSchemaAction_K2Enum (const FString& InNodeCategory, const FText& InMenuDesc, const FString& InToolTip, const int32 InGrouping)
+	FEdGraphSchemaAction_K2Enum (const FText& InNodeCategory, const FText& InMenuDesc, const FString& InToolTip, const int32 InGrouping)
 		: FEdGraphSchemaAction(InNodeCategory, InMenuDesc, InToolTip, InGrouping)
 	{}
 };
@@ -535,7 +531,7 @@ public:
 		: FEdGraphSchemaAction()
 	{}
 
-	FEdGraphSchemaAction_K2Var (const FString& InNodeCategory, const FText& InMenuDesc, const FString& InToolTip, const int32 InGrouping)
+	FEdGraphSchemaAction_K2Var (const FText& InNodeCategory, const FText& InMenuDesc, const FString& InToolTip, const int32 InGrouping)
 		: FEdGraphSchemaAction(InNodeCategory, InMenuDesc, InToolTip, InGrouping)
 	{}
 
@@ -602,7 +598,7 @@ public:
 		: FEdGraphSchemaAction()
 	{}
 
-	FEdGraphSchemaAction_K2LocalVar (const FString& InNodeCategory, const FText& InMenuDesc, const FString& InToolTip, const int32 InGrouping)
+	FEdGraphSchemaAction_K2LocalVar (const FText& InNodeCategory, const FText& InMenuDesc, const FString& InToolTip, const int32 InGrouping)
 		: FEdGraphSchemaAction(InNodeCategory, InMenuDesc, InToolTip, InGrouping)
 	{}
 
@@ -673,7 +669,7 @@ struct BLUEPRINTGRAPH_API FEdGraphSchemaAction_K2Graph : public FEdGraphSchemaAc
 		: FEdGraphSchemaAction()
 	{}
 
-	FEdGraphSchemaAction_K2Graph (EEdGraphSchemaAction_K2Graph::Type InType, const FString& InNodeCategory, const FText& InMenuDesc, const FString& InToolTip, const int32 InGrouping)
+	FEdGraphSchemaAction_K2Graph (EEdGraphSchemaAction_K2Graph::Type InType, const FText& InNodeCategory, const FText& InMenuDesc, const FString& InToolTip, const int32 InGrouping)
 		: FEdGraphSchemaAction(InNodeCategory, InMenuDesc, InToolTip, InGrouping), GraphType(InType), EdGraph(NULL)
 	{}
 
@@ -711,7 +707,7 @@ struct BLUEPRINTGRAPH_API FEdGraphSchemaAction_K2Event : public FEdGraphSchemaAc
 	 * @param   Tooltip			A string to display when hovering over this action entry.
 	 * @param   Grouping		Used to override list ordering (actions with the same number get grouped together, higher numbers get sorted first).
 	 */
-	FEdGraphSchemaAction_K2Event(FString const& Category, FText const& MenuDescription, FString const& Tooltip, int32 const Grouping)
+	FEdGraphSchemaAction_K2Event(FText const& Category, FText const& MenuDescription, FString const& Tooltip, int32 const Grouping)
 		: FEdGraphSchemaAction_K2TargetNode(Category, MenuDescription, Tooltip, Grouping)
 	{}
 
@@ -761,7 +757,7 @@ struct BLUEPRINTGRAPH_API FEdGraphSchemaAction_K2InputAction : public FEdGraphSc
 	* @param   Tooltip			A string to display when hovering over this action entry.
 	* @param   Grouping		Used to override list ordering (actions with the same number get grouped together, higher numbers get sorted first).
 	*/
-	FEdGraphSchemaAction_K2InputAction(FString const& Category, FText const& MenuDescription, FString const& Tooltip, int32 const Grouping)
+	FEdGraphSchemaAction_K2InputAction(FText const& Category, FText const& MenuDescription, FString const& Tooltip, int32 const Grouping)
 		: FEdGraphSchemaAction_K2TargetNode(Category, MenuDescription, Tooltip, Grouping)
 	{}
 
@@ -819,7 +815,7 @@ public:
 		: FEdGraphSchemaAction(), EdGraph(NULL)
 	{}
 
-	FEdGraphSchemaAction_K2Delegate(const FString& InNodeCategory, const FText& InMenuDesc, const FString& InToolTip, const int32 InGrouping)
+	FEdGraphSchemaAction_K2Delegate(const FText& InNodeCategory, const FText& InMenuDesc, const FString& InToolTip, const int32 InGrouping)
 		: FEdGraphSchemaAction(InNodeCategory, InMenuDesc, InToolTip, InGrouping), EdGraph(NULL)
 	{}
 
diff --git a/Engine/Source/Editor/BlueprintGraph/Classes/K2Node_CallFunction.h b/Engine/Source/Editor/BlueprintGraph/Classes/K2Node_CallFunction.h
--- a/Engine/Source/Editor/BlueprintGraph/Classes/K2Node_CallFunction.h
+++ b/Engine/Source/Editor/BlueprintGraph/Classes/K2Node_CallFunction.h
@@ -144,7 +144,7 @@ public:
 	/** Gets the non-specific tooltip for the function */
 	static FString GetDefaultTooltipForFunction(const UFunction* Function);
 	/** Get default category for this function in action menu */
-	static FString GetDefaultCategoryForFunction(const UFunction* Function, const FString& BaseCategory);
+	static FText GetDefaultCategoryForFunction(const UFunction* Function, const FText& BaseCategory);
 	/** Get keywords for this function in the action menu */
 	static FText GetKeywordsForFunction(const UFunction* Function);
 	/** Should be drawn compact for this function */
diff --git a/Engine/Source/Editor/BlueprintGraph/Classes/K2Node_EditablePinBase.h b/Engine/Source/Editor/BlueprintGraph/Classes/K2Node_EditablePinBase.h
--- a/Engine/Source/Editor/BlueprintGraph/Classes/K2Node_EditablePinBase.h
+++ b/Engine/Source/Editor/BlueprintGraph/Classes/K2Node_EditablePinBase.h
@@ -63,7 +63,7 @@ struct FKismetUserDeclaredFunctionMetadata
 	FString ToolTip;
 
 	UPROPERTY()
-	FString Category;
+	FText Category;
 
 	UPROPERTY()
 	FLinearColor InstanceTitleColor;
diff --git a/Engine/Source/Editor/BlueprintGraph/Private/BlueprintEventNodeSpawner.cpp b/Engine/Source/Editor/BlueprintGraph/Private/BlueprintEventNodeSpawner.cpp
--- a/Engine/Source/Editor/BlueprintGraph/Private/BlueprintEventNodeSpawner.cpp
+++ b/Engine/Source/Editor/BlueprintGraph/Private/BlueprintEventNodeSpawner.cpp
@@ -98,8 +98,7 @@ UBlueprintEventNodeSpawner* UBlueprintEventNodeSpawner::Create(UFunction const*
 	FBlueprintActionUiSpec& MenuSignature = NodeSpawner->DefaultMenuSignature;
 	FText const FuncName = UEdGraphSchema_K2::GetFriendlySignatureName(EventFunc);
 	MenuSignature.MenuName = FText::Format(LOCTEXT("EventWithSignatureName", "Event {0}"), FuncName);
-	FString const FuncCategory = UK2Node_CallFunction::GetDefaultCategoryForFunction(EventFunc, TEXT(""));
-	MenuSignature.Category = FText::FromString(LOCTEXT("AddEventCategory", "Add Event").ToString() + TEXT("|") + FuncCategory);
+	MenuSignature.Category = UK2Node_CallFunction::GetDefaultCategoryForFunction(EventFunc, LOCTEXT("AddEventCategory", "Add Event"));
 	//MenuSignature.Tooltip, will be pulled from the node template
 	MenuSignature.Keywords = UK2Node_CallFunction::GetKeywordsForFunction(EventFunc);
 	if (MenuSignature.Keywords.IsEmpty())
diff --git a/Engine/Source/Editor/BlueprintGraph/Private/BlueprintFunctionNodeSpawner.cpp b/Engine/Source/Editor/BlueprintGraph/Private/BlueprintFunctionNodeSpawner.cpp
--- a/Engine/Source/Editor/BlueprintGraph/Private/BlueprintFunctionNodeSpawner.cpp
+++ b/Engine/Source/Editor/BlueprintGraph/Private/BlueprintFunctionNodeSpawner.cpp
@@ -243,7 +243,7 @@ UBlueprintFunctionNodeSpawner* UBlueprintFunctionNodeSpawner::Create(TSubclassOf
 
 	FBlueprintActionUiSpec& MenuSignature = NodeSpawner->DefaultMenuSignature;
 	MenuSignature.MenuName = UK2Node_CallFunction::GetUserFacingFunctionName(Function);
-	MenuSignature.Category = FText::FromString( UK2Node_CallFunction::GetDefaultCategoryForFunction(Function, TEXT("")) );
+	MenuSignature.Category = UK2Node_CallFunction::GetDefaultCategoryForFunction(Function, FText::GetEmpty());
 	MenuSignature.Tooltip  = FText::FromString( UK2Node_CallFunction::GetDefaultTooltipForFunction(Function) );
 	// add at least one character, so that PrimeDefaultUiSpec() doesn't attempt to query the template node
 	MenuSignature.Keywords = UK2Node_CallFunction::GetKeywordsForFunction(Function);
diff --git a/Engine/Source/Editor/BlueprintGraph/Private/BlueprintVariableNodeSpawner.cpp b/Engine/Source/Editor/BlueprintGraph/Private/BlueprintVariableNodeSpawner.cpp
--- a/Engine/Source/Editor/BlueprintGraph/Private/BlueprintVariableNodeSpawner.cpp
+++ b/Engine/Source/Editor/BlueprintGraph/Private/BlueprintVariableNodeSpawner.cpp
@@ -115,7 +115,7 @@ UBlueprintVariableNodeSpawner* UBlueprintVariableNodeSpawner::Create(TSubclassOf
 	//--------------------------------------
 
 	FBlueprintActionUiSpec& MenuSignature = NodeSpawner->DefaultMenuSignature;
-	MenuSignature.Category = FEditorCategoryUtils::BuildCategoryString(FCommonEditorCategory::Variables, FText::FromName(VarDesc.Category));
+	MenuSignature.Category = FEditorCategoryUtils::BuildCategoryString(FCommonEditorCategory::Variables, VarDesc.Category);
 
 	FText const VarName = NodeSpawner->GetVariableName();
 	// @TODO: NodeClass could be modified post Create()
diff --git a/Engine/Source/Editor/BlueprintGraph/Private/EdGraphSchema_K2.cpp b/Engine/Source/Editor/BlueprintGraph/Private/EdGraphSchema_K2.cpp
--- a/Engine/Source/Editor/BlueprintGraph/Private/EdGraphSchema_K2.cpp
+++ b/Engine/Source/Editor/BlueprintGraph/Private/EdGraphSchema_K2.cpp
@@ -475,7 +475,7 @@ const FName UEdGraphSchema_K2::FN_UserConstructionScript(TEXT("UserConstructionS
 const FName UEdGraphSchema_K2::FN_ExecuteUbergraphBase(TEXT("ExecuteUbergraph"));
 const FName UEdGraphSchema_K2::GN_EventGraph(TEXT("EventGraph"));
 const FName UEdGraphSchema_K2::GN_AnimGraph(TEXT("AnimGraph"));
-const FName UEdGraphSchema_K2::VR_DefaultCategory(TEXT("Default"));
+const FText UEdGraphSchema_K2::VR_DefaultCategory(LOCTEXT("Default", "Default"));
 
 const int32 UEdGraphSchema_K2::AG_LevelReference = 100;
 
diff --git a/Engine/Source/Editor/BlueprintGraph/Private/K2Node_BaseAsyncTask.cpp b/Engine/Source/Editor/BlueprintGraph/Private/K2Node_BaseAsyncTask.cpp
--- a/Engine/Source/Editor/BlueprintGraph/Private/K2Node_BaseAsyncTask.cpp
+++ b/Engine/Source/Editor/BlueprintGraph/Private/K2Node_BaseAsyncTask.cpp
@@ -373,7 +373,7 @@ FName UK2Node_BaseAsyncTask::GetCornerIcon() const
 FText UK2Node_BaseAsyncTask::GetMenuCategory() const
 {	
 	UFunction* TargetFunction = GetFactoryFunction();
-	return FText::FromString(UK2Node_CallFunction::GetDefaultCategoryForFunction(TargetFunction, TEXT("")));
+	return UK2Node_CallFunction::GetDefaultCategoryForFunction(TargetFunction, FText::GetEmpty());
 }
 
 void UK2Node_BaseAsyncTask::GetMenuActions(FBlueprintActionDatabaseRegistrar& ActionRegistrar) const
diff --git a/Engine/Source/Editor/BlueprintGraph/Private/K2Node_CallFunction.cpp b/Engine/Source/Editor/BlueprintGraph/Private/K2Node_CallFunction.cpp
--- a/Engine/Source/Editor/BlueprintGraph/Private/K2Node_CallFunction.cpp
+++ b/Engine/Source/Editor/BlueprintGraph/Private/K2Node_CallFunction.cpp
@@ -1332,24 +1332,34 @@ FString UK2Node_CallFunction::GetDefaultTooltipForFunction(const UFunction* Func
 	}
 }
 
-FString UK2Node_CallFunction::GetDefaultCategoryForFunction(const UFunction* Function, const FString& BaseCategory)
+FText UK2Node_CallFunction::GetDefaultCategoryForFunction(const UFunction* Function, const FText& BaseCategory)
 {
-	FString NodeCategory = BaseCategory;
+	FText NodeCategory = BaseCategory;
 	if( Function->HasMetaData(FBlueprintMetadata::MD_FunctionCategory) )
 	{
-		// Add seperator if base category is supplied
-		if(NodeCategory.Len() > 0)
+		FText FuncCategory;
+		// If we are not showing friendly names, return the metadata stored, without localization
+		if( GEditor && !GetDefault<UEditorStyleSettings>()->bShowFriendlyNames )
 		{
-			NodeCategory += TEXT("|");
+			FuncCategory = FText::FromString(Function->GetMetaData(FBlueprintMetadata::MD_FunctionCategory));
+		}
+		else
+		{
+			// Look for localized metadata
+			FuncCategory = Function->GetMetaDataText(FBlueprintMetadata::MD_FunctionCategory, TEXT("UObjectCategory"), Function->GetFullGroupName(false));
+
+			// If the result is culture invariant, force it into a display string
+			if (FuncCategory.IsCultureInvariant())
+			{
+				FuncCategory = FText::FromString(FName::NameToDisplayString(FuncCategory.ToString(), false));
+			}
 		}
 
-		// Add category from function
-		FString FuncCategory = Function->GetMetaData(FBlueprintMetadata::MD_FunctionCategory);
-		if( GEditor && GetDefault<UEditorStyleSettings>()->bShowFriendlyNames )
+		// Combine with the BaseCategory to form the full category, delimited by "|"
+		if (!FuncCategory.IsEmpty() && !NodeCategory.IsEmpty())
 		{
-			FuncCategory = FName::NameToDisplayString( FuncCategory, false );
+			NodeCategory = FText::Format(FText::FromString(TEXT("{0}|{1}")), NodeCategory, FuncCategory);
 		}
-		NodeCategory += FuncCategory;
 	}
 	return NodeCategory;
 }
@@ -2181,7 +2191,7 @@ FText UK2Node_CallFunction::GetMenuCategory() const
 	UFunction* TargetFunction = GetTargetFunction();
 	if (TargetFunction != nullptr)
 	{
-		return FText::FromString(GetDefaultCategoryForFunction(TargetFunction, TEXT("")));
+		return GetDefaultCategoryForFunction(TargetFunction, FText::GetEmpty());
 	}
 	return FText::GetEmpty();
 }
diff --git a/Engine/Source/Editor/BlueprintGraph/Private/K2Node_Event.cpp b/Engine/Source/Editor/BlueprintGraph/Private/K2Node_Event.cpp
--- a/Engine/Source/Editor/BlueprintGraph/Private/K2Node_Event.cpp
+++ b/Engine/Source/Editor/BlueprintGraph/Private/K2Node_Event.cpp
@@ -676,14 +676,12 @@ void UK2Node_Event::GetNodeAttributes( TArray<TKeyValuePair<FString, FString>>&
 
 FText UK2Node_Event::GetMenuCategory() const
 {
-	FString FuncSubCategory;
+	FText FunctionCategory = LOCTEXT("AddEventCategory", "Add Event");
 	if (UFunction* Function = EventReference.ResolveMember<UFunction>(GetBlueprintClassFromNode()))
 	{
-		FuncSubCategory = UK2Node_CallFunction::GetDefaultCategoryForFunction(Function, TEXT(""));
+		FunctionCategory = UK2Node_CallFunction::GetDefaultCategoryForFunction(Function, FunctionCategory);
 	}
-
-	FString RootCategory = LOCTEXT("AddEventCategory", "Add Event").ToString();
-	return FText::FromString(RootCategory + TEXT("|") + FuncSubCategory);
+	return FunctionCategory;
 }
 
 bool UK2Node_Event::IsDeprecated() const
diff --git a/Engine/Source/Editor/BlueprintGraph/Private/K2Node_MacroInstance.cpp b/Engine/Source/Editor/BlueprintGraph/Private/K2Node_MacroInstance.cpp
--- a/Engine/Source/Editor/BlueprintGraph/Private/K2Node_MacroInstance.cpp
+++ b/Engine/Source/Editor/BlueprintGraph/Private/K2Node_MacroInstance.cpp
@@ -381,7 +381,7 @@ FText UK2Node_MacroInstance::GetCompactNodeTitle() const
 		else if( MacroName == TEXT("NegateFloat") || 
 				 MacroName == TEXT("NegateInt") )
 		{
-			return LOCTEXT("DecrementCompactNodeTitle", "-");
+			return LOCTEXT("NegateCompactNodeTitle", "-");
 		}
 	}	
 
@@ -535,7 +535,7 @@ FText UK2Node_MacroInstance::GetMenuCategory() const
 		FKismetUserDeclaredFunctionMetadata* MacroGraphMetadata = UK2Node_MacroInstance::GetAssociatedGraphMetadata(MacroGraph);
 		if ((MacroGraphMetadata != nullptr) && !MacroGraphMetadata->Category.IsEmpty())
 		{
-			MenuCategory = FText::FromString(MacroGraphMetadata->Category);
+			MenuCategory = MacroGraphMetadata->Category;
 		}
 	}
 
diff --git a/Engine/Source/Editor/GraphEditor/Private/GraphActionNode.cpp b/Engine/Source/Editor/GraphEditor/Private/GraphActionNode.cpp
--- a/Engine/Source/Editor/GraphEditor/Private/GraphActionNode.cpp
+++ b/Engine/Source/Editor/GraphEditor/Private/GraphActionNode.cpp
@@ -330,12 +330,12 @@ FText const& FGraphActionNode::GetDisplayName() const
 }
 
 //------------------------------------------------------------------------------
-FString FGraphActionNode::GetCategoryPath() const
+FText FGraphActionNode::GetCategoryPath() const
 {
-	FString CategoryPath;
+	FText CategoryPath;
 	if (IsCategoryNode())
 	{
-		CategoryPath = DisplayText.ToString();
+		CategoryPath = DisplayText;
 	}
 
 	TWeakPtr<FGraphActionNode> AncestorNode = ParentNode;
@@ -345,7 +345,7 @@ FString FGraphActionNode::GetCategoryPath() const
 
 		if( !AncestorDisplayText.IsEmpty() )
 		{
-			CategoryPath = AncestorDisplayText.ToString() + TEXT("|") + CategoryPath;
+			CategoryPath = FText::Format(FText::FromString(TEXT("{0}|{1}")), AncestorDisplayText, CategoryPath);
 		}
 		AncestorNode = AncestorNode.Pin()->GetParentNode();
 	}
diff --git a/Engine/Source/Editor/GraphEditor/Private/GraphActionNode.h b/Engine/Source/Editor/GraphEditor/Private/GraphActionNode.h
--- a/Engine/Source/Editor/GraphEditor/Private/GraphActionNode.h
+++ b/Engine/Source/Editor/GraphEditor/Private/GraphActionNode.h
@@ -121,7 +121,7 @@ public:
 	 *
 	 * @return A category path string, denoting the category hierarchy up to this node.
 	 */
-	FString GetCategoryPath() const;
+	FText GetCategoryPath() const;
 
 	/**
 	 * Checks to see if this node contains at least one valid action.
diff --git a/Engine/Source/Editor/GraphEditor/Private/GraphEditorDragDropAction.cpp b/Engine/Source/Editor/GraphEditor/Private/GraphEditorDragDropAction.cpp
--- a/Engine/Source/Editor/GraphEditor/Private/GraphEditorDragDropAction.cpp
+++ b/Engine/Source/Editor/GraphEditor/Private/GraphEditorDragDropAction.cpp
@@ -60,9 +60,9 @@ void FGraphEditorDragDropAction::SetHoveredGraph(const TSharedPtr<SGraphPanel>&
 	}
 }
 
-void FGraphEditorDragDropAction::SetHoveredCategoryName(const FString& InHoverCategoryName)
+void FGraphEditorDragDropAction::SetHoveredCategoryName(const FText& InHoverCategoryName)
 {
-	if(HoveredCategoryName != InHoverCategoryName)
+	if(!HoveredCategoryName.EqualTo(InHoverCategoryName))
 	{
 		HoveredCategoryName = InHoverCategoryName;
 		HoverTargetChanged();
diff --git a/Engine/Source/Editor/GraphEditor/Private/SGraphActionMenu.cpp b/Engine/Source/Editor/GraphEditor/Private/SGraphActionMenu.cpp
--- a/Engine/Source/Editor/GraphEditor/Private/SGraphActionMenu.cpp
+++ b/Engine/Source/Editor/GraphEditor/Private/SGraphActionMenu.cpp
@@ -209,7 +209,7 @@ public:
 		TSharedPtr<FGraphEditorDragDropAction> GraphDropOp = DragDropEvent.GetOperationAs<FGraphEditorDragDropAction>();
 		if (GraphDropOp.IsValid())
 		{
-			GraphDropOp->SetHoveredCategoryName( ActionNode.Pin()->GetDisplayName().ToString() );
+			GraphDropOp->SetHoveredCategoryName( ActionNode.Pin()->GetDisplayName() );
 		}
 	}
 
@@ -218,7 +218,7 @@ public:
 		TSharedPtr<FGraphEditorDragDropAction> GraphDropOp = DragDropEvent.GetOperationAs<FGraphEditorDragDropAction>();
 		if (GraphDropOp.IsValid())
 		{
-			GraphDropOp->SetHoveredCategoryName( FString(TEXT("")) );
+			GraphDropOp->SetHoveredCategoryName( FText::GetEmpty() );
 		}
 	}
 
@@ -543,15 +543,15 @@ bool SGraphActionMenu::SelectItemByName(const FName& ItemName, ESelectInfo::Type
 	return false;
 }
 
-void SGraphActionMenu::ExpandCategory(const FString& CategoryName)
+void SGraphActionMenu::ExpandCategory(const FText& CategoryName)
 {
-	if (CategoryName.Len())
+	if (!CategoryName.IsEmpty())
 	{
 		TArray<TSharedPtr<FGraphActionNode>> GraphNodes;
 		FilteredRootAction->GetAllNodes(GraphNodes);
 		for (int32 i = 0; i < GraphNodes.Num(); ++i)
 		{
-			if (GraphNodes[i]->GetDisplayName().ToString() == CategoryName)
+			if (GraphNodes[i]->GetDisplayName().EqualTo(CategoryName))
 			{
 				GraphNodes[i]->ExpandAllChildren(TreeView);
 			}
@@ -664,7 +664,7 @@ void SGraphActionMenu::GenerateFilteredItems(bool bPreserveExpansion)
 			// Combine the actions string, separate with \n so terms don't run into each other, and remove the spaces (incase the user is searching for a variable)
 			// In the case of groups containing multiple actions, they will have been created and added at the same place in the code, using the same description
 			// and keywords, so we only need to use the first one for filtering.
-			FString SearchText = CurrentAction.Actions[0]->GetSearchTitle().ToString() + LINE_TERMINATOR + CurrentAction.Actions[0]->GetSearchKeywords().ToString() + LINE_TERMINATOR +CurrentAction.Actions[0]->Category;
+			FString SearchText = CurrentAction.Actions[0]->GetSearchTitle().ToString() + LINE_TERMINATOR + CurrentAction.Actions[0]->GetSearchKeywords().ToString() + LINE_TERMINATOR +CurrentAction.Actions[0]->GetSearchCategory().ToString();
 			SearchText = SearchText.Replace( TEXT( " " ), TEXT( "" ) );
 
 			FString EachTermSanitized;
@@ -774,7 +774,7 @@ int32 SGraphActionMenu::GetActionFilteredWeight( const FGraphActionListBuilderBa
 		// Combine the actions string, separate with \n so terms don't run into each other, and remove the spaces (incase the user is searching for a variable)
 		// In the case of groups containing multiple actions, they will have been created and added at the same place in the code, using the same description
 		// and keywords, so we only need to use the first one for filtering.
-		FString SearchText = InCurrentAction.Actions[Action]->GetSearchTitle().ToString() + LINE_TERMINATOR + InCurrentAction.Actions[Action]->GetSearchKeywords().ToString() + LINE_TERMINATOR +InCurrentAction.Actions[Action]->Category;
+		FString SearchText = InCurrentAction.Actions[Action]->GetSearchTitle().ToString() + LINE_TERMINATOR + InCurrentAction.Actions[Action]->GetSearchKeywords().ToString() + LINE_TERMINATOR + InCurrentAction.Actions[Action]->GetSearchCategory().ToString();
 		SearchText = SearchText.Replace( TEXT( " " ), TEXT( "" ) );
 
 		// First the keywords
@@ -793,7 +793,7 @@ int32 SGraphActionMenu::GetActionFilteredWeight( const FGraphActionListBuilderBa
 		WeightedArrayList.Add( EachEntry );
 
 		// The category
-		InCurrentAction.Actions[Action]->Category.ParseIntoArray( EachEntry.Array, TEXT(" "), true );
+		InCurrentAction.Actions[Action]->GetSearchCategory().ToString().ParseIntoArray( EachEntry.Array, TEXT(" "), true );
 		EachEntry.Weight = CategoryWeight;
 		WeightedArrayList.Add( EachEntry );
 
diff --git a/Engine/Source/Editor/GraphEditor/Public/GraphEditorDragDropAction.h b/Engine/Source/Editor/GraphEditor/Public/GraphEditorDragDropAction.h
--- a/Engine/Source/Editor/GraphEditor/Public/GraphEditorDragDropAction.h
+++ b/Engine/Source/Editor/GraphEditor/Public/GraphEditorDragDropAction.h
@@ -13,7 +13,7 @@ public:
 	void SetHoveredPin(class UEdGraphPin* InPin);
 	void SetHoveredNode(const TSharedPtr<class SGraphNode>& InNode);
 	void SetHoveredGraph(const TSharedPtr<class SGraphPanel>& InGraph);
-	void SetHoveredCategoryName(const FString& InHoverCategoryName);
+	void SetHoveredCategoryName(const FText& InHoverCategoryName);
 	void SetHoveredAction(TSharedPtr<struct FEdGraphSchemaAction> Action);
 	void SetDropTargetValid( bool bValid ) { bDropTargetValid = bValid; }
 
@@ -23,7 +23,7 @@ public:
 	virtual FReply DroppedOnNode(FVector2D ScreenPosition, FVector2D GraphPosition) { return FReply::Unhandled(); }
 	virtual FReply DroppedOnPanel( const TSharedRef< class SWidget >& Panel, FVector2D ScreenPosition, FVector2D GraphPosition, UEdGraph& Graph) { return FReply::Unhandled(); }
 	virtual FReply DroppedOnAction(TSharedRef<struct FEdGraphSchemaAction> Action) { return FReply::Unhandled(); }
-	virtual FReply DroppedOnCategory(FString Category) { return FReply::Unhandled(); }
+	virtual FReply DroppedOnCategory(FText Category) { return FReply::Unhandled(); }
 	virtual void OnDragBegin(const TSharedRef<class SGraphPin>& InPin) {}
 	// End of interface to override
 	
@@ -58,7 +58,7 @@ private:
 protected:
 
 	// Name of category we are hovering over
-	FString HoveredCategoryName;
+	FText HoveredCategoryName;
 
 	// Action we are hovering over
 	TWeakPtr<struct FEdGraphSchemaAction> HoveredAction;
diff --git a/Engine/Source/Editor/GraphEditor/Public/SGraphActionMenu.h b/Engine/Source/Editor/GraphEditor/Public/SGraphActionMenu.h
--- a/Engine/Source/Editor/GraphEditor/Public/SGraphActionMenu.h
+++ b/Engine/Source/Editor/GraphEditor/Public/SGraphActionMenu.h
@@ -78,7 +78,7 @@ public:
 	/** Delegate executed when an action is dragged */
 	DECLARE_DELEGATE_RetVal_TwoParams( FReply, FOnActionDragged, const TArray< TSharedPtr<FEdGraphSchemaAction> >&, const FPointerEvent& );
 	/** Delegate executed when a category is dragged */
-	DECLARE_DELEGATE_RetVal_TwoParams( FReply, FOnCategoryDragged, const FString&, const FPointerEvent& );
+	DECLARE_DELEGATE_RetVal_TwoParams( FReply, FOnCategoryDragged, const FText&, const FPointerEvent& );
 	/** Delegate executed when the list of all actions needs to be refreshed */
 	DECLARE_DELEGATE_OneParam( FOnCollectAllActions, FGraphActionListBuilderBase& );
 	/** Delegate executed when the list of all actions needs to be refreshed, should return any sections that should always be visible, even if they don't have children. */
@@ -242,7 +242,7 @@ public:
 	bool SelectItemByName(const FName& ItemName, ESelectInfo::Type SelectInfo = ESelectInfo::Direct, int32 SectionId = INDEX_NONE, bool bIsCategory = false );
 
 	/** Expands any category with the associated name */
-	void ExpandCategory(const FString& CategoryName);
+	void ExpandCategory(const FText& CategoryName);
 
 	/* Handler for mouse button going down */
 	bool OnMouseButtonDownEvent( TWeakPtr<FEdGraphSchemaAction> InAction );
diff --git a/Engine/Source/Editor/Kismet/Private/BPVariableDragDropAction.cpp b/Engine/Source/Editor/Kismet/Private/BPVariableDragDropAction.cpp
--- a/Engine/Source/Editor/Kismet/Private/BPVariableDragDropAction.cpp
+++ b/Engine/Source/Editor/Kismet/Private/BPVariableDragDropAction.cpp
@@ -204,10 +204,10 @@ void FKismetVariableDragDropAction::HoverTargetChanged()
 			Message = FText::Format( LOCTEXT("IncorrectGraphForNodeReplace_Error", "Cannot replace node with local variable '{VariableName}' in external scope '{Scope}'"), Args);
 		}
 	}
-	else if (HoveredCategoryName.Len() > 0)
+	else if (!HoveredCategoryName.IsEmpty())
 	{
 		// Find Blueprint that made this class and get category of variable
-		FName Category = NAME_None;
+		FText Category;
 		UBlueprint* Blueprint;
 		
 		// Find the Blueprint for this property
@@ -225,8 +225,6 @@ void FKismetVariableDragDropAction::HoverTargetChanged()
 			Category = FBlueprintEditorUtils::GetBlueprintVariableCategory(Blueprint, VariableProperty->GetFName(), GetLocalVariableScope() );
 		}
 
-		FName NewCategory = FName(*HoveredCategoryName,  FNAME_Find);
-
 		// See if class is native
 		UClass* OuterClass = Cast<UClass>(VariableProperty->GetOuter());
 		if(OuterClass || Cast<UFunction>(VariableProperty->GetOuter()))
@@ -235,14 +233,14 @@ void FKismetVariableDragDropAction::HoverTargetChanged()
 
 			FFormatNamedArguments Args;
 			Args.Add(TEXT("VariableName"), FText::FromString(VariableString));
-			Args.Add(TEXT("HoveredCategoryName"), FText::FromString(HoveredCategoryName));
+			Args.Add(TEXT("HoveredCategoryName"), HoveredCategoryName);
 
 			if (bIsNativeVar)
 			{
 				StatusSymbol = FEditorStyle::GetBrush(TEXT("Graph.ConnectorFeedback.Error"));
 				Message = FText::Format( LOCTEXT("ChangingCatagoryNotThisVar", "Cannot change category for variable '{VariableName}'"), Args );
 			}
-			else if (Category == NewCategory)
+			else if (Category.EqualTo(HoveredCategoryName))
 			{
 				StatusSymbol = FEditorStyle::GetBrush(TEXT("Graph.ConnectorFeedback.Error"));
 				Message = FText::Format( LOCTEXT("ChangingCatagoryAlreadyIn", "Variable '{VariableName}' is already in category '{HoveredCategoryName}'"), Args );
@@ -548,9 +546,9 @@ FReply FKismetVariableDragDropAction::DroppedOnAction(TSharedRef<FEdGraphSchemaA
 			if(bMoved)
 			{
 				// Change category of var to match the one we dragged on to as well
-				FName MovedVarCategory = FBlueprintEditorUtils::GetBlueprintVariableCategory(BP, VariableName, GetLocalVariableScope());
-				FName TargetVarCategory = FBlueprintEditorUtils::GetBlueprintVariableCategory(BP, TargetVarName, GetLocalVariableScope());
-				if(MovedVarCategory != TargetVarCategory)
+				FText MovedVarCategory = FBlueprintEditorUtils::GetBlueprintVariableCategory(BP, VariableName, GetLocalVariableScope());
+				FText TargetVarCategory = FBlueprintEditorUtils::GetBlueprintVariableCategory(BP, TargetVarName, GetLocalVariableScope());
+				if(!MovedVarCategory.EqualTo(TargetVarCategory))
 				{
 					FBlueprintEditorUtils::SetBlueprintVariableCategory(BP, VariableName, GetLocalVariableScope(), TargetVarCategory, true);
 				}
@@ -565,9 +563,9 @@ FReply FKismetVariableDragDropAction::DroppedOnAction(TSharedRef<FEdGraphSchemaA
 	return FReply::Unhandled();
 }
 
-FReply FKismetVariableDragDropAction::DroppedOnCategory(FString Category)
+FReply FKismetVariableDragDropAction::DroppedOnCategory(FText Category)
 {
-	UE_LOG(LogTemp, Log, TEXT("Dropped %s on Category %s"), *VariableName.ToString(), *Category);
+	UE_LOG(LogTemp, Log, TEXT("Dropped %s on Category %s"), *VariableName.ToString(), *Category.ToString());
 
 	UBlueprint* BP = NULL;
 	if(VariableSource.Get()->IsA(UFunction::StaticClass()))
@@ -582,10 +580,10 @@ FReply FKismetVariableDragDropAction::DroppedOnCategory(FString Category)
 	if(BP != NULL)
 	{
 		// Check this is actually a different category
-		FName CurrentCategory = FBlueprintEditorUtils::GetBlueprintVariableCategory(BP, VariableName, GetLocalVariableScope());
-		if(FName(*Category) != CurrentCategory)
+		FText CurrentCategory = FBlueprintEditorUtils::GetBlueprintVariableCategory(BP, VariableName, GetLocalVariableScope());
+		if(!Category.EqualTo(CurrentCategory))
 		{
-			FBlueprintEditorUtils::SetBlueprintVariableCategory(BP, VariableName, GetLocalVariableScope(), FName(*Category), false);
+			FBlueprintEditorUtils::SetBlueprintVariableCategory(BP, VariableName, GetLocalVariableScope(), Category, false);
 		}
 	}
 
diff --git a/Engine/Source/Editor/Kismet/Private/BPVariableDragDropAction.h b/Engine/Source/Editor/Kismet/Private/BPVariableDragDropAction.h
--- a/Engine/Source/Editor/Kismet/Private/BPVariableDragDropAction.h
+++ b/Engine/Source/Editor/Kismet/Private/BPVariableDragDropAction.h
@@ -17,7 +17,7 @@ public:
 	virtual FReply DroppedOnNode(FVector2D ScreenPosition, FVector2D GraphPosition) override;
 	virtual FReply DroppedOnPanel(const TSharedRef< class SWidget >& Panel, FVector2D ScreenPosition, FVector2D GraphPosition, UEdGraph& Graph) override;
 	virtual FReply DroppedOnAction(TSharedRef<struct FEdGraphSchemaAction> Action) override;
-	virtual FReply DroppedOnCategory(FString Category) override;
+	virtual FReply DroppedOnCategory(FText Category) override;
 	// End of FGraphEditorDragDropAction
 
 	static TSharedRef<FKismetVariableDragDropAction> New(FName InVariableName, UStruct* InVariableSource, FNodeCreationAnalytic AnalyticCallback)
diff --git a/Engine/Source/Editor/Kismet/Private/BlueprintActionMenuBuilder.cpp b/Engine/Source/Editor/Kismet/Private/BlueprintActionMenuBuilder.cpp
--- a/Engine/Source/Editor/Kismet/Private/BlueprintActionMenuBuilder.cpp
+++ b/Engine/Source/Editor/Kismet/Private/BlueprintActionMenuBuilder.cpp
@@ -110,7 +110,7 @@ TSharedPtr<FBlueprintActionMenuItem> FBlueprintActionMenuItemFactory::MakeAction
 	FBlueprintActionMenuItem* NewMenuItem = new FBlueprintActionMenuItem(Action, UiSignature);
 
 	NewMenuItem->Grouping = MenuGrouping;
-	NewMenuItem->Category = FString::Printf(TEXT("%s|%s"), *RootCategory.ToString(), *NewMenuItem->Category);
+	NewMenuItem->Category = FText::Format(FText::FromString(TEXT("{0}|{1}")), RootCategory, NewMenuItem->Category);
 
 	return MakeShareable(NewMenuItem);
 }
@@ -121,7 +121,7 @@ TSharedPtr<FBlueprintDragDropMenuItem> FBlueprintActionMenuItemFactory::MakeDrag
 	// FBlueprintDragDropMenuItem takes care of its own menu MenuDescription, etc.
 	FBlueprintDragDropMenuItem* NewMenuItem = new FBlueprintDragDropMenuItem(Context, SampleAction, MenuGrouping);
 
-	NewMenuItem->Category = FString::Printf(TEXT("%s|%s"), *RootCategory.ToString(), *NewMenuItem->Category);
+	NewMenuItem->Category = FText::Format(FText::FromString(TEXT("{0}|{1}")), RootCategory, NewMenuItem->Category);
 	return MakeShareable(NewMenuItem);
 }
 
@@ -134,7 +134,7 @@ TSharedPtr<FBlueprintActionMenuItem> FBlueprintActionMenuItemFactory::MakeBoundM
 	FBlueprintActionMenuItem* NewMenuItem = new FBlueprintActionMenuItem(Action, UiSignature, ActionInfo.GetBindings());
 
 	NewMenuItem->Grouping = MenuGrouping;
-	NewMenuItem->Category = FString::Printf(TEXT("%s|%s"), *RootCategory.ToString(), *NewMenuItem->Category);
+	NewMenuItem->Category = FText::Format(FText::FromString(TEXT("{0}|{1}")), RootCategory, NewMenuItem->Category);
 
 	return MakeShareable(NewMenuItem);
 }
@@ -331,7 +331,7 @@ void FBlueprintActionMenuBuilderImpl::FMenuSectionDefinition::AddBoundMenuItems(
 
 					if (Flags & FBlueprintActionMenuBuilder::FlattenCategoryHierarcy)
 					{
-						LastMadeMenuItem->Category = ItemFactory.RootCategory.ToString();
+						LastMadeMenuItem->Category = ItemFactory.RootCategory;
 					}
 				}
 				else
@@ -396,7 +396,7 @@ FBlueprintActionMenuBuilderImpl::MenuItemList FBlueprintActionMenuBuilderImpl::F
 		UnBoundMenuEntry = ItemFactory.MakeActionMenuItem(EditorContext, DatabaseAction);
 		if (Flags & FBlueprintActionMenuBuilder::FlattenCategoryHierarcy)
 		{
-			UnBoundMenuEntry->Category = ItemFactory.RootCategory.ToString();
+			UnBoundMenuEntry->Category = ItemFactory.RootCategory;
 		}
 	}
 
diff --git a/Engine/Source/Editor/Kismet/Private/BlueprintActionMenuItem.cpp b/Engine/Source/Editor/Kismet/Private/BlueprintActionMenuItem.cpp
--- a/Engine/Source/Editor/Kismet/Private/BlueprintActionMenuItem.cpp
+++ b/Engine/Source/Editor/Kismet/Private/BlueprintActionMenuItem.cpp
@@ -208,7 +208,7 @@ FBlueprintActionMenuItem::FBlueprintActionMenuItem(UBlueprintNodeSpawner const*
 	check(Action != nullptr);
 
 	MenuDescription = UiSpec.MenuName;
-	Category = UiSpec.Category.ToString();
+	Category = UiSpec.Category;
 	TooltipDescription = UiSpec.Tooltip.ToString();
 	Keywords = UiSpec.Keywords;
 	DocExcerptRef.DocLink = UiSpec.DocLink;
diff --git a/Engine/Source/Editor/Kismet/Private/BlueprintActionMenuUtils.cpp b/Engine/Source/Editor/Kismet/Private/BlueprintActionMenuUtils.cpp
--- a/Engine/Source/Editor/Kismet/Private/BlueprintActionMenuUtils.cpp
+++ b/Engine/Source/Editor/Kismet/Private/BlueprintActionMenuUtils.cpp
@@ -673,7 +673,7 @@ void FBlueprintActionMenuUtils::MakeContextMenu(FBlueprintActionContext const& C
 		if (FKismetEditorUtilities::CanPasteNodes(Graph))
 		{
 			// @TODO: Grey out menu option with tooltip if one of the nodes cannot paste into this graph
-			TSharedPtr<FEdGraphSchemaAction> PasteHereAction(new FEdGraphSchemaAction_K2PasteHere(TEXT(""), LOCTEXT("PasteHereMenuName", "Paste here"), TEXT(""), MainMenuSectionGroup));
+			TSharedPtr<FEdGraphSchemaAction> PasteHereAction(new FEdGraphSchemaAction_K2PasteHere(FText::GetEmpty(), LOCTEXT("PasteHereMenuName", "Paste here"), TEXT(""), MainMenuSectionGroup));
 			MenuOut.AddAction(PasteHereAction);
 			break;
 		}
@@ -683,7 +683,7 @@ void FBlueprintActionMenuUtils::MakeContextMenu(FBlueprintActionContext const& C
 	{
 		FText SelectComponentMsg = LOCTEXT("SelectComponentForEvents", "Select a Component to see available Events & Functions");
 		FText SelectComponentToolTip = LOCTEXT("SelectComponentForEventsTooltip", "Select a Component in the MyBlueprint tab to see available Events and Functions in this menu.");
-		TSharedPtr<FEdGraphSchemaAction> MsgAction = TSharedPtr<FEdGraphSchemaAction>(new FEdGraphSchemaAction_Dummy(TEXT(""), SelectComponentMsg, SelectComponentToolTip.ToString(), ComponentsSectionGroup));
+		TSharedPtr<FEdGraphSchemaAction> MsgAction = TSharedPtr<FEdGraphSchemaAction>(new FEdGraphSchemaAction_Dummy(FText::GetEmpty(), SelectComponentMsg, SelectComponentToolTip.ToString(), ComponentsSectionGroup));
 		MenuOut.AddAction(MsgAction);
 	}
 
@@ -691,7 +691,7 @@ void FBlueprintActionMenuUtils::MakeContextMenu(FBlueprintActionContext const& C
 	{
 		FText SelectActorsMsg = LOCTEXT("SelectActorForEvents", "Select Actor(s) to see available Events & Functions");
 		FText SelectActorsToolTip = LOCTEXT("SelectActorForEventsTooltip", "Select Actor(s) in the level to see available Events and Functions in this menu.");
-		TSharedPtr<FEdGraphSchemaAction> MsgAction = TSharedPtr<FEdGraphSchemaAction>(new FEdGraphSchemaAction_Dummy(TEXT(""), SelectActorsMsg, SelectActorsToolTip.ToString(), LevelActorSectionGroup));
+		TSharedPtr<FEdGraphSchemaAction> MsgAction = TSharedPtr<FEdGraphSchemaAction>(new FEdGraphSchemaAction_Dummy(FText::GetEmpty(), SelectActorsMsg, SelectActorsToolTip.ToString(), LevelActorSectionGroup));
 		MenuOut.AddAction(MsgAction);
 	}
 }
diff --git a/Engine/Source/Editor/Kismet/Private/BlueprintDetailsCustomization.cpp b/Engine/Source/Editor/Kismet/Private/BlueprintDetailsCustomization.cpp
--- a/Engine/Source/Editor/Kismet/Private/BlueprintDetailsCustomization.cpp
+++ b/Engine/Source/Editor/Kismet/Private/BlueprintDetailsCustomization.cpp
@@ -402,7 +402,7 @@ void FBlueprintVarActionDetails::CustomizeDetails( IDetailLayoutBuilder& DetailL
 
 	PopulateCategories(MyBlueprint.Pin().Get(), CategorySource);
 	TSharedPtr<SComboButton> NewComboButton;
-	TSharedPtr<SListView<TSharedPtr<FString>>> NewListView;
+	TSharedPtr<SListView<TSharedPtr<FText>>> NewListView;
 
 	TSharedPtr<SToolTip> CategoryTooltip = IDocumentation::Get()->CreateToolTip(LOCTEXT("EditCategoryName_Tooltip", "The category of the variable; editing this will place the variable into another category or create a new one."), NULL, DocLink, TEXT("Category"));
 
@@ -442,7 +442,7 @@ void FBlueprintVarActionDetails::CustomizeDetails( IDetailLayoutBuilder& DetailL
 			.AutoHeight()
 			.MaxHeight(400.0f)
 			[
-				SAssignNew(NewListView, SListView<TSharedPtr<FString>>)
+				SAssignNew(NewListView, SListView<TSharedPtr<FText>>)
 					.ListItemsSource(&CategorySource)
 					.OnGenerateRow(this, &FBlueprintVarActionDetails::MakeCategoryViewWidget)
 					.OnSelectionChanged(this, &FBlueprintVarActionDetails::OnCategorySelectionChanged)
@@ -1066,7 +1066,7 @@ void FBlueprintVarActionDetails::OnTooltipTextCommitted(const FText& NewText, ET
 	FBlueprintEditorUtils::SetBlueprintVariableMetaData(GetBlueprintObj(), VarName, GetLocalVariableScope(CachedVariableProperty.Get()), TEXT("tooltip"), NewText.ToString() );
 }
 
-void FBlueprintVarActionDetails::PopulateCategories(SMyBlueprint* MyBlueprint, TArray<TSharedPtr<FString>>& CategorySource)
+void FBlueprintVarActionDetails::PopulateCategories(SMyBlueprint* MyBlueprint, TArray<TSharedPtr<FText>>& CategorySource)
 {
 	// Used to compare found categories to prevent double adds
 	TArray<FString> CategoryNameList;
@@ -1099,20 +1099,20 @@ void FBlueprintVarActionDetails::PopulateCategories(SMyBlueprint* MyBlueprint, T
 	}
 
 	CategorySource.Empty();
-	CategorySource.Add(MakeShareable(new FString(TEXT("Default"))));
+	CategorySource.Add(MakeShareable(new FText(LOCTEXT("Default", "Default"))));
 	for (int32 i = 0; i < VisibleVariables.Num(); ++i)
 	{
-		FName Category = FBlueprintEditorUtils::GetBlueprintVariableCategory(Blueprint, VisibleVariables[i], NULL);
-		if (Category != NAME_None && Category != Blueprint->GetFName())
+		FText Category = FBlueprintEditorUtils::GetBlueprintVariableCategory(Blueprint, VisibleVariables[i], NULL);
+		if (!Category.IsEmpty() && !Category.EqualTo(FText::FromString(Blueprint->GetName())))
 		{
 			bool bNewCategory = true;
 			for (int32 j = 0; j < CategorySource.Num() && bNewCategory; ++j)
 			{
-				bNewCategory &= *CategorySource[j].Get() != Category.ToString();
+				bNewCategory &= !CategorySource[j].Get()->EqualTo(Category);
 			}
 			if (bNewCategory)
 			{
-				CategorySource.Add(MakeShareable(new FString(Category.ToString())));
+				CategorySource.Add(MakeShareable(new FText(Category)));
 			}
 		}
 	}
@@ -1122,19 +1122,19 @@ void FBlueprintVarActionDetails::PopulateCategories(SMyBlueprint* MyBlueprint, T
 	{
 		if(UFunction* Function = Blueprint->SkeletonGeneratedClass->FindFunctionByName(FunctionGraph->GetFName()))
 		{
-			FString FunctionCategory = Function->GetMetaData(FBlueprintMetadata::MD_FunctionCategory);
+			FText FunctionCategory = Function->GetMetaDataText(FBlueprintMetadata::MD_FunctionCategory, TEXT("UObjectCategory"), Function->GetFullGroupName(false));
 
 			if(!FunctionCategory.IsEmpty())
 			{
 				bool bNewCategory = true;
 				for (int32 j = 0; j < CategorySource.Num() && bNewCategory; ++j)
 				{
-					bNewCategory &= *CategorySource[j].Get() != FunctionCategory;
+					bNewCategory &= !CategorySource[j].Get()->EqualTo(FunctionCategory);
 				}
 
 				if(bNewCategory)
 				{
-					CategorySource.Add(MakeShareable(new FString(FunctionCategory)));
+					CategorySource.Add(MakeShareable(new FText(FunctionCategory)));
 				}
 			}
 		}
@@ -1149,11 +1149,11 @@ void FBlueprintVarActionDetails::PopulateCategories(SMyBlueprint* MyBlueprint, T
 				bool bNewCategory = true;
 				for (int32 j = 0; j < CategorySource.Num() && bNewCategory; ++j)
 				{
-					bNewCategory &= *CategorySource[j].Get() != Variable.Category.ToString();
+					bNewCategory &= !CategorySource[j].Get()->EqualTo(Variable.Category);
 				}
 				if (bNewCategory)
 				{
-					CategorySource.Add(MakeShareable(new FString(Variable.Category.ToString())));
+					CategorySource.Add(MakeShareable(new FText(Variable.Category)));
 				}
 			}
 		}
@@ -1169,11 +1169,11 @@ void FBlueprintVarActionDetails::PopulateCategories(SMyBlueprint* MyBlueprint, T
 			bool bNewCategory = true;
 			for (int32 j = 0; j < CategorySource.Num() && bNewCategory; ++j)
 			{
-				bNewCategory &= *CategorySource[j].Get() != TypedEntryNode->MetaData.Category;
+				bNewCategory &= !CategorySource[j].Get()->EqualTo(TypedEntryNode->MetaData.Category);
 			}
 			if (bNewCategory)
 			{
-				CategorySource.Add(MakeShareable(new FString(TypedEntryNode->MetaData.Category)));
+				CategorySource.Add(MakeShareable(new FText(TypedEntryNode->MetaData.Category)));
 			}
 		}
 	}
@@ -1186,19 +1186,19 @@ void FBlueprintVarActionDetails::PopulateCategories(SMyBlueprint* MyBlueprint, T
 
 		if (UEdGraphSchema_K2::CanKismetOverrideFunction(Function) && !UEdGraphSchema_K2::FunctionCanBePlacedAsEvent(Function))
 		{
-			FString FunctionCategory = Function->GetMetaData(FBlueprintMetadata::MD_FunctionCategory);
+			FText FunctionCategory = Function->GetMetaDataText(FBlueprintMetadata::MD_FunctionCategory, TEXT("UObjectCategory"), Function->GetFullGroupName(false));
 
 			if (!FunctionCategory.IsEmpty())
 			{
 				bool bNewCategory = true;
 				for (int32 j = 0; j < CategorySource.Num() && bNewCategory; ++j)
 				{
-					bNewCategory &= *CategorySource[j].Get() != FunctionCategory;
+					bNewCategory &= !CategorySource[j].Get()->EqualTo(FunctionCategory);
 				}
 
 				if (bNewCategory)
 				{
-					CategorySource.Add(MakeShareable(new FString(FunctionCategory)));
+					CategorySource.Add(MakeShareable(new FText(FunctionCategory)));
 				}
 			}
 		}
@@ -1275,16 +1275,16 @@ FText FBlueprintVarActionDetails::OnGetCategoryText() const
 	{
 		const UEdGraphSchema_K2* K2Schema = GetDefault<UEdGraphSchema_K2>();
 
-		FName Category = FBlueprintEditorUtils::GetBlueprintVariableCategory(GetBlueprintObj(), VarName, GetLocalVariableScope(CachedVariableProperty.Get()));
+		FText Category = FBlueprintEditorUtils::GetBlueprintVariableCategory(GetBlueprintObj(), VarName, GetLocalVariableScope(CachedVariableProperty.Get()));
 
 		// Older blueprints will have their name as the default category and whenever it is the same as the default category, display localized text
-		if( Category == GetBlueprintObj()->GetFName() || Category == K2Schema->VR_DefaultCategory )
+		if( Category.EqualTo(FText::FromString(GetBlueprintObj()->GetName())) || Category.EqualTo(K2Schema->VR_DefaultCategory) )
 		{
-			return LOCTEXT("DefaultCategory", "Default");
+			return K2Schema->VR_DefaultCategory;
 		}
 		else
 		{
-			return FText::FromName(Category);
+			return Category;
 		}
 		return FText::FromName(VarName);
 	}
@@ -1298,33 +1298,29 @@ void FBlueprintVarActionDetails::OnCategoryTextCommitted(const FText& NewText, E
 		// Remove excess whitespace and prevent categories with just spaces
 		FText CategoryName = FText::TrimPrecedingAndTrailing(NewText);
 
-		FString NewCategory = CategoryName.ToString();
-		if(NewCategory.Len() <= NAME_SIZE)
-		{
-			FBlueprintEditorUtils::SetBlueprintVariableCategory(GetBlueprintObj(), VarName, GetLocalVariableScope(CachedVariableProperty.Get()), FName( *NewCategory ));
-			check(MyBlueprint.IsValid());
-			PopulateCategories(MyBlueprint.Pin().Get(), CategorySource);
-			MyBlueprint.Pin()->ExpandCategory(NewCategory);
-		}
+		FBlueprintEditorUtils::SetBlueprintVariableCategory(GetBlueprintObj(), VarName, GetLocalVariableScope(CachedVariableProperty.Get()), CategoryName);
+		check(MyBlueprint.IsValid());
+		PopulateCategories(MyBlueprint.Pin().Get(), CategorySource);
+		MyBlueprint.Pin()->ExpandCategory(CategoryName);
 	}
 }
 
-TSharedRef< ITableRow > FBlueprintVarActionDetails::MakeCategoryViewWidget( TSharedPtr<FString> Item, const TSharedRef< STableViewBase >& OwnerTable )
+TSharedRef< ITableRow > FBlueprintVarActionDetails::MakeCategoryViewWidget( TSharedPtr<FText> Item, const TSharedRef< STableViewBase >& OwnerTable )
 {
 	return SNew(STableRow<TSharedPtr<FString>>, OwnerTable)
 		[
-			SNew(STextBlock) .Text(FText::FromString(*Item.Get()))
+			SNew(STextBlock) .Text(*Item.Get())
 		];
 }
 
-void FBlueprintVarActionDetails::OnCategorySelectionChanged( TSharedPtr<FString> ProposedSelection, ESelectInfo::Type /*SelectInfo*/ )
+void FBlueprintVarActionDetails::OnCategorySelectionChanged( TSharedPtr<FText> ProposedSelection, ESelectInfo::Type /*SelectInfo*/ )
 {
 	FName VarName = CachedVariableName;
 	if (ProposedSelection.IsValid() && VarName != NAME_None)
 	{
-		FString NewCategory = *ProposedSelection.Get();
+		FText NewCategory = *ProposedSelection.Get();
 
-		FBlueprintEditorUtils::SetBlueprintVariableCategory(GetBlueprintObj(), VarName, GetLocalVariableScope(CachedVariableProperty.Get()), FName( *NewCategory ));
+		FBlueprintEditorUtils::SetBlueprintVariableCategory(GetBlueprintObj(), VarName, GetLocalVariableScope(CachedVariableProperty.Get()), NewCategory );
 		CategoryListView.Pin()->ClearSelection();
 		CategoryComboButton.Pin()->SetIsOpen(false);
 		MyBlueprint.Pin()->ExpandCategory(NewCategory);
@@ -2343,7 +2339,7 @@ void FBlueprintGraphActionDetails::CustomizeDetails( IDetailLayoutBuilder& Detai
 			{
 				FBlueprintVarActionDetails::PopulateCategories(MyBlueprint.Pin().Get(), CategorySource);
 				TSharedPtr<SComboButton> NewComboButton;
-				TSharedPtr<SListView<TSharedPtr<FString>>> NewListView;
+				TSharedPtr<SListView<TSharedPtr<FText>>> NewListView;
 
 				const FString DocLink = TEXT("Shared/Editors/BlueprintEditor/VariableDetails");
 				TSharedPtr<SToolTip> CategoryTooltip = IDocumentation::Get()->CreateToolTip(LOCTEXT("EditCategoryName_Tooltip", "The category of the variable; editing this will place the variable into another category or create a new one."), NULL, DocLink, TEXT("Category"));
@@ -2383,7 +2379,7 @@ void FBlueprintGraphActionDetails::CustomizeDetails( IDetailLayoutBuilder& Detai
 								.AutoHeight()
 								.MaxHeight(400.0f)
 								[
-									SAssignNew(NewListView, SListView<TSharedPtr<FString>>)
+									SAssignNew(NewListView, SListView<TSharedPtr<FText>>)
 									.ListItemsSource(&CategorySource)
 									.OnGenerateRow(this, &FBlueprintGraphActionDetails::MakeCategoryViewWidget)
 									.OnSelectionChanged(this, &FBlueprintGraphActionDetails::OnCategorySelectionChanged)
@@ -2914,18 +2910,17 @@ FText FBlueprintDelegateActionDetails::OnGetCategoryText() const
 	{
 		const UEdGraphSchema_K2* K2Schema = GetDefault<UEdGraphSchema_K2>();
 		FName DelegateName = DelegateProperty->GetFName();
-		FName Category = FBlueprintEditorUtils::GetBlueprintVariableCategory(GetBlueprintObj(), DelegateName, NULL);
+		FText Category = FBlueprintEditorUtils::GetBlueprintVariableCategory(GetBlueprintObj(), DelegateName, NULL);
 
 		// Older blueprints will have their name as the default category
-		if( Category == GetBlueprintObj()->GetFName() || Category == K2Schema->VR_DefaultCategory )
+		if( Category.EqualTo(FText::FromString(GetBlueprintObj()->GetName())) || Category.EqualTo(K2Schema->VR_DefaultCategory) )
 		{
 			return LOCTEXT("DefaultCategory", "Default");
 		}
 		else
 		{
-			return FText::FromName(Category);
+			return Category;
 		}
-		return FText::FromName(DelegateName);
 	}
 	return FText();
 }
@@ -2938,32 +2933,31 @@ void FBlueprintDelegateActionDetails::OnCategoryTextCommitted(const FText& NewTe
 		{
 			// Remove excess whitespace and prevent categories with just spaces
 			FText CategoryName = FText::TrimPrecedingAndTrailing(NewText);
-			FString NewCategory = CategoryName.ToString();
 			
-			FBlueprintEditorUtils::SetBlueprintVariableCategory(GetBlueprintObj(), DelegateProperty->GetFName(), NULL, FName( *NewCategory ));
+			FBlueprintEditorUtils::SetBlueprintVariableCategory(GetBlueprintObj(), DelegateProperty->GetFName(), NULL, CategoryName);
 			check(MyBlueprint.IsValid());
 			FBlueprintVarActionDetails::PopulateCategories(MyBlueprint.Pin().Get(), CategorySource);
-			MyBlueprint.Pin()->ExpandCategory(NewCategory);
+			MyBlueprint.Pin()->ExpandCategory(CategoryName);
 		}
 	}
 }
 
-TSharedRef< ITableRow > FBlueprintDelegateActionDetails::MakeCategoryViewWidget( TSharedPtr<FString> Item, const TSharedRef< STableViewBase >& OwnerTable )
+TSharedRef< ITableRow > FBlueprintDelegateActionDetails::MakeCategoryViewWidget( TSharedPtr<FText> Item, const TSharedRef< STableViewBase >& OwnerTable )
 {
 	return SNew(STableRow<TSharedPtr<FString>>, OwnerTable)
 		[
-			SNew(STextBlock) .Text(FText::FromString(*Item.Get()))
+			SNew(STextBlock) .Text(*Item.Get())
 		];
 }
 
-void FBlueprintDelegateActionDetails::OnCategorySelectionChanged( TSharedPtr<FString> ProposedSelection, ESelectInfo::Type /*SelectInfo*/ )
+void FBlueprintDelegateActionDetails::OnCategorySelectionChanged( TSharedPtr<FText> ProposedSelection, ESelectInfo::Type /*SelectInfo*/ )
 {
 	UMulticastDelegateProperty* DelegateProperty = GetDelegatePoperty();
 	if (DelegateProperty && ProposedSelection.IsValid())
 	{
-		FString NewCategory = *ProposedSelection.Get();
+		FText NewCategory = *ProposedSelection.Get();
 
-		FBlueprintEditorUtils::SetBlueprintVariableCategory(GetBlueprintObj(), DelegateProperty->GetFName(), NULL, FName( *NewCategory ));
+		FBlueprintEditorUtils::SetBlueprintVariableCategory(GetBlueprintObj(), DelegateProperty->GetFName(), NULL, NewCategory);
 		CategoryListView.Pin()->ClearSelection();
 		CategoryComboButton.Pin()->SetIsOpen(false);
 		MyBlueprint.Pin()->ExpandCategory(NewCategory);
@@ -2999,7 +2993,7 @@ void FBlueprintDelegateActionDetails::CustomizeDetails( IDetailLayoutBuilder& De
 
 		FBlueprintVarActionDetails::PopulateCategories(MyBlueprint.Pin().Get(), CategorySource);
 		TSharedPtr<SComboButton> NewComboButton;
-		TSharedPtr<SListView<TSharedPtr<FString>>> NewListView;
+		TSharedPtr<SListView<TSharedPtr<FText>>> NewListView;
 
 		Category.AddCustomRow( LOCTEXT("CategoryLabel", "Category") )
 		.NameContent()
@@ -3034,7 +3028,7 @@ void FBlueprintDelegateActionDetails::CustomizeDetails( IDetailLayoutBuilder& De
 				.AutoHeight()
 				.MaxHeight(400.0f)
 				[
-					SAssignNew(NewListView, SListView<TSharedPtr<FString>>)
+					SAssignNew(NewListView, SListView<TSharedPtr<FText>>)
 						.ListItemsSource(&CategorySource)
 						.OnGenerateRow(this, &FBlueprintDelegateActionDetails::MakeCategoryViewWidget)
 						.OnSelectionChanged(this, &FBlueprintDelegateActionDetails::OnCategorySelectionChanged)
@@ -3403,12 +3397,12 @@ FText FBlueprintGraphActionDetails::OnGetCategoryText() const
 	if (FKismetUserDeclaredFunctionMetadata* Metadata = GetMetadataBlock())
 	{
 		const UEdGraphSchema_K2* K2Schema = GetDefault<UEdGraphSchema_K2>();
-		if( Metadata->Category.IsEmpty() || Metadata->Category == K2Schema->VR_DefaultCategory.ToString() )
+		if( Metadata->Category.IsEmpty() || Metadata->Category.EqualTo(K2Schema->VR_DefaultCategory) )
 		{
 			return LOCTEXT("DefaultCategory", "Default");
 		}
 		
-		return FText::FromString(Metadata->Category);
+		return Metadata->Category;
 	}
 	else
 	{
@@ -3428,11 +3422,11 @@ void FBlueprintGraphActionDetails::OnCategoryTextCommitted(const FText& NewText,
 			if(CategoryName.IsEmpty())
 			{
 				const UEdGraphSchema_K2* K2Schema = GetDefault<UEdGraphSchema_K2>();
-				Metadata->Category = K2Schema->VR_DefaultCategory.ToString();
+				Metadata->Category = K2Schema->VR_DefaultCategory;
 			}
 			else
 			{
-				Metadata->Category = CategoryName.ToString();
+				Metadata->Category = CategoryName;
 			}
 
 			if(auto Function = FindFunction())
@@ -3446,7 +3440,7 @@ void FBlueprintGraphActionDetails::OnCategoryTextCommitted(const FText& NewText,
 	}
 }
 
-void FBlueprintGraphActionDetails::OnCategorySelectionChanged( TSharedPtr<FString> ProposedSelection, ESelectInfo::Type /*SelectInfo*/ )
+void FBlueprintGraphActionDetails::OnCategorySelectionChanged( TSharedPtr<FText> ProposedSelection, ESelectInfo::Type /*SelectInfo*/ )
 {
 	if(ProposedSelection.IsValid())
 	{
@@ -3456,7 +3450,7 @@ void FBlueprintGraphActionDetails::OnCategorySelectionChanged( TSharedPtr<FStrin
 			if(auto Function = FindFunction())
 			{
 				Function->Modify();
-				Function->SetMetaData(FBlueprintMetadata::MD_FunctionCategory, **ProposedSelection.Get());
+				Function->SetMetaData(FBlueprintMetadata::MD_FunctionCategory, *ProposedSelection.Get()->ToString());
 			}
 			MyBlueprint.Pin()->Refresh();
 			FBlueprintEditorUtils::MarkBlueprintAsModified(GetBlueprintObj());
@@ -3468,11 +3462,11 @@ void FBlueprintGraphActionDetails::OnCategorySelectionChanged( TSharedPtr<FStrin
 	}
 }
 
-TSharedRef< ITableRow > FBlueprintGraphActionDetails::MakeCategoryViewWidget( TSharedPtr<FString> Item, const TSharedRef< STableViewBase >& OwnerTable )
+TSharedRef< ITableRow > FBlueprintGraphActionDetails::MakeCategoryViewWidget( TSharedPtr<FText> Item, const TSharedRef< STableViewBase >& OwnerTable )
 {
 	return SNew(STableRow<TSharedPtr<FString>>, OwnerTable)
 		[
-			SNew(STextBlock) .Text(FText::FromString(*Item.Get()))
+			SNew(STextBlock) .Text(*Item.Get())
 		];
 }
 
@@ -4211,9 +4205,9 @@ void FBlueprintGlobalOptionsDetails::CustomizeDetails(IDetailLayoutBuilder& Deta
 		for (TFieldIterator<UProperty> PropertyIt(Blueprint->GetClass(), EFieldIteratorFlags::IncludeSuper); PropertyIt; ++PropertyIt)
 		{
 			UProperty* Property = *PropertyIt;
-			FString Category = FObjectEditorUtils::GetCategory(Property);
+			FText Category = FObjectEditorUtils::GetCategoryText(Property);
 
-			if ( Category != TEXT("BlueprintOptions") && Category != TEXT("ClassOptions") )
+			if ( Category.ToString() != TEXT("BlueprintOptions") && Category.ToString() != TEXT("ClassOptions") )
 			{
 				DetailLayout.HideProperty(DetailLayout.GetProperty(Property->GetFName()));
 			}
@@ -4387,7 +4381,7 @@ void FBlueprintComponentDetails::CustomizeDetails(IDetailLayoutBuilder& DetailLa
 				.AutoHeight()
 				.MaxHeight(400.0f)
 				[
-					SAssignNew(VariableCategoryListView, SListView<TSharedPtr<FString>>)
+					SAssignNew(VariableCategoryListView, SListView<TSharedPtr<FText>>)
 					.ListItemsSource(&VariableCategorySource)
 					.OnGenerateRow(this, &FBlueprintComponentDetails::MakeVariableCategoryViewWidget)
 					.OnSelectionChanged(this, &FBlueprintComponentDetails::OnVariableCategorySelectionChanged)
@@ -4546,16 +4540,16 @@ FText FBlueprintComponentDetails::OnGetVariableCategoryText() const
 	{
 		const UEdGraphSchema_K2* K2Schema = GetDefault<UEdGraphSchema_K2>();
 
-		FName Category = FBlueprintEditorUtils::GetBlueprintVariableCategory(GetBlueprintObj(), VarName, NULL);
+		FText Category = FBlueprintEditorUtils::GetBlueprintVariableCategory(GetBlueprintObj(), VarName, NULL);
 
 		// Older blueprints will have their name as the default category
-		if( Category == GetBlueprintObj()->GetFName() )
+		if( Category.EqualTo(FText::FromString(GetBlueprintObj()->GetName())) )
 		{
-			return FText::FromName(K2Schema->VR_DefaultCategory);
+			return K2Schema->VR_DefaultCategory;
 		}
 		else
 		{
-			return FText::FromName(Category);
+			return Category;
 		}
 	}
 
@@ -4568,22 +4562,20 @@ void FBlueprintComponentDetails::OnVariableCategoryTextCommitted(const FText& Ne
 
 	if (InTextCommit == ETextCommit::OnEnter || InTextCommit == ETextCommit::OnUserMovedFocus)
 	{
-		FString NewCategory = NewText.ToString();
-
-		FBlueprintEditorUtils::SetBlueprintVariableCategory(GetBlueprintObj(), CachedNodePtr->GetVariableName(), NULL, FName( *NewCategory ));
+		FBlueprintEditorUtils::SetBlueprintVariableCategory(GetBlueprintObj(), CachedNodePtr->GetVariableName(), NULL, NewText);
 		PopulateVariableCategories();
 	}
 }
 
-void FBlueprintComponentDetails::OnVariableCategorySelectionChanged( TSharedPtr<FString> ProposedSelection, ESelectInfo::Type /*SelectInfo*/ )
+void FBlueprintComponentDetails::OnVariableCategorySelectionChanged( TSharedPtr<FText> ProposedSelection, ESelectInfo::Type /*SelectInfo*/ )
 {
 	check(CachedNodePtr.IsValid());
 
 	FName VarName = CachedNodePtr->GetVariableName();
 	if (ProposedSelection.IsValid() && VarName != NAME_None)
 	{
-		FString NewCategory = *ProposedSelection.Get();
-		FBlueprintEditorUtils::SetBlueprintVariableCategory(GetBlueprintObj(), VarName, NULL, FName( *NewCategory ));
+		FText NewCategory = *ProposedSelection.Get();
+		FBlueprintEditorUtils::SetBlueprintVariableCategory(GetBlueprintObj(), VarName, NULL, NewCategory);
 
 		check(VariableCategoryListView.IsValid());
 		check(VariableCategoryComboButton.IsValid());
@@ -4593,12 +4585,12 @@ void FBlueprintComponentDetails::OnVariableCategorySelectionChanged( TSharedPtr<
 	}
 }
 
-TSharedRef< ITableRow > FBlueprintComponentDetails::MakeVariableCategoryViewWidget( TSharedPtr<FString> Item, const TSharedRef< STableViewBase >& OwnerTable )
+TSharedRef< ITableRow > FBlueprintComponentDetails::MakeVariableCategoryViewWidget( TSharedPtr<FText> Item, const TSharedRef< STableViewBase >& OwnerTable )
 {
 	return SNew(STableRow<TSharedPtr<FString>>, OwnerTable)
 	[
 		SNew(STextBlock)
-		.Text(FText::FromString(*Item.Get()))
+			.Text(*Item.Get())
 	];
 }
 
@@ -4623,20 +4615,20 @@ void FBlueprintComponentDetails::PopulateVariableCategories()
 	FBlueprintEditorUtils::GetSCSVariableNameList(Blueprint, VisibleVariables);
 
 	VariableCategorySource.Empty();
-	VariableCategorySource.Add(MakeShareable(new FString(TEXT("Default"))));
+	VariableCategorySource.Add(MakeShareable(new FText(LOCTEXT("Default", "Default"))));
 	for (int32 i = 0; i < VisibleVariables.Num(); ++i)
 	{
-		FName Category = FBlueprintEditorUtils::GetBlueprintVariableCategory(Blueprint, VisibleVariables[i], NULL);
-		if (Category != NAME_None && Category != Blueprint->GetFName())
+		FText Category = FBlueprintEditorUtils::GetBlueprintVariableCategory(Blueprint, VisibleVariables[i], NULL);
+		if (!Category.IsEmpty() && !Category.EqualTo(FText::FromString(Blueprint->GetName())))
 		{
 			bool bNewCategory = true;
 			for (int32 j = 0; j < VariableCategorySource.Num() && bNewCategory; ++j)
 			{
-				bNewCategory &= *VariableCategorySource[j].Get() != Category.ToString();
+				bNewCategory &= !VariableCategorySource[j].Get()->EqualTo(Category);
 			}
 			if (bNewCategory)
 			{
-				VariableCategorySource.Add(MakeShareable(new FString(Category.ToString())));
+				VariableCategorySource.Add(MakeShareable(new FText(Category)));
 			}
 		}
 	}
diff --git a/Engine/Source/Editor/Kismet/Private/BlueprintDetailsCustomization.h b/Engine/Source/Editor/Kismet/Private/BlueprintDetailsCustomization.h
--- a/Engine/Source/Editor/Kismet/Private/BlueprintDetailsCustomization.h
+++ b/Engine/Source/Editor/Kismet/Private/BlueprintDetailsCustomization.h
@@ -81,7 +81,7 @@ public:
 	/** IDetailCustomization interface */
 	virtual void CustomizeDetails( IDetailLayoutBuilder& DetailLayout ) override;
 	
-	static void PopulateCategories(SMyBlueprint* MyBlueprint, TArray<TSharedPtr<FString>>& CategorySource);
+	static void PopulateCategories(SMyBlueprint* MyBlueprint, TArray<TSharedPtr<FText>>& CategorySource);
 
 private:
 	/** Accessors passed to parent */
@@ -124,8 +124,8 @@ private:
 	
 	FText OnGetCategoryText() const;
 	void OnCategoryTextCommitted(const FText& NewText, ETextCommit::Type InTextCommit, FName VarName);
-	TSharedRef< ITableRow > MakeCategoryViewWidget( TSharedPtr<FString> Item, const TSharedRef< STableViewBase >& OwnerTable );
-	void OnCategorySelectionChanged( TSharedPtr<FString> ProposedSelection, ESelectInfo::Type /*SelectInfo*/ );
+	TSharedRef< ITableRow > MakeCategoryViewWidget( TSharedPtr<FText> Item, const TSharedRef< STableViewBase >& OwnerTable );
+	void OnCategorySelectionChanged( TSharedPtr<FText> ProposedSelection, ESelectInfo::Type /*SelectInfo*/ );
 	
 	EVisibility ShowEditableCheckboxVisibilty() const;
 	ECheckBoxState OnEditableCheckboxState() const;
@@ -195,10 +195,10 @@ private:
 	bool bIsVarNameInvalid;
 	
 	/** A list of all category names to choose from */
-	TArray<TSharedPtr<FString>> CategorySource;
+	TArray<TSharedPtr<FText>> CategorySource;
 	/** Widgets for the categories */
 	TWeakPtr<SComboButton> CategoryComboButton;
-	TWeakPtr<SListView<TSharedPtr<FString>>> CategoryListView;
+	TWeakPtr<SListView<TSharedPtr<FText>>> CategoryListView;
 
 	/** Array of names of property flags on the selected property */
 	TArray< TSharedPtr< FString > > PropertyFlags;
@@ -315,8 +315,8 @@ private:
 	void OnTooltipTextCommitted(const FText& NewText, ETextCommit::Type InTextCommit);
 	FText OnGetCategoryText() const;
 	void OnCategoryTextCommitted(const FText& NewText, ETextCommit::Type InTextCommit);
-	TSharedRef< ITableRow > MakeCategoryViewWidget( TSharedPtr<FString> Item, const TSharedRef< STableViewBase >& OwnerTable );
-	void OnCategorySelectionChanged( TSharedPtr<FString> ProposedSelection, ESelectInfo::Type /*SelectInfo*/ );
+	TSharedRef< ITableRow > MakeCategoryViewWidget( TSharedPtr<FText> Item, const TSharedRef< STableViewBase >& OwnerTable );
+	void OnCategorySelectionChanged( TSharedPtr<FText> ProposedSelection, ESelectInfo::Type /*SelectInfo*/ );
 
 	void CollectAvailibleSignatures();
 	void OnFunctionSelected(TSharedPtr<FString> FunctionItemData, ESelectInfo::Type SelectInfo);
@@ -325,11 +325,11 @@ private:
 private:
 
 	/** A list of all category names to choose from */
-	TArray<TSharedPtr<FString>> CategorySource;
+	TArray<TSharedPtr<FText>> CategorySource;
 
 	/** Widgets for the categories */
 	TWeakPtr<SComboButton> CategoryComboButton;
-	TWeakPtr<SListView<TSharedPtr<FString>>> CategoryListView;
+	TWeakPtr<SListView<TSharedPtr<FText>>> CategoryListView;
 
 	TArray<TSharedPtr<FString>> FunctionsToCopySignatureFrom;
 	TSharedPtr<STextComboBox> CopySignatureComboButton;
@@ -534,10 +534,10 @@ private:
 	static void SetNetFlags( TWeakObjectPtr<UK2Node_EditablePinBase> FunctionEntryNode, uint32 NetFlags);
 
 	/** Callback when a graph category is changed */
-	void OnCategorySelectionChanged( TSharedPtr<FString> ProposedSelection, ESelectInfo::Type /*SelectInfo*/ );
+	void OnCategorySelectionChanged( TSharedPtr<FText> ProposedSelection, ESelectInfo::Type /*SelectInfo*/ );
 
 	/** Callback to make category widgets */
-	TSharedRef< ITableRow > MakeCategoryViewWidget( TSharedPtr<FString> Item, const TSharedRef< STableViewBase >& OwnerTable );
+	TSharedRef< ITableRow > MakeCategoryViewWidget( TSharedPtr<FText> Item, const TSharedRef< STableViewBase >& OwnerTable );
 
 private:
 
@@ -551,11 +551,11 @@ private:
 	TSharedPtr<class SColorBlock> ColorBlock;
 
 	/** A list of all category names to choose from */
-	TArray<TSharedPtr<FString>> CategorySource;
+	TArray<TSharedPtr<FText>> CategorySource;
 
 	/** Widgets for the categories */
 	TWeakPtr<SComboButton> CategoryComboButton;
-	TWeakPtr<SListView<TSharedPtr<FString>>> CategoryListView;
+	TWeakPtr<SListView<TSharedPtr<FText>>> CategoryListView;
 };
 
 /** Blueprint Interface List Details */
@@ -714,8 +714,8 @@ protected:
 	bool OnVariableCategoryChangeEnabled() const;
 	FText OnGetVariableCategoryText() const;
 	void OnVariableCategoryTextCommitted(const FText& NewText, ETextCommit::Type InTextCommit, FName VarName);
-	void OnVariableCategorySelectionChanged(TSharedPtr<FString> ProposedSelection, ESelectInfo::Type /*SelectInfo*/);
-	TSharedRef<ITableRow> MakeVariableCategoryViewWidget(TSharedPtr<FString> Item, const TSharedRef< STableViewBase >& OwnerTable);
+	void OnVariableCategorySelectionChanged(TSharedPtr<FText> ProposedSelection, ESelectInfo::Type /*SelectInfo*/);
+	TSharedRef<ITableRow> MakeVariableCategoryViewWidget(TSharedPtr<FText> Item, const TSharedRef< STableViewBase >& OwnerTable);
 
 	FText GetSocketName() const;
 	void OnBrowseSocket();
@@ -739,11 +739,11 @@ private:
 	bool bIsVariableNameInvalid;
 
 	/** A list of all category names to choose from */
-	TArray<TSharedPtr<FString>> VariableCategorySource;
+	TArray<TSharedPtr<FText>> VariableCategorySource;
 
 	/** Widgets for the categories */
 	TSharedPtr<SComboButton> VariableCategoryComboButton;
-	TSharedPtr<SListView<TSharedPtr<FString>>> VariableCategoryListView;
+	TSharedPtr<SListView<TSharedPtr<FText>>> VariableCategoryListView;
 };
 
 /** Details customization for All Graph Nodes */
diff --git a/Engine/Source/Editor/Kismet/Private/BlueprintDragDropMenuItem.cpp b/Engine/Source/Editor/Kismet/Private/BlueprintDragDropMenuItem.cpp
--- a/Engine/Source/Editor/Kismet/Private/BlueprintDragDropMenuItem.cpp
+++ b/Engine/Source/Editor/Kismet/Private/BlueprintDragDropMenuItem.cpp
@@ -43,12 +43,12 @@ FBlueprintDragDropMenuItem::FBlueprintDragDropMenuItem(FBlueprintActionContext c
 		MenuDescription = bShowFriendlyNames ? FText::FromString(UEditorEngine::GetFriendlyName(SampleProperty)) : FText::FromName(SampleProperty->GetFName());
 
 		TooltipDescription = SampleProperty->GetToolTipText().ToString();
-		Category = FObjectEditorUtils::GetCategory(SampleProperty);
+		Category = FObjectEditorUtils::GetCategoryText(SampleProperty);
 
 		bool const bIsDelegateProperty = SampleProperty->IsA<UMulticastDelegateProperty>();
 		if (bIsDelegateProperty && Category.IsEmpty())
 		{
-			Category = FEditorCategoryUtils::GetCommonCategory(FCommonEditorCategory::Delegates).ToString();
+			Category = FEditorCategoryUtils::GetCommonCategory(FCommonEditorCategory::Delegates);
 		}
 		else if (!bIsDelegateProperty)
 		{
@@ -67,14 +67,14 @@ FBlueprintDragDropMenuItem::FBlueprintDragDropMenuItem(FBlueprintActionContext c
 			}
 			else if (bIsMemberProperty)
 			{
-				TextCategory = FEditorCategoryUtils::BuildCategoryString(FCommonEditorCategory::Variables, FText::FromString(Category));
+				TextCategory = FEditorCategoryUtils::BuildCategoryString(FCommonEditorCategory::Variables, Category);
 			}
 			
 			if (!bIsMemberProperty)
 			{
 				TextCategory = FText::Format(LOCTEXT("NonMemberVarCategory", "Class|{0}|{1}"), PropertyClass->GetDisplayNameText(), TextCategory);
 			}
-			Category = TextCategory.ToString();
+			Category = TextCategory;
 		}
 	}
 	else
diff --git a/Engine/Source/Editor/Kismet/Private/SBlueprintActionMenu.cpp b/Engine/Source/Editor/Kismet/Private/SBlueprintActionMenu.cpp
--- a/Engine/Source/Editor/Kismet/Private/SBlueprintActionMenu.cpp
+++ b/Engine/Source/Editor/Kismet/Private/SBlueprintActionMenu.cpp
@@ -26,7 +26,6 @@ struct FBlueprintAction_PromoteVariable : public FEdGraphSchemaAction
 	FBlueprintAction_PromoteVariable()
 		: FEdGraphSchemaAction()
 	{
-		Category = TEXT("");
 		MenuDescription = LOCTEXT("PromoteToVariable", "Promote to variable");
 		TooltipDescription = LOCTEXT("PromoteToVariable", "Promote to variable").ToString();
 		Grouping = 1;
diff --git a/Engine/Source/Editor/Kismet/Private/SMyBlueprint.cpp b/Engine/Source/Editor/Kismet/Private/SMyBlueprint.cpp
--- a/Engine/Source/Editor/Kismet/Private/SMyBlueprint.cpp
+++ b/Engine/Source/Editor/Kismet/Private/SMyBlueprint.cpp
@@ -69,14 +69,14 @@ public:
 	virtual void HoverTargetChanged() override
 	{
 		const FSlateBrush* StatusSymbol = FEditorStyle::GetBrush(TEXT("NoBrush")); 
-		FText Message = FText::FromString(DraggedCategory);
+		FText Message = DraggedCategory;
 
-		if (HoveredCategoryName.Len() > 0)
+		if (!HoveredCategoryName.IsEmpty())
 		{
 			FFormatNamedArguments Args;
-			Args.Add(TEXT("DraggedCategory"), FText::FromString(DraggedCategory));
+			Args.Add(TEXT("DraggedCategory"), DraggedCategory);
 
-			if(HoveredCategoryName == DraggedCategory)
+			if(HoveredCategoryName.EqualTo(DraggedCategory))
 			{
 				StatusSymbol = FEditorStyle::GetBrush(TEXT("Graph.ConnectorFeedback.Error"));
 
@@ -86,7 +86,7 @@ public:
 			else
 			{
 				StatusSymbol = FEditorStyle::GetBrush(TEXT("Graph.ConnectorFeedback.OK"));
-				Args.Add(TEXT("HoveredCategory"), FText::FromString(HoveredCategoryName));
+				Args.Add(TEXT("HoveredCategory"), HoveredCategoryName);
 				Message = FText::Format( LOCTEXT("MoveCatOK", "Move category '{DraggedCategory}' before '{HoveredCategory}'"), Args );
 			}
 		}
@@ -99,7 +99,7 @@ public:
 		SetSimpleFeedbackMessage(StatusSymbol, FLinearColor::White, Message);
 	}
 	
-	virtual FReply DroppedOnCategory(FString OnCategory) override
+	virtual FReply DroppedOnCategory(FText OnCategory) override
 	{
 		// Get MyBlueprint via MyBlueprintPtr
 		TSharedPtr<SMyBlueprint> MyBlueprint = MyBlueprintPtr.Pin();
@@ -112,7 +112,7 @@ public:
 		return FReply::Handled();
 	}
 
-	static TSharedRef<FMyBlueprintCategoryDragDropAction> New(const FString& InCategory, TSharedPtr<SMyBlueprint> InMyBlueprint)
+	static TSharedRef<FMyBlueprintCategoryDragDropAction> New(const FText& InCategory, TSharedPtr<SMyBlueprint> InMyBlueprint)
 	{
 		TSharedRef<FMyBlueprintCategoryDragDropAction> Operation = MakeShareable(new FMyBlueprintCategoryDragDropAction);
 		Operation->DraggedCategory = InCategory;
@@ -122,7 +122,7 @@ public:
 	}
 
 	/** Category we were dragging */
-	FString DraggedCategory;
+	FText DraggedCategory;
 	/** MyBlueprint widget we dragged from */
 	TWeakPtr<SMyBlueprint>	MyBlueprintPtr;
 };
@@ -440,7 +440,7 @@ void SMyBlueprint::OnCategoryNameCommitted(const FText& InNewText, ETextCommit::
 					// If the variable is not native and it's outer is the skeleton generated class, we can rename the category
 					if(!bIsNativeVar && OuterClass == GetBlueprintObj()->SkeletonGeneratedClass)
 					{
-						FBlueprintEditorUtils::SetBlueprintVariableCategory(GetBlueprintObj(), VarAction->GetVariableName(), NULL, FName( *CategoryName.ToString() ), true);
+						FBlueprintEditorUtils::SetBlueprintVariableCategory(GetBlueprintObj(), VarAction->GetVariableName(), NULL, CategoryName, true);
 					}
 				}
 			}
@@ -448,12 +448,12 @@ void SMyBlueprint::OnCategoryNameCommitted(const FText& InNewText, ETextCommit::
 			{
 				FEdGraphSchemaAction_K2LocalVar* LocalVarAction = (FEdGraphSchemaAction_K2LocalVar*)Actions[i].Get();
 
-				FBlueprintEditorUtils::SetBlueprintVariableCategory(GetBlueprintObj(), LocalVarAction->GetVariableName(), LocalVarAction->GetVariableScope(), FName( *CategoryName.ToString() ), true);
+				FBlueprintEditorUtils::SetBlueprintVariableCategory(GetBlueprintObj(), LocalVarAction->GetVariableName(), LocalVarAction->GetVariableScope(), CategoryName, true);
 			}
 			else if (Actions[i]->GetTypeId() == FEdGraphSchemaAction_K2Delegate::StaticGetTypeId())
 			{
 				FEdGraphSchemaAction_K2Delegate* DelegateAction = (FEdGraphSchemaAction_K2Delegate*)Actions[i].Get();
-				FBlueprintEditorUtils::SetBlueprintVariableCategory(GetBlueprintObj(), DelegateAction->GetDelegatePoperty()->GetFName(), NULL, FName( *CategoryName.ToString() ), true);
+				FBlueprintEditorUtils::SetBlueprintVariableCategory(GetBlueprintObj(), DelegateAction->GetDelegatePoperty()->GetFName(), NULL, CategoryName, true);
 			}
 			else if (Actions[i]->GetTypeId() == FEdGraphSchemaAction_K2Graph::StaticGetTypeId())
 			{
@@ -471,11 +471,11 @@ void SMyBlueprint::OnCategoryNameCommitted(const FText& InNewText, ETextCommit::
 						EntryNode->Modify();
 						if (UK2Node_FunctionEntry* FunctionEntryNode = Cast<UK2Node_FunctionEntry>(EntryNode.Get()))
 						{
-							FunctionEntryNode->MetaData.Category = CategoryName.ToString();
+							FunctionEntryNode->MetaData.Category = CategoryName;
 						}
 						else if (UK2Node_Tunnel* TypedEntryNode = ExactCast<UK2Node_Tunnel>(EntryNode.Get()))
 						{
-							TypedEntryNode->MetaData.Category = CategoryName.ToString();
+							TypedEntryNode->MetaData.Category = CategoryName;
 						}
 
 						if(UFunction* Function = GetBlueprintObj()->SkeletonGeneratedClass->FindFunctionByName(GraphAction->EdGraph->GetFName()))
@@ -816,7 +816,7 @@ TSharedRef<SWidget> SMyBlueprint::OnCreateWidgetForAction(FCreateWidgetForAction
 	return BlueprintEditorPtr.IsValid() ? SNew(SBlueprintPaletteItem, InCreateData, BlueprintEditorPtr.Pin()) : SNew(SBlueprintPaletteItem, InCreateData, GetBlueprintObj());
 }
 
-void SMyBlueprint::GetChildGraphs(UEdGraph* InEdGraph, int32 const SectionId, FGraphActionListBuilderBase& OutAllActions, FString ParentCategory)
+void SMyBlueprint::GetChildGraphs(UEdGraph* InEdGraph, int32 const SectionId, FGraphActionListBuilderBase& OutAllActions, FText ParentCategory)
 {
 	// Grab subgraphs
 	const UEdGraphSchema* Schema = InEdGraph->GetSchema();
@@ -838,7 +838,12 @@ void SMyBlueprint::GetChildGraphs(UEdGraph* InEdGraph, int32 const SectionId, FG
 
 		FText DisplayText = ChildGraphDisplayInfo.DisplayName;
 
-		FString Category = ((ParentCategory.IsEmpty()) ? "" : ParentCategory + "|") + EdGraphDisplayName.ToString();
+		FText Category;
+		if (!ParentCategory.IsEmpty())
+		{
+			Category = FText::Format(FText::FromString(TEXT("{0}|{1}")), ParentCategory, EdGraphDisplayName);
+		}
+
 		FString ChildTooltip = DisplayText.ToString();
 		FText ChildDesc = DisplayText;
 		const FName DisplayName =  FName(*DisplayText.ToString());
@@ -857,7 +862,7 @@ void SMyBlueprint::GetChildGraphs(UEdGraph* InEdGraph, int32 const SectionId, FG
 struct FCreateEdGraphSchemaActionHelper
 {
 	template<class ActionType, class NodeType>
-	static void CreateAll(UEdGraph const* EdGraph, int32 SectionId, FGraphActionListBuilderBase& OutAllActions, const FString& ActionCategory)
+	static void CreateAll(UEdGraph const* EdGraph, int32 SectionId, FGraphActionListBuilderBase& OutAllActions, const FText& ActionCategory)
 	{
 		TArray<NodeType*> EventNodes;
 		EdGraph->GetNodesOfClass<NodeType>(EventNodes);
@@ -876,7 +881,7 @@ struct FCreateEdGraphSchemaActionHelper
 	}
 };
 
-void SMyBlueprint::GetChildEvents(UEdGraph const* InEdGraph, int32 const SectionId, FGraphActionListBuilderBase& OutAllActions, FString ParentCategory) const
+void SMyBlueprint::GetChildEvents(UEdGraph const* InEdGraph, int32 const SectionId, FGraphActionListBuilderBase& OutAllActions, FText ParentCategory) const
 {
 	if (!ensure(InEdGraph != NULL))
 	{
@@ -888,7 +893,11 @@ void SMyBlueprint::GetChildEvents(UEdGraph const* InEdGraph, int32 const Section
 	FGraphDisplayInfo EdGraphDisplayInfo;
 	Schema->GetGraphDisplayInformation(*InEdGraph, EdGraphDisplayInfo);
 	FText const EdGraphDisplayName = EdGraphDisplayInfo.DisplayName;
-	FString ActionCategory = ((ParentCategory.IsEmpty()) ? "" : ParentCategory + "|") + EdGraphDisplayName.ToString();
+	FText ActionCategory;
+	if (!ParentCategory.IsEmpty())
+	{
+		ActionCategory = FText::Format(FText::FromString(TEXT("{0}|{1}")), ParentCategory, EdGraphDisplayName);
+	}
 
 	FCreateEdGraphSchemaActionHelper::CreateAll<FEdGraphSchemaAction_K2Event, UK2Node_Event>(InEdGraph, SectionId, OutAllActions, ActionCategory);
 	FCreateEdGraphSchemaActionHelper::CreateAll<FEdGraphSchemaAction_K2InputAction, UK2Node_InputKey>(InEdGraph, SectionId, OutAllActions, ActionCategory);
@@ -911,16 +920,16 @@ void SMyBlueprint::GetLocalVariables(FGraphActionListBuilderBase& OutAllActions)
 		TopLevelGraph->GetNodesOfClass<UK2Node_FunctionEntry>(FunctionEntryNodes);
 
 		// Search in all FunctionEntry nodes for their local variables
-		FString ActionCategory;
+		FText ActionCategory;
 		const UEdGraphSchema_K2* K2Schema = GetDefault<UEdGraphSchema_K2>();
 		for (UK2Node_FunctionEntry* const FunctionEntry : FunctionEntryNodes)
 		{
 			for( const FBPVariableDescription& Variable : FunctionEntry->LocalVariables )
 			{
-				FString Category = Variable.Category.ToString();
-				if (Variable.Category == K2Schema->VR_DefaultCategory)
+				FText Category = Variable.Category;
+				if (Variable.Category.EqualTo(K2Schema->VR_DefaultCategory))
 				{
-					Category = FString();
+					Category = FText::GetEmpty();
 				}
 
 				UFunction* Func = FindField<UFunction>(GetBlueprintObj()->SkeletonGeneratedClass, TopLevelGraph->GetFName());
@@ -1073,14 +1082,14 @@ void SMyBlueprint::CollectAllActions(FGraphActionListBuilderBase& OutAllActions)
 		const FName PropertyName = Property->GetFName();
 		const FText PropertyDesc = FText::FromName(PropertyName);
 
-		FName CategoryName = FObjectEditorUtils::GetCategoryFName(Property);
-		FString PropertyCategory = FObjectEditorUtils::GetCategory(Property);
-		const FString UserCategoryName = FEditorCategoryUtils::GetCategoryDisplayString( PropertyCategory );
+		FText CategoryName = FObjectEditorUtils::GetCategoryText(Property);
+		FText PropertyCategory = FObjectEditorUtils::GetCategoryText(Property);
+		const FString UserCategoryName = FEditorCategoryUtils::GetCategoryDisplayString( PropertyCategory.ToString() );
 
-		if ((CategoryName == BlueprintObj->GetFName()) || (CategoryName == K2Schema->VR_DefaultCategory))
+		if (CategoryName.EqualTo(FText::FromString(BlueprintObj->GetName())) || CategoryName.EqualTo(K2Schema->VR_DefaultCategory))
 		{
-			CategoryName = NAME_None;		// default, so place in 'non' category
-			PropertyCategory = FString();
+			CategoryName = FText::GetEmpty();		// default, so place in 'non' category
+			PropertyCategory = FText::GetEmpty();
 		}
 
 		if (bShouldShowAsVar)
@@ -1088,9 +1097,9 @@ void SMyBlueprint::CollectAllActions(FGraphActionListBuilderBase& OutAllActions)
 			const bool bComponentProperty = Obj && Obj->PropertyClass ? Obj->PropertyClass->IsChildOf<UActorComponent>() : false;
 
 			// By default components go into the variable section under the component category unless a custom category is specified.
-			if ( bComponentProperty && CategoryName == NAME_None )
+			if ( bComponentProperty && CategoryName.IsEmpty() )
 			{
-				PropertyCategory = LOCTEXT("Components", "Components").ToString();
+				PropertyCategory = LOCTEXT("Components", "Components");
 			}
 
 			TSharedPtr<FEdGraphSchemaAction_K2Var> NewVarAction = MakeShareable(new FEdGraphSchemaAction_K2Var(PropertyCategory, PropertyDesc, PropertyTooltip.ToString(), 0));
@@ -1141,10 +1150,10 @@ void SMyBlueprint::CollectAllActions(FGraphActionListBuilderBase& OutAllActions)
 		FGraphDisplayInfo DisplayInfo;
 		Graph->GetSchema()->GetGraphDisplayInformation(*Graph, DisplayInfo);
 
-		FString FunctionCategory;
+		FText FunctionCategory;
 		if(UFunction* Function = BlueprintObj->SkeletonGeneratedClass->FindFunctionByName(Graph->GetFName()))
 		{
-			FunctionCategory = Function->GetMetaData(FBlueprintMetadata::MD_FunctionCategory);
+			FunctionCategory = Function->GetMetaDataText(FBlueprintMetadata::MD_FunctionCategory, TEXT("UObjectCategory"), Function->GetFullGroupName(false));
 		}
 
 		TSharedPtr<FEdGraphSchemaAction_K2Graph> NewFuncAction = MakeShareable(new FEdGraphSchemaAction_K2Graph(EEdGraphSchemaAction_K2Graph::Function, FunctionCategory, DisplayInfo.PlainName, DisplayInfo.Tooltip, bIsConstructionScript ? 2 : 1));
@@ -1176,7 +1185,7 @@ void SMyBlueprint::CollectAllActions(FGraphActionListBuilderBase& OutAllActions)
 		FGraphDisplayInfo DisplayInfo;
 		Graph->GetSchema()->GetGraphDisplayInformation(*Graph, DisplayInfo);
 
-		FString MacroCategory = GetGraphCategory(Graph);
+		FText MacroCategory = GetGraphCategory(Graph);
 
 		TSharedPtr<FEdGraphSchemaAction_K2Graph> NewMacroAction = MakeShareable(new FEdGraphSchemaAction_K2Graph(EEdGraphSchemaAction_K2Graph::Macro, MacroCategory, DisplayInfo.PlainName, DisplayInfo.Tooltip, 1));
 		NewMacroAction->SectionID = NodeSectionID::MACRO;
@@ -1214,7 +1223,7 @@ void SMyBlueprint::CollectAllActions(FGraphActionListBuilderBase& OutAllActions)
 				FunctionDesc = FText::FromString(Function->GetName());
 			}
 
-			FString FunctionCategory = Function->GetMetaData(FBlueprintMetadata::MD_FunctionCategory);
+			FText FunctionCategory = Function->GetMetaDataText(FBlueprintMetadata::MD_FunctionCategory, TEXT("UObjectCategory"), Function->GetFullGroupName(false));
 
 			TSharedPtr<FEdGraphSchemaAction_K2Graph> NewFuncAction = MakeShareable(new FEdGraphSchemaAction_K2Graph(EEdGraphSchemaAction_K2Graph::Function, FunctionCategory, FunctionDesc, FunctionTooltip, 1));
 			NewFuncAction->FuncName = FunctionName;
@@ -1238,10 +1247,10 @@ void SMyBlueprint::CollectAllActions(FGraphActionListBuilderBase& OutAllActions)
 			FString FunctionTooltip = FunctionName.ToString();
 			FString FunctionDesc = FunctionName.ToString();
 			
-			FString FunctionCategory;
+			FText FunctionCategory;
 			if ( UFunction* Function = BlueprintObj->SkeletonGeneratedClass->FindFunctionByName(Graph->GetFName()) )
 			{
-				FunctionCategory = Function->GetMetaData(FBlueprintMetadata::MD_FunctionCategory);
+				FunctionCategory = Function->GetMetaDataText(FBlueprintMetadata::MD_FunctionCategory, TEXT("UObjectCategory"), Function->GetFullGroupName(false));
 			}
 
 			TSharedPtr<FEdGraphSchemaAction_K2Graph> NewFuncAction = MakeShareable(new FEdGraphSchemaAction_K2Graph(EEdGraphSchemaAction_K2Graph::Interface, FunctionCategory, FText::FromString(FunctionDesc), FunctionTooltip, 1));
@@ -1274,7 +1283,7 @@ void SMyBlueprint::CollectAllActions(FGraphActionListBuilderBase& OutAllActions)
 						FString FunctionTooltip = Function->GetToolTipText().ToString();
 						FText FunctionDesc = K2Schema->GetFriendlySignatureName(Function);
 
-						FString FunctionCategory = Function->GetMetaData(FBlueprintMetadata::MD_FunctionCategory);
+						FText FunctionCategory = Function->GetMetaDataText(FBlueprintMetadata::MD_FunctionCategory, TEXT("UObjectCategory"), Function->GetFullGroupName(false));
 
 						TSharedPtr<FEdGraphSchemaAction_K2Graph> NewFuncAction = MakeShareable(new FEdGraphSchemaAction_K2Graph(EEdGraphSchemaAction_K2Graph::Function, FunctionCategory, FunctionDesc, FunctionTooltip, 1));
 						NewFuncAction->FuncName = FunctionName;
@@ -1294,7 +1303,7 @@ void SMyBlueprint::CollectAllActions(FGraphActionListBuilderBase& OutAllActions)
 		FGraphDisplayInfo DisplayInfo;
 		Graph->GetSchema()->GetGraphDisplayInformation(*Graph, DisplayInfo);
 
-		TSharedPtr<FEdGraphSchemaAction_K2Graph> NeUbergraphAction = MakeShareable(new FEdGraphSchemaAction_K2Graph(EEdGraphSchemaAction_K2Graph::Graph, FString(), DisplayInfo.PlainName, DisplayInfo.Tooltip, 2));
+		TSharedPtr<FEdGraphSchemaAction_K2Graph> NeUbergraphAction = MakeShareable(new FEdGraphSchemaAction_K2Graph(EEdGraphSchemaAction_K2Graph::Graph, FText::GetEmpty(), DisplayInfo.PlainName, DisplayInfo.Tooltip, 2));
 		NeUbergraphAction->FuncName = Graph->GetFName();
 		NeUbergraphAction->EdGraph = Graph;
 		NeUbergraphAction->SectionID = NodeSectionID::GRAPH;
@@ -1313,7 +1322,7 @@ void SMyBlueprint::CollectAllActions(FGraphActionListBuilderBase& OutAllActions)
 		const FName IntermediateName(*(FString(TEXT("$INTERMEDIATE$_")) + Graph->GetName()));
 		FString IntermediateTooltip = IntermediateName.ToString();
 		FString IntermediateDesc = IntermediateName.ToString();
-		TSharedPtr<FEdGraphSchemaAction_K2Graph> NewIntermediateAction = MakeShareable(new FEdGraphSchemaAction_K2Graph(EEdGraphSchemaAction_K2Graph::Graph, FString(), FText::FromString(IntermediateDesc), IntermediateTooltip, 1));
+		TSharedPtr<FEdGraphSchemaAction_K2Graph> NewIntermediateAction = MakeShareable(new FEdGraphSchemaAction_K2Graph(EEdGraphSchemaAction_K2Graph::Graph, FText::GetEmpty(), FText::FromString(IntermediateDesc), IntermediateTooltip, 1));
 		NewIntermediateAction->FuncName = IntermediateName;
 		NewIntermediateAction->EdGraph = Graph;
 		OutAllActions.AddAction(NewIntermediateAction);
@@ -1496,7 +1505,7 @@ FReply SMyBlueprint::OnActionDragged( const TArray< TSharedPtr<FEdGraphSchemaAct
 	return FReply::Unhandled();
 }
 
-FReply SMyBlueprint::OnCategoryDragged(const FString& InCategory, const FPointerEvent& MouseEvent)
+FReply SMyBlueprint::OnCategoryDragged(const FText& InCategory, const FPointerEvent& MouseEvent)
 {
 	TSharedRef<FMyBlueprintCategoryDragDropAction> DragOperation = FMyBlueprintCategoryDragDropAction::New(InCategory, SharedThis(this));
 	return FReply::Handled().BeginDragDrop(DragOperation);
@@ -1746,9 +1755,9 @@ bool SMyBlueprint::SelectionHasContextMenu() const
 	return SelectedActions.Num() > 0;
 }
 
-FString SMyBlueprint::GetGraphCategory(UEdGraph* InGraph) const
+FText SMyBlueprint::GetGraphCategory(UEdGraph* InGraph) const
 {
-	FString ReturnCategory;
+	FText ReturnCategory;
 
 	// Pull the category from the required metadata based on the types of nodes we can discover in the graph
 	TWeakObjectPtr<UK2Node_EditablePinBase> EntryNode;
@@ -1765,9 +1774,9 @@ FString SMyBlueprint::GetGraphCategory(UEdGraph* InGraph) const
 
 	// Empty the category if it's default, we don't want to display the "default" category and items will just appear without a category
 	const UEdGraphSchema_K2* K2Schema = GetDefault<UEdGraphSchema_K2>();
-	if(ReturnCategory == K2Schema->VR_DefaultCategory.ToString())
+	if(ReturnCategory.EqualTo(K2Schema->VR_DefaultCategory))
 	{
-		ReturnCategory.Empty();
+		ReturnCategory = FText::GetEmpty();
 	}
 
 	return ReturnCategory;
@@ -2623,23 +2632,25 @@ void SMyBlueprint::ClearGraphActionMenuSelection()
 	GraphActionMenu->SelectItemByName(NAME_None);
 }
 
-void SMyBlueprint::ExpandCategory(const FString& CategoryName)
+void SMyBlueprint::ExpandCategory(const FText& CategoryName)
 {
 	GraphActionMenu->ExpandCategory(CategoryName);
 }
 
-bool SMyBlueprint::MoveCategoryBeforeCategory( const FString& InCategoryToMove, const FString& InTargetCategory )
+bool SMyBlueprint::MoveCategoryBeforeCategory( const FText& InCategoryToMove, const FText& InTargetCategory )
 {
 	bool bResult = false;
 	UBlueprint* BlueprintObj = BlueprintEditorPtr.Pin()->GetBlueprintObj();
 
+	FString CategoryToMoveString = InCategoryToMove.ToString();
+	FString TargetCategoryString = InTargetCategory.ToString();
 	if( BlueprintObj )
 	{
 		// Find root categories
-		int32 RootCategoryDelim = InCategoryToMove.Find( TEXT( "|" ));
-		FName CategoryToMove = RootCategoryDelim == INDEX_NONE ? *InCategoryToMove : *InCategoryToMove.Left( RootCategoryDelim );
-		RootCategoryDelim = InTargetCategory.Find( TEXT( "|" ));
-		FName TargetCategory = RootCategoryDelim == INDEX_NONE ? *InTargetCategory : *InTargetCategory.Left( RootCategoryDelim );
+		int32 RootCategoryDelim = CategoryToMoveString.Find( TEXT( "|" ));
+		FName CategoryToMove = RootCategoryDelim == INDEX_NONE ? *CategoryToMoveString : *CategoryToMoveString.Left( RootCategoryDelim );
+		RootCategoryDelim = TargetCategoryString.Find( TEXT( "|" ));
+		FName TargetCategory = RootCategoryDelim == INDEX_NONE ? *TargetCategoryString : *TargetCategoryString.Left( RootCategoryDelim );
 
 		TArray<FName>& CategorySort = BlueprintObj->CategorySorting;
 		const int32 RemovalIndex = CategorySort.Find( CategoryToMove );
diff --git a/Engine/Source/Editor/Kismet/Private/SMyBlueprint.h b/Engine/Source/Editor/Kismet/Private/SMyBlueprint.h
--- a/Engine/Source/Editor/Kismet/Private/SMyBlueprint.h
+++ b/Engine/Source/Editor/Kismet/Private/SMyBlueprint.h
@@ -130,10 +130,10 @@ public:
 	void OnRequestRenameOnActionNode();
 
 	/** Expands any category with the associated name */
-	void ExpandCategory(const FString& CategoryName);
+	void ExpandCategory(const FText& CategoryName);
 
 	/** Move the category before the target category */
-	bool MoveCategoryBeforeCategory( const FString& CategoryToMove, const FString& TargetCategory );
+	bool MoveCategoryBeforeCategory( const FText& CategoryToMove, const FText& TargetCategory );
 private:
 	/** Creates widgets for the graph schema actions */
 	TSharedRef<SWidget> OnCreateWidgetForAction(struct FCreateWidgetForActionData* const InCreateData);
@@ -141,8 +141,8 @@ private:
 	/** Callback used to populate all actions list in SGraphActionMenu */
 	void CollectAllActions(FGraphActionListBuilderBase& OutAllActions);
 	void CollectStaticSections(TArray<int32>& StaticSectionIDs);
-	void GetChildGraphs(UEdGraph* EdGraph, int32 const SectionId, FGraphActionListBuilderBase& OutAllActions, FString ParentCategory = FString());
-	void GetChildEvents(UEdGraph const* EdGraph, int32 const SectionId, FGraphActionListBuilderBase& OutAllActions, FString ParentCategory = FString()) const;
+	void GetChildGraphs(UEdGraph* EdGraph, int32 const SectionId, FGraphActionListBuilderBase& OutAllActions, FText ParentCategory = FText::GetEmpty());
+	void GetChildEvents(UEdGraph const* EdGraph, int32 const SectionId, FGraphActionListBuilderBase& OutAllActions, FText ParentCategory = FText::GetEmpty()) const;
 	void GetLocalVariables(FGraphActionListBuilderBase& OutAllActions) const;
 	
 	/** Handles the visibility of the local action list */
@@ -150,7 +150,7 @@ private:
 
 	/** Callbacks for the graph action menu */
 	FReply OnActionDragged(const TArray< TSharedPtr<FEdGraphSchemaAction> >& InActions, const FPointerEvent& MouseEvent);
-	FReply OnCategoryDragged(const FString& InCategory, const FPointerEvent& MouseEvent);
+	FReply OnCategoryDragged(const FText& InCategory, const FPointerEvent& MouseEvent);
 	void OnActionSelected(const TArray< TSharedPtr<FEdGraphSchemaAction> >& InActions);
 	static void OnActionSelectedHelper(TSharedPtr<FEdGraphSchemaAction> InAction, UBlueprint* Blueprint, TSharedRef<SKismetInspector> Inspector);
 	void OnGlobalActionSelected(const TArray< TSharedPtr<FEdGraphSchemaAction> >& InActions, ESelectInfo::Type InSelectionType);
@@ -218,7 +218,7 @@ private:
 	void UpdateNodeCreation();
 
 	/** Returns the displayed category, if any, of a graph */
-	FString GetGraphCategory(UEdGraph* InGraph) const;
+	FText GetGraphCategory(UEdGraph* InGraph) const;
 
 	/** Helper function to delete a graph in the MyBlueprint window */
 	void OnDeleteGraph(UEdGraph* InGraph, EEdGraphSchemaAction_K2Graph::Type);
diff --git a/Engine/Source/Editor/KismetCompiler/Private/KismetCompiler.cpp b/Engine/Source/Editor/KismetCompiler/Private/KismetCompiler.cpp
--- a/Engine/Source/Editor/KismetCompiler/Private/KismetCompiler.cpp
+++ b/Engine/Source/Editor/KismetCompiler/Private/KismetCompiler.cpp
@@ -620,9 +620,9 @@ void FKismetCompilerContext::CreateClassVariablesFromBlueprint()
 				UProperty* NewProperty = CreateVariable(VarName, Type);
 				if (NewProperty != NULL)
 				{
-					const FString CategoryName = Node->CategoryName != NAME_None ? Node->CategoryName.ToString() : Blueprint->GetName();
+					const FText CategoryName = Node->CategoryName.IsEmpty() ? FText::FromString(Blueprint->GetName()) : Node->CategoryName ;
 					
-					NewProperty->SetMetaData(TEXT("Category"), *CategoryName);
+					NewProperty->SetMetaData(TEXT("Category"), *CategoryName.ToString());
 					NewProperty->SetPropertyFlags(CPF_BlueprintVisible | CPF_NonTransactional );
 				}
 			}
@@ -1218,7 +1218,7 @@ void FKismetCompilerContext::PrecompileFunction(FKismetFunctionContext& Context)
 		FKismetUserDeclaredFunctionMetadata& FunctionMetaData = Context.EntryPoint->MetaData;
 		if( !FunctionMetaData.Category.IsEmpty() )
 		{
-			Context.Function->SetMetaData(FBlueprintMetadata::MD_FunctionCategory, *FunctionMetaData.Category);
+			Context.Function->SetMetaData(FBlueprintMetadata::MD_FunctionCategory, *FunctionMetaData.Category.ToString());
 		}
 		// Set as blutility function
 		if( FunctionMetaData.bCallInEditor )
diff --git a/Engine/Source/Editor/MaterialEditor/Private/MaterialEditor.h b/Engine/Source/Editor/MaterialEditor/Private/MaterialEditor.h
--- a/Engine/Source/Editor/MaterialEditor/Private/MaterialEditor.h
+++ b/Engine/Source/Editor/MaterialEditor/Private/MaterialEditor.h
@@ -161,7 +161,7 @@ struct FMaterialExpression
 /** Static array of categorized material expression classes. */
 struct FCategorizedMaterialExpressionNode
 {
-	FString	CategoryName;
+	FText	CategoryName;
 	TArray<FMaterialExpression> MaterialExpressions;
 };
 
diff --git a/Engine/Source/Editor/MaterialEditor/Private/MaterialEditorUtilities.cpp b/Engine/Source/Editor/MaterialEditor/Private/MaterialEditorUtilities.cpp
--- a/Engine/Source/Editor/MaterialEditor/Private/MaterialEditorUtilities.cpp
+++ b/Engine/Source/Editor/MaterialEditor/Private/MaterialEditorUtilities.cpp
@@ -145,13 +145,13 @@ void FMaterialEditorUtilities::GetMaterialExpressionActions(FGraphActionMenuBuil
 
 	if (bUseUnsortedMenus)
 	{
-		AddMaterialExpressionCategory(ActionMenuBuilder, TEXT(""), &ExpressionClasses->AllExpressionClasses, bMaterialFunction);
+		AddMaterialExpressionCategory(ActionMenuBuilder, FText::GetEmpty(), &ExpressionClasses->AllExpressionClasses, bMaterialFunction);
 	}
 	else
 	{
 		// Add Favourite expressions as a category
 		const FText FavouritesCategory = LOCTEXT("FavoritesMenu", "Favorites");
-		AddMaterialExpressionCategory(ActionMenuBuilder, FavouritesCategory.ToString(), &ExpressionClasses->FavoriteExpressionClasses, bMaterialFunction);
+		AddMaterialExpressionCategory(ActionMenuBuilder, FavouritesCategory, &ExpressionClasses->FavoriteExpressionClasses, bMaterialFunction);
 
 		// Add each category to the menu
 		for (int32 CategoryIndex = 0; CategoryIndex < ExpressionClasses->CategorizedExpressionClasses.Num(); ++CategoryIndex)
@@ -162,7 +162,7 @@ void FMaterialEditorUtilities::GetMaterialExpressionActions(FGraphActionMenuBuil
 
 		if (ExpressionClasses->UnassignedExpressionClasses.Num() > 0)
 		{
-			AddMaterialExpressionCategory(ActionMenuBuilder, TEXT(""), &ExpressionClasses->UnassignedExpressionClasses, bMaterialFunction);
+			AddMaterialExpressionCategory(ActionMenuBuilder, FText::GetEmpty(), &ExpressionClasses->UnassignedExpressionClasses, bMaterialFunction);
 		}
 	}
 }
@@ -545,7 +545,7 @@ TSharedPtr<class IMaterialEditor> FMaterialEditorUtilities::GetIMaterialEditorFo
 	return MaterialEditor;
 }
 
-void FMaterialEditorUtilities::AddMaterialExpressionCategory(FGraphActionMenuBuilder& ActionMenuBuilder, FString CategoryName, TArray<struct FMaterialExpression>* MaterialExpressions, bool bMaterialFunction)
+void FMaterialEditorUtilities::AddMaterialExpressionCategory(FGraphActionMenuBuilder& ActionMenuBuilder, FText CategoryName, TArray<struct FMaterialExpression>* MaterialExpressions, bool bMaterialFunction)
 {
 	// Get type of dragged pin
 	uint32 FromPinType = 0;
diff --git a/Engine/Source/Editor/MaterialEditor/Private/MaterialExpressionClasses.cpp b/Engine/Source/Editor/MaterialEditor/Private/MaterialExpressionClasses.cpp
--- a/Engine/Source/Editor/MaterialEditor/Private/MaterialExpressionClasses.cpp
+++ b/Engine/Source/Editor/MaterialEditor/Private/MaterialExpressionClasses.cpp
@@ -39,14 +39,14 @@ const UStruct* MaterialExpressionClasses::GetExpressionInputStruct()
 	return ExpressionInputStruct;
 }
 
-FCategorizedMaterialExpressionNode* MaterialExpressionClasses::GetCategoryNode(const FString& InCategoryName, bool bCreate)
+FCategorizedMaterialExpressionNode* MaterialExpressionClasses::GetCategoryNode(const FText& InCategoryName, bool bCreate)
 {
 	for (int32 CheckIndex = 0; CheckIndex < CategorizedExpressionClasses.Num(); CheckIndex++)
 	{
 		FCategorizedMaterialExpressionNode* CheckNode = &(CategorizedExpressionClasses[CheckIndex]);
 		if (CheckNode)
 		{
-			if (CheckNode->CategoryName == InCategoryName)
+			if (CheckNode->CategoryName.EqualTo(InCategoryName))
 			{
 				return CheckNode;
 			}
@@ -158,7 +158,7 @@ void MaterialExpressionClasses::InitMaterialExpressionClasses()
 		{
 			FORCEINLINE bool operator()( const FCategorizedMaterialExpressionNode& A, const FCategorizedMaterialExpressionNode& B ) const
 			{
-				return A.CategoryName < B.CategoryName;
+				return A.CategoryName.CompareTo(B.CategoryName) < 0;
 			}
 		};
 		CategorizedExpressionClasses.Sort( FCompareFCategorizedMaterialExpressionNode() );
diff --git a/Engine/Source/Editor/MaterialEditor/Private/MaterialExpressionClasses.h b/Engine/Source/Editor/MaterialEditor/Private/MaterialExpressionClasses.h
--- a/Engine/Source/Editor/MaterialEditor/Private/MaterialExpressionClasses.h
+++ b/Engine/Source/Editor/MaterialEditor/Private/MaterialExpressionClasses.h
@@ -55,7 +55,7 @@ private:
 	 *
 	 *	@return	The category node.
 	 */
-	struct FCategorizedMaterialExpressionNode* GetCategoryNode(const FString& InCategoryName, bool bCreate);
+	struct FCategorizedMaterialExpressionNode* GetCategoryNode(const FText& InCategoryName, bool bCreate);
 
 	/**
 	 * Initializes the list of UMaterialExpression-derived classes shared between all material editor instances.
diff --git a/Engine/Source/Editor/MaterialEditor/Public/MaterialEditorUtilities.h b/Engine/Source/Editor/MaterialEditor/Public/MaterialEditorUtilities.h
--- a/Engine/Source/Editor/MaterialEditor/Public/MaterialEditorUtilities.h
+++ b/Engine/Source/Editor/MaterialEditor/Public/MaterialEditorUtilities.h
@@ -188,7 +188,7 @@ private:
 	 * @param	MaterialExpressions	List of Material Expressions in the category.
 	 * @param	bMaterialFunction	Whether we are building for a material function.
 	 */
-	static void AddMaterialExpressionCategory(FGraphActionMenuBuilder& ActionMenuBuilder, FString CategoryName, TArray<struct FMaterialExpression>* MaterialExpressions, bool bMaterialFunction);
+	static void AddMaterialExpressionCategory(FGraphActionMenuBuilder& ActionMenuBuilder, FText CategoryName, TArray<struct FMaterialExpression>* MaterialExpressions, bool bMaterialFunction);
 
 	/**
 	 * Checks whether a Material Expression class has any connections that are compatible with a type/direction
diff --git a/Engine/Source/Editor/UMGEditor/Private/Customizations/UMGDetailCustomizations.cpp b/Engine/Source/Editor/UMGEditor/Private/Customizations/UMGDetailCustomizations.cpp
--- a/Engine/Source/Editor/UMGEditor/Private/Customizations/UMGDetailCustomizations.cpp
+++ b/Engine/Source/Editor/UMGEditor/Private/Customizations/UMGDetailCustomizations.cpp
@@ -171,7 +171,7 @@ void FBlueprintWidgetCustomization::CreateMulticastEventCustomization(IDetailLay
 	}
 
 	// Add on category for delegate property
-	const FString EventCategory = FObjectEditorUtils::GetCategory(DelegateProperty);
+	const FText EventCategory = FObjectEditorUtils::GetCategoryText(DelegateProperty);
 
 	UObjectProperty* ComponentProperty = FindField<UObjectProperty>(Blueprint->SkeletonGeneratedClass, ThisComponentName);
 
diff --git a/Engine/Source/Editor/UnrealEd/Classes/MaterialGraph/MaterialGraphSchema.h b/Engine/Source/Editor/UnrealEd/Classes/MaterialGraph/MaterialGraphSchema.h
--- a/Engine/Source/Editor/UnrealEd/Classes/MaterialGraph/MaterialGraphSchema.h
+++ b/Engine/Source/Editor/UnrealEd/Classes/MaterialGraph/MaterialGraphSchema.h
@@ -23,7 +23,7 @@ struct UNREALED_API FMaterialGraphSchemaAction_NewNode : public FEdGraphSchemaAc
 		, MaterialExpressionClass(NULL)
 	{}
 
-	FMaterialGraphSchemaAction_NewNode(const FString& InNodeCategory, const FText& InMenuDesc, const FString& InToolTip, const int32 InGrouping)
+	FMaterialGraphSchemaAction_NewNode(const FText& InNodeCategory, const FText& InMenuDesc, const FString& InToolTip, const int32 InGrouping)
 		: FEdGraphSchemaAction(InNodeCategory, InMenuDesc, InToolTip, InGrouping) 
 		, MaterialExpressionClass(NULL)
 	{}
@@ -59,7 +59,7 @@ struct UNREALED_API FMaterialGraphSchemaAction_NewFunctionCall : public FEdGraph
 		: FEdGraphSchemaAction()
 	{}
 
-	FMaterialGraphSchemaAction_NewFunctionCall(const FString& InNodeCategory, const FText& InMenuDesc, const FString& InToolTip, const int32 InGrouping)
+	FMaterialGraphSchemaAction_NewFunctionCall(const FText& InNodeCategory, const FText& InMenuDesc, const FString& InToolTip, const int32 InGrouping)
 		: FEdGraphSchemaAction(InNodeCategory, InMenuDesc, InToolTip, InGrouping)
 	{}
 
@@ -82,7 +82,7 @@ struct UNREALED_API FMaterialGraphSchemaAction_NewComment : public FEdGraphSchem
 		: FEdGraphSchemaAction()
 	{}
 
-	FMaterialGraphSchemaAction_NewComment(const FString& InNodeCategory, const FText& InMenuDesc, const FString& InToolTip, const int32 InGrouping)
+	FMaterialGraphSchemaAction_NewComment(const FText& InNodeCategory, const FText& InMenuDesc, const FString& InToolTip, const int32 InGrouping)
 		: FEdGraphSchemaAction(InNodeCategory, InMenuDesc, InToolTip, InGrouping)
 	{}
 
@@ -105,7 +105,7 @@ struct UNREALED_API FMaterialGraphSchemaAction_Paste : public FEdGraphSchemaActi
 		: FEdGraphSchemaAction()
 	{}
 
-	FMaterialGraphSchemaAction_Paste(const FString& InNodeCategory, const FText& InMenuDesc, const FString& InToolTip, const int32 InGrouping)
+	FMaterialGraphSchemaAction_Paste(const FText& InNodeCategory, const FText& InMenuDesc, const FString& InToolTip, const int32 InGrouping)
 		: FEdGraphSchemaAction(InNodeCategory, InMenuDesc, InToolTip, InGrouping)
 	{}
 
diff --git a/Engine/Source/Editor/UnrealEd/Classes/Niagara/EdGraphSchema_Niagara.h b/Engine/Source/Editor/UnrealEd/Classes/Niagara/EdGraphSchema_Niagara.h
--- a/Engine/Source/Editor/UnrealEd/Classes/Niagara/EdGraphSchema_Niagara.h
+++ b/Engine/Source/Editor/UnrealEd/Classes/Niagara/EdGraphSchema_Niagara.h
@@ -21,7 +21,7 @@ struct UNREALED_API FNiagaraSchemaAction_NewNode : public FEdGraphSchemaAction
 		, NodeTemplate(NULL)
 	{}
 
-	FNiagaraSchemaAction_NewNode(const FString& InNodeCategory, const FText& InMenuDesc, const FString& InToolTip, const int32 InGrouping)
+	FNiagaraSchemaAction_NewNode(const FText& InNodeCategory, const FText& InMenuDesc, const FString& InToolTip, const int32 InGrouping)
 		: FEdGraphSchemaAction(InNodeCategory, InMenuDesc, InToolTip, InGrouping) 
 		, NodeTemplate(NULL)
 	{}
diff --git a/Engine/Source/Editor/UnrealEd/Classes/SoundClassGraph/SoundClassGraphSchema.h b/Engine/Source/Editor/UnrealEd/Classes/SoundClassGraph/SoundClassGraphSchema.h
--- a/Engine/Source/Editor/UnrealEd/Classes/SoundClassGraph/SoundClassGraphSchema.h
+++ b/Engine/Source/Editor/UnrealEd/Classes/SoundClassGraph/SoundClassGraphSchema.h
@@ -17,7 +17,7 @@ struct UNREALED_API FSoundClassGraphSchemaAction_NewNode : public FEdGraphSchema
 		, NewSoundClassName(TEXT("ClassName"))
 	{}
 
-	FSoundClassGraphSchemaAction_NewNode(const FString& InNodeCategory, const FText& InMenuDesc, const FString& InToolTip, const int32 InGrouping)
+	FSoundClassGraphSchemaAction_NewNode(const FText& InNodeCategory, const FText& InMenuDesc, const FString& InToolTip, const int32 InGrouping)
 		: FEdGraphSchemaAction(InNodeCategory, InMenuDesc, InToolTip, InGrouping)
 		, NewSoundClassName(TEXT("ClassName"))
 	{}
diff --git a/Engine/Source/Editor/UnrealEd/Classes/SoundCueGraph/SoundCueGraphSchema.h b/Engine/Source/Editor/UnrealEd/Classes/SoundCueGraph/SoundCueGraphSchema.h
--- a/Engine/Source/Editor/UnrealEd/Classes/SoundCueGraph/SoundCueGraphSchema.h
+++ b/Engine/Source/Editor/UnrealEd/Classes/SoundCueGraph/SoundCueGraphSchema.h
@@ -23,7 +23,7 @@ struct UNREALED_API FSoundCueGraphSchemaAction_NewNode : public FEdGraphSchemaAc
 		, SoundNodeClass(NULL)
 	{}
 
-	FSoundCueGraphSchemaAction_NewNode(const FString& InNodeCategory, const FText& InMenuDesc, const FString& InToolTip, const int32 InGrouping)
+	FSoundCueGraphSchemaAction_NewNode(const FText& InNodeCategory, const FText& InMenuDesc, const FString& InToolTip, const int32 InGrouping)
 		: FEdGraphSchemaAction(InNodeCategory, InMenuDesc, InToolTip, InGrouping) 
 		, SoundNodeClass(NULL)
 	{}
@@ -47,7 +47,7 @@ struct UNREALED_API FSoundCueGraphSchemaAction_NewFromSelected : public FSoundCu
 		: FSoundCueGraphSchemaAction_NewNode()
 	{}
 
-	FSoundCueGraphSchemaAction_NewFromSelected(const FString& InNodeCategory, const FText& InMenuDesc, const FString& InToolTip, const int32 InGrouping)
+	FSoundCueGraphSchemaAction_NewFromSelected(const FText& InNodeCategory, const FText& InMenuDesc, const FString& InToolTip, const int32 InGrouping)
 		: FSoundCueGraphSchemaAction_NewNode(InNodeCategory, InMenuDesc, InToolTip, InGrouping) 
 	{}
 
@@ -66,7 +66,7 @@ struct UNREALED_API FSoundCueGraphSchemaAction_NewComment : public FEdGraphSchem
 		: FEdGraphSchemaAction()
 	{}
 
-	FSoundCueGraphSchemaAction_NewComment(const FString& InNodeCategory, const FText& InMenuDesc, const FString& InToolTip, const int32 InGrouping)
+	FSoundCueGraphSchemaAction_NewComment(const FText& InNodeCategory, const FText& InMenuDesc, const FString& InToolTip, const int32 InGrouping)
 		: FEdGraphSchemaAction(InNodeCategory, InMenuDesc, InToolTip, InGrouping)
 	{}
 
@@ -85,7 +85,7 @@ struct UNREALED_API FSoundCueGraphSchemaAction_Paste : public FEdGraphSchemaActi
 		: FEdGraphSchemaAction()
 	{}
 
-	FSoundCueGraphSchemaAction_Paste(const FString& InNodeCategory, const FText& InMenuDesc, const FString& InToolTip, const int32 InGrouping)
+	FSoundCueGraphSchemaAction_Paste(const FText& InNodeCategory, const FText& InMenuDesc, const FString& InToolTip, const int32 InGrouping)
 		: FEdGraphSchemaAction(InNodeCategory, InMenuDesc, InToolTip, InGrouping)
 	{}
 
diff --git a/Engine/Source/Editor/UnrealEd/Private/Kismet2/BlueprintEditorUtils.cpp b/Engine/Source/Editor/UnrealEd/Private/Kismet2/BlueprintEditorUtils.cpp
--- a/Engine/Source/Editor/UnrealEd/Private/Kismet2/BlueprintEditorUtils.cpp
+++ b/Engine/Source/Editor/UnrealEd/Private/Kismet2/BlueprintEditorUtils.cpp
@@ -3139,11 +3139,11 @@ bool FBlueprintEditorUtils::MoveVariableBeforeVariable(UBlueprint* Blueprint, FN
 	return bMoved;
 }
 
-int32 FBlueprintEditorUtils::FindFirstNewVarOfCategory(const UBlueprint* Blueprint, FName Category)
+int32 FBlueprintEditorUtils::FindFirstNewVarOfCategory(const UBlueprint* Blueprint, FText Category)
 {
 	for(int32 VarIdx=0; VarIdx<Blueprint->NewVariables.Num(); VarIdx++)
 	{
-		if(Blueprint->NewVariables[VarIdx].Category == Category)
+		if(Blueprint->NewVariables[VarIdx].Category.EqualTo(Category))
 		{
 			return VarIdx;
 		}
@@ -3516,15 +3516,15 @@ void FBlueprintEditorUtils::RemoveBlueprintVariableMetaData(UBlueprint* Blueprin
 	FBlueprintEditorUtils::MarkBlueprintAsStructurallyModified(Blueprint);
 }
 
-void FBlueprintEditorUtils::SetBlueprintVariableCategory(UBlueprint* Blueprint, const FName& VarName, const UStruct* InLocalVarScope, const FName& NewCategory, bool bDontRecompile)
+void FBlueprintEditorUtils::SetBlueprintVariableCategory(UBlueprint* Blueprint, const FName& VarName, const UStruct* InLocalVarScope, const FText& NewCategory, bool bDontRecompile)
 {
 	const FScopedTransaction Transaction( LOCTEXT("ChangeVariableCategory", "Change Variable Category") );
 	Blueprint->Modify();
 
 	// Ensure we always set a category
 	const UEdGraphSchema_K2* K2Schema = GetDefault<UEdGraphSchema_K2>();
-	FName SetCategory = NewCategory;
-	if (SetCategory == NAME_None)
+	FText SetCategory = NewCategory;
+	if (SetCategory.IsEmpty())
 	{
 		SetCategory = K2Schema->VR_DefaultCategory; 
 	}
@@ -3543,7 +3543,7 @@ void FBlueprintEditorUtils::SetBlueprintVariableCategory(UBlueprint* Blueprint,
 			const int32 VarIndex = FBlueprintEditorUtils::FindNewVariableIndex(Blueprint, VarName);
 			if (VarIndex != INDEX_NONE)
 			{
-				bIsCategoryChanged = Blueprint->NewVariables[VarIndex].Category != SetCategory;
+				bIsCategoryChanged = !Blueprint->NewVariables[VarIndex].Category.EqualTo(SetCategory);
 				
 				if(bIsCategoryChanged)
 				{
@@ -3555,7 +3555,7 @@ void FBlueprintEditorUtils::SetBlueprintVariableCategory(UBlueprint* Blueprint,
 				const int32 SCS_NodeIndex = FBlueprintEditorUtils::FindSCS_Node(Blueprint, VarName);
 				if (SCS_NodeIndex != INDEX_NONE)
 				{
-					bIsCategoryChanged = Blueprint->SimpleConstructionScript->GetAllNodes()[SCS_NodeIndex]->CategoryName != SetCategory;
+					bIsCategoryChanged = !Blueprint->SimpleConstructionScript->GetAllNodes()[SCS_NodeIndex]->CategoryName.EqualTo(SetCategory);
 					
 					if(bIsCategoryChanged)
 					{
@@ -3577,7 +3577,7 @@ void FBlueprintEditorUtils::SetBlueprintVariableCategory(UBlueprint* Blueprint,
 		if(FBPVariableDescription* LocalVariable = FindLocalVariable(Blueprint, InLocalVarScope, VarName, &OutFunctionEntryNode))
 		{
 			// If the category does not change, we will not recompile the Blueprint
-			bool bIsCategoryChanged = LocalVariable->Category != SetCategory;
+			bool bIsCategoryChanged = !LocalVariable->Category.EqualTo(SetCategory);
 
 			if(bIsCategoryChanged)
 			{
@@ -3594,31 +3594,31 @@ void FBlueprintEditorUtils::SetBlueprintVariableCategory(UBlueprint* Blueprint,
 	}
 }
 
-FName FBlueprintEditorUtils::GetBlueprintVariableCategory(UBlueprint* Blueprint, const FName& VarName, const UStruct* InLocalVarScope)
+FText FBlueprintEditorUtils::GetBlueprintVariableCategory(UBlueprint* Blueprint, const FName& VarName, const UStruct* InLocalVarScope)
 {
-	FName CategoryName = NAME_None;
+	FText CategoryName;
 	UClass* SkeletonGeneratedClass = Blueprint->SkeletonGeneratedClass;
 	UProperty* TargetProperty = FindField<UProperty>(SkeletonGeneratedClass, VarName);
 	if(TargetProperty != NULL)
 	{
-		CategoryName = FObjectEditorUtils::GetCategoryFName(TargetProperty);
+		CategoryName = FObjectEditorUtils::GetCategoryText(TargetProperty);
 	}
 	else if(InLocalVarScope)
 	{
 		// Check to see if it is a local variable
 		if(FBPVariableDescription* LocalVariable = FindLocalVariable(Blueprint, InLocalVarScope, VarName))
 		{
-			return LocalVariable->Category;
+			CategoryName = LocalVariable->Category;
 		}
 	}
 
-	if(CategoryName == NAME_None && Blueprint->SimpleConstructionScript != NULL)
+	if(CategoryName.IsEmpty() && Blueprint->SimpleConstructionScript != NULL)
 	{
 		// Look for the variable in the SCS (in case the Blueprint has not been compiled yet)
 		const int32 SCS_NodeIndex = FBlueprintEditorUtils::FindSCS_Node(Blueprint, VarName);
 		if (SCS_NodeIndex != INDEX_NONE)
 		{
-			return Blueprint->SimpleConstructionScript->GetAllNodes()[SCS_NodeIndex]->CategoryName;
+			CategoryName = Blueprint->SimpleConstructionScript->GetAllNodes()[SCS_NodeIndex]->CategoryName;
 		}
 	}
 
diff --git a/Engine/Source/Editor/UnrealEd/Private/MaterialGraphSchema.cpp b/Engine/Source/Editor/UnrealEd/Private/MaterialGraphSchema.cpp
--- a/Engine/Source/Editor/UnrealEd/Private/MaterialGraphSchema.cpp
+++ b/Engine/Source/Editor/UnrealEd/Private/MaterialGraphSchema.cpp
@@ -329,7 +329,7 @@ void UMaterialGraphSchema::GetGraphContextActions(FGraphContextMenuBuilder& Cont
 	{
 		const FText PasteDesc = LOCTEXT("PasteDesc", "Paste Here");
 		const FText PasteToolTip = LOCTEXT("PasteToolTip", "Pastes copied items at this location.");
-		TSharedPtr<FMaterialGraphSchemaAction_Paste> PasteAction(new FMaterialGraphSchemaAction_Paste(TEXT(""), PasteDesc, PasteToolTip.ToString(), 0));
+		TSharedPtr<FMaterialGraphSchemaAction_Paste> PasteAction(new FMaterialGraphSchemaAction_Paste(FText::GetEmpty(), PasteDesc, PasteToolTip.ToString(), 0));
 		ContextMenuBuilder.AddAction(PasteAction);
 	}
 }
@@ -722,14 +722,31 @@ void UMaterialGraphSchema::GetMaterialFunctionActions(FGraphActionMenuBuilder& A
 					const FString LibraryCategoriesString = AssetData.TagsAndValues.FindRef("LibraryCategories");
 					if ( !LibraryCategoriesString.IsEmpty() )
 					{
-						UArrayProperty* LibraryCategoriesProperty = FindFieldChecked<UArrayProperty>(UMaterialFunction::StaticClass(), GET_MEMBER_NAME_CHECKED(UMaterialFunction, LibraryCategories));
-						uint8* DestAddr = (uint8*)(&LibraryCategories);
+						if (UArrayProperty* LibraryCategoriesProperty = FindFieldChecked<UArrayProperty>(UMaterialFunction::StaticClass(), TEXT("LibraryCategories")))
+						{
+							uint8* DestAddr = (uint8*)(&LibraryCategories);
+							LibraryCategoriesProperty->ImportText(*LibraryCategoriesString, DestAddr, PPF_None, NULL, GWarn);
+						}
+					}
+				}
+				TArray<FText> LibraryCategoriesText;
+				{
+					const FString LibraryCategoriesString = AssetData.TagsAndValues.FindRef("LibraryCategoriesText");
+					if ( !LibraryCategoriesString.IsEmpty() )
+					{
+						UArrayProperty* LibraryCategoriesProperty = FindFieldChecked<UArrayProperty>(UMaterialFunction::StaticClass(), GET_MEMBER_NAME_CHECKED(UMaterialFunction, LibraryCategoriesText));
+						uint8* DestAddr = (uint8*)(&LibraryCategoriesText);
 						LibraryCategoriesProperty->ImportText(*LibraryCategoriesString, DestAddr, PPF_None, NULL, GWarn);
 					}
 
-					if ( LibraryCategories.Num() == 0 )
+					for (FString& Category : LibraryCategories)
+					{
+						LibraryCategoriesText.Add(FText::FromString(Category));
+					}
+
+					if ( LibraryCategoriesText.Num() == 0 )
 					{
-						LibraryCategories.Add( LOCTEXT("UncategorizedMaterialFunction", "Uncategorized").ToString() );
+						LibraryCategoriesText.Add( LOCTEXT("UncategorizedMaterialFunction", "Uncategorized") );
 					}
 				}
 
@@ -745,7 +762,7 @@ void UMaterialGraphSchema::GetMaterialFunctionActions(FGraphActionMenuBuilder& A
 				// For each category the function should belong to...
 				for (int32 CategoryIndex = 0; CategoryIndex < LibraryCategories.Num(); CategoryIndex++)
 				{
-					const FString& CategoryName = LibraryCategories[CategoryIndex];
+					const FText& CategoryName = LibraryCategoriesText[CategoryIndex];
 					TSharedPtr<FMaterialGraphSchemaAction_NewFunctionCall> NewFunctionAction(new FMaterialGraphSchemaAction_NewFunctionCall(
 						CategoryName,
 						FText::FromString(FunctionName),
@@ -767,7 +784,7 @@ void UMaterialGraphSchema::GetCommentAction(FGraphActionMenuBuilder& ActionMenuB
 		const FText MultiCommentDesc = LOCTEXT("MultiCommentDesc", "Create Comment from Selection");
 		const FText CommentToolTip = LOCTEXT("CommentToolTip", "Creates a comment.");
 		const FText MenuDescription = bIsManyNodesSelected ? MultiCommentDesc : CommentDesc;
-		TSharedPtr<FMaterialGraphSchemaAction_NewComment> NewAction(new FMaterialGraphSchemaAction_NewComment(TEXT(""), MenuDescription, CommentToolTip.ToString(), 0));
+		TSharedPtr<FMaterialGraphSchemaAction_NewComment> NewAction(new FMaterialGraphSchemaAction_NewComment(FText::GetEmpty(), MenuDescription, CommentToolTip.ToString(), 0));
 		ActionMenuBuilder.AddAction( NewAction );
 	}
 }
diff --git a/Engine/Source/Editor/UnrealEd/Private/Niagara/EdGraphSchema_Niagara.cpp b/Engine/Source/Editor/UnrealEd/Private/Niagara/EdGraphSchema_Niagara.cpp
--- a/Engine/Source/Editor/UnrealEd/Private/Niagara/EdGraphSchema_Niagara.cpp
+++ b/Engine/Source/Editor/UnrealEd/Private/Niagara/EdGraphSchema_Niagara.cpp
@@ -110,7 +110,7 @@ UEdGraphSchema_Niagara::UEdGraphSchema_Niagara(const FObjectInitializer& ObjectI
 	PC_Curve = TEXT("curve");
 }
 
-TSharedPtr<FNiagaraSchemaAction_NewNode> AddNewNodeAction(FGraphContextMenuBuilder& ContextMenuBuilder, const FString& Category, const FText& MenuDesc, const FString& Tooltip)
+TSharedPtr<FNiagaraSchemaAction_NewNode> AddNewNodeAction(FGraphContextMenuBuilder& ContextMenuBuilder, const FText& Category, const FText& MenuDesc, const FString& Tooltip)
 {
 	TSharedPtr<FNiagaraSchemaAction_NewNode> NewAction = TSharedPtr<FNiagaraSchemaAction_NewNode>(new FNiagaraSchemaAction_NewNode(Category, MenuDesc, Tooltip, 0));
 	ContextMenuBuilder.AddAction( NewAction );
@@ -133,7 +133,7 @@ void UEdGraphSchema_Niagara::GetGraphContextActions(FGraphContextMenuBuilder& Co
 			Args.Add(TEXT("Attribute"), FText::FromName(Attr.Name));
 			const FText MenuDesc = FText::Format(LOCTEXT("GetAttribute", "Get {Attribute}"), Args);
 
-			TSharedPtr<FNiagaraSchemaAction_NewNode> GetAttrAction = AddNewNodeAction(ContextMenuBuilder, *LOCTEXT("Attributes Menu Title", "Attributes").ToString(), MenuDesc, TEXT(""));
+			TSharedPtr<FNiagaraSchemaAction_NewNode> GetAttrAction = AddNewNodeAction(ContextMenuBuilder, LOCTEXT("Attributes Menu Title", "Attributes"), MenuDesc, TEXT(""));
 
 			UNiagaraNodeInput* InputNode = NewObject<UNiagaraNodeInput>(ContextMenuBuilder.OwnerOfTemporaries);
 			InputNode->Input = Attr;
@@ -144,7 +144,7 @@ void UEdGraphSchema_Niagara::GetGraphContextActions(FGraphContextMenuBuilder& Co
 #define NiagaraOp(OPNAME) \
 		if(const FNiagaraOpInfo* OpInfo = FNiagaraOpInfo::GetOpInfo(FNiagaraOpInfo::OPNAME))\
 		{\
-		TSharedPtr<FNiagaraSchemaAction_NewNode> AddOpAction = AddNewNodeAction(ContextMenuBuilder, *LOCTEXT("Operations Menu Title", "Operations").ToString(), OpInfo->FriendlyName, TEXT("")); \
+		TSharedPtr<FNiagaraSchemaAction_NewNode> AddOpAction = AddNewNodeAction(ContextMenuBuilder, LOCTEXT("Operations Menu Title", "Operations"), OpInfo->FriendlyName, TEXT("")); \
 			UNiagaraNodeOp* OpNode = NewObject<UNiagaraNodeOp>(ContextMenuBuilder.OwnerOfTemporaries); \
 			OpNode->OpName = OpInfo->Name; \
 			AddOpAction->NodeTemplate = OpNode; \
@@ -162,7 +162,7 @@ void UEdGraphSchema_Niagara::GetGraphContextActions(FGraphContextMenuBuilder& Co
 		Args.Add(TEXT("Constant"), FText::FromName(SysConst.Name));
 		const FText MenuDesc = FText::Format(LOCTEXT("GetSystemConstant", "Get {Constant}"), Args);
 
-		TSharedPtr<FNiagaraSchemaAction_NewNode> GetConstAction = AddNewNodeAction(ContextMenuBuilder, *LOCTEXT("System Constants Menu Title", "System Constants").ToString(), MenuDesc, TEXT(""));
+		TSharedPtr<FNiagaraSchemaAction_NewNode> GetConstAction = AddNewNodeAction(ContextMenuBuilder, LOCTEXT("System Constants Menu Title", "System Constants"), MenuDesc, TEXT(""));
 
 		UNiagaraNodeInput* InputNode = NewObject<UNiagaraNodeInput>(ContextMenuBuilder.OwnerOfTemporaries);
 		InputNode->Input = SysConst;
@@ -176,7 +176,7 @@ void UEdGraphSchema_Niagara::GetGraphContextActions(FGraphContextMenuBuilder& Co
 		const FText MenuDesc = LOCTEXT("GetInput", "Input");
 
 		//TSharedPtr<FNiagaraSchemaAction_NewNode> GetConstAction = AddNewNodeAction(ContextMenuBuilder, *LOCTEXT("System Constants Menu Title", "System Constants").ToString(), MenuDesc, TEXT(""));
-		TSharedPtr<FNiagaraSchemaAction_NewNode> InputAction = AddNewNodeAction(ContextMenuBuilder, TEXT(""), MenuDesc, TEXT(""));
+		TSharedPtr<FNiagaraSchemaAction_NewNode> InputAction = AddNewNodeAction(ContextMenuBuilder, FText::GetEmpty(), MenuDesc, TEXT(""));
 
 		UNiagaraNodeInput* InputNode = NewObject<UNiagaraNodeInput>(ContextMenuBuilder.OwnerOfTemporaries);
 		InputNode->Input.Name = NAME_None;
@@ -188,7 +188,7 @@ void UEdGraphSchema_Niagara::GetGraphContextActions(FGraphContextMenuBuilder& Co
 	{
 		const FText MenuDesc = LOCTEXT("NiagaraFunctionCall", "Function Call");
 
-		TSharedPtr<FNiagaraSchemaAction_NewNode> FunctionCallAction = AddNewNodeAction(ContextMenuBuilder, TEXT(""), MenuDesc, TEXT(""));
+		TSharedPtr<FNiagaraSchemaAction_NewNode> FunctionCallAction = AddNewNodeAction(ContextMenuBuilder, FText::GetEmpty(), MenuDesc, TEXT(""));
 
 		UNiagaraNodeFunctionCall* FunctionCallNode = NewObject<UNiagaraNodeFunctionCall>(ContextMenuBuilder.OwnerOfTemporaries);
 		FunctionCallAction->NodeTemplate = FunctionCallNode;
diff --git a/Engine/Source/Editor/UnrealEd/Private/SoundClassGraphSchema.cpp b/Engine/Source/Editor/UnrealEd/Private/SoundClassGraphSchema.cpp
--- a/Engine/Source/Editor/UnrealEd/Private/SoundClassGraphSchema.cpp
+++ b/Engine/Source/Editor/UnrealEd/Private/SoundClassGraphSchema.cpp
@@ -81,7 +81,7 @@ void USoundClassGraphSchema::GetGraphContextActions(FGraphContextMenuBuilder& Co
 	const FText Name = LOCTEXT("NewSoundClass", "New Sound Class");
 	const FText ToolTip = LOCTEXT("NewSoundClassTooltip", "Create a new sound class");
 	
-	TSharedPtr<FSoundClassGraphSchemaAction_NewNode> NewAction(new FSoundClassGraphSchemaAction_NewNode(TEXT(""), Name, ToolTip.ToString(), 0));
+	TSharedPtr<FSoundClassGraphSchemaAction_NewNode> NewAction(new FSoundClassGraphSchemaAction_NewNode(FText::GetEmpty(), Name, ToolTip.ToString(), 0));
 
 	ContextMenuBuilder.AddAction( NewAction );
 }
diff --git a/Engine/Source/Editor/UnrealEd/Private/SoundCueGraphSchema.cpp b/Engine/Source/Editor/UnrealEd/Private/SoundCueGraphSchema.cpp
--- a/Engine/Source/Editor/UnrealEd/Private/SoundCueGraphSchema.cpp
+++ b/Engine/Source/Editor/UnrealEd/Private/SoundCueGraphSchema.cpp
@@ -248,7 +248,7 @@ void USoundCueGraphSchema::GetGraphContextActions(FGraphContextMenuBuilder& Cont
 
 	if (!ContextMenuBuilder.FromPin && FSoundCueEditorUtilities::CanPasteNodes(ContextMenuBuilder.CurrentGraph))
 	{
-		TSharedPtr<FSoundCueGraphSchemaAction_Paste> NewAction( new FSoundCueGraphSchemaAction_Paste(TEXT(""), LOCTEXT("PasteHereAction", "Paste here"), TEXT(""), 0) );
+		TSharedPtr<FSoundCueGraphSchemaAction_Paste> NewAction( new FSoundCueGraphSchemaAction_Paste(FText::GetEmpty(), LOCTEXT("PasteHereAction", "Paste here"), TEXT(""), 0) );
 		ContextMenuBuilder.AddAction( NewAction );
 	}
 }
@@ -446,7 +446,7 @@ void USoundCueGraphSchema::GetAllSoundNodeActions(FGraphActionMenuBuilder& Actio
 				FFormatNamedArguments Arguments;
 				Arguments.Add(TEXT("Name"), Name);
 				const FText AddToolTip = FText::Format(LOCTEXT("NewSoundCueNodeTooltip", "Adds {Name} node here"), Arguments);
-				TSharedPtr<FSoundCueGraphSchemaAction_NewNode> NewNodeAction(new FSoundCueGraphSchemaAction_NewNode(LOCTEXT("SoundNodeAction", "Sound Node").ToString(), Name, AddToolTip.ToString(), 0));
+				TSharedPtr<FSoundCueGraphSchemaAction_NewNode> NewNodeAction(new FSoundCueGraphSchemaAction_NewNode(LOCTEXT("SoundNodeAction", "Sound Node"), Name, AddToolTip.ToString(), 0));
 				ActionMenuBuilder.AddAction(NewNodeAction);
 				NewNodeAction->SoundNodeClass = SoundNodeClass;
 			}
@@ -458,7 +458,7 @@ void USoundCueGraphSchema::GetAllSoundNodeActions(FGraphActionMenuBuilder& Actio
 				Arguments.Add(TEXT("SelectedItems"), SelectedItemText);
 				const FText MenuDesc = FText::Format(LOCTEXT("NewSoundNodeRandom", "{Name}: {SelectedItems}"), Arguments);
 				const FText ToolTip = FText::Format(LOCTEXT("NewSoundNodeRandomTooltip", "Adds a {Name} node for {SelectedItems} here"), Arguments);
-				TSharedPtr<FSoundCueGraphSchemaAction_NewFromSelected> NewNodeAction(new FSoundCueGraphSchemaAction_NewFromSelected(TEXT("From Selected"),
+				TSharedPtr<FSoundCueGraphSchemaAction_NewFromSelected> NewNodeAction(new FSoundCueGraphSchemaAction_NewFromSelected(LOCTEXT("FromSelected", "From Selected"),
 					MenuDesc,
 					ToolTip.ToString(), 0));
 				ActionMenuBuilder.AddAction(NewNodeAction);
@@ -476,7 +476,7 @@ void USoundCueGraphSchema::GetCommentAction(FGraphActionMenuBuilder& ActionMenuB
 		const FText MenuDescription = bIsManyNodesSelected ? LOCTEXT("CreateCommentAction", "Create Comment from Selection") : LOCTEXT("AddCommentAction", "Add Comment...");
 		const FString ToolTip = LOCTEXT("CreateCommentToolTip", "Creates a comment.").ToString();
 
-		TSharedPtr<FSoundCueGraphSchemaAction_NewComment> NewAction(new FSoundCueGraphSchemaAction_NewComment(TEXT(""), MenuDescription, ToolTip, 0));
+		TSharedPtr<FSoundCueGraphSchemaAction_NewComment> NewAction(new FSoundCueGraphSchemaAction_NewComment(FText::GetEmpty(), MenuDescription, ToolTip, 0));
 		ActionMenuBuilder.AddAction( NewAction );
 	}
 }
diff --git a/Engine/Source/Editor/UnrealEd/Private/Tests/EditorBuildPromotionTests.cpp b/Engine/Source/Editor/UnrealEd/Private/Tests/EditorBuildPromotionTests.cpp
--- a/Engine/Source/Editor/UnrealEd/Private/Tests/EditorBuildPromotionTests.cpp
+++ b/Engine/Source/Editor/UnrealEd/Private/Tests/EditorBuildPromotionTests.cpp
@@ -574,7 +574,7 @@ namespace EditorBuildPromotionTestUtils
 	static UEdGraphNode* CreateNewGraphNodeFromTemplate(UK2Node* NodeTemplate, UEdGraph* InGraph, const FVector2D& GraphLocation, UEdGraphPin* ConnectPin = NULL)
 	{
 		const FString EmptyString(TEXT(""));
-		TSharedPtr<FEdGraphSchemaAction_K2NewNode> Action = TSharedPtr<FEdGraphSchemaAction_K2NewNode>(new FEdGraphSchemaAction_K2NewNode(EmptyString, FText::FromString(EmptyString), EmptyString, 0));
+		TSharedPtr<FEdGraphSchemaAction_K2NewNode> Action = TSharedPtr<FEdGraphSchemaAction_K2NewNode>(new FEdGraphSchemaAction_K2NewNode(FText::GetEmpty(), FText::GetEmpty(), EmptyString, 0));
 		Action->NodeTemplate = NodeTemplate;
 
 		return Action->PerformAction(InGraph, ConnectPin, GraphLocation, false);
diff --git a/Engine/Source/Editor/UnrealEd/Public/Kismet2/BlueprintEditorUtils.h b/Engine/Source/Editor/UnrealEd/Public/Kismet2/BlueprintEditorUtils.h
--- a/Engine/Source/Editor/UnrealEd/Public/Kismet2/BlueprintEditorUtils.h
+++ b/Engine/Source/Editor/UnrealEd/Public/Kismet2/BlueprintEditorUtils.h
@@ -878,7 +878,7 @@ public:
 	 * @param	VarCategory			The new value of the custom category for the variable
 	 * @param	bDontRecompile		If true, the blueprint will not be marked as modified, and will not be recompiled.  
 	 */
-	static void SetBlueprintVariableCategory(UBlueprint* Blueprint, const FName& VarName, const UStruct* InLocalVarScope, const FName& NewCategory, bool bDontRecompile=false);
+	static void SetBlueprintVariableCategory(UBlueprint* Blueprint, const FName& VarName, const UStruct* InLocalVarScope, const FText& NewCategory, bool bDontRecompile=false);
 
 	/**
 	 * Gets the custom category on the variable with the specified name.
@@ -887,7 +887,7 @@ public:
 	 * @param	InLocalVarScope		Local variable's scope, if looking to modify a local variable
 	 * @return						The custom category (None indicates the name will be the same as the blueprint)
 	 */
-	static FName GetBlueprintVariableCategory(UBlueprint* Blueprint, const FName& VarName, const UStruct* InLocalVarScope);
+	static FText GetBlueprintVariableCategory(UBlueprint* Blueprint, const FName& VarName, const UStruct* InLocalVarScope);
 
 	/** Gets pointer to PropertyFlags of variable */
 	static uint64* GetBlueprintVariablePropertyFlags(UBlueprint* Blueprint, const FName& VarName);
@@ -911,7 +911,7 @@ public:
 	static bool MoveVariableBeforeVariable(UBlueprint* Blueprint, FName VarNameToMove, FName TargetVarName, bool bDontRecompile);
 
 	/** Find first variable of the supplied category */
-	static int32 FindFirstNewVarOfCategory(const UBlueprint* Blueprint, FName Category);
+	static int32 FindFirstNewVarOfCategory(const UBlueprint* Blueprint, FText Category);
 
 	/**
 	 * Find the index of a timeline first declared in this blueprint. Returns INDEX_NONE if not found.
diff --git a/Engine/Source/Runtime/Core/Public/UObject/ObjectVersion.h b/Engine/Source/Runtime/Core/Public/UObject/ObjectVersion.h
--- a/Engine/Source/Runtime/Core/Public/UObject/ObjectVersion.h
+++ b/Engine/Source/Runtime/Core/Public/UObject/ObjectVersion.h
@@ -508,6 +508,8 @@ enum EUnrealEngineObjectUE4Version
 	VER_UE4_CONSTRAINT_INSTANCE_MOTOR_FLAGS,
 	// Serialize bIsConst in FEdGraphPinType
 	VER_UE4_SERIALIZE_PINTYPE_CONST,
+	// Change UMaterialFunction::LibraryCategories to LibraryCategoriesText (old assets were saved before auto-conversion of UArrayProperty was possible)
+	VER_UE4_LIBRARY_CATEGORIES_AS_FTEXT,
 
 	// -----<new versions can be added before this line>-------------------------------------------------
 	// - this needs to be the last line (see note below)
diff --git a/Engine/Source/Runtime/CoreUObject/Private/UObject/Class.cpp b/Engine/Source/Runtime/CoreUObject/Private/UObject/Class.cpp
--- a/Engine/Source/Runtime/CoreUObject/Private/UObject/Class.cpp
+++ b/Engine/Source/Runtime/CoreUObject/Private/UObject/Class.cpp
@@ -332,7 +332,7 @@ const FText UField::GetMetaDataText(const TCHAR* MetaDataKey, const FString Loca
 		if( HasMetaData( MetaDataKey ))
 		{
 			DefaultMetaData = GetMetaData(MetaDataKey);
-			LocalizedMetaData = FText::FromString(DefaultMetaData);
+			LocalizedMetaData = FText::AsCultureInvariant(DefaultMetaData);
 		}
 	}
 
@@ -348,7 +348,7 @@ const FText UField::GetMetaDataText(const FName& MetaDataKey, const FString Loca
 		if( HasMetaData( MetaDataKey ))
 		{
 			DefaultMetaData = GetMetaData(MetaDataKey);
-			LocalizedMetaData = FText::FromString(DefaultMetaData);
+			LocalizedMetaData = FText::AsCultureInvariant(DefaultMetaData);
 		}
 	}
 	return LocalizedMetaData;
@@ -1012,6 +1012,24 @@ void UStruct::SerializeTaggedProperties(FArchive& Ar, uint8* Data, UStruct* Defa
 				AdvanceProperty = true;
 				continue; 
 			}
+			else if( Tag.Type==NAME_NameProperty && dynamic_cast<UTextProperty*>(Property) ) // Convert serialized name to text.
+			{ 
+				FName Name;  
+				Ar << Name;
+				FText Text = FText::FromName(Name);
+				CastChecked<UTextProperty>(Property)->SetPropertyValue_InContainer(Data, Text, Tag.ArrayIndex);
+				AdvanceProperty = true;
+				continue; 
+			}
+			else if( Tag.Type==NAME_TextProperty && dynamic_cast<UNameProperty*>(Property) ) // Convert serialized text to name.
+			{ 
+				FText Text;  
+				Ar << Text;
+				FName Name = FName(*Text.ToString());
+				CastChecked<UNameProperty>(Property)->SetPropertyValue_InContainer(Data, Name, Tag.ArrayIndex);
+				AdvanceProperty = true;
+				continue; 
+			}
 			else if ( Tag.Type == NAME_ByteProperty && Property->GetID() == NAME_IntProperty )
 			{
 				// this property's data was saved as a uint8, but the property has been changed to an int32.  Since there is no loss of data
@@ -1162,6 +1180,31 @@ void UStruct::SerializeTaggedProperties(FArchive& Ar, uint8* Data, UStruct* Defa
 					}
 					continue; 
 				}
+				else if (Tag.InnerType == NAME_NameProperty && dynamic_cast<UTextProperty*>(ArrayProperty->Inner)) // Convert serialized name to text.
+				{ 
+					for(int32 i = 0; i < ElementCount; ++i)
+					{
+						FName Name;
+						Ar << Name;
+						FText Text = FText::FromName(Name);
+						Text.Flags |= ETextFlag::ConvertedProperty;
+						CastChecked<UTextProperty>(ArrayProperty->Inner)->SetPropertyValue(ScriptArrayHelper.GetRawPtr(i), Text);
+						AdvanceProperty = true;
+					}
+					continue;
+				}
+				else if( Tag.InnerType==NAME_TextProperty && dynamic_cast<UNameProperty*>(ArrayProperty->Inner) ) // Convert serialized text to name.
+				{ 
+					for(int32 i = 0; i < ElementCount; ++i)
+					{
+						FText Text;  
+						Ar << Text;
+						FName Name = FTextInspector::GetSourceString(Text) ? FName(**FTextInspector::GetSourceString(Text)) : NAME_None;
+						static_cast<UNameProperty*>(ArrayProperty->Inner)->SetPropertyValue(ScriptArrayHelper.GetRawPtr(i), Name);
+						AdvanceProperty = true;
+					}
+					continue; 
+				}
 				else if ((Tag.InnerType == NAME_AssetObjectProperty || Tag.InnerType == NAME_AssetSubclassOfProperty) && (ArrayProperty->Inner->GetID() == NAME_ObjectProperty || ArrayProperty->Inner->GetID() == NAME_ClassProperty))
 				{
 					for (int32 i = 0; i < ElementCount; ++i)
diff --git a/Engine/Source/Runtime/Engine/Classes/EdGraph/EdGraphSchema.h b/Engine/Source/Runtime/Engine/Classes/EdGraph/EdGraphSchema.h
--- a/Engine/Source/Runtime/Engine/Classes/EdGraph/EdGraphSchema.h
+++ b/Engine/Source/Runtime/Engine/Classes/EdGraph/EdGraphSchema.h
@@ -66,7 +66,7 @@ struct ENGINE_API FEdGraphSchemaAction
 
 	/** This is the UI centric category the action fits in (e.g., Functions, Variables).  Use this instead of the NodeType.NodeCategory because multiple NodeCategories might visually belong together. */
 	UPROPERTY()
-	FString Category;
+	FText Category;
 
 	/** This is just an arbitrary dump of extra text that search will match on, in addition to the description and tooltip, e.g., Add might have the keyword Math */
 	UPROPERTY()
@@ -88,6 +88,10 @@ struct ENGINE_API FEdGraphSchemaAction
 	UPROPERTY()
 	FText CachedSearchKeywords;
 
+	/** Search categories for the action (doesn't have to be set when instantiated, will be constructed by GetSearchTitle() if left empty)*/
+	UPROPERTY()
+	FText CachedSearchCategories;
+
 	FEdGraphSchemaAction() 
 		: Grouping(0)
 		, SectionID(0)
@@ -95,7 +99,7 @@ struct ENGINE_API FEdGraphSchemaAction
 	
 	virtual ~FEdGraphSchemaAction() {}
 
-	FEdGraphSchemaAction(const FString& InNodeCategory, const FText& InMenuDesc, const FString& InToolTip, const int32 InGrouping)
+	FEdGraphSchemaAction(const FText& InNodeCategory, const FText& InMenuDesc, const FString& InToolTip, const int32 InGrouping)
 		: MenuDescription(InMenuDesc)
 		, TooltipDescription(InToolTip)
 		, Category(InNodeCategory)
@@ -139,7 +143,7 @@ struct ENGINE_API FEdGraphSchemaAction
 		return CachedSearchTitle;
 	}
 
-	/** Retrieves the full searchable title for this action */
+	/** Retrieves the full searchable keywords for this action */
 	FText GetSearchKeywords()
 	{
 		if(CachedSearchKeywords.IsEmpty())
@@ -152,6 +156,19 @@ struct ENGINE_API FEdGraphSchemaAction
 		return CachedSearchKeywords;
 	}
 
+	/** Retrieves the full searchable categories for this action */
+	FText GetSearchCategory()
+	{
+		if(CachedSearchCategories.IsEmpty())
+		{
+			FFormatNamedArguments Args;
+			Args.Add(TEXT("LocalizedCategories"), Category);
+			Args.Add(TEXT("SourceCategories"), FText::FromString(Category.BuildSourceString()));
+			CachedSearchCategories = FText::Format(FText::FromString("{LocalizedCategories} {SourceCategories}"), Args);
+		}
+		return CachedSearchCategories;
+	}
+
 	// GC.
 	virtual void AddReferencedObjects(FReferenceCollector& Collector) {}
 };
@@ -176,7 +193,7 @@ struct ENGINE_API FEdGraphSchemaAction_NewNode : public FEdGraphSchemaAction
 		, NodeTemplate(NULL)
 	{}
 
-	FEdGraphSchemaAction_NewNode(const FString& InNodeCategory, const FText& InMenuDesc, const FString& InToolTip, const int32 InGrouping)
+	FEdGraphSchemaAction_NewNode(const FText& InNodeCategory, const FText& InMenuDesc, const FString& InToolTip, const int32 InGrouping)
 		: FEdGraphSchemaAction(InNodeCategory, InMenuDesc, InToolTip, InGrouping) 
 		, NodeTemplate(NULL)
 	{}
@@ -209,7 +226,7 @@ struct FEdGraphSchemaAction_Dummy : public FEdGraphSchemaAction
 	: FEdGraphSchemaAction()
 	{}
 
-	FEdGraphSchemaAction_Dummy(const FString& InNodeCategory, const FText& InMenuDesc, const FString& InToolTip, const int32 InGrouping)
+	FEdGraphSchemaAction_Dummy(const FText& InNodeCategory, const FText& InMenuDesc, const FString& InToolTip, const int32 InGrouping)
 		: FEdGraphSchemaAction(InNodeCategory, InMenuDesc, InToolTip, InGrouping)		
 	{}
 };
diff --git a/Engine/Source/Runtime/Engine/Classes/Engine/Blueprint.h b/Engine/Source/Runtime/Engine/Classes/Engine/Blueprint.h
--- a/Engine/Source/Runtime/Engine/Classes/Engine/Blueprint.h
+++ b/Engine/Source/Runtime/Engine/Classes/Engine/Blueprint.h
@@ -159,7 +159,7 @@ struct FBPVariableDescription
 
 	/** Category this variable should be in */
 	UPROPERTY(EditAnywhere, Category=BPVariableDescription)
-	FName Category;
+	FText Category;
 
 	/** Property flags for this variable - Changed from int32 to uint64*/
 	UPROPERTY(EditAnywhere, Category=BPVariableDescription)
diff --git a/Engine/Source/Runtime/Engine/Classes/Engine/SCS_Node.h b/Engine/Source/Runtime/Engine/Classes/Engine/SCS_Node.h
--- a/Engine/Source/Runtime/Engine/Classes/Engine/SCS_Node.h
+++ b/Engine/Source/Runtime/Engine/Classes/Engine/SCS_Node.h
@@ -22,7 +22,7 @@ class USCS_Node : public UObject
 
 	/** If non-None, the assigned category name */
 	UPROPERTY()
-	FName CategoryName;
+	FText CategoryName;
 
 	/** Socket/Bone that Node might attach to */
 	UPROPERTY()
diff --git a/Engine/Source/Runtime/Engine/Classes/Materials/MaterialExpression.h b/Engine/Source/Runtime/Engine/Classes/Materials/MaterialExpression.h
--- a/Engine/Source/Runtime/Engine/Classes/Materials/MaterialExpression.h
+++ b/Engine/Source/Runtime/Engine/Classes/Materials/MaterialExpression.h
@@ -165,7 +165,7 @@ class ENGINE_API UMaterialExpression : public UObject
 
 	/** Localized categories to sort this expression into... */
 	UPROPERTY()
-	TArray<FString> MenuCategories;
+	TArray<FText> MenuCategories;
 
 	/** The expression's outputs, which are set in default properties by derived classes. */
 	UPROPERTY()
diff --git a/Engine/Source/Runtime/Engine/Classes/Materials/MaterialFunction.h b/Engine/Source/Runtime/Engine/Classes/Materials/MaterialFunction.h
--- a/Engine/Source/Runtime/Engine/Classes/Materials/MaterialFunction.h
+++ b/Engine/Source/Runtime/Engine/Classes/Materials/MaterialFunction.h
@@ -29,12 +29,19 @@ class UMaterialFunction : public UObject
 	UPROPERTY(EditAnywhere, Category=MaterialFunction, AssetRegistrySearchable)
 	uint32 bExposeToLibrary:1;
 
+	/** 
+	 * Categories that this function belongs to in the material function library.  
+	 * Ideally categories should be chosen carefully so that there are not too many.
+	 */
+	UPROPERTY(AssetRegistrySearchable)
+	TArray<FString> LibraryCategories_DEPRECATED;
+
 	/** 
 	 * Categories that this function belongs to in the material function library.  
 	 * Ideally categories should be chosen carefully so that there are not too many.
 	 */
 	UPROPERTY(EditAnywhere, Category=MaterialFunction, AssetRegistrySearchable)
-	TArray<FString> LibraryCategories;
+	TArray<FText> LibraryCategoriesText;
 
 	/** Array of material expressions, excluding Comments.  Used by the material editor. */
 	UPROPERTY()
diff --git a/Engine/Source/Runtime/Engine/Private/EdGraph/EdGraphSchema.cpp b/Engine/Source/Runtime/Engine/Private/EdGraph/EdGraphSchema.cpp
--- a/Engine/Source/Runtime/Engine/Private/EdGraph/EdGraphSchema.cpp
+++ b/Engine/Source/Runtime/Engine/Private/EdGraph/EdGraphSchema.cpp
@@ -73,7 +73,7 @@ void FGraphActionListBuilderBase::ActionGroup::GetCategoryChain(TArray<FString>&
 	{
 		TArray<FString> SubCategoryChain;
 
-		FString SubCategory = FEditorCategoryUtils::GetCategoryDisplayString(Actions[0]->Category);
+		FString SubCategory = FEditorCategoryUtils::GetCategoryDisplayString(Actions[0]->Category.ToString());
 		SubCategory.ParseIntoArray(SubCategoryChain, *CategoryDelim, true);
 
 		HierarchyOut.Append(SubCategoryChain);
diff --git a/Engine/Source/Runtime/Engine/Private/Materials/MaterialExpressions.cpp b/Engine/Source/Runtime/Engine/Private/Materials/MaterialExpressions.cpp
--- a/Engine/Source/Runtime/Engine/Private/Materials/MaterialExpressions.cpp
+++ b/Engine/Source/Runtime/Engine/Private/Materials/MaterialExpressions.cpp
@@ -1112,9 +1112,9 @@ UMaterialExpressionTextureSample::UMaterialExpressionTextureSample(const FObject
 	// Structure to hold one-time initialization
 	struct FConstructorStatics
 	{
-		FString NAME_Texture;
+		FText NAME_Texture;
 		FConstructorStatics()
-			: NAME_Texture(LOCTEXT( "Texture", "Texture" ).ToString())
+			: NAME_Texture(LOCTEXT( "Texture", "Texture" ))
 		{
 		}
 	};
@@ -1413,9 +1413,9 @@ UMaterialExpressionAdd::UMaterialExpressionAdd(const FObjectInitializer& ObjectI
 	// Structure to hold one-time initialization
 	struct FConstructorStatics
 	{
-		FString NAME_Math;
+		FText NAME_Math;
 		FConstructorStatics()
-			: NAME_Math(LOCTEXT( "Math", "Math" ).ToString())
+			: NAME_Math(LOCTEXT( "Math", "Math" ))
 		{
 		}
 	};
@@ -1437,9 +1437,9 @@ UMaterialExpressionTextureSampleParameter::UMaterialExpressionTextureSampleParam
 	// Structure to hold one-time initialization
 	struct FConstructorStatics
 	{
-		FString NAME_Obsolete;
+		FText NAME_Obsolete;
 		FConstructorStatics()
-			: NAME_Obsolete(LOCTEXT( "Obsolete", "Obsolete" ).ToString())
+			: NAME_Obsolete(LOCTEXT( "Obsolete", "Obsolete" ))
 		{
 		}
 	};
@@ -1539,12 +1539,12 @@ UMaterialExpressionTextureObjectParameter::UMaterialExpressionTextureObjectParam
 	struct FConstructorStatics
 	{
 		ConstructorHelpers::FObjectFinder<UTexture2D> DefaultTexture2D;
-		FString NAME_Texture;
-		FString NAME_Parameters;
+		FText NAME_Texture;
+		FText NAME_Parameters;
 		FConstructorStatics()
 			: DefaultTexture2D(TEXT("/Engine/EngineResources/DefaultTexture"))
-			, NAME_Texture(LOCTEXT( "Texture", "Texture" ).ToString())
-			, NAME_Parameters(LOCTEXT( "Parameters", "Parameters" ).ToString())
+			, NAME_Texture(LOCTEXT( "Texture", "Texture" ))
+			, NAME_Parameters(LOCTEXT( "Parameters", "Parameters" ))
 		{
 		}
 	};
@@ -1608,12 +1608,12 @@ UMaterialExpressionTextureObject::UMaterialExpressionTextureObject(const FObject
 	struct FConstructorStatics
 	{
 		ConstructorHelpers::FObjectFinder<UTexture2D> Object0;
-		FString NAME_Texture;
-		FString NAME_Functions;
+		FText NAME_Texture;
+		FText NAME_Functions;
 		FConstructorStatics()
 			: Object0(TEXT("/Engine/EngineResources/DefaultTexture"))
-			, NAME_Texture(LOCTEXT( "Texture", "Texture" ).ToString())
-			, NAME_Functions(LOCTEXT( "Functions", "Functions" ).ToString())
+			, NAME_Texture(LOCTEXT( "Texture", "Texture" ))
+			, NAME_Functions(LOCTEXT( "Functions", "Functions" ))
 		{
 		}
 	};
@@ -1695,12 +1695,12 @@ UMaterialExpressionTextureSampleParameter2D::UMaterialExpressionTextureSamplePar
 	struct FConstructorStatics
 	{
 		ConstructorHelpers::FObjectFinder<UTexture2D> DefaultTexture;
-		FString NAME_Texture;
-		FString NAME_Parameters;
+		FText NAME_Texture;
+		FText NAME_Parameters;
 		FConstructorStatics()
 			: DefaultTexture(TEXT("/Engine/EngineResources/DefaultTexture"))
-			, NAME_Texture(LOCTEXT( "Texture", "Texture" ).ToString())
-			, NAME_Parameters(LOCTEXT( "Parameters", "Parameters" ).ToString())
+			, NAME_Texture(LOCTEXT( "Texture", "Texture" ))
+			, NAME_Parameters(LOCTEXT( "Parameters", "Parameters" ))
 		{
 		}
 	};
@@ -1784,12 +1784,12 @@ UMaterialExpressionTextureSampleParameterCube::UMaterialExpressionTextureSampleP
 	struct FConstructorStatics
 	{
 		ConstructorHelpers::FObjectFinder<UTextureCube> DefaultTextureCube;
-		FString NAME_Texture;
-		FString NAME_Parameters;
+		FText NAME_Texture;
+		FText NAME_Parameters;
 		FConstructorStatics()
 			: DefaultTextureCube(TEXT("/Engine/EngineResources/DefaultTextureCube"))
-			, NAME_Texture(LOCTEXT( "Texture", "Texture" ).ToString())
-			, NAME_Parameters(LOCTEXT( "Parameters", "Parameters" ).ToString())
+			, NAME_Texture(LOCTEXT( "Texture", "Texture" ))
+			, NAME_Parameters(LOCTEXT( "Parameters", "Parameters" ))
 		{
 		}
 	};
@@ -1860,9 +1860,9 @@ UMaterialExpressionTextureSampleParameterSubUV::UMaterialExpressionTextureSample
 	// Structure to hold one-time initialization
 	struct FConstructorStatics
 	{
-		FString NAME_Particles;
+		FText NAME_Particles;
 		FConstructorStatics()
-			: NAME_Particles(LOCTEXT( "Particles", "Particles" ).ToString())
+			: NAME_Particles(LOCTEXT( "Particles", "Particles" ))
 		{
 		}
 	};
@@ -1944,9 +1944,9 @@ UMaterialExpressionMultiply::UMaterialExpressionMultiply(const FObjectInitialize
 	// Structure to hold one-time initialization
 	struct FConstructorStatics
 	{
-		FString NAME_Math;
+		FText NAME_Math;
 		FConstructorStatics()
-			: NAME_Math(LOCTEXT( "Math", "Math" ).ToString())
+			: NAME_Math(LOCTEXT( "Math", "Math" ))
 		{
 		}
 	};
@@ -1989,9 +1989,9 @@ UMaterialExpressionDivide::UMaterialExpressionDivide(const FObjectInitializer& O
 	// Structure to hold one-time initialization
 	struct FConstructorStatics
 	{
-		FString NAME_Math;
+		FText NAME_Math;
 		FConstructorStatics()
-			: NAME_Math(LOCTEXT( "Math", "Math" ).ToString())
+			: NAME_Math(LOCTEXT( "Math", "Math" ))
 		{
 		}
 	};
@@ -2035,9 +2035,9 @@ UMaterialExpressionSubtract::UMaterialExpressionSubtract(const FObjectInitialize
 	// Structure to hold one-time initialization
 	struct FConstructorStatics
 	{
-		FString NAME_Math;
+		FText NAME_Math;
 		FConstructorStatics()
-			: NAME_Math(LOCTEXT( "Math", "Math" ).ToString())
+			: NAME_Math(LOCTEXT( "Math", "Math" ))
 		{
 		}
 	};
@@ -2083,11 +2083,11 @@ UMaterialExpressionLinearInterpolate::UMaterialExpressionLinearInterpolate(const
 	// Structure to hold one-time initialization
 	struct FConstructorStatics
 	{
-		FString NAME_Math;
-		FString NAME_Utility;
+		FText NAME_Math;
+		FText NAME_Utility;
 		FConstructorStatics()
-			: NAME_Math(LOCTEXT( "Math", "Math" ).ToString())
-			, NAME_Utility(LOCTEXT( "Utility", "Utility" ).ToString())
+			: NAME_Math(LOCTEXT( "Math", "Math" ))
+			, NAME_Utility(LOCTEXT( "Utility", "Utility" ))
 		{
 		}
 	};
@@ -2134,9 +2134,9 @@ UMaterialExpressionConstant::UMaterialExpressionConstant(const FObjectInitialize
 	// Structure to hold one-time initialization
 	struct FConstructorStatics
 	{
-		FString NAME_Constants;
+		FText NAME_Constants;
 		FConstructorStatics()
-			: NAME_Constants(LOCTEXT( "Constants", "Constants" ).ToString())
+			: NAME_Constants(LOCTEXT( "Constants", "Constants" ))
 		{
 		}
 	};
@@ -2175,11 +2175,11 @@ UMaterialExpressionConstant2Vector::UMaterialExpressionConstant2Vector(const FOb
 	// Structure to hold one-time initialization
 	struct FConstructorStatics
 	{
-		FString NAME_Constants;
-		FString NAME_Vectors;
+		FText NAME_Constants;
+		FText NAME_Vectors;
 		FConstructorStatics()
-			: NAME_Constants(LOCTEXT( "Constants", "Constants" ).ToString())
-			, NAME_Vectors(LOCTEXT( "Vectors", "Vectors" ).ToString())
+			: NAME_Constants(LOCTEXT( "Constants", "Constants" ))
+			, NAME_Vectors(LOCTEXT( "Vectors", "Vectors" ))
 		{
 		}
 	};
@@ -2218,11 +2218,11 @@ UMaterialExpressionConstant3Vector::UMaterialExpressionConstant3Vector(const FOb
 	// Structure to hold one-time initialization
 	struct FConstructorStatics
 	{
-		FString NAME_Constants;
-		FString NAME_Vectors;
+		FText NAME_Constants;
+		FText NAME_Vectors;
 		FConstructorStatics()
-			: NAME_Constants(LOCTEXT( "Constants", "Constants" ).ToString())
-			, NAME_Vectors(LOCTEXT( "Vectors", "Vectors" ).ToString())
+			: NAME_Constants(LOCTEXT( "Constants", "Constants" ))
+			, NAME_Vectors(LOCTEXT( "Vectors", "Vectors" ))
 		{
 		}
 	};
@@ -2261,11 +2261,11 @@ UMaterialExpressionConstant4Vector::UMaterialExpressionConstant4Vector(const FOb
 	// Structure to hold one-time initialization
 	struct FConstructorStatics
 	{
-		FString NAME_Constants;
-		FString NAME_Vectors;
+		FText NAME_Constants;
+		FText NAME_Vectors;
 		FConstructorStatics()
-			: NAME_Constants(LOCTEXT( "Constants", "Constants" ).ToString())
-			, NAME_Vectors(LOCTEXT( "Vectors", "Vectors" ).ToString())
+			: NAME_Constants(LOCTEXT( "Constants", "Constants" ))
+			, NAME_Vectors(LOCTEXT( "Vectors", "Vectors" ))
 		{
 		}
 	};
@@ -2305,9 +2305,9 @@ UMaterialExpressionClamp::UMaterialExpressionClamp(const FObjectInitializer& Obj
 	// Structure to hold one-time initialization
 	struct FConstructorStatics
 	{
-		FString NAME_Math;
+		FText NAME_Math;
 		FConstructorStatics()
-			: NAME_Math(LOCTEXT( "Math", "Math" ).ToString())
+			: NAME_Math(LOCTEXT( "Math", "Math" ))
 		{
 		}
 	};
@@ -2390,9 +2390,9 @@ UMaterialExpressionMin::UMaterialExpressionMin(const FObjectInitializer& ObjectI
 	// Structure to hold one-time initialization
 	struct FConstructorStatics
 	{
-		FString NAME_Math;
+		FText NAME_Math;
 		FConstructorStatics()
-			: NAME_Math(LOCTEXT( "Math", "Math" ).ToString())
+			: NAME_Math(LOCTEXT( "Math", "Math" ))
 		{
 		}
 	};
@@ -2438,9 +2438,9 @@ UMaterialExpressionMax::UMaterialExpressionMax(const FObjectInitializer& ObjectI
 	// Structure to hold one-time initialization
 	struct FConstructorStatics
 	{
-		FString NAME_Math;
+		FText NAME_Math;
 		FConstructorStatics()
-			: NAME_Math(LOCTEXT( "Math", "Math" ).ToString())
+			: NAME_Math(LOCTEXT( "Math", "Math" ))
 		{
 		}
 	};
@@ -2486,9 +2486,9 @@ UMaterialExpressionTextureCoordinate::UMaterialExpressionTextureCoordinate(const
 	// Structure to hold one-time initialization
 	struct FConstructorStatics
 	{
-		FString NAME_Coordinates;
+		FText NAME_Coordinates;
 		FConstructorStatics()
-			: NAME_Coordinates(LOCTEXT( "Coordinates", "Coordinates" ).ToString())
+			: NAME_Coordinates(LOCTEXT( "Coordinates", "Coordinates" ))
 		{
 		}
 	};
@@ -2530,11 +2530,11 @@ UMaterialExpressionDotProduct::UMaterialExpressionDotProduct(const FObjectInitia
 	// Structure to hold one-time initialization
 	struct FConstructorStatics
 	{
-		FString NAME_Math;
-		FString NAME_VectorOps;
+		FText NAME_Math;
+		FText NAME_VectorOps;
 		FConstructorStatics()
-			: NAME_Math(LOCTEXT( "Math", "Math" ).ToString())
-			, NAME_VectorOps(LOCTEXT( "VectorOps", "VectorOps" ).ToString())
+			: NAME_Math(LOCTEXT( "Math", "Math" ))
+			, NAME_VectorOps(LOCTEXT( "VectorOps", "VectorOps" ))
 		{
 		}
 	};
@@ -2576,11 +2576,11 @@ UMaterialExpressionCrossProduct::UMaterialExpressionCrossProduct(const FObjectIn
 	// Structure to hold one-time initialization
 	struct FConstructorStatics
 	{
-		FString NAME_Math;
-		FString NAME_VectorOps;
+		FText NAME_Math;
+		FText NAME_VectorOps;
 		FConstructorStatics()
-			: NAME_Math(LOCTEXT( "Math", "Math" ).ToString())
-			, NAME_VectorOps(LOCTEXT( "VectorOps", "VectorOps" ).ToString())
+			: NAME_Math(LOCTEXT( "Math", "Math" ))
+			, NAME_VectorOps(LOCTEXT( "VectorOps", "VectorOps" ))
 		{
 		}
 	};
@@ -2622,11 +2622,11 @@ UMaterialExpressionComponentMask::UMaterialExpressionComponentMask(const FObject
 	// Structure to hold one-time initialization
 	struct FConstructorStatics
 	{
-		FString NAME_Math;
-		FString NAME_VectorOps;
+		FText NAME_Math;
+		FText NAME_VectorOps;
 		FConstructorStatics()
-			: NAME_Math(LOCTEXT( "Math", "Math" ).ToString())
-			, NAME_VectorOps(LOCTEXT( "VectorOps", "VectorOps" ).ToString())
+			: NAME_Math(LOCTEXT( "Math", "Math" ))
+			, NAME_VectorOps(LOCTEXT( "VectorOps", "VectorOps" ))
 		{
 		}
 	};
@@ -2673,9 +2673,9 @@ UMaterialExpressionStaticComponentMaskParameter::UMaterialExpressionStaticCompon
 	// Structure to hold one-time initialization
 	struct FConstructorStatics
 	{
-		FString NAME_Parameters;
+		FText NAME_Parameters;
 		FConstructorStatics()
-			: NAME_Parameters(LOCTEXT( "Parameters", "Parameters" ).ToString())
+			: NAME_Parameters(LOCTEXT( "Parameters", "Parameters" ))
 		{
 		}
 	};
@@ -2734,9 +2734,9 @@ UMaterialExpressionTime::UMaterialExpressionTime(const FObjectInitializer& Objec
 	// Structure to hold one-time initialization
 	struct FConstructorStatics
 	{
-		FString NAME_Constants;
+		FText NAME_Constants;
 		FConstructorStatics()
-			: NAME_Constants(LOCTEXT( "Constants", "Constants" ).ToString())
+			: NAME_Constants(LOCTEXT( "Constants", "Constants" ))
 		{
 		}
 	};
@@ -2779,9 +2779,9 @@ UMaterialExpressionCameraVectorWS::UMaterialExpressionCameraVectorWS(const FObje
 	// Structure to hold one-time initialization
 	struct FConstructorStatics
 	{
-		FString NAME_Vectors;
+		FText NAME_Vectors;
 		FConstructorStatics()
-			: NAME_Vectors(LOCTEXT( "Vectors", "Vectors" ).ToString())
+			: NAME_Vectors(LOCTEXT( "Vectors", "Vectors" ))
 		{
 		}
 	};
@@ -2810,11 +2810,11 @@ UMaterialExpressionCameraPositionWS::UMaterialExpressionCameraPositionWS(const F
 	// Structure to hold one-time initialization
 	struct FConstructorStatics
 	{
-		FString NAME_Vectors;
-		FString NAME_Coordinates;
+		FText NAME_Vectors;
+		FText NAME_Coordinates;
 		FConstructorStatics()
-			: NAME_Vectors(LOCTEXT( "Vectors", "Vectors" ).ToString())
-			, NAME_Coordinates(LOCTEXT( "Coordinates", "Coordinates" ).ToString())
+			: NAME_Vectors(LOCTEXT( "Vectors", "Vectors" ))
+			, NAME_Coordinates(LOCTEXT( "Coordinates", "Coordinates" ))
 		{
 		}
 	};
@@ -2846,9 +2846,9 @@ UMaterialExpressionReflectionVectorWS::UMaterialExpressionReflectionVectorWS(con
 	// Structure to hold one-time initialization
 	struct FConstructorStatics
 	{
-		FString NAME_Vectors;
+		FText NAME_Vectors;
 		FConstructorStatics()
-			: NAME_Vectors(LOCTEXT( "Vectors", "Vectors" ).ToString())
+			: NAME_Vectors(LOCTEXT( "Vectors", "Vectors" ))
 		{
 		}
 	};
@@ -2886,9 +2886,9 @@ UMaterialExpressionPanner::UMaterialExpressionPanner(const FObjectInitializer& O
 	// Structure to hold one-time initialization
 	struct FConstructorStatics
 	{
-		FString NAME_Coordinates;
+		FText NAME_Coordinates;
 		FConstructorStatics()
-			: NAME_Coordinates(LOCTEXT( "Coordinates", "Coordinates" ).ToString())
+			: NAME_Coordinates(LOCTEXT( "Coordinates", "Coordinates" ))
 		{
 		}
 	};
@@ -2941,9 +2941,9 @@ UMaterialExpressionRotator::UMaterialExpressionRotator(const FObjectInitializer&
 	// Structure to hold one-time initialization
 	struct FConstructorStatics
 	{
-		FString NAME_Coordinates;
+		FText NAME_Coordinates;
 		FConstructorStatics()
-			: NAME_Coordinates(LOCTEXT( "Coordinates", "Coordinates" ).ToString())
+			: NAME_Coordinates(LOCTEXT( "Coordinates", "Coordinates" ))
 		{
 		}
 	};
@@ -3012,9 +3012,9 @@ UMaterialExpressionSine::UMaterialExpressionSine(const FObjectInitializer& Objec
 	// Structure to hold one-time initialization
 	struct FConstructorStatics
 	{
-		FString NAME_Math;
+		FText NAME_Math;
 		FConstructorStatics()
-			: NAME_Math(LOCTEXT( "Math", "Math" ).ToString())
+			: NAME_Math(LOCTEXT( "Math", "Math" ))
 		{
 		}
 	};
@@ -3046,9 +3046,9 @@ UMaterialExpressionCosine::UMaterialExpressionCosine(const FObjectInitializer& O
 	// Structure to hold one-time initialization
 	struct FConstructorStatics
 	{
-		FString NAME_Math;
+		FText NAME_Math;
 		FConstructorStatics()
-			: NAME_Math(LOCTEXT( "Math", "Math" ).ToString())
+			: NAME_Math(LOCTEXT( "Math", "Math" ))
 		{
 		}
 	};
@@ -3081,9 +3081,9 @@ UMaterialExpressionBumpOffset::UMaterialExpressionBumpOffset(const FObjectInitia
 	// Structure to hold one-time initialization
 	struct FConstructorStatics
 	{
-		FString NAME_Utility;
+		FText NAME_Utility;
 		FConstructorStatics()
-			: NAME_Utility(LOCTEXT( "Utility", "Utility" ).ToString())
+			: NAME_Utility(LOCTEXT( "Utility", "Utility" ))
 		{
 		}
 	};
@@ -3131,11 +3131,11 @@ UMaterialExpressionAppendVector::UMaterialExpressionAppendVector(const FObjectIn
 	// Structure to hold one-time initialization
 	struct FConstructorStatics
 	{
-		FString NAME_Math;
-		FString NAME_VectorOps;
+		FText NAME_Math;
+		FText NAME_VectorOps;
 		FConstructorStatics()
-			: NAME_Math(LOCTEXT( "Math", "Math" ).ToString())
-			, NAME_VectorOps(LOCTEXT( "VectorOps", "VectorOps" ).ToString())
+			: NAME_Math(LOCTEXT( "Math", "Math" ))
+			, NAME_VectorOps(LOCTEXT( "VectorOps", "VectorOps" ))
 		{
 		}
 	};
@@ -3180,9 +3180,9 @@ UMaterialExpressionMakeMaterialAttributes::UMaterialExpressionMakeMaterialAttrib
 	// Structure to hold one-time initialization
 	struct FConstructorStatics
 	{
-		FString NAME_MaterialAttributes;
+		FText NAME_MaterialAttributes;
 		FConstructorStatics()
-			: NAME_MaterialAttributes(LOCTEXT( "MaterialAttributes", "Material Attributes" ).ToString())
+			: NAME_MaterialAttributes(LOCTEXT( "MaterialAttributes", "Material Attributes" ))
 		{
 		}
 	};
@@ -3251,9 +3251,9 @@ UMaterialExpressionBreakMaterialAttributes::UMaterialExpressionBreakMaterialAttr
 	// Structure to hold one-time initialization
 	struct FConstructorStatics
 	{
-		FString NAME_MaterialAttributes;
+		FText NAME_MaterialAttributes;
 		FConstructorStatics()
-			: NAME_MaterialAttributes(LOCTEXT( "MaterialAttributes", "Material Attributes" ).ToString())
+			: NAME_MaterialAttributes(LOCTEXT( "MaterialAttributes", "Material Attributes" ))
 		{
 		}
 	};
@@ -3342,9 +3342,9 @@ UMaterialExpressionFloor::UMaterialExpressionFloor(const FObjectInitializer& Obj
 	// Structure to hold one-time initialization
 	struct FConstructorStatics
 	{
-		FString NAME_Math;
+		FText NAME_Math;
 		FConstructorStatics()
-			: NAME_Math(LOCTEXT( "Math", "Math" ).ToString())
+			: NAME_Math(LOCTEXT( "Math", "Math" ))
 		{
 		}
 	};
@@ -3374,9 +3374,9 @@ UMaterialExpressionCeil::UMaterialExpressionCeil(const FObjectInitializer& Objec
 	// Structure to hold one-time initialization
 	struct FConstructorStatics
 	{
-		FString NAME_Math;
+		FText NAME_Math;
 		FConstructorStatics()
-			: NAME_Math(LOCTEXT( "Math", "Math" ).ToString())
+			: NAME_Math(LOCTEXT( "Math", "Math" ))
 		{
 		}
 	};
@@ -3412,9 +3412,9 @@ UMaterialExpressionFmod::UMaterialExpressionFmod(const FObjectInitializer& Objec
 	// Structure to hold one-time initialization
 	struct FConstructorStatics
 	{
-		FString NAME_Math;
+		FText NAME_Math;
 		FConstructorStatics()
-			: NAME_Math(LOCTEXT( "Math", "Math" ).ToString())
+			: NAME_Math(LOCTEXT( "Math", "Math" ))
 		{
 		}
 	};
@@ -3449,9 +3449,9 @@ UMaterialExpressionFrac::UMaterialExpressionFrac(const FObjectInitializer& Objec
 	// Structure to hold one-time initialization
 	struct FConstructorStatics
 	{
-		FString NAME_Math;
+		FText NAME_Math;
 		FConstructorStatics()
-			: NAME_Math(LOCTEXT( "Math", "Math" ).ToString())
+			: NAME_Math(LOCTEXT( "Math", "Math" ))
 		{
 		}
 	};
@@ -3482,11 +3482,11 @@ UMaterialExpressionDesaturation::UMaterialExpressionDesaturation(const FObjectIn
 	// Structure to hold one-time initialization
 	struct FConstructorStatics
 	{
-		FString NAME_Color;
-		FString NAME_Utility;
+		FText NAME_Color;
+		FText NAME_Utility;
 		FConstructorStatics()
-			: NAME_Color(LOCTEXT( "Color", "Color" ).ToString())
-			, NAME_Utility(LOCTEXT( "Utility", "Utility" ).ToString())
+			: NAME_Color(LOCTEXT( "Color", "Color" ))
+			, NAME_Utility(LOCTEXT( "Utility", "Utility" ))
 		{
 		}
 	};
@@ -3522,9 +3522,9 @@ UMaterialExpressionParameter::UMaterialExpressionParameter(const FObjectInitiali
 	// Structure to hold one-time initialization
 	struct FConstructorStatics
 	{
-		FString NAME_Parameters;
+		FText NAME_Parameters;
 		FConstructorStatics()
-			: NAME_Parameters(LOCTEXT( "Parameters", "Parameters" ).ToString())
+			: NAME_Parameters(LOCTEXT( "Parameters", "Parameters" ))
 		{
 		}
 	};
@@ -3806,9 +3806,9 @@ UMaterialExpressionStaticBool::UMaterialExpressionStaticBool(const FObjectInitia
 	// Structure to hold one-time initialization
 	struct FConstructorStatics
 	{
-		FString NAME_Functions;
+		FText NAME_Functions;
 		FConstructorStatics()
-			: NAME_Functions(LOCTEXT( "Functions", "Functions" ).ToString())
+			: NAME_Functions(LOCTEXT( "Functions", "Functions" ))
 		{
 		}
 	};
@@ -3845,9 +3845,9 @@ UMaterialExpressionStaticSwitch::UMaterialExpressionStaticSwitch(const FObjectIn
 	// Structure to hold one-time initialization
 	struct FConstructorStatics
 	{
-		FString NAME_Functions;
+		FText NAME_Functions;
 		FConstructorStatics()
-			: NAME_Functions(LOCTEXT( "Functions", "Functions" ).ToString())
+			: NAME_Functions(LOCTEXT( "Functions", "Functions" ))
 		{
 		}
 	};
@@ -3942,9 +3942,9 @@ UMaterialExpressionQualitySwitch::UMaterialExpressionQualitySwitch(const FObject
 	// Structure to hold one-time initialization
 	struct FConstructorStatics
 	{
-		FString NAME_Utility;
+		FText NAME_Utility;
 		FConstructorStatics()
-			: NAME_Utility(LOCTEXT( "Utility", "Utility" ).ToString())
+			: NAME_Utility(LOCTEXT( "Utility", "Utility" ))
 		{
 		}
 	};
@@ -4045,9 +4045,9 @@ UMaterialExpressionFeatureLevelSwitch::UMaterialExpressionFeatureLevelSwitch(con
 	// Structure to hold one-time initialization
 	struct FConstructorStatics
 	{
-		FString NAME_Utility;
+		FText NAME_Utility;
 		FConstructorStatics()
-			: NAME_Utility(LOCTEXT( "Utility", "Utility" ).ToString())
+			: NAME_Utility(LOCTEXT( "Utility", "Utility" ))
 		{
 		}
 	};
@@ -4157,11 +4157,11 @@ UMaterialExpressionNormalize::UMaterialExpressionNormalize(const FObjectInitiali
 	// Structure to hold one-time initialization
 	struct FConstructorStatics
 	{
-		FString NAME_Math;
-		FString NAME_VectorOps;
+		FText NAME_Math;
+		FText NAME_VectorOps;
 		FConstructorStatics()
-			: NAME_Math(LOCTEXT( "Math", "Math" ).ToString())
-			, NAME_VectorOps(LOCTEXT( "VectorOps", "VectorOps" ).ToString())
+			: NAME_Math(LOCTEXT( "Math", "Math" ))
+			, NAME_VectorOps(LOCTEXT( "VectorOps", "VectorOps" ))
 		{
 		}
 	};
@@ -4188,9 +4188,9 @@ UMaterialExpressionVertexColor::UMaterialExpressionVertexColor(const FObjectInit
 	// Structure to hold one-time initialization
 	struct FConstructorStatics
 	{
-		FString NAME_Constants;
+		FText NAME_Constants;
 		FConstructorStatics()
-			: NAME_Constants(LOCTEXT( "Constants", "Constants" ).ToString())
+			: NAME_Constants(LOCTEXT( "Constants", "Constants" ))
 		{
 		}
 	};
@@ -4222,11 +4222,11 @@ UMaterialExpressionParticleColor::UMaterialExpressionParticleColor(const FObject
 	// Structure to hold one-time initialization
 	struct FConstructorStatics
 	{
-		FString NAME_Particles;
-		FString NAME_Constants;
+		FText NAME_Particles;
+		FText NAME_Constants;
 		FConstructorStatics()
-			: NAME_Particles(LOCTEXT( "Particles", "Particles" ).ToString())
-			, NAME_Constants(LOCTEXT( "Constants", "Constants" ).ToString())
+			: NAME_Particles(LOCTEXT( "Particles", "Particles" ))
+			, NAME_Constants(LOCTEXT( "Constants", "Constants" ))
 		{
 		}
 	};
@@ -4261,13 +4261,13 @@ UMaterialExpressionParticlePositionWS::UMaterialExpressionParticlePositionWS(con
 	// Structure to hold one-time initialization
 	struct FConstructorStatics
 	{
-		FString NAME_Particles;
-		FString NAME_Coordinates;
-		FString NAME_Vectors;
+		FText NAME_Particles;
+		FText NAME_Coordinates;
+		FText NAME_Vectors;
 		FConstructorStatics()
-			: NAME_Particles(LOCTEXT( "Particles", "Particles" ).ToString())
-			, NAME_Coordinates(LOCTEXT( "Coordinates", "Coordinates" ).ToString())
-			, NAME_Vectors(LOCTEXT( "Vectors", "Vectors" ).ToString())
+			: NAME_Particles(LOCTEXT( "Particles", "Particles" ))
+			, NAME_Coordinates(LOCTEXT( "Coordinates", "Coordinates" ))
+			, NAME_Vectors(LOCTEXT( "Vectors", "Vectors" ))
 		{
 		}
 	};
@@ -4295,11 +4295,11 @@ UMaterialExpressionParticleRadius::UMaterialExpressionParticleRadius(const FObje
 	// Structure to hold one-time initialization
 	struct FConstructorStatics
 	{
-		FString NAME_Particles;
-		FString NAME_Constants;
+		FText NAME_Particles;
+		FText NAME_Constants;
 		FConstructorStatics()
-			: NAME_Particles(LOCTEXT( "Particles", "Particles" ).ToString())
-			, NAME_Constants(LOCTEXT( "Constants", "Constants" ).ToString())
+			: NAME_Particles(LOCTEXT( "Particles", "Particles" ))
+			, NAME_Constants(LOCTEXT( "Constants", "Constants" ))
 		{
 		}
 	};
@@ -4327,11 +4327,11 @@ UMaterialExpressionDynamicParameter::UMaterialExpressionDynamicParameter(const F
 	// Structure to hold one-time initialization
 	struct FConstructorStatics
 	{
-		FString NAME_Particles;
-		FString NAME_Parameters;
+		FText NAME_Particles;
+		FText NAME_Parameters;
 		FConstructorStatics()
-			: NAME_Particles(LOCTEXT( "Particles", "Particles" ).ToString())
-			, NAME_Parameters(LOCTEXT( "Parameters", "Parameters" ).ToString())
+			: NAME_Particles(LOCTEXT( "Particles", "Particles" ))
+			, NAME_Parameters(LOCTEXT( "Parameters", "Parameters" ))
 		{
 		}
 	};
@@ -4455,9 +4455,9 @@ UMaterialExpressionParticleSubUV::UMaterialExpressionParticleSubUV(const FObject
 	// Structure to hold one-time initialization
 	struct FConstructorStatics
 	{
-		FString NAME_Particles;
+		FText NAME_Particles;
 		FConstructorStatics()
-			: NAME_Particles(LOCTEXT( "Particles", "Particles" ).ToString())
+			: NAME_Particles(LOCTEXT( "Particles", "Particles" ))
 		{
 		}
 	};
@@ -4505,9 +4505,9 @@ UMaterialExpressionParticleMacroUV::UMaterialExpressionParticleMacroUV(const FOb
 	// Structure to hold one-time initialization
 	struct FConstructorStatics
 	{
-		FString NAME_Particles;
+		FText NAME_Particles;
 		FConstructorStatics()
-			: NAME_Particles(LOCTEXT( "Particles", "Particles" ).ToString())
+			: NAME_Particles(LOCTEXT( "Particles", "Particles" ))
 		{
 		}
 	};
@@ -4536,9 +4536,9 @@ UMaterialExpressionLightVector::UMaterialExpressionLightVector(const FObjectInit
 	// Structure to hold one-time initialization
 	struct FConstructorStatics
 	{
-		FString NAME_Vectors;
+		FText NAME_Vectors;
 		FConstructorStatics()
-			: NAME_Vectors(LOCTEXT( "Vectors", "Vectors" ).ToString())
+			: NAME_Vectors(LOCTEXT( "Vectors", "Vectors" ))
 		{
 		}
 	};
@@ -4566,9 +4566,9 @@ UMaterialExpressionScreenPosition::UMaterialExpressionScreenPosition(const FObje
 	// Structure to hold one-time initialization
 	struct FConstructorStatics
 	{
-		FString NAME_Coordinates;
+		FText NAME_Coordinates;
 		FConstructorStatics()
-			: NAME_Coordinates(LOCTEXT( "Coordinates", "Coordinates" ).ToString())
+			: NAME_Coordinates(LOCTEXT( "Coordinates", "Coordinates" ))
 		{
 		}
 	};
@@ -4594,9 +4594,9 @@ UMaterialExpressionViewSize::UMaterialExpressionViewSize(const FObjectInitialize
 	// Structure to hold one-time initialization
 	struct FConstructorStatics
 	{
-		FString NAME_Coordinates;
+		FText NAME_Coordinates;
 		FConstructorStatics()
-			: NAME_Coordinates(LOCTEXT( "Coordinates", "Coordinates" ).ToString())
+			: NAME_Coordinates(LOCTEXT( "Coordinates", "Coordinates" ))
 		{
 		}
 	};
@@ -4622,9 +4622,9 @@ UMaterialExpressionSceneTexelSize::UMaterialExpressionSceneTexelSize(const FObje
 	// Structure to hold one-time initialization
 	struct FConstructorStatics
 	{
-		FString NAME_Coordinates;
+		FText NAME_Coordinates;
 		FConstructorStatics()
-			: NAME_Coordinates(LOCTEXT( "Coordinates", "Coordinates" ).ToString())
+			: NAME_Coordinates(LOCTEXT( "Coordinates", "Coordinates" ))
 		{
 		}
 	};
@@ -4650,9 +4650,9 @@ UMaterialExpressionSquareRoot::UMaterialExpressionSquareRoot(const FObjectInitia
 	// Structure to hold one-time initialization
 	struct FConstructorStatics
 	{
-		FString NAME_Math;
+		FText NAME_Math;
 		FConstructorStatics()
-			: NAME_Math(LOCTEXT( "Math", "Math" ).ToString())
+			: NAME_Math(LOCTEXT( "Math", "Math" ))
 		{
 		}
 	};
@@ -4684,9 +4684,9 @@ UMaterialExpressionPixelDepth::UMaterialExpressionPixelDepth(const FObjectInitia
 	// Structure to hold one-time initialization
 	struct FConstructorStatics
 	{
-		FString NAME_Depth;
+		FText NAME_Depth;
 		FConstructorStatics()
-			: NAME_Depth(LOCTEXT( "Depth", "Depth" ).ToString())
+			: NAME_Depth(LOCTEXT( "Depth", "Depth" ))
 		{
 		}
 	};
@@ -4721,9 +4721,9 @@ UMaterialExpressionSceneDepth::UMaterialExpressionSceneDepth(const FObjectInitia
 	// Structure to hold one-time initialization
 	struct FConstructorStatics
 	{
-		FString NAME_Depth;
+		FText NAME_Depth;
 		FConstructorStatics()
-			: NAME_Depth(LOCTEXT( "Depth", "Depth" ).ToString())
+			: NAME_Depth(LOCTEXT( "Depth", "Depth" ))
 		{
 		}
 	};
@@ -4805,9 +4805,9 @@ UMaterialExpressionSceneTexture::UMaterialExpressionSceneTexture(const FObjectIn
 	// Structure to hold one-time initialization
 	struct FConstructorStatics
 	{
-		FString NAME_Texture;
+		FText NAME_Texture;
 		FConstructorStatics()
-			: NAME_Texture(LOCTEXT( "Texture", "Texture" ).ToString())
+			: NAME_Texture(LOCTEXT( "Texture", "Texture" ))
 		{
 		}
 	};
@@ -4890,9 +4890,9 @@ UMaterialExpressionSceneColor::UMaterialExpressionSceneColor(const FObjectInitia
 	// Structure to hold one-time initialization
 	struct FConstructorStatics
 	{
-		FString NAME_Texture;
+		FText NAME_Texture;
 		FConstructorStatics()
-			: NAME_Texture(LOCTEXT( "Texture", "Texture" ).ToString())
+			: NAME_Texture(LOCTEXT( "Texture", "Texture" ))
 		{
 		}
 	};
@@ -4958,9 +4958,9 @@ UMaterialExpressionPower::UMaterialExpressionPower(const FObjectInitializer& Obj
 	// Structure to hold one-time initialization
 	struct FConstructorStatics
 	{
-		FString NAME_Math;
+		FText NAME_Math;
 		FConstructorStatics()
-			: NAME_Math(LOCTEXT( "Math", "Math" ).ToString())
+			: NAME_Math(LOCTEXT( "Math", "Math" ))
 		{
 		}
 	};
@@ -5003,9 +5003,9 @@ UMaterialExpressionLogarithm2::UMaterialExpressionLogarithm2(const FObjectInitia
 	// Structure to hold one-time initialization
 	struct FConstructorStatics
 	{
-		FString NAME_Math;
+		FText NAME_Math;
 		FConstructorStatics()
-			: NAME_Math(LOCTEXT( "Math", "Math" ).ToString())
+			: NAME_Math(LOCTEXT( "Math", "Math" ))
 		{
 		}
 	};
@@ -5047,9 +5047,9 @@ UMaterialExpressionIf::UMaterialExpressionIf(const FObjectInitializer& ObjectIni
 	// Structure to hold one-time initialization
 	struct FConstructorStatics
 	{
-		FString NAME_Math;
+		FText NAME_Math;
 		FConstructorStatics()
-			: NAME_Math(LOCTEXT( "Math", "Math" ).ToString())
+			: NAME_Math(LOCTEXT( "Math", "Math" ))
 		{
 		}
 	};
@@ -5124,9 +5124,9 @@ UMaterialExpressionOneMinus::UMaterialExpressionOneMinus(const FObjectInitialize
 	// Structure to hold one-time initialization
 	struct FConstructorStatics
 	{
-		FString NAME_Math;
+		FText NAME_Math;
 		FConstructorStatics()
-			: NAME_Math(LOCTEXT( "Math", "Math" ).ToString())
+			: NAME_Math(LOCTEXT( "Math", "Math" ))
 		{
 		}
 	};
@@ -5155,9 +5155,9 @@ UMaterialExpressionAbs::UMaterialExpressionAbs(const FObjectInitializer& ObjectI
 	// Structure to hold one-time initialization
 	struct FConstructorStatics
 	{
-		FString NAME_Math;
+		FText NAME_Math;
 		FConstructorStatics()
-			: NAME_Math(LOCTEXT( "Math", "Math" ).ToString())
+			: NAME_Math(LOCTEXT( "Math", "Math" ))
 		{
 		}
 	};
@@ -5262,9 +5262,9 @@ UMaterialExpressionTransform::UMaterialExpressionTransform(const FObjectInitiali
 	// Structure to hold one-time initialization
 	struct FConstructorStatics
 	{
-		FString NAME_VectorOps;
+		FText NAME_VectorOps;
 		FConstructorStatics()
-			: NAME_VectorOps(LOCTEXT( "VectorOps", "VectorOps" ).ToString())
+			: NAME_VectorOps(LOCTEXT( "VectorOps", "VectorOps" ))
 		{
 		}
 	};
@@ -5284,9 +5284,9 @@ UMaterialExpressionTransformPosition::UMaterialExpressionTransformPosition(const
 	// Structure to hold one-time initialization
 	struct FConstructorStatics
 	{
-		FString NAME_VectorOps;
+		FText NAME_VectorOps;
 		FConstructorStatics()
-			: NAME_VectorOps(LOCTEXT( "VectorOps", "VectorOps" ).ToString())
+			: NAME_VectorOps(LOCTEXT( "VectorOps", "VectorOps" ))
 		{
 		}
 	};
@@ -5348,9 +5348,9 @@ UMaterialExpressionComment::UMaterialExpressionComment(const FObjectInitializer&
 	// Structure to hold one-time initialization
 	struct FConstructorStatics
 	{
-		FString NAME_Utility;
+		FText NAME_Utility;
 		FConstructorStatics()
-			: NAME_Utility(LOCTEXT( "Utility", "Utility" ).ToString())
+			: NAME_Utility(LOCTEXT( "Utility", "Utility" ))
 		{
 		}
 	};
@@ -5423,9 +5423,9 @@ UMaterialExpressionFresnel::UMaterialExpressionFresnel(const FObjectInitializer&
 	// Structure to hold one-time initialization
 	struct FConstructorStatics
 	{
-		FString NAME_Utility;
+		FText NAME_Utility;
 		FConstructorStatics()
-			: NAME_Utility(LOCTEXT( "Utility", "Utility" ).ToString())
+			: NAME_Utility(LOCTEXT( "Utility", "Utility" ))
 		{
 		}
 	};
@@ -5463,11 +5463,11 @@ UMaterialExpressionFontSample::UMaterialExpressionFontSample(const FObjectInitia
 	// Structure to hold one-time initialization
 	struct FConstructorStatics
 	{
-		FString NAME_Font;
-		FString NAME_Texture;
+		FText NAME_Font;
+		FText NAME_Texture;
 		FConstructorStatics()
-			: NAME_Font(LOCTEXT( "Font", "Font" ).ToString())
-			, NAME_Texture(LOCTEXT( "Texture", "Texture" ).ToString())
+			: NAME_Font(LOCTEXT( "Font", "Font" ))
+			, NAME_Texture(LOCTEXT( "Texture", "Texture" ))
 		{
 		}
 	};
@@ -5591,11 +5591,11 @@ UMaterialExpressionFontSampleParameter::UMaterialExpressionFontSampleParameter(c
 	// Structure to hold one-time initialization
 	struct FConstructorStatics
 	{
-		FString NAME_Font;
-		FString NAME_Parameters;
+		FText NAME_Font;
+		FText NAME_Parameters;
 		FConstructorStatics()
-			: NAME_Font(LOCTEXT( "Font", "Font" ).ToString())
-			, NAME_Parameters(LOCTEXT( "Parameters", "Parameters" ).ToString())
+			: NAME_Font(LOCTEXT( "Font", "Font" ))
+			, NAME_Parameters(LOCTEXT( "Parameters", "Parameters" ))
 		{
 		}
 	};
@@ -5717,9 +5717,9 @@ UMaterialExpressionWorldPosition::UMaterialExpressionWorldPosition(const FObject
 	// Structure to hold one-time initialization
 	struct FConstructorStatics
 	{
-		FString NAME_Coordinates;
+		FText NAME_Coordinates;
 		FConstructorStatics()
-			: NAME_Coordinates(LOCTEXT( "Coordinates", "Coordinates" ).ToString())
+			: NAME_Coordinates(LOCTEXT( "Coordinates", "Coordinates" ))
 		{
 		}
 	};
@@ -5779,11 +5779,11 @@ UMaterialExpressionObjectPositionWS::UMaterialExpressionObjectPositionWS(const F
 	// Structure to hold one-time initialization
 	struct FConstructorStatics
 	{
-		FString NAME_Vectors;
-		FString NAME_Coordinates;
+		FText NAME_Vectors;
+		FText NAME_Coordinates;
 		FConstructorStatics()
-			: NAME_Vectors(LOCTEXT( "Vectors", "Vectors" ).ToString())
-			, NAME_Coordinates(LOCTEXT( "Coordinates", "Coordinates" ).ToString())
+			: NAME_Vectors(LOCTEXT( "Vectors", "Vectors" ))
+			, NAME_Coordinates(LOCTEXT( "Coordinates", "Coordinates" ))
 		{
 		}
 	};
@@ -5818,9 +5818,9 @@ UMaterialExpressionObjectRadius::UMaterialExpressionObjectRadius(const FObjectIn
 	// Structure to hold one-time initialization
 	struct FConstructorStatics
 	{
-		FString NAME_Coordinates;
+		FText NAME_Coordinates;
 		FConstructorStatics()
-			: NAME_Coordinates(LOCTEXT( "Coordinates", "Coordinates" ).ToString())
+			: NAME_Coordinates(LOCTEXT( "Coordinates", "Coordinates" ))
 		{
 		}
 	};
@@ -5855,9 +5855,9 @@ UMaterialExpressionObjectBounds::UMaterialExpressionObjectBounds(const FObjectIn
 	// Structure to hold one-time initialization
 	struct FConstructorStatics
 	{
-		FString NAME_Vectors;
+		FText NAME_Vectors;
 		FConstructorStatics()
-			: NAME_Vectors(LOCTEXT( "Vectors", "Vectors" ).ToString())
+			: NAME_Vectors(LOCTEXT( "Vectors", "Vectors" ))
 		{
 		}
 	};
@@ -5892,9 +5892,9 @@ UMaterialExpressionDistanceCullFade::UMaterialExpressionDistanceCullFade(const F
 	// Structure to hold one-time initialization
 	struct FConstructorStatics
 	{
-		FString NAME_Constants;
+		FText NAME_Constants;
 		FConstructorStatics()
-			: NAME_Constants(LOCTEXT( "Constants", "Constants" ).ToString())
+			: NAME_Constants(LOCTEXT( "Constants", "Constants" ))
 		{
 		}
 	};
@@ -5925,11 +5925,11 @@ UMaterialExpressionActorPositionWS::UMaterialExpressionActorPositionWS(const FOb
 	// Structure to hold one-time initialization
 	struct FConstructorStatics
 	{
-		FString NAME_Vectors;
-		FString NAME_Coordinates;
+		FText NAME_Vectors;
+		FText NAME_Coordinates;
 		FConstructorStatics()
-			: NAME_Vectors(LOCTEXT( "Vectors", "Vectors" ).ToString())
-			, NAME_Coordinates(LOCTEXT( "Coordinates", "Coordinates" ).ToString())
+			: NAME_Vectors(LOCTEXT( "Vectors", "Vectors" ))
+			, NAME_Coordinates(LOCTEXT( "Coordinates", "Coordinates" ))
 		{
 		}
 	};
@@ -5961,9 +5961,9 @@ UMaterialExpressionDeriveNormalZ::UMaterialExpressionDeriveNormalZ(const FObject
 	// Structure to hold one-time initialization
 	struct FConstructorStatics
 	{
-		FString NAME_VectorOps;
+		FText NAME_VectorOps;
 		FConstructorStatics()
-			: NAME_VectorOps(LOCTEXT( "VectorOps", "VectorOps" ).ToString())
+			: NAME_VectorOps(LOCTEXT( "VectorOps", "VectorOps" ))
 		{
 		}
 	};
@@ -6004,9 +6004,9 @@ UMaterialExpressionConstantBiasScale::UMaterialExpressionConstantBiasScale(const
 	// Structure to hold one-time initialization
 	struct FConstructorStatics
 	{
-		FString NAME_Utility;
+		FText NAME_Utility;
 		FConstructorStatics()
-			: NAME_Utility(LOCTEXT( "Utility", "Utility" ).ToString())
+			: NAME_Utility(LOCTEXT( "Utility", "Utility" ))
 		{
 		}
 	};
@@ -6044,9 +6044,9 @@ UMaterialExpressionCustom::UMaterialExpressionCustom(const FObjectInitializer& O
 	// Structure to hold one-time initialization
 	struct FConstructorStatics
 	{
-		FString NAME_Custom;
+		FText NAME_Custom;
 		FConstructorStatics()
-			: NAME_Custom(LOCTEXT( "Custom", "Custom" ).ToString())
+			: NAME_Custom(LOCTEXT( "Custom", "Custom" ))
 		{
 		}
 	};
@@ -6180,7 +6180,7 @@ uint32 UMaterialExpressionCustom::GetOutputType(int32 OutputIndex)
 UMaterialFunction::UMaterialFunction(const FObjectInitializer& ObjectInitializer)
 	: Super(ObjectInitializer)
 {
-	LibraryCategories.Add(TEXT("Misc"));
+	LibraryCategoriesText.Add(LOCTEXT("Misc", "Misc"));
 
 #if WITH_EDITORONLY_DATA
 	PreviewMaterial = NULL;
@@ -6273,6 +6273,14 @@ void UMaterialFunction::Serialize(FArchive& Ar)
 	{
 		GMaterialFunctionsThatNeedSamplerFixup.Set(this);
 	}
+
+	if (Ar.UE4Ver() < VER_UE4_LIBRARY_CATEGORIES_AS_FTEXT)
+	{
+		for (FString& Category : LibraryCategories_DEPRECATED)
+		{
+			LibraryCategoriesText.Add(FText::FromString(Category));
+		}
+	}
 #endif // #if WITH_EDITOR
 }
 
@@ -6701,9 +6709,9 @@ UMaterialExpressionMaterialFunctionCall::UMaterialExpressionMaterialFunctionCall
 	// Structure to hold one-time initialization
 	struct FConstructorStatics
 	{
-		FString NAME_Functions;
+		FText NAME_Functions;
 		FConstructorStatics()
-			: NAME_Functions(LOCTEXT( "Functions", "Functions" ).ToString())
+			: NAME_Functions(LOCTEXT( "Functions", "Functions" ))
 		{
 		}
 	};
@@ -7236,9 +7244,9 @@ UMaterialExpressionFunctionInput::UMaterialExpressionFunctionInput(const FObject
 	// Structure to hold one-time initialization
 	struct FConstructorStatics
 	{
-		FString NAME_Functions;
+		FText NAME_Functions;
 		FConstructorStatics()
-			: NAME_Functions(LOCTEXT( "Functions", "Functions" ).ToString())
+			: NAME_Functions(LOCTEXT( "Functions", "Functions" ))
 		{
 		}
 	};
@@ -7530,9 +7538,9 @@ UMaterialExpressionFunctionOutput::UMaterialExpressionFunctionOutput(const FObje
 	// Structure to hold one-time initialization
 	struct FConstructorStatics
 	{
-		FString NAME_Functions;
+		FText NAME_Functions;
 		FConstructorStatics()
-			: NAME_Functions(LOCTEXT( "Functions", "Functions" ).ToString())
+			: NAME_Functions(LOCTEXT( "Functions", "Functions" ))
 		{
 		}
 	};
@@ -7695,9 +7703,9 @@ UMaterialExpressionCollectionParameter::UMaterialExpressionCollectionParameter(c
 	// Structure to hold one-time initialization
 	struct FConstructorStatics
 	{
-		FString NAME_Parameters;
+		FText NAME_Parameters;
 		FConstructorStatics()
-			: NAME_Parameters(LOCTEXT( "Parameters", "Parameters" ).ToString())
+			: NAME_Parameters(LOCTEXT( "Parameters", "Parameters" ))
 		{
 		}
 	};
@@ -7821,9 +7829,9 @@ UMaterialExpressionLightmapUVs::UMaterialExpressionLightmapUVs(const FObjectInit
 	// Structure to hold one-time initialization
 	struct FConstructorStatics
 	{
-		FString NAME_Coordinates;
+		FText NAME_Coordinates;
 		FConstructorStatics()
-			: NAME_Coordinates(LOCTEXT( "Coordinates", "Coordinates" ).ToString())
+			: NAME_Coordinates(LOCTEXT( "Coordinates", "Coordinates" ))
 		{
 		}
 	};
@@ -7861,9 +7869,9 @@ UMaterialExpressionLightmassReplace::UMaterialExpressionLightmassReplace(const F
 	// Structure to hold one-time initialization
 	struct FConstructorStatics
 	{
-		FString NAME_Utility;
+		FText NAME_Utility;
 		FConstructorStatics()
-			: NAME_Utility(LOCTEXT( "Utility", "Utility" ).ToString())
+			: NAME_Utility(LOCTEXT( "Utility", "Utility" ))
 		{
 		}
 	};
@@ -7905,9 +7913,9 @@ UMaterialExpressionGIReplace::UMaterialExpressionGIReplace(const FObjectInitiali
 	// Structure to hold one-time initialization
 	struct FConstructorStatics
 	{
-		FString NAME_Utility;
+		FText NAME_Utility;
 		FConstructorStatics()
-			: NAME_Utility(LOCTEXT( "Utility", "Utility" ).ToString())
+			: NAME_Utility(LOCTEXT( "Utility", "Utility" ))
 		{
 		}
 	};
@@ -7945,11 +7953,11 @@ UMaterialExpressionObjectOrientation::UMaterialExpressionObjectOrientation(const
 	// Structure to hold one-time initialization
 	struct FConstructorStatics
 	{
-		FString NAME_Vectors;
-		FString NAME_Coordinates;
+		FText NAME_Vectors;
+		FText NAME_Coordinates;
 		FConstructorStatics()
-			: NAME_Vectors(LOCTEXT( "Vectors", "Vectors" ).ToString())
-			, NAME_Coordinates(LOCTEXT( "Coordinates", "Coordinates" ).ToString())
+			: NAME_Vectors(LOCTEXT( "Vectors", "Vectors" ))
+			, NAME_Coordinates(LOCTEXT( "Coordinates", "Coordinates" ))
 		{
 		}
 	};
@@ -7981,9 +7989,9 @@ UMaterialExpressionRotateAboutAxis::UMaterialExpressionRotateAboutAxis(const FOb
 	// Structure to hold one-time initialization
 	struct FConstructorStatics
 	{
-		FString NAME_Utility;
+		FText NAME_Utility;
 		FConstructorStatics()
-			: NAME_Utility(LOCTEXT( "Utility", "Utility" ).ToString())
+			: NAME_Utility(LOCTEXT( "Utility", "Utility" ))
 		{
 		}
 	};
@@ -8064,9 +8072,9 @@ UMaterialExpressionSphereMask::UMaterialExpressionSphereMask(const FObjectInitia
 	// Structure to hold one-time initialization
 	struct FConstructorStatics
 	{
-		FString NAME_Utility;
+		FText NAME_Utility;
 		FConstructorStatics()
-			: NAME_Utility(LOCTEXT( "Utility", "Utility" ).ToString())
+			: NAME_Utility(LOCTEXT( "Utility", "Utility" ))
 		{
 		}
 	};
@@ -8146,9 +8154,9 @@ UMaterialExpressionNoise::UMaterialExpressionNoise(const FObjectInitializer& Obj
 	// Structure to hold one-time initialization
 	struct FConstructorStatics
 	{
-		FString NAME_Utility;
+		FText NAME_Utility;
 		FConstructorStatics()
-			: NAME_Utility(LOCTEXT( "Utility", "Utility" ).ToString())
+			: NAME_Utility(LOCTEXT( "Utility", "Utility" ))
 		{
 		}
 	};
@@ -8208,9 +8216,9 @@ UMaterialExpressionBlackBody::UMaterialExpressionBlackBody(const FObjectInitiali
 	// Structure to hold one-time initialization
 	struct FConstructorStatics
 	{
-		FString NAME_Utility;
+		FText NAME_Utility;
 		FConstructorStatics()
-			: NAME_Utility(LOCTEXT( "Utility", "Utility" ).ToString())
+			: NAME_Utility(LOCTEXT( "Utility", "Utility" ))
 		{
 		}
 	};
@@ -8251,9 +8259,9 @@ UMaterialExpressionDistance::UMaterialExpressionDistance(const FObjectInitialize
 	// Structure to hold one-time initialization
 	struct FConstructorStatics
 	{
-		FString NAME_Utility;
+		FText NAME_Utility;
 		FConstructorStatics()
-			: NAME_Utility(LOCTEXT( "Utility", "Utility" ).ToString())
+			: NAME_Utility(LOCTEXT( "Utility", "Utility" ))
 		{
 		}
 	};
@@ -8292,9 +8300,9 @@ UMaterialExpressionTwoSidedSign::UMaterialExpressionTwoSidedSign(const FObjectIn
 	// Structure to hold one-time initialization
 	struct FConstructorStatics
 	{
-		FString NAME_Constants;
+		FText NAME_Constants;
 		FConstructorStatics()
-			: NAME_Constants(LOCTEXT( "Constants", "Constants" ).ToString())
+			: NAME_Constants(LOCTEXT( "Constants", "Constants" ))
 		{
 		}
 	};
@@ -8320,11 +8328,11 @@ UMaterialExpressionVertexNormalWS::UMaterialExpressionVertexNormalWS(const FObje
 	// Structure to hold one-time initialization
 	struct FConstructorStatics
 	{
-		FString NAME_Vectors;
-		FString NAME_Coordinates;
+		FText NAME_Vectors;
+		FText NAME_Coordinates;
 		FConstructorStatics()
-			: NAME_Vectors(LOCTEXT( "Vectors", "Vectors" ).ToString())
-			, NAME_Coordinates(LOCTEXT( "Coordinates", "Coordinates" ).ToString())
+			: NAME_Vectors(LOCTEXT( "Vectors", "Vectors" ))
+			, NAME_Coordinates(LOCTEXT( "Coordinates", "Coordinates" ))
 		{
 		}
 	};
@@ -8351,11 +8359,11 @@ UMaterialExpressionPixelNormalWS::UMaterialExpressionPixelNormalWS(const FObject
 	// Structure to hold one-time initialization
 	struct FConstructorStatics
 	{
-		FString NAME_Vectors;
-		FString NAME_Coordinates;
+		FText NAME_Vectors;
+		FText NAME_Coordinates;
 		FConstructorStatics()
-			: NAME_Vectors(LOCTEXT( "Vectors", "Vectors" ).ToString())
-			, NAME_Coordinates(LOCTEXT( "Coordinates", "Coordinates" ).ToString())
+			: NAME_Vectors(LOCTEXT( "Vectors", "Vectors" ))
+			, NAME_Coordinates(LOCTEXT( "Coordinates", "Coordinates" ))
 		{
 		}
 	};
@@ -8385,9 +8393,9 @@ UMaterialExpressionPerInstanceRandom::UMaterialExpressionPerInstanceRandom(const
 	// Structure to hold one-time initialization
 	struct FConstructorStatics
 	{
-		FString NAME_Constants;
+		FText NAME_Constants;
 		FConstructorStatics()
-			: NAME_Constants(LOCTEXT( "Constants", "Constants" ).ToString())
+			: NAME_Constants(LOCTEXT( "Constants", "Constants" ))
 		{
 		}
 	};
@@ -8416,9 +8424,9 @@ UMaterialExpressionPerInstanceFadeAmount::UMaterialExpressionPerInstanceFadeAmou
 	// Structure to hold one-time initialization
 	struct FConstructorStatics
 	{
-		FString NAME_Constants;
+		FText NAME_Constants;
 		FConstructorStatics()
-			: NAME_Constants(LOCTEXT( "Constants", "Constants" ).ToString())
+			: NAME_Constants(LOCTEXT( "Constants", "Constants" ))
 		{
 		}
 	};
@@ -8449,11 +8457,11 @@ UMaterialExpressionAntialiasedTextureMask::UMaterialExpressionAntialiasedTexture
 	struct FConstructorStatics
 	{
 		ConstructorHelpers::FObjectFinder<UTexture2D> DefaultTexture;
-		FString NAME_Utility;
+		FText NAME_Utility;
 		FName NAME_None;
 		FConstructorStatics()
 			: DefaultTexture(TEXT("/Engine/EngineResources/DefaultTexture"))
-			, NAME_Utility(LOCTEXT( "Utility", "Utility" ).ToString())
+			, NAME_Utility(LOCTEXT( "Utility", "Utility" ))
 			, NAME_None(TEXT("None"))
 		{
 		}
@@ -8550,9 +8558,9 @@ UMaterialExpressionDecalMipmapLevel::UMaterialExpressionDecalMipmapLevel(const F
 	// Structure to hold one-time initialization
 	struct FConstructorStatics
 	{
-		FString NAME_Vectors;
+		FText NAME_Vectors;
 		FConstructorStatics()
-			: NAME_Vectors(LOCTEXT("Utils", "Utils").ToString())
+			: NAME_Vectors(LOCTEXT("Utils", "Utils"))
 		{
 		}
 	};
@@ -8599,11 +8607,11 @@ UMaterialExpressionDepthFade::UMaterialExpressionDepthFade(const FObjectInitiali
 	// Structure to hold one-time initialization
 	struct FConstructorStatics
 	{
-		FString NAME_Depth;
-		FString NAME_Utility;
+		FText NAME_Depth;
+		FText NAME_Utility;
 		FConstructorStatics()
-			: NAME_Depth(LOCTEXT( "Depth", "Depth" ).ToString())
-			, NAME_Utility(LOCTEXT( "Utility", "Utility" ).ToString())
+			: NAME_Depth(LOCTEXT( "Depth", "Depth" ))
+			, NAME_Utility(LOCTEXT( "Utility", "Utility" ))
 		{
 		}
 	};
@@ -8635,9 +8643,9 @@ UMaterialExpressionSphericalParticleOpacity::UMaterialExpressionSphericalParticl
 	// Structure to hold one-time initialization
 	struct FConstructorStatics
 	{
-		FString NAME_Particles;
+		FText NAME_Particles;
 		FConstructorStatics()
-			: NAME_Particles(LOCTEXT( "Particles", "Particles" ).ToString())
+			: NAME_Particles(LOCTEXT( "Particles", "Particles" ))
 		{
 		}
 	};
@@ -8663,9 +8671,9 @@ UMaterialExpressionDepthOfFieldFunction::UMaterialExpressionDepthOfFieldFunction
 	// Structure to hold one-time initialization
 	struct FConstructorStatics
 	{
-		FString NAME_Utility;
+		FText NAME_Utility;
 		FConstructorStatics()
-			: NAME_Utility(LOCTEXT( "Utility", "Utility" ).ToString())
+			: NAME_Utility(LOCTEXT( "Utility", "Utility" ))
 		{
 		}
 	};
@@ -8714,9 +8722,9 @@ UMaterialExpressionDDX::UMaterialExpressionDDX(const FObjectInitializer& ObjectI
 	// Structure to hold one-time initialization
 	struct FConstructorStatics
 	{
-		FString NAME_Utility;
+		FText NAME_Utility;
 		FConstructorStatics()
-			: NAME_Utility(LOCTEXT( "Utility", "Utility" ).ToString())
+			: NAME_Utility(LOCTEXT( "Utility", "Utility" ))
 		{
 		}
 	};
@@ -8759,9 +8767,9 @@ UMaterialExpressionDDY::UMaterialExpressionDDY(const FObjectInitializer& ObjectI
 	// Structure to hold one-time initialization
 	struct FConstructorStatics
 	{
-		FString NAME_Utility;
+		FText NAME_Utility;
 		FConstructorStatics()
-			: NAME_Utility(LOCTEXT( "Utility", "Utility" ).ToString())
+			: NAME_Utility(LOCTEXT( "Utility", "Utility" ))
 		{
 		}
 	};
@@ -8804,11 +8812,11 @@ UMaterialExpressionParticleRelativeTime::UMaterialExpressionParticleRelativeTime
 	// Structure to hold one-time initialization
 	struct FConstructorStatics
 	{
-		FString NAME_Particles;
-		FString NAME_Constants;
+		FText NAME_Particles;
+		FText NAME_Constants;
 		FConstructorStatics()
-			: NAME_Particles(LOCTEXT( "Particles", "Particles" ).ToString())
-			, NAME_Constants(LOCTEXT( "Constants", "Constants" ).ToString())
+			: NAME_Particles(LOCTEXT( "Particles", "Particles" ))
+			, NAME_Constants(LOCTEXT( "Constants", "Constants" ))
 		{
 		}
 	};
@@ -8838,11 +8846,11 @@ UMaterialExpressionParticleMotionBlurFade::UMaterialExpressionParticleMotionBlur
 	// Structure to hold one-time initialization
 	struct FConstructorStatics
 	{
-		FString NAME_Particles;
-		FString NAME_Constants;
+		FText NAME_Particles;
+		FText NAME_Constants;
 		FConstructorStatics()
-			: NAME_Particles(LOCTEXT( "Particles", "Particles" ).ToString())
-			, NAME_Constants(LOCTEXT( "Constants", "Constants" ).ToString())
+			: NAME_Particles(LOCTEXT( "Particles", "Particles" ))
+			, NAME_Constants(LOCTEXT( "Constants", "Constants" ))
 		{
 		}
 	};
@@ -8872,11 +8880,11 @@ UMaterialExpressionParticleDirection::UMaterialExpressionParticleDirection(const
 	// Structure to hold one-time initialization
 	struct FConstructorStatics
 	{
-		FString NAME_Particles;
-		FString NAME_Constants;
+		FText NAME_Particles;
+		FText NAME_Constants;
 		FConstructorStatics()
-			: NAME_Particles(LOCTEXT( "Particles", "Particles" ).ToString())
-			, NAME_Constants(LOCTEXT( "Constants", "Constants" ).ToString())
+			: NAME_Particles(LOCTEXT( "Particles", "Particles" ))
+			, NAME_Constants(LOCTEXT( "Constants", "Constants" ))
 		{
 		}
 	};
@@ -8906,11 +8914,11 @@ UMaterialExpressionParticleSpeed::UMaterialExpressionParticleSpeed(const FObject
 	// Structure to hold one-time initialization
 	struct FConstructorStatics
 	{
-		FString NAME_Particles;
-		FString NAME_Constants;
+		FText NAME_Particles;
+		FText NAME_Constants;
 		FConstructorStatics()
-			: NAME_Particles(LOCTEXT( "Particles", "Particles" ).ToString())
-			, NAME_Constants(LOCTEXT( "Constants", "Constants" ).ToString())
+			: NAME_Particles(LOCTEXT( "Particles", "Particles" ))
+			, NAME_Constants(LOCTEXT( "Constants", "Constants" ))
 		{
 		}
 	};
@@ -8940,11 +8948,11 @@ UMaterialExpressionParticleSize::UMaterialExpressionParticleSize(const FObjectIn
 	// Structure to hold one-time initialization
 	struct FConstructorStatics
 	{
-		FString NAME_Particles;
-		FString NAME_Constants;
+		FText NAME_Particles;
+		FText NAME_Constants;
 		FConstructorStatics()
-			: NAME_Particles(LOCTEXT( "Particles", "Particles" ).ToString())
-			, NAME_Constants(LOCTEXT( "Constants", "Constants" ).ToString())
+			: NAME_Particles(LOCTEXT( "Particles", "Particles" ))
+			, NAME_Constants(LOCTEXT( "Constants", "Constants" ))
 		{
 		}
 	};
@@ -8974,9 +8982,9 @@ UMaterialExpressionAtmosphericFogColor::UMaterialExpressionAtmosphericFogColor(c
 	// Structure to hold one-time initialization
 	struct FConstructorStatics
 	{
-		FString NAME_Atmosphere;
+		FText NAME_Atmosphere;
 		FConstructorStatics()
-			: NAME_Atmosphere(LOCTEXT( "Atmosphere", "Atmosphere" ).ToString())
+			: NAME_Atmosphere(LOCTEXT( "Atmosphere", "Atmosphere" ))
 		{
 		}
 	};
@@ -9012,9 +9020,9 @@ UMaterialExpressionSpeedTree::UMaterialExpressionSpeedTree(const FObjectInitiali
 	// Structure to hold one-time initialization
 	struct FConstructorStatics
 	{
-		FString NAME_SpeedTree;
+		FText NAME_SpeedTree;
 		FConstructorStatics()
-			: NAME_SpeedTree(LOCTEXT( "SpeedTree", "SpeedTree" ).ToString())
+			: NAME_SpeedTree(LOCTEXT( "SpeedTree", "SpeedTree" ))
 		{
 		}
 	};
@@ -9106,9 +9114,9 @@ UMaterialExpressionEyeAdaptation::UMaterialExpressionEyeAdaptation(const FObject
 	// Structure to hold one-time initialization
 	struct FConstructorStatics
 	{
-		FString NAME_Utility;
+		FText NAME_Utility;
 		FConstructorStatics()
-			: NAME_Utility(LOCTEXT( "Utility", "Utility" ).ToString())
+			: NAME_Utility(LOCTEXT( "Utility", "Utility" ))
 		{
 		}
 	};
diff --git a/Engine/Source/Runtime/Engine/Private/ObjectEditorUtils.cpp b/Engine/Source/Runtime/Engine/Private/ObjectEditorUtils.cpp
--- a/Engine/Source/Runtime/Engine/Private/ObjectEditorUtils.cpp
+++ b/Engine/Source/Runtime/Engine/Private/ObjectEditorUtils.cpp
@@ -9,26 +9,31 @@
 namespace FObjectEditorUtils
 {
 
-	FString GetCategory( const UProperty* InProperty )
+	FText GetCategoryText( const UProperty* InProperty )
 	{
 		static const FName NAME_Category(TEXT("Category"));
-		if (InProperty->HasMetaData(NAME_Category))
+		if (InProperty && InProperty->HasMetaData(NAME_Category))
 		{
-			return InProperty->GetMetaData(NAME_Category);
+			return InProperty->GetMetaDataText(NAME_Category, TEXT("UObjectCategory"), InProperty->GetFullGroupName(false));
 		}
 		else
 		{
-			return FString();
+			return FText::GetEmpty();
 		}
 	}
 
+	FString GetCategory( const UProperty* InProperty )
+	{
+		return GetCategoryText(InProperty).ToString();
+	}
+
 
 	FName GetCategoryFName( const UProperty* InProperty )
 	{
 		FName CategoryName(NAME_None);
 		if( InProperty )
 		{
-			CategoryName = *GetCategory( InProperty );
+			CategoryName = *GetCategoryText( InProperty ).ToString();
 		}
 		return CategoryName;
 	}
diff --git a/Engine/Source/Runtime/Engine/Private/SimpleConstructionScript.cpp b/Engine/Source/Runtime/Engine/Private/SimpleConstructionScript.cpp
--- a/Engine/Source/Runtime/Engine/Private/SimpleConstructionScript.cpp
+++ b/Engine/Source/Runtime/Engine/Private/SimpleConstructionScript.cpp
@@ -167,9 +167,9 @@ void USimpleConstructionScript::PostLoad()
 		USCS_Node* Node = Nodes[NodeIndex];
 
 		// Fix up any uninitialized category names
-		if(Node->CategoryName == NAME_None)
+		if(Node->CategoryName.IsEmpty())
 		{
-			Node->CategoryName = TEXT("Default");
+			Node->CategoryName = NSLOCTEXT("SCS", "Default", "Default");
 		}
 
 		// Fix up components that may have switched from scene to non-scene type and vice-versa
@@ -1016,7 +1016,7 @@ USCS_Node* USimpleConstructionScript::CreateNodeImpl(UActorComponent* NewCompone
 	NewNode->VariableName = ComponentVariableName;
 
 	// Note: This should match up with UEdGraphSchema_K2::VR_DefaultCategory
-	NewNode->CategoryName = TEXT("Default");
+	NewNode->CategoryName = NSLOCTEXT("SCS", "Default", "Default");
 	NewNode->VariableGuid = FGuid::NewGuid();
 	return NewNode;
 }
diff --git a/Engine/Source/Runtime/Engine/Public/ObjectEditorUtils.h b/Engine/Source/Runtime/Engine/Public/ObjectEditorUtils.h
--- a/Engine/Source/Runtime/Engine/Public/ObjectEditorUtils.h
+++ b/Engine/Source/Runtime/Engine/Public/ObjectEditorUtils.h
@@ -8,6 +8,14 @@ class UProperty;
 namespace FObjectEditorUtils
 {
 
+	/**
+	 * Gets the category this property belongs to.
+	 *
+	 * @param	InProperty	Property we want the category name of.
+	 * @return	Category name of the given property.
+	 */
+	ENGINE_API FText GetCategoryText( const class UProperty* InProperty );	
+
 	/**
 	 * Gets the category this property belongs to.
 	 *
diff --git a/Engine/Source/Runtime/Landscape/Private/LandscapeGrass.cpp b/Engine/Source/Runtime/Landscape/Private/LandscapeGrass.cpp
--- a/Engine/Source/Runtime/Landscape/Private/LandscapeGrass.cpp
+++ b/Engine/Source/Runtime/Landscape/Private/LandscapeGrass.cpp
@@ -650,10 +650,10 @@ UMaterialExpressionLandscapeGrassOutput::UMaterialExpressionLandscapeGrassOutput
 	// Structure to hold one-time initialization
 	struct FConstructorStatics
 	{
-		FString STRING_Landscape;
+		FText STRING_Landscape;
 		FName NAME_Grass;
 		FConstructorStatics()
-			: STRING_Landscape(LOCTEXT("Landscape", "Landscape").ToString())
+			: STRING_Landscape(LOCTEXT("Landscape", "Landscape"))
 			, NAME_Grass("Grass")
 		{
 		}
diff --git a/Engine/Source/Runtime/Landscape/Private/Materials/MaterialExpressionLandscapeLayerBlend.cpp b/Engine/Source/Runtime/Landscape/Private/Materials/MaterialExpressionLandscapeLayerBlend.cpp
--- a/Engine/Source/Runtime/Landscape/Private/Materials/MaterialExpressionLandscapeLayerBlend.cpp
+++ b/Engine/Source/Runtime/Landscape/Private/Materials/MaterialExpressionLandscapeLayerBlend.cpp
@@ -19,9 +19,9 @@ UMaterialExpressionLandscapeLayerBlend::UMaterialExpressionLandscapeLayerBlend(c
 	// Structure to hold one-time initialization
 	struct FConstructorStatics
 	{
-		FString NAME_Landscape;
+		FText NAME_Landscape;
 		FConstructorStatics()
-			: NAME_Landscape(LOCTEXT("Landscape", "Landscape").ToString())
+			: NAME_Landscape(LOCTEXT("Landscape", "Landscape"))
 		{
 		}
 	};
diff --git a/Engine/Source/Runtime/Landscape/Private/Materials/MaterialExpressionLandscapeLayerCoords.cpp b/Engine/Source/Runtime/Landscape/Private/Materials/MaterialExpressionLandscapeLayerCoords.cpp
--- a/Engine/Source/Runtime/Landscape/Private/Materials/MaterialExpressionLandscapeLayerCoords.cpp
+++ b/Engine/Source/Runtime/Landscape/Private/Materials/MaterialExpressionLandscapeLayerCoords.cpp
@@ -18,9 +18,9 @@ UMaterialExpressionLandscapeLayerCoords::UMaterialExpressionLandscapeLayerCoords
 	// Structure to hold one-time initialization
 	struct FConstructorStatics
 	{
-		FString NAME_Landscape;
+		FText NAME_Landscape;
 		FConstructorStatics()
-			: NAME_Landscape(LOCTEXT("Landscape", "Landscape").ToString())
+			: NAME_Landscape(LOCTEXT("Landscape", "Landscape"))
 		{
 		}
 	};
diff --git a/Engine/Source/Runtime/Landscape/Private/Materials/MaterialExpressionLandscapeLayerSample.cpp b/Engine/Source/Runtime/Landscape/Private/Materials/MaterialExpressionLandscapeLayerSample.cpp
--- a/Engine/Source/Runtime/Landscape/Private/Materials/MaterialExpressionLandscapeLayerSample.cpp
+++ b/Engine/Source/Runtime/Landscape/Private/Materials/MaterialExpressionLandscapeLayerSample.cpp
@@ -19,9 +19,9 @@ UMaterialExpressionLandscapeLayerSample::UMaterialExpressionLandscapeLayerSample
 	// Structure to hold one-time initialization
 	struct FConstructorStatics
 	{
-		FString NAME_Landscape;
+		FText NAME_Landscape;
 		FConstructorStatics()
-			: NAME_Landscape(LOCTEXT("Landscape", "Landscape").ToString())
+			: NAME_Landscape(LOCTEXT("Landscape", "Landscape"))
 		{
 		}
 	};
diff --git a/Engine/Source/Runtime/Landscape/Private/Materials/MaterialExpressionLandscapeLayerSwitch.cpp b/Engine/Source/Runtime/Landscape/Private/Materials/MaterialExpressionLandscapeLayerSwitch.cpp
--- a/Engine/Source/Runtime/Landscape/Private/Materials/MaterialExpressionLandscapeLayerSwitch.cpp
+++ b/Engine/Source/Runtime/Landscape/Private/Materials/MaterialExpressionLandscapeLayerSwitch.cpp
@@ -19,9 +19,9 @@ UMaterialExpressionLandscapeLayerSwitch::UMaterialExpressionLandscapeLayerSwitch
 	// Structure to hold one-time initialization
 	struct FConstructorStatics
 	{
-		FString NAME_Landscape;
+		FText NAME_Landscape;
 		FConstructorStatics()
-			: NAME_Landscape(LOCTEXT("Landscape", "Landscape").ToString())
+			: NAME_Landscape(LOCTEXT("Landscape", "Landscape"))
 		{
 		}
 	};
diff --git a/Engine/Source/Runtime/Landscape/Private/Materials/MaterialExpressionLandscapeLayerWeight.cpp b/Engine/Source/Runtime/Landscape/Private/Materials/MaterialExpressionLandscapeLayerWeight.cpp
--- a/Engine/Source/Runtime/Landscape/Private/Materials/MaterialExpressionLandscapeLayerWeight.cpp
+++ b/Engine/Source/Runtime/Landscape/Private/Materials/MaterialExpressionLandscapeLayerWeight.cpp
@@ -19,9 +19,9 @@ UMaterialExpressionLandscapeLayerWeight::UMaterialExpressionLandscapeLayerWeight
 	// Structure to hold one-time initialization
 	struct FConstructorStatics
 	{
-		FString NAME_Landscape;
+		FText NAME_Landscape;
 		FConstructorStatics()
-			: NAME_Landscape(LOCTEXT("Landscape", "Landscape").ToString())
+			: NAME_Landscape(LOCTEXT("Landscape", "Landscape"))
 		{
 		}
 	};
diff --git a/Engine/Source/Runtime/Landscape/Private/Materials/MaterialExpressionLandscapeVisibilityMask.cpp b/Engine/Source/Runtime/Landscape/Private/Materials/MaterialExpressionLandscapeVisibilityMask.cpp
--- a/Engine/Source/Runtime/Landscape/Private/Materials/MaterialExpressionLandscapeVisibilityMask.cpp
+++ b/Engine/Source/Runtime/Landscape/Private/Materials/MaterialExpressionLandscapeVisibilityMask.cpp
@@ -21,9 +21,9 @@ UMaterialExpressionLandscapeVisibilityMask::UMaterialExpressionLandscapeVisibili
 	// Structure to hold one-time initialization
 	struct FConstructorStatics
 	{
-		FString NAME_Landscape;
+		FText NAME_Landscape;
 		FConstructorStatics()
-			: NAME_Landscape(LOCTEXT("Landscape", "Landscape").ToString())
+			: NAME_Landscape(LOCTEXT("Landscape", "Landscape"))
 		{
 		}
 	};
