#ue4 - Perf: cache GetWorld() calls in MoveComponentImpl(). Shrink size of inline TArray in UpdateOverlaps (FOverlapInfo is large).

diff --git a/Engine/Source/Runtime/Engine/Private/Components/PrimitiveComponent.cpp b/Engine/Source/Runtime/Engine/Private/Components/PrimitiveComponent.cpp
--- a/Engine/Source/Runtime/Engine/Private/Components/PrimitiveComponent.cpp
+++ b/Engine/Source/Runtime/Engine/Private/Components/PrimitiveComponent.cpp
@@ -1414,11 +1414,12 @@ bool UPrimitiveComponent::MoveComponentImpl( const FVector& Delta, const FQuat&
 			MoveTimer.bDidLineCheck = true;
 #endif 
 			static const FName Name_MoveComponent(TEXT("MoveComponent"));
+			UWorld* const World = GetWorld();
 
 			FComponentQueryParams Params(Name_MoveComponent, Actor);
 			FCollisionResponseParams ResponseParam;
 			InitSweepCollisionParams(Params, ResponseParam);
-			bool const bHadBlockingHit = GetWorld()->ComponentSweepMulti(Hits, this, TraceStart, TraceEnd, GetComponentQuat(), Params);
+			bool const bHadBlockingHit = World->ComponentSweepMulti(Hits, this, TraceStart, TraceEnd, GetComponentQuat(), Params);
 
 			if (Hits.Num() > 0)
 			{
@@ -1440,9 +1441,9 @@ bool UPrimitiveComponent::MoveComponentImpl( const FVector& Delta, const FQuat&
 				{
 					const FHitResult& TestHit = Hits[HitIdx];
 
-					if ( !ShouldIgnoreHitResult(GetWorld(), TestHit, Delta, Actor, MoveFlags) )
+					if (TestHit.bBlockingHit)
 					{
-						if (TestHit.bBlockingHit)
+						if (!ShouldIgnoreHitResult(World, TestHit, Delta, Actor, MoveFlags))
 						{
 							if (TestHit.Time == 0.f)
 							{
@@ -1460,30 +1461,30 @@ bool UPrimitiveComponent::MoveComponentImpl( const FVector& Delta, const FQuat&
 								// This should be the only non-overlapping blocking hit, and last in the results.
 								BlockingHitIndex = HitIdx;
 								break;
-							}							
+							}
 						}
-						else if (bGenerateOverlapEvents)
+					}
+					else if (bGenerateOverlapEvents)
+					{
+						UPrimitiveComponent* OverlapComponent = TestHit.Component.Get();
+						if (OverlapComponent && OverlapComponent->bGenerateOverlapEvents)
 						{
-							UPrimitiveComponent * OverlapComponent = TestHit.Component.Get();
-							if (OverlapComponent && OverlapComponent->bGenerateOverlapEvents)
+							if (!ShouldIgnoreOverlapResult(World, Actor, *this, TestHit.GetActor(), *OverlapComponent))
 							{
-								if (!ShouldIgnoreOverlapResult(GetWorld(), Actor, *this, TestHit.GetActor(), *OverlapComponent))
+								// don't process touch events after initial blocking hits
+								if (BlockingHitIndex >= 0 && TestHit.Time > Hits[BlockingHitIndex].Time)
 								{
-									// don't process touch events after initial blocking hits
-									if (BlockingHitIndex >= 0 && TestHit.Time > Hits[BlockingHitIndex].Time)
-									{
-										break;
-									}
-
-									if (FirstNonInitialOverlapIdx == INDEX_NONE && TestHit.Time > 0.f)
-									{
-										// We are about to add the first non-initial overlap.
-										FirstNonInitialOverlapIdx = PendingOverlaps.Num();
-									}
-
-									// cache touches
-									PendingOverlaps.AddUnique(FOverlapInfo(TestHit));
+									break;
 								}
+
+								if (FirstNonInitialOverlapIdx == INDEX_NONE && TestHit.Time > 0.f)
+								{
+									// We are about to add the first non-initial overlap.
+									FirstNonInitialOverlapIdx = PendingOverlaps.Num();
+								}
+
+								// cache touches
+								PendingOverlaps.AddUnique(FOverlapInfo(TestHit));
 							}
 						}
 					}
@@ -2148,7 +2149,7 @@ void UPrimitiveComponent::UpdateOverlaps(TArray<FOverlapInfo> const* PendingOver
 
 			// now generate full list of new touches, so we can compare to existing list and
 			// determine what changed
-			typedef TArray<FOverlapInfo, TInlineAllocator<4>> TInlineOverlapInfoArray;
+			typedef TArray<FOverlapInfo, TInlineAllocator<3>> TInlineOverlapInfoArray;
 			TInlineOverlapInfoArray NewOverlappingComponents;
 
 			// If pending kill, we should not generate any new overlaps
