#ue4 - Fix shadowed variable.

diff --git a/Engine/Source/Runtime/Engine/Private/Components/PrimitiveComponent.cpp b/Engine/Source/Runtime/Engine/Private/Components/PrimitiveComponent.cpp
--- a/Engine/Source/Runtime/Engine/Private/Components/PrimitiveComponent.cpp
+++ b/Engine/Source/Runtime/Engine/Private/Components/PrimitiveComponent.cpp
@@ -1414,12 +1414,12 @@ bool UPrimitiveComponent::MoveComponentImpl( const FVector& Delta, const FQuat&
 			MoveTimer.bDidLineCheck = true;
 #endif 
 			static const FName Name_MoveComponent(TEXT("MoveComponent"));
-			UWorld* const World = GetWorld();
+			UWorld* const MyWorld = GetWorld();
 
 			FComponentQueryParams Params(Name_MoveComponent, Actor);
 			FCollisionResponseParams ResponseParam;
 			InitSweepCollisionParams(Params, ResponseParam);
-			bool const bHadBlockingHit = World->ComponentSweepMulti(Hits, this, TraceStart, TraceEnd, GetComponentQuat(), Params);
+			bool const bHadBlockingHit = MyWorld->ComponentSweepMulti(Hits, this, TraceStart, TraceEnd, GetComponentQuat(), Params);
 
 			if (Hits.Num() > 0)
 			{
@@ -1443,7 +1443,7 @@ bool UPrimitiveComponent::MoveComponentImpl( const FVector& Delta, const FQuat&
 
 					if (TestHit.bBlockingHit)
 					{
-						if (!ShouldIgnoreHitResult(World, TestHit, Delta, Actor, MoveFlags))
+						if (!ShouldIgnoreHitResult(MyWorld, TestHit, Delta, Actor, MoveFlags))
 						{
 							if (TestHit.Time == 0.f)
 							{
@@ -1469,7 +1469,7 @@ bool UPrimitiveComponent::MoveComponentImpl( const FVector& Delta, const FQuat&
 						UPrimitiveComponent* OverlapComponent = TestHit.Component.Get();
 						if (OverlapComponent && OverlapComponent->bGenerateOverlapEvents)
 						{
-							if (!ShouldIgnoreOverlapResult(World, Actor, *this, TestHit.GetActor(), *OverlapComponent))
+							if (!ShouldIgnoreOverlapResult(MyWorld, Actor, *this, TestHit.GetActor(), *OverlapComponent))
 							{
 								// don't process touch events after initial blocking hits
 								if (BlockingHitIndex >= 0 && TestHit.Time > Hits[BlockingHitIndex].Time)
