Fix up a number of issues with SoundNodeWavePlayer using TAssetPtr
* Avoid unnecessary weak pointer evaluations by caching SoundWave pointer after loading from TAssetPtr
* Encapsulated SoundWave and SoundWaveAssetPtr to prevent getting out of sync
* Prevent sound waves from being garbage collected

diff --git a/Engine/Config/BaseEngine.ini b/Engine/Config/BaseEngine.ini
--- a/Engine/Config/BaseEngine.ini
+++ b/Engine/Config/BaseEngine.ini
@@ -1223,6 +1223,9 @@ bUseStreamingPause=false
 +K2FieldRedirects=(OldFieldName="KismetSystemLibrary.EXPERIMENTAL_HideAdBanner",NewFieldName="KismetSystemLibrary.HideAdBanner")
 +K2FieldRedirects=(OldFieldName="KismetSystemLibrary.EXPERIMENTAL_CloseAdBanner",NewFieldName="KismetSystemLibrary.ForceCloseAdBanner")
 
+;SoundNodeWavePlayer using asset pointers - DEPRECATED(4.9)
++TaggedPropertyRedirects=(ClassName="SoundNodeWavePlayer",OldPropertyName="SoundWave",NewPropertyName="SoundWaveAssetPtr")
+
 [CoreUObject.Metadata]
 ; Note: UnrealHeaderTool should be rerun after making changes to MetadataRedirects to catch any keys specified in class headers
 +MetadataRedirects=(OldKey="K2Protected", NewKey="BlueprintProtected")
diff --git a/Engine/Source/Developer/MovieSceneTools/Private/AudioTrackEditor.cpp b/Engine/Source/Developer/MovieSceneTools/Private/AudioTrackEditor.cpp
--- a/Engine/Source/Developer/MovieSceneTools/Private/AudioTrackEditor.cpp
+++ b/Engine/Source/Developer/MovieSceneTools/Private/AudioTrackEditor.cpp
@@ -62,7 +62,7 @@ USoundWave* DeriveSoundWave(UMovieSceneAudioSection* AudioSection)
 		{
 			if (AllNodes[i]->IsA<USoundNodeWavePlayer>())
 			{
-				SoundWave = Cast<USoundNodeWavePlayer>(AllNodes[i])->SoundWave.Get();
+				SoundWave = Cast<USoundNodeWavePlayer>(AllNodes[i])->GetSoundWave();
 			}
 		}
 	}
diff --git a/Engine/Source/Editor/SoundCueEditor/Private/SoundCueEditor.cpp b/Engine/Source/Editor/SoundCueEditor/Private/SoundCueEditor.cpp
--- a/Engine/Source/Editor/SoundCueEditor/Private/SoundCueEditor.cpp
+++ b/Engine/Source/Editor/SoundCueEditor/Private/SoundCueEditor.cpp
@@ -388,9 +388,9 @@ void FSoundCueEditor::SyncInBrowser()
 		if (SelectedNode)
 		{
 			USoundNodeWavePlayer* SelectedWave = Cast<USoundNodeWavePlayer>(SelectedNode->SoundNode);
-			if (SelectedWave && SelectedWave->SoundWave)
+			if (SelectedWave && SelectedWave->GetSoundWave())
 			{
-				ObjectsToSync.AddUnique(SelectedWave->SoundWave.Get());
+				ObjectsToSync.AddUnique(SelectedWave->GetSoundWave());
 			}
 			//ObjectsToSync.AddUnique(SelectedNode->SoundNode);
 		}
@@ -414,7 +414,7 @@ bool FSoundCueEditor::CanSyncInBrowser() const
 		{
 			USoundNodeWavePlayer* WavePlayer = Cast<USoundNodeWavePlayer>(SelectedNode->SoundNode);
 
-			if (WavePlayer && WavePlayer->SoundWave)
+			if (WavePlayer && WavePlayer->GetSoundWave())
 			{
 				return true;
 			}
diff --git a/Engine/Source/Editor/SoundCueEditor/Private/SoundCueEditorUtilities.cpp b/Engine/Source/Editor/SoundCueEditor/Private/SoundCueEditorUtilities.cpp
--- a/Engine/Source/Editor/SoundCueEditor/Private/SoundCueEditorUtilities.cpp
+++ b/Engine/Source/Editor/SoundCueEditor/Private/SoundCueEditorUtilities.cpp
@@ -42,7 +42,7 @@ void FSoundCueEditorUtilities::CreateWaveContainers(TArray<USoundWave*>& Selecte
 		{
 			USoundNodeWavePlayer* WavePlayer = SoundCue->ConstructSoundNode<USoundNodeWavePlayer>();
 
-			WavePlayer->SoundWave = NewWave;
+			WavePlayer->SetSoundWave(NewWave);
 
 			WavePlayer->GraphNode->NodePosX = Location.X - WavePlayer->GetGraphNode()->EstimateNodeWidth();
 			WavePlayer->GraphNode->NodePosY = Location.Y + (NodeSpacing * WaveIndex);
diff --git a/Engine/Source/Editor/UnrealEd/Private/Factories/EditorFactories.cpp b/Engine/Source/Editor/UnrealEd/Private/Factories/EditorFactories.cpp
--- a/Engine/Source/Editor/UnrealEd/Private/Factories/EditorFactories.cpp
+++ b/Engine/Source/Editor/UnrealEd/Private/Factories/EditorFactories.cpp
@@ -1572,7 +1572,7 @@ void CreateSoundCue( USoundWave* Sound, UObject* InParent, EObjectFlags Flags, b
 	// Apply the initial volume.
 	SoundCue->VolumeMultiplier = CueVolume;
 
-	WavePlayer->SoundWave = Sound;
+	WavePlayer->SetSoundWave(Sound);
 	SoundCue->FirstNode = WavePlayer;
 	SoundCue->LinkGraphNodesFromSoundNodes();
 
@@ -2397,7 +2397,7 @@ UObject* USoundCueFactoryNew::FactoryCreateNew( UClass* Class, UObject* InParent
 		USoundNodeWavePlayer* WavePlayer = SoundCue->ConstructSoundNode<USoundNodeWavePlayer>();
 		SoundCue->FirstNode = WavePlayer;
 		SoundCue->LinkGraphNodesFromSoundNodes();
-		WavePlayer->SoundWave = InitialSoundWave;
+		WavePlayer->SetSoundWave(InitialSoundWave);
 		WavePlayer->GraphNode->NodePosX = -250;
 		WavePlayer->GraphNode->NodePosY = -35;
 	}
diff --git a/Engine/Source/Runtime/Engine/Classes/Sound/SoundNodeWavePlayer.h b/Engine/Source/Runtime/Engine/Classes/Sound/SoundNodeWavePlayer.h
--- a/Engine/Source/Runtime/Engine/Classes/Sound/SoundNodeWavePlayer.h
+++ b/Engine/Source/Runtime/Engine/Classes/Sound/SoundNodeWavePlayer.h
@@ -13,13 +13,30 @@ class USoundNodeWavePlayer : public USoundNode
 {
 	GENERATED_UCLASS_BODY()
 
-	UPROPERTY(EditAnywhere, Category=WavePlayer)
-	TAssetPtr<USoundWave> SoundWave;
+private:
+	UPROPERTY(EditAnywhere, Category=WavePlayer, meta=(DisplayName="Sound Wave"))
+	TAssetPtr<USoundWave> SoundWaveAssetPtr;
+
+	USoundWave* SoundWave;
+
+public:	
 
 	UPROPERTY(EditAnywhere, Category=WavePlayer)
 	uint32 bLooping:1;
 
-public:	
+	ENGINE_API USoundWave* GetSoundWave() const { return SoundWave; }
+	ENGINE_API void SetSoundWave(USoundWave* SoundWave);
+	void LoadSoundWave();
+
+
+	// Begin UObject Interface
+	static void AddReferencedObjects(UObject* InThis, FReferenceCollector& Collector);
+
+#if WITH_EDITOR
+	virtual void PostEditChangeProperty(FPropertyChangedEvent& PropertyChangedEvent) override;
+#endif
+	// End UObject Interface
+
 	// Begin USoundNode Interface
 	virtual int32 GetMaxChildNodes() const override;
 	virtual float GetDuration() override;
diff --git a/Engine/Source/Runtime/Engine/Private/Audio.cpp b/Engine/Source/Runtime/Engine/Private/Audio.cpp
--- a/Engine/Source/Runtime/Engine/Private/Audio.cpp
+++ b/Engine/Source/Runtime/Engine/Private/Audio.cpp
@@ -99,7 +99,7 @@ FName FSoundBuffer::GetSoundClassName()
 			// look through them to see if this cue uses a wave this buffer is bound to, via ResourceID
 			for (int32 WaveIndex = 0; WaveIndex < WavePlayers.Num(); ++WaveIndex)
 			{
-				USoundWave* WaveNode = WavePlayers[WaveIndex]->SoundWave.Get();
+				USoundWave* WaveNode = WavePlayers[WaveIndex]->GetSoundWave();
 				if (WaveNode != NULL)
 				{
 					if (WaveNode->ResourceID == ResourceID)
diff --git a/Engine/Source/Runtime/Engine/Private/AudioDevice.cpp b/Engine/Source/Runtime/Engine/Private/AudioDevice.cpp
--- a/Engine/Source/Runtime/Engine/Private/AudioDevice.cpp
+++ b/Engine/Source/Runtime/Engine/Private/AudioDevice.cpp
@@ -315,7 +315,7 @@ void FAudioDevice::GetSoundClassInfo( TMap<FName, FAudioClassInfo>& AudioClassIn
 		for( int32 WaveIndex = 0; WaveIndex < WavePlayers.Num(); ++WaveIndex )
 		{
 			// Presume one class per sound node wave
-			USoundWave *SoundWave = WavePlayers[ WaveIndex ]->SoundWave.Get();
+			USoundWave *SoundWave = WavePlayers[ WaveIndex ]->GetSoundWave();
 			if (SoundWave && SoundCue->GetSoundClass())
 			{
 				SoundWaveClasses.Add( SoundWave, SoundCue->GetSoundClass()->GetFName() );
@@ -594,7 +594,7 @@ bool FAudioDevice::HandlePlaySoundCueCommand( const TCHAR* Cmd, FOutputDevice& A
 			Cue->RecursiveFindNode<USoundNodeWavePlayer>( Cue->FirstNode, WavePlayers );
 			for( int32 i = 0; i < WavePlayers.Num(); ++i )
 			{
-				USoundWave* SoundWave = WavePlayers[ i ]->SoundWave.Get();
+				USoundWave* SoundWave = WavePlayers[ i ]->GetSoundWave();
 				if (SoundWave)
 				{
 					SoundWave->LogSubtitle( Ar );
diff --git a/Engine/Source/Runtime/Engine/Private/SoundCue.cpp b/Engine/Source/Runtime/Engine/Private/SoundCue.cpp
--- a/Engine/Source/Runtime/Engine/Private/SoundCue.cpp
+++ b/Engine/Source/Runtime/Engine/Private/SoundCue.cpp
@@ -81,7 +81,7 @@ void USoundCue::PostLoad()
 		{
 			if (USoundNodeWavePlayer* WavePlayerNode = Cast<USoundNodeWavePlayer>(SoundNode))
 			{
-				WavePlayerNode->SoundWave.LoadSynchronous();
+				WavePlayerNode->LoadSoundWave();
 			}
 		}
 	}
@@ -97,7 +97,7 @@ void USoundCue::PostLoad()
 			{
 				if (USoundNodeWavePlayer* WavePlayerNode = Cast<USoundNodeWavePlayer>(SoundNode))
 				{
-					WavePlayerNode->SoundWave.LoadSynchronous();
+					WavePlayerNode->LoadSoundWave();
 				}
 				/* else if USoundNodeQuality
 				{
@@ -241,7 +241,7 @@ SIZE_T USoundCue::GetResourceSize(EResourceSizeMode::Type Mode)
 
 		for( int32 WaveIndex = 0; WaveIndex < WavePlayers.Num(); ++WaveIndex )
 		{
-			USoundWave* SoundWave = WavePlayers[WaveIndex]->SoundWave.Get();
+			USoundWave* SoundWave = WavePlayers[WaveIndex]->GetSoundWave();
 			if (SoundWave)
 			{
 				ResourceSize += SoundWave->GetResourceSize(Mode);
@@ -260,7 +260,7 @@ int32 USoundCue::GetResourceSizeForFormat(FName Format)
 	int32 ResourceSize = 0;
 	for (int32 WaveIndex = 0; WaveIndex < WavePlayers.Num(); ++WaveIndex)
 	{
-		USoundWave* SoundWave = WavePlayers[WaveIndex]->SoundWave.Get();
+		USoundWave* SoundWave = WavePlayers[WaveIndex]->GetSoundWave();
 		if (SoundWave)
 		{
 			ResourceSize += SoundWave->GetResourceSizeForFormat(Format);
diff --git a/Engine/Source/Runtime/Engine/Private/SoundNodeMature.cpp b/Engine/Source/Runtime/Engine/Private/SoundNodeMature.cpp
--- a/Engine/Source/Runtime/Engine/Private/SoundNodeMature.cpp
+++ b/Engine/Source/Runtime/Engine/Private/SoundNodeMature.cpp
@@ -44,9 +44,9 @@ EMaturityChildType GetMaturityTypeForChild( USoundNode* Node )
 	{
 		// Try to see if the given node is a sound wave node.
 		USoundNodeWavePlayer* WavePlayer = Cast<USoundNodeWavePlayer>( Node );
-		if( WavePlayer && WavePlayer->SoundWave )
+		if( WavePlayer && WavePlayer->GetSoundWave() )
 		{
-			if( WavePlayer->SoundWave->bMature )
+			if( WavePlayer->GetSoundWave()->bMature )
 			{
 				Type = ChildType_Mature;
 			}
@@ -154,10 +154,12 @@ void USoundNodeMature::PostLoad()
 
 	if( !GIsEditor && GEngine && !HasAnyFlags(RF_RootSet) && ChildNodes.Num() >= 2 )
 	{
+		// Make sure the SoundCue has gotten all the SoundWavePlayers in to memory
+		GetOuter()->ConditionalPostLoad();
 		for( int32 i = ChildNodes.Num() - 1; i >= 0 ; --i )
 		{
 			USoundNodeWavePlayer *WavePlayer = Cast<USoundNodeWavePlayer>(ChildNodes[i]);
-			if (WavePlayer && WavePlayer->SoundWave && WavePlayer->SoundWave->bMature != GEngine->bAllowMatureLanguage)
+			if (WavePlayer && WavePlayer->GetSoundWave() && WavePlayer->GetSoundWave()->bMature != GEngine->bAllowMatureLanguage)
 			{
 				ChildNodes.RemoveAt(i);
 			}
diff --git a/Engine/Source/Runtime/Engine/Private/SoundNodeWavePlayer.cpp b/Engine/Source/Runtime/Engine/Private/SoundNodeWavePlayer.cpp
--- a/Engine/Source/Runtime/Engine/Private/SoundNodeWavePlayer.cpp
+++ b/Engine/Source/Runtime/Engine/Private/SoundNodeWavePlayer.cpp
@@ -12,6 +12,35 @@ USoundNodeWavePlayer::USoundNodeWavePlayer(const FObjectInitializer& ObjectIniti
 {
 }
 
+void USoundNodeWavePlayer::AddReferencedObjects(UObject* InThis, FReferenceCollector& Collector)
+{
+	USoundNodeWavePlayer* This = CastChecked<USoundNodeWavePlayer>(InThis);
+	Super::AddReferencedObjects(This, Collector);
+
+	Collector.AddReferencedObject(This->SoundWave);
+}
+
+void USoundNodeWavePlayer::LoadSoundWave()
+{
+	SoundWave = SoundWaveAssetPtr.LoadSynchronous();
+}
+
+void USoundNodeWavePlayer::SetSoundWave(USoundWave* InSoundWave)
+{
+	SoundWave = InSoundWave;
+	SoundWaveAssetPtr = InSoundWave;
+}
+
+#if WITH_EDITOR
+void USoundNodeWavePlayer::PostEditChangeProperty(FPropertyChangedEvent& PropertyChangedEvent)
+{
+	if (PropertyChangedEvent.Property && PropertyChangedEvent.Property->GetFName() == GET_MEMBER_NAME_CHECKED(USoundNodeWavePlayer, SoundWaveAssetPtr))
+	{
+		LoadSoundWave();
+	}
+}
+#endif
+
 void USoundNodeWavePlayer::ParseNodes( FAudioDevice* AudioDevice, const UPTRINT NodeWaveInstanceHash, FActiveSound& ActiveSound, const FSoundParseParameters& ParseParams, TArray<FWaveInstance*>& WaveInstances )
 {
 	if (SoundWave)
