#ue4 - Perf: Controller::SetControlRotation() avoids work if nothing changes.

diff --git a/Engine/Source/Runtime/AIModule/Private/AIController.cpp b/Engine/Source/Runtime/AIModule/Private/AIController.cpp
--- a/Engine/Source/Runtime/AIModule/Private/AIController.cpp
+++ b/Engine/Source/Runtime/AIModule/Private/AIController.cpp
@@ -419,7 +419,6 @@ void AAIController::UpdateControlRotation(float DeltaTime, bool bUpdatePawn)
 		{
 			NewControlRotation.Pitch = 0.f;
 		}
-		NewControlRotation.Yaw = FRotator::ClampAxis(NewControlRotation.Yaw);
 
 		if (GetControlRotation().Equals(NewControlRotation, 1e-3f) == false)
 		{
diff --git a/Engine/Source/Runtime/Engine/Private/Controller.cpp b/Engine/Source/Runtime/Engine/Private/Controller.cpp
--- a/Engine/Source/Runtime/Engine/Private/Controller.cpp
+++ b/Engine/Source/Runtime/Engine/Private/Controller.cpp
@@ -103,11 +103,18 @@ FRotator AController::GetControlRotation() const
 
 void AController::SetControlRotation(const FRotator& NewRotation)
 {
-	ControlRotation = NewRotation;
+	if (!ControlRotation.Equals(NewRotation, 1e-3f))
+	{
+		ControlRotation = NewRotation;
 
-	if (RootComponent && RootComponent->bAbsoluteRotation)
+		if (RootComponent && RootComponent->bAbsoluteRotation)
+		{
+			RootComponent->SetWorldRotation(GetControlRotation());
+		}
+	}
+	else
 	{
-		RootComponent->SetWorldRotation(GetControlRotation());
+		//UE_LOG(LogPlayerController, Log, TEXT("Skipping SetControlRotation for %s (Pawn %s)"), *GetNameSafe(this), *GetNameSafe(GetPawn()));
 	}
 }
 
@@ -211,6 +218,13 @@ void AController::PostInitializeComponents()
 	if ( !IsPendingKill() )
 	{
 		GetWorld()->AddController( this );
+
+		// Since we avoid updating rotation in SetControlRotation() if it hasn't changed,
+		// we should make sure that the initial RootComponent rotation matches it if ControlRotation was set directly.
+		if (RootComponent && RootComponent->bAbsoluteRotation)
+		{
+			RootComponent->SetWorldRotation(GetControlRotation());
+		}
 	}
 }
 
