Restore view target when scrubbing

diff --git a/Engine/Source/Runtime/Engine/Private/DemoNetDriver.cpp b/Engine/Source/Runtime/Engine/Private/DemoNetDriver.cpp
--- a/Engine/Source/Runtime/Engine/Private/DemoNetDriver.cpp
+++ b/Engine/Source/Runtime/Engine/Private/DemoNetDriver.cpp
@@ -1469,6 +1469,19 @@ void UDemoNetDriver::LoadCheckpoint()
 		AddNonQueuedActorForScrubbing( SpectatorController );
 	}
 
+	// Remember the spectator controller's view target so we can restore it
+	FNetworkGUID ViewTargetGUID;
+
+	if ( SpectatorController && SpectatorController->GetViewTarget() )
+	{
+		ViewTargetGUID = GuidCache->NetGUIDLookup.FindRef( SpectatorController->GetViewTarget() );
+
+		if ( ViewTargetGUID.IsValid() )
+		{
+			AddNonQueuedActorForScrubbing( SpectatorController->GetViewTarget() );
+		}
+	}
+
 	PauseChannels( false );
 
 #if 1
@@ -1639,6 +1652,16 @@ void UDemoNetDriver::LoadCheckpoint()
 	GotoCheckpointSkipExtraTimeInMS = -1;
 	GotoCheckpointArchive			= NULL;
 	bDemoPlaybackDone				= false;
+
+	if ( SpectatorController && ViewTargetGUID.IsValid() )
+	{
+		AActor* ViewTarget = Cast< AActor >( GuidCache->GetObjectFromNetGUID( ViewTargetGUID, false ) );
+
+		if ( ViewTarget )
+		{
+			SpectatorController->SetViewTarget( ViewTarget );
+		}
+	}
 }
 
 bool UDemoNetDriver::ShouldQueueBunchesForActorGUID(FNetworkGUID InGUID) const
