#ue4 - Mask out complicated ensure code in shipping, I don't trust the compiler.

diff --git a/Engine/Source/Runtime/Engine/Private/Components/SceneComponent.cpp b/Engine/Source/Runtime/Engine/Private/Components/SceneComponent.cpp
--- a/Engine/Source/Runtime/Engine/Private/Components/SceneComponent.cpp
+++ b/Engine/Source/Runtime/Engine/Private/Components/SceneComponent.cpp
@@ -341,7 +341,10 @@ void USceneComponent::UpdateComponentToWorldWithParent(USceneComponent* Parent,
 	
 	// Minimize accumulation of errors after many composed transforms.
 	NewTransform.NormalizeRotation();
+
+#if !(UE_BUILD_SHIPPING || UE_BUILD_TEST)
 	ensureOnce(NewTransform.IsValid());
+#endif
 
 	// If transform has changed..
 	if (!ComponentToWorld.Equals(NewTransform, SMALL_NUMBER))
