Turn off physics serialization while waiting for support from PhysX

UE-15267,15144,14376

diff --git a/Engine/Source/Runtime/Engine/Private/PhysicsEngine/PhysicsSerializer.cpp b/Engine/Source/Runtime/Engine/Private/PhysicsEngine/PhysicsSerializer.cpp
--- a/Engine/Source/Runtime/Engine/Private/PhysicsEngine/PhysicsSerializer.cpp
+++ b/Engine/Source/Runtime/Engine/Private/PhysicsEngine/PhysicsSerializer.cpp
@@ -16,6 +16,11 @@ UPhysicsSerializer::UPhysicsSerializer(const FObjectInitializer& ObjectInitializ
 
 FByteBulkData* UPhysicsSerializer::GetBinaryData(FName Format, const TArray<FBodyInstance*>& Bodies, const TArray<class UBodySetup*>& BodySetups, const TArray<class UPhysicalMaterial*>& PhysicalMaterials)
 {
+	if (!FParse::Param(FCommandLine::Get(), TEXT("PhysxSerialization")))
+	{
+		return nullptr;
+	}
+
 #if PLATFORM_MAC
 	return nullptr;	//This is not supported right now
 #endif
@@ -95,6 +100,11 @@ void UPhysicsSerializer::SerializePhysics(const TArray<FBodyInstance*>& Bodies,
 
 void UPhysicsSerializer::CreatePhysicsData(const TArray<UBodySetup*>& BodySetups, const TArray<UPhysicalMaterial*>& PhysicalMaterials)
 {
+	if (!FParse::Param(FCommandLine::Get(), TEXT("PhysxSerialization")))
+	{
+		return;
+	}
+
 #if PLATFORM_MAC
 	return;	//This is not supported right now
 #endif
