#ue4 - Perf: symmetric changes in rotation without translation within a FScopedMovementUpdate can now avoid UpdateOverlaps() when there are no current and pending overlaps.

diff --git a/Engine/Source/Runtime/Engine/Classes/Components/SceneComponent.h b/Engine/Source/Runtime/Engine/Classes/Components/SceneComponent.h
--- a/Engine/Source/Runtime/Engine/Classes/Components/SceneComponent.h
+++ b/Engine/Source/Runtime/Engine/Classes/Components/SceneComponent.h
@@ -970,9 +970,12 @@ public:
 	/** Returns true if the Component's transform has changed since the start of the scoped update. */
 	bool IsTransformDirty() const;
 
-	/** Returns true if there are pending overlaps queued up in this scope. */
+	/** Returns true if there are pending overlaps queued in this scope. */
 	bool HasPendingOverlaps() const;
 
+	/** Returns true if there are pending overlaps queued in this scope or any outer scope. */
+	bool HasPendingOverlapsInAncestorOrSelf() const;
+
 	/** Returns the pending overlaps within this scope. */
 	const TArray<struct FOverlapInfo>& GetPendingOverlaps() const;
 	
@@ -1003,6 +1006,7 @@ private:
 private:
 
 	class USceneComponent* Owner;
+	FScopedMovementUpdate* OuterScope;
 	uint32 bDeferUpdates:1;
 	uint32 bHasValidOverlapsAtEnd:1;
 	FTransform InitialTransform;
@@ -1030,6 +1034,22 @@ FORCEINLINE bool FScopedMovementUpdate::HasPendingOverlaps() const
 	return PendingOverlaps.Num() > 0;
 }
 
+FORCEINLINE_DEBUGGABLE bool FScopedMovementUpdate::HasPendingOverlapsInAncestorOrSelf() const
+{
+	const FScopedMovementUpdate* Scope = this;
+	do
+	{
+		if (Scope->HasPendingOverlaps())
+		{
+			return true;
+		}
+		Scope = Scope->OuterScope;
+	}
+	while (Scope != nullptr);
+
+	return false;
+}
+
 FORCEINLINE const TArray<struct FOverlapInfo>& FScopedMovementUpdate::GetPendingOverlaps() const
 {
 	return PendingOverlaps;
diff --git a/Engine/Source/Runtime/Engine/Private/Components/PrimitiveComponent.cpp b/Engine/Source/Runtime/Engine/Private/Components/PrimitiveComponent.cpp
--- a/Engine/Source/Runtime/Engine/Private/Components/PrimitiveComponent.cpp
+++ b/Engine/Source/Runtime/Engine/Private/Components/PrimitiveComponent.cpp
@@ -1565,15 +1565,19 @@ bool UPrimitiveComponent::MoveComponentImpl( const FVector& Delta, const FQuat&
 		else if (DeltaSizeSq == 0.f && bCollisionEnabled)
 		{
 			// We didn't move, and any rotation that doesn't change our overlap bounds means we already know what we are overlapping at this point.
-			// Can only do this if current known overlaps are valid (not deferring updates)
-			if (CVarAllowCachedOverlaps->GetInt() && Actor && Actor->GetRootComponent() == this && !IsDeferringMovementUpdates())
+			if (Actor && Actor->GetRootComponent() == this && CVarAllowCachedOverlaps->GetInt())
 			{
-				// Only if we know that we won't change our overlap status with children by moving.
-				if (AreSymmetricRotations(NewRotationQuat, ComponentToWorld.GetRotation(), ComponentToWorld.GetScale3D()) &&
-					AreAllCollideableDescendantsRelative())
+				// Can only do this if current known overlaps are valid (not deferring updates, or we know there are no new pending overlaps).
+				const FScopedMovementUpdate* ScopedUpdate = GetCurrentScopedMovement();
+				if (ScopedUpdate == nullptr || (OverlappingComponents.Num() == 0 && !ScopedUpdate->HasPendingOverlapsInAncestorOrSelf()))
 				{
-					OverlapsAtEndLocation = OverlappingComponents;
-					OverlapsAtEndLocationPtr = &OverlapsAtEndLocation;
+					// Only if we know that we won't change our overlap status with children by moving.
+					if (AreSymmetricRotations(NewRotationQuat, ComponentToWorld.GetRotation(), ComponentToWorld.GetScale3D()) &&
+						AreAllCollideableDescendantsRelative())
+					{
+						OverlapsAtEndLocation = OverlappingComponents;
+						OverlapsAtEndLocationPtr = &OverlapsAtEndLocation;
+					}
 				}
 			}
 		}
diff --git a/Engine/Source/Runtime/Engine/Private/Components/SceneComponent.cpp b/Engine/Source/Runtime/Engine/Private/Components/SceneComponent.cpp
--- a/Engine/Source/Runtime/Engine/Private/Components/SceneComponent.cpp
+++ b/Engine/Source/Runtime/Engine/Private/Components/SceneComponent.cpp
@@ -521,6 +521,7 @@ void USceneComponent::EndScopedMovementUpdate(class FScopedMovementUpdate& Compl
 		{
 			// Combine with next item on the stack
 			FScopedMovementUpdate* OuterScopedUpdate = ScopedMovementStack.Last();
+			check(OuterScopedUpdate == CurrentScopedUpdate->OuterScope);
 			OuterScopedUpdate->OnInnerScopeComplete(*CurrentScopedUpdate);
 		}
 	}
@@ -2240,11 +2241,13 @@ static uint32 s_ScopedWarningCount = 0;
 
 FScopedMovementUpdate::FScopedMovementUpdate( class USceneComponent* Component, EScopedUpdate::Type ScopeBehavior )
 : Owner(Component)
+, OuterScope(nullptr)
 , bDeferUpdates(ScopeBehavior == EScopedUpdate::DeferredUpdates)
 , bHasValidOverlapsAtEnd(false)
 {
 	if (IsValid(Component))
 	{
+		OuterScope = Component->GetCurrentScopedMovement();
 		InitialTransform = Component->GetComponentToWorld();
 		InitialRelativeLocation = Component->RelativeLocation;
 		InitialRelativeRotation = Component->RelativeRotation;
@@ -2253,7 +2256,6 @@ FScopedMovementUpdate::FScopedMovementUpdate( class USceneComponent* Component,
 		if (ScopeBehavior == EScopedUpdate::ImmediateUpdates)
 		{
 			// We only allow ScopeUpdateImmediately if there is no outer scope, or if the outer scope is also ScopeUpdateImmediately.
-			FScopedMovementUpdate* OuterScope = Component->GetCurrentScopedMovement();
 			if (OuterScope && OuterScope->bDeferUpdates)
 			{
 				if (s_ScopedWarningCount < 100 || (GFrameCounter & 31) == 0)
@@ -2340,13 +2342,13 @@ void FScopedMovementUpdate::AppendOverlaps(const TArray<struct FOverlapInfo>& Ot
 	}
 }
 
-
 void FScopedMovementUpdate::OnInnerScopeComplete(const FScopedMovementUpdate& InnerScope)
 {
 	if (IsValid(Owner))
 	{
 		checkSlow(IsDeferringUpdates());
 		checkSlow(InnerScope.IsDeferringUpdates());
+		checkSlow(InnerScope.OuterScope == this);
 
 		// Combine with the next item on the stack
 		AppendOverlaps(InnerScope.GetPendingOverlaps(), InnerScope.GetOverlapsAtEnd());
