Action mapping array order matters

Merging 2543923

//depot/UE4-UT/...

to //depot/UE4/...

diff --git a/Engine/Source/Runtime/Engine/Private/UserInterface/InputSettings.cpp b/Engine/Source/Runtime/Engine/Private/UserInterface/InputSettings.cpp
--- a/Engine/Source/Runtime/Engine/Private/UserInterface/InputSettings.cpp
+++ b/Engine/Source/Runtime/Engine/Private/UserInterface/InputSettings.cpp
@@ -110,7 +110,7 @@ void UInputSettings::RemoveActionMapping(const FInputActionKeyMapping& KeyMappin
 	{
 		if (ActionMappings[ActionIndex] == KeyMapping)
 		{
-			ActionMappings.RemoveAtSwap(ActionIndex);
+			ActionMappings.RemoveAt(ActionIndex);
 			// we don't break because the mapping may have been in the array twice
 		}
 	}
@@ -131,7 +131,7 @@ void UInputSettings::RemoveAxisMapping(const FInputAxisKeyMapping& InKeyMapping)
 		if (KeyMapping.AxisName == InKeyMapping.AxisName
 			&& KeyMapping.Key == InKeyMapping.Key)
 		{
-			AxisMappings.RemoveAtSwap(AxisIndex);
+			AxisMappings.RemoveAt(AxisIndex);
 			// we don't break because the mapping may have been in the array twice
 		}
 	}
