#ue4 - Remove unnecessary check().

diff --git a/Engine/Source/Runtime/Engine/Private/Components/SceneComponent.cpp b/Engine/Source/Runtime/Engine/Private/Components/SceneComponent.cpp
--- a/Engine/Source/Runtime/Engine/Private/Components/SceneComponent.cpp
+++ b/Engine/Source/Runtime/Engine/Private/Components/SceneComponent.cpp
@@ -521,7 +521,6 @@ void USceneComponent::EndScopedMovementUpdate(class FScopedMovementUpdate& Compl
 		{
 			// Combine with next item on the stack
 			FScopedMovementUpdate* OuterScopedUpdate = ScopedMovementStack.Last();
-			check(OuterScopedUpdate == CurrentScopedUpdate->OuterScope);
 			OuterScopedUpdate->OnInnerScopeComplete(*CurrentScopedUpdate);
 		}
 	}
