#ue4 - Perf: allow cached overlaps in FScopedMovementUpdate when just rotating and there is not outer scope.

- This will later be expanded to allow more cases (parent scopes valid).

diff --git a/Engine/Source/Runtime/Engine/Classes/Components/SceneComponent.h b/Engine/Source/Runtime/Engine/Classes/Components/SceneComponent.h
--- a/Engine/Source/Runtime/Engine/Classes/Components/SceneComponent.h
+++ b/Engine/Source/Runtime/Engine/Classes/Components/SceneComponent.h
@@ -961,6 +961,9 @@ public:
 	FScopedMovementUpdate( class USceneComponent* Component, EScopedUpdate::Type ScopeBehavior = EScopedUpdate::DeferredUpdates );
 	~FScopedMovementUpdate();
 
+	/** Get the scope containing this scope. A scope only has an outer scope if they both defer updates. */
+	const FScopedMovementUpdate* GetOuterDeferredScope() const;
+
 	/** Return true if deferring updates, false if updates are applied immediately. */
 	bool IsDeferringUpdates() const;
 	
@@ -982,6 +985,9 @@ public:
 	/** Add overlaps to the queued overlaps array. This is intended for use only by SceneComponent and its derived classes. */
 	void AppendOverlaps(const TArray<struct FOverlapInfo>& OtherOverlaps, const TArray<FOverlapInfo>* OverlapsAtEndLocation);
 
+	/** Returns whether there are valid pending overlaps at the end location. */
+	bool HasValidOverlapsAtEnd() const;
+
 	/** Returns the list of overlaps at the end location, or null if the list is invalid. */
 	const TArray<struct FOverlapInfo>* GetOverlapsAtEnd() const;
 
@@ -1006,7 +1012,7 @@ private:
 private:
 
 	class USceneComponent* Owner;
-	FScopedMovementUpdate* OuterScope;
+	FScopedMovementUpdate* OuterDeferredScope;
 	uint32 bDeferUpdates:1;
 	uint32 bHasValidOverlapsAtEnd:1;
 	FTransform InitialTransform;
@@ -1024,6 +1030,11 @@ private:
 //////////////////////////////////////////////////////////////////////////
 // FScopedMovementUpdate inlines
 
+FORCEINLINE const FScopedMovementUpdate* FScopedMovementUpdate::GetOuterDeferredScope() const
+{
+	return OuterDeferredScope;
+}
+
 FORCEINLINE bool FScopedMovementUpdate::IsDeferringUpdates() const
 {
 	return bDeferUpdates;
@@ -1043,7 +1054,7 @@ FORCEINLINE_DEBUGGABLE bool FScopedMovementUpdate::HasPendingOverlapsInAncestorO
 		{
 			return true;
 		}
-		Scope = Scope->OuterScope;
+		Scope = Scope->OuterDeferredScope;
 	}
 	while (Scope != nullptr);
 
@@ -1055,6 +1066,11 @@ FORCEINLINE const TArray<struct FOverlapInfo>& FScopedMovementUpdate::GetPending
 	return PendingOverlaps;
 }
 
+FORCEINLINE bool FScopedMovementUpdate::HasValidOverlapsAtEnd() const
+{
+	return bHasValidOverlapsAtEnd;
+}
+
 FORCEINLINE const TArray<struct FOverlapInfo>* FScopedMovementUpdate::GetOverlapsAtEnd() const
 {
 	return bHasValidOverlapsAtEnd ? &OverlapsAtEnd : NULL;
diff --git a/Engine/Source/Runtime/Engine/Private/Components/PrimitiveComponent.cpp b/Engine/Source/Runtime/Engine/Private/Components/PrimitiveComponent.cpp
--- a/Engine/Source/Runtime/Engine/Private/Components/PrimitiveComponent.cpp
+++ b/Engine/Source/Runtime/Engine/Private/Components/PrimitiveComponent.cpp
@@ -1569,7 +1569,9 @@ bool UPrimitiveComponent::MoveComponentImpl( const FVector& Delta, const FQuat&
 			{
 				// Can only do this if current known overlaps are valid (not deferring updates, or we know there are no new pending overlaps).
 				const FScopedMovementUpdate* ScopedUpdate = GetCurrentScopedMovement();
-				if (ScopedUpdate == nullptr || (OverlappingComponents.Num() == 0 && !ScopedUpdate->HasPendingOverlapsInAncestorOrSelf()))
+				if (ScopedUpdate == nullptr
+					|| (OverlappingComponents.Num() == 0 && !ScopedUpdate->HasPendingOverlapsInAncestorOrSelf())
+					|| (!ScopedUpdate->GetOuterDeferredScope()))
 				{
 					// Only if we know that we won't change our overlap status with children by moving.
 					if (AreSymmetricRotations(NewRotationQuat, ComponentToWorld.GetRotation(), ComponentToWorld.GetScale3D()) &&
diff --git a/Engine/Source/Runtime/Engine/Private/Components/SceneComponent.cpp b/Engine/Source/Runtime/Engine/Private/Components/SceneComponent.cpp
--- a/Engine/Source/Runtime/Engine/Private/Components/SceneComponent.cpp
+++ b/Engine/Source/Runtime/Engine/Private/Components/SceneComponent.cpp
@@ -497,7 +497,7 @@ void USceneComponent::EndScopedMovementUpdate(class FScopedMovementUpdate& Compl
 			}
 
 			// We may have moved somewhere and then moved back to the start, we still need to update overlaps if we touched things along the way.
-			if (bMoved || CurrentScopedUpdate->HasPendingOverlaps() || CurrentScopedUpdate->GetOverlapsAtEnd())
+			if (bMoved || CurrentScopedUpdate->HasPendingOverlaps() || CurrentScopedUpdate->HasValidOverlapsAtEnd())
 			{
 				UpdateOverlaps(&CurrentScopedUpdate->GetPendingOverlaps(), true, CurrentScopedUpdate->GetOverlapsAtEnd());
 			}
@@ -2240,13 +2240,13 @@ static uint32 s_ScopedWarningCount = 0;
 
 FScopedMovementUpdate::FScopedMovementUpdate( class USceneComponent* Component, EScopedUpdate::Type ScopeBehavior )
 : Owner(Component)
-, OuterScope(nullptr)
+, OuterDeferredScope(nullptr)
 , bDeferUpdates(ScopeBehavior == EScopedUpdate::DeferredUpdates)
 , bHasValidOverlapsAtEnd(false)
 {
 	if (IsValid(Component))
 	{
-		OuterScope = Component->GetCurrentScopedMovement();
+		OuterDeferredScope = Component->GetCurrentScopedMovement();
 		InitialTransform = Component->GetComponentToWorld();
 		InitialRelativeLocation = Component->RelativeLocation;
 		InitialRelativeRotation = Component->RelativeRotation;
@@ -2255,7 +2255,7 @@ FScopedMovementUpdate::FScopedMovementUpdate( class USceneComponent* Component,
 		if (ScopeBehavior == EScopedUpdate::ImmediateUpdates)
 		{
 			// We only allow ScopeUpdateImmediately if there is no outer scope, or if the outer scope is also ScopeUpdateImmediately.
-			if (OuterScope && OuterScope->bDeferUpdates)
+			if (OuterDeferredScope && OuterDeferredScope->bDeferUpdates)
 			{
 				if (s_ScopedWarningCount < 100 || (GFrameCounter & 31) == 0)
 				{
@@ -2307,6 +2307,7 @@ void FScopedMovementUpdate::RevertMove()
 	{
 		bHasValidOverlapsAtEnd = false;
 		PendingOverlaps.Reset();
+		OverlapsAtEnd.Reset();
 		BlockingHits.Reset();
 		
 		if (IsTransformDirty())
@@ -2347,7 +2348,7 @@ void FScopedMovementUpdate::OnInnerScopeComplete(const FScopedMovementUpdate& In
 	{
 		checkSlow(IsDeferringUpdates());
 		checkSlow(InnerScope.IsDeferringUpdates());
-		checkSlow(InnerScope.OuterScope == this);
+		checkSlow(InnerScope.OuterDeferredScope == this);
 
 		// Combine with the next item on the stack
 		AppendOverlaps(InnerScope.GetPendingOverlaps(), InnerScope.GetOverlapsAtEnd());
