Made PawnNoiseEmmiterComponent talk to AIPerceptionSystem (can opt-out) #UE4

UE-15318

diff --git a/Engine/Source/Runtime/Engine/Classes/Components/PawnNoiseEmitterComponent.h b/Engine/Source/Runtime/Engine/Classes/Components/PawnNoiseEmitterComponent.h
--- a/Engine/Source/Runtime/Engine/Classes/Components/PawnNoiseEmitterComponent.h
+++ b/Engine/Source/Runtime/Engine/Classes/Components/PawnNoiseEmitterComponent.h
@@ -14,6 +14,12 @@ class ENGINE_API UPawnNoiseEmitterComponent : public UActorComponent
 {
 	GENERATED_UCLASS_BODY()
 
+protected:
+	// If set to true (default value) will notify AIPerceptionSystem about noise events
+	// otherwise only PawnSensingComponents will be able to pick up noises generated by this component
+	UPROPERTY(EditDefaultsOnly, Category = "AI|Perception", AdvancedDisplay)
+	uint32 bAIPerceptionSystemCompatibilityMode : 1;
+
 public:
 	// Most recent noise made by this pawn not at its own location
 	UPROPERTY()
diff --git a/Engine/Source/Runtime/Engine/Private/Components/PawnNoiseEmitterComponent.cpp b/Engine/Source/Runtime/Engine/Private/Components/PawnNoiseEmitterComponent.cpp
--- a/Engine/Source/Runtime/Engine/Private/Components/PawnNoiseEmitterComponent.cpp
+++ b/Engine/Source/Runtime/Engine/Private/Components/PawnNoiseEmitterComponent.cpp
@@ -4,14 +4,12 @@
 #include "EnginePrivate.h"
 #include "Components/PawnNoiseEmitterComponent.h"
 
-
-
-
 UPawnNoiseEmitterComponent::UPawnNoiseEmitterComponent(const FObjectInitializer& ObjectInitializer)
 	: Super(ObjectInitializer)
 {
 	NoiseLifetime = 1.0f;
 	PrimaryComponentTick.bCanEverTick = false;
+	bAIPerceptionSystemCompatibilityMode = true;
 }
 
 
@@ -63,6 +61,11 @@ void UPawnNoiseEmitterComponent::MakeNoise(AActor* NoiseMaker, float Loudness, c
 			LastRemoteNoiseTime = GetWorld()->GetTimeSeconds();
 		}
 	}
+
+	if (bAIPerceptionSystemCompatibilityMode)
+	{
+		NoiseMaker->MakeNoise(Loudness, PawnOwner, NoiseLocation);
+	}
 }
 
 float UPawnNoiseEmitterComponent::GetLastNoiseVolume(bool bSourceWithinNoiseEmitter) const
