Engine: Fixed an issue where convex collision elements for a scaled component would be drawn at the wrong size (scaled twice) [UE-15327]

diff --git a/Engine/Source/Runtime/Engine/Private/PhysicsEngine/PhysDrawing.cpp b/Engine/Source/Runtime/Engine/Private/PhysicsEngine/PhysDrawing.cpp
--- a/Engine/Source/Runtime/Engine/Private/PhysicsEngine/PhysDrawing.cpp
+++ b/Engine/Source/Runtime/Engine/Private/PhysicsEngine/PhysDrawing.cpp
@@ -596,7 +596,7 @@ void FKAggregateGeom::GetAggGeom(const FTransform& Transform, const FColor Color
 				for(int32 i=0; i<ConvexElems.Num(); i++)
 				{
 					// Get vertices/triangles from this hull.
-					ConvexElems[i].AddCachedSolidConvexGeom(ThisGeom.RenderInfo->VertexBuffer->Vertices, ThisGeom.RenderInfo->IndexBuffer->Indices, Scale3D.X, FColor::White);
+					ConvexElems[i].AddCachedSolidConvexGeom(ThisGeom.RenderInfo->VertexBuffer->Vertices, ThisGeom.RenderInfo->IndexBuffer->Indices, 1.0f, FColor::White);
 				}
 
 				// Only continue if we actually got some valid geometry
