UE-15206 This user's project freezes indefinitely when scaling a mesh after its relative scale was set to negative
Data Only Blueprint is recompiled while regeneration, if it has no SceneRoot. While compiation the DefaultSceneRoot is added and a new property is created, to the data layout is changed.

diff --git a/Engine/Source/Editor/UnrealEd/Private/Kismet2/BlueprintEditorUtils.cpp b/Engine/Source/Editor/UnrealEd/Private/Kismet2/BlueprintEditorUtils.cpp
--- a/Engine/Source/Editor/UnrealEd/Private/Kismet2/BlueprintEditorUtils.cpp
+++ b/Engine/Source/Editor/UnrealEd/Private/Kismet2/BlueprintEditorUtils.cpp
@@ -1426,7 +1426,12 @@ UClass* FBlueprintEditorUtils::RegenerateBlueprintClass(UBlueprint* Blueprint, U
 		const bool bHasPendingUberGraphFrame = BPGClassToRegenerate
 			&& (BPGClassToRegenerate->UberGraphFramePointerProperty || BPGClassToRegenerate->UberGraphFunction);
 
-		const bool bShouldBeRecompiled = bHasCode || bDataOnlyClassThatMustBeRecompiled || bHasPendingUberGraphFrame;
+		const bool bDefaultComponentMustBeAdded = !bHasCode 
+			&& BPGClassToRegenerate
+			&& SupportsConstructionScript(Blueprint) 
+			&& BPGClassToRegenerate->SimpleConstructionScript
+			&& (nullptr == BPGClassToRegenerate->SimpleConstructionScript->GetSceneRootComponentTemplate());
+		const bool bShouldBeRecompiled = bHasCode || bDataOnlyClassThatMustBeRecompiled || bHasPendingUberGraphFrame || bDefaultComponentMustBeAdded;
 
 		if (bShouldBeRecompiled)
 		{
diff --git a/Engine/Source/Runtime/Engine/Classes/Engine/SimpleConstructionScript.h b/Engine/Source/Runtime/Engine/Classes/Engine/SimpleConstructionScript.h
--- a/Engine/Source/Runtime/Engine/Classes/Engine/SimpleConstructionScript.h
+++ b/Engine/Source/Runtime/Engine/Classes/Engine/SimpleConstructionScript.h
@@ -32,6 +32,9 @@ class USimpleConstructionScript : public UObject
 #if WITH_EDITOR
 	/** Return the Blueprint associated with this SCS instance */
 	ENGINE_API class UBlueprint* GetBlueprint() const;
+
+	/** Helper function to find the current scene root component template and/or owning SCS node */
+	ENGINE_API USceneComponent* GetSceneRootComponentTemplate(USCS_Node** OutSCSNode = nullptr) const;
 #endif
 
 	/** Return the Blueprint associated with this SCS instance */
@@ -100,9 +103,6 @@ private:
 	void FixupSceneNodeHierarchy();
 
 #if WITH_EDITOR
-	/** Helper function to find the current scene root component template and/or owning SCS node */
-	USceneComponent* GetSceneRootComponentTemplate(USCS_Node** OutSCSNode = nullptr) const;
-
 	/** Helper function for generating list of currently used names */
 	void GenerateListOfExistingNames(TArray<FName>& CurrentNames) const;
 
