A fix to infinite recursion resulting from calling MakeNoise on a pawn while UAISense_Hearing is not being used #UE4

Note that solution is a bit hacky and the only proper way to fix it requires moving NoiseEmmiter component to AIModule, which I just filed as a task

diff --git a/Engine/Source/Runtime/Engine/Private/Components/PawnNoiseEmitterComponent.cpp b/Engine/Source/Runtime/Engine/Private/Components/PawnNoiseEmitterComponent.cpp
--- a/Engine/Source/Runtime/Engine/Private/Components/PawnNoiseEmitterComponent.cpp
+++ b/Engine/Source/Runtime/Engine/Private/Components/PawnNoiseEmitterComponent.cpp
@@ -15,6 +15,16 @@ UPawnNoiseEmitterComponent::UPawnNoiseEmitterComponent(const FObjectInitializer&
 
 void UPawnNoiseEmitterComponent::MakeNoise(AActor* NoiseMaker, float Loudness, const FVector& NoiseLocation)
 {
+	// @hack! this won't be needed once UPawnNoiseEmitterComponent gets moved to AIModule
+	// there's no other way to easily and efficiently prevent infinite recursion when 
+	// bAIPerceptionSystemCompatibilityMode == true and UAISense_Hearing is not being used (yet)
+	static bool bMakeNoiseLockHack = false;
+	if (bMakeNoiseLockHack)
+	{
+		bMakeNoiseLockHack = false;
+		return;
+	}
+
 	if (!NoiseMaker || Loudness <= 0.f)
 	{
 		return;
@@ -64,6 +74,7 @@ void UPawnNoiseEmitterComponent::MakeNoise(AActor* NoiseMaker, float Loudness, c
 
 	if (bAIPerceptionSystemCompatibilityMode)
 	{
+		bMakeNoiseLockHack = true;
 		NoiseMaker->MakeNoise(Loudness, PawnOwner, NoiseLocation);
 	}
 }
