Fixed crash when assigning a parameter to a LandscapeGrassOutput material node (UE-15069)

diff --git a/Engine/Source/Runtime/Engine/Private/Materials/HLSLMaterialTranslator.h b/Engine/Source/Runtime/Engine/Private/Materials/HLSLMaterialTranslator.h
--- a/Engine/Source/Runtime/Engine/Private/Materials/HLSLMaterialTranslator.h
+++ b/Engine/Source/Runtime/Engine/Private/Materials/HLSLMaterialTranslator.h
@@ -3753,7 +3753,15 @@ protected:
 
 		FString Definitions;
 		FString Body;
-		GetFixedParameterCode(OutputCode, *CurrentScopeChunks, Definitions, Body);
+
+		if ((*CurrentScopeChunks)[OutputCode].UniformExpression && !(*CurrentScopeChunks)[OutputCode].UniformExpression->IsConstant())
+		{
+			Body = GetParameterCode(OutputCode);
+		}
+		else
+		{
+			GetFixedParameterCode(OutputCode, *CurrentScopeChunks, Definitions, Body);
+		}
 
 		FString ImplementationCode = FString::Printf(TEXT("%s %s%d(FMaterial%sParameters Parameters)\r\n{\r\n%s return %s;\r\n}\r\n"), *OutputTypeString, *Custom->GetFunctionName(), OutputIndex, ShaderFrequency == SF_Vertex ? TEXT("Vertex") : TEXT("Pixel"), *Definitions, *Body);
 		CustomOutputImplementations.Add(ImplementationCode);
