UE-15007 Re-enable LOD category for (Material)BillboardComponent

diff --git a/Engine/Source/Runtime/Engine/Classes/Components/BillboardComponent.h b/Engine/Source/Runtime/Engine/Classes/Components/BillboardComponent.h
--- a/Engine/Source/Runtime/Engine/Classes/Components/BillboardComponent.h
+++ b/Engine/Source/Runtime/Engine/Classes/Components/BillboardComponent.h
@@ -10,7 +10,7 @@ struct FConvexVolume;
 /** 
  * A 2d texture that will be rendered always facing the camera.
  */
-UCLASS(ClassGroup=Rendering, collapsecategories, hidecategories=(Object,Activation,"Components|Activation",Physics,Collision,Lighting,LOD,Mesh,PhysicsVolume), editinlinenew, meta=(BlueprintSpawnableComponent))
+UCLASS(ClassGroup=Rendering, collapsecategories, hidecategories=(Object,Activation,"Components|Activation",Physics,Collision,Lighting,Mesh,PhysicsVolume), editinlinenew, meta=(BlueprintSpawnableComponent))
 class ENGINE_API UBillboardComponent : public UPrimitiveComponent
 {
 	GENERATED_UCLASS_BODY()
diff --git a/Engine/Source/Runtime/Engine/Classes/Components/MaterialBillboardComponent.h b/Engine/Source/Runtime/Engine/Classes/Components/MaterialBillboardComponent.h
--- a/Engine/Source/Runtime/Engine/Classes/Components/MaterialBillboardComponent.h
+++ b/Engine/Source/Runtime/Engine/Classes/Components/MaterialBillboardComponent.h
@@ -50,7 +50,7 @@ struct FMaterialSpriteElement
 /** 
  * A 2d material that will be rendered always facing the camera.
  */
-UCLASS(ClassGroup=Rendering, collapsecategories, hidecategories=(Object,Activation,"Components|Activation",Physics,Collision,Lighting,LOD,Mesh,PhysicsVolume), editinlinenew, meta=(BlueprintSpawnableComponent))
+UCLASS(ClassGroup=Rendering, collapsecategories, hidecategories=(Object,Activation,"Components|Activation",Physics,Collision,Lighting,Mesh,PhysicsVolume), editinlinenew, meta=(BlueprintSpawnableComponent))
 class ENGINE_API UMaterialBillboardComponent : public UPrimitiveComponent
 {
 	GENERATED_UCLASS_BODY()
