Fix for UE-14456: make sure to clamp alpha when building lookup tables for distributions

diff --git a/Engine/Source/Runtime/Engine/Private/Distributions.cpp b/Engine/Source/Runtime/Engine/Private/Distributions.cpp
--- a/Engine/Source/Runtime/Engine/Private/Distributions.cpp
+++ b/Engine/Source/Runtime/Engine/Private/Distributions.cpp
@@ -133,7 +133,7 @@ void BuildLookupTable( FDistributionLookupTable* OutTable, const DistributionTyp
 	// Sample the distribution.
 	for ( uint32 SampleIndex = 0; SampleIndex < EntryCount; SampleIndex++ )
 	{
-		const float Time = MinIn + SampleIndex * TimeScale;
+		const float Time = FMath::Clamp(MinIn + SampleIndex * TimeScale, 0.0f, 1.0f);
 		float Values[8];
 		Distribution->InitializeRawEntry( Time, Values );
 		for ( uint32 ValueIndex = 0; ValueIndex < EntryStride; ValueIndex++ )
