Fix for UE-10588
Let's not crash when trying to show collision on a cooked build.

diff --git a/Engine/Source/Runtime/Engine/Private/StaticMeshRender.cpp b/Engine/Source/Runtime/Engine/Private/StaticMeshRender.cpp
--- a/Engine/Source/Runtime/Engine/Private/StaticMeshRender.cpp
+++ b/Engine/Source/Runtime/Engine/Private/StaticMeshRender.cpp
@@ -735,27 +735,30 @@ void FStaticMeshSceneProxy::GetDynamicMeshElements(const TArray<const FSceneView
 											{
 												UMaterial* MaterialToUse = (bDrawComplexWireframeCollision) ? GEngine->WireframeMaterial : GEngine->ShadedLevelColorationUnlitMaterial;
 
-												// Override the mesh's material with our material that draws the collision color
-												auto CollisionMaterialInstance = new FColoredMaterialRenderProxy(
-													MaterialToUse->GetRenderProxy(bProxyIsSelected, IsHovered()),
-													WireframeColor
-													);
-
-												Collector.RegisterOneFrameMaterialProxy(CollisionMaterialInstance);
-
-												// If drawing the complex collision in wireframe, we do that in _addition_ to the normal mesh
-												if(bDrawComplexWireframeCollision)
-												{
-													FMeshBatch& CollisionElement = Collector.AllocateMesh();
-													GetMeshElement(LODIndex, BatchIndex, SectionIndex, SDPG_World, bSectionIsSelected, IsHovered(), CollisionElement);
-													CollisionElement.MaterialRenderProxy = CollisionMaterialInstance;
-													Collector.AddMesh(ViewIndex, CollisionElement);
-													INC_DWORD_STAT_BY(STAT_StaticMeshTriangles, CollisionElement.GetNumPrimitives());
-												}
-												// Otherwise just change material on existing mesh
-												else
+												if (MaterialToUse)
 												{
-													MeshElement.MaterialRenderProxy = CollisionMaterialInstance;
+													// Override the mesh's material with our material that draws the collision color
+													auto CollisionMaterialInstance = new FColoredMaterialRenderProxy(
+														MaterialToUse ? MaterialToUse->GetRenderProxy(bProxyIsSelected, IsHovered()) : NULL,
+														WireframeColor
+														);
+
+													Collector.RegisterOneFrameMaterialProxy(CollisionMaterialInstance);
+
+													// If drawing the complex collision in wireframe, we do that in _addition_ to the normal mesh
+													if (bDrawComplexWireframeCollision)
+													{
+														FMeshBatch& CollisionElement = Collector.AllocateMesh();
+														GetMeshElement(LODIndex, BatchIndex, SectionIndex, SDPG_World, bSectionIsSelected, IsHovered(), CollisionElement);
+														CollisionElement.MaterialRenderProxy = CollisionMaterialInstance;
+														Collector.AddMesh(ViewIndex, CollisionElement);
+														INC_DWORD_STAT_BY(STAT_StaticMeshTriangles, CollisionElement.GetNumPrimitives());
+													}
+													// Otherwise just change material on existing mesh
+													else
+													{
+														MeshElement.MaterialRenderProxy = CollisionMaterialInstance;
+													}
 												}
 											}
 											else
