Fixed issue where a Matinee event which stops the Matinee sequence could cause a crash when updating subsequent tracks.

diff --git a/Engine/Source/Runtime/Engine/Classes/Matinee/MatineeActor.h b/Engine/Source/Runtime/Engine/Classes/Matinee/MatineeActor.h
--- a/Engine/Source/Runtime/Engine/Classes/Matinee/MatineeActor.h
+++ b/Engine/Source/Runtime/Engine/Classes/Matinee/MatineeActor.h
@@ -188,6 +188,9 @@ public:
 	UPROPERTY(replicated,transient)
 	uint32 bPaused:1;
 
+	UPROPERTY(replicated,transient)
+	uint32 bPendingStop:1;
+
 	UPROPERTY(replicated)
 	float InterpPosition;
 
diff --git a/Engine/Source/Runtime/Engine/Private/Interpolation.cpp b/Engine/Source/Runtime/Engine/Private/Interpolation.cpp
--- a/Engine/Source/Runtime/Engine/Private/Interpolation.cpp
+++ b/Engine/Source/Runtime/Engine/Private/Interpolation.cpp
@@ -482,17 +482,7 @@ void AMatineeActor::Reverse()
 
 void AMatineeActor::Stop()
 {
-	// Re-enable the radio filter
-	EnableRadioFilter();
-
-	bIsPlaying = false;
-	bPaused = false;
-
-	if( GetWorld()->IsGameWorld() )
-	{
-		// We should only terminate the interp in the game.  The editor handles this from inside the matinee editor
-		TermInterp();
-	}
+	bPendingStop = true;
 }
 
 void AMatineeActor::Pause()
@@ -607,6 +597,22 @@ void AMatineeActor::Tick(float DeltaTime)
 	{
 		StepInterp(DeltaTime, false);
 	}
+
+	if (bPendingStop)
+	{
+		// Re-enable the radio filter
+		EnableRadioFilter();
+
+		bIsPlaying = false;
+		bPaused = false;
+		bPendingStop = false;
+
+		if (GetWorld()->IsGameWorld())
+		{
+			// We should only terminate the interp in the game.  The editor handles this from inside the matinee editor
+			TermInterp();
+		}
+	}
 }
 
 
@@ -731,22 +737,9 @@ void AMatineeActor::UpdateInterp( float NewPosition, bool bPreview, bool bJump )
 			for( int32 GroupIndex = 0; GroupIndex < Groups.Num(); ++GroupIndex )
 			{
 				Groups[GroupIndex]->Group->UpdateGroup( NewPosition, Groups[GroupIndex], bPreview, bJump );
-
-				const bool bhasBeenTerminated = (GroupInst.Num() == 0);
-#if WITH_EDITORONLY_DATA
-				if (bhasBeenTerminated && !bIsBeingEdited)
-#else
-				if (bhasBeenTerminated)
-#endif
-				{
-					UE_LOG(LogMatinee, Log, TEXT("WARNING: A matinee was stopped while updating group '%s'; the next groups will not be updated."), *Groups[GroupIndex]->Group->GetFullGroupName(true));
-					InterpPosition = NewPosition;
-					return;
-				}
 			}
 		}
 
-
 		InterpPosition = NewPosition;
 	}
 }
