Niagara: set material usage flags

diff --git a/Engine/Source/Runtime/Engine/Private/NiagaraEffectRenderer.cpp b/Engine/Source/Runtime/Engine/Private/NiagaraEffectRenderer.cpp
--- a/Engine/Source/Runtime/Engine/Private/NiagaraEffectRenderer.cpp
+++ b/Engine/Source/Runtime/Engine/Private/NiagaraEffectRenderer.cpp
@@ -90,9 +90,9 @@ void NiagaraEffectRendererSprites::GetDynamicMeshElements(const TArray<const FSc
 
 	SimpleTimer MeshElementsTimer;
 
-	check(DynamicDataRender)
+	//check(DynamicDataRender)
 		
-	if (DynamicDataRender->VertexData.Num() == 0)
+	if (!DynamicDataRender || DynamicDataRender->VertexData.Num() == 0)
 	{
 		return;
 	}
@@ -205,8 +205,7 @@ void NiagaraEffectRendererSprites::GetDynamicMeshElements(const TArray<const FSc
 bool NiagaraEffectRendererSprites::SetMaterialUsage()
 {
 	//Causes deadlock :S Need to look at / rework the setting of materials and render modules.
-	//return Material && Material->CheckMaterialUsage_Concurrent(MATUSAGE_ParticleSprites);
-	return true;
+	return Material && Material->CheckMaterialUsage(MATUSAGE_ParticleSprites);
 }
 
 /** Update render data buffer from attributes */
@@ -488,9 +487,7 @@ const TArray<FNiagaraVariableInfo>& NiagaraEffectRendererRibbon::GetRequiredAttr
 
 bool NiagaraEffectRendererRibbon::SetMaterialUsage()
 {
-	//Causes deadlock :S Need to look at / rework the setting of materials and render modules.
-	//return Material && Material->CheckMaterialUsage_Concurrent(MATUSAGE_BeamTrails);
-	return true;
+	return Material && Material->CheckMaterialUsage_Concurrent(MATUSAGE_BeamTrails);
 }
 
 FNiagaraDynamicDataBase *NiagaraEffectRendererRibbon::GenerateVertexData(const FNiagaraEmitterParticleData &Data)
