Fixed auto merge error.

diff --git a/Engine/Source/Runtime/Engine/Private/Components/CharacterMovementComponent.cpp b/Engine/Source/Runtime/Engine/Private/Components/CharacterMovementComponent.cpp
--- a/Engine/Source/Runtime/Engine/Private/Components/CharacterMovementComponent.cpp
+++ b/Engine/Source/Runtime/Engine/Private/Components/CharacterMovementComponent.cpp
@@ -1562,7 +1562,6 @@ void UCharacterMovementComponent::PerformMovement(float DeltaSeconds)
 			{
 				// Convert Local Space Root Motion to world space. Do it right before used by physics to make sure we use up to date transforms, as translation is relative to rotation.
 				RootMotionParams.Set( SkelMeshComp->ConvertLocalRootMotionToWorld(RootMotionParams.RootMotionTransform) );
-					*RootMotionParams.RootMotionTransform.GetTranslation().ToCompactString(), *RootMotionParams.RootMotionTransform.GetRotation().Rotator().ToCompactString(), *UpdatedComponent->GetForwardVector().ToCompactString());
 
 				// Then turn root motion to velocity to be used by various physics modes.
 				if (DeltaSeconds > 0.f)
