Fix child actor component parent being set to null

diff --git a/Engine/Source/Runtime/Engine/Private/Components/ChildActorComponent.cpp b/Engine/Source/Runtime/Engine/Private/Components/ChildActorComponent.cpp
--- a/Engine/Source/Runtime/Engine/Private/Components/ChildActorComponent.cpp
+++ b/Engine/Source/Runtime/Engine/Private/Components/ChildActorComponent.cpp
@@ -234,7 +234,8 @@ void UChildActorComponent::CreateChildActor()
 		{
 			// Before we spawn let's try and prevent cyclic disaster
 			bool bSpawn = true;
-			AActor* Actor = GetOwner();
+			AActor* MyOwner = GetOwner();
+			AActor* Actor = MyOwner;
 			while (Actor && bSpawn)
 			{
 				if (Actor->GetClass() == ChildActorClass)
@@ -251,7 +252,7 @@ void UChildActorComponent::CreateChildActor()
 				Params.bNoCollisionFail = true;
 				Params.bDeferConstruction = true; // We defer construction so that we set ParentComponentActor prior to component registration so they appear selected
 				Params.bAllowDuringConstructionScript = true;
-				Params.OverrideLevel = (Actor ? Actor->GetLevel() : nullptr);
+				Params.OverrideLevel = (MyOwner ? MyOwner->GetLevel() : nullptr);
 				Params.Name = ChildActorName;
 				if (!HasAllFlags(RF_Transactional))
 				{
@@ -269,7 +270,7 @@ void UChildActorComponent::CreateChildActor()
 					ChildActorName = ChildActor->GetFName();
 
 					// Remember which actor spawned it (for selection in editor etc)
-					ChildActor->ParentComponentActor = Actor;
+					ChildActor->ParentComponentActor = MyOwner;
 
 					ChildActor->AttachRootComponentTo(this);
 
