#ue4 - Rotational movement can now avoid UpdateOverlaps when the rotation is symmetric even when currently overlapping components in most cases.

- Also fixes invalidated outer scopes when inner scopes have never moved or have been reverted.

diff --git a/Engine/Source/Runtime/Engine/Classes/Components/SceneComponent.h b/Engine/Source/Runtime/Engine/Classes/Components/SceneComponent.h
--- a/Engine/Source/Runtime/Engine/Classes/Components/SceneComponent.h
+++ b/Engine/Source/Runtime/Engine/Classes/Components/SceneComponent.h
@@ -967,10 +967,19 @@ public:
 	/** Return true if deferring updates, false if updates are applied immediately. */
 	bool IsDeferringUpdates() const;
 	
-	/** Revert movement to the initial location of the Component at the start of the scoped update. Also clears pending overlaps. */
+	/** Revert movement to the initial location of the Component at the start of the scoped update. Also clears pending overlaps and sets bHasMoved to false. */
 	void RevertMove();
 
-	/** Returns true if the Component's transform has changed since the start of the scoped update. */
+	enum class EHasMovedTransformOption
+	{
+		eTestTransform,
+		eIgnoreTransform
+	};
+
+	/** Returns whether movement has occurred at all during this scope, optionally checking if the transform is different (since changing scale does not go through a move). RevertMove() sets this back to false. */
+	bool HasMoved(EHasMovedTransformOption CheckTransform) const;
+
+	/** Returns true if the Component's transform differs from that at the start of the scoped update. */
 	bool IsTransformDirty() const;
 
 	/** Returns true if there are pending overlaps queued in this scope. */
@@ -980,25 +989,41 @@ public:
 	bool HasPendingOverlapsInAncestorOrSelf() const;
 
 	/** Returns the pending overlaps within this scope. */
-	const TArray<struct FOverlapInfo>& GetPendingOverlaps() const;
-	
-	/** Add overlaps to the queued overlaps array. This is intended for use only by SceneComponent and its derived classes. */
-	void AppendOverlaps(const TArray<struct FOverlapInfo>& OtherOverlaps, const TArray<FOverlapInfo>* OverlapsAtEndLocation);
+	const TArray<FOverlapInfo>& GetPendingOverlaps() const;
 
-	/** Returns whether there are valid pending overlaps at the end location. */
-	bool HasValidOverlapsAtEnd() const;
+	/** Returns whether the overlaps at the current location are completely known. May walk the outer scope chain to compute this if no movement has occurred. */
+	bool HasKnownOverlapStateAtEnd() const;
 
-	/** Returns the list of overlaps at the end location, or null if the list is invalid. */
-	const TArray<struct FOverlapInfo>* GetOverlapsAtEnd() const;
+	/** Returns the overlaps at the current location if they are known, or null otherwise. See HasKnownOverlapStateAtEnd(). */
+	const TArray<FOverlapInfo>* GetKnownOverlapsAtEnd(const TArray<struct FOverlapInfo>* DefaultOverlaps) const;
+
+	/** Add overlaps to the queued overlaps array. This is intended for use only by SceneComponent and its derived classes whenever movement is performed. Null NewOverlapsAtEndLocation indicates an unknown overlap state. */
+	void AppendOverlapsAfterMove(const TArray<FOverlapInfo>& NewPendingOverlaps, const TArray<FOverlapInfo>* NewOverlapsAtEndLocation);
 
 	/** Add blocking hit that will get processed once the move is committed. This is intended for use only by SceneComponent and its derived classes. */
-	void AppendBlockingHit(const FHitResult& Hit);
+	void AppendBlockingHitAfterMove(const FHitResult& Hit);
 	
 	/** Returns the list of pending blocking hits, which will be used for notifications once the move is committed. */
 	const TBlockingHitArray& GetPendingBlockingHits() const;
 
+
+	DEPRECATED(4.9, "This function has been renamed to AppendOverlapsAfterMove.")
+	void AppendOverlaps(const TArray<FOverlapInfo>& NewPendingOverlaps, const TArray<FOverlapInfo>* NewOverlapsAtEndLocation);
+
+	DEPRECATED(4.9, "This function has been renamed to AppendOverlapsAfterMove.")
+	void AppendBlockingHit(const FHitResult& Hit);
+
+	DEPRECATED(4.9, "This function is ambiguous. Use either HasStoredOverlapsAtEnd() or HasKnownOverlapStateAtEnd().")
+	bool HasValidOverlapsAtEnd() const;
+
 private:
 
+	/** Returns whether there are stored pending overlaps at the end location. If there has been no movement, this returns false, but overlaps may be extracted from the outer scope. */
+	bool HasStoredOverlapsAtEnd() const;
+
+	/** Returns the list of overlaps at the end location, or null if HasStoredOverlapsAtEnd() is false. */
+	const TArray<FOverlapInfo>* GetOverlapsAtEnd() const;
+
 	/** Notify this scope that the given inner scope completed its update (ie is going out of scope). Only occurs for deferred updates. */
 	void OnInnerScopeComplete(const FScopedMovementUpdate& InnerScope);
 	
@@ -1014,14 +1039,15 @@ private:
 	class USceneComponent* Owner;
 	FScopedMovementUpdate* OuterDeferredScope;
 	uint32 bDeferUpdates:1;
-	uint32 bHasValidOverlapsAtEnd:1;
+	uint32 bHasStoredOverlapsAtEnd:1;
+	uint32 bHasMoved:1;
 	FTransform InitialTransform;
 	FVector InitialRelativeLocation;
 	FRotator InitialRelativeRotation;
 	FVector InitialRelativeScale;
 
-	TArray<struct FOverlapInfo> PendingOverlaps;
-	TArray<struct FOverlapInfo> OverlapsAtEnd;
+	TArray<FOverlapInfo> PendingOverlaps;
+	TArray<FOverlapInfo> OverlapsAtEnd;
 	TBlockingHitArray BlockingHits;
 
 	friend class USceneComponent;
@@ -1040,6 +1066,11 @@ FORCEINLINE bool FScopedMovementUpdate::IsDeferringUpdates() const
 	return bDeferUpdates;
 }
 
+FORCEINLINE bool FScopedMovementUpdate::HasMoved(EHasMovedTransformOption CheckTransform) const
+{
+	return bHasMoved || (CheckTransform == EHasMovedTransformOption::eTestTransform && IsTransformDirty());
+}
+
 FORCEINLINE bool FScopedMovementUpdate::HasPendingOverlaps() const
 {
 	return PendingOverlaps.Num() > 0;
@@ -1066,14 +1097,14 @@ FORCEINLINE const TArray<struct FOverlapInfo>& FScopedMovementUpdate::GetPending
 	return PendingOverlaps;
 }
 
-FORCEINLINE bool FScopedMovementUpdate::HasValidOverlapsAtEnd() const
+FORCEINLINE bool FScopedMovementUpdate::HasStoredOverlapsAtEnd() const
 {
-	return bHasValidOverlapsAtEnd;
+	return bHasStoredOverlapsAtEnd;
 }
 
 FORCEINLINE const TArray<struct FOverlapInfo>* FScopedMovementUpdate::GetOverlapsAtEnd() const
 {
-	return bHasValidOverlapsAtEnd ? &OverlapsAtEnd : NULL;
+	return bHasStoredOverlapsAtEnd ? &OverlapsAtEnd : nullptr;
 }
 
 FORCEINLINE const FScopedMovementUpdate::TBlockingHitArray& FScopedMovementUpdate::GetPendingBlockingHits() const
@@ -1081,7 +1112,23 @@ FORCEINLINE const FScopedMovementUpdate::TBlockingHitArray& FScopedMovementUpdat
 	return BlockingHits;
 }
 
-FORCEINLINE_DEBUGGABLE void FScopedMovementUpdate::AppendBlockingHit(const FHitResult& Hit)
+FORCEINLINE_DEBUGGABLE void FScopedMovementUpdate::AppendBlockingHitAfterMove(const FHitResult& Hit)
 {
 	BlockingHits.Add(Hit);
 }
+
+FORCEINLINE void FScopedMovementUpdate::AppendOverlaps(const TArray<FOverlapInfo>& NewPendingOverlaps, const TArray<FOverlapInfo>* NewOverlapsAtEndLocation)
+{
+	AppendOverlapsAfterMove(NewPendingOverlaps, NewOverlapsAtEndLocation);
+}
+
+FORCEINLINE void FScopedMovementUpdate::AppendBlockingHit(const FHitResult& Hit)
+{
+	AppendBlockingHitAfterMove(Hit);
+}
+
+FORCEINLINE bool FScopedMovementUpdate::HasValidOverlapsAtEnd() const
+{
+	// This was the old behavior, but should be re-evaluated by calling code.
+	return HasStoredOverlapsAtEnd();
+}
diff --git a/Engine/Source/Runtime/Engine/Private/Components/PrimitiveComponent.cpp b/Engine/Source/Runtime/Engine/Private/Components/PrimitiveComponent.cpp
--- a/Engine/Source/Runtime/Engine/Private/Components/PrimitiveComponent.cpp
+++ b/Engine/Source/Runtime/Engine/Private/Components/PrimitiveComponent.cpp
@@ -1379,7 +1379,7 @@ bool UPrimitiveComponent::MoveComponentImpl( const FVector& Delta, const FQuat&
 	bool bMoved = false;
 	TArray<FOverlapInfo> PendingOverlaps;
 	TArray<FOverlapInfo> OverlapsAtEndLocation;
-	TArray<FOverlapInfo>* OverlapsAtEndLocationPtr = NULL; // When non-null, used as optimization to avoid work in UpdateOverlaps.
+	const TArray<FOverlapInfo>* OverlapsAtEndLocationPtr = nullptr; // When non-null, used as optimization to avoid work in UpdateOverlaps.
 	AActor* const Actor = GetOwner();
 
 	if ( !bSweep )
@@ -1569,16 +1569,13 @@ bool UPrimitiveComponent::MoveComponentImpl( const FVector& Delta, const FQuat&
 			{
 				// Can only do this if current known overlaps are valid (not deferring updates, or we know there are no new pending overlaps).
 				const FScopedMovementUpdate* ScopedUpdate = GetCurrentScopedMovement();
-				if (ScopedUpdate == nullptr
-					|| (OverlappingComponents.Num() == 0 && !ScopedUpdate->HasPendingOverlapsInAncestorOrSelf())
-					|| (!ScopedUpdate->GetOuterDeferredScope()))
+				if (ScopedUpdate == nullptr || ScopedUpdate->HasKnownOverlapStateAtEnd())
 				{
 					// Only if we know that we won't change our overlap status with children by moving.
 					if (AreSymmetricRotations(NewRotationQuat, ComponentToWorld.GetRotation(), ComponentToWorld.GetScale3D()) &&
 						AreAllCollideableDescendantsRelative())
 					{
-						OverlapsAtEndLocation = OverlappingComponents;
-						OverlapsAtEndLocationPtr = &OverlapsAtEndLocation;
+						OverlapsAtEndLocationPtr = (ScopedUpdate ? ScopedUpdate->GetKnownOverlapsAtEnd(&OverlappingComponents) : &OverlappingComponents);
 					}
 				}
 			}
@@ -1597,7 +1594,7 @@ bool UPrimitiveComponent::MoveComponentImpl( const FVector& Delta, const FQuat&
 			// Defer UpdateOverlaps until the scoped move ends.
 			FScopedMovementUpdate* ScopedUpdate = GetCurrentScopedMovement();
 			checkSlow(ScopedUpdate != NULL);
-			ScopedUpdate->AppendOverlaps(PendingOverlaps, OverlapsAtEndLocationPtr);
+			ScopedUpdate->AppendOverlapsAfterMove(PendingOverlaps, OverlapsAtEndLocationPtr);
 		}
 		else
 		{
@@ -1613,7 +1610,7 @@ bool UPrimitiveComponent::MoveComponentImpl( const FVector& Delta, const FQuat&
 		{
 			FScopedMovementUpdate* ScopedUpdate = GetCurrentScopedMovement();
 			checkSlow(ScopedUpdate != NULL);
-			ScopedUpdate->AppendBlockingHit(BlockingHit);
+			ScopedUpdate->AppendBlockingHitAfterMove(BlockingHit);
 		}
 		else
 		{
@@ -2128,7 +2125,7 @@ void UPrimitiveComponent::ClearMoveIgnoreActors()
 }
 
 
-void UPrimitiveComponent::UpdateOverlaps(TArray<FOverlapInfo> const* PendingOverlaps, bool bDoNotifies, const TArray<FOverlapInfo>* OverlapsAtEndLocation)
+void UPrimitiveComponent::UpdateOverlaps(const TArray<FOverlapInfo>* PendingOverlaps, bool bDoNotifies, const TArray<FOverlapInfo>* OverlapsAtEndLocation)
 {
 	SCOPE_CYCLE_COUNTER(STAT_UpdateOverlaps); 
 
diff --git a/Engine/Source/Runtime/Engine/Private/Components/SceneComponent.cpp b/Engine/Source/Runtime/Engine/Private/Components/SceneComponent.cpp
--- a/Engine/Source/Runtime/Engine/Private/Components/SceneComponent.cpp
+++ b/Engine/Source/Runtime/Engine/Private/Components/SceneComponent.cpp
@@ -483,34 +483,33 @@ void USceneComponent::EndScopedMovementUpdate(class FScopedMovementUpdate& Compl
 	// Process top of the stack
 	FScopedMovementUpdate* CurrentScopedUpdate = ScopedMovementStack.Pop(false);
 	check(CurrentScopedUpdate == &CompletedScope);
-
-	if (CurrentScopedUpdate)
 	{
 		checkSlow(CurrentScopedUpdate->IsDeferringUpdates());
 		if (ScopedMovementStack.Num() == 0)
 		{
 			// This was the last item on the stack, time to apply the updates if necessary
-			const bool bMoved = CurrentScopedUpdate->IsTransformDirty();
-			if (bMoved)
+			const bool bTransformChanged = CurrentScopedUpdate->IsTransformDirty();
+			if (bTransformChanged)
 			{
 				PropagateTransformUpdate(true);
 			}
 
 			// We may have moved somewhere and then moved back to the start, we still need to update overlaps if we touched things along the way.
-			if (bMoved || CurrentScopedUpdate->HasPendingOverlaps() || CurrentScopedUpdate->HasValidOverlapsAtEnd())
+			// If no movement and no change in transform, nothing changed.
+			if (bTransformChanged || CurrentScopedUpdate->bHasMoved)
 			{
 				UpdateOverlaps(&CurrentScopedUpdate->GetPendingOverlaps(), true, CurrentScopedUpdate->GetOverlapsAtEnd());
 			}
 
 			// Dispatch all deferred blocking hits
-			if (CompletedScope.BlockingHits.Num() > 0)
+			if (CurrentScopedUpdate->BlockingHits.Num() > 0)
 			{
 				AActor* const Owner = GetOwner();
 				if (Owner)
 				{
 					// If we have blocking hits, we must be a primitive component.
 					UPrimitiveComponent* PrimitiveThis = CastChecked<UPrimitiveComponent>(this);
-					for (const FHitResult& Hit : CompletedScope.BlockingHits)
+					for (const FHitResult& Hit : CurrentScopedUpdate->BlockingHits)
 					{
 						PrimitiveThis->DispatchBlockingHit(*Owner, Hit);
 					}
@@ -2242,7 +2241,8 @@ FScopedMovementUpdate::FScopedMovementUpdate( class USceneComponent* Component,
 : Owner(Component)
 , OuterDeferredScope(nullptr)
 , bDeferUpdates(ScopeBehavior == EScopedUpdate::DeferredUpdates)
-, bHasValidOverlapsAtEnd(false)
+, bHasStoredOverlapsAtEnd(false)
+, bHasMoved(false)
 {
 	if (IsValid(Component))
 	{
@@ -2305,7 +2305,7 @@ void FScopedMovementUpdate::RevertMove()
 	USceneComponent* Component = Owner;
 	if (IsValid(Component))
 	{
-		bHasValidOverlapsAtEnd = false;
+		bHasStoredOverlapsAtEnd = false;
 		PendingOverlaps.Reset();
 		OverlapsAtEnd.Reset();
 		BlockingHits.Reset();
@@ -2325,20 +2325,24 @@ void FScopedMovementUpdate::RevertMove()
 			}
 		}
 	}
+	bHasMoved = false;
 }
 
 
-void FScopedMovementUpdate::AppendOverlaps(const TArray<struct FOverlapInfo>& OtherOverlaps, const TArray<FOverlapInfo>* OverlapsAtEndLocation)
+void FScopedMovementUpdate::AppendOverlapsAfterMove(const TArray<FOverlapInfo>& NewPendingOverlaps, const TArray<FOverlapInfo>* NewOverlapsAtEndLocation)
 {
-	PendingOverlaps.Append(OtherOverlaps);
-	if (OverlapsAtEndLocation)
+	bHasMoved = true;
+	PendingOverlaps.Append(NewPendingOverlaps);
+	if (NewOverlapsAtEndLocation)
 	{
-		bHasValidOverlapsAtEnd = true;
-		OverlapsAtEnd = *OverlapsAtEndLocation;		
+		bHasStoredOverlapsAtEnd = true;
+		OverlapsAtEnd = *NewOverlapsAtEndLocation;		
 	}
 	else
 	{
-		bHasValidOverlapsAtEnd = false;
+		// null indicates unknown overlap state at the end.
+		bHasStoredOverlapsAtEnd = false;
+		OverlapsAtEnd.Reset();
 	}
 }
 
@@ -2350,12 +2354,58 @@ void FScopedMovementUpdate::OnInnerScopeComplete(const FScopedMovementUpdate& In
 		checkSlow(InnerScope.IsDeferringUpdates());
 		checkSlow(InnerScope.OuterDeferredScope == this);
 
-		// Combine with the next item on the stack
-		AppendOverlaps(InnerScope.GetPendingOverlaps(), InnerScope.GetOverlapsAtEnd());
+		// Combine with the next item on the stack.
+		if (InnerScope.HasMoved(EHasMovedTransformOption::eTestTransform))
+		{
+			bHasMoved = true;
+			AppendOverlapsAfterMove(InnerScope.GetPendingOverlaps(), InnerScope.GetOverlapsAtEnd());
+		}
+		else
+		{
+			// Don't want to invalidate a parent scope when nothing changed in the child.
+		}
+
 		BlockingHits.Append(InnerScope.GetPendingBlockingHits());
 	}	
 }
 
+bool FScopedMovementUpdate::HasKnownOverlapStateAtEnd() const
+{
+	const FScopedMovementUpdate* CurrentScope = this;
+	do
+	{
+		if (CurrentScope->HasMoved(EHasMovedTransformOption::eTestTransform))
+		{
+			return bHasStoredOverlapsAtEnd;
+		}
+
+		CurrentScope = CurrentScope->GetOuterDeferredScope();
+	}
+	while (CurrentScope != nullptr);
+
+	// Top level scope that has not moved. Current state of overlaps is same as component.
+	return true;
+}
+
+const TArray<FOverlapInfo>* FScopedMovementUpdate::GetKnownOverlapsAtEnd(const TArray<FOverlapInfo>* DefaultOverlaps) const
+{
+	const FScopedMovementUpdate* CurrentScope = this;
+	do
+	{
+		if (CurrentScope->HasMoved(EHasMovedTransformOption::eTestTransform))
+		{
+			return CurrentScope->GetOverlapsAtEnd();
+		}
+
+		CurrentScope = CurrentScope->GetOuterDeferredScope();
+	}
+	while (CurrentScope != nullptr);
+
+	// Top level scope that has not moved.
+	return DefaultOverlaps;
+}
+
+
 #if WITH_EDITOR
 const int32 USceneComponent::GetNumUncachedStaticLightingInteractions() const
 {
