Modified pieced of navmesh streaming code that were using un-safe pointers #UE4

UE-15364

diff --git a/Engine/Source/Runtime/Engine/Classes/AI/Navigation/RecastNavMeshDataChunk.h b/Engine/Source/Runtime/Engine/Classes/AI/Navigation/RecastNavMeshDataChunk.h
--- a/Engine/Source/Runtime/Engine/Classes/AI/Navigation/RecastNavMeshDataChunk.h
+++ b/Engine/Source/Runtime/Engine/Classes/AI/Navigation/RecastNavMeshDataChunk.h
@@ -43,10 +43,10 @@ class ENGINE_API URecastNavMeshDataChunk : public UNavigationDataChunk
 	// End UObject Interface
 
 	/** Attaches tiles to specified navmesh, transferring tile ownership to navmesh */
-	TArray<uint32> AttachTiles(FPImplRecastNavMesh* NavMeshImpl);
+	TArray<uint32> AttachTiles(FPImplRecastNavMesh& NavMeshImpl);
 	
 	/** Detaches tiles from specified navmesh, taking tile ownership */
-	TArray<uint32> DetachTiles(FPImplRecastNavMesh* NavMeshImpl);
+	TArray<uint32> DetachTiles(FPImplRecastNavMesh& NavMeshImpl);
 	
 	/** Number of tiles in this chunk */
 	int32 GetNumTiles() const;
@@ -64,6 +64,12 @@ private:
 	void SerializeRecastData(FArchive& Ar, int32 NavMeshVersion);
 #endif//WITH_RECAST
 
+	
+	DEPRECATED(4.8, "AttachTiles is deprecated. Use the version takin a reference instead.")
+	TArray<uint32> AttachTiles(FPImplRecastNavMesh* NavMeshImpl);
+	DEPRECATED(4.8, "DetachTiles is deprecated. Use the version takin a reference instead.")
+	TArray<uint32> DetachTiles(FPImplRecastNavMesh* NavMeshImpl);
+
 private:
 	TArray<FRecastTileData> Tiles;
 };
diff --git a/Engine/Source/Runtime/Engine/Private/AI/Navigation/RecastNavMesh.cpp b/Engine/Source/Runtime/Engine/Private/AI/Navigation/RecastNavMesh.cpp
--- a/Engine/Source/Runtime/Engine/Private/AI/Navigation/RecastNavMesh.cpp
+++ b/Engine/Source/Runtime/Engine/Private/AI/Navigation/RecastNavMesh.cpp
@@ -1655,12 +1655,12 @@ void ARecastNavMesh::OnStreamingLevelAdded(ULevel* InLevel, UWorld* InWorld)
 	
 	bWantsUpdate = true;
 
-	if (SupportsStreaming())
+	if (SupportsStreaming() && RecastNavMeshImpl)
 	{
 		URecastNavMeshDataChunk* NavDataChunk = GetNavigationDataChunk(InLevel);
 		if (NavDataChunk)
 		{
-			TArray<uint32> AttachedIndices = NavDataChunk->AttachTiles(RecastNavMeshImpl);
+			TArray<uint32> AttachedIndices = NavDataChunk->AttachTiles(*RecastNavMeshImpl);
 			if (AttachedIndices.Num() > 0)
 			{
 				InvalidateAffectedPaths(AttachedIndices);
@@ -1674,12 +1674,12 @@ void ARecastNavMesh::OnStreamingLevelRemoved(ULevel* InLevel, UWorld* InWorld)
 {
 	QUICK_SCOPE_CYCLE_COUNTER(STAT_RecastNavMesh_OnStreamingLevelRemoved);
 	
-	if (SupportsStreaming())
+	if (SupportsStreaming() && RecastNavMeshImpl)
 	{
 		URecastNavMeshDataChunk* NavDataChunk = GetNavigationDataChunk(InLevel);
 		if (NavDataChunk)
 		{
-			TArray<uint32> DetachedIndices = NavDataChunk->DetachTiles(RecastNavMeshImpl);
+			TArray<uint32> DetachedIndices = NavDataChunk->DetachTiles(*RecastNavMeshImpl);
 			if (DetachedIndices.Num() > 0)
 			{
 				InvalidateAffectedPaths(DetachedIndices);
diff --git a/Engine/Source/Runtime/Engine/Private/AI/Navigation/RecastNavMeshDataChunk.cpp b/Engine/Source/Runtime/Engine/Private/AI/Navigation/RecastNavMeshDataChunk.cpp
--- a/Engine/Source/Runtime/Engine/Private/AI/Navigation/RecastNavMeshDataChunk.cpp
+++ b/Engine/Source/Runtime/Engine/Private/AI/Navigation/RecastNavMeshDataChunk.cpp
@@ -155,66 +155,70 @@ void URecastNavMeshDataChunk::SerializeRecastData(FArchive& Ar, int32 NavMeshVer
 }
 #endif// WITH_RECAST
 
-TArray<uint32> URecastNavMeshDataChunk::AttachTiles(FPImplRecastNavMesh* NavMeshImpl)
+TArray<uint32> URecastNavMeshDataChunk::AttachTiles(FPImplRecastNavMesh& NavMeshImpl)
 {
 	TArray<uint32> Result;
 	Result.Reserve(Tiles.Num());
 
 #if WITH_RECAST	
-	const bool bIsGame = NavMeshImpl->NavMeshOwner->GetWorld()->IsGameWorld();
-	dtNavMesh* NavMesh = NavMeshImpl->DetourNavMesh;
+	check(NavMeshImpl.NavMeshOwner && NavMeshImpl.NavMeshOwner->GetWorld());
+	const bool bIsGame = NavMeshImpl.NavMeshOwner->GetWorld()->IsGameWorld();
+	dtNavMesh* NavMesh = NavMeshImpl.DetourNavMesh;
 
-	for (FRecastTileData& TileData : Tiles)
+	if (NavMesh != nullptr)
 	{
-		if (TileData.TileRawData.IsValid())
+		for (FRecastTileData& TileData : Tiles)
 		{
-			// Attach mesh tile to target nav mesh 
-			dtStatus status = NavMesh->addTile(TileData.TileRawData->RawData, TileData.TileDataSize, DT_TILE_FREE_DATA, 0, &TileData.TileRef);
-			if (dtStatusFailed(status))
-			{
-				TileData.TileRef = 0;
-				continue;
-			}
-			
-			if (bIsGame)
-			{
-				// We don't own tile data anymore it will be released by recast navmesh 
-				TileData.TileDataSize = 0;
-				TileData.TileRawData->RawData = nullptr;
-			}
-			else
-			{
-				// In the editor we still need to own data, so make a copy of it
-				TileData.TileRawData->RawData = DuplicateRecastRawData(TileData.TileRawData->RawData, TileData.TileDataSize);
-			}
-			
-			// Attach tile cache layer to target nav mesh
-			if (TileData.TileCacheDataSize > 0)
+			if (TileData.TileRawData.IsValid())
 			{
-				const dtMeshTile* MeshTile = NavMesh->getTileByRef(TileData.TileRef);
-				check(MeshTile);
-				int32 TileX = MeshTile->header->x;
-				int32 TileY = MeshTile->header->y;
-				int32 TileLayerIdx = MeshTile->header->layer;
-				FBox TileBBox = Recast2UnrealBox(MeshTile->header->bmin, MeshTile->header->bmax);
-				
-				FNavMeshTileData LayerData(TileData.TileCacheRawData->RawData, TileData.TileCacheDataSize, TileLayerIdx, TileBBox);
-				NavMeshImpl->AddTileCacheLayer(TileX, TileY, TileLayerIdx, LayerData);
+				// Attach mesh tile to target nav mesh 
+				dtStatus status = NavMesh->addTile(TileData.TileRawData->RawData, TileData.TileDataSize, DT_TILE_FREE_DATA, 0, &TileData.TileRef);
+				if (dtStatusFailed(status))
+				{
+					TileData.TileRef = 0;
+					continue;
+				}
 
 				if (bIsGame)
 				{
-					// We don't own tile cache data anymore it will be released by navmesh
-					TileData.TileCacheDataSize = 0;
-					TileData.TileCacheRawData->RawData = nullptr;
+					// We don't own tile data anymore it will be released by recast navmesh 
+					TileData.TileDataSize = 0;
+					TileData.TileRawData->RawData = nullptr;
 				}
 				else
 				{
 					// In the editor we still need to own data, so make a copy of it
-					TileData.TileCacheRawData->RawData = DuplicateRecastRawData(TileData.TileCacheRawData->RawData, TileData.TileCacheDataSize);
+					TileData.TileRawData->RawData = DuplicateRecastRawData(TileData.TileRawData->RawData, TileData.TileDataSize);
 				}
+
+				// Attach tile cache layer to target nav mesh
+				if (TileData.TileCacheDataSize > 0)
+				{
+					const dtMeshTile* MeshTile = NavMesh->getTileByRef(TileData.TileRef);
+					check(MeshTile);
+					int32 TileX = MeshTile->header->x;
+					int32 TileY = MeshTile->header->y;
+					int32 TileLayerIdx = MeshTile->header->layer;
+					FBox TileBBox = Recast2UnrealBox(MeshTile->header->bmin, MeshTile->header->bmax);
+
+					FNavMeshTileData LayerData(TileData.TileCacheRawData->RawData, TileData.TileCacheDataSize, TileLayerIdx, TileBBox);
+					NavMeshImpl.AddTileCacheLayer(TileX, TileY, TileLayerIdx, LayerData);
+
+					if (bIsGame)
+					{
+						// We don't own tile cache data anymore it will be released by navmesh
+						TileData.TileCacheDataSize = 0;
+						TileData.TileCacheRawData->RawData = nullptr;
+					}
+					else
+					{
+						// In the editor we still need to own data, so make a copy of it
+						TileData.TileCacheRawData->RawData = DuplicateRecastRawData(TileData.TileCacheRawData->RawData, TileData.TileCacheDataSize);
+					}
+				}
+
+				Result.Add(NavMesh->decodePolyIdTile(TileData.TileRef));
 			}
-			
-			Result.Add(NavMesh->decodePolyIdTile(TileData.TileRef));
 		}
 	}
 #endif// WITH_RECAST
@@ -223,48 +227,52 @@ TArray<uint32> URecastNavMeshDataChunk::AttachTiles(FPImplRecastNavMesh* NavMesh
 	return Result;
 }
 
-TArray<uint32> URecastNavMeshDataChunk::DetachTiles(FPImplRecastNavMesh* NavMeshImpl)
+TArray<uint32> URecastNavMeshDataChunk::DetachTiles(FPImplRecastNavMesh& NavMeshImpl)
 {
 	TArray<uint32> Result;
 	Result.Reserve(Tiles.Num());
 
-#if WITH_RECAST	
-	const bool bIsGame = NavMeshImpl->NavMeshOwner->GetWorld()->IsGameWorld();
-	dtNavMesh* NavMesh = NavMeshImpl->DetourNavMesh;
+#if WITH_RECAST
+	check(NavMeshImpl.NavMeshOwner && NavMeshImpl.NavMeshOwner->GetWorld());
+	const bool bIsGame = NavMeshImpl.NavMeshOwner->GetWorld()->IsGameWorld();
+	dtNavMesh* NavMesh = NavMeshImpl.DetourNavMesh;
 
-	for (FRecastTileData& TileData : Tiles)
+	if (NavMesh != nullptr)
 	{
-		if (TileData.TileRef != 0)
+		for (FRecastTileData& TileData : Tiles)
 		{
-			// Detach tile cache layer and take ownership over compressed data
-			const dtMeshTile* MeshTile = NavMesh->getTileByRef(TileData.TileRef);
-			if (MeshTile)
+			if (TileData.TileRef != 0)
 			{
-				if (bIsGame)
+				// Detach tile cache layer and take ownership over compressed data
+				const dtMeshTile* MeshTile = NavMesh->getTileByRef(TileData.TileRef);
+				if (MeshTile)
 				{
-					FNavMeshTileData TileCacheData = NavMeshImpl->GetTileCacheLayer(MeshTile->header->x, MeshTile->header->y, MeshTile->header->layer);
-					if (TileCacheData.IsValid())
+					if (bIsGame)
 					{
-						TileData.TileCacheDataSize = TileCacheData.DataSize;
-						TileData.TileCacheRawData->RawData = TileCacheData.Release();
+						FNavMeshTileData TileCacheData = NavMeshImpl.GetTileCacheLayer(MeshTile->header->x, MeshTile->header->y, MeshTile->header->layer);
+						if (TileCacheData.IsValid())
+						{
+							TileData.TileCacheDataSize = TileCacheData.DataSize;
+							TileData.TileCacheRawData->RawData = TileCacheData.Release();
+						}
 					}
-				}
 				
-				NavMeshImpl->RemoveTileCacheLayer(MeshTile->header->x, MeshTile->header->y, MeshTile->header->layer);
-			}
+					NavMeshImpl.RemoveTileCacheLayer(MeshTile->header->x, MeshTile->header->y, MeshTile->header->layer);
+				}
 						
-			if (bIsGame)
-			{
-				// Remove tile from navmesh and take ownership of tile raw data
-				NavMesh->removeTile(TileData.TileRef, &TileData.TileRawData->RawData, &TileData.TileDataSize);
-			}
-			else
-			{
-				// In the editor we have a copy of tile data so just release tile in navmesh
-				NavMesh->removeTile(TileData.TileRef, nullptr, nullptr);
-			}
+				if (bIsGame)
+				{
+					// Remove tile from navmesh and take ownership of tile raw data
+					NavMesh->removeTile(TileData.TileRef, &TileData.TileRawData->RawData, &TileData.TileDataSize);
+				}
+				else
+				{
+					// In the editor we have a copy of tile data so just release tile in navmesh
+					NavMesh->removeTile(TileData.TileRef, nullptr, nullptr);
+				}
 						
-			Result.Add(NavMesh->decodePolyIdTile(TileData.TileRef));
+				Result.Add(NavMesh->decodePolyIdTile(TileData.TileRef));
+			}
 		}
 	}
 #endif// WITH_RECAST
@@ -316,4 +324,27 @@ void URecastNavMeshDataChunk::GatherTiles(const FPImplRecastNavMesh* NavMeshImpl
 		}
 	}
 }
-#endif// WITH_EDITOR
\ No newline at end of file
+#endif// WITH_EDITOR
+
+//----------------------------------------------------------------------//
+// DEPRECATED
+//----------------------------------------------------------------------//
+TArray<uint32> URecastNavMeshDataChunk::AttachTiles(FPImplRecastNavMesh* NavMeshImpl)
+{
+	if (NavMeshImpl)
+	{
+		return AttachTiles(*NavMeshImpl);
+	}
+
+	return TArray<uint32>();
+}
+
+TArray<uint32> URecastNavMeshDataChunk::DetachTiles(FPImplRecastNavMesh* NavMeshImpl)
+{
+	if (NavMeshImpl)
+	{
+		return DetachTiles(*NavMeshImpl);
+	}
+
+	return TArray<uint32>();
+}
