BP: Logging a warning instead of throwing an assertion if DaysInMonth() has invalid parameters

diff --git a/Engine/Source/Runtime/Engine/Private/KismetMathLibrary.cpp b/Engine/Source/Runtime/Engine/Private/KismetMathLibrary.cpp
--- a/Engine/Source/Runtime/Engine/Private/KismetMathLibrary.cpp
+++ b/Engine/Source/Runtime/Engine/Private/KismetMathLibrary.cpp
@@ -1308,6 +1308,13 @@ bool UKismetMathLibrary::IsMorning( FDateTime A )
 
 int32 UKismetMathLibrary::DaysInMonth( int32 Year, int32 Month )
 {
+	if ((Month < 1) || (Month > 12))
+	{
+		//@TODO: EXCEPTION: Throw script exception
+		FFrame::KismetExecutionMessage(TEXT("Invalid month (must be between 1 and 12): DaysInMonth"), ELogVerbosity::Warning);
+		return 0;
+	}
+
 	return FDateTime::DaysInMonth(Year, Month);
 }
 
