Add memory validation for physx allocator

diff --git a/Engine/Source/Runtime/Engine/Private/PhysicsEngine/PhysLevel.cpp b/Engine/Source/Runtime/Engine/Private/PhysicsEngine/PhysLevel.cpp
--- a/Engine/Source/Runtime/Engine/Private/PhysicsEngine/PhysLevel.cpp
+++ b/Engine/Source/Runtime/Engine/Private/PhysicsEngine/PhysLevel.cpp
@@ -257,6 +257,15 @@ void FEndPhysicsTickFunction::ExecuteTick(float DeltaTime, enum ELevelTick TickT
 		// it was already done, so let just do it.
 		Target->FinishPhysicsSim();
 	}
+
+#if PHYSX_MEMORY_VALIDATION
+	static int32 Frequency = 0;
+	if (Frequency++ > 10)
+	{
+		Frequency = 0;
+		GPhysXAllocator->ValidateHeaders();
+	}
+#endif
 }
 
 FString FEndPhysicsTickFunction::DiagnosticMessage()
diff --git a/Engine/Source/Runtime/Engine/Private/PhysicsEngine/PhysXSupport.h b/Engine/Source/Runtime/Engine/Private/PhysicsEngine/PhysXSupport.h
--- a/Engine/Source/Runtime/Engine/Private/PhysicsEngine/PhysXSupport.h
+++ b/Engine/Source/Runtime/Engine/Private/PhysicsEngine/PhysXSupport.h
@@ -16,9 +16,14 @@
 #define USE_SCENE_LOCK			1
 #endif
 
+// Whether to track PhysX memory allocations
+#ifndef PHYSX_MEMORY_VALIDATION
+#define PHYSX_MEMORY_VALIDATION		1
+#endif
+
 // Whether to track PhysX memory allocations
 #ifndef PHYSX_MEMORY_STATS
-#define PHYSX_MEMORY_STATS		0
+#define PHYSX_MEMORY_STATS		0 || PHYSX_MEMORY_VALIDATION
 #endif
 
 // binary serialization requires 128 byte alignment
@@ -503,24 +508,80 @@ class FPhysXAllocator : public PxAllocatorCallback
 		FPhysXAllocationHeader(	FName InAllocationTypeName, size_t InAllocationSize )
 		:	AllocationTypeName(InAllocationTypeName)
 		,	AllocationSize(InAllocationSize)
-		{}
+		{
+			static_assert(sizeof(FPhysXAllocationHeader) == 32, "FPhysXAllocationHeader size must be 32 bytes.");
+			MagicPadding();
+		}
+
+		void MagicPadding()
+		{
+			for (uint8 ByteCount = 0; ByteCount < sizeof(Padding); ++ByteCount)
+			{
+				Padding[ByteCount] = 'A' + ByteCount % 4;
+			}
+		}
+
+		bool operator==(const FPhysXAllocationHeader& OtherHeader) const
+		{
+			bool bHeaderSame = AllocationTypeName == OtherHeader.AllocationTypeName && AllocationSize == OtherHeader.AllocationSize;
+			for (uint8 ByteCount = 0; ByteCount < sizeof(Padding); ++ByteCount)
+			{
+				bHeaderSame &= Padding[ByteCount] == OtherHeader.Padding[ByteCount];
+			}
+
+			return bHeaderSame;
+		}
 
 		FName AllocationTypeName;
 		size_t	AllocationSize;
+		uint8 Padding[8];	//physx needs 16 byte alignment. Additionally we fill padding with a pattern to see if there's any memory stomps
+
+		void Validate() const
+		{
+			bool bValid = true;
+			for (uint8 ByteCount = 0; ByteCount < sizeof(Padding); ++ByteCount)
+			{
+				bValid &= Padding[ByteCount] == 'A' + ByteCount % 4;
+			}
+
+			check(bValid);
+
+			FPhysXAllocationHeader* AllocationFooter = (FPhysXAllocationHeader*) (((uint8*)this) + sizeof(FPhysXAllocationHeader)+AllocationSize);
+			check(*AllocationFooter == *this);
+		}
 	};
 	
 #endif
 
 public:
-	FPhysXAllocator()
-#if PHYSX_MEMORY_STATS
-	: HeaderSize(32) //physx needs 16 byte alignment. We need to inflate the header size to account for this.
-#endif
+
+#if PHYSX_MEMORY_VALIDATION
+
+	/** Iterates over all allocations and checks that they the headers and footers are valid */
+	void ValidateHeaders()
 	{
-#if PHYSX_MEMORY_STATS
-		static_assert(sizeof(FPhysXAllocationHeader) <= 32, "FPhysXAllocationHeader size must be less than HeaderSize bytes.");
-#endif
+		check(IsInGameThread());
+		FPhysXAllocationHeader* TmpHeader = nullptr;
+		while(NewHeaders.Dequeue(TmpHeader))
+		{
+			AllocatedHeaders.Add(TmpHeader);
+		}
+
+		while (OldHeaders.Dequeue(TmpHeader))
+		{
+			AllocatedHeaders.Remove(TmpHeader);
+		}
+
+		FScopeLock Lock(&ValidationCS);	//this is needed in case another thread is freeing the header
+		for (FPhysXAllocationHeader* Header : AllocatedHeaders)
+		{
+			Header->Validate();
+		}
 	}
+#endif
+
+	FPhysXAllocator()
+	{}
 
 	virtual ~FPhysXAllocator() 
 	{}
@@ -528,19 +589,23 @@ public:
 	virtual void* allocate(size_t size, const char* typeName, const char* filename, int line) override
 	{
 #if PHYSX_MEMORY_STATS
-
 		INC_DWORD_STAT_BY(STAT_MemoryPhysXTotalAllocationSize, size);
 
+
 		FString AllocationString = FString::Printf(TEXT("%s %s:%d"), ANSI_TO_TCHAR(typeName), ANSI_TO_TCHAR(filename), line);
 		FName AllocationName(*AllocationString);
 
 		// Assign header
-		FPhysXAllocationHeader* AllocationHeader = (FPhysXAllocationHeader*)FMemory::Malloc(size + HeaderSize, 16);
+		FPhysXAllocationHeader* AllocationHeader = (FPhysXAllocationHeader*)FMemory::Malloc(size + sizeof(FPhysXAllocationHeader) * 2, 16);
 		AllocationHeader->AllocationTypeName = AllocationName;
 		AllocationHeader->AllocationSize = size;
+		AllocationHeader->MagicPadding();
+		FPhysXAllocationHeader* AllocationFooter = (FPhysXAllocationHeader*) (((uint8*)AllocationHeader) + size + sizeof(FPhysXAllocationHeader));
+		AllocationFooter->AllocationTypeName = AllocationName;
+		AllocationFooter->AllocationSize = size;
+		AllocationFooter->MagicPadding();
 
-		// Assign map to track by type
-		size_t* TotalByType = AllocationsByType.Find(AllocationName);
+		size_t* TotalByType = AllocationsByType.Find(AllocationName);	//TODO: this is not thread safe!
 		if( TotalByType )
 		{
 			*TotalByType += size;
@@ -550,7 +615,11 @@ public:
 			AllocationsByType.Add(AllocationName, size);
 		}
 
-		return (uint8*)AllocationHeader + HeaderSize;
+#if PHYSX_MEMORY_VALIDATION
+		NewHeaders.Enqueue(AllocationHeader);
+#endif
+
+		return (uint8*)AllocationHeader + sizeof(FPhysXAllocationHeader);
 #else
 		void* ptr = FMemory::Malloc(size, 16);
 		#if PHYSX_MEMORY_STAT_ONLY
@@ -565,7 +634,13 @@ public:
 #if PHYSX_MEMORY_STATS
 		if( ptr )
 		{
-			FPhysXAllocationHeader* AllocationHeader = (FPhysXAllocationHeader*)((uint8*)ptr - HeaderSize);
+			FPhysXAllocationHeader* AllocationHeader = (FPhysXAllocationHeader*)((uint8*)ptr - sizeof(FPhysXAllocationHeader));
+#if PHYSX_MEMORY_VALIDATION
+			AllocationHeader->Validate();
+			OldHeaders.Enqueue(AllocationHeader);
+			FScopeLock Lock(&ValidationCS);	//this is needed in case we are in the middle of validating the headers
+#endif
+
 			DEC_DWORD_STAT_BY(STAT_MemoryPhysXTotalAllocationSize, AllocationHeader->AllocationSize);
 			size_t* TotalByType = AllocationsByType.Find(AllocationHeader->AllocationTypeName);
 			*TotalByType -= AllocationHeader->AllocationSize;
@@ -599,9 +674,15 @@ public:
 	}
 #endif
 
+#if PHYSX_MEMORY_VALIDATION
 private:
-#if PHYSX_MEMORY_STATS
-	const uint32 HeaderSize;
+	FCriticalSection ValidationCS;
+	TSet<FPhysXAllocationHeader*> AllocatedHeaders;
+
+	//Since this needs to be thread safe we can't add to the allocated headers set until we're on the game thread
+	TQueue<FPhysXAllocationHeader*, EQueueMode::Mpsc> NewHeaders;
+	TQueue<FPhysXAllocationHeader*, EQueueMode::Mpsc> OldHeaders;
+
 #endif
 };
 
diff --git a/Engine/Source/Runtime/PhysXFormats/Private/PhysXFormats.cpp b/Engine/Source/Runtime/PhysXFormats/Private/PhysXFormats.cpp
--- a/Engine/Source/Runtime/PhysXFormats/Private/PhysXFormats.cpp
+++ b/Engine/Source/Runtime/PhysXFormats/Private/PhysXFormats.cpp
@@ -260,10 +260,17 @@ public:
 		FPhysXOutputStream Buffer(&OutBuffer);
 		PxSerialization::complete(*PCollection, *PRegistry, PExceptFor);
 		PxSerialization::serializeCollectionToBinary(Buffer, *PCollection, *PRegistry, PExceptFor);
+
+#if PHYSX_MEMORY_VALIDATION
+		GPhysXAllocator->ValidateHeaders();
+#endif
 		PCollection->release();
 		PExceptFor->release();
 		PRegistry->release();
-		
+
+#if PHYSX_MEMORY_VALIDATION
+		GPhysXAllocator->ValidateHeaders();
+#endif
 		return true;
 #endif
 		return false;
