#ue4 - Fix character mesh relative translation being applied incorrectly in UCharacterMovementComponent::SmoothClientPosition().

diff --git a/Engine/Source/Runtime/Engine/Private/Components/CharacterMovementComponent.cpp b/Engine/Source/Runtime/Engine/Private/Components/CharacterMovementComponent.cpp
--- a/Engine/Source/Runtime/Engine/Private/Components/CharacterMovementComponent.cpp
+++ b/Engine/Source/Runtime/Engine/Private/Components/CharacterMovementComponent.cpp
@@ -5646,7 +5646,7 @@ void UCharacterMovementComponent::SmoothClientPosition(float DeltaSeconds)
 			ClientData->MeshTranslationOffset.Z = 0;
 		}
 
-		const FVector NewRelTranslation = CharacterOwner->ActorToWorld().InverseTransformVectorNoScale(ClientData->MeshTranslationOffset + CharacterOwner->GetBaseTranslationOffset());
+		const FVector NewRelTranslation = CharacterOwner->ActorToWorld().InverseTransformVectorNoScale(ClientData->MeshTranslationOffset) + CharacterOwner->GetBaseTranslationOffset();
 		const FQuat NewRelRotation = ClientData->MeshRotationOffset * CharacterOwner->GetBaseRotationOffset();
 		CharacterOwner->GetMesh()->SetRelativeLocationAndRotation(NewRelTranslation, NewRelRotation);
 	}
