Remove unnecessary pragma disables

diff --git a/Engine/Source/Runtime/Engine/Private/UserInterface/PlayerInput.cpp b/Engine/Source/Runtime/Engine/Private/UserInterface/PlayerInput.cpp
--- a/Engine/Source/Runtime/Engine/Private/UserInterface/PlayerInput.cpp
+++ b/Engine/Source/Runtime/Engine/Private/UserInterface/PlayerInput.cpp
@@ -357,9 +357,7 @@ void UPlayerInput::SetMouseSensitivityToDefault()
 		const FKey AxisKey = AxisConfigEntry.AxisKeyName;
 		if (AxisKey == EKeys::MouseX)
 		{
-			PRAGMA_DISABLE_DEPRECATION_WARNINGS
 			SetMouseSensitivity(AxisConfigEntry.AxisProperties.Sensitivity);
-			PRAGMA_ENABLE_DEPRECATION_WARNINGS
 			break;
 		}
 	}
