UE-14514 - Fix issues where component default subobjects would not be flagged RF_Public causing GLEO errors.

diff --git a/Engine/Source/Editor/Kismet/Private/SSCSEditor.cpp b/Engine/Source/Editor/Kismet/Private/SSCSEditor.cpp
--- a/Engine/Source/Editor/Kismet/Private/SSCSEditor.cpp
+++ b/Engine/Source/Editor/Kismet/Private/SSCSEditor.cpp
@@ -4873,8 +4873,6 @@ void SSCSEditor::PasteNodes()
 			// Get the component object instance
 			UActorComponent* NewActorComponent = NewObjectIt->Value;
 			check(NewActorComponent);
-			// make sure that our flags match those set in USimpleConstructionScript::CreateNode:
-			NewActorComponent->SetFlags(RF_ArchetypeObject | RF_Transactional | RF_Public);
 
 			// Create a new SCS node to contain the new component and add it to the tree
 			NewActorComponent = AddNewNode(Blueprint->SimpleConstructionScript->CreateNodeAndRenameComponent(NewActorComponent), NULL, false, false);
diff --git a/Engine/Source/Editor/UnrealEd/Private/Kismet2/ComponentEditorUtils.cpp b/Engine/Source/Editor/UnrealEd/Private/Kismet2/ComponentEditorUtils.cpp
--- a/Engine/Source/Editor/UnrealEd/Private/Kismet2/ComponentEditorUtils.cpp
+++ b/Engine/Source/Editor/UnrealEd/Private/Kismet2/ComponentEditorUtils.cpp
@@ -63,7 +63,7 @@ struct FComponentObjectTextFactory : public FCustomizableTextObjectFactory
 			EObjectFlags ObjectFlags = RF_Transactional;
 			if (bPasteAsArchetypes)
 			{
-				ObjectFlags |= RF_ArchetypeObject;
+				ObjectFlags |= RF_ArchetypeObject | RF_Public;
 			}
 
 			// Use the transient package initially for creating the objects, since the variable name is used when copying
diff --git a/Engine/Source/Runtime/Engine/Private/SimpleConstructionScript.cpp b/Engine/Source/Runtime/Engine/Private/SimpleConstructionScript.cpp
--- a/Engine/Source/Runtime/Engine/Private/SimpleConstructionScript.cpp
+++ b/Engine/Source/Runtime/Engine/Private/SimpleConstructionScript.cpp
@@ -1031,8 +1031,7 @@ USCS_Node* USimpleConstructionScript::CreateNode(UClass* NewComponentClass, FNam
 	// note that naming logic is duplicated in CreateNodeAndRenameComponent:
 	NewComponentVariableName = GenerateNewComponentName(NewComponentClass, NewComponentVariableName);
 
-	auto NewComponentTemplate = NewObject<UActorComponent>(Blueprint->GeneratedClass, NewComponentClass, *(NewComponentVariableName.GetPlainNameString() + FGuid::NewGuid().ToString() ) );
-	NewComponentTemplate->SetFlags(RF_ArchetypeObject|RF_Transactional|RF_Public);
+	UActorComponent* NewComponentTemplate = NewObject<UActorComponent>(Blueprint->GeneratedClass, NewComponentClass, *(NewComponentVariableName.GetPlainNameString() + FGuid::NewGuid().ToString() ), RF_ArchetypeObject|RF_Transactional|RF_Public);
 
 	return CreateNodeImpl(NewComponentTemplate, NewComponentVariableName);
 }
