UE-15437 Fix for crash when trying to load a streaming wavesound on a dedicated server.

- Dedicated servers shouldn't be trying to load derived data for audio streams, so we won't add streaming sounds to the DD if running a dedicated server.
- Also added a check on the pointer at the point of the crash in the callstack and some error-logging.

diff --git a/Engine/Source/Runtime/Engine/Private/AudioDerivedData.cpp b/Engine/Source/Runtime/Engine/Private/AudioDerivedData.cpp
--- a/Engine/Source/Runtime/Engine/Private/AudioDerivedData.cpp
+++ b/Engine/Source/Runtime/Engine/Private/AudioDerivedData.cpp
@@ -243,40 +243,50 @@ class FStreamedAudioCacheDerivedDataWorker : public FNonAbandonableTask
 			DerivedData->AudioFormat = AudioFormatName;
 
 			FByteBulkData* CompressedData = SoundWave.GetCompressedData(AudioFormatName);
-			TArray<uint8> CompressedBuffer;
-			CompressedBuffer.Empty(CompressedData->GetBulkDataSize());
-			CompressedBuffer.AddUninitialized(CompressedData->GetBulkDataSize());
-			void* BufferData = CompressedBuffer.GetData();
-			CompressedData->GetCopy(&BufferData, false);
-			TArray<TArray<uint8>> ChunkBuffers;
-
-			if (AudioFormat->SplitDataForStreaming(CompressedBuffer, ChunkBuffers))
+			if (CompressedData)
 			{
-				for (int32 ChunkIndex = 0; ChunkIndex < ChunkBuffers.Num(); ++ChunkIndex)
+				TArray<uint8> CompressedBuffer;
+				CompressedBuffer.Empty(CompressedData->GetBulkDataSize());
+				CompressedBuffer.AddUninitialized(CompressedData->GetBulkDataSize());
+				void* BufferData = CompressedBuffer.GetData();
+				CompressedData->GetCopy(&BufferData, false);
+				TArray<TArray<uint8>> ChunkBuffers;
+
+				if (AudioFormat->SplitDataForStreaming(CompressedBuffer, ChunkBuffers))
 				{
+					for (int32 ChunkIndex = 0; ChunkIndex < ChunkBuffers.Num(); ++ChunkIndex)
+					{
+						FStreamedAudioChunk* NewChunk = new(DerivedData->Chunks) FStreamedAudioChunk();
+						NewChunk->DataSize = ChunkBuffers[ChunkIndex].Num();
+						NewChunk->BulkData.Lock(LOCK_READ_WRITE);
+						void* NewChunkData = NewChunk->BulkData.Realloc(ChunkBuffers[ChunkIndex].Num());
+						FMemory::Memcpy(NewChunkData, ChunkBuffers[ChunkIndex].GetData(), ChunkBuffers[ChunkIndex].Num());
+						NewChunk->BulkData.Unlock();
+					}
+				}
+				else
+				{
+					// Could not split so copy compressed data into a single chunk
 					FStreamedAudioChunk* NewChunk = new(DerivedData->Chunks) FStreamedAudioChunk();
-					NewChunk->DataSize = ChunkBuffers[ChunkIndex].Num();
+					NewChunk->DataSize = CompressedBuffer.Num();
 					NewChunk->BulkData.Lock(LOCK_READ_WRITE);
-					void* NewChunkData = NewChunk->BulkData.Realloc(ChunkBuffers[ChunkIndex].Num());
-					FMemory::Memcpy(NewChunkData, ChunkBuffers[ChunkIndex].GetData(), ChunkBuffers[ChunkIndex].Num());
+					void* NewChunkData = NewChunk->BulkData.Realloc(CompressedBuffer.Num());
+					FMemory::Memcpy(NewChunkData, CompressedBuffer.GetData(), CompressedBuffer.Num());
 					NewChunk->BulkData.Unlock();
 				}
+
+				DerivedData->NumChunks = DerivedData->Chunks.Num();
+
+				// Store it in the cache.
+				PutDerivedDataInCache(DerivedData, KeySuffix);
 			}
 			else
 			{
-				// Could not split so copy compressed data into a single chunk
-				FStreamedAudioChunk* NewChunk = new(DerivedData->Chunks) FStreamedAudioChunk();
-				NewChunk->DataSize = CompressedBuffer.Num();
-				NewChunk->BulkData.Lock(LOCK_READ_WRITE);
-				void* NewChunkData = NewChunk->BulkData.Realloc(CompressedBuffer.Num());
-				FMemory::Memcpy(NewChunkData, CompressedBuffer.GetData(), CompressedBuffer.Num());
-				NewChunk->BulkData.Unlock();
+				UE_LOG(LogAudio, Warning, TEXT("Failed to retrieve compressed data for format %s and soundwave %s"),
+					   *AudioFormatName.GetPlainNameString(),
+					   *SoundWave.GetPathName()
+					);
 			}
-
-			DerivedData->NumChunks = DerivedData->Chunks.Num();
-
-			// Store it in the cache.
-			PutDerivedDataInCache(DerivedData, KeySuffix);
 		}
 		
 		if (DerivedData->Chunks.Num())
diff --git a/Engine/Source/Runtime/Engine/Private/SoundWave.cpp b/Engine/Source/Runtime/Engine/Private/SoundWave.cpp
--- a/Engine/Source/Runtime/Engine/Private/SoundWave.cpp
+++ b/Engine/Source/Runtime/Engine/Private/SoundWave.cpp
@@ -340,7 +340,8 @@ void USoundWave::PostLoad()
 		}
 	}
 
-	if (IsStreaming())
+	// Only add this streaming sound if we're not a dedicated server or if there is an audio device manager
+	if (IsStreaming() && !IsRunningDedicatedServer() && GEngine && GEngine->GetAudioDeviceManager())
 	{
 #if WITH_EDITORONLY_DATA
 		FinishCachePlatformData();
