Added guard against crash due to dereferencing a NULL sound asset.

diff --git a/Engine/Source/Runtime/Engine/Private/Interpolation.cpp b/Engine/Source/Runtime/Engine/Private/Interpolation.cpp
--- a/Engine/Source/Runtime/Engine/Private/Interpolation.cpp
+++ b/Engine/Source/Runtime/Engine/Private/Interpolation.cpp
@@ -7847,7 +7847,7 @@ void UInterpTrackSound::UpdateTrack(float NewPosition, UInterpTrackInst* TrInst,
 		// Check if we're in the audio range, and if we need to start playing the audio,
 		// either because it has never been played, or isn't currently playing.
 		// We only do this when we've jumped position.
-		bool bIsInRangeAndNeedsStart = !bPlaying && NewPosition >= SoundTrackKey.Time && NewPosition <= ( SoundTrackKey.Time + SoundTrackKey.Sound->Duration );
+		bool bIsInRangeAndNeedsStart = !bPlaying && SoundTrackKey.Sound != nullptr && NewPosition >= SoundTrackKey.Time && NewPosition <= ( SoundTrackKey.Time + SoundTrackKey.Sound->Duration );
 		if ( bIsInRangeAndNeedsStart )
 		{
 			bIsInRangeAndNeedsStart = SoundInst->PlayAudioComp == NULL || !SoundInst->PlayAudioComp->IsPlaying();
