Engine: Added a vaildation check before the use of SetNumUninitialized to InstancedStaticMeshComponent

diff --git a/Engine/Source/Runtime/Engine/Private/InstancedStaticMesh.cpp b/Engine/Source/Runtime/Engine/Private/InstancedStaticMesh.cpp
--- a/Engine/Source/Runtime/Engine/Private/InstancedStaticMesh.cpp
+++ b/Engine/Source/Runtime/Engine/Private/InstancedStaticMesh.cpp
@@ -988,7 +988,8 @@ void UInstancedStaticMeshComponent::CreateAllInstanceBodies()
 		FPhysScene* PhysScene = GetWorld()->GetPhysicsScene();
 
 	    const int32 NumBodies = PerInstanceSMData.Num();
-	    InstanceBodies.SetNumUninitialized(NumBodies);
+		check(InstanceBodies.Num() == 0);
+		InstanceBodies.SetNumUninitialized(NumBodies);
     
 	    TArray<FTransform> Transforms;
 	    Transforms.Reserve(NumBodies);
