Let memory stats track allocated persistent frames. - STAT_PersistentUberGraphFrameMemory - Very simple version, that doesn't distinguish owners.

diff --git a/Engine/Source/Runtime/Engine/Classes/Engine/BlueprintGeneratedClass.h b/Engine/Source/Runtime/Engine/Classes/Engine/BlueprintGeneratedClass.h
--- a/Engine/Source/Runtime/Engine/Classes/Engine/BlueprintGeneratedClass.h
+++ b/Engine/Source/Runtime/Engine/Classes/Engine/BlueprintGeneratedClass.h
@@ -5,6 +5,7 @@
 #include "MeshBatch.h"
 #include "BlueprintGeneratedClass.generated.h"
 
+DECLARE_MEMORY_STAT_EXTERN(TEXT("Persistent Uber Graph Frame memory"), STAT_PersistentUberGraphFrameMemory, STATGROUP_Memory, );
 
 class UEdGraphPin;
 
diff --git a/Engine/Source/Runtime/Engine/Private/BlueprintGeneratedClass.cpp b/Engine/Source/Runtime/Engine/Private/BlueprintGeneratedClass.cpp
--- a/Engine/Source/Runtime/Engine/Private/BlueprintGeneratedClass.cpp
+++ b/Engine/Source/Runtime/Engine/Private/BlueprintGeneratedClass.cpp
@@ -15,6 +15,7 @@
 #include "KismetEditorUtilities.h"
 #endif //WITH_EDITOR
 
+DEFINE_STAT(STAT_PersistentUberGraphFrameMemory);
 
 UBlueprintGeneratedClass::UBlueprintGeneratedClass(const FObjectInitializer& ObjectInitializer)
 	: Super(ObjectInitializer)
@@ -600,7 +601,9 @@ void UBlueprintGeneratedClass::CreatePersistentUberGraphFrame(UObject* Obj, bool
 			const bool bUberGraphFunctionIsReady = UberGraphFunction->HasAllFlags(RF_LoadCompleted); // is fully loaded
 			if (bUberGraphFunctionIsReady)
 			{
+				INC_MEMORY_STAT_BY(STAT_PersistentUberGraphFrameMemory, UberGraphFunction->GetStructureSize());
 				FrameMemory = (uint8*)FMemory::Malloc(UberGraphFunction->GetStructureSize());
+
 				FMemory::Memzero(FrameMemory, UberGraphFunction->GetStructureSize());
 				for (UProperty* Property = UberGraphFunction->PropertyLink; Property; Property = Property->PropertyLinkNext)
 				{
@@ -637,6 +640,7 @@ void UBlueprintGeneratedClass::DestroyPersistentUberGraphFrame(UObject* Obj) con
 				Property->DestroyValue_InContainer(FrameMemory);
 			}
 			FMemory::Free(FrameMemory);
+			DEC_MEMORY_STAT_BY(STAT_PersistentUberGraphFrameMemory, UberGraphFunction->GetStructureSize());
 		}
 		else
 		{
