Fixed ChildActorComponent scaling regression (UE-14966). Also fixes issue where scaling the parent actor would not scale the children individually. Added clearer display name for EAttachLocation::SnapToTarget ("Snap to Target, Keep World Scale").

diff --git a/Engine/Source/Runtime/Engine/Classes/Engine/EngineTypes.h b/Engine/Source/Runtime/Engine/Classes/Engine/EngineTypes.h
--- a/Engine/Source/Runtime/Engine/Classes/Engine/EngineTypes.h
+++ b/Engine/Source/Runtime/Engine/Classes/Engine/EngineTypes.h
@@ -42,7 +42,7 @@ namespace EBrowseReturnVal
 {
 	enum Type
 	{
-		Success,		// Sucessfully browsed to a new map
+		Success,		// Successfully browsed to a new map
 		Failure,		// Immediately failed to browse
 		Pending,		// A connection is pending
 	};
@@ -54,9 +54,17 @@ namespace EAttachLocation
 {
 	enum Type
 	{
+		/** Keeps current relative transform as the relative transform to the new parent. */
 		KeepRelativeOffset,
+		
+		/** Automatically calculates the relative transform such that the attached component maintains the same world transform. */
 		KeepWorldPosition,
-		SnapToTarget
+
+		/** Snaps location and rotation to the attach point. Calculates the relative scale so that the final world scale of the component remains the same. */
+		SnapToTarget					UMETA(DisplayName = "Snap to Target, Keep World Scale"),
+		
+		/** Snaps entire transform to target, including scale. */
+		SnapToTargetIncludingScale		UMETA(DisplayName = "Snap to Target, Including Scale"),
 	};
 }
 
diff --git a/Engine/Source/Runtime/Engine/Private/Components/ChildActorComponent.cpp b/Engine/Source/Runtime/Engine/Private/Components/ChildActorComponent.cpp
--- a/Engine/Source/Runtime/Engine/Private/Components/ChildActorComponent.cpp
+++ b/Engine/Source/Runtime/Engine/Private/Components/ChildActorComponent.cpp
@@ -29,7 +29,7 @@ void UChildActorComponent::OnRegister()
 			// we can't attach in CreateChildActor since it has intermediate Mobility set up
 			// causing spam with inconsistent mobility set up
 			// so moving Attach to happen in Register
-			ChildActor->AttachRootComponentTo(this, NAME_None, EAttachLocation::SnapToTarget);
+			ChildActor->AttachRootComponentTo(this, NAME_None, EAttachLocation::SnapToTargetIncludingScale);
 		}
 	}
 	else if (ChildActorClass)
@@ -272,7 +272,7 @@ void UChildActorComponent::CreateChildActor()
 					// Remember which actor spawned it (for selection in editor etc)
 					ChildActor->ParentComponentActor = MyOwner;
 
-					ChildActor->AttachRootComponentTo(this);
+					ChildActor->AttachRootComponentTo(this, NAME_None, EAttachLocation::SnapToTargetIncludingScale);
 
 					// Parts that we deferred from SpawnActor
 					ChildActor->FinishSpawning(ComponentToWorld);
diff --git a/Engine/Source/Runtime/Engine/Private/Components/SceneComponent.cpp b/Engine/Source/Runtime/Engine/Private/Components/SceneComponent.cpp
--- a/Engine/Source/Runtime/Engine/Private/Components/SceneComponent.cpp
+++ b/Engine/Source/Runtime/Engine/Private/Components/SceneComponent.cpp
@@ -1255,7 +1255,13 @@ void USceneComponent::AttachTo(class USceneComponent* Parent, FName InSocketName
 				}
 			}
 			break;
-
+		case EAttachLocation::SnapToTargetIncludingScale:
+			{
+				RelativeLocation = FVector::ZeroVector;
+				RelativeRotation = FRotator::ZeroRotator;
+				RelativeScale3D = FVector(1.f, 1.f, 1.f);
+			}
+			break;
 		default:
 			// Leave the transform alone (relative offset to parent stays the same)
 			break;
