UE-15486 Fixing sounds for UI slate when using PIE.

- Old design assumed a singleton audio device and used the "main" audio device to play all its UI sounds from, but now the main audio device may be muted (since its not active).
- Fix is to make a new function on GEngine which returns the active audio device and plays the UI sound on that.
- Testing indicates that the UI sound will always play correctly on the active PIE window, tabbing between windows correctly mutes the UI sound from one to the other

diff --git a/Engine/Source/Runtime/Engine/Classes/Engine/Engine.h b/Engine/Source/Runtime/Engine/Classes/Engine/Engine.h
--- a/Engine/Source/Runtime/Engine/Classes/Engine/Engine.h
+++ b/Engine/Source/Runtime/Engine/Classes/Engine/Engine.h
@@ -2049,6 +2049,9 @@ public:
 	/** @return the main audio device. */
 	class FAudioDevice* GetMainAudioDevice();
 
+	/** @return the currently active audio device */
+	class FAudioDevice* GetActiveAudioDevice();
+
 	DEPRECATED(4.8, "GetAudioDevice is deprecated UEngine::GetMainAudioDevice instead.")
 	/** @return the main audio device. */
 	class FAudioDevice* GetAudioDevice();
diff --git a/Engine/Source/Runtime/Engine/Private/AudioDeviceManager.cpp b/Engine/Source/Runtime/Engine/Private/AudioDeviceManager.cpp
--- a/Engine/Source/Runtime/Engine/Private/AudioDeviceManager.cpp
+++ b/Engine/Source/Runtime/Engine/Private/AudioDeviceManager.cpp
@@ -34,6 +34,7 @@ FAudioDeviceManager::FAudioDeviceManager()
 	, NumWorldsUsingMainAudioDevice(0)
 	, NextResourceID(1)
 	, SoloDeviceHandle(INDEX_NONE)
+	, ActiveAudioDeviceHandle(INDEX_NONE)
 {
 }
 
@@ -248,6 +249,15 @@ class FAudioDevice* FAudioDeviceManager::GetAudioDevice(uint32 Handle)
 	return AudioDevice;
 }
 
+class FAudioDevice* FAudioDeviceManager::GetActiveAudioDevice()
+{
+	if (ActiveAudioDeviceHandle != INDEX_NONE)
+	{
+		return GetAudioDevice(ActiveAudioDeviceHandle);
+	}
+	return GEngine->GetMainAudioDevice();
+}
+
 void FAudioDeviceManager::UpdateActiveAudioDevices(bool bGameTicking)
 {
 	for (FAudioDevice* AudioDevice : Devices)
@@ -337,6 +347,7 @@ void FAudioDeviceManager::SetActiveDevice(uint32 InAudioDeviceHandle)
 			{
 				if (AudioDevice->DeviceHandle == InAudioDeviceHandle)
 				{
+					ActiveAudioDeviceHandle = InAudioDeviceHandle;
 					AudioDevice->bIsDeviceMuted = false;
 				}
 				else
@@ -357,8 +368,10 @@ void FAudioDeviceManager::SetSoloDevice(uint32 InAudioDeviceHandle)
 		{
 			if (AudioDevice)
 			{
+				// Un-mute the active audio device and mute non-active device, as long as its not the main audio device (which is used to play UI sounds)
 				if (AudioDevice->DeviceHandle == InAudioDeviceHandle)
 				{
+					ActiveAudioDeviceHandle = InAudioDeviceHandle;
 					AudioDevice->bIsDeviceMuted = false;
 				}
 				else
diff --git a/Engine/Source/Runtime/Engine/Private/Slate/SlateSoundDevice.cpp b/Engine/Source/Runtime/Engine/Private/Slate/SlateSoundDevice.cpp
--- a/Engine/Source/Runtime/Engine/Private/Slate/SlateSoundDevice.cpp
+++ b/Engine/Source/Runtime/Engine/Private/Slate/SlateSoundDevice.cpp
@@ -11,7 +11,7 @@ DEFINE_LOG_CATEGORY_STATIC(LogSlateSoundDevice, Log, All);
 
 void FSlateSoundDevice::PlaySound(const FSlateSound& Sound, int32 UserIndex) const
 {
-	FAudioDevice* const AudioDevice = GEngine->GetMainAudioDevice();
+	FAudioDevice* const AudioDevice = GEngine->GetActiveAudioDevice();
 	if(AudioDevice)
 	{
 		UObject* const Object = Sound.GetResourceObject();
diff --git a/Engine/Source/Runtime/Engine/Private/UnrealEngine.cpp b/Engine/Source/Runtime/Engine/Private/UnrealEngine.cpp
--- a/Engine/Source/Runtime/Engine/Private/UnrealEngine.cpp
+++ b/Engine/Source/Runtime/Engine/Private/UnrealEngine.cpp
@@ -1656,6 +1656,15 @@ FAudioDevice* UEngine::GetMainAudioDevice()
 	return nullptr;
 }
 
+FAudioDevice* UEngine::GetActiveAudioDevice()
+{
+	if (AudioDeviceManager != nullptr)
+	{
+		return AudioDeviceManager->GetActiveAudioDevice();
+	}
+	return nullptr;
+}
+
 FAudioDevice* UEngine::GetAudioDevice()
 {
 	return GetMainAudioDevice();
diff --git a/Engine/Source/Runtime/Engine/Public/AudioDeviceManager.h b/Engine/Source/Runtime/Engine/Public/AudioDeviceManager.h
--- a/Engine/Source/Runtime/Engine/Public/AudioDeviceManager.h
+++ b/Engine/Source/Runtime/Engine/Public/AudioDeviceManager.h
@@ -54,6 +54,12 @@ public:
 	*/
 	class FAudioDevice* GetAudioDevice(uint32 Handle);
 
+	/**
+	* Returns a ptr to the active audio device. If there is no active 
+	* device then it will return the main audio device.
+	*/
+	class FAudioDevice* GetActiveAudioDevice();
+
 	/** Returns the current number of active audio devices. */
 	uint8 GetNumActiveAudioDevices() const;
 
@@ -157,6 +163,9 @@ private:
 
 	/** Which audio device is solo'd */
 	uint32 SoloDeviceHandle;
+
+	/** Which audio device is currently active */
+	uint32 ActiveAudioDeviceHandle;
 };
 
 
