#ue4 - Perf: CharacterMovement nav mesh projection uses cheaper line-plane-intersection, and avoids it altogether when at the destination.

diff --git a/Engine/Source/Runtime/Engine/Private/Components/CharacterMovementComponent.cpp b/Engine/Source/Runtime/Engine/Private/Components/CharacterMovementComponent.cpp
--- a/Engine/Source/Runtime/Engine/Private/Components/CharacterMovementComponent.cpp
+++ b/Engine/Source/Runtime/Engine/Private/Components/CharacterMovementComponent.cpp
@@ -3976,19 +3976,24 @@ FVector UCharacterMovementComponent::ProjectLocationFromNavMesh(float DeltaSecon
 
 	FVector NewLocation = TargetNavLocation;
 
-	const FVector RayCastOffsetUp(0.0f, 0.0f, UpOffset);
-	const FVector RayCastOffsetDown(0.0f, 0.0f, DownOffset);
+	const float ZOffset = -(DownOffset + UpOffset);
+	if (ZOffset > -SMALL_NUMBER)
+	{
+		return NewLocation;
+	}
+
+	const FVector TraceStart = FVector(TargetNavLocation.X, TargetNavLocation.Y, TargetNavLocation.Z + UpOffset);
+	const FVector TraceEnd   = FVector(TargetNavLocation.X, TargetNavLocation.Y, TargetNavLocation.Z - DownOffset);
+
+	// We can skip this trace if we are checking at the same location as the last trace (ie, we haven't moved).
+	const bool bCachedLocationStillValid = (CachedProjectedNavMeshHitResult.bBlockingHit &&
+											CachedProjectedNavMeshHitResult.TraceStart == TraceStart &&
+											CachedProjectedNavMeshHitResult.TraceEnd == TraceEnd);
 
 	NavMeshProjectionTimer -= DeltaSeconds;
 	if (NavMeshProjectionTimer <= 0.0f)
 	{
-		// We can skip this trace if we are checking at the same location as the last trace (ie, we haven't moved).
-		const bool bCachedLocationStillValid = (!bAlwaysCheckFloor &&
-												CachedProjectedNavMeshHitResult.bBlockingHit &&
-												CachedProjectedNavMeshHitResult.TraceStart == (TargetNavLocation + RayCastOffsetUp) &&
-												CachedProjectedNavMeshHitResult.TraceEnd == (TargetNavLocation - RayCastOffsetDown));
-
-		if (!bCachedLocationStillValid)
+		if (!bCachedLocationStillValid || bAlwaysCheckFloor)
 		{
 			UE_LOG(LogNavMeshMovement, VeryVerbose, TEXT("ProjectLocationFromNavMesh(): %s interval: %.3f velocity: %s"), *GetNameSafe(CharacterOwner), NavMeshProjectionInterval, *Velocity.ToString());
 
@@ -3999,7 +4004,7 @@ FVector UCharacterMovementComponent::ProjectLocationFromNavMesh(float DeltaSecon
 			FCollisionQueryParams Params(ProjectLocationName, false);
 			FCollisionResponseParams ResponseParams(ECR_Ignore); // ignore everything
 			ResponseParams.CollisionResponse.SetResponse(ECC_WorldStatic, ECR_Block); // get blocked only by WorldStatic
-			GetWorld()->LineTraceSingleByChannel(CachedProjectedNavMeshHitResult, TargetNavLocation + RayCastOffsetUp, TargetNavLocation - RayCastOffsetDown, ECC_WorldStatic, Params, ResponseParams);
+			GetWorld()->LineTraceSingleByChannel(CachedProjectedNavMeshHitResult, TraceStart, TraceEnd, ECC_WorldStatic, Params, ResponseParams);
 
 			// discard result if we were already inside something
 			if (CachedProjectedNavMeshHitResult.bStartPenetrating)
@@ -4017,27 +4022,41 @@ FVector UCharacterMovementComponent::ProjectLocationFromNavMesh(float DeltaSecon
 		float ModTime = 0.f;
 		if (NavMeshProjectionInterval > SMALL_NUMBER)
 		{
-			ModTime = FMath::Fmod(FMath::Abs(NavMeshProjectionTimer), NavMeshProjectionInterval);
+			ModTime = FMath::Fmod(-NavMeshProjectionTimer, NavMeshProjectionInterval);
 		}
 
 		NavMeshProjectionTimer = NavMeshProjectionInterval - ModTime;
 	}
-		
-	// project to last plane we found
+	
+	// Project to last plane we found.
 	if (CachedProjectedNavMeshHitResult.bBlockingHit)
 	{
-		FVector ProjectedPoint = FMath::LinePlaneIntersection(TargetNavLocation + RayCastOffsetUp, TargetNavLocation - RayCastOffsetDown, CachedProjectedNavMeshHitResult.Location, CachedProjectedNavMeshHitResult.Normal);
+		if (bCachedLocationStillValid && FMath::IsNearlyEqual(CurrentFeetLocation.Z, CachedProjectedNavMeshHitResult.ImpactPoint.Z, 0.01f))
+		{
+			// Already at destination.
+			NewLocation.Z = CurrentFeetLocation.Z;
+		}
+		else
+		{
+			//const FVector ProjectedPoint = FMath::LinePlaneIntersection(TraceStart, TraceEnd, CachedProjectedNavMeshHitResult.ImpactPoint, CachedProjectedNavMeshHitResult.Normal);
+			//float ProjectedZ = ProjectedPoint.Z;
+
+			// Optimized assuming we only care about Z coordinate of result.
+			const FVector& PlaneOrigin = CachedProjectedNavMeshHitResult.ImpactPoint;
+			const FVector& PlaneNormal = CachedProjectedNavMeshHitResult.Normal;
+			float ProjectedZ = TraceStart.Z + ZOffset * (((PlaneOrigin - TraceStart)|PlaneNormal) / (ZOffset * PlaneNormal.Z));
 
-		// Limit to not be too far above or below NavMesh location
-		ProjectedPoint.Z = FMath::Clamp(ProjectedPoint.Z, TargetNavLocation.Z - DownOffset, TargetNavLocation.Z + UpOffset);
+			// Limit to not be too far above or below NavMesh location
+			ProjectedZ = FMath::Clamp(ProjectedZ, TraceEnd.Z, TraceStart.Z);
 
-		// Interp for smoother updates (less "pop" when trace hits something new). 0 interp speed is instant.
-		const float InterpSpeed = FMath::Max(0.f, NavMeshProjectionInterpSpeed);
-		ProjectedPoint.Z = FMath::FInterpTo(CurrentFeetLocation.Z, ProjectedPoint.Z, DeltaSeconds, InterpSpeed);
-		ProjectedPoint.Z = FMath::Clamp(ProjectedPoint.Z, TargetNavLocation.Z - DownOffset, TargetNavLocation.Z + UpOffset);
+			// Interp for smoother updates (less "pop" when trace hits something new). 0 interp speed is instant.
+			const float InterpSpeed = FMath::Max(0.f, NavMeshProjectionInterpSpeed);
+			ProjectedZ = FMath::FInterpTo(CurrentFeetLocation.Z, ProjectedZ, DeltaSeconds, InterpSpeed);
+			ProjectedZ = FMath::Clamp(ProjectedZ, TraceEnd.Z, TraceStart.Z);
 
-		// Final result
-		NewLocation.Z = ProjectedPoint.Z;
+			// Final result
+			NewLocation.Z = ProjectedZ;
+		}
 	}
 
 	return NewLocation;
