Asserts to catch garbage shadermap

diff --git a/Engine/Source/Runtime/Engine/Private/Materials/MaterialInstance.cpp b/Engine/Source/Runtime/Engine/Private/Materials/MaterialInstance.cpp
--- a/Engine/Source/Runtime/Engine/Private/Materials/MaterialInstance.cpp
+++ b/Engine/Source/Runtime/Engine/Private/Materials/MaterialInstance.cpp
@@ -671,6 +671,8 @@ void UMaterialInstance::GetTextureExpressionValues(const FMaterialResource* Mate
 {
 	const TArray<TRefCountPtr<FMaterialUniformExpressionTexture> >* ExpressionsByType[2];
 
+	check(MaterialResource);
+
 	ExpressionsByType[0] = &MaterialResource->GetUniform2DTextureExpressions();
 	ExpressionsByType[1] = &MaterialResource->GetUniformCubeTextureExpressions();
 
diff --git a/Engine/Source/Runtime/Engine/Private/Materials/MaterialShader.cpp b/Engine/Source/Runtime/Engine/Private/Materials/MaterialShader.cpp
--- a/Engine/Source/Runtime/Engine/Private/Materials/MaterialShader.cpp
+++ b/Engine/Source/Runtime/Engine/Private/Materials/MaterialShader.cpp
@@ -1505,7 +1505,7 @@ void FMaterialShaderMap::AddRef()
 
 void FMaterialShaderMap::Release()
 {
-	check(NumRefs != 0);
+	check(NumRefs > 0);
 	if(--NumRefs == 0)
 	{
 		if (bRegistered)
diff --git a/Engine/Source/Runtime/Engine/Private/Materials/MaterialShared.cpp b/Engine/Source/Runtime/Engine/Private/Materials/MaterialShared.cpp
--- a/Engine/Source/Runtime/Engine/Private/Materials/MaterialShared.cpp
+++ b/Engine/Source/Runtime/Engine/Private/Materials/MaterialShared.cpp
@@ -348,6 +348,8 @@ const TArray<TRefCountPtr<FMaterialUniformExpressionTexture> >& FMaterial::GetUn
 		// If we are accessing uniform texture expressions on the game thread, use results from a shader map whose compile is in flight that matches this material
 		// This allows querying what textures a material uses even when it is being asynchronously compiled
 		ShaderMapToUse = GetGameThreadShaderMap() ? GetGameThreadShaderMap() : FMaterialShaderMap::GetShaderMapBeingCompiled(this);
+
+		checkf(!ShaderMapToUse || ShaderMapToUse->GetNumRefs() > 0, TEXT("NumRefs %i, GameThreadShaderMap 0x%08x"), ShaderMapToUse->GetNumRefs(), GetGameThreadShaderMap());
 	}
 	else 
 	{
@@ -642,12 +644,12 @@ void FMaterial::SerializeInlineShaderMap(FArchive& Ar)
 				//@todo - don't cook it in the first place
 				if (FApp::CanEverRender())
 				{
-				GameThreadShaderMap = RenderingThreadShaderMap = LoadedShaderMap;
+					GameThreadShaderMap = RenderingThreadShaderMap = LoadedShaderMap;
+				}
 			}
 		}
 	}
 }
-}
 
 void FMaterialResource::LegacySerialize(FArchive& Ar)
 {
diff --git a/Engine/Source/Runtime/Engine/Public/MaterialShared.h b/Engine/Source/Runtime/Engine/Public/MaterialShared.h
--- a/Engine/Source/Runtime/Engine/Public/MaterialShared.h
+++ b/Engine/Source/Runtime/Engine/Public/MaterialShared.h
@@ -667,6 +667,8 @@ public:
 
 	const FUniformExpressionSet& GetUniformExpressionSet() const { return MaterialCompilationOutput.UniformExpressionSet; }
 
+	int32 GetNumRefs() const { return NumRefs; }
+
 private:
 
 	/** 
@@ -712,7 +714,7 @@ private:
 	/** Uniquely identifies this shader map during compilation, needed for deferred compilation where shaders from multiple shader maps are compiled together. */
 	uint32 CompilingId;
 
-	mutable uint32 NumRefs;
+	mutable int32 NumRefs;
 
 	/** Used to catch errors where the shader map is deleted directly. */
 	bool bDeletedThroughDeferredCleanup;
