Fixed primitives with bLightAsIfStatic enabled causing lighting to be unbuilt if they move into a light's influence radius that they weren't built against

diff --git a/Engine/Source/Runtime/Engine/Private/PrimitiveSceneProxy.cpp b/Engine/Source/Runtime/Engine/Private/PrimitiveSceneProxy.cpp
--- a/Engine/Source/Runtime/Engine/Private/PrimitiveSceneProxy.cpp
+++ b/Engine/Source/Runtime/Engine/Private/PrimitiveSceneProxy.cpp
@@ -48,6 +48,7 @@ FPrimitiveSceneProxy::FPrimitiveSceneProxy(const UPrimitiveComponent* InComponen
 ,	bCastInsetShadow(InComponent->bSelfShadowOnly ? true : InComponent->bCastInsetShadow)	// Assumed to be enabled if bSelfShadowOnly is enabled.
 ,	bCastCinematicShadow(InComponent->bCastCinematicShadow)
 ,	bCastFarShadow(InComponent->bCastFarShadow)
+,	bLightAsIfStatic(InComponent->bLightAsIfStatic)
 ,	bLightAttachmentsAsGroup(InComponent->bLightAttachmentsAsGroup)
 ,	bStaticElementsAlwaysUseProxyPrimitiveUniformBuffer(false)
 ,	bAlwaysHasVelocity(false)
diff --git a/Engine/Source/Runtime/Engine/Public/PrimitiveSceneProxy.h b/Engine/Source/Runtime/Engine/Public/PrimitiveSceneProxy.h
--- a/Engine/Source/Runtime/Engine/Public/PrimitiveSceneProxy.h
+++ b/Engine/Source/Runtime/Engine/Public/PrimitiveSceneProxy.h
@@ -399,6 +399,7 @@ public:
 	inline bool CastsInsetShadow() const { return bCastInsetShadow; }
 	inline bool CastsCinematicShadow() const { return bCastCinematicShadow; }
 	inline bool CastsFarShadow() const { return bCastFarShadow; }
+	inline bool LightAsIfStatic() const { return bLightAsIfStatic; }
 	inline bool LightAttachmentsAsGroup() const { return bLightAttachmentsAsGroup; }
 	inline bool StaticElementsAlwaysUseProxyPrimitiveUniformBuffer() const { return bStaticElementsAlwaysUseProxyPrimitiveUniformBuffer; }
 	inline bool ShouldUseAsOccluder() const { return bUseAsOccluder; }
@@ -596,6 +597,12 @@ protected:
 	/* When enabled, the component will be rendering into the distant shadow cascades (only for directional lights). */
 	uint32 bCastFarShadow : 1;
 
+	/** 
+	 * This has to be known by the rendering thread to avoid marking lighting dirty when new interactions are created,
+	 * Which happens when a movable mesh with bLightAsIfStatic moves into the influence of a light it was not baked against.
+	 */
+	uint32 bLightAsIfStatic : 1;
+
 	/** 
 	 * Whether to light this component and any attachments as a group.  This only has effect on the root component of an attachment tree.
 	 * When enabled, attached component shadowing settings like bCastInsetShadow, bCastVolumetricTranslucentShadow, etc, will be ignored.
diff --git a/Engine/Source/Runtime/Renderer/Private/SceneCore.cpp b/Engine/Source/Runtime/Renderer/Private/SceneCore.cpp
--- a/Engine/Source/Runtime/Renderer/Private/SceneCore.cpp
+++ b/Engine/Source/Runtime/Renderer/Private/SceneCore.cpp
@@ -172,6 +172,8 @@ FLightPrimitiveInteraction::FLightPrimitiveInteraction(
 		// Determine the type of dynamic shadow produced by this light.
 		if (PrimitiveSceneInfo->Proxy->HasStaticLighting()
 			&& PrimitiveSceneInfo->Proxy->CastsStaticShadow()
+			// Don't mark unbuilt for movable primitives which were built with lightmaps but moved into a new light's influence
+			&& !PrimitiveSceneInfo->Proxy->LightAsIfStatic()
 			&& (LightSceneInfo->Proxy->HasStaticLighting() || (LightSceneInfo->Proxy->HasStaticShadowing() && !bInIsShadowMapped)))
 		{
 			// Update the game thread's counter of number of uncached static lighting interactions.
