UE4 - fix threading with the atmospheric fog component

diff --git a/Engine/Source/Runtime/Engine/Classes/Atmosphere/AtmosphericFogComponent.h b/Engine/Source/Runtime/Engine/Classes/Atmosphere/AtmosphericFogComponent.h
--- a/Engine/Source/Runtime/Engine/Classes/Atmosphere/AtmosphericFogComponent.h
+++ b/Engine/Source/Runtime/Engine/Classes/Atmosphere/AtmosphericFogComponent.h
@@ -91,6 +91,8 @@ class UAtmosphericFogComponent : public USceneComponent
 {
 	GENERATED_UCLASS_BODY()
 
+	~UAtmosphericFogComponent();
+
 	/** Global scattering factor. */
 	UPROPERTY(EditAnywhere, BlueprintReadWrite, interp, Category=Atmosphere)
 	float SunMultiplier;
@@ -213,11 +215,13 @@ public:
 	enum EPrecomputeState
 	{
 		EInvalid = 0,
-		EFinishedComputation = 1,
 		EValid = 2,
 	};
 
-	FThreadSafeCounter PrecomputeCounter;
+	// this is mostly a legacy thing, it is only modified by the game thread
+	uint32 PrecomputeCounter;
+	// When non-zero, the component should flush rendering commands and see if there is any atmosphere stuff to deal with, then decrement it
+	mutable FThreadSafeCounter GameThreadServiceRequest;
 
 	UFUNCTION(BlueprintCallable, Category="Rendering|Components|AtmosphericFog")
 	void StartPrecompute();	
@@ -257,13 +261,6 @@ public:
 	ENGINE_API void InitResource();
 	ENGINE_API void ReleaseResource();
 
-	// Begin UActorComponent interface.
-	virtual FActorComponentInstanceData* GetComponentInstanceData() const override;
-	virtual FName GetComponentInstanceDataType() const override;
-	// End UActorComponent interface.
-
-	void ApplyComponentInstanceData(class FAtmospherePrecomputeInstanceData* ComponentInstanceData);
-
 	const FAtmospherePrecomputeParameters& GetPrecomputeParameters() const { return PrecomputeParams;  }
 
 private:
diff --git a/Engine/Source/Runtime/Engine/Private/Atmosphere/Atmosphere.cpp b/Engine/Source/Runtime/Engine/Private/Atmosphere/Atmosphere.cpp
--- a/Engine/Source/Runtime/Engine/Private/Atmosphere/Atmosphere.cpp
+++ b/Engine/Source/Runtime/Engine/Private/Atmosphere/Atmosphere.cpp
@@ -189,9 +189,10 @@ public:
 	/** FTickableEditorObject interface */
 	virtual void Tick( float DeltaTime ) override
 	{
-		if( Component && Component->PrecomputeCounter.GetValue() == UAtmosphericFogComponent::EFinishedComputation ) 
+		if( Component && Component->GameThreadServiceRequest.GetValue()) 
 		{
 			Component->UpdatePrecomputedData();
+			Component->GameThreadServiceRequest.Decrement();
 		}              
 	}
 
@@ -226,11 +227,11 @@ void UAtmosphericFogComponent::InitResource()
 		return; // Don't do initialize resource when Atmosphere Render is off
 	}
 
-	if (PrecomputeCounter.GetValue() >= EValid)
+	if (PrecomputeCounter >= EValid)
 	{
 		// A little inefficient for thread-safe
 		if (TransmittanceData.GetElementCount() && TransmittanceResource == NULL)
-	{
+		{
 			TransmittanceResource = new FAtmosphereTextureResource( PrecomputeParams, TransmittanceData, FAtmosphereTextureResource::E_Transmittance ); 
 			BeginInitResource( TransmittanceResource );
 		}
@@ -248,7 +249,7 @@ void UAtmosphericFogComponent::InitResource()
 		}
 	}
 #if WITH_EDITOR
-	else if (!PrecomputeDataHandler)
+	else if (!PrecomputeDataHandler && !IsTemplate())
 	{
 		PrecomputeDataHandler = new FAtmospherePrecomputeDataHandler(this);
 	}
@@ -263,12 +264,18 @@ void UAtmosphericFogComponent::BeginDestroy()
 
 void UAtmosphericFogComponent::ReleaseResource()
 {
+	FSceneInterface* Scene = GetScene();
 	if ( TransmittanceResource != NULL )
 	{
-		ENQUEUE_UNIQUE_RENDER_COMMAND_ONEPARAMETER(
+		ENQUEUE_UNIQUE_RENDER_COMMAND_TWOPARAMETER(
 			ReleaseAtmosphericFogTransmittanceTextureCommand,
 			FRenderResource*,TransmittanceResource,TransmittanceResource,
+			FSceneInterface*,Scene,Scene,
 		{
+			if (Scene)
+			{
+				Scene->RemoveAtmosphericFogResource_RenderThread(TransmittanceResource);
+			}
 			TransmittanceResource->ReleaseResource();
 			delete TransmittanceResource;
 		});
@@ -278,10 +285,15 @@ void UAtmosphericFogComponent::ReleaseResource()
 
 	if ( IrradianceResource != NULL )
 	{
-		ENQUEUE_UNIQUE_RENDER_COMMAND_ONEPARAMETER(
+		ENQUEUE_UNIQUE_RENDER_COMMAND_TWOPARAMETER(
 			ReleaseAtmosphericFogIrradianceTextureCommand,
 			FRenderResource*,IrradianceResource,IrradianceResource,
+			FSceneInterface*,Scene,Scene,
 		{
+			if (Scene)
+			{
+				Scene->RemoveAtmosphericFogResource_RenderThread(IrradianceResource);
+			}
 			IrradianceResource->ReleaseResource();
 			delete IrradianceResource;
 		});
@@ -291,24 +303,21 @@ void UAtmosphericFogComponent::ReleaseResource()
 
 	if ( InscatterResource != NULL )
 	{
-		ENQUEUE_UNIQUE_RENDER_COMMAND_ONEPARAMETER(
+		ENQUEUE_UNIQUE_RENDER_COMMAND_TWOPARAMETER(
 			ReleaseAtmosphericFogInscatterTextureCommand,
 			FRenderResource*,InscatterResource,InscatterResource,
+			FSceneInterface*,Scene,Scene,
 		{
+			if (Scene)
+			{
+				Scene->RemoveAtmosphericFogResource_RenderThread(InscatterResource);
+			}
 			InscatterResource->ReleaseResource();
 			delete InscatterResource;
 		});
 
 		InscatterResource = NULL;
 	}
-
-#if WITH_EDITOR
-	if (PrecomputeDataHandler)
-	{
-		delete PrecomputeDataHandler;
-		PrecomputeDataHandler = NULL;
-	}
-#endif
 }
 
 void UAtmosphericFogComponent::CreateRenderState_Concurrent()
@@ -543,34 +552,31 @@ void UAtmosphericFogComponent::PostInterpChange(UProperty* PropertyThatChanged)
 void UAtmosphericFogComponent::StartPrecompute()
 {
 #if WITH_EDITOR
-	if (GIsEditor)
+	if (GIsEditor && !IsTemplate())
 	{
-		if (GetScene())
+		FSceneInterface* AtmosphericFogScene = GetScene();
+		if (AtmosphericFogScene)
 		{
-			PrecomputeCounter.Reset();
+			PrecomputeCounter = EInvalid;
 
-			if (PrecomputeDataHandler)
+			if (!PrecomputeDataHandler)
 			{
-				delete PrecomputeDataHandler;
-				PrecomputeDataHandler = NULL;
+				PrecomputeDataHandler = new FAtmospherePrecomputeDataHandler(this);
 			}
-			PrecomputeDataHandler = new FAtmospherePrecomputeDataHandler(this);
-
-			FAtmosphericFogSceneInfo* AtmosphericFogSceneInfo = GetScene()->GetAtmosphericFogSceneInfo();
 
-			if (AtmosphericFogSceneInfo)
+			// This is largely redundant, the component will be reattached anyway, thus it will be recomputed.
+			ENQUEUE_UNIQUE_RENDER_COMMAND_TWOPARAMETER(
+				FStartPrecompute,
+				FSceneInterface*, AtmosphericFogScene, AtmosphericFogScene, 
+				UAtmosphericFogComponent*, Component, this,
 			{
-				ENQUEUE_UNIQUE_RENDER_COMMAND_TWOPARAMETER(
-					FStartPrecompute,
-					FAtmosphericFogSceneInfo*, AtmosphericFogSceneInfo, AtmosphericFogSceneInfo, 
-					UAtmosphericFogComponent*, Component, this,
+				FAtmosphericFogSceneInfo* AtmosphericFogSceneInfo = AtmosphericFogScene->GetAtmosphericFogSceneInfo();
+
+				if (AtmosphericFogSceneInfo->Component == Component)
 				{
-					if (AtmosphericFogSceneInfo->Component == Component)
-					{
-						AtmosphericFogSceneInfo->bNeedRecompute = true;
-					}
-				});
-			}
+					AtmosphericFogSceneInfo->bNeedRecompute = true;
+				}
+			});
 		}
 	}
 #endif
@@ -590,14 +596,25 @@ void AAtmosphericFog::PostActorCreated()
 	}
 }
 
+UAtmosphericFogComponent::~UAtmosphericFogComponent()
+{
+#if WITH_EDITOR
+	if (PrecomputeDataHandler)
+	{
+		delete PrecomputeDataHandler;
+		PrecomputeDataHandler = NULL;
+	}
+#endif
+}
+
 void UAtmosphericFogComponent::UpdatePrecomputedData()
 {
-	if (GIsEditor && PrecomputeCounter.GetValue() == EFinishedComputation) 
+	if (GIsEditor) 
 	{
 		// Prevent atmosphere pre-computation texture read/write from rendering thread during this process
 		FlushRenderingCommands();
 		FScene* Scene = GetScene() ? GetScene()->GetRenderScene() : NULL;
-		if (Scene && Scene->AtmosphericFog && this == Scene->AtmosphericFog->Component)
+		if (Scene && Scene->AtmosphericFog && this == Scene->AtmosphericFog->Component && Scene->AtmosphericFog->bPrecomputationFinished && !Scene->AtmosphericFog->bPrecomputationAcceptedByGameThread)
 		{
 			// When the precomputation is done, should save result to final texture for rendering
 			// Need to resolve to actor textures and remove render targets
@@ -641,12 +658,14 @@ void UAtmosphericFogComponent::UpdatePrecomputedData()
 				InscatterData.Unlock();
 				Scene->AtmosphericFog->PrecomputeInscatter.Unlock();
 			}
+			PrecomputeCounter = EValid;
+			FPlatformMisc::MemoryBarrier();
+			Scene->AtmosphericFog->bPrecomputationAcceptedByGameThread = true;
 
 			// Resolve to data...
 			ReleaseResource();
 			// Wait for release...
 			FlushRenderingCommands();
-			PrecomputeCounter.Increment();
 
 			InitResource();
 			FComponentReregisterContext ReregisterContext(this);
@@ -677,19 +696,19 @@ void UAtmosphericFogComponent::Serialize(FArchive& Ar)
 		int32 CounterVal;
 		Ar << CounterVal;
 		// Precomputation was not successful, just ignore it
-		if (CounterVal < EValid)
+		if (CounterVal < EValid || !TransmittanceData.GetElementCount())
 		{
 			CounterVal = EInvalid;
 		}
-		PrecomputeCounter.Set(CounterVal);
+		PrecomputeCounter = CounterVal;
 	}
 	else
 	{
-		int32 CounterVal = PrecomputeCounter.GetValue();
+		int32 CounterVal = PrecomputeCounter;
 		Ar << CounterVal;
 	}
 
-	if (Ar.IsLoading() && Ar.UE4Ver() < VER_UE4_ATMOSPHERIC_FOG_CACHE_DATA && PrecomputeCounter.GetValue() == EValid) 
+	if (Ar.IsLoading() && Ar.UE4Ver() < VER_UE4_ATMOSPHERIC_FOG_CACHE_DATA && PrecomputeCounter == EValid) 
 	{
 		// InscatterAltitudeSampleNum default value has been changed (32 -> 2)
 		// Recalculate InscatterAltitudeSampleNum based on Inscatter Size
@@ -700,121 +719,6 @@ void UAtmosphericFogComponent::Serialize(FArchive& Ar)
 	}
 }
 
-/** Used to store lightmap data during RerunConstructionScripts */
-class FAtmospherePrecomputeInstanceData : public FSceneComponentInstanceData
-{
-public:
-	FAtmospherePrecomputeInstanceData(const UAtmosphericFogComponent* SourceComponent)
-		: FSceneComponentInstanceData(SourceComponent)
-	{}
-
-	virtual ~FAtmospherePrecomputeInstanceData()
-	{}
-
-	virtual void ApplyToComponent(UActorComponent* Component, const ECacheApplyPhase CacheApplyPhase) override
-	{
-		FSceneComponentInstanceData::ApplyToComponent(Component, CacheApplyPhase);
-		CastChecked<UAtmosphericFogComponent>(Component)->ApplyComponentInstanceData(this);
-	}
-
-	struct FAtmospherePrecomputeParameters PrecomputeParameter;
-
-	FByteBulkData TransmittanceData;
-	FByteBulkData IrradianceData;
-	FByteBulkData InscatterData;
-};
-
-FName UAtmosphericFogComponent::GetComponentInstanceDataType() const
-{
-	static const FName InstanceDataTypeName(TEXT("AtmospherePrecomputedInstanceData"));
-	return InstanceDataTypeName;
-}
-
-// Backup the precomputed data before re-running Blueprint construction script
-FActorComponentInstanceData* UAtmosphericFogComponent::GetComponentInstanceData() const
-{
-	FActorComponentInstanceData* InstanceData = nullptr;
-
-	if (TransmittanceData.GetElementCount() && IrradianceData.GetElementCount() && InscatterData.GetElementCount() && PrecomputeCounter.GetValue() == EValid)
-	{
-		// Allocate new struct for holding light map data
-		 FAtmospherePrecomputeInstanceData* PrecomputedData = new FAtmospherePrecomputeInstanceData(this);
-		 InstanceData = PrecomputedData;
-
-		// Fill in info
-		PrecomputedData->PrecomputeParameter = PrecomputeParams;
-		{
-			int32 TotalByte = TransmittanceData.GetBulkDataSize();
-			PrecomputedData->TransmittanceData.Lock(LOCK_READ_WRITE);
-			void* OutData = PrecomputedData->TransmittanceData.Realloc(TotalByte);
-			TransmittanceData.GetCopy(&OutData, false);
-			PrecomputedData->TransmittanceData.Unlock();
-		}
-
-		{
-			int32 TotalByte = IrradianceData.GetBulkDataSize();
-			PrecomputedData->IrradianceData.Lock(LOCK_READ_WRITE);
-			void* OutData = PrecomputedData->IrradianceData.Realloc(TotalByte);
-			IrradianceData.GetCopy(&OutData, false);
-			PrecomputedData->IrradianceData.Unlock();
-		}
-
-		{
-			int32 TotalByte = InscatterData.GetBulkDataSize();
-			PrecomputedData->InscatterData.Lock(LOCK_READ_WRITE);
-			void* OutData = PrecomputedData->InscatterData.Realloc(TotalByte);
-			InscatterData.GetCopy(&OutData, false);
-			PrecomputedData->InscatterData.Unlock();
-		}
-	}
-	else
-	{
-		InstanceData = Super::GetComponentInstanceData();
-	}
-
-	return InstanceData;
-}
-
-// Restore the precomputed data after re-running Blueprint construction script
-void UAtmosphericFogComponent::ApplyComponentInstanceData(FAtmospherePrecomputeInstanceData* PrecomputedData)
-{
-	check(PrecomputedData);
-
-	if (PrecomputedData->PrecomputeParameter != GetPrecomputeParameters())
-	{
-		return;
-	}
-
-	FComponentReregisterContext ReregisterContext(this);
-	ReleaseResource();
-
-	{
-		int32 TotalByte = PrecomputedData->TransmittanceData.GetBulkDataSize();
-		TransmittanceData.Lock(LOCK_READ_WRITE);
-		void* OutData = TransmittanceData.Realloc(TotalByte);
-		PrecomputedData->TransmittanceData.GetCopy(&OutData, false);
-		TransmittanceData.Unlock();
-	}
-
-	{
-		int32 TotalByte = PrecomputedData->IrradianceData.GetBulkDataSize();
-		IrradianceData.Lock(LOCK_READ_WRITE);
-		void* OutData = IrradianceData.Realloc(TotalByte);
-		PrecomputedData->IrradianceData.GetCopy(&OutData, false);
-		IrradianceData.Unlock();
-	}
-
-	{
-		int32 TotalByte = PrecomputedData->InscatterData.GetBulkDataSize();
-		InscatterData.Lock(LOCK_READ_WRITE);
-		void* OutData = InscatterData.Realloc(TotalByte);
-		PrecomputedData->InscatterData.GetCopy(&OutData, false);
-		InscatterData.Unlock();
-	}
-
-	PrecomputeCounter.Set(EValid);
-	InitResource();
-}
 
 /**
  * Gets called any time cvars change (on the main thread), we check if r.Atmosphere has changed and update the components.
diff --git a/Engine/Source/Runtime/Engine/Public/SceneInterface.h b/Engine/Source/Runtime/Engine/Public/SceneInterface.h
--- a/Engine/Source/Runtime/Engine/Public/SceneInterface.h
+++ b/Engine/Source/Runtime/Engine/Public/SceneInterface.h
@@ -171,6 +171,13 @@ public:
 	 */	
 	virtual void RemoveAtmosphericFog(class UAtmosphericFogComponent* FogComponent) = 0;
 
+	/** 
+	 * Removes a atmospheric fog resource from the scene...this is just a double check to make sure we don't have stale stuff hanging around; should already be gone.
+	 * 
+	 * @param FogResource - fog resource to remove
+	 */	
+	virtual void RemoveAtmosphericFogResource_RenderThread(FRenderResource* FogResource) = 0;
+
 	/**
 	 * Returns the scene's FAtmosphericFogSceneInfo if it exists
 	 */
diff --git a/Engine/Source/Runtime/Renderer/Private/AtmosphereRendering.cpp b/Engine/Source/Runtime/Renderer/Private/AtmosphereRendering.cpp
--- a/Engine/Source/Runtime/Renderer/Private/AtmosphereRendering.cpp
+++ b/Engine/Source/Runtime/Renderer/Private/AtmosphereRendering.cpp
@@ -324,21 +324,10 @@ void FSceneRenderer::InitAtmosphereConstants()
 			{
 				const FAtmosphericFogSceneInfo& FogInfo = *Scene->AtmosphericFog;
 
-				if (FogInfo.Component->PrecomputeCounter.GetValue() >= UAtmosphericFogComponent::EValid)
-				{
-					View.AtmosphereTransmittanceTexture = FogInfo.Component->TransmittanceResource
-						? (FTextureRHIRef)FogInfo.Component->TransmittanceResource->TextureRHI
-						: (FogInfo.Component->TransmittanceTexture_DEPRECATED && FogInfo.Component->TransmittanceTexture_DEPRECATED->Resource) 
-								? FogInfo.Component->TransmittanceTexture_DEPRECATED->Resource->TextureRHI  : GBlackTexture->TextureRHI;  
-					View.AtmosphereIrradianceTexture = FogInfo.Component->IrradianceResource
-						? (FTextureRHIRef)FogInfo.Component->IrradianceResource->TextureRHI
-						: (FogInfo.Component->IrradianceTexture_DEPRECATED && FogInfo.Component->IrradianceTexture_DEPRECATED->Resource)
-								? FogInfo.Component->IrradianceTexture_DEPRECATED->Resource->TextureRHI : GBlackTexture->TextureRHI;
-					View.AtmosphereInscatterTexture = FogInfo.Component->InscatterResource 
-						? FogInfo.Component->InscatterResource->TextureRHI : (FTextureRHIRef)GBlackVolumeTexture->TextureRHI;
-
-					bInitTextures = true;
-				}
+				View.AtmosphereTransmittanceTexture = (FogInfo.TransmittanceResource && FogInfo.TransmittanceResource->TextureRHI.GetReference()) ? (FTextureRHIRef)FogInfo.TransmittanceResource->TextureRHI : GBlackTexture->TextureRHI;
+				View.AtmosphereIrradianceTexture = (FogInfo.IrradianceResource && FogInfo.IrradianceResource->TextureRHI.GetReference()) ? (FTextureRHIRef)FogInfo.IrradianceResource->TextureRHI : GBlackTexture->TextureRHI;
+				View.AtmosphereInscatterTexture = (FogInfo.InscatterResource && FogInfo.InscatterResource->TextureRHI.GetReference()) ? (FTextureRHIRef)FogInfo.InscatterResource->TextureRHI : GBlackVolumeTexture->TextureRHI;
+				bInitTextures = true;
 			}
 		}
 
@@ -1029,6 +1018,9 @@ namespace
 void FAtmosphericFogSceneInfo::StartPrecompute()
 {
 	bNeedRecompute = false;
+	bPrecomputationStarted = true;
+	check(!bPrecomputationFinished);
+	check(!bPrecomputationAcceptedByGameThread);
 	AtmospherePhase = 0;
 	Atmosphere3DTextureIndex = 0;
 	AtmoshpereOrder = 2;
@@ -1461,13 +1453,21 @@ void FAtmosphericFogSceneInfo::PrecomputeTextures(FRHICommandListImmediate& RHIC
 		AtmosphereTextures = new FAtmosphereTextures(&Component->PrecomputeParams);
 	}
 
-	if (bNeedRecompute)
+	if (bPrecomputationAcceptedByGameThread)
+	{
+		// we finished everything and so now can start a new one if another one came in
+		bPrecomputationStarted = false;
+		bPrecomputationFinished = false;
+		bPrecomputationAcceptedByGameThread = false;
+	}
+
+	if (bNeedRecompute && !bPrecomputationStarted)
 	{
 		StartPrecompute();
 	}
 
 	// Atmosphere 
-	if (PrecomputeCounter.GetValue() < UAtmosphericFogComponent::EFinishedComputation)
+	if (bPrecomputationStarted && !bPrecomputationFinished)
 	{
 		PrecomputeAtmosphereData(RHICmdList, View, *ViewFamily);
 
@@ -1555,7 +1555,8 @@ void FAtmosphericFogSceneInfo::PrecomputeTextures(FRHICommandListImmediate& RHIC
 			AtmosphereTextures = NULL;
 
 			// Save to bulk data is done
-			PrecomputeCounter.Increment(); // Notice to component that pre-computation is done...
+			bPrecomputationFinished = true;
+			Component->GameThreadServiceRequest.Increment();
 		}
 	}
 }
@@ -1577,14 +1578,17 @@ FAtmosphericFogSceneInfo::FAtmosphericFogSceneInfo(UAtmosphericFogComponent* InC
 	, SunDiscScale(InComponent->SunDiscScale)
 	, RenderFlag(EAtmosphereRenderFlag::E_EnableAll)
 	, InscatterAltitudeSampleNum(InComponent->PrecomputeParams.InscatterAltitudeSampleNum)
+
 #if WITH_EDITORONLY_DATA
 	, bNeedRecompute(false)
+	, bPrecomputationStarted(false)
+	, bPrecomputationFinished(false)
+	, bPrecomputationAcceptedByGameThread(false)
 	, MaxScatteringOrder(InComponent->PrecomputeParams.MaxScatteringOrder)
 	, AtmospherePhase(0)
 	, Atmosphere3DTextureIndex(0)
 	, AtmoshpereOrder(2)
 	, AtmosphereTextures(NULL)
-	, PrecomputeCounter(InComponent->PrecomputeCounter)
 #endif
 {
 	StartDistance *= DistanceScale * 0.00001f; // Convert to km in Atmospheric fog shader
@@ -1596,6 +1600,15 @@ FAtmosphericFogSceneInfo::FAtmosphericFogSceneInfo(UAtmosphericFogComponent* InC
 	GroundOffset += InComponent->GetComponentLocation().Z;
 	FMatrix WorldToLight = InComponent->ComponentToWorld.ToMatrixNoScale().InverseFast();
 	DefaultSunDirection = FVector(WorldToLight.M[0][0],WorldToLight.M[1][0],WorldToLight.M[2][0]);
+
+	if (Component->PrecomputeCounter != UAtmosphericFogComponent::EValid)
+	{
+		bNeedRecompute = true;
+	}
+
+	TransmittanceResource = Component->TransmittanceResource;
+	IrradianceResource = Component->IrradianceResource;
+	InscatterResource = Component->InscatterResource;
 }
 
 FAtmosphericFogSceneInfo::~FAtmosphericFogSceneInfo()
@@ -1617,3 +1630,50 @@ bool ShouldRenderAtmosphere(const FSceneViewFamily& Family)
 		&& EngineShowFlags.AtmosphericFog
 		&& EngineShowFlags.Fog;
 }
+
+void FScene::AddAtmosphericFog(UAtmosphericFogComponent* FogComponent)
+{
+	check(FogComponent);
+
+	FAtmosphericFogSceneInfo* FogSceneInfo = new FAtmosphericFogSceneInfo(FogComponent, this);
+
+	ENQUEUE_UNIQUE_RENDER_COMMAND_TWOPARAMETER(
+		FAddAtmosphericFogCommand,
+		FScene*,Scene,this,
+		FAtmosphericFogSceneInfo*,FogSceneInfo,FogSceneInfo,
+	{
+		delete Scene->AtmosphericFog;
+		Scene->AtmosphericFog = FogSceneInfo;
+	});
+}
+
+void FScene::RemoveAtmosphericFog(UAtmosphericFogComponent* FogComponent)
+{
+	ENQUEUE_UNIQUE_RENDER_COMMAND_TWOPARAMETER(
+		FRemoveAtmosphericFogCommand,
+		FScene*,Scene,this,
+		UAtmosphericFogComponent*,FogComponent,FogComponent,
+	{
+		// Remove the given component's FExponentialHeightFogSceneInfo from the scene's fog array.
+		if (Scene->AtmosphericFog && Scene->AtmosphericFog->Component == FogComponent)
+		{
+			delete Scene->AtmosphericFog;
+			Scene->AtmosphericFog = NULL;
+		}
+	});
+}
+
+
+void FScene::RemoveAtmosphericFogResource_RenderThread(FRenderResource* FogResource)
+{
+	check(IsInRenderingThread());
+
+	if (AtmosphericFog)
+	{
+		if (AtmosphericFog->TransmittanceResource == FogResource || AtmosphericFog->IrradianceResource == FogResource || AtmosphericFog->InscatterResource == FogResource)
+		{
+			delete AtmosphericFog;
+			AtmosphericFog = NULL;
+		}
+	}
+}
diff --git a/Engine/Source/Runtime/Renderer/Private/AtmosphereRendering.h b/Engine/Source/Runtime/Renderer/Private/AtmosphereRendering.h
--- a/Engine/Source/Runtime/Renderer/Private/AtmosphereRendering.h
+++ b/Engine/Source/Runtime/Renderer/Private/AtmosphereRendering.h
@@ -24,7 +24,7 @@ namespace EAtmosphereRenderFlag
 }
 
 /** The properties of a atmospheric fog layer which are used for rendering. */
-class FAtmosphericFogSceneInfo
+class FAtmosphericFogSceneInfo : public FRenderResource
 {
 public:
 	/** The fog component the scene info is for. */
@@ -44,16 +44,22 @@ public:
 	FVector DefaultSunDirection;
 	uint32 RenderFlag;
 	uint32 InscatterAltitudeSampleNum;
+	class FAtmosphereTextureResource* TransmittanceResource;
+	class FAtmosphereTextureResource* IrradianceResource;
+	class FAtmosphereTextureResource* InscatterResource;
+
 
 #if WITH_EDITORONLY_DATA
 	/** Atmosphere pre-computation related data */
 	bool bNeedRecompute;
+	bool bPrecomputationStarted;
+	bool bPrecomputationFinished;
+	bool bPrecomputationAcceptedByGameThread;
 	int32 MaxScatteringOrder;
 	int32 AtmospherePhase;
 	int32 Atmosphere3DTextureIndex;
 	int32 AtmoshpereOrder;
 	class FAtmosphereTextures* AtmosphereTextures;
-	class FThreadSafeCounter& PrecomputeCounter;
 	FByteBulkData PrecomputeTransmittance;
 	FByteBulkData PrecomputeIrradiance;
 	FByteBulkData PrecomputeInscatter;
diff --git a/Engine/Source/Runtime/Renderer/Private/RendererScene.cpp b/Engine/Source/Runtime/Renderer/Private/RendererScene.cpp
--- a/Engine/Source/Runtime/Renderer/Private/RendererScene.cpp
+++ b/Engine/Source/Runtime/Renderer/Private/RendererScene.cpp
@@ -1428,50 +1428,6 @@ void FScene::RemoveExponentialHeightFog(UExponentialHeightFogComponent* FogCompo
 		});
 }
 
-void FScene::AddAtmosphericFog(UAtmosphericFogComponent* FogComponent)
-{
-	check(FogComponent);
-
-	FAtmosphericFogSceneInfo* FogSceneInfo = new FAtmosphericFogSceneInfo(FogComponent, this);
-
-	ENQUEUE_UNIQUE_RENDER_COMMAND_TWOPARAMETER(
-		FAddAtmosphericFogCommand,
-		FScene*,Scene,this,
-		FAtmosphericFogSceneInfo*,FogSceneInfo,FogSceneInfo,
-	{
-		if (Scene->AtmosphericFog && Scene->AtmosphericFog->Component != FogSceneInfo->Component)
-		{
-			delete Scene->AtmosphericFog;
-			Scene->AtmosphericFog = NULL;
-		}
-
-		if (Scene->AtmosphericFog == NULL)
-		{
-			Scene->AtmosphericFog = FogSceneInfo;
-		}
-		else
-		{
-			delete FogSceneInfo;
-		}
-	});
-}
-
-void FScene::RemoveAtmosphericFog(UAtmosphericFogComponent* FogComponent)
-{
-	ENQUEUE_UNIQUE_RENDER_COMMAND_TWOPARAMETER(
-		FRemoveAtmosphericFogCommand,
-		FScene*,Scene,this,
-		UAtmosphericFogComponent*,FogComponent,FogComponent,
-	{
-		// Remove the given component's FExponentialHeightFogSceneInfo from the scene's fog array.
-		if (Scene->AtmosphericFog && Scene->AtmosphericFog->Component == FogComponent)
-		{
-			delete Scene->AtmosphericFog;
-			Scene->AtmosphericFog = NULL;
-		}
-	});
-}
-
 void FScene::AddWindSource(UWindDirectionalSourceComponent* WindComponent)
 {
 	// if this wind component is not activated (or Auto Active is set to false), then don't add to WindSources
@@ -2385,6 +2341,7 @@ public:
 	virtual void RemoveExponentialHeightFog(class UExponentialHeightFogComponent* FogComponent) override {}
 	virtual void AddAtmosphericFog(class UAtmosphericFogComponent* FogComponent) override {}
 	virtual void RemoveAtmosphericFog(class UAtmosphericFogComponent* FogComponent) override {}
+	virtual void RemoveAtmosphericFogResource_RenderThread(FRenderResource* FogResource) override {}
 	virtual FAtmosphericFogSceneInfo* GetAtmosphericFogSceneInfo() override { return NULL; }
 	virtual void AddWindSource(class UWindDirectionalSourceComponent* WindComponent) override {}
 	virtual void RemoveWindSource(class UWindDirectionalSourceComponent* WindComponent) override {}
diff --git a/Engine/Source/Runtime/Renderer/Private/ScenePrivate.h b/Engine/Source/Runtime/Renderer/Private/ScenePrivate.h
--- a/Engine/Source/Runtime/Renderer/Private/ScenePrivate.h
+++ b/Engine/Source/Runtime/Renderer/Private/ScenePrivate.h
@@ -1644,6 +1644,7 @@ public:
 	virtual void RemoveExponentialHeightFog(UExponentialHeightFogComponent* FogComponent) override;
 	virtual void AddAtmosphericFog(UAtmosphericFogComponent* FogComponent) override;
 	virtual void RemoveAtmosphericFog(UAtmosphericFogComponent* FogComponent) override;
+	virtual void RemoveAtmosphericFogResource_RenderThread(FRenderResource* FogResource) override;
 	virtual FAtmosphericFogSceneInfo* GetAtmosphericFogSceneInfo() override { return AtmosphericFog; }
 	virtual void AddWindSource(UWindDirectionalSourceComponent* WindComponent) override;
 	virtual void RemoveWindSource(UWindDirectionalSourceComponent* WindComponent) override;
