Fixed mismatch in display category between RecastNavMesh's PolyRefTileBits, PolyRefNavPolyBits, PolyRefSaltBits #UE4

All of them should be in the same category, otherwise this preview doesn't make sense.

diff --git a/Engine/Source/Runtime/Engine/Classes/AI/Navigation/RecastNavMesh.h b/Engine/Source/Runtime/Engine/Classes/AI/Navigation/RecastNavMesh.h
--- a/Engine/Source/Runtime/Engine/Classes/AI/Navigation/RecastNavMesh.h
+++ b/Engine/Source/Runtime/Engine/Classes/AI/Navigation/RecastNavMesh.h
@@ -546,7 +546,7 @@ class ENGINE_API ARecastNavMesh : public ANavigationData
 	UPROPERTY(EditAnywhere, Category = Generation, config, meta = (ClampMin = "1", UIMin = "1"), AdvancedDisplay)
 	int32 TileNumberHardLimit;
 
-	UPROPERTY(VisibleAnywhere, Category = "Generation|Poly Ref Bits", AdvancedDisplay)
+	UPROPERTY(VisibleAnywhere, Category = Generation, AdvancedDisplay)
 	int32 PolyRefTileBits;
 
 	UPROPERTY(VisibleAnywhere, Category = Generation, AdvancedDisplay)
