Removed old code that shouldnt have been left in

diff --git a/Engine/Source/Runtime/Engine/Private/Components/SkeletalMeshComponent.cpp b/Engine/Source/Runtime/Engine/Private/Components/SkeletalMeshComponent.cpp
--- a/Engine/Source/Runtime/Engine/Private/Components/SkeletalMeshComponent.cpp
+++ b/Engine/Source/Runtime/Engine/Private/Components/SkeletalMeshComponent.cpp
@@ -1821,14 +1821,6 @@ FTransform USkeletalMeshComponent::ConvertLocalRootMotionToWorld(const FTransfor
 		UpdateComponentToWorld();
 	}
 
-	/*check(ComponentToWorld.IsRotationNormalized());
-	check(InTransform.IsRotationNormalized());
-	const FTransform NewWorldTransform = InTransform * ComponentToWorld;
-	check(NewWorldTransform.IsRotationNormalized());
-	const FVector DeltaWorldTranslation = NewWorldTransform.GetTranslation() - ComponentToWorld.GetTranslation();
-	const FQuat DeltaWorldRotation = ComponentToWorld.GetRotation().Inverse() * NewWorldTransform.GetRotation();
-	check(DeltaWorldRotation.IsNormalized());*/
-
 	const FTransform NewWorldTransform = InTransform * ComponentToWorld;
 	const FQuat NewWorldRotation = ComponentToWorld.GetRotation() * InTransform.GetRotation();
 	const FVector DeltaWorldTranslation = NewWorldTransform.GetTranslation() - ComponentToWorld.GetTranslation();
