Fix bad assumption about threads which Slate can run on causing crashes when a loading screen is opened while blocked waiting for texture streaming to flush

diff --git a/Engine/Source/Runtime/Engine/Private/Slate/SceneViewport.cpp b/Engine/Source/Runtime/Engine/Private/Slate/SceneViewport.cpp
--- a/Engine/Source/Runtime/Engine/Private/Slate/SceneViewport.cpp
+++ b/Engine/Source/Runtime/Engine/Private/Slate/SceneViewport.cpp
@@ -877,7 +877,7 @@ void FSceneViewport::OnViewportClosed()
 
 FSlateShaderResource* FSceneViewport::GetViewportRenderTargetTexture() const
 { 
-	check(IsInGameThread());
+	check(IsThreadSafeForSlateRendering());
 	return (BufferedSlateHandles.Num() != 0) ? BufferedSlateHandles[CurrentBufferedTargetIndex] : nullptr;
 }
 
