Fixed crash when retargetting single frame animation assets.

diff --git a/Engine/Source/Runtime/Engine/Private/Animation/AnimSequence.cpp b/Engine/Source/Runtime/Engine/Private/Animation/AnimSequence.cpp
--- a/Engine/Source/Runtime/Engine/Private/Animation/AnimSequence.cpp
+++ b/Engine/Source/Runtime/Engine/Private/Animation/AnimSequence.cpp
@@ -1581,12 +1581,6 @@ bool UAnimSequence::CompressRawAnimData(float MaxPosDiff, float MaxAngleDiff)
 		}
 	}
 
-	// Only bother doing anything if we have some keys!
-	if( NumFrames == 1 )
-	{
-		return bRemovedKeys;
-	}
-
 	CompressedTrackOffsets.Empty();
 	CompressedScaleOffsets.Empty();
 #endif
