Linear character movement component smoothing during replays

diff --git a/Engine/Source/Runtime/Engine/Classes/GameFramework/CharacterMovementComponent.h b/Engine/Source/Runtime/Engine/Classes/GameFramework/CharacterMovementComponent.h
--- a/Engine/Source/Runtime/Engine/Classes/GameFramework/CharacterMovementComponent.h
+++ b/Engine/Source/Runtime/Engine/Classes/GameFramework/CharacterMovementComponent.h
@@ -2012,6 +2012,21 @@ public:
 	/** Used for rotation smoothing in net games */
 	FQuat MeshRotationOffset;
 
+	/** Used for position smoothing in net games */
+	FVector OriginalMeshTranslationOffset;
+
+	/** Used for rotation smoothing in net games */
+	FQuat OriginalMeshRotationOffset;
+
+	/** Used for remembering how much time has passed between server corrections */
+	float LastCorrectionDelta;
+
+	/** Used to track how much time has elapsed since last correction */
+	float CurrentSmoothTime;
+
+	/** Used to signify that linear smoothing is desired */
+	bool bUseLinearSmoothing;
+
 	/** Maximum location correction distance for which other pawn positions on a client will be smoothly updated */
 	float MaxSmoothNetUpdateDist;
 
@@ -2019,10 +2034,12 @@ public:
 	If it is between MaxSmoothNetUpdateDist and NoSmoothNetUpdateDist, pop to MaxSmoothNetUpdateDist away from the updated location */
 	float NoSmoothNetUpdateDist;
 
-	/** How long to take to smoothly interpolate from the old pawn position on the client to the corrected one sent by the server.  Must be >= 0.0 */
+	/** How long to take to smoothly interpolate from the old pawn position on the client to the corrected one sent by the server.  Must be >= 0.0 
+	This variable isn't used when bUseLinearSmoothing = true */
 	float SmoothNetUpdateTime;
 
-	/** How long to take to smoothly interpolate from the old pawn rotation on the client to the corrected one sent by the server.  Must be >= 0.0 */
+	/** How long to take to smoothly interpolate from the old pawn rotation on the client to the corrected one sent by the server.  Must be >= 0.0
+	This variable isn't used when bUseLinearSmoothing = true */
 	float SmoothNetUpdateRotationTime;
 	
 	// how long server will wait for client move update before setting position
diff --git a/Engine/Source/Runtime/Engine/Private/Components/CharacterMovementComponent.cpp b/Engine/Source/Runtime/Engine/Private/Components/CharacterMovementComponent.cpp
--- a/Engine/Source/Runtime/Engine/Private/Components/CharacterMovementComponent.cpp
+++ b/Engine/Source/Runtime/Engine/Private/Components/CharacterMovementComponent.cpp
@@ -5623,6 +5623,14 @@ void UCharacterMovementComponent::SmoothCorrection(const FVector& OldLocation, c
 
 		// Take difference between where we were rotated before, and where we're going
 		ClientData->MeshRotationOffset = (UpdatedComponent->GetComponentQuat().Inverse() * OldRotation) * ClientData->MeshRotationOffset;
+
+		if (ClientData->bUseLinearSmoothing)
+		{
+			ClientData->OriginalMeshTranslationOffset	= ClientData->MeshTranslationOffset;
+			ClientData->OriginalMeshRotationOffset		= ClientData->MeshRotationOffset;
+			ClientData->LastCorrectionDelta				= ClientData->CurrentSmoothTime;
+			ClientData->CurrentSmoothTime				= 0;
+		}
 	}
 }
 
@@ -5637,26 +5645,51 @@ void UCharacterMovementComponent::SmoothClientPosition(float DeltaSeconds)
 	FNetworkPredictionData_Client_Character* ClientData = GetPredictionData_Client_Character();
 	if (ClientData && ClientData->bSmoothNetUpdates && CharacterOwner->GetMesh() && !CharacterOwner->GetMesh()->IsSimulatingPhysics())
 	{
-		// smooth interpolation of mesh translation to avoid popping of other client pawns unless under a low tick rate
-		if (DeltaSeconds < ClientData->SmoothNetUpdateTime)
+		if (ClientData->bUseLinearSmoothing)
 		{
-			// Slowly decay translation offset
-			ClientData->MeshTranslationOffset = (ClientData->MeshTranslationOffset * (1.f - DeltaSeconds / ClientData->SmoothNetUpdateTime));
-		}
-		else
-		{
-			ClientData->MeshTranslationOffset = FVector::ZeroVector;
-		}
+			ClientData->CurrentSmoothTime += DeltaSeconds;
 
-		// Smooth rotation
-		if (DeltaSeconds < ClientData->SmoothNetUpdateRotationTime)
-		{
-			// Slowly decay rotation offset
-			ClientData->MeshRotationOffset = FQuat::Slerp( ClientData->MeshRotationOffset, FQuat::Identity, DeltaSeconds / ClientData->SmoothNetUpdateRotationTime );
+			if (ClientData->LastCorrectionDelta < SMALL_NUMBER || ClientData->CurrentSmoothTime >= ClientData->LastCorrectionDelta)
+			{
+				// This is either:
+				//	1. The very first update
+				//	2. Time between updates was really small
+				//  3. We've arrived at the target position
+				ClientData->MeshTranslationOffset	= FVector::ZeroVector;
+				ClientData->MeshRotationOffset		= FQuat::Identity;
+			}
+			else
+			{
+				// Linearly interpolate between correction updates
+				const float LerpPercent	= ClientData->CurrentSmoothTime / ClientData->LastCorrectionDelta;
+
+				ClientData->MeshTranslationOffset	= FMath::Lerp( ClientData->OriginalMeshTranslationOffset, FVector::ZeroVector, LerpPercent );
+				ClientData->MeshRotationOffset		= FQuat::Slerp( ClientData->OriginalMeshRotationOffset, FQuat::Identity, LerpPercent );
+			}
 		}
 		else
 		{
-			ClientData->MeshRotationOffset = FQuat::Identity;
+			// smooth interpolation of mesh translation to avoid popping of other client pawns unless under a low tick rate
+			if (DeltaSeconds < ClientData->SmoothNetUpdateTime)
+			{
+				// Slowly decay translation offset
+				ClientData->MeshTranslationOffset = (ClientData->MeshTranslationOffset * (1.f - DeltaSeconds / ClientData->SmoothNetUpdateTime));
+			}
+			else
+			{
+				ClientData->MeshTranslationOffset = FVector::ZeroVector;
+			}
+
+			// Smooth rotation
+			if (DeltaSeconds < ClientData->SmoothNetUpdateRotationTime)
+			{
+				// Slowly decay rotation offset
+				ClientData->MeshRotationOffset = FQuat::Slerp( ClientData->MeshRotationOffset, FQuat::Identity, DeltaSeconds / ClientData->SmoothNetUpdateRotationTime );
+			}
+			else
+			{
+				ClientData->MeshRotationOffset = FQuat::Identity;
+			}
 		}
 
 		if (IsMovingOnGround())
@@ -7077,6 +7110,11 @@ FNetworkPredictionData_Client_Character::FNetworkPredictionData_Client_Character
 	, bSmoothNetUpdates(false)
 	, MeshTranslationOffset(ForceInitToZero)
 	, MeshRotationOffset(FQuat::Identity)
+	, OriginalMeshTranslationOffset(ForceInitToZero)
+	, OriginalMeshRotationOffset(FQuat::Identity)
+	, LastCorrectionDelta(0.f)
+	, CurrentSmoothTime(0.f)
+	, bUseLinearSmoothing(false)
 	, MaxSmoothNetUpdateDist(0.f)
 	, NoSmoothNetUpdateDist(0.f)
 	, SmoothNetUpdateTime(0.f)
