Turn on linear smoothing for all characters when in replays

diff --git a/Engine/Source/Runtime/Engine/Private/Components/CharacterMovementComponent.cpp b/Engine/Source/Runtime/Engine/Private/Components/CharacterMovementComponent.cpp
--- a/Engine/Source/Runtime/Engine/Private/Components/CharacterMovementComponent.cpp
+++ b/Engine/Source/Runtime/Engine/Private/Components/CharacterMovementComponent.cpp
@@ -22,6 +22,8 @@
 #include "Components/BrushComponent.h"
 #include "Components/DestructibleComponent.h"
 
+#include "Engine/DemoNetDriver.h"
+
 DEFINE_LOG_CATEGORY_STATIC(LogCharacterMovement, Log, All);
 DEFINE_LOG_CATEGORY_STATIC(LogNavMeshMovement, Log, All);
 
@@ -5609,6 +5611,18 @@ void UCharacterMovementComponent::SmoothCorrection(const FVector& OldLocation, c
 	FNetworkPredictionData_Client_Character* ClientData = GetPredictionData_Client_Character();
 	if (ClientData && ClientData->bSmoothNetUpdates)
 	{
+		if (!ClientData->bUseLinearSmoothing && GetWorld() && GetWorld()->DemoNetDriver && GetWorld()->DemoNetDriver->ServerConnection)
+		{
+			FNetworkPredictionData_Client_Character* ClientData = GetPredictionData_Client_Character();
+
+			ClientData->bUseLinearSmoothing	= true;
+
+			// Really large since we want to smooth most of the time during playback
+			// We do however want to compensate for large deltas (like teleporting, etc)
+			ClientData->MaxSmoothNetUpdateDist	= 512.0f;
+			ClientData->NoSmoothNetUpdateDist	= 512.0f;
+		}
+
 		float DistSq = (OldLocation - UpdatedComponent->GetComponentLocation()).SizeSquared();
 		if (DistSq > FMath::Square(ClientData->MaxSmoothNetUpdateDist))
 		{
