#ue4 - Revert change to UCharacterMovementComponent::SmoothCorrection() meant to fix offsets in the character mesh, it's causing artifacts (change 2546662). Offset is in root space, not mesh.

diff --git a/Engine/Source/Runtime/Engine/Private/Components/CharacterMovementComponent.cpp b/Engine/Source/Runtime/Engine/Private/Components/CharacterMovementComponent.cpp
--- a/Engine/Source/Runtime/Engine/Private/Components/CharacterMovementComponent.cpp
+++ b/Engine/Source/Runtime/Engine/Private/Components/CharacterMovementComponent.cpp
@@ -5712,7 +5712,7 @@ void UCharacterMovementComponent::SmoothClientPosition(float DeltaSeconds)
 			ClientData->MeshTranslationOffset.Z = 0;
 		}
 
-		const FVector NewRelTranslation = CharacterOwner->GetMesh()->GetComponentToWorld().InverseTransformVectorNoScale(ClientData->MeshTranslationOffset) + CharacterOwner->GetBaseTranslationOffset();
+		const FVector NewRelTranslation = UpdatedComponent->GetComponentToWorld().InverseTransformVectorNoScale(ClientData->MeshTranslationOffset) + CharacterOwner->GetBaseTranslationOffset();
 		const FQuat NewRelRotation = ClientData->MeshRotationOffset * CharacterOwner->GetBaseRotationOffset();
 		CharacterOwner->GetMesh()->SetRelativeLocationAndRotation(NewRelTranslation, NewRelRotation);
 	}
