Changed the USceneComponent::PhysicsVolume property to a TWeakObjectPtr.

Previously, the PhysicsVolume pointer was creating a dependency between sublevels and the persistent level as its default value was the default physics volume of the persistent level, no matter which level the component was in.  This was causing GC fails when loading a sublevel which was already loaded.  Changing to a TWeakObjectPtr breaks this dependency.

diff --git a/Engine/Source/Runtime/Engine/Classes/Components/SceneComponent.h b/Engine/Source/Runtime/Engine/Classes/Components/SceneComponent.h
--- a/Engine/Source/Runtime/Engine/Classes/Components/SceneComponent.h
+++ b/Engine/Source/Runtime/Engine/Classes/Components/SceneComponent.h
@@ -167,9 +167,10 @@ protected:
 	// Transient flag that temporarily disables UpdateOverlaps within DetachFromParent().
 	uint32 bDisableDetachmentUpdateOverlaps:1;
 
+private:
 	/** Physics Volume in which this SceneComponent is located **/
 	UPROPERTY(transient)
-	class APhysicsVolume* PhysicsVolume;
+	TWeakObjectPtr<class APhysicsVolume> PhysicsVolume;
 
 public:
 	/** Current transform of this component, relative to the world */
diff --git a/Engine/Source/Runtime/Engine/Private/Components/PrimitiveComponent.cpp b/Engine/Source/Runtime/Engine/Private/Components/PrimitiveComponent.cpp
--- a/Engine/Source/Runtime/Engine/Private/Components/PrimitiveComponent.cpp
+++ b/Engine/Source/Runtime/Engine/Private/Components/PrimitiveComponent.cpp
@@ -2306,7 +2306,7 @@ void UPrimitiveComponent::UpdatePhysicsVolume( bool bTriggerNotifiers )
 				}
 			}
 
-			if (BestVolume != PhysicsVolume)
+			if (BestVolume != GetPhysicsVolume())
 			{
 				SetPhysicsVolume(BestVolume, bTriggerNotifiers);
 			}
diff --git a/Engine/Source/Runtime/Engine/Private/Components/SceneComponent.cpp b/Engine/Source/Runtime/Engine/Private/Components/SceneComponent.cpp
--- a/Engine/Source/Runtime/Engine/Private/Components/SceneComponent.cpp
+++ b/Engine/Source/Runtime/Engine/Private/Components/SceneComponent.cpp
@@ -1697,9 +1697,9 @@ bool USceneComponent::IsAnySimulatingPhysics() const
 
 APhysicsVolume* USceneComponent::GetPhysicsVolume() const
 {
-	if (PhysicsVolume)
+	if (PhysicsVolume.IsValid())
 	{
-		return PhysicsVolume;
+		return PhysicsVolume.Get();
 	}
 	else if (const UWorld* MyWorld = GetWorld())
 	{
@@ -1801,13 +1801,13 @@ void USceneComponent::SetPhysicsVolume( APhysicsVolume * NewVolume,  bool bTrigg
 		if( NewVolume != PhysicsVolume )
 		{
 			AActor *A = GetOwner();
-			if( PhysicsVolume )
+			if( PhysicsVolume.IsValid() )
 			{
 				PhysicsVolume->ActorLeavingVolume(A);
 				PhysicsVolumeChangedDelegate.Broadcast(NewVolume);
 			}
 			PhysicsVolume = NewVolume;
-			if( PhysicsVolume )
+			if( PhysicsVolume.IsValid() )
 			{
 				PhysicsVolume->ActorEnteredVolume(A);
 			}
