#ue4 - Fix shadowed variable from CL 2549347.

diff --git a/Engine/Source/Runtime/Engine/Private/Components/CharacterMovementComponent.cpp b/Engine/Source/Runtime/Engine/Private/Components/CharacterMovementComponent.cpp
--- a/Engine/Source/Runtime/Engine/Private/Components/CharacterMovementComponent.cpp
+++ b/Engine/Source/Runtime/Engine/Private/Components/CharacterMovementComponent.cpp
@@ -5613,8 +5613,6 @@ void UCharacterMovementComponent::SmoothCorrection(const FVector& OldLocation, c
 	{
 		if (!ClientData->bUseLinearSmoothing && GetWorld() && GetWorld()->DemoNetDriver && GetWorld()->DemoNetDriver->ServerConnection)
 		{
-			FNetworkPredictionData_Client_Character* ClientData = GetPredictionData_Client_Character();
-
 			ClientData->bUseLinearSmoothing	= true;
 
 			// Really large since we want to smooth most of the time during playback
