#ue4 - Network character proxies only need to check for encroachment when they teleport with zero velocity. Normal movement handles encroachment already.

- Removes uncessary encroach check every time a simulated proxy position or rotation replication is received.

diff --git a/Engine/Source/Runtime/Engine/Private/Character.cpp b/Engine/Source/Runtime/Engine/Private/Character.cpp
--- a/Engine/Source/Runtime/Engine/Private/Character.cpp
+++ b/Engine/Source/Runtime/Engine/Private/Character.cpp
@@ -1167,13 +1167,19 @@ void ACharacter::UpdateSimulatedPosition(const FVector& NewLocation, const FRota
 	if( (NewLocation != GetActorLocation()) || (CreationTime == GetWorld()->TimeSeconds) )
 	{
 		FVector FinalLocation = NewLocation;
-		if( GetWorld()->EncroachingBlockingGeometry(this, NewLocation, NewRotation) )
-		{
-			bSimGravityDisabled = true;
-		}
-		else
+
+		// Only need to check for encroachment when teleported without any velocity.
+		// Normal movement pops the character out of geometry anyway, no use doing it before and after (with different rules).
+		if (CharacterMovement && CharacterMovement->Velocity.IsZero())
 		{
-			bSimGravityDisabled = false;
+			if (GetWorld()->EncroachingBlockingGeometry(this, NewLocation, NewRotation))
+			{
+				bSimGravityDisabled = true;
+			}
+			else
+			{
+				bSimGravityDisabled = false;
+			}
 		}
 		
 		// Don't use TeleportTo(), that clears our base.
