Fixed crash in ClientReceiveLocalizedMessage() if the received message type is invalid (bad call or networking serialization failure)

Merging 2549599

//depot/UE4-UT/...

to //depot/UE4/...

diff --git a/Engine/Source/Runtime/Engine/Private/PlayerController.cpp b/Engine/Source/Runtime/Engine/Private/PlayerController.cpp
--- a/Engine/Source/Runtime/Engine/Private/PlayerController.cpp
+++ b/Engine/Source/Runtime/Engine/Private/PlayerController.cpp
@@ -1284,7 +1284,7 @@ void APlayerController::OnNetCleanup(UNetConnection* Connection)
 void APlayerController::ClientReceiveLocalizedMessage_Implementation( TSubclassOf<ULocalMessage> Message, int32 Switch, APlayerState* RelatedPlayerState_1, APlayerState* RelatedPlayerState_2, UObject* OptionalObject )
 {
 	// Wait for player to be up to date with replication when joining a server, before stacking up messages
-	if ( GetNetMode() == NM_DedicatedServer || GetWorld()->GameState == NULL )
+	if (GetNetMode() == NM_DedicatedServer || GetWorld()->GameState == NULL || Message == NULL)
 	{
 		return;
 	}
