Add optional parameter to GetLocalNetworkVersion to skip delegate call.
Allows project implementations to fall back to engine implementation without causing infinite loop.

diff --git a/Engine/Source/Runtime/Engine/Private/UnrealNetwork.cpp b/Engine/Source/Runtime/Engine/Private/UnrealNetwork.cpp
--- a/Engine/Source/Runtime/Engine/Private/UnrealNetwork.cpp
+++ b/Engine/Source/Runtime/Engine/Private/UnrealNetwork.cpp
@@ -21,9 +21,9 @@ enum ENetworkVersionHistory
 
 const uint32 FNetworkVersion::InternalProtocolVersion = HISTORY_REPLAY_RPC_CHECKSUMS;
 
-uint32 FNetworkVersion::GetLocalNetworkVersion()
+uint32 FNetworkVersion::GetLocalNetworkVersion(bool AllowOverrideDelegate /*=true*/)
 {
-	if ( GetLocalNetworkVersionOverride.IsBound() )
+	if ( AllowOverrideDelegate && GetLocalNetworkVersionOverride.IsBound() )
 	{
 		const uint32 LocalNetworkVersion = GetLocalNetworkVersionOverride.Execute();
 
diff --git a/Engine/Source/Runtime/Engine/Public/Net/UnrealNetwork.h b/Engine/Source/Runtime/Engine/Public/Net/UnrealNetwork.h
--- a/Engine/Source/Runtime/Engine/Public/Net/UnrealNetwork.h
+++ b/Engine/Source/Runtime/Engine/Public/Net/UnrealNetwork.h
@@ -52,8 +52,9 @@ struct ENGINE_API FNetworkVersion
 	/**
 	 * Generates a version number, that by default, is based on a checksum of the engine version + project name + project version string
 	 * Game/project code can completely override what this value returns through the GetLocalNetworkVersionOverride delegate
+	 * If called with AllowOverrideDelegate=false, we will not call the game project override. (This allows projects to call base implementation in their project implementation)
 	 */
-	static uint32 GetLocalNetworkVersion();
+	static uint32 GetLocalNetworkVersion(bool AllowOverrideDelegate=true);
 
 	/**
 	 * Determine if a connection is compatible with this instance
