[UE-14774] Post Edit can be called before we have ControlPoints

diff --git a/Engine/Source/Runtime/Engine/Private/Components/InterpToMovementComponent.cpp b/Engine/Source/Runtime/Engine/Private/Components/InterpToMovementComponent.cpp
--- a/Engine/Source/Runtime/Engine/Private/Components/InterpToMovementComponent.cpp
+++ b/Engine/Source/Runtime/Engine/Private/Components/InterpToMovementComponent.cpp
@@ -472,6 +472,10 @@ void UInterpToMovementComponent::AddControlPointPosition(FVector Pos)
 void UInterpToMovementComponent::PostEditChangeProperty(FPropertyChangedEvent& PropertyChangedEvent)
 {
 	Super::PostEditChangeProperty(PropertyChangedEvent);
-	UpdateControlPoints(true);
+
+	if (ControlPoints.Num() != 0)
+	{
+		UpdateControlPoints(true);
+	}
 }
 #endif // WITH_EDITOR
