Fixed: Option to enable HLOD is hidden and separate from related options
UE-15403

diff --git a/Engine/Source/Runtime/Engine/Classes/GameFramework/WorldSettings.h b/Engine/Source/Runtime/Engine/Classes/GameFramework/WorldSettings.h
--- a/Engine/Source/Runtime/Engine/Classes/GameFramework/WorldSettings.h
+++ b/Engine/Source/Runtime/Engine/Classes/GameFramework/WorldSettings.h
@@ -259,11 +259,6 @@ class ENGINE_API AWorldSettings : public AInfo, public IInterface_AssetUserData
 	UPROPERTY(EditAnywhere, BlueprintReadOnly, config, Category=World, AdvancedDisplay)
 	uint32 bEnableNavigationSystem:1;
 
-#if WITH_EDITORONLY_DATA
-	/** if set to true, hierarchical LODs will be built, which will create hierarchical LODActors*/
-	UPROPERTY(EditAnywhere, config, Category=World, AdvancedDisplay)
-	uint32 bEnableHierarchicalLODSystem:1;
-#endif
 	/** 
 	 * Enables tools for composing a tiled world. 
 	 * Level has to be saved and all sub-levels removed before enabling this option.
@@ -402,8 +397,12 @@ class ENGINE_API AWorldSettings : public AInfo, public IInterface_AssetUserData
 	class USoundMix* DefaultBaseSoundMix;
 
 #if WITH_EDITORONLY_DATA
+	/** if set to true, hierarchical LODs will be built, which will create hierarchical LODActors*/
+	UPROPERTY(EditAnywhere, config, Category=LODSystem)
+	uint32 bEnableHierarchicalLODSystem:1;
+
 	/** Hierarchical LOD Setup */
-	UPROPERTY(EditAnywhere, Category=LODSystem, AdvancedDisplay, meta=(editcondition = "bEnableHierarchicalLODSystem"))
+	UPROPERTY(EditAnywhere, Category=LODSystem, meta=(editcondition = "bEnableHierarchicalLODSystem"))
 	TArray<struct FHierarchicalSimplification>	HierarchicalLODSetup;
 #endif
 	/************************************/
