Fixed: HLOD actors can be moved by changing transform in Details panel
UE-14946

diff --git a/Engine/Source/Runtime/Engine/Classes/Engine/LODActor.h b/Engine/Source/Runtime/Engine/Classes/Engine/LODActor.h
--- a/Engine/Source/Runtime/Engine/Classes/Engine/LODActor.h
+++ b/Engine/Source/Runtime/Engine/Classes/Engine/LODActor.h
@@ -10,7 +10,7 @@
  * @see https://docs.unrealengine.com/latest/INT/Engine/Actors/LODActor/
  * @see UStaticMesh
  */
-UCLASS(notplaceable)
+UCLASS(notplaceable, hidecategories=(Object, Collision, Display, Input, Blueprint, Transform))
 class ENGINE_API ALODActor : public AActor
 {
 	GENERATED_UCLASS_BODY()
