#ue4 - Fix FBodyInstance::OverlapTestForBodies() only working against RigidBodies, not any RigidActor, which broke ComponentOverlapComponent() in many cases.

diff --git a/Engine/Source/Runtime/Engine/Private/PhysicsEngine/BodyInstance.cpp b/Engine/Source/Runtime/Engine/Private/PhysicsEngine/BodyInstance.cpp
--- a/Engine/Source/Runtime/Engine/Private/PhysicsEngine/BodyInstance.cpp
+++ b/Engine/Source/Runtime/Engine/Private/PhysicsEngine/BodyInstance.cpp
@@ -4017,15 +4017,15 @@ bool FBodyInstance::OverlapTestForBodies(const FVector& Pos, const FQuat& Rot, c
 	bool bHaveOverlap = false;
 #if WITH_PHYSX
 
-	ExecuteOnPxRigidBodyReadOnly(this, [&] (const PxRigidBody* TargetRigidBody)
+	ExecuteOnPxRigidActorReadOnly(this, [&] (const PxRigidActor* TargetRigidActor)
 	{
 		// calculate the test global pose of the rigid body
 		PxTransform PTestGlobalPose = U2PTransform(FTransform(Rot, Pos));
 
 		// Get all the shapes from the actor
 		TArray<PxShape*, TInlineAllocator<16>> PTargetShapes;
-		PTargetShapes.AddUninitialized(TargetRigidBody->getNbShapes());
-		int32 NumTargetShapes = TargetRigidBody->getShapes(PTargetShapes.GetData(), PTargetShapes.Num());
+		PTargetShapes.AddUninitialized(TargetRigidActor->getNbShapes());
+		int32 NumTargetShapes = TargetRigidActor->getShapes(PTargetShapes.GetData(), PTargetShapes.Num());
 
 		for (int32 TargetShapeIdx = 0; TargetShapeIdx < PTargetShapes.Num(); ++TargetShapeIdx)
 		{
