Added BindClothToMasterPoseComponent node to allow cloth simulations to be bound to cloth elements in master pose components rather than run their own simulation.

diff --git a/Engine/Source/Runtime/Engine/Classes/Components/SkeletalMeshComponent.h b/Engine/Source/Runtime/Engine/Classes/Components/SkeletalMeshComponent.h
--- a/Engine/Source/Runtime/Engine/Classes/Components/SkeletalMeshComponent.h
+++ b/Engine/Source/Runtime/Engine/Classes/Components/SkeletalMeshComponent.h
@@ -611,6 +611,25 @@ public:
 	UFUNCTION(BlueprintCallable, Category="Components|SkeletalMesh")
 	void ResetClothTeleportMode();
 
+	/** 
+	 * If this component has a valid MasterPoseComponent then this function makes cloth items on the slave component
+	 * take the transforms of the cloth items on the master component instead of simulating separately.
+	 * @Note This will FORCE any cloth actor on the master component to simulate in local space. Also
+	 * The meshes used in the components must be identical for the cloth to bind correctly
+	 */
+	UFUNCTION(BlueprintCallable, Category="Components|SkeletalMesh", meta=(UnsafeDuringActorConstruction="true"))
+	void BindClothToMasterPoseComponent();
+
+	/** 
+	 * If this component has a valid MasterPoseComponent and has previously had its cloth bound to the
+	 * MCP, this function will unbind the cloth and resume simulation.
+	 * @param bRestoreSimulationSpace if true and the master pose cloth was originally simulating in world
+	 * space, we will restore this setting. This will cause the master component to reset which may be
+	 * undesirable.
+	 */
+	UFUNCTION(BlueprintCallable, Category="Components|SkeletalMesh", meta=(UnsafeDuringActorConstruction="true"))
+	void UnbindClothFromMasterPoseComponent(bool bRestoreSimulationSpace = true);
+
 	/** We detach the Component once we are done playing it.
 	 *
 	 * @param	ParticleSystemComponent that finished
@@ -697,6 +716,11 @@ public:
 	/** precomputed actual cloth morph target data */
 	TArray<FClothMorphTargetData> ClothMorphTargets;
 
+	/** Whether or not we're taking cloth sim information from our master component */
+	bool bBindClothToMasterComponent;
+	/** The previous state of the master component simulation coord space, so we can restore on unbind */
+	bool bPrevMasterSimulateLocalSpace;
+
 #if WITH_CLOTH_COLLISION_DETECTION
 	/** increase every tick to update clothing collision  */
 	uint32 ClothingCollisionRevision; 
@@ -728,7 +752,7 @@ public:
 	void TickClothing(float DeltaTime);
 
 	/** Store cloth simulation data into OutClothSimData */
-	void GetUpdateClothSimulationData(TArray<FClothSimulData>& OutClothSimData);
+	void GetUpdateClothSimulationData(TArray<FClothSimulData>& OutClothSimData, USkeletalMeshComponent* OverrideLocalRootComponent = nullptr);
 	void ApplyWindForCloth(FClothingActor& ClothingActor);
 	void RemoveAllClothingActors();
 	void ReleaseAllClothingResources();
diff --git a/Engine/Source/Runtime/Engine/Private/Components/SkeletalMeshComponent.cpp b/Engine/Source/Runtime/Engine/Private/Components/SkeletalMeshComponent.cpp
--- a/Engine/Source/Runtime/Engine/Private/Components/SkeletalMeshComponent.cpp
+++ b/Engine/Source/Runtime/Engine/Private/Components/SkeletalMeshComponent.cpp
@@ -27,6 +27,12 @@
 #include "PhysicsEngine/PhysicsSettings.h"
 #include "PhysicsEngine/PhysicsAsset.h"
 
+#if WITH_APEX_CLOTHING
+#include "PhysicsEngine/PhysXSupport.h"
+#include "NxClothingActor.h"
+#include "NxClothingAsset.h"
+#endif
+
 TAutoConsoleVariable<int32> CVarUseParallelAnimationEvaluation(TEXT("a.ParallelAnimEvaluation"), 1, TEXT("If 1, animation evaluation will be run across the task graph system. If 0, evaluation will run purely on the game thread"));
 
 DECLARE_CYCLE_STAT_EXTERN(TEXT("Anim Instance Spawn Time"), STAT_AnimSpawnTime, STATGROUP_Anim, );
@@ -130,6 +136,8 @@ USkeletalMeshComponent::USkeletalMeshComponent(const FObjectInitializer& ObjectI
 	bNeedTeleportAndResetOnceMore = false;
 	ClothBlendWeight = 1.0f;
 	bPreparedClothMorphTargets = false;
+	bBindClothToMasterComponent = false;
+	bPrevMasterSimulateLocalSpace = false;
 #if WITH_CLOTH_COLLISION_DETECTION
 	ClothingCollisionRevision = 0;
 #endif// #if WITH_CLOTH_COLLISION_DETECTION
@@ -2164,3 +2172,92 @@ bool USkeletalMeshComponent::IsRunningParallelEvaluation(bool bBlockOnTask, bool
 	}
 	return false;
 }
+
+void USkeletalMeshComponent::BindClothToMasterPoseComponent()
+{
+#if WITH_APEX_CLOTHING
+	if(USkeletalMeshComponent* MasterComp = Cast<USkeletalMeshComponent>(MasterPoseComponent.Get()))
+	{
+		if(SkeletalMesh != MasterComp->SkeletalMesh)
+		{
+			// Not the same mesh, can't bind
+			return;
+		}
+
+		ValidateClothingActors();
+		MasterComp->ValidateClothingActors();
+
+		int32 NumClothingActors = ClothingActors.Num();
+		
+		for(int32 ActorIdx = 0 ; ActorIdx < NumClothingActors ; ++ActorIdx)
+		{
+			FClothingActor& Actor = ClothingActors[ActorIdx];
+			FClothingActor& MasterActor = MasterComp->ClothingActors[ActorIdx];
+			NxClothingActor* ApexActor = Actor.ApexClothingActor;
+			NxClothingActor* MasterApexActor = MasterActor.ApexClothingActor;
+			if(ApexActor && MasterApexActor)
+			{
+				// Disable our actors
+				ApexActor->setFrozen(true);
+
+				// Force local space simulation
+				NxParameterized::Interface* MasterActorInterface = MasterApexActor->getActorDesc();
+				verify(NxParameterized::setParamBool(*MasterActorInterface, "localSpaceSim", true));
+
+				// Make sure the master component starts extracting in local space
+				bPrevMasterSimulateLocalSpace = MasterComp->bLocalSpaceSimulation;
+				MasterComp->bLocalSpaceSimulation = true;
+			}
+			else
+			{
+				// Something has gone wrong here, don't attempt to extract cloth positions
+				UE_LOG(LogAnimation, Warning, TEXT("BindClothToMasterPoseComponent: Failed to bind to master component, missing actor."));
+				bBindClothToMasterComponent = false;
+				return;
+			}
+		}
+
+		// When we extract positions from now we'll just take the master components positions
+		bBindClothToMasterComponent = true;
+	}
+#endif
+}
+
+void USkeletalMeshComponent::UnbindClothFromMasterPoseComponent(bool bRestoreSimulationSpace)
+{
+#if WITH_APEX_CLOTHING
+	USkeletalMeshComponent* MasterComp = Cast<USkeletalMeshComponent>(MasterPoseComponent.Get());
+	if(MasterComp && bBindClothToMasterComponent)
+	{
+		bBindClothToMasterComponent = false;
+
+		int32 NumClothingActors = ClothingActors.Num();
+
+		for(int32 ActorIdx = 0 ; ActorIdx < NumClothingActors ; ++ActorIdx)
+		{
+			FClothingActor& Actor = ClothingActors[ActorIdx];
+			NxClothingActor* ApexActor = Actor.ApexClothingActor;
+
+			if(ApexActor)
+			{
+				ApexActor->setFrozen(false);
+
+				bool bMasterPoseSpaceChanged = (MasterComp->bLocalSpaceSimulation && !bPrevMasterSimulateLocalSpace);
+				if(bMasterPoseSpaceChanged && bRestoreSimulationSpace)
+				{
+					// Need to undo local space
+					FClothingActor& MasterActor = MasterComp->ClothingActors[ActorIdx];
+					NxClothingActor* MasterApexActor = MasterActor.ApexClothingActor;
+					if(MasterApexActor)
+					{
+						NxParameterized::Interface* MasterActorInterface = MasterApexActor->getActorDesc();
+						verify(NxParameterized::setParamBool(*MasterActorInterface, "localSpaceSim", false));
+
+						MasterComp->bLocalSpaceSimulation = false;
+					}
+				}
+			}
+		}
+	}
+#endif
+}
diff --git a/Engine/Source/Runtime/Engine/Private/SkeletalMeshComponentPhysics.cpp b/Engine/Source/Runtime/Engine/Private/SkeletalMeshComponentPhysics.cpp
--- a/Engine/Source/Runtime/Engine/Private/SkeletalMeshComponentPhysics.cpp
+++ b/Engine/Source/Runtime/Engine/Private/SkeletalMeshComponentPhysics.cpp
@@ -3660,7 +3660,7 @@ void USkeletalMeshComponent::TickClothing(float DeltaTime)
 #endif// #if WITH_APEX_CLOTHING
 }
 
-void USkeletalMeshComponent::GetUpdateClothSimulationData(TArray<FClothSimulData>& OutClothSimData)
+void USkeletalMeshComponent::GetUpdateClothSimulationData(TArray<FClothSimulData>& OutClothSimData, USkeletalMeshComponent* OverrideLocalRootComponent)
 {
 #if WITH_APEX_CLOTHING
 
@@ -3717,7 +3717,7 @@ void USkeletalMeshComponent::GetUpdateClothSimulationData(TArray<FClothSimulData
 				if (bLocalSpaceSimulation)
 				{
 					FMatrix ClothRootBoneMatrix;
-					GetClothRootBoneMatrix(ActorIndex, ClothRootBoneMatrix);
+					(OverrideLocalRootComponent) ? OverrideLocalRootComponent->GetClothRootBoneMatrix(ActorIndex, ClothRootBoneMatrix) : GetClothRootBoneMatrix(ActorIndex, ClothRootBoneMatrix);
 					for (uint32 VertexIndex = 0; VertexIndex < NumSimulVertices; VertexIndex++)
 					{						
 						ClothData.ClothSimulPositions[VertexIndex] = ClothRootBoneMatrix.TransformPosition(P2UVector(Vertices[VertexIndex]));
diff --git a/Engine/Source/Runtime/Engine/Private/SkeletalRenderGPUSkin.cpp b/Engine/Source/Runtime/Engine/Private/SkeletalRenderGPUSkin.cpp
--- a/Engine/Source/Runtime/Engine/Private/SkeletalRenderGPUSkin.cpp
+++ b/Engine/Source/Runtime/Engine/Private/SkeletalRenderGPUSkin.cpp
@@ -1058,6 +1058,21 @@ bool FDynamicSkelMeshObjectDataGPUSkin::ActiveVertexAnimsEqual( const TArray<FAc
 bool FDynamicSkelMeshObjectDataGPUSkin::UpdateClothSimulationData(USkinnedMeshComponent* InMeshComponent)
 {
 	USkeletalMeshComponent * SkelMeshComponent = Cast<USkeletalMeshComponent>(InMeshComponent);
+
+	if(InMeshComponent->MasterPoseComponent.IsValid() && SkelMeshComponent->bBindClothToMasterComponent)
+	{
+		USkeletalMeshComponent* OriginalComponent = Cast<USkeletalMeshComponent>(InMeshComponent);
+		SkelMeshComponent = Cast<USkeletalMeshComponent>(InMeshComponent->MasterPoseComponent.Get());
+
+		// If either component is invalid, don't attempt to update
+		if(!OriginalComponent || !SkelMeshComponent) return false;
+
+		ClothBlendWeight = OriginalComponent->ClothBlendWeight;
+		SkelMeshComponent->GetUpdateClothSimulationData(ClothSimulUpdateData, OriginalComponent);
+
+		return true;
+	}
+
 	if(SkelMeshComponent)
 	{
 		ClothBlendWeight = SkelMeshComponent->ClothBlendWeight;
