Init AudioDevice from within CreateAudioDevice

diff --git a/Engine/Source/Runtime/Engine/Private/AudioDeviceManager.cpp b/Engine/Source/Runtime/Engine/Private/AudioDeviceManager.cpp
--- a/Engine/Source/Runtime/Engine/Private/AudioDeviceManager.cpp
+++ b/Engine/Source/Runtime/Engine/Private/AudioDeviceManager.cpp
@@ -137,6 +137,14 @@ FAudioDevice* FAudioDeviceManager::CreateAudioDevice(uint32& HandleOut, bool bCr
 
 	++NumActiveAudioDevices;
 
+	const UAudioSettings* AudioSettings = GetDefault<UAudioSettings>();
+	if (!NewAudioDevice->Init())
+	{
+		ShutdownAudioDevice(HandleOut);
+		HandleOut = AUDIO_DEVICE_HANDLE_INVALID;
+		NewAudioDevice = nullptr;
+	}
+
 	return NewAudioDevice;
 }
 
diff --git a/Engine/Source/Runtime/Engine/Private/GameViewportClient.cpp b/Engine/Source/Runtime/Engine/Private/GameViewportClient.cpp
--- a/Engine/Source/Runtime/Engine/Private/GameViewportClient.cpp
+++ b/Engine/Source/Runtime/Engine/Private/GameViewportClient.cpp
@@ -259,26 +259,18 @@ void UGameViewportClient::Init(struct FWorldContext& WorldContext, UGameInstance
 			FAudioDevice* NewAudioDevice = AudioDeviceManager->CreateAudioDevice(AudioDeviceHandle, bCreateNewAudioDevice);
 			if (NewAudioDevice)
 			{
-				if (NewAudioDevice->Init())
+				// Set the base mix of the new device based on the world settings of the world
+				if (World)
 				{
-					// Set the base mix of the new device based on the world settings of the world
-					if (World)
-					{
-						NewAudioDevice->SetDefaultBaseSoundMix(World->GetWorldSettings()->DefaultBaseSoundMix);
-
-						// Set the world's audio device handle to use so that sounds which play in that world will use the correct audio device
-						World->SetAudioDeviceHandle(AudioDeviceHandle);
-					}
+					NewAudioDevice->SetDefaultBaseSoundMix(World->GetWorldSettings()->DefaultBaseSoundMix);
 
-					// Set this audio device handle on the world context so future world's set onto the world context
-					// will pass the audio device handle to them and audio will play on the correct audio device
-					WorldContext.AudioDeviceHandle = AudioDeviceHandle;
-				}
-				else
-				{
-					// Shut it down if we failed to initialize
-					AudioDeviceManager->ShutdownAudioDevice(AudioDeviceHandle);
+					// Set the world's audio device handle to use so that sounds which play in that world will use the correct audio device
+					World->SetAudioDeviceHandle(AudioDeviceHandle);
 				}
+
+				// Set this audio device handle on the world context so future world's set onto the world context
+				// will pass the audio device handle to them and audio will play on the correct audio device
+				WorldContext.AudioDeviceHandle = AudioDeviceHandle;
 			}
 		}
 	}
diff --git a/Engine/Source/Runtime/Engine/Private/UnrealEngine.cpp b/Engine/Source/Runtime/Engine/Private/UnrealEngine.cpp
--- a/Engine/Source/Runtime/Engine/Private/UnrealEngine.cpp
+++ b/Engine/Source/Runtime/Engine/Private/UnrealEngine.cpp
@@ -1699,27 +1699,13 @@ bool UEngine::InitializeAudioDeviceManager()
 					AudioDeviceManager = new FAudioDeviceManager();
 					AudioDeviceManager->RegisterAudioDeviceModule(AudioDeviceModule);
 
-					bool bSucceeded = false;
-
 					// Create a new audio device.
 					FAudioDevice* NewAudioDevice = AudioDeviceManager->CreateAudioDevice(MainAudioDeviceHandle, true);
 					if (NewAudioDevice)
 					{
-						// Initialize the audio device
-						if (NewAudioDevice->Init())
-						{
-							bSucceeded = true;
-							AudioDeviceManager->SetActiveDevice(MainAudioDeviceHandle);
-						}
-						else
-						{
-							// Shut it down if we failed to initialize
-							AudioDeviceManager->ShutdownAudioDevice(MainAudioDeviceHandle);
-						}
+						AudioDeviceManager->SetActiveDevice(MainAudioDeviceHandle);
 					}
-
-					// if we failed to create or init a main audio device, shut down the audio device manager
-					if (!bSucceeded)
+					else
 					{
 						ShutdownAudioDeviceManager();
 					}
