Physical materials are returned in async traces.

diff --git a/Engine/Source/Runtime/Engine/Private/Collision/CollisionConversions.cpp b/Engine/Source/Runtime/Engine/Private/Collision/CollisionConversions.cpp
--- a/Engine/Source/Runtime/Engine/Private/Collision/CollisionConversions.cpp
+++ b/Engine/Source/Runtime/Engine/Private/Collision/CollisionConversions.cpp
@@ -361,19 +361,10 @@ static void SetHitResultFromShapeAndFaceIndex(const PxShape* PShape,  const PxRi
 
 		if (bReturnPhysMat)
 		{
-			// @fixme: only do this for InGameThread, otherwise, this will be done in AsyncTrace
-			if ( IsInGameThread() )
+			// This function returns the single material in all cases other than trimesh or heightfield
+			if(PxMaterial* PxMat = PShape->getMaterialFromInternalFaceIndex(FaceIndex))
 			{
-				// This function returns the single material in all cases other than trimesh or heightfield
-				PxMaterial* PxMat = PShape->getMaterialFromInternalFaceIndex(FaceIndex);
-				if(PxMat != NULL)
-				{
-					OutResult.PhysMaterial = FPhysxUserData::Get<UPhysicalMaterial>(PxMat->userData);
-				}
-			}
-			else
-			{
-				//@fixme: this will be fixed properly when we can make FBodyInstance to be TWeakPtr - TTP (263842)
+				OutResult.PhysMaterial = FPhysxUserData::Get<UPhysicalMaterial>(PxMat->userData);
 			}
 		}
 
