#ue4 - Add separate stat for SceneComponent::MoveComponent. It was not tracked as part of "MoveComponent Time" previously.

diff --git a/Engine/Source/Runtime/Engine/Private/Components/SceneComponent.cpp b/Engine/Source/Runtime/Engine/Private/Components/SceneComponent.cpp
--- a/Engine/Source/Runtime/Engine/Private/Components/SceneComponent.cpp
+++ b/Engine/Source/Runtime/Engine/Private/Components/SceneComponent.cpp
@@ -1911,6 +1911,8 @@ bool USceneComponent::CheckStaticMobilityAndWarn(const FText& ActionText) const
 
 bool USceneComponent::MoveComponentImpl( const FVector& Delta, const FQuat& NewRotation, bool bSweep, FHitResult* OutHit, EMoveComponentFlags MoveFlags )
 {
+	SCOPE_CYCLE_COUNTER(STAT_MoveComponentSceneComponentTime);
+
 	// static things can move before they are registered (e.g. immediately after streaming), but not after.
 	static const FText WarnText = LOCTEXT("InvalidMove", "move");
 	if (CheckStaticMobilityAndWarn(WarnText))
diff --git a/Engine/Source/Runtime/Engine/Private/LevelTick.cpp b/Engine/Source/Runtime/Engine/Private/LevelTick.cpp
--- a/Engine/Source/Runtime/Engine/Private/LevelTick.cpp
+++ b/Engine/Source/Runtime/Engine/Private/LevelTick.cpp
@@ -48,6 +48,7 @@ DEFINE_STAT(STAT_GCMarkTime);
 
 DEFINE_STAT(STAT_TeleportToTime);
 DEFINE_STAT(STAT_MoveComponentTime);
+DEFINE_STAT(STAT_MoveComponentSceneComponentTime);
 DEFINE_STAT(STAT_UpdateOverlaps);
 DEFINE_STAT(STAT_UpdatePhysicsVolume);
 DEFINE_STAT(STAT_EndScopedMovementUpdate);
diff --git a/Engine/Source/Runtime/Engine/Public/EngineStats.h b/Engine/Source/Runtime/Engine/Public/EngineStats.h
--- a/Engine/Source/Runtime/Engine/Public/EngineStats.h
+++ b/Engine/Source/Runtime/Engine/Public/EngineStats.h
@@ -35,7 +35,8 @@ DECLARE_CYCLE_STAT_EXTERN(TEXT("Deferred Tick Time"),STAT_DeferredTickTime,STATG
  */
 DECLARE_CYCLE_STAT_EXTERN(TEXT("Async Physics Time"),STAT_PhysicsTime,STATGROUP_Game, );
 DECLARE_CYCLE_STAT_EXTERN(TEXT("Spawn Actor Time"),STAT_SpawnActorTime,STATGROUP_Game, );
-DECLARE_CYCLE_STAT_EXTERN(TEXT("MoveComponent Time"),STAT_MoveComponentTime,STATGROUP_Game, );
+DECLARE_CYCLE_STAT_EXTERN(TEXT("MoveComponent(Primitive) Time"),STAT_MoveComponentTime,STATGROUP_Game, );
+DECLARE_CYCLE_STAT_EXTERN(TEXT("MoveComponent(SceneComp) Time"), STAT_MoveComponentSceneComponentTime, STATGROUP_Game, );
 DECLARE_CYCLE_STAT_EXTERN(TEXT("UpdateOverlaps Time"),STAT_UpdateOverlaps,STATGROUP_Game, );
 DECLARE_CYCLE_STAT_EXTERN(TEXT("UpdatePhysicsVolume Time"),STAT_UpdatePhysicsVolume,STATGROUP_Game, );
 DECLARE_CYCLE_STAT_EXTERN(TEXT("EndScopedMovementUpdate Time"),STAT_EndScopedMovementUpdate,STATGROUP_Game, );
