Probable fix for HTML5.

diff --git a/Engine/Source/Runtime/Engine/Private/SkeletalRenderGPUSkin.cpp b/Engine/Source/Runtime/Engine/Private/SkeletalRenderGPUSkin.cpp
--- a/Engine/Source/Runtime/Engine/Private/SkeletalRenderGPUSkin.cpp
+++ b/Engine/Source/Runtime/Engine/Private/SkeletalRenderGPUSkin.cpp
@@ -1059,6 +1059,7 @@ bool FDynamicSkelMeshObjectDataGPUSkin::UpdateClothSimulationData(USkinnedMeshCo
 {
 	USkeletalMeshComponent * SkelMeshComponent = Cast<USkeletalMeshComponent>(InMeshComponent);
 
+#if WITH_APEX_CLOTHING
 	if(InMeshComponent->MasterPoseComponent.IsValid() && SkelMeshComponent->bBindClothToMasterComponent)
 	{
 		USkeletalMeshComponent* OriginalComponent = Cast<USkeletalMeshComponent>(InMeshComponent);
@@ -1072,6 +1073,7 @@ bool FDynamicSkelMeshObjectDataGPUSkin::UpdateClothSimulationData(USkinnedMeshCo
 
 		return true;
 	}
+#endif
 
 	if(SkelMeshComponent)
 	{
