Added support for configuring AnimUpdateRateParams in DefaultGame.ini

OrionGame/Config/DefaultGame.ini

[/Script/OrionGame.OrionChar]
; Update 1 in 4 frames when not rendered (and dedicated servers)
BaseNonRenderedAnimUpdateRate=4
; 0 frame skip, MaxDistanceFactor >  0.4f
; 1 frame skip, MaxDistanceFactor >  0.2f
; 2 frame skip, MaxDistanceFactor >= 0.0f
!AnimUpdateRateVisibleMaxDistanceFactor=ClearArray
+AnimUpdateRateVisibleMaxDistanceFactor=0.4f
+AnimUpdateRateVisibleMaxDistanceFactor=0.2f

--------
Integrated using branch UE4-To-UE4-Orion (reversed) of change#2551395 by [REDACTED_PERSON] on 2015/05/14 16:42:12.

diff --git a/Engine/Source/Runtime/Engine/Classes/Components/SkinnedMeshComponent.h b/Engine/Source/Runtime/Engine/Classes/Components/SkinnedMeshComponent.h
--- a/Engine/Source/Runtime/Engine/Classes/Components/SkinnedMeshComponent.h
+++ b/Engine/Source/Runtime/Engine/Classes/Components/SkinnedMeshComponent.h
@@ -11,6 +11,9 @@
 class FSkeletalMeshResource;
 struct FSkelMeshChunk;
 class FSkeletalMeshVertexBuffer;
+
+DECLARE_DELEGATE_OneParam(FOnAnimUpdateRateParamsCreated, FAnimUpdateRateParameters*)
+
 //
 // Bone Visibility.
 //
@@ -895,6 +898,9 @@ public:
 	/** Animation Update Rate optimization parameters. */
 	struct FAnimUpdateRateParameters* AnimUpdateRateParams;
 
+	/** Delegate when AnimUpdateRateParams is created, to override its default settings. */
+	FOnAnimUpdateRateParamsCreated OnAnimUpdateRateParamsCreated;
+
 	/** Updates AnimUpdateRateParams, used by SkinnedMeshComponents.
 	* 
 	* @param bRecentlyRendered : true if at least one SkinnedMeshComponent on this Actor has been rendered in the last second.
diff --git a/Engine/Source/Runtime/Engine/Classes/Engine/EngineTypes.h b/Engine/Source/Runtime/Engine/Classes/Engine/EngineTypes.h
--- a/Engine/Source/Runtime/Engine/Classes/Engine/EngineTypes.h
+++ b/Engine/Source/Runtime/Engine/Classes/Engine/EngineTypes.h
@@ -1790,6 +1790,24 @@ public:
 	/** The delta time of the last tick */
 	float ThisTickDelta;
 
+	/** Rate of animation evaluation when non rendered (off screen and dedicated servers).
+	 * a value of 4 means evaluated 1 frame, then 3 frames skipped */
+	UPROPERTY()
+	int32 BaseNonRenderedUpdateRate;
+
+	/** Array of MaxDistanceFactor to use for AnimUpdateRate when mesh is visible (rendered).
+	 * MaxDistanceFactor is size on screen, as used by LODs
+	 * Example:
+	 *		BaseVisibleDistanceFactorThesholds.Add(0.4f)
+	 *		BaseVisibleDistanceFactorThesholds.Add(0.2f)
+	 * means:
+	 *		0 frame skip, MaxDistanceFactor > 0.4f
+	 *		1 frame skip, MaxDistanceFactor > 0.2f
+	 *		2 frame skip, MaxDistanceFactor > 0.0f
+	 */
+	UPROPERTY()
+	TArray<float> BaseVisibleDistanceFactorThesholds;
+
 public:
 
 	/** Default constructor. */
@@ -1802,7 +1820,11 @@ public:
 		, TickedPoseOffestTime(0.f)
 		, AdditionalTime(0.f)
 		, ThisTickDelta(0.f)
-	{ }
+		, BaseNonRenderedUpdateRate(4)
+	{ 
+		BaseVisibleDistanceFactorThesholds.Add(0.4f);
+		BaseVisibleDistanceFactorThesholds.Add(0.2f);
+	}
 
 	/** Set parameters and verify inputs for Trail Mode (original behaviour - skip frames, track skipped time and then catch up afterwards).
 	 * @param : UpdateShiftRate. Shift our update frames so that updates across all skinned components are staggered
@@ -1851,12 +1873,12 @@ public:
 	{
 		if (OptimizeMode == TrailMode)
 		{
-			/*switch (UpdateRate)
+			switch (UpdateRate)
 			{
 			case 1: return FColor::Red;
 			case 2: return FColor::Green;
 			case 3: return FColor::Blue;
-			}*/
+			}
 			return FColor::Black;
 		}
 		else
diff --git a/Engine/Source/Runtime/Engine/Private/Components/SkinnedMeshComponent.cpp b/Engine/Source/Runtime/Engine/Private/Components/SkinnedMeshComponent.cpp
--- a/Engine/Source/Runtime/Engine/Private/Components/SkinnedMeshComponent.cpp
+++ b/Engine/Source/Runtime/Engine/Private/Components/SkinnedMeshComponent.cpp
@@ -83,22 +83,29 @@ namespace FAnimUpdateRateManager
 
 	FAnimUpdateRateParameters* GetUpdateRateParameters(USkinnedMeshComponent* SkinnedComponent)
 	{
+		if (!SkinnedComponent)
+		{
+			return NULL;
+		}
 		UObject* TrackerIndex = GetMapIndexForComponent(SkinnedComponent);
 
-		FAnimUpdateRateParametersTracker** ExistingTracker = ActorToUpdateRateParams.Find(TrackerIndex);
-		if (!ExistingTracker)
+		FAnimUpdateRateParametersTracker** ExistingTrackerPtr = ActorToUpdateRateParams.Find(TrackerIndex);
+		if (!ExistingTrackerPtr)
 		{
-			ExistingTracker = &ActorToUpdateRateParams.Add(TrackerIndex);
-			(*ExistingTracker) = new FAnimUpdateRateParametersTracker();
+			ExistingTrackerPtr = &ActorToUpdateRateParams.Add(TrackerIndex);
+			(*ExistingTrackerPtr) = new FAnimUpdateRateParametersTracker();
 		}
 		
+		check(ExistingTrackerPtr);
+		FAnimUpdateRateParametersTracker* ExistingTracker = *ExistingTrackerPtr;
 		check(ExistingTracker);
-		check(*ExistingTracker);
+		checkSlow(!ExistingTracker->RegisteredComponents.Contains(SkinnedComponent)); // We have already been registered? Something has gone very wrong!
 
-		checkSlow(!(*ExistingTracker)->RegisteredComponents.Contains(SkinnedComponent)); // We have already been registered? Something has gone very wrong!
+		ExistingTracker->RegisteredComponents.Add(SkinnedComponent);
+		FAnimUpdateRateParameters* UpdateRateParams = &ExistingTracker->UpdateRateParameters;
+		SkinnedComponent->OnAnimUpdateRateParamsCreated.ExecuteIfBound(UpdateRateParams);
 
-		(*ExistingTracker)->RegisteredComponents.Add(SkinnedComponent);
-		return &(*ExistingTracker)->UpdateRateParameters;
+		return UpdateRateParams;
 	}
 
 	void CleanupUpdateRateParametersRef(USkinnedMeshComponent* SkinnedComponent)
@@ -126,7 +133,9 @@ namespace FAnimUpdateRateManager
 		// Not rendered, including dedicated servers. we can skip the Evaluation part.
 		if (!bRecentlyRendered)
 		{
-			Tracker->UpdateRateParameters.SetTrailMode(DeltaTime, Tracker->GetAnimUpdateRateShiftTag(), ((bHumanControlled || bNeedsEveryFrame) ? 1 : 4), 4, false);
+			const int32 NewUpdateRate = (((bHumanControlled || bNeedsEveryFrame) ? 1 : Tracker->UpdateRateParameters.BaseNonRenderedUpdateRate);
+			const int32 NewEvaluationRate = Tracker->UpdateRateParameters.BaseNonRenderedUpdateRate;
+			Tracker->UpdateRateParameters.SetTrailMode(DeltaTime, Tracker->GetAnimUpdateRateShiftTag(), NewUpdateRate, NewEvaluationRate, false);
 		}
 		// Visible controlled characters or playing root motion. Need evaluation and ticking done every frame.
 		else  if (bHumanControlled || bNeedsEveryFrame)
@@ -135,20 +144,15 @@ namespace FAnimUpdateRateManager
 		}
 		else
 		{
-			// For visible meshes, figure out how often bones should be evaluated VS interpolated to previous evaluation.
-			// This is based on screen size and makes sense for a 3D FPS game, different games or Actors can implement different rules.
-			int32 DesiredEvaluationRate;
-			if (MaxDistanceFactor > 0.4f)
+			int32 DesiredEvaluationRate = Tracker->UpdateRateParameters.BaseVisibleDistanceFactorThesholds.Num() + 1;
+			for (int32 Index = 0; Index < Tracker->UpdateRateParameters.BaseVisibleDistanceFactorThesholds.Num(); Index++)
 			{
-				DesiredEvaluationRate = 1;
-			}
-			else if (MaxDistanceFactor > .2f)
-			{
-				DesiredEvaluationRate = 2;
-			}
-			else
-			{
-				DesiredEvaluationRate = 3;
+				const float& DistanceFactorThreadhold = Tracker->UpdateRateParameters.BaseVisibleDistanceFactorThesholds[Index];
+				if (MaxDistanceFactor > DistanceFactorThreadhold)
+				{
+					DesiredEvaluationRate = Index + 1;
+					break;
+				}
 			}
 
 			if (bUsingRootMotionFromEverything && DesiredEvaluationRate > 1)
