Fix shadowed variable
Fix mac compile error

diff --git a/Engine/Source/Runtime/Engine/Private/AudioSettings.cpp b/Engine/Source/Runtime/Engine/Private/AudioSettings.cpp
--- a/Engine/Source/Runtime/Engine/Private/AudioSettings.cpp
+++ b/Engine/Source/Runtime/Engine/Private/AudioSettings.cpp
@@ -58,9 +58,9 @@ void UAudioSettings::PostEditChangeChainProperty(FPropertyChangedChainEvent& Pro
 						{
 							bFoundDuplicate = false;
 							NewLevelName = FText::Format(LOCTEXT("NewQualityLevelName","New Level{0}"), (NewQualityLevelIndex > 0 ? FText::FromString(FString::Printf(TEXT(" %d"),NewQualityLevelIndex)) : FText::GetEmpty()));
-							for (const FAudioQualitySettings& AQSettings : QualityLevels)
+							for (const FAudioQualitySettings& QualityLevelSettings : QualityLevels)
 							{
-								if (AQSettings.DisplayName.EqualTo(NewLevelName))
+								if (QualityLevelSettings.DisplayName.EqualTo(NewLevelName))
 								{
 									bFoundDuplicate = true;
 									break;
diff --git a/Engine/Source/Runtime/Engine/Private/SoundNodeQualityLevel.cpp b/Engine/Source/Runtime/Engine/Private/SoundNodeQualityLevel.cpp
--- a/Engine/Source/Runtime/Engine/Private/SoundNodeQualityLevel.cpp
+++ b/Engine/Source/Runtime/Engine/Private/SoundNodeQualityLevel.cpp
@@ -5,6 +5,7 @@
 #include "ActiveSound.h"
 #include "Sound/SoundNodeQualityLevel.h"
 #include "Sound/AudioSettings.h"
+#include "GameFramework/GameUserSettings.h"
 
 #if WITH_EDITORONLY_DATA
 #include "Editor.h"
