Fixed merge error

diff --git a/Engine/Source/Runtime/Engine/Private/Components/SkinnedMeshComponent.cpp b/Engine/Source/Runtime/Engine/Private/Components/SkinnedMeshComponent.cpp
--- a/Engine/Source/Runtime/Engine/Private/Components/SkinnedMeshComponent.cpp
+++ b/Engine/Source/Runtime/Engine/Private/Components/SkinnedMeshComponent.cpp
@@ -133,7 +133,7 @@ namespace FAnimUpdateRateManager
 		// Not rendered, including dedicated servers. we can skip the Evaluation part.
 		if (!bRecentlyRendered)
 		{
-			const int32 NewUpdateRate = (((bHumanControlled || bNeedsEveryFrame) ? 1 : Tracker->UpdateRateParameters.BaseNonRenderedUpdateRate);
+			const int32 NewUpdateRate = ((bHumanControlled || bNeedsEveryFrame) ? 1 : Tracker->UpdateRateParameters.BaseNonRenderedUpdateRate);
 			const int32 NewEvaluationRate = Tracker->UpdateRateParameters.BaseNonRenderedUpdateRate;
 			Tracker->UpdateRateParameters.SetTrailMode(DeltaTime, Tracker->GetAnimUpdateRateShiftTag(), NewUpdateRate, NewEvaluationRate, false);
 		}
