Fix crash accessing the engine when attempting to play a slate sound after the engine has been shut down

diff --git a/Engine/Source/Runtime/Engine/Private/Slate/SlateSoundDevice.cpp b/Engine/Source/Runtime/Engine/Private/Slate/SlateSoundDevice.cpp
--- a/Engine/Source/Runtime/Engine/Private/Slate/SlateSoundDevice.cpp
+++ b/Engine/Source/Runtime/Engine/Private/Slate/SlateSoundDevice.cpp
@@ -11,24 +11,27 @@ DEFINE_LOG_CATEGORY_STATIC(LogSlateSoundDevice, Log, All);
 
 void FSlateSoundDevice::PlaySound(const FSlateSound& Sound, int32 UserIndex) const
 {
-	FAudioDevice* const AudioDevice = GEngine->GetActiveAudioDevice();
-	if(AudioDevice)
+	if( GEngine && Sound.GetResourceObject() != nullptr )
 	{
-		UObject* const Object = Sound.GetResourceObject();
-		USoundBase* const SoundResource = Cast<USoundBase>(Object);
-		if(SoundResource)
+		FAudioDevice* const AudioDevice = GEngine->GetActiveAudioDevice();
+		if(AudioDevice)
 		{
-			FActiveSound NewActiveSound;
-			NewActiveSound.Sound = SoundResource;
-			NewActiveSound.bIsUISound = true;
-			NewActiveSound.UserIndex = UserIndex;
+			UObject* const Object = Sound.GetResourceObject();
+			USoundBase* const SoundResource = Cast<USoundBase>(Object);
+			if(SoundResource)
+			{
+				FActiveSound NewActiveSound;
+				NewActiveSound.Sound = SoundResource;
+				NewActiveSound.bIsUISound = true;
+				NewActiveSound.UserIndex = UserIndex;
 
-			AudioDevice->AddNewActiveSound(NewActiveSound);
-		}
-		else if(Object)
-		{
-			// An Object but no SoundResource means that the FSlateSound is holding an invalid object; report that as an error
-			UE_LOG(LogSlateSoundDevice, Error, TEXT("A sound contains a non-sound resource '%s'"), *Object->GetName());
+				AudioDevice->AddNewActiveSound(NewActiveSound);
+			}
+			else if(Object)
+			{
+				// An Object but no SoundResource means that the FSlateSound is holding an invalid object; report that as an error
+				UE_LOG(LogSlateSoundDevice, Error, TEXT("A sound contains a non-sound resource '%s'"), *Object->GetName());
+			}
 		}
 	}
 }
