UE-15644 Fix typo in InterpTo Movement Component

diff --git a/Engine/Source/Runtime/Engine/Classes/Components/InterpToMovementComponent.h b/Engine/Source/Runtime/Engine/Classes/Components/InterpToMovementComponent.h
--- a/Engine/Source/Runtime/Engine/Classes/Components/InterpToMovementComponent.h
+++ b/Engine/Source/Runtime/Engine/Classes/Components/InterpToMovementComponent.h
@@ -66,7 +66,7 @@ class ENGINE_API UInterpToMovementComponent : public UMovementComponent
 	DECLARE_DYNAMIC_MULTICAST_DELEGATE_TwoParams( FOnInterpToWaitEndDelegate , const FHitResult&, ImpactResult, float, Time );
 	DECLARE_DYNAMIC_MULTICAST_DELEGATE_TwoParams( FOnInterpToResetDelegate , const FHitResult&, ImpactResult, float, Time );
 
-	/* How to take to move from the first point to the last (or vice versa) */
+	/* How long to take to move from the first point to the last (or vice versa) */
 	UPROPERTY(EditAnywhere, BlueprintReadWrite, Category=Control,meta=(UIMin=0.1f, ClampMin=0.1f))
 	float Duration;
 	
