UE-15405 Add columns in the FPS charts to track time spent running above 60 fps. This used to be reflected in the 60-INF column, but was lost when we added data for up to 120FPS.

Removed some redundant STAT enums and declarations.

diff --git a/Engine/Source/Runtime/Engine/Classes/Engine/Engine.h b/Engine/Source/Runtime/Engine/Classes/Engine/Engine.h
--- a/Engine/Source/Runtime/Engine/Classes/Engine/Engine.h
+++ b/Engine/Source/Runtime/Engine/Classes/Engine/Engine.h
@@ -1991,6 +1991,11 @@ public:
 
 private:
 
+	/**
+	* Calculates the range of FPS values for the given bucket index
+	*/
+	void CalcQuantisedFPSRange(int32 BucketIndex, int32& StartFPS, int32& EndFPS);
+
 	/**
 	 * Dumps the FPS chart information to HTML.
 	 */
diff --git a/Engine/Source/Runtime/Engine/Private/ChartCreation.cpp b/Engine/Source/Runtime/Engine/Private/ChartCreation.cpp
--- a/Engine/Source/Runtime/Engine/Private/ChartCreation.cpp
+++ b/Engine/Source/Runtime/Engine/Private/ChartCreation.cpp
@@ -14,7 +14,7 @@
 DEFINE_LOG_CATEGORY_STATIC(LogChartCreation, Log, All);
 
 /** The total GPU time taken to render all frames. In seconds. */
-double			GTotalGPUTime = 0;
+double			GTotalGPUTime = 0; 
 
 #if DO_CHARTING
 
@@ -410,13 +410,45 @@ void UEngine::StopFPSChart()
 	GEnableDataCapture = false;
 }
 
+/**
+* Calculates the range of FPS values for the given bucket index
+*/
+
+void UEngine::CalcQuantisedFPSRange(int32 BucketIndex, int32& StartFPS, int32& EndFPS )
+{
+	// Figure out bucket range. Buckets start at 5 frame intervals then change to 10.
+	StartFPS = 0;
+	EndFPS = 0;
+	
+	if (BucketIndex < STAT_FPSChart_30_40)
+	{
+		// Still incrementing by 5
+		StartFPS = BucketIndex * 5;
+		EndFPS = StartFPS + 5;
+	}
+	else
+	{
+		// Now incrementing by 10
+		StartFPS = STAT_FPSChart_30_40 * 5;
+		StartFPS += (BucketIndex - STAT_FPSChart_30_40) * 10;
+		EndFPS = StartFPS + 10;
+	}
+
+	if (BucketIndex + STAT_FPSChartFirstStat == STAT_FPSChart_120_INF)
+	{
+		EndFPS = 999;
+	}
+}
+
 /**
  * Dumps the FPS chart information to the log.
  */
 void UEngine::DumpFPSChartToLog( float TotalTime, float DeltaTime, int32 NumFrames, const FString& InMapName )
 {
 	int32 NumFramesBelow30 = 0; // keep track of the number of frames below 30 FPS
-	float PctTimeAbove30 = 0; // Keep track of percentage of time at 30+ FPS.
+	float PctTimeAbove30 = 0;	// Keep track of percentage of time at 30+ FPS.
+	int32 NumFramesBelow60 = 0; // keep track of the number of frames below 60 FPS
+	float PctTimeAbove60 = 0;	// Keep track of percentage of time at 60+ FPS.
 
 	UE_LOG(LogChartCreation, Log, TEXT("--- Begin : FPS chart dump for level '%s'"), *InMapName );
 
@@ -430,55 +462,51 @@ void UEngine::DumpFPSChartToLog( float TotalTime, float DeltaTime, int32 NumFram
 		// Figure out bucket range. Buckets start at 5 frame intervals then change to 10.
 		int32 StartFPS = 0;
 		int32 EndFPS = 0;
-		if ( BucketIndex < STAT_FPSChart_30_40 )
+		CalcQuantisedFPSRange(BucketIndex, StartFPS, EndFPS);
+
+		// Keep track of time spent at 30+ FPS.
+		if( StartFPS >= 30 )
 		{
-			// Still incrementing by 5
-			StartFPS = BucketIndex * 5;
-			EndFPS = StartFPS + 5;
+			PctTimeAbove30 += BucketTimePercentage;
 		}
 		else
 		{
-			// Now incrementing by 10
-			StartFPS = STAT_FPSChart_30_40 * 5;
-			StartFPS += (BucketIndex - STAT_FPSChart_30_40) * 10;
-			EndFPS = StartFPS + 10;
-		}
-		
-		if( BucketIndex + STAT_FPSChartFirstStat == STAT_FPSChart_120_INF )
-		{
-			EndFPS = 999;
+			NumFramesBelow30 += GFPSChart[BucketIndex].Count;
 		}
 
-		// Keep track of time spent at 30+ FPS.
-		if( StartFPS >= 30 )
+		// Keep track of time spent at 60+ FPS.
+		if (StartFPS >= 60)
 		{
-			PctTimeAbove30 += BucketTimePercentage;
+			PctTimeAbove60 += BucketTimePercentage;
 		}
 		else
 		{
-			NumFramesBelow30 += GFPSChart[BucketIndex].Count;
+			NumFramesBelow60 += GFPSChart[BucketIndex].Count;
 		}
 
 		// Log bucket index, time and frame Percentage.
 		UE_LOG(LogChartCreation, Log, TEXT("Bucket: %2i - %2i  Time: %5.2f  Frame: %5.2f"), StartFPS, EndFPS, BucketTimePercentage, BucketFramePercentage );
 	}
 
-	UE_LOG(LogChartCreation, Log, TEXT("%i frames collected over %4.2f seconds, disregarding %4.2f seconds for a %4.2f FPS average, %4.2f percent of time spent > 30 FPS"),
+	UE_LOG(LogChartCreation, Log, TEXT("%i frames collected over %4.2f seconds, disregarding %4.2f seconds for a %4.2f FPS average, %4.2f percent of time spent > 30 FPS, %4.2f percent of time spent > 60 FPS"),
 		NumFrames, 
 		DeltaTime, 
 		FMath::Max<float>( 0, DeltaTime - TotalTime ), 
 		NumFrames / TotalTime,
-		PctTimeAbove30);
+		PctTimeAbove30,
+		PctTimeAbove60);
 	UE_LOG(LogChartCreation, Log, TEXT("Average GPU frametime: %4.2f ms"), float((GTotalGPUTime / NumFrames)*1000.0));
-	UE_LOG(LogChartCreation, Log, TEXT("BoundGameThreadPct: %4.2f  BoundRenderThreadPct: %4.2f  BoundGPUPct: %4.2f PercentFrames30+: %f   BoundGameTime: %f  BoundRenderTime: %f  BoundGPUTime: %f  PctTimeAbove30: %f ")
+	UE_LOG(LogChartCreation, Log, TEXT("BoundGameThreadPct: %4.2f  BoundRenderThreadPct: %4.2f  BoundGPUPct: %4.2f PercentFrames30+: %f   PercentFrames60+: %f   BoundGameTime: %f  BoundRenderTime: %f  BoundGPUTime: %f  PctTimeAbove30: %f  PctTimeAbove60: %f ")
 		, (float(GNumFramesBound_GameThread)/float(NumFrames))*100.0f
 		, (float(GNumFramesBound_RenderThread)/float(NumFrames))*100.0f
 		, (float(GNumFramesBound_GPU)/float(NumFrames))*100.0f
 		, float(NumFrames - NumFramesBelow30) / float(NumFrames)*100.0f
-		, (GTotalFramesBoundTime_GameThread/DeltaTime)*100.0f
+		, float(NumFrames - NumFramesBelow60) / float(NumFrames)*100.0f
+		, (GTotalFramesBoundTime_GameThread / DeltaTime)*100.0f
 		, ((GTotalFramesBoundTime_RenderThread)/DeltaTime)*100.0f
 		, ((GTotalFramesBoundTime_GPU)/DeltaTime)*100.0f
 		, PctTimeAbove30
+		, PctTimeAbove60
 		);
 
 	UE_LOG(LogChartCreation, Log, TEXT("--- End"));
@@ -685,7 +713,9 @@ void UEngine::DumpFPSChartToStatsLog( float TotalTime, float DeltaTime, int32 Nu
 		OutputFile->Logf(TEXT("\tEffects Quality: %d"), Quality.EffectsQuality);
 
 		int32 NumFramesBelow30 = 0; // keep track of the number of frames below 30 FPS
-		float PctTimeAbove30 = 0; // Keep track of percentage of time at 30+ FPS.
+		float PctTimeAbove30 = 0;	// Keep track of percentage of time at 30+ FPS.
+		int32 NumFramesBelow60 = 0; // keep track of the number of frames below 60 FPS
+		float PctTimeAbove60 = 0;	// Keep track of percentage of time at 60+ FPS.
 
 		// Iterate over all buckets, dumping percentages.
 		for( int32 BucketIndex=0; BucketIndex<ARRAY_COUNT(GFPSChart); BucketIndex++ )
@@ -694,58 +724,52 @@ void UEngine::DumpFPSChartToStatsLog( float TotalTime, float DeltaTime, int32 Nu
 			const float BucketTimePercentage  = 100.f * GFPSChart[BucketIndex].CummulativeTime / TotalTime;
 			const float BucketFramePercentage = 100.f * GFPSChart[BucketIndex].Count / NumFrames;
 
-			// Figure out bucket range. Buckets start at 5 frame intervals then change to 10.
 			int32 StartFPS = 0;
 			int32 EndFPS = 0;
-			if (BucketIndex < STAT_FPSChart_30_40)
+			CalcQuantisedFPSRange(BucketIndex, StartFPS, EndFPS);
+
+			// Keep track of time spent at 30+ FPS.
+			if (StartFPS >= 30)
 			{
-				// Still incrementing by 5
-				StartFPS = BucketIndex * 5;
-				EndFPS = StartFPS + 5;
+				PctTimeAbove30 += BucketTimePercentage;
 			}
 			else
 			{
-				// Now incrementing by 10
-				StartFPS = STAT_FPSChart_30_40 * 5;
-				StartFPS += (BucketIndex - STAT_FPSChart_30_40) * 10;
-				EndFPS = StartFPS + 10;
-			}
-
-			if( BucketIndex + STAT_FPSChartFirstStat == STAT_FPSChart_120_INF )
-			{
-				EndFPS = 999;
+				NumFramesBelow30 += GFPSChart[BucketIndex].Count;
 			}
-
-			// Keep track of time spent at 30+ FPS.
-			if( StartFPS >= 30 )
+			
+			// Keep track of time spent at 60+ FPS.
+			if (StartFPS >= 60)
 			{
-				PctTimeAbove30 += BucketTimePercentage;
+				PctTimeAbove60 += BucketTimePercentage;
 			}
 			else
 			{
-				NumFramesBelow30 += GFPSChart[BucketIndex].Count;
+				NumFramesBelow60 += GFPSChart[BucketIndex].Count;
 			}
 
 			// Log bucket index, time and frame Percentage.
 			OutputFile->Logf(TEXT("Bucket: %2i - %2i  Time: %5.2f  Frame: %5.2f"), StartFPS, EndFPS, BucketTimePercentage, BucketFramePercentage);
 		}
 
-		OutputFile->Logf(TEXT("%i frames collected over %4.2f seconds, disregarding %4.2f seconds for a %4.2f FPS average, %4.2f percent of time spent > 30 FPS"), 
+		OutputFile->Logf(TEXT("%i frames collected over %4.2f seconds, disregarding %4.2f seconds for a %4.2f FPS average, %4.2f percent of time spent > 30 FPS, %4.2f percent of time spent > 60 FPS"), 
 			NumFrames, 
 			DeltaTime, 
 			FMath::Max<float>( 0, DeltaTime - TotalTime ), 
 			NumFrames / TotalTime,
-			PctTimeAbove30);
+			PctTimeAbove30, PctTimeAbove60 );
 		OutputFile->Logf(TEXT("Average GPU frame time: %4.2f ms"), float((GTotalGPUTime / NumFrames)*1000.0));
-		OutputFile->Logf(TEXT("BoundGameThreadPct: %4.2f  BoundRenderThreadPct: %4.2f  BoundGPUPct: %4.2f PercentFrames30+: %f   BoundGameTime: %f  BoundRenderTime: %f  BoundGPUTime: %f  PctTimeAbove30: %f ")
+		OutputFile->Logf(TEXT("BoundGameThreadPct: %4.2f  BoundRenderThreadPct: %4.2f  BoundGPUPct: %4.2f PercentFrames30+: %f   PercentFrames60+: %f   BoundGameTime: %f  BoundRenderTime: %f  BoundGPUTime: %f  PctTimeAbove30: %f  PctTimeAbove60: %f ")
 			, (float(GNumFramesBound_GameThread)/float(NumFrames))*100.0f
 			, (float(GNumFramesBound_RenderThread)/float(NumFrames))*100.0f
 			, (float(GNumFramesBound_GPU)/float(NumFrames))*100.0f
 			, float(NumFrames - NumFramesBelow30) / float(NumFrames)*100.0f
-			, (GTotalFramesBoundTime_GameThread/DeltaTime)*100.0f
+			, float(NumFrames - NumFramesBelow60) / float(NumFrames)*100.0f
+			, (GTotalFramesBoundTime_GameThread / DeltaTime)*100.0f
 			, ((GTotalFramesBoundTime_RenderThread)/DeltaTime)*100.0f
 			, ((GTotalFramesBoundTime_GPU)/DeltaTime)*100.0f
 			, PctTimeAbove30
+			, PctTimeAbove60
 			);
 
 		// Dump hitch data
@@ -828,6 +852,8 @@ void UEngine::DumpFPSChartToHTML( float TotalTime, float DeltaTime, int32 NumFra
 	{
 		// Keep track of percentage of time at 30+ FPS.
 		float PctTimeAbove30 = 0;
+		// Keep track of percentage of time at 60+ FPS.
+		float PctTimeAbove60 = 0;
 
 		// Iterate over all buckets, updating row 
 		for( int32 BucketIndex=0; BucketIndex<ARRAY_COUNT(GFPSChart); BucketIndex++ )
@@ -839,24 +865,7 @@ void UEngine::DumpFPSChartToHTML( float TotalTime, float DeltaTime, int32 NumFra
 			// Figure out bucket range. Buckets start at 5 frame intervals then change to 10.
 			int32 StartFPS = 0;
 			int32 EndFPS = 0;
-			if (BucketIndex < STAT_FPSChart_30_40)
-			{
-				// Still incrementing by 5
-				StartFPS = BucketIndex * 5;
-				EndFPS = StartFPS + 5;
-			}
-			else
-			{
-				// Now incrementing by 10
-				StartFPS = STAT_FPSChart_30_40 * 5;
-				StartFPS += (BucketIndex - STAT_FPSChart_30_40) * 10;
-				EndFPS = StartFPS + 10;
-			}
-
-			if( BucketIndex + STAT_FPSChartFirstStat == STAT_FPSChart_120_INF )
-			{
-				EndFPS = 999;
-			}
+			CalcQuantisedFPSRange(BucketIndex, StartFPS, EndFPS);
 
 			// Keep track of time spent at 30+ FPS.
 			if( StartFPS >= 30 )
@@ -864,7 +873,13 @@ void UEngine::DumpFPSChartToHTML( float TotalTime, float DeltaTime, int32 NumFra
 				PctTimeAbove30 += BucketTimePercentage;
 			}
 
-			const FString SrcToken = FString::Printf( TEXT("TOKEN_%i_%i"), StartFPS, EndFPS );
+			// Keep track of time spent at 60+ FPS.
+			if (StartFPS >= 60)
+			{
+				PctTimeAbove60 += BucketTimePercentage;
+			}
+
+			const FString SrcToken = FString::Printf(TEXT("TOKEN_%i_%i"), StartFPS, EndFPS);
 			const FString DstToken = FString::Printf( TEXT("%5.2f"), BucketTimePercentage );
 
 			// Replace token with actual values.
@@ -934,8 +949,9 @@ void UEngine::DumpFPSChartToHTML( float TotalTime, float DeltaTime, int32 NumFra
 		FPSChartRow = FPSChartRow.Replace( TEXT("TOKEN_SETTINGS_FX"),	    *FString::Printf(TEXT("%d"), Quality.EffectsQuality ), ESearchCase::CaseSensitive );
 
 		FPSChartRow = FPSChartRow.Replace( TEXT("TOKEN_AVG_FPS"),			*FString::Printf(TEXT("%4.2f"), NumFrames / TotalTime), ESearchCase::CaseSensitive );
-		FPSChartRow = FPSChartRow.Replace( TEXT("TOKEN_PCT_ABOVE_30"),		*FString::Printf(TEXT("%4.2f"), PctTimeAbove30), ESearchCase::CaseSensitive );
-		FPSChartRow = FPSChartRow.Replace( TEXT("TOKEN_TIME_DISREGARDED"),	*FString::Printf(TEXT("%4.2f"), FMath::Max<float>( 0, DeltaTime - TotalTime ) ), ESearchCase::CaseSensitive );
+		FPSChartRow = FPSChartRow.Replace(TEXT("TOKEN_PCT_ABOVE_30"), *FString::Printf(TEXT("%4.2f"), PctTimeAbove30), ESearchCase::CaseSensitive);
+		FPSChartRow = FPSChartRow.Replace(TEXT("TOKEN_PCT_ABOVE_60"), *FString::Printf(TEXT("%4.2f"), PctTimeAbove60), ESearchCase::CaseSensitive);
+		FPSChartRow = FPSChartRow.Replace(TEXT("TOKEN_TIME_DISREGARDED"), *FString::Printf(TEXT("%4.2f"), FMath::Max<float>(0, DeltaTime - TotalTime)), ESearchCase::CaseSensitive);
 		FPSChartRow = FPSChartRow.Replace( TEXT("TOKEN_TIME"),				*FString::Printf(TEXT("%4.2f"), DeltaTime), ESearchCase::CaseSensitive );
 		FPSChartRow = FPSChartRow.Replace( TEXT("TOKEN_FRAMECOUNT"),		*FString::Printf(TEXT("%i"), NumFrames), ESearchCase::CaseSensitive );
 		FPSChartRow = FPSChartRow.Replace( TEXT("TOKEN_AVG_GPUTIME"),		*FString::Printf(TEXT("%4.2f ms"), float((GTotalGPUTime / NumFrames)*1000.0) ), ESearchCase::CaseSensitive );
diff --git a/Engine/Source/Runtime/Engine/Private/UnrealClient.cpp b/Engine/Source/Runtime/Engine/Private/UnrealClient.cpp
--- a/Engine/Source/Runtime/Engine/Private/UnrealClient.cpp
+++ b/Engine/Source/Runtime/Engine/Private/UnrealClient.cpp
@@ -306,10 +306,10 @@ int32 FStatUnitData::DrawStat(FViewport* InViewport, FCanvas* InCanvas, int32 In
 	RawGPUFrameTime = FPlatformTime::ToMilliseconds(GPUCycles);
 	GPUFrameTime = 0.9 * GPUFrameTime + 0.1 * RawGPUFrameTime;
 
-	SET_FLOAT_STAT(STAT_FPSChart_UnitFrame, FrameTime);
-	SET_FLOAT_STAT(STAT_FPSChart_UnitRender, RenderThreadTime);
-	SET_FLOAT_STAT(STAT_FPSChart_UnitGame, GameThreadTime);
-	SET_FLOAT_STAT(STAT_FPSChart_UnitGPU, GPUFrameTime);
+	SET_FLOAT_STAT(STAT_UnitFrame, FrameTime);
+	SET_FLOAT_STAT(STAT_UnitRender, RenderThreadTime);
+	SET_FLOAT_STAT(STAT_UnitGame, GameThreadTime);
+	SET_FLOAT_STAT(STAT_UnitGPU, GPUFrameTime);
 
 	GEngine->SetAverageUnitTimes(FrameTime, RenderThreadTime, GameThreadTime, GPUFrameTime);
 
diff --git a/Engine/Source/Runtime/Engine/Private/UnrealEngine.cpp b/Engine/Source/Runtime/Engine/Private/UnrealEngine.cpp
--- a/Engine/Source/Runtime/Engine/Private/UnrealEngine.cpp
+++ b/Engine/Source/Runtime/Engine/Private/UnrealEngine.cpp
@@ -7275,42 +7275,11 @@ DEFINE_STAT(STAT_SkelMeshDrawCalls);
 DEFINE_STAT(STAT_CPUSkinVertices);
 DEFINE_STAT(STAT_GPUSkinVertices);
 
-/** Frame chart stats */
-
-DEFINE_STAT(STAT_FPSChart_0_5);
-DEFINE_STAT(STAT_FPSChart_5_10);
-DEFINE_STAT(STAT_FPSChart_10_15);
-DEFINE_STAT(STAT_FPSChart_15_20);
-DEFINE_STAT(STAT_FPSChart_20_25);
-DEFINE_STAT(STAT_FPSChart_25_30);
-DEFINE_STAT(STAT_FPSChart_30_35);
-DEFINE_STAT(STAT_FPSChart_35_40);
-DEFINE_STAT(STAT_FPSChart_40_45);
-DEFINE_STAT(STAT_FPSChart_45_50);
-DEFINE_STAT(STAT_FPSChart_50_55);
-DEFINE_STAT(STAT_FPSChart_55_60);
-DEFINE_STAT(STAT_FPSChart_60_INF);
-DEFINE_STAT(STAT_FPSChart_30Plus);
-DEFINE_STAT(STAT_FPSChart_UnaccountedTime);
-DEFINE_STAT(STAT_FPSChart_FrameCount);
-DEFINE_STAT(STAT_FPSChart_Hitch_5000_Plus);
-DEFINE_STAT(STAT_FPSChart_Hitch_2500_5000);
-DEFINE_STAT(STAT_FPSChart_Hitch_2000_2500);
-DEFINE_STAT(STAT_FPSChart_Hitch_1500_2000);
-DEFINE_STAT(STAT_FPSChart_Hitch_1000_1500);
-DEFINE_STAT(STAT_FPSChart_Hitch_750_1000);
-DEFINE_STAT(STAT_FPSChart_Hitch_500_750);
-DEFINE_STAT(STAT_FPSChart_Hitch_300_500);
-DEFINE_STAT(STAT_FPSChart_Hitch_200_300);
-DEFINE_STAT(STAT_FPSChart_Hitch_150_200);
-DEFINE_STAT(STAT_FPSChart_Hitch_100_150);
-DEFINE_STAT(STAT_FPSChart_Hitch_60_100);
-DEFINE_STAT(STAT_FPSChart_TotalHitchCount);
-
-DEFINE_STAT(STAT_FPSChart_UnitFrame);
-DEFINE_STAT(STAT_FPSChart_UnitGame);
-DEFINE_STAT(STAT_FPSChart_UnitRender);
-DEFINE_STAT(STAT_FPSChart_UnitGPU);
+/** Unit times */
+DEFINE_STAT(STAT_UnitFrame);
+DEFINE_STAT(STAT_UnitGame);
+DEFINE_STAT(STAT_UnitRender);
+DEFINE_STAT(STAT_UnitGPU);
 
 /*-----------------------------------------------------------------------------
 	Lightmass object/actor implementations.
diff --git a/Engine/Source/Runtime/Engine/Public/ChartCreation.h b/Engine/Source/Runtime/Engine/Public/ChartCreation.h
--- a/Engine/Source/Runtime/Engine/Public/ChartCreation.h
+++ b/Engine/Source/Runtime/Engine/Public/ChartCreation.h
@@ -45,9 +45,6 @@ enum FPSChartStats
 	STAT_FPSChart_110_120,
 	STAT_FPSChart_120_INF,
 	STAT_FPSChartLastBucketStat,
-	STAT_FPSChart_30Plus,
-	STAT_FPSChart_UnaccountedTime,
-	STAT_FPSChart_FrameCount,
 
 	/** Hitch stats */
 	STAT_FPSChart_FirstHitchStat,
@@ -65,18 +62,6 @@ enum FPSChartStats
 	STAT_FPSChart_Hitch_60_100,
 	STAT_FPSChart_LastHitchBucketStat,
 	STAT_FPSChart_TotalHitchCount,
-
-	/** Unit time stats */
-	STAT_FPSChart_UnitFrame,
-	STAT_FPSChart_UnitRender,
-	STAT_FPSChart_UnitGame,
-	STAT_FPSChart_UnitGPU,
-
-	STAT_FPSChart_TotalFrameCount,
-	STAT_FPSChart_TotalUnitFrame,
-	STAT_FPSChart_TotalUnitRender,
-	STAT_FPSChart_TotalUnitGame,
-	STAT_FPSChart_TotalUnitGPU,
 };
 
 /** Start time of current FPS chart.										*/
diff --git a/Engine/Source/Runtime/Engine/Public/EngineStats.h b/Engine/Source/Runtime/Engine/Public/EngineStats.h
--- a/Engine/Source/Runtime/Engine/Public/EngineStats.h
+++ b/Engine/Source/Runtime/Engine/Public/EngineStats.h
@@ -30,6 +30,13 @@ DECLARE_CYCLE_STAT_EXTERN(TEXT("Platform Message Time"),STAT_PlatformMessageTime
 DECLARE_CYCLE_STAT_EXTERN(TEXT("Frame Sync Time"),STAT_FrameSyncTime,STATGROUP_Engine, ENGINE_API);
 DECLARE_CYCLE_STAT_EXTERN(TEXT("Deferred Tick Time"),STAT_DeferredTickTime,STATGROUP_Engine, ENGINE_API);
 
+/** Unit time stats */
+DECLARE_FLOAT_COUNTER_STAT_EXTERN(TEXT("StatUnit FrameTime"), STAT_UnitFrame, STATGROUP_Engine, );
+DECLARE_FLOAT_COUNTER_STAT_EXTERN(TEXT("StatUnit RenderThreadTime"), STAT_UnitRender, STATGROUP_Engine, );
+DECLARE_FLOAT_COUNTER_STAT_EXTERN(TEXT("StatUnit GameThreadTime"), STAT_UnitGame, STATGROUP_Engine, );
+DECLARE_FLOAT_COUNTER_STAT_EXTERN(TEXT("StatUnit GPUTime"), STAT_UnitGPU, STATGROUP_Engine, );
+
+
 /**
  * Game stats
  */
@@ -76,48 +83,6 @@ DECLARE_CYCLE_STAT_EXTERN(TEXT("Net Post BC Tick Time"),STAT_NetBroadcastPostTic
 DECLARE_CYCLE_STAT_EXTERN(TEXT("Net PackageMap SerializeObject"),STAT_PackageMap_SerializeObjectTime,STATGROUP_Game, );
 
 
-/**
- * FPS chart stats
- */
-DECLARE_FLOAT_ACCUMULATOR_STAT_EXTERN(TEXT("00 - 05 FPS"), STAT_FPSChart_0_5,		STATGROUP_FPSChart, );
-DECLARE_FLOAT_ACCUMULATOR_STAT_EXTERN(TEXT("05 - 10 FPS"), STAT_FPSChart_5_10,		STATGROUP_FPSChart, );
-DECLARE_FLOAT_ACCUMULATOR_STAT_EXTERN(TEXT("10 - 15 FPS"), STAT_FPSChart_10_15,	STATGROUP_FPSChart, );
-DECLARE_FLOAT_ACCUMULATOR_STAT_EXTERN(TEXT("15 - 20 FPS"), STAT_FPSChart_15_20,	STATGROUP_FPSChart, );
-DECLARE_FLOAT_ACCUMULATOR_STAT_EXTERN(TEXT("20 - 25 FPS"), STAT_FPSChart_20_25,	STATGROUP_FPSChart, );
-DECLARE_FLOAT_ACCUMULATOR_STAT_EXTERN(TEXT("25 - 30 FPS"), STAT_FPSChart_25_30,	STATGROUP_FPSChart, );
-DECLARE_FLOAT_ACCUMULATOR_STAT_EXTERN(TEXT("30 - 35 FPS"), STAT_FPSChart_30_35,	STATGROUP_FPSChart, );
-DECLARE_FLOAT_ACCUMULATOR_STAT_EXTERN(TEXT("35 - 40 FPS"), STAT_FPSChart_35_40,	STATGROUP_FPSChart, );
-DECLARE_FLOAT_ACCUMULATOR_STAT_EXTERN(TEXT("40 - 45 FPS"), STAT_FPSChart_40_45,	STATGROUP_FPSChart, );
-DECLARE_FLOAT_ACCUMULATOR_STAT_EXTERN(TEXT("45 - 50 FPS"), STAT_FPSChart_45_50,	STATGROUP_FPSChart, );
-DECLARE_FLOAT_ACCUMULATOR_STAT_EXTERN(TEXT("50 - 55 FPS"), STAT_FPSChart_50_55,	STATGROUP_FPSChart, );
-DECLARE_FLOAT_ACCUMULATOR_STAT_EXTERN(TEXT("55 - 60 FPS"), STAT_FPSChart_55_60,	STATGROUP_FPSChart, );
-DECLARE_FLOAT_ACCUMULATOR_STAT_EXTERN(TEXT("60 - .. FPS"), STAT_FPSChart_60_INF,	STATGROUP_FPSChart, );
-DECLARE_FLOAT_ACCUMULATOR_STAT_EXTERN(TEXT("30+ FPS"), STAT_FPSChart_30Plus,	STATGROUP_FPSChart, );
-DECLARE_FLOAT_ACCUMULATOR_STAT_EXTERN(TEXT("Unaccounted time"), STAT_FPSChart_UnaccountedTime,	STATGROUP_FPSChart, );
-DECLARE_DWORD_ACCUMULATOR_STAT_EXTERN(TEXT("Frame count"),	STAT_FPSChart_FrameCount, STATGROUP_FPSChart, );
-
-/** Hitch stats */
-DECLARE_DWORD_ACCUMULATOR_STAT_EXTERN(TEXT("5.0s  - ....  hitches"), STAT_FPSChart_Hitch_5000_Plus, STATGROUP_FPSChart, );
-DECLARE_DWORD_ACCUMULATOR_STAT_EXTERN(TEXT("2.5s  - 5.0s  hitches"), STAT_FPSChart_Hitch_2500_5000, STATGROUP_FPSChart, );
-DECLARE_DWORD_ACCUMULATOR_STAT_EXTERN(TEXT("2.0s  - 2.5s  hitches"), STAT_FPSChart_Hitch_2000_2500, STATGROUP_FPSChart, );
-DECLARE_DWORD_ACCUMULATOR_STAT_EXTERN(TEXT("1.5s  - 2.0s  hitches"), STAT_FPSChart_Hitch_1500_2000, STATGROUP_FPSChart, );
-DECLARE_DWORD_ACCUMULATOR_STAT_EXTERN(TEXT("1.0s  - 1.5s  hitches"), STAT_FPSChart_Hitch_1000_1500, STATGROUP_FPSChart, );
-DECLARE_DWORD_ACCUMULATOR_STAT_EXTERN(TEXT("0.75s - 1.0s  hitches"), STAT_FPSChart_Hitch_750_1000, STATGROUP_FPSChart, );
-DECLARE_DWORD_ACCUMULATOR_STAT_EXTERN(TEXT("0.5s  - 0.75s hitches"), STAT_FPSChart_Hitch_500_750, STATGROUP_FPSChart, );
-DECLARE_DWORD_ACCUMULATOR_STAT_EXTERN(TEXT("0.3s  - 0.5s  hitches"), STAT_FPSChart_Hitch_300_500, STATGROUP_FPSChart, );
-DECLARE_DWORD_ACCUMULATOR_STAT_EXTERN(TEXT("0.2s  - 0.3s  hitches"), STAT_FPSChart_Hitch_200_300, STATGROUP_FPSChart, );
-DECLARE_DWORD_ACCUMULATOR_STAT_EXTERN(TEXT("0.15s - 0.2s  hitches"), STAT_FPSChart_Hitch_150_200, STATGROUP_FPSChart, );
-DECLARE_DWORD_ACCUMULATOR_STAT_EXTERN(TEXT("0.1s  - 0.15s hitches"), STAT_FPSChart_Hitch_100_150, STATGROUP_FPSChart, );
-DECLARE_DWORD_ACCUMULATOR_STAT_EXTERN(TEXT("0.06s  - 0.1s hitches"), STAT_FPSChart_Hitch_60_100, STATGROUP_FPSChart, );
-DECLARE_DWORD_ACCUMULATOR_STAT_EXTERN(TEXT("Total hitches"), STAT_FPSChart_TotalHitchCount, STATGROUP_FPSChart, );
-
-/** Unit time stats */
-DECLARE_FLOAT_COUNTER_STAT_EXTERN(TEXT("StatUnit FrameTime"), STAT_FPSChart_UnitFrame, STATGROUP_FPSChart, );
-DECLARE_FLOAT_COUNTER_STAT_EXTERN(TEXT("StatUnit RenderThreadTime"), STAT_FPSChart_UnitRender, STATGROUP_FPSChart, );
-DECLARE_FLOAT_COUNTER_STAT_EXTERN(TEXT("StatUnit GameThreadTime"), STAT_FPSChart_UnitGame, STATGROUP_FPSChart, );
-DECLARE_FLOAT_COUNTER_STAT_EXTERN(TEXT("StatUnit GPUTime"), STAT_FPSChart_UnitGPU, STATGROUP_FPSChart, );
-
-
 /**
  * Path finding stats
  */
