UE4 - fix UE-15630, editing of ISMC and HISMC

diff --git a/Engine/Source/Runtime/Engine/Classes/Components/HierarchicalInstancedStaticMeshComponent.h b/Engine/Source/Runtime/Engine/Classes/Components/HierarchicalInstancedStaticMeshComponent.h
--- a/Engine/Source/Runtime/Engine/Classes/Components/HierarchicalInstancedStaticMeshComponent.h
+++ b/Engine/Source/Runtime/Engine/Classes/Components/HierarchicalInstancedStaticMeshComponent.h
@@ -121,6 +121,9 @@ public:
 	virtual FPrimitiveSceneProxy* CreateSceneProxy() override;
 	virtual void PostLoad() override;
 	virtual FBoxSphereBounds CalcBounds(const FTransform& BoundTransform) const override;
+#if WITH_EDITOR
+	virtual void PostEditChangeChainProperty(FPropertyChangedChainEvent& PropertyChangedEvent) override;
+#endif
 
 	// UInstancedStaticMesh interface
 	virtual int32 AddInstance(const FTransform& InstanceTransform) override;
diff --git a/Engine/Source/Runtime/Engine/Private/HierarchicalInstancedStaticMesh.cpp b/Engine/Source/Runtime/Engine/Private/HierarchicalInstancedStaticMesh.cpp
--- a/Engine/Source/Runtime/Engine/Private/HierarchicalInstancedStaticMesh.cpp
+++ b/Engine/Source/Runtime/Engine/Private/HierarchicalInstancedStaticMesh.cpp
@@ -1183,8 +1183,6 @@ void FHierarchicalStaticMeshSceneProxy::FillDynamicMeshElements(FMeshElementColl
 #endif
 }
 
-
-
 void FHierarchicalStaticMeshSceneProxy::GetDynamicMeshElements(const TArray<const FSceneView*>& Views, const FSceneViewFamily& ViewFamily, uint32 VisibilityMap, FMeshElementCollector& Collector) const
 {
 	if (Views[0]->bRenderFirstInstanceOnly)
@@ -1546,6 +1544,19 @@ UHierarchicalInstancedStaticMeshComponent::~UHierarchicalInstancedStaticMeshComp
 	FlushAsyncBuildInstanceBufferTask();
 }
 
+#if WITH_EDITOR
+void UHierarchicalInstancedStaticMeshComponent::PostEditChangeChainProperty(FPropertyChangedChainEvent& PropertyChangedEvent)
+{
+	if ((PropertyChangedEvent.Property != NULL && PropertyChangedEvent.Property->GetFName() == "PerInstanceSMData") ||
+		(PropertyChangedEvent.Property != NULL && PropertyChangedEvent.Property->GetFName() == "Transform"))
+	{
+		BuildTree();
+	}
+
+	Super::PostEditChangeChainProperty(PropertyChangedEvent);
+}
+#endif
+
 void UHierarchicalInstancedStaticMeshComponent::FlushAsyncBuildInstanceBufferTask()
 {
 	if (AsyncBuildInstanceBufferTask)
diff --git a/Engine/Source/Runtime/Engine/Private/InstancedStaticMesh.cpp b/Engine/Source/Runtime/Engine/Private/InstancedStaticMesh.cpp
--- a/Engine/Source/Runtime/Engine/Private/InstancedStaticMesh.cpp
+++ b/Engine/Source/Runtime/Engine/Private/InstancedStaticMesh.cpp
@@ -1683,6 +1683,11 @@ void UInstancedStaticMeshComponent::PostEditChangeChainProperty(FPropertyChanged
 		ReleasePerInstanceRenderData();
 		MarkRenderStateDirty();
 	}
+	else if (PropertyChangedEvent.Property != NULL && PropertyChangedEvent.Property->GetFName() == "Transform")
+	{
+		ReleasePerInstanceRenderData();
+		MarkRenderStateDirty();
+	}
 
 	Super::PostEditChangeChainProperty(PropertyChangedEvent);
 }
