Localization should work by default on Node titles and pins.

diff --git a/Engine/Source/Runtime/Engine/Private/EdGraph/EdGraphPin.cpp b/Engine/Source/Runtime/Engine/Private/EdGraph/EdGraphPin.cpp
--- a/Engine/Source/Runtime/Engine/Private/EdGraph/EdGraphPin.cpp
+++ b/Engine/Source/Runtime/Engine/Private/EdGraph/EdGraphPin.cpp
@@ -272,7 +272,7 @@ FText UEdGraphPin::GetDisplayName() const
 	{
 		DisplayName = (!PinFriendlyName.IsEmpty()) ? PinFriendlyName : FText::FromString(PinName);
 
-		bool bShowNodesAndPinsUnlocalized;
+		bool bShowNodesAndPinsUnlocalized = false;
 		GConfig->GetBool( TEXT("Internationalization"), TEXT("ShowNodesAndPinsUnlocalized"), bShowNodesAndPinsUnlocalized, GEditorSettingsIni );
 		if (bShowNodesAndPinsUnlocalized)
 		{
diff --git a/Engine/Source/Runtime/Engine/Public/EdGraph/EdGraphNodeUtils.h b/Engine/Source/Runtime/Engine/Public/EdGraph/EdGraphNodeUtils.h
--- a/Engine/Source/Runtime/Engine/Public/EdGraph/EdGraphNodeUtils.h
+++ b/Engine/Source/Runtime/Engine/Public/EdGraph/EdGraphNodeUtils.h
@@ -104,7 +104,7 @@ private:
 	 *  node and pin titles localized and updates the cache accordingly */
 	FORCEINLINE void UpdateCacheIternal(FText const& InText) const
 	{
-		bool bShowNodesAndPinsUnlocalized;
+		bool bShowNodesAndPinsUnlocalized = false;
 		GConfig->GetBool( TEXT("Internationalization"), TEXT("ShowNodesAndPinsUnlocalized"), bShowNodesAndPinsUnlocalized, GEditorSettingsIni );
 
 		if (bShowNodesAndPinsUnlocalized)
