#ue4 - Remove unused var in InterpToMovementComponent

UE-15645

diff --git a/Engine/Source/Runtime/Engine/Classes/Components/InterpToMovementComponent.h b/Engine/Source/Runtime/Engine/Classes/Components/InterpToMovementComponent.h
--- a/Engine/Source/Runtime/Engine/Classes/Components/InterpToMovementComponent.h
+++ b/Engine/Source/Runtime/Engine/Classes/Components/InterpToMovementComponent.h
@@ -181,10 +181,6 @@ class ENGINE_API UInterpToMovementComponent : public UMovementComponent
 	// End UObject interface.
 #endif // WITH_EDITOR
 
-	// The position of the point 
-	UPROPERTY(EditAnywhere, BlueprintReadWrite, Category = Location, meta = (MakeEditWidget = true))
-		FVector PositionControlPointTest;
-
 protected:
 
 	/** @return true if the simulation should stop. */
