Fix crash when calling SetWorldLocation on a destructible mesh in its first frame in the world. See comment about APEX support

UE-15611

diff --git a/Engine/Source/Runtime/Engine/Private/Components/DestructibleComponent.cpp b/Engine/Source/Runtime/Engine/Private/Components/DestructibleComponent.cpp
--- a/Engine/Source/Runtime/Engine/Private/Components/DestructibleComponent.cpp
+++ b/Engine/Source/Runtime/Engine/Private/Components/DestructibleComponent.cpp
@@ -162,7 +162,13 @@ void UDestructibleComponent::OnUpdateTransform(bool bSkipPhysicsMove)
 	if (ApexDestructibleActor != NULL && !BodyInstance.bSimulatePhysics)
 	{
 		PxMat44 GlobalPose(PxMat33(U2PQuat(CurrentLocalToWorld.GetRotation())), U2PVector(CurrentLocalToWorld.GetTranslation()));
-		ApexDestructibleActor->setGlobalPose(GlobalPose);
+
+		PxRigidDynamic* PRootActor = ApexDestructibleActor->getChunkPhysXActor(0);
+		if(!PRootActor || PRootActor->getScene())	//either root chunk is null meaning fractured (so there's a scene), or the root has a scene
+		{
+			ApexDestructibleActor->setGlobalPose(GlobalPose);
+		}
+		//TODO: This misses an OnUpdateTransform for the initial frame. This needs APEX support so waiting on that.
 	}
 #endif // #if WITH_APEX
 }
