#ue4 - Make StopMovementImmediately() work for InterpToMovementComponent.

UE-15646

diff --git a/Engine/Source/Runtime/Engine/Classes/Components/InterpToMovementComponent.h b/Engine/Source/Runtime/Engine/Classes/Components/InterpToMovementComponent.h
--- a/Engine/Source/Runtime/Engine/Classes/Components/InterpToMovementComponent.h
+++ b/Engine/Source/Runtime/Engine/Classes/Components/InterpToMovementComponent.h
@@ -114,6 +114,7 @@ class ENGINE_API UInterpToMovementComponent : public UMovementComponent
 
 	//Begin UMovementComponent Interface
 	virtual void InitializeComponent() override;
+	virtual void StopMovementImmediately() override;
 	//End UMovementComponent Interface
 
 	/**
diff --git a/Engine/Source/Runtime/Engine/Private/Components/InterpToMovementComponent.cpp b/Engine/Source/Runtime/Engine/Private/Components/InterpToMovementComponent.cpp
--- a/Engine/Source/Runtime/Engine/Private/Components/InterpToMovementComponent.cpp
+++ b/Engine/Source/Runtime/Engine/Private/Components/InterpToMovementComponent.cpp
@@ -32,6 +32,17 @@ void UInterpToMovementComponent::InitializeComponent()
 }
 
 
+void UInterpToMovementComponent::StopMovementImmediately()
+{
+	bStopped = true;
+	FHitResult FakeHit;
+	OnInterpToStop.Broadcast(FakeHit, CurrentTime);
+	// Not calling StopSimulating here (that nulls UpdatedComponent). Users can call that explicitly instead if they wish.
+
+	Super::StopMovementImmediately();
+}
+
+
 void UInterpToMovementComponent::TickComponent(float DeltaTime, enum ELevelTick TickType, FActorComponentTickFunction *ThisTickFunction)
 {
 	QUICK_SCOPE_CYCLE_COUNTER(STAT_InterpToMovementComponent_TickComponent);
