cleanup

diff --git a/Engine/Source/Runtime/Engine/Private/ShaderCompiler/ShaderCompiler.cpp b/Engine/Source/Runtime/Engine/Private/ShaderCompiler/ShaderCompiler.cpp
--- a/Engine/Source/Runtime/Engine/Private/ShaderCompiler/ShaderCompiler.cpp
+++ b/Engine/Source/Runtime/Engine/Private/ShaderCompiler/ShaderCompiler.cpp
@@ -512,9 +512,6 @@ struct FShaderCompileWorkerInfo
 
 };
 
-//--------------------
-// @LH : LUTHA - XGE Distributed Shader Compilation
-//-----
 FShaderCompileThreadRunnableBase::FShaderCompileThreadRunnableBase(FShaderCompilingManager* InManager)
 	: Manager(InManager)
 	, Thread(nullptr)
@@ -524,9 +521,6 @@ FShaderCompileThreadRunnableBase::FShaderCompileThreadRunnableBase(FShaderCompil
 }
 void FShaderCompileThreadRunnableBase::StartThread()
 {
-//-----
-// @LH : END
-//--------------------
 	if (Manager->bAllowAsynchronousShaderCompiling && !FPlatformProperties::RequiresCookedData())
 	{
 		Thread = FRunnableThread::Create(this, TEXT("ShaderCompilingThread"), 0, TPri_Normal, FPlatformAffinity::GetPoolThreadMask());
