Change 2553100 by [REDACTED_PERSON][REDACTED_HANDLE] on 2015/05/15 13:12:13

	Fixed lighting builds invalidating navigation

 Merging 2553100

//depot/UE4-UT/...

to //depot/UE4/...

diff --git a/Engine/Source/Runtime/Engine/Private/AI/Navigation/NavigationSystem.cpp b/Engine/Source/Runtime/Engine/Private/AI/Navigation/NavigationSystem.cpp
--- a/Engine/Source/Runtime/Engine/Private/AI/Navigation/NavigationSystem.cpp
+++ b/Engine/Source/Runtime/Engine/Private/AI/Navigation/NavigationSystem.cpp
@@ -2815,8 +2815,29 @@ void UNavigationSystem::AddNavigationBoundsUpdateRequest(const FNavigationBounds
 
 	if (ExistingIdx != INDEX_NONE)
 	{
-		// Overwrite any previous updates
-		PendingNavBoundsUpdates[ExistingIdx] = UpdateRequest;
+		// catch the case where the bounds was removed and immediately re-added with the same bounds as before
+		// in that case, we can cancel any update at all
+		bool bCanCancelUpdate = false;
+		if (PendingNavBoundsUpdates[ExistingIdx].UpdateRequest == FNavigationBoundsUpdateRequest::Removed && UpdateRequest.UpdateRequest == FNavigationBoundsUpdateRequest::Added)
+		{
+			for (TSet<FNavigationBounds>::TConstIterator It(RegisteredNavBounds); It; ++It)
+			{
+				if (*It == UpdateRequest.NavBounds)
+				{
+					bCanCancelUpdate = true;
+					break;
+				}
+			}
+		}
+		if (bCanCancelUpdate)
+		{
+			PendingNavBoundsUpdates.RemoveAt(ExistingIdx);
+		}
+		else
+		{
+			// Overwrite any previous updates
+			PendingNavBoundsUpdates[ExistingIdx] = UpdateRequest;
+		}
 	}
 	else
 	{
