Fix crash caused by misuse of FMath::Clamp

diff --git a/Engine/Source/Runtime/Engine/Private/AudioSettings.cpp b/Engine/Source/Runtime/Engine/Private/AudioSettings.cpp
--- a/Engine/Source/Runtime/Engine/Private/AudioSettings.cpp
+++ b/Engine/Source/Runtime/Engine/Private/AudioSettings.cpp
@@ -93,7 +93,7 @@ void UAudioSettings::PostEditChangeChainProperty(FPropertyChangedChainEvent& Pro
 
 const FAudioQualitySettings& UAudioSettings::GetQualityLevelSettings(int32 QualityLevel) const
 {
-	return QualityLevels[FMath::Clamp(0, QualityLevel, QualityLevels.Num() - 1)];
+	return QualityLevels[FMath::Clamp(QualityLevel, 0, QualityLevels.Num() - 1)];
 }
 
 #undef LOCTEXT_NAMESPACE
\ No newline at end of file
