Fix threading issue where pending hits were not separated by scene type.

UE-15690

diff --git a/Engine/Source/Runtime/Engine/Private/PhysicsEngine/PhysScene.cpp b/Engine/Source/Runtime/Engine/Private/PhysicsEngine/PhysScene.cpp
--- a/Engine/Source/Runtime/Engine/Private/PhysicsEngine/PhysScene.cpp
+++ b/Engine/Source/Runtime/Engine/Private/PhysicsEngine/PhysScene.cpp
@@ -820,8 +820,10 @@ void FPhysScene::DispatchPhysNotifications_AssumesLocked()
 {
 	SCOPE_CYCLE_COUNTER(STAT_PhysicsEventTime);
 
-	//Collision notification
+	for(int32 SceneType = 0; SceneType < PST_MAX; ++SceneType)
 	{
+		TArray<FCollisionNotifyInfo>& PendingCollisionNotifies = GetPendingCollisionNotifies(SceneType);
+
 		// Let the game-specific PhysicsCollisionHandler process any physics collisions that took place
 		if (OwningWorld != NULL && OwningWorld->PhysicsCollisionHandler != NULL)
 		{
diff --git a/Engine/Source/Runtime/Engine/Private/PhysicsEngine/PhysXSupport.cpp b/Engine/Source/Runtime/Engine/Private/PhysicsEngine/PhysXSupport.cpp
--- a/Engine/Source/Runtime/Engine/Private/PhysicsEngine/PhysXSupport.cpp
+++ b/Engine/Source/Runtime/Engine/Private/PhysicsEngine/PhysXSupport.cpp
@@ -473,8 +473,11 @@ void FPhysXSimEventCallback::onContact(const PxContactPairHeader& PairHeader, co
 	FPhysScene* PhysScene = FPhysxUserData::Get<FPhysScene>(PScene->userData);
 	check(PhysScene);
 
-	uint32 PreAddingCollisionNotify = PhysScene->PendingCollisionNotifies.Num() - 1;
-	TArray<int32> PairNotifyMapping = FBodyInstance::AddCollisionNotifyInfo(BodyInst0, BodyInst1, Pairs, NumPairs, PhysScene->PendingCollisionNotifies);
+	const EPhysicsSceneType SceneType = PhysScene->GetPhysXScene(PST_Sync) == PScene ? PST_Sync : PST_Async;
+	TArray<FCollisionNotifyInfo>& PendingCollisionNotifies = PhysScene->GetPendingCollisionNotifies(SceneType);
+
+	uint32 PreAddingCollisionNotify = PendingCollisionNotifies.Num() - 1;
+	TArray<int32> PairNotifyMapping = FBodyInstance::AddCollisionNotifyInfo(BodyInst0, BodyInst1, Pairs, NumPairs, PendingCollisionNotifies);
 
 	// Iterate through contact points
 	for(uint32 PairIdx=0; PairIdx<NumPairs; PairIdx++)
@@ -485,7 +488,7 @@ void FPhysXSimEventCallback::onContact(const PxContactPairHeader& PairHeader, co
 			continue;
 		}
 
-		FCollisionNotifyInfo * NotifyInfo = &PhysScene->PendingCollisionNotifies[NotifyIdx];
+		FCollisionNotifyInfo * NotifyInfo = &PendingCollisionNotifies[NotifyIdx];
 		FCollisionImpactData* ImpactInfo = &(NotifyInfo->RigidCollisionData);
 
 		const PxContactPair* Pair = Pairs + PairIdx;
@@ -525,14 +528,14 @@ void FPhysXSimEventCallback::onContact(const PxContactPairHeader& PairHeader, co
 		}	
 	}
 
-	for (int32 NotifyIdx = PreAddingCollisionNotify + 1; NotifyIdx < PhysScene->PendingCollisionNotifies.Num(); NotifyIdx++)
+	for (int32 NotifyIdx = PreAddingCollisionNotify + 1; NotifyIdx < PendingCollisionNotifies.Num(); NotifyIdx++)
 	{
-		FCollisionNotifyInfo * NotifyInfo = &PhysScene->PendingCollisionNotifies[NotifyIdx];
+		FCollisionNotifyInfo * NotifyInfo = &PendingCollisionNotifies[NotifyIdx];
 		FCollisionImpactData* ImpactInfo = &(NotifyInfo->RigidCollisionData);
 		// Discard pairs that don't generate any force (eg. have been rejected through a modify contact callback).
 		if (ImpactInfo->TotalNormalImpulse.SizeSquared() < KINDA_SMALL_NUMBER)
 		{
-			PhysScene->PendingCollisionNotifies.RemoveAt(NotifyIdx);
+			PendingCollisionNotifies.RemoveAt(NotifyIdx);
 			NotifyIdx--;
 		}
 	}
diff --git a/Engine/Source/Runtime/Engine/Public/PhysicsPublic.h b/Engine/Source/Runtime/Engine/Public/PhysicsPublic.h
--- a/Engine/Source/Runtime/Engine/Public/PhysicsPublic.h
+++ b/Engine/Source/Runtime/Engine/Public/PhysicsPublic.h
@@ -273,9 +273,9 @@ public:
 	/** Stores the number of valid scenes we are working with. This will be PST_MAX or PST_Async, 
 		depending on whether the async scene is enabled or not*/
 	uint32							NumPhysScenes;
-
-	/** Array of collision notifications, pending execution at the end of the physics engine run. */
-	TArray<FCollisionNotifyInfo>	PendingCollisionNotifies;
+	
+	/** Gets the array of collision notifications, pending execution at the end of the physics engine run. */
+	TArray<FCollisionNotifyInfo>& GetPendingCollisionNotifies(int32 SceneType){ return PendingCollisionData[SceneType].PendingCollisionNotifies; }
 
 	/** World that owns this physics scene */
 	UWorld*							OwningWorld;
@@ -379,6 +379,14 @@ private:
 #endif
 #endif	//
 
+	struct FPendingCollisionData
+	{
+		/** Array of collision notifications, pending execution at the end of the physics engine run. */
+		TArray<FCollisionNotifyInfo>	PendingCollisionNotifies;
+	};
+
+	FPendingCollisionData PendingCollisionData[PST_MAX];
+
 public:
 
 #if WITH_PHYSX
