Fix destructible mesh not getting its position set during OnBeginPlay

UE-15684

diff --git a/Engine/Source/Runtime/Engine/Private/Components/DestructibleComponent.cpp b/Engine/Source/Runtime/Engine/Private/Components/DestructibleComponent.cpp
--- a/Engine/Source/Runtime/Engine/Private/Components/DestructibleComponent.cpp
+++ b/Engine/Source/Runtime/Engine/Private/Components/DestructibleComponent.cpp
@@ -161,14 +161,15 @@ void UDestructibleComponent::OnUpdateTransform(bool bSkipPhysicsMove)
 #if WITH_APEX
 	if (ApexDestructibleActor != NULL && !BodyInstance.bSimulatePhysics)
 	{
-		PxMat44 GlobalPose(PxMat33(U2PQuat(CurrentLocalToWorld.GetRotation())), U2PVector(CurrentLocalToWorld.GetTranslation()));
-
 		PxRigidDynamic* PRootActor = ApexDestructibleActor->getChunkPhysXActor(0);
+		PxMat44 GlobalPose(PxMat33(U2PQuat(CurrentLocalToWorld.GetRotation())), U2PVector(CurrentLocalToWorld.GetTranslation()));
 		if(!PRootActor || PRootActor->getScene())	//either root chunk is null meaning fractured (so there's a scene), or the root has a scene
 		{
 			ApexDestructibleActor->setGlobalPose(GlobalPose);
+		}else
+		{
+			PRootActor->setGlobalPose(PxTransform(GlobalPose));	//we're not in a scene yet, so place the root actor in this new position
 		}
-		//TODO: This misses an OnUpdateTransform for the initial frame. This needs APEX support so waiting on that.
 	}
 #endif // #if WITH_APEX
 }
