Major refactor of actor iterators to support ranges.
New FActorRange, TActorRange and FSelectedActorRange.
Some example usage of the new range types.
FNetRelevantActorIterator removed, as it's not used.
FTickableLevelFilter removed.
Actor iterators' operator bool made explicit.

diff --git a/Engine/Source/Editor/LevelEditor/Private/LevelEditorActions.cpp b/Engine/Source/Editor/LevelEditor/Private/LevelEditorActions.cpp
--- a/Engine/Source/Editor/LevelEditor/Private/LevelEditorActions.cpp
+++ b/Engine/Source/Editor/LevelEditor/Private/LevelEditorActions.cpp
@@ -1603,9 +1603,9 @@ void FLevelEditorActionCallbacks::OnSelectAllLights()
 {
 	GEditor->GetSelectedActors()->BeginBatchSelectOperation();
 	// Select all light actors.
-	for( TActorIterator<ALight> It(GetWorld()); It; ++It )
+	for( ALight* Light : TActorRange<ALight>(GetWorld()) )
 	{
-		GUnrealEd->SelectActor( *It, true, false, false );
+		GUnrealEd->SelectActor( Light, true, false, false );
 	}
 
 	GEditor->GetSelectedActors()->EndBatchSelectOperation();
diff --git a/Engine/Source/Editor/LevelEditor/Private/LevelEditorToolBar.cpp b/Engine/Source/Editor/LevelEditor/Private/LevelEditorToolBar.cpp
--- a/Engine/Source/Editor/LevelEditor/Private/LevelEditorToolBar.cpp
+++ b/Engine/Source/Editor/LevelEditor/Private/LevelEditorToolBar.cpp
@@ -2076,7 +2076,7 @@ TSharedRef< SWidget > FLevelEditorToolBar::GenerateMatineeMenuContent( TSharedRe
 	bool bHasAnyMatineeActors = false;
 	TActorIterator<AMatineeActor> MatineeIt( LevelEditorWeakPtr.Pin()->GetWorld() );
 
-	bHasAnyMatineeActors = MatineeIt;
+	bHasAnyMatineeActors = !!MatineeIt;
 
 	//Add a heading to separate the existing matinees from the 'Add New Matinee Actor' button
 	MenuBuilder.BeginSection("LevelEditorExistingMatinee", LOCTEXT( "MatineeMenuCombo_ExistingHeading", "Edit Existing Matinee" ) );
diff --git a/Engine/Source/Editor/StatsViewer/Private/StatsPages/TextureStatsPage.cpp b/Engine/Source/Editor/StatsViewer/Private/StatsPages/TextureStatsPage.cpp
--- a/Engine/Source/Editor/StatsViewer/Private/StatsPages/TextureStatsPage.cpp
+++ b/Engine/Source/Editor/StatsViewer/Private/StatsPages/TextureStatsPage.cpp
@@ -38,9 +38,9 @@ struct TextureStatsGenerator : public FFindReferencedAssets
 		if( InObjectSet == TextureObjectSet_SelectedActors )
 		{
 			// In this mode only get selected actors
-			for( FSelectedActorIterator It(GetWorld()); It; ++It )
+			for( AActor* Actor : FSelectedActorRange(GetWorld()) )
 			{
-				OutObjectsToSearch.Add( *It );
+				OutObjectsToSearch.Add( Actor );
 			}
 		}
 		else if( InObjectSet == TextureObjectSet_SelectedMaterials )
diff --git a/Engine/Source/Editor/UnrealEd/Private/Layers/Layers.cpp b/Engine/Source/Editor/UnrealEd/Private/Layers/Layers.cpp
--- a/Engine/Source/Editor/UnrealEd/Private/Layers/Layers.cpp
+++ b/Engine/Source/Editor/UnrealEd/Private/Layers/Layers.cpp
@@ -447,9 +447,8 @@ bool FLayers::SelectActorsInLayers( const TArray< FName >& LayerNames, bool bSel
 	bool bChangesOccurred = false;
 
 	// Iterate over all actors, looking for actors in the specified layers.
-	for( FActorIterator It(GetWorld()); It ; ++It )
+	for( const TWeakObjectPtr< AActor > Actor : FActorRange(GetWorld()) )
 	{
-		const TWeakObjectPtr< AActor > Actor = *It;
 		if( !IsActorValidForLayer( Actor ) )
 		{
 			continue;
diff --git a/Engine/Source/Runtime/Engine/Private/TickTaskManager.cpp b/Engine/Source/Runtime/Engine/Private/TickTaskManager.cpp
--- a/Engine/Source/Runtime/Engine/Private/TickTaskManager.cpp
+++ b/Engine/Source/Runtime/Engine/Private/TickTaskManager.cpp
@@ -721,7 +721,7 @@ private:
 		for( int32 LevelIndex = 0; LevelIndex < World->GetNumLevels(); LevelIndex++ )
 		{
 			ULevel* Level = World->GetLevel(LevelIndex);
-			if (FTickableLevelFilter::CanIterateLevel(Level))
+			if (Level->bIsVisible)
 			{
 				check(Level->TickTaskLevel);
 				LevelList.Add(Level->TickTaskLevel);
@@ -751,7 +751,7 @@ private:
 		for( int32 LevelIndex = 0; LevelIndex < InWorld->GetNumLevels(); LevelIndex++ )
 		{
 			ULevel* Level = InWorld->GetLevel(LevelIndex);
-			if (FTickableLevelFilter::CanIterateLevel(Level))
+			if (Level->bIsVisible)
 			{
 				check(Level->TickTaskLevel);
 				Level->TickTaskLevel->DumpAllTickFunctions(Ar, EnabledCount, DisabledCount, bEnabled, bDisabled);
diff --git a/Engine/Source/Runtime/Engine/Private/World.cpp b/Engine/Source/Runtime/Engine/Private/World.cpp
--- a/Engine/Source/Runtime/Engine/Private/World.cpp
+++ b/Engine/Source/Runtime/Engine/Private/World.cpp
@@ -1320,9 +1320,9 @@ void UWorld::UpdateCullDistanceVolumes(AActor* ActorToUpdate, UPrimitiveComponen
 		}
 		else
 		{
-			for( FActorIterator It(this); It; ++It )
+			for( AActor* Actor : FActorRange(this) )
 			{
-				TInlineComponentArray<UPrimitiveComponent*> PrimitiveComponents(*It);
+				TInlineComponentArray<UPrimitiveComponent*> PrimitiveComponents(Actor);
 				for (UPrimitiveComponent* PrimitiveComponent : PrimitiveComponents)
 				{
 					if (ACullDistanceVolume::CanBeAffectedByVolumes(PrimitiveComponent))
@@ -1331,7 +1331,7 @@ void UWorld::UpdateCullDistanceVolumes(AActor* ActorToUpdate, UPrimitiveComponen
 					}
 				}
 
-				ACullDistanceVolume* CullDistanceVolume = Cast<ACullDistanceVolume>(*It);
+				ACullDistanceVolume* CullDistanceVolume = Cast<ACullDistanceVolume>(Actor);
 				if (CullDistanceVolume)
 				{
 					CullDistanceVolumes.Add(CullDistanceVolume);
diff --git a/Engine/Source/Runtime/Engine/Public/EngineUtils.h b/Engine/Source/Runtime/Engine/Public/EngineUtils.h
--- a/Engine/Source/Runtime/Engine/Public/EngineUtils.h
+++ b/Engine/Source/Runtime/Engine/Public/EngineUtils.h
@@ -11,6 +11,9 @@
 #include "Engine/Level.h"
 #include "GameFramework/WorldSettings.h"
 #include "HitProxies.h"
+#include "Misc/Optional.h"
+
+class FActorRange;
 
 /*-----------------------------------------------------------------------------
 	Hit proxies.
@@ -136,17 +139,13 @@ struct HTranslucentActor : public HActor
 -----------------------------------------------------------------------------*/
 
 /**
- * Abstract base class for actor iteration. Implements all operators and relies on IsSuitable
- * to be overridden by derived classed. Also the code currently relies on the ++ operator being
- * called from within the constructor which is bad as it uses IsSuitable which is a virtual function.
- * This means that all derived classes are treated as "final" and they need to manually call
- * ++(*this) in their constructor so it ends up calling the right one. The long term plan is to
- * replace the use of virtual functions with template tricks.
+ * Abstract base class for actor iteration. Implements all operators and relies on IsActorSuitable
+ * to be overridden by derived class.
  * Note that when Playing In Editor, this will find actors only in CurrentWorld.
  */
-class FActorIteratorBase
+class FActorIteratorState
 {
-protected:
+public:
 	/** Current world we are iterating upon						*/
 	UWorld* CurrentWorld;
 	/** Results from the GetObjectsOfClass query				*/
@@ -159,8 +158,6 @@ protected:
 	int32		ConsideredCount;
 	/** Current actor pointed to by actor iterator				*/
 	AActor*	CurrentActor;
-	/** Delegate for catching spawned actors					*/
-	FOnActorSpawned::FDelegate ActorSpawnedDelegate;
 	/** Contains any actors spawned during iteration			*/
 	TArray<AActor*> SpawnedActorArray;
 	/** The class type we are iterating, kept for filtering		*/
@@ -171,7 +168,7 @@ protected:
 	/**
 	 * Default ctor, inits everything
 	 */
-	FActorIteratorBase( UWorld* InWorld, TSubclassOf<class AActor> InClass = AActor::StaticClass() ) :
+	FActorIteratorState( UWorld* InWorld, TSubclassOf<AActor> InClass ) :
 		CurrentWorld( InWorld ),
 		Index( -1 ),
 		ReachedEnd( false ),
@@ -184,15 +181,28 @@ protected:
 		EObjectFlags ExcludeFlags = RF_ClassDefaultObject|RF_PendingKill;
 		GetObjectsOfClass(InClass, ObjectArray, true, ExcludeFlags);
 
-		ActorSpawnedDelegate = FOnActorSpawned::FDelegate::CreateRaw(this, &FActorIteratorBase::OnActorSpawned);
+		auto ActorSpawnedDelegate = FOnActorSpawned::FDelegate::CreateRaw(this, &FActorIteratorState::OnActorSpawned);
 		ActorSpawnedDelegateHandle = CurrentWorld->AddOnActorSpawnedHandler(ActorSpawnedDelegate);
 	}
 
-	~FActorIteratorBase()
+	~FActorIteratorState()
 	{
 		CurrentWorld->RemoveOnActorSpawnedHandler(ActorSpawnedDelegateHandle);
 	}
 
+	/**
+	 * Returns the current suitable actor pointed at by the Iterator
+	 *
+	 * @return	Current suitable actor
+	 */
+	FORCEINLINE AActor* GetActorChecked() const
+	{
+		check(CurrentActor);
+		checkf(!CurrentActor->HasAnyFlags(RF_Unreachable), TEXT("%s"), *CurrentActor->GetFullName());
+		return CurrentActor;
+	}
+
+private:
 	void OnActorSpawned(AActor* InActor)
 	{
 		if (InActor->IsA(DesiredClass))
@@ -200,32 +210,79 @@ protected:
 			SpawnedActorArray.AddUnique(InActor);
 		}
 	}
+};
 
+enum class EActorIteratorType
+{
+	End
+};
+
+/**
+ * Template class used to filter actors by certain characteristics
+ */
+template <typename Derived>
+class TActorIteratorBase
+{
 public:
+	/**
+	 * Iterates to next suitable actor.
+	 */
+	void operator++()
+	{
+		// Use local version to avoid LHSs as compiler is not required to write out member variables to memory.
+		AActor*           LocalCurrentActor      = nullptr;
+		int32             LocalIndex             = State->Index;
+		TArray<UObject*>& LocalObjectArray       = State->ObjectArray;
+		TArray<AActor*>&  LocalSpawnedActorArray = State->SpawnedActorArray;
+		UWorld*           LocalCurrentWorld      = State->CurrentWorld;
+		while(++LocalIndex < (LocalObjectArray.Num() + LocalSpawnedActorArray.Num()))
+		{
+			if (LocalIndex < LocalObjectArray.Num())
+			{
+				LocalCurrentActor = static_cast<AActor*>(LocalObjectArray[LocalIndex]);
+			}
+			else
+			{
+				LocalCurrentActor = LocalSpawnedActorArray[LocalIndex - LocalObjectArray.Num()];
+			}
+			State->ConsideredCount++;
 
-	
+			if ( LocalCurrentActor != NULL
+				&& static_cast<const Derived*>(this)->IsActorSuitable(LocalCurrentActor)
+				&& static_cast<const Derived*>(this)->CanIterateLevel(LocalCurrentActor->GetLevel())
+				&& LocalCurrentActor->GetWorld() == LocalCurrentWorld)
+			{
+				// ignore non-persistent world settings
+				if (LocalCurrentActor->GetLevel() == LocalCurrentWorld->PersistentLevel || !LocalCurrentActor->IsA(AWorldSettings::StaticClass()))
+				{
+					State->CurrentActor = LocalCurrentActor;
+					State->Index = LocalIndex;
+					return;
+				}
+			}
+		}
+		State->CurrentActor = NULL;
+		State->ReachedEnd = true;
+	}
 
 	/**
 	 * Returns the current suitable actor pointed at by the Iterator
 	 *
 	 * @return	Current suitable actor
 	 */
-	FORCEINLINE AActor* operator*()
+	FORCEINLINE AActor* operator*() const
 	{
-		check(CurrentActor);
-		checkf(!CurrentActor->HasAnyFlags(RF_Unreachable), TEXT("%s"), *CurrentActor->GetFullName());
-		return CurrentActor;
+		return State->GetActorChecked();
 	}
+
 	/**
 	 * Returns the current suitable actor pointed at by the Iterator
 	 *
 	 * @return	Current suitable actor
 	 */
-	FORCEINLINE AActor* operator->()
+	FORCEINLINE AActor* operator->() const
 	{
-		check(CurrentActor);
-		checkf(!CurrentActor->HasAnyFlags(RF_Unreachable), TEXT("%s"), *CurrentActor->GetFullName());
-		return CurrentActor;
+		return State->GetActorChecked();
 	}
 	/**
 	 * Returns whether the iterator has reached the end and no longer points
@@ -233,9 +290,9 @@ public:
 	 *
 	 * @return true if iterator points to a suitable actor, false if it has reached the end
 	 */
-	FORCEINLINE operator bool()
+	FORCEINLINE_EXPLICIT_OPERATOR_BOOL() const
 	{
-		return !ReachedEnd;
+		return !State->ReachedEnd;
 	}
 
 	/**
@@ -243,8 +300,8 @@ public:
 	 */
 	void ClearCurrent()
 	{
-		check(!ReachedEnd);
-		CurrentWorld->RemoveActor(CurrentActor, true);
+		check(!State->ReachedEnd);
+		State->CurrentWorld->RemoveActor(State->CurrentActor, true);
 	}
 
 	/**
@@ -253,290 +310,261 @@ public:
 	 *
 	 * @return number of actors considered thus far.
 	 */
-	int32 GetProgressNumerator()
+	int32 GetProgressNumerator() const
 	{
-		return ConsideredCount;
+		return State->ConsideredCount;
 	}
-};
 
-/**
- * Default level iteration filter. Doesn't cull levels out.
- */
-struct FDefaultLevelFilter
-{
+protected:
 	/**
-	 * Used to examine whether this level is valid for iteration or not
-	 *
-	 * @param Level the level to check for iteration
+	 * Hide the constructors as construction on this class should only be done by subclasses
+	 */
+	explicit TActorIteratorBase(EActorIteratorType)
+	{
+	}
+
+	TActorIteratorBase(UWorld* InWorld, TSubclassOf<AActor> InClass)
+	{
+		State.Emplace(InWorld, InClass);
+	}
+
+	/**
+	 * Determines whether this is a valid actor or not.
+	 * This function should be redefined (thus hiding this one) in a derived class if it wants special actor filtering.
 	 *
-	 * @return true if the level can be iterated, false otherwise
+	 * @param	Actor	Actor to check
+	 * @return	true
 	 */
-	bool CanIterateLevel(ULevel* Level) const
+	FORCEINLINE static bool IsActorSuitable(AActor* Actor)
 	{
 		return true;
 	}
-};
 
-/**
- * Filter class that prevents ticking of levels that are still loading
- */
-struct FTickableLevelFilter
-{
 	/**
 	 * Used to examine whether this level is valid for iteration or not
+	 * This function should be redefined (thus hiding this one) in a derived class if it wants special level filtering.
 	 *
 	 * @param Level the level to check for iteration
-	 *
 	 * @return true if the level can be iterated, false otherwise
 	 */
-	static bool CanIterateLevel(ULevel* Level)
+	FORCEINLINE static bool CanIterateLevel(ULevel* Level)
 	{
-		return Level->bIsVisible;
+		return true;
 	}
+
+private:
+	TOptional<FActorIteratorState> State;
+
+	friend bool operator==(const TActorIteratorBase& Lhs, const TActorIteratorBase& Rhs) { check(!Rhs.State); return  !Lhs; }
+	friend bool operator!=(const TActorIteratorBase& Lhs, const TActorIteratorBase& Rhs) { check(!Rhs.State); return !!Lhs; }
 };
 
 /**
- * Template class used to avoid the virtual function call cost for filtering
- * actors by certain characteristics
+ * Actor iterator
+ * Note that when Playing In Editor, this will find actors only in CurrentWorld
  */
-template<typename FILTER_CLASS,typename LEVEL_FILTER_CLASS = FDefaultLevelFilter>
-class TActorIteratorBase :
-	public FActorIteratorBase
+class FActorIterator : public TActorIteratorBase<FActorIterator>
 {
-protected:
-	/**
-	 * Ref to the class that is going to filter actors
-	 */
-	const FILTER_CLASS& FilterClass;
-	/**
-	 * Ref to the object that is going to filter levels
-	 */
-	const LEVEL_FILTER_CLASS& LevelFilterClass;
+	friend class TActorIteratorBase<FActorIterator>;
+	typedef TActorIteratorBase<FActorIterator> Super;
 
+public:
 	/**
-	 * Hide the constructor as construction on this class should only be done by subclasses
+	 * Constructor
+	 *
+	 * @param  InWorld  The world whose actors are to be iterated over.
 	 */
-	TActorIteratorBase( UWorld* InWorld, TSubclassOf<class AActor> InClass = AActor::StaticClass(),
-		const FILTER_CLASS& InFilterClass = FILTER_CLASS(),
-		const LEVEL_FILTER_CLASS& InLevelClass = LEVEL_FILTER_CLASS())
-	:	FActorIteratorBase( InWorld, InClass ),
-		FilterClass(InFilterClass),
-		LevelFilterClass(InLevelClass)
+	explicit FActorIterator(UWorld* InWorld)
+		: Super(InWorld, AActor::StaticClass())
 	{
+		++(*this);
 	}
 
-public:
 	/**
-	 * Iterates to next suitable actor.
+	 * Constructor for creating an end iterator
 	 */
-	void operator++()
+	explicit FActorIterator(EActorIteratorType)
+		: Super(EActorIteratorType::End)
 	{
-		// Use local version to avoid LHSs as compiler is not required to write out member variables to memory.
-		AActor*		LocalCurrentActor	= NULL;
-		int32		LocalIndex			= Index;
-		while(++LocalIndex < (ObjectArray.Num() + SpawnedActorArray.Num()))
-		{
-			if (LocalIndex < ObjectArray.Num())
-			{
-				LocalCurrentActor = static_cast<AActor*>(ObjectArray[LocalIndex]);
-			}
-			else
-			{
-				LocalCurrentActor = SpawnedActorArray[LocalIndex - ObjectArray.Num()];
-			}
-			ConsideredCount++;
-
-			if ( LocalCurrentActor != NULL
-				&& FilterClass.IsSuitable(LocalCurrentActor)
-				&& LevelFilterClass.CanIterateLevel(LocalCurrentActor->GetLevel())
-				&& LocalCurrentActor->GetWorld() == CurrentWorld)
-			{
-				// ignore non-persistent world settings
-				if (LocalCurrentActor->GetLevel() == CurrentWorld->PersistentLevel
-					|| !LocalCurrentActor->IsA(AWorldSettings::StaticClass()))
-				{
-					CurrentActor = LocalCurrentActor;
-					Index = LocalIndex;
-					return;
-				}
-			}
-		}
-		CurrentActor = NULL;
-		ReachedEnd = true;
 	}
-};
 
-/**
- * Base filter class for actor filtering in iterators.
- */
-class FActorFilter
-{
-public:
+private:
 	/**
-	 * Determines whether this is a valid actor or not
+	 * Used to examine whether this level is valid for iteration or not
 	 *
-	 * @param	Actor	Actor to check
-	 * @return	true if actor is != NULL, false otherwise
+	 * @param Level the level to check for iteration
+	 * @return true if the level can be iterated, false otherwise
 	 */
-	FORCEINLINE bool IsSuitable( AActor* Actor ) const
+	static bool CanIterateLevel(ULevel* Level)
 	{
-		return Actor != NULL;
+		return Level->bIsVisible;
 	}
 };
 
 /**
- * Actor iterator
+ * Actor range for ranged-for support.
  * Note that when Playing In Editor, this will find actors only in CurrentWorld
  */
-class FActorIterator :
-	public TActorIteratorBase<FActorFilter, FTickableLevelFilter>
+class FActorRange
 {
 public:
 	/**
-	 * Constructor, inits the starting position for iteration
-	 * Also optionally takes in a custom UWorld
+	 * Constructor
+	 *
+	 * @param  InWorld  The world whose actors are to be iterated over.
 	 */
-	FActorIterator(UWorld* InWorld )
-		:TActorIteratorBase(InWorld)
+	explicit FActorRange(UWorld* InWorld)
+		: World(InWorld)
 	{
-		++(*this);
 	}
+
+private:
+	UWorld* World;
+
+	friend FActorIterator begin(const FActorRange& Range) { return FActorIterator(Range.World); }
+	friend FActorIterator end  (const FActorRange& Range) { return FActorIterator(EActorIteratorType::End); }
 };
 
 /**
  * Template actor iterator.
  */
 template <typename ActorType>
-class TActorIterator : 
-	public TActorIteratorBase< FActorFilter >
+class TActorIterator : public TActorIteratorBase<TActorIterator<ActorType>>
 {
+	friend class TActorIteratorBase<TActorIterator>;
+	typedef TActorIteratorBase<TActorIterator> Super;
+
 public:
 	/**
-	 * Constructor, inits the starting position for iteration
+	 * Constructor
+	 *
+	 * @param  InWorld  The world whose actors are to be iterated over.
+	 * @param  InClass  The subclass of actors to be iterated over.
 	 */
-	TActorIterator( UWorld* InWorld, TSubclassOf<ActorType> InClass = ActorType::StaticClass() )
-		: TActorIteratorBase< FActorFilter >( InWorld, InClass )
+	explicit TActorIterator( UWorld* InWorld, TSubclassOf<ActorType> InClass = ActorType::StaticClass() )
+		: Super( InWorld, InClass )
 	{
 		++(*this);
 	}
 
+	/**
+	 * Constructor for creating an end iterator
+	 */
+	explicit TActorIterator(EActorIteratorType)
+		: Super(EActorIteratorType::End)
+	{
+	}
+
 	/**
 	 * Returns the current suitable actor pointed at by the Iterator
 	 *
 	 * @return	Current suitable actor
 	 */
-	FORCEINLINE ActorType* operator*()
+	FORCEINLINE ActorType* operator*() const
 	{
-		check(FActorIteratorBase::CurrentActor);
-		checkf(!FActorIteratorBase::CurrentActor->HasAnyFlags(RF_Unreachable), TEXT("%s"), *FActorIteratorBase::CurrentActor->GetFullName());
-		return CastChecked<ActorType>(FActorIteratorBase::CurrentActor);
+		return CastChecked<ActorType>(**static_cast<const Super*>(this));
 	}
+
 	/**
 	 * Returns the current suitable actor pointed at by the Iterator
 	 *
 	 * @return	Current suitable actor
 	 */
-	FORCEINLINE ActorType* operator->()
+	FORCEINLINE ActorType* operator->() const
 	{
-		check(FActorIteratorBase::CurrentActor);
-		checkf(!FActorIteratorBase::CurrentActor->HasAnyFlags(RF_Unreachable), TEXT("%s"), *FActorIteratorBase::CurrentActor->GetFullName());
-		return CastChecked<ActorType>(FActorIteratorBase::CurrentActor);
+		return **this;
 	}
 };
 
 /**
- * Filters actors by whether or not they are selected
+ * Template actor range for ranged-for support.
  */
-class FSelectedActorFilter :
-	public FActorFilter
+template <typename ActorType>
+class TActorRange
 {
 public:
 	/**
-	 * Determines if the actor should be returned during iteration or not
+	 * Constructor
 	 *
-	 * @param	Actor	Actor to check
-	 * @return	true if actor is selected, false otherwise
+	 * @param  InWorld  The world whose actors are to be iterated over.
+	 * @param  InClass  The subclass of actors to be iterated over.
 	 */
-	FORCEINLINE bool IsSuitable( AActor* Actor ) const
+	explicit TActorRange(UWorld* InWorld, TSubclassOf<ActorType> InClass = ActorType::StaticClass())
+		: World(InWorld)
+		, Class(InClass)
 	{
-		return Actor ? Actor->IsSelected() : false;
 	}
+
+private:
+	UWorld*                World;
+	TSubclassOf<ActorType> Class;
+
+	friend TActorIterator<ActorType> begin(const TActorRange& Range) { return TActorIterator<ActorType>(Range.World, Range.Class); }
+	friend TActorIterator<ActorType> end  (const TActorRange& Range) { return TActorIterator<ActorType>(EActorIteratorType::End); }
 };
 
 /**
  * Selected actor iterator
  */
-class FSelectedActorIterator :
-	public TActorIteratorBase<FSelectedActorFilter>
+class FSelectedActorIterator : public TActorIteratorBase<FSelectedActorIterator>
 {
+	friend class TActorIteratorBase<FSelectedActorIterator>;
+	typedef TActorIteratorBase<FSelectedActorIterator> Super;
+
 public:
 	/**
-	 * Constructor, inits the starting position for iteration
+	 * Constructor
+	 *
+	 * @param  InWorld  The world whose actors are to be iterated over.
 	 */
-	FSelectedActorIterator( UWorld* InWorld )
-		:TActorIteratorBase( InWorld )
+	explicit FSelectedActorIterator( UWorld* InWorld )
+		:Super( InWorld, AActor::StaticClass() )
 	{
 		++(*this);
 	}
-};
 
-/**
- * Filters actors by whether or not they are network relevant
- */
-class FNetRelevantActorFilter :
-	public FActorFilter
-{
 	/**
-	 * The name of the net driver to filter actors by
+	 * Constructor for creating an end iterator
 	 */
-	FName NetDriverName;
-
-public:
-	/** Filter by the name specified */
-	FNetRelevantActorFilter(FName InNetDriverName) :
-		NetDriverName(InNetDriverName)
-	{
-	}
-	/** Default to the game network driver if not specified */
-	FNetRelevantActorFilter() :
-		NetDriverName(NAME_GameNetDriver)
+	explicit FSelectedActorIterator(EActorIteratorType)
+		: Super(EActorIteratorType::End)
 	{
 	}
 
+protected:
 	/**
-	 * Determines whether this actor is marked for the specified net driver or not
+	 * Determines if the actor should be returned during iteration or not
 	 *
-	 * @param	Actor	Actor to check the net driver name of
-	 * @return	true if actor is != NULL, false otherwise
+	 * @param	Actor	Actor to check
+	 * @return	true if actor is not null and is selected, false otherwise
 	 */
-	FORCEINLINE bool IsSuitable( AActor* Actor ) const
+	FORCEINLINE static bool IsActorSuitable(AActor* Actor)
 	{
-		return Actor != NULL && Actor->NetDriverName == NetDriverName;
+		return Actor->IsSelected();
 	}
 };
 
 /**
- * Net relevant actor iterator
+ * Selected actor range for ranged-for support.
  */
-class FNetRelevantActorIterator :
-	public TActorIteratorBase<FNetRelevantActorFilter, FTickableLevelFilter>
+class FSelectedActorRange
 {
-	/** Holds a filter class for the lifetime of this iterator */
-	FNetRelevantActorFilter FilterByNetDriverName;
-
-	/** Hidden on purpose */
-	FNetRelevantActorIterator();
-
 public:
 	/**
-	 * Positions the iterator at the first network relevant actor
+	 * Constructor
+	 *
+	 * @param  InWorld  The world whose actors are to be iterated over.
 	 */
-	FNetRelevantActorIterator( UWorld* InWorld, FName NetDriverName ) :
-		TActorIteratorBase( InWorld, AActor::StaticClass(), FilterByNetDriverName ),
-		FilterByNetDriverName( NetDriverName )
+	explicit FSelectedActorRange(UWorld* InWorld)
+		: World(InWorld)
 	{
-		++(*this);
 	}
+
+private:
+	UWorld* World;
+
+	friend FSelectedActorIterator begin(const FSelectedActorRange& Range) { return FSelectedActorIterator(Range.World); }
+	friend FSelectedActorIterator end  (const FSelectedActorRange& Range) { return FSelectedActorIterator(EActorIteratorType::End); }
 };
 
 /**
@@ -544,13 +572,15 @@ public:
  */
 class ENGINE_API FConsoleOutputDevice : public FStringOutputDevice
 {
+	typedef FStringOutputDevice Super;
+
 public:
 
 	/**
 	 * Minimal initialization constructor.
 	 */
 	FConsoleOutputDevice(class UConsole* InConsole):
-		FStringOutputDevice(TEXT("")),
+		Super(TEXT("")),
 		Console(InConsole)
 	{}
 
