AnimNotifyState_Trail
- Fixed an issue where the PSTemplate wasn't getting overriden because the if check was doing incorrect logic and the order of operations was wrong.
-Add dedicated server checks to AnimNoitifyState_Trail

diff --git a/Engine/Source/Runtime/Engine/Private/Animation/AnimNotifyState_Trail.cpp b/Engine/Source/Runtime/Engine/Private/Animation/AnimNotifyState_Trail.cpp
--- a/Engine/Source/Runtime/Engine/Private/Animation/AnimNotifyState_Trail.cpp
+++ b/Engine/Source/Runtime/Engine/Private/Animation/AnimNotifyState_Trail.cpp
@@ -38,6 +38,11 @@ void UAnimNotifyState_Trail::NotifyBegin(class USkeletalMeshComponent * MeshComp
 {
 	bool bError = ValidateInput(MeshComp);
 
+	if (MeshComp->GetWorld()->GetNetMode() == NM_DedicatedServer)
+	{
+		return;
+	}
+
 	TArray<USceneComponent*> Children;
 	MeshComp->GetChildrenComponents(false, Children);
 	
@@ -49,9 +54,9 @@ void UAnimNotifyState_Trail::NotifyBegin(class USkeletalMeshComponent * MeshComp
 	}
 
 	UParticleSystem* ParticleSystemTemplate = OverridePSTemplate(MeshComp, Animation);
-	if (ParticleSystemTemplate == nullptr)
+	if (ParticleSystemTemplate != nullptr)
 	{
-		ParticleSystemTemplate = PSTemplate;
+		PSTemplate = ParticleSystemTemplate;
 	}
 
 	bool bFoundExistingTrail = false;
@@ -129,6 +134,11 @@ void UAnimNotifyState_Trail::NotifyTick(class USkeletalMeshComponent * MeshComp,
 {
 	bool bError = ValidateInput(MeshComp, true);
 
+	if (MeshComp->GetWorld()->GetNetMode() == NM_DedicatedServer)
+	{
+		return;
+	}
+
 	TArray<USceneComponent*> Children;
 	MeshComp->GetChildrenComponents(false, Children);
 
@@ -169,6 +179,11 @@ void UAnimNotifyState_Trail::NotifyTick(class USkeletalMeshComponent * MeshComp,
 
 void UAnimNotifyState_Trail::NotifyEnd(class USkeletalMeshComponent * MeshComp, class UAnimSequenceBase * Animation)
 {
+	if (MeshComp->GetWorld()->GetNetMode() == NM_DedicatedServer)
+	{
+		return;
+	}
+
 	TArray<USceneComponent*> Children;
 	MeshComp->GetChildrenComponents(false, Children);
 
