- Fixing crash when the demo net driver is left idleing for too long.

Merging

//depot/UE4-Orion/Engine/Source/Runtime/Engine/Private/DemoNetDriver.cpp

to //depot/UE4/Engine/Source/Runtime/Engine/Private/DemoNetDriver.cpp

diff --git a/Engine/Source/Runtime/Engine/Private/DemoNetDriver.cpp b/Engine/Source/Runtime/Engine/Private/DemoNetDriver.cpp
--- a/Engine/Source/Runtime/Engine/Private/DemoNetDriver.cpp
+++ b/Engine/Source/Runtime/Engine/Private/DemoNetDriver.cpp
@@ -944,6 +944,18 @@ void UDemoNetDriver::TickDemoRecord( float DeltaSeconds )
 	{
 		AActor* Actor = World->NetworkActors[i];
 
+		if ( Actor->IsPendingKill() )
+		{
+			World->NetworkActors.RemoveAtSwap( i );
+			continue;
+		}
+
+		if ( Actor->GetRemoteRole() == ROLE_None )
+		{
+			World->NetworkActors.RemoveAtSwap( i );
+			continue;
+		}
+
 		Actor->PreReplication( *FindOrCreateRepChangedPropertyTracker( Actor ).Get() );
 		DemoReplicateActor( Actor, ClientConnections[0], IsNetClient );
 	}
