Actor calls PreReplication for all replicated subcomponents

Changes:
PreReplication added to UActorComponents
PreReplication called on all owned and replicated components by Actor in AActor::PreReplication

UE-4416

diff --git a/Engine/Source/Runtime/Engine/Classes/Components/ActorComponent.h b/Engine/Source/Runtime/Engine/Classes/Components/ActorComponent.h
--- a/Engine/Source/Runtime/Engine/Classes/Components/ActorComponent.h
+++ b/Engine/Source/Runtime/Engine/Classes/Components/ActorComponent.h
@@ -245,6 +245,9 @@ public:
 	/** Allows a component to replicate other subobject on the actor  */
 	virtual bool ReplicateSubobjects(class UActorChannel *Channel, class FOutBunch *Bunch, FReplicationFlags *RepFlags);
 
+	/** Called on the component right before replication occurs */
+	virtual void PreReplication(IRepChangedPropertyTracker & ChangedPropertyTracker);
+
 	virtual bool GetComponentClassCanReplicate() const;
 
 	/** Returns whether this component is an editor-only object or not */
diff --git a/Engine/Source/Runtime/Engine/Private/Actor.cpp b/Engine/Source/Runtime/Engine/Private/Actor.cpp
--- a/Engine/Source/Runtime/Engine/Private/Actor.cpp
+++ b/Engine/Source/Runtime/Engine/Private/Actor.cpp
@@ -10,6 +10,7 @@
 #include "LatentActions.h"
 #include "MessageLog.h"
 #include "Net/UnrealNetwork.h"
+#include "Net/RepLayout.h"
 #include "DisplayDebugHelpers.h"
 #include "Matinee/MatineeActor.h"
 #include "Matinee/InterpGroup.h"
@@ -731,6 +732,16 @@ void AActor::PreReplication( IRepChangedPropertyTracker & ChangedPropertyTracker
 	{
 		BPClass->InstancePreReplication(ChangedPropertyTracker);
 	}
+
+	// Call PreReplication on all owned components that are replicated
+	for (UActorComponent* Component : OwnedComponents)
+	{
+		// Only call on components that aren't pending kill
+		if (Component && !Component->IsPendingKill() && Component->GetIsReplicated())
+		{
+			Component->PreReplication(*GetNetDriver()->FindOrCreateRepChangedPropertyTracker(Component).Get());
+		}
+	}
 }
 
 void AActor::PostActorCreated()
diff --git a/Engine/Source/Runtime/Engine/Private/Components/ActorComponent.cpp b/Engine/Source/Runtime/Engine/Private/Components/ActorComponent.cpp
--- a/Engine/Source/Runtime/Engine/Private/Components/ActorComponent.cpp
+++ b/Engine/Source/Runtime/Engine/Private/Components/ActorComponent.cpp
@@ -1409,6 +1409,15 @@ bool UActorComponent::ReplicateSubobjects(class UActorChannel *Channel, class FO
 	return false;
 }
 
+void UActorComponent::PreReplication(IRepChangedPropertyTracker & ChangedPropertyTracker)
+{
+	UBlueprintGeneratedClass* BPClass = Cast<UBlueprintGeneratedClass>(GetClass());
+	if (BPClass != NULL)
+	{
+		BPClass->InstancePreReplication(ChangedPropertyTracker);
+	}
+}
+
 bool UActorComponent::GetComponentClassCanReplicate() const
 {
 	return true;
