Fix move component ignoring teleport flag. Also fix skeletal mesh flappy bits not being affected by world position changes.

diff --git a/Engine/Source/Runtime/Engine/Classes/Components/ActorComponent.h b/Engine/Source/Runtime/Engine/Classes/Components/ActorComponent.h
--- a/Engine/Source/Runtime/Engine/Classes/Components/ActorComponent.h
+++ b/Engine/Source/Runtime/Engine/Classes/Components/ActorComponent.h
@@ -474,7 +474,7 @@ public:
 	void DoDeferredRenderUpdates_Concurrent();
 
 	/** Recalculate the value of our component to world transform */
-	virtual void UpdateComponentToWorld(bool bSkipPhysicsMove = false) {}
+	virtual void UpdateComponentToWorld(bool bSkipPhysicsMove = false, bool bTeleport = false) {}
 
 	/** Mark the render state as dirty - will be sent to the render thread at the end of the frame. */
 	void MarkRenderStateDirty();
diff --git a/Engine/Source/Runtime/Engine/Classes/Components/AudioComponent.h b/Engine/Source/Runtime/Engine/Classes/Components/AudioComponent.h
--- a/Engine/Source/Runtime/Engine/Classes/Components/AudioComponent.h
+++ b/Engine/Source/Runtime/Engine/Classes/Components/AudioComponent.h
@@ -270,7 +270,7 @@ public:
 	// Begin USceneComponent Interface
 	virtual void Activate(bool bReset=false) override;
 	virtual void Deactivate() override;
-	virtual void OnUpdateTransform(bool bSkipPhysicsMove) override;
+	virtual void OnUpdateTransform(bool bSkipPhysicsMove, bool bTeleport = false) override;
 	// do nothing! on purpose
 	virtual void UpdateNavigationData() override {}
 	// End USceneComponent Interface
diff --git a/Engine/Source/Runtime/Engine/Classes/Components/DestructibleComponent.h b/Engine/Source/Runtime/Engine/Classes/Components/DestructibleComponent.h
--- a/Engine/Source/Runtime/Engine/Classes/Components/DestructibleComponent.h
+++ b/Engine/Source/Runtime/Engine/Classes/Components/DestructibleComponent.h
@@ -107,7 +107,7 @@ public:
 
 	// Begin USceneComponent interface.
 	virtual FBoxSphereBounds CalcBounds(const FTransform& LocalToWorld) const override;
-	virtual void OnUpdateTransform(bool bSkipPhysicsMove) override;
+	virtual void OnUpdateTransform(bool bSkipPhysicsMove, bool bTeleport = false) override;
 	virtual void Activate(bool bReset=false) override;
 	virtual void Deactivate() override;
 	// End USceneComponent interface.
diff --git a/Engine/Source/Runtime/Engine/Classes/Components/PrimitiveComponent.h b/Engine/Source/Runtime/Engine/Classes/Components/PrimitiveComponent.h
--- a/Engine/Source/Runtime/Engine/Classes/Components/PrimitiveComponent.h
+++ b/Engine/Source/Runtime/Engine/Classes/Components/PrimitiveComponent.h
@@ -1262,7 +1262,7 @@ protected:
 	friend class FStaticMeshComponentRecreateRenderStateContext;
 
 	// Begin USceneComponent Interface
-	virtual void OnUpdateTransform(bool bSkipPhysicsMove) override;
+	virtual void OnUpdateTransform(bool bSkipPhysicsMove, bool bTeleport = false) override;
 
 	/** Event called when AttachParent changes, to allow the scene to update its attachment state. */
 	virtual void OnAttachmentChanged() override;
diff --git a/Engine/Source/Runtime/Engine/Classes/Components/SceneComponent.h b/Engine/Source/Runtime/Engine/Classes/Components/SceneComponent.h
--- a/Engine/Source/Runtime/Engine/Classes/Components/SceneComponent.h
+++ b/Engine/Source/Runtime/Engine/Classes/Components/SceneComponent.h
@@ -517,7 +517,7 @@ public:
 	{
 		return true;
 	}
-	virtual void UpdateComponentToWorld(bool bSkipPhysicsMove = false) override final;
+	virtual void UpdateComponentToWorld(bool bSkipPhysicsMove = false, bool bTeleport = false) override final;
 	virtual void DestroyComponent(bool bPromoteChildren = false) override;
 	virtual void OnComponentDestroyed() override;
 	virtual void ApplyWorldOffset(const FVector& InOffset, bool bWorldShift) override;
@@ -545,17 +545,17 @@ protected:
 	 * Internal helper, for use from MoveComponent().  Special codepath since the normal setters call MoveComponent.
 	 * @return: true if location or rotation was changed.
 	 */
-	bool InternalSetWorldLocationAndRotation(FVector NewLocation, const FQuat& NewQuat, bool bNoPhysics=false);
+	bool InternalSetWorldLocationAndRotation(FVector NewLocation, const FQuat& NewQuat, bool bNoPhysics=false, bool bTeleport = false);
 
-	virtual void OnUpdateTransform(bool bSkipPhysicsMove);
+	virtual void OnUpdateTransform(bool bSkipPhysicsMove, bool bTeleport = false);
 
 	/** Check if mobility is set to non-static. If it's static we trigger a PIE warning and return true*/
 	bool CheckStaticMobilityAndWarn(const FText& ActionText) const;
 
 private:
 
-	void PropagateTransformUpdate(bool bTransformChanged, bool bSkipPhysicsMove = false);
-	void UpdateComponentToWorldWithParent(USceneComponent* Parent, bool bSkipPhysicsMove, const FQuat& RelativeRotationQuat);
+	void PropagateTransformUpdate(bool bTransformChanged, bool bSkipPhysicsMove = false, bool bTeleport = false);
+	void UpdateComponentToWorldWithParent(USceneComponent* Parent, bool bSkipPhysicsMove, const FQuat& RelativeRotationQuat, bool bTeleport = false);
 
 
 public:
@@ -670,7 +670,7 @@ public:
 	}
 
 	/** Update transforms of any components attached to this one. */
-	void UpdateChildTransforms();
+	void UpdateChildTransforms(bool bSkipPhysicsMove = false, bool bTeleport = false);
 
 	/** Calculate the bounds of this component. Default behavior is a bounding box/sphere of zero size. */
 	virtual FBoxSphereBounds CalcBounds(const FTransform& LocalToWorld) const;
diff --git a/Engine/Source/Runtime/Engine/Classes/Components/SkeletalMeshComponent.h b/Engine/Source/Runtime/Engine/Classes/Components/SkeletalMeshComponent.h
--- a/Engine/Source/Runtime/Engine/Classes/Components/SkeletalMeshComponent.h
+++ b/Engine/Source/Runtime/Engine/Classes/Components/SkeletalMeshComponent.h
@@ -837,7 +837,7 @@ public:
 	virtual void UpdateBounds() override;
 	virtual FBoxSphereBounds CalcBounds(const FTransform& LocalToWorld) const override;
 	virtual bool IsAnySimulatingPhysics() const override;
-	virtual void OnUpdateTransform(bool bSkipPhysicsMove) override;
+	virtual void OnUpdateTransform(bool bSkipPhysicsMove, bool bTeleport = false) override;
 	// End USceneComponent interface.
 
 	// Begin UPrimitiveComponent interface.
@@ -1062,7 +1062,13 @@ public:
 	 *	Iterate over each physics body in the physics for this mesh, and for each 'kinematic' (ie fixed or default if owner isn't simulating) one, update its
 	 *	transform based on the animated transform.
 	 */
-	void UpdateKinematicBonesToAnim(const TArray<FTransform>& InSpaceBases, bool bTeleport, bool bNeedsSkinning, bool bForceUpdate = false);
+	void UpdateKinematicBonesToAnim(const TArray<FTransform>& InSpaceBases, bool bTeleport, bool bNeedsSkinning);
+
+	DEPRECATED(4.9, "bForceUpdate is no longer used. Please use bTeleport")
+	void UpdateKinematicBonesToAnim(const TArray<FTransform>& InSpaceBases, bool bTeleport, bool bNeedsSkinning, bool bForceUpdate)
+	{
+		UpdateKinematicBonesToAnim(InSpaceBases, bTeleport || bForceUpdate, bNeedsSkinning);
+	}
 	
 	/**
 	 * Look up all bodies for broken constraints.
diff --git a/Engine/Source/Runtime/Engine/Private/Components/AudioComponent.cpp b/Engine/Source/Runtime/Engine/Private/Components/AudioComponent.cpp
--- a/Engine/Source/Runtime/Engine/Private/Components/AudioComponent.cpp
+++ b/Engine/Source/Runtime/Engine/Private/Components/AudioComponent.cpp
@@ -92,9 +92,9 @@ void UAudioComponent::SetSound( USoundBase* NewSound )
 	}
 }
 
-void UAudioComponent::OnUpdateTransform(bool bSkipPhysicsMove)
+void UAudioComponent::OnUpdateTransform(bool bSkipPhysicsMove, bool bTeleport)
 {
-	Super::OnUpdateTransform(bSkipPhysicsMove);
+	Super::OnUpdateTransform(bSkipPhysicsMove, bTeleport);
 
 	if (bIsActive && !bPreviewComponent)
 	{
diff --git a/Engine/Source/Runtime/Engine/Private/Components/DestructibleComponent.cpp b/Engine/Source/Runtime/Engine/Private/Components/DestructibleComponent.cpp
--- a/Engine/Source/Runtime/Engine/Private/Components/DestructibleComponent.cpp
+++ b/Engine/Source/Runtime/Engine/Private/Components/DestructibleComponent.cpp
@@ -127,10 +127,10 @@ bool IsImpactDamageEnabled(const UDestructibleMesh* TheDestructibleMesh, int32 L
 	}
 }
 
-void UDestructibleComponent::OnUpdateTransform(bool bSkipPhysicsMove)
+void UDestructibleComponent::OnUpdateTransform(bool bSkipPhysicsMove, bool bTeleport)
 {
 	// We are handling the physics move below, so don't handle it at higher levels
-	Super::OnUpdateTransform(true);
+	Super::OnUpdateTransform(true, bTeleport);
 
 	if (SkeletalMesh == NULL)
 	{
@@ -165,7 +165,7 @@ void UDestructibleComponent::OnUpdateTransform(bool bSkipPhysicsMove)
 		PxMat44 GlobalPose(PxMat33(U2PQuat(CurrentLocalToWorld.GetRotation())), U2PVector(CurrentLocalToWorld.GetTranslation()));
 		if(!PRootActor || PRootActor->getScene())	//either root chunk is null meaning fractured (so there's a scene), or the root has a scene
 		{
-			ApexDestructibleActor->setGlobalPose(GlobalPose);
+			ApexDestructibleActor->setGlobalPose(GlobalPose);	//this assumes ignores bTeleport and treats it like it's true 
 		}else
 		{
 			PRootActor->setGlobalPose(PxTransform(GlobalPose));	//we're not in a scene yet, so place the root actor in this new position
diff --git a/Engine/Source/Runtime/Engine/Private/Components/PrimitiveComponent.cpp b/Engine/Source/Runtime/Engine/Private/Components/PrimitiveComponent.cpp
--- a/Engine/Source/Runtime/Engine/Private/Components/PrimitiveComponent.cpp
+++ b/Engine/Source/Runtime/Engine/Private/Components/PrimitiveComponent.cpp
@@ -465,15 +465,14 @@ void UPrimitiveComponent::EnsurePhysicsStateCreated()
 	}
 }
 
-void UPrimitiveComponent::OnUpdateTransform(bool bSkipPhysicsMove)
+void UPrimitiveComponent::OnUpdateTransform(bool bSkipPhysicsMove, bool bTeleport)
 {
 	Super::OnUpdateTransform(bSkipPhysicsMove);
 
 	// Always send new transform to physics
 	if(bPhysicsStateCreated && !bSkipPhysicsMove)
 	{
-		// @todo UE4 rather than always pass false, this function should know if it is being teleported or not!
-		SendPhysicsTransform(false);
+		SendPhysicsTransform(bTeleport);
 	}
 }
 
@@ -1385,7 +1384,7 @@ bool UPrimitiveComponent::MoveComponentImpl( const FVector& Delta, const FQuat&
 	if ( !bSweep )
 	{
 		// not sweeping, just go directly to the new transform
-		bMoved = InternalSetWorldLocationAndRotation(TraceEnd, NewRotationQuat, bSkipPhysicsMove);
+		bMoved = InternalSetWorldLocationAndRotation(TraceEnd, NewRotationQuat, bSkipPhysicsMove, !bSweep);
 	}
 	else
 	{
diff --git a/Engine/Source/Runtime/Engine/Private/Components/SceneComponent.cpp b/Engine/Source/Runtime/Engine/Private/Components/SceneComponent.cpp
--- a/Engine/Source/Runtime/Engine/Private/Components/SceneComponent.cpp
+++ b/Engine/Source/Runtime/Engine/Private/Components/SceneComponent.cpp
@@ -315,11 +315,11 @@ FTransform USceneComponent::CalcNewComponentToWorld(const FTransform& NewRelativ
 	}
 }
 
-void USceneComponent::OnUpdateTransform(bool bSkipPhysicsMove)
+void USceneComponent::OnUpdateTransform(bool bSkipPhysicsMove, bool bTeleport)
 {
 }
 
-void USceneComponent::UpdateComponentToWorldWithParent(USceneComponent* Parent, bool bSkipPhysicsMove, const FQuat& RelativeRotationQuat)
+void USceneComponent::UpdateComponentToWorldWithParent(USceneComponent* Parent, bool bSkipPhysicsMove, const FQuat& RelativeRotationQuat, bool bTeleport)
 {
 	// If our parent hasn't been updated before, we'll need walk up our parent attach hierarchy
 	if (Parent && !Parent->bWorldToComponentUpdated)
@@ -352,7 +352,7 @@ void USceneComponent::UpdateComponentToWorldWithParent(USceneComponent* Parent,
 		// Update transform
 		ComponentToWorld = NewTransform;
 
-		PropagateTransformUpdate(true, bSkipPhysicsMove);
+		PropagateTransformUpdate(true, bSkipPhysicsMove, bTeleport);
 	}
 	else
 	{
@@ -392,13 +392,13 @@ void USceneComponent::OnRegister()
 #endif
 }
 
-void USceneComponent::UpdateComponentToWorld(bool bSkipPhysicsMove)
+void USceneComponent::UpdateComponentToWorld(bool bSkipPhysicsMove, bool bTeleport)
 {
-	UpdateComponentToWorldWithParent(AttachParent, bSkipPhysicsMove, RelativeRotation.Quaternion());
+	UpdateComponentToWorldWithParent(AttachParent, bSkipPhysicsMove, RelativeRotation.Quaternion(), bTeleport);
 }
 
 
-void USceneComponent::PropagateTransformUpdate(bool bTransformChanged, bool bSkipPhysicsMove)
+void USceneComponent::PropagateTransformUpdate(bool bTransformChanged, bool bSkipPhysicsMove, bool bTeleport)
 {
 	if (IsDeferringMovementUpdates())
 	{
@@ -412,13 +412,13 @@ void USceneComponent::PropagateTransformUpdate(bool bTransformChanged, bool bSki
 		UpdateBounds();
 
 		// Always send new transform to physics
-		OnUpdateTransform(bSkipPhysicsMove);
+		OnUpdateTransform(bSkipPhysicsMove, bTeleport);
 
 		// Flag render transform as dirty
 		MarkRenderTransformDirty();
 		
 		// Now go and update children
-		UpdateChildTransforms();
+		UpdateChildTransforms(bSkipPhysicsMove, bTeleport);
 
 		// Refresh navigation
 		UpdateNavigationData();
@@ -1487,14 +1487,14 @@ FName USceneComponent::GetComponentInstanceDataType() const
 	return SceneComponentInstanceDataTypeName;
 }
 
-void USceneComponent::UpdateChildTransforms()
+void USceneComponent::UpdateChildTransforms(bool bSkipPhysicsMove, bool bTeleport)
 {
 	for(int32 i=0; i<AttachChildren.Num(); i++)
 	{
 		USceneComponent* ChildComp = AttachChildren[i];
 		if(ChildComp != NULL)
 		{
-			ChildComp->UpdateComponentToWorld();
+			ChildComp->UpdateComponentToWorld(bSkipPhysicsMove, bTeleport);
 		}
 	}
 }
@@ -1826,7 +1826,7 @@ void USceneComponent::BeginDestroy()
 	Super::BeginDestroy();
 }
 
-bool USceneComponent::InternalSetWorldLocationAndRotation(FVector NewLocation, const FQuat& RotationQuat, bool bNoPhysics)
+bool USceneComponent::InternalSetWorldLocationAndRotation(FVector NewLocation, const FQuat& RotationQuat, bool bNoPhysics, bool bTeleport)
 {
 	FQuat NewRotationQuat(RotationQuat);
 
@@ -1852,7 +1852,7 @@ bool USceneComponent::InternalSetWorldLocationAndRotation(FVector NewLocation, c
 	{
 		RelativeLocation = NewLocation;
 		RelativeRotation = NewRotator;
-		UpdateComponentToWorldWithParent(AttachParent, bNoPhysics, NewRotationQuat);
+		UpdateComponentToWorldWithParent(AttachParent, bNoPhysics, NewRotationQuat, bTeleport);
 		return true;
 	}
 
@@ -1943,7 +1943,7 @@ bool USceneComponent::MoveComponentImpl( const FVector& Delta, const FQuat& NewR
 	}
 
 	// just teleport, sweep is supported for PrimitiveComponents. This will update child components as well.
-	const bool bMoved = InternalSetWorldLocationAndRotation(GetComponentLocation() + Delta, NewRotation);
+	const bool bMoved = InternalSetWorldLocationAndRotation(GetComponentLocation() + Delta, NewRotation, false, !bSweep);
 
 	// Only update overlaps if not deferring updates within a scope
 	if (bMoved && !IsDeferringMovementUpdates())
diff --git a/Engine/Source/Runtime/Engine/Private/PhysicsEngine/PhysAnim.cpp b/Engine/Source/Runtime/Engine/Private/PhysicsEngine/PhysAnim.cpp
--- a/Engine/Source/Runtime/Engine/Private/PhysicsEngine/PhysAnim.cpp
+++ b/Engine/Source/Runtime/Engine/Private/PhysicsEngine/PhysAnim.cpp
@@ -245,8 +245,7 @@ void USkeletalMeshComponent::BlendInPhysics()
 }
 
 
-
-void USkeletalMeshComponent::UpdateKinematicBonesToAnim(const TArray<FTransform>& InSpaceBases, bool bTeleport, bool bNeedsSkinning, bool bForceUpdate)
+void USkeletalMeshComponent::UpdateKinematicBonesToAnim(const TArray<FTransform>& InSpaceBases, bool bTeleport, bool bNeedsSkinning)
 {
 	SCOPE_CYCLE_COUNTER(STAT_UpdateRBBones);
 
@@ -258,7 +257,7 @@ void USkeletalMeshComponent::UpdateKinematicBonesToAnim(const TArray<FTransform>
 	const bool bUpdateKinematics = (KinematicBonesUpdateType != EKinematicBonesUpdateToPhysics::SkipAllBones);
 
 	// If desired, update physics bodies associated with skeletal mesh component to match.
-	if( !bUpdateKinematics && !(bForceUpdate && IsAnySimulatingPhysics()))
+	if(!bUpdateKinematics && !(bTeleport && IsAnySimulatingPhysics()))
 	{
 		// nothing to do 
 		return;
@@ -343,7 +342,7 @@ void USkeletalMeshComponent::UpdateKinematicBonesToAnim(const TArray<FTransform>
 				FBodyInstance* BodyInst = Bodies[i];
 				check(BodyInst);
 
-				if (bForceUpdate || (BodyInst->IsValidBodyInstance() && !BodyInst->IsInstanceSimulatingPhysics()))
+				if (bTeleport || (BodyInst->IsValidBodyInstance() && !BodyInst->IsInstanceSimulatingPhysics()))
 				{
 					const int32 BoneIndex = BodyInst->InstanceBoneIndex;
 
diff --git a/Engine/Source/Runtime/Engine/Private/SkeletalMeshComponentPhysics.cpp b/Engine/Source/Runtime/Engine/Private/SkeletalMeshComponentPhysics.cpp
--- a/Engine/Source/Runtime/Engine/Private/SkeletalMeshComponentPhysics.cpp
+++ b/Engine/Source/Runtime/Engine/Private/SkeletalMeshComponentPhysics.cpp
@@ -1327,20 +1327,29 @@ FConstraintInstance* USkeletalMeshComponent::FindConstraintInstance(FName ConNam
 	return NULL;
 }
 
-void USkeletalMeshComponent::OnUpdateTransform(bool bSkipPhysicsMove)
+#ifndef OLD_FORCE_UPDATE_BEHAVIOR
+#define OLD_FORCE_UPDATE_BEHAVIOR 0
+#endif
+
+void USkeletalMeshComponent::OnUpdateTransform(bool bSkipPhysicsMove, bool bTeleport)
 {
 	// We are handling the physics move below, so don't handle it at higher levels
-	Super::OnUpdateTransform(true);
+	Super::OnUpdateTransform(true, bTeleport);
 
 	// Always send new transform to physics
 	if(bPhysicsStateCreated && !bSkipPhysicsMove)
 	{
-		UpdateKinematicBonesToAnim(GetSpaceBases(), false, false, true);
+#if !OLD_FORCE_UPDATE_BEHAVIOR
+		UpdateKinematicBonesToAnim(GetSpaceBases(), bTeleport, false);
+#else
+		UpdateKinematicBonesToAnim(GetSpaceBases(), true, false);
+#endif
 	}
 
 #if WITH_APEX_CLOTHING
 	if(ClothingActors.Num() > 0)
 	{
+		//@todo: Should cloth know whether we're teleporting?
 		// Updates cloth animation states because transform is updated
 		UpdateClothTransform();
 	}
