Added validation and range checking to NMT_ActorChannelFailure and tagged a case with PotentialDOSAttackDetection

UE-2190

diff --git a/Engine/Source/Runtime/Engine/Private/DataChannel.cpp b/Engine/Source/Runtime/Engine/Private/DataChannel.cpp
--- a/Engine/Source/Runtime/Engine/Private/DataChannel.cpp
+++ b/Engine/Source/Runtime/Engine/Private/DataChannel.cpp
@@ -14,8 +14,10 @@
 #include "Engine/PackageMapClient.h"
 
 DEFINE_LOG_CATEGORY(LogNet);
+DEFINE_LOG_CATEGORY(LogRep);
 DEFINE_LOG_CATEGORY(LogNetPlayerMovement);
 DEFINE_LOG_CATEGORY(LogNetTraffic);
+DEFINE_LOG_CATEGORY(LogRepTraffic);
 DEFINE_LOG_CATEGORY(LogNetDormancy);
 DEFINE_LOG_CATEGORY_STATIC(LogNetPartialBunch, Warning, All);
 
@@ -1131,12 +1133,17 @@ void UControlChannel::ReceivedBunch( FInBunch& Bunch )
 				UE_LOG(LogNet, Log, TEXT("Server connection received: %s"), FNetControlMessageInfo::GetName(MessageType));
 				int32 ChannelIndex;
 				FNetControlMessage<NMT_ActorChannelFailure>::Receive(Bunch, ChannelIndex);
-				if (ChannelIndex < ARRAY_COUNT(Connection->Channels))
+
+				// Check if Channel index provided by client is valid and within range of channel on server
+				if (ChannelIndex >= 0 && ChannelIndex < ARRAY_COUNT(Connection->Channels))
 				{
+					// Get the actor channel that the client provided as having failed
 					UActorChannel* ActorChan = Cast<UActorChannel>(Connection->Channels[ChannelIndex]);
-					if (ActorChan != NULL && ActorChan->Actor != NULL)
+
+					// The channel and the actor attached to the channel exists on the server
+					if (ActorChan != nullptr && ActorChan->Actor != nullptr)
 					{
-						// if the client failed to initialize the PlayerController channel, the connection is broken
+						// The channel that failed is the player controller thus the connection is broken
 						if (ActorChan->Actor == Connection->PlayerController)
 						{
 							UE_LOG(LogNet, Warning, TEXT("UControlChannel::ReceivedBunch: NetConnection::Close() [%s] [%s] [%s] from failed to initialize the PlayerController channel. Closing connection."), 
@@ -1146,12 +1153,27 @@ void UControlChannel::ReceivedBunch( FInBunch& Bunch )
 
 							Connection->Close();
 						}
-						else if (Connection->PlayerController != NULL)
+						// The client has a PlayerController connection, report the actor failure to PlayerController
+						else if (Connection->PlayerController != nullptr)
 						{
 							Connection->PlayerController->NotifyActorChannelFailure(ActorChan);
 						}
+						// The PlayerController connection doesn't exist for the client
+						// but the client is reporting an actor channel failure that isn't the PlayerController
+						else
+						{
+							//UE_LOG(LogNet, Warning, TEXT("UControlChannel::RecievedBunch: PlayerController doesn't exist for the client, but the client is reporting an actor channel failure that isn't the PlayerController."));
+						}
 					}
 				}
+				// The client is sending an actor channel failure message with an invalid
+				// actor channel index
+				// @PotentialDOSAttackDetection
+				else
+				{
+					UE_LOG(LogNet, Warning, TEXT("UControlChannel::RecievedBunch: The client is sending an actor channel failure message with an invalid actor channel index."));
+				}
+
 			}
 		}
 		else
@@ -1222,7 +1244,10 @@ void UControlChannel::ReceivedBunch( FInBunch& Bunch )
 					FNetControlMessage<NMT_BeaconNetGUIDAck>::Discard(Bunch);
 					break;
 				default:
-					check(!FNetControlMessageInfo::IsRegistered(MessageType)); // if this fails, a case is missing above for an implemented message type
+					// if this fails, a case is missing above for an implemented message type
+					// or the connection is being sent potentially malformed packets
+					// @PotentialDOSAttackDetection
+					check(!FNetControlMessageInfo::IsRegistered(MessageType));
 
 					UE_LOG(LogNet, Error, TEXT("Received unknown control channel message"));
 					ensureMsgf(false, TEXT("Failed to read control channel message %i"), int32(MessageType));
