InternalProtocolVersion Update

diff --git a/Engine/Source/Runtime/Engine/Private/UnrealNetwork.cpp b/Engine/Source/Runtime/Engine/Private/UnrealNetwork.cpp
--- a/Engine/Source/Runtime/Engine/Private/UnrealNetwork.cpp
+++ b/Engine/Source/Runtime/Engine/Private/UnrealNetwork.cpp
@@ -10,16 +10,9 @@ FNetworkVersion::FIsNetworkCompatibleOverride FNetworkVersion::IsNetworkCompatib
 
 enum ENetworkVersionHistory
 {
-	HISTORY_INITIAL					= 1,
-	HISTORY_INTERNAL_ACK			= 3,				// We no longer save packet/channel sequence in stream. We can derive this for 100% reliable connections.
-	HISTORY_REPLAY_CHECKSUMS		= 4,				// We now save replay checksums into stream for backwards compatibility checks
-	HISTORY_REPLAY_CHECKSUMS2		= 5,				// Saving package and network checksum to packagemap, no longer save checksum in BeginContentBlockHeader
-	HISTORY_REPLAY_CHECKSUMS3		= 6,				// Various network checksum fixes
-	HISTORY_REPLAY_CHECKSUMS4		= 7,				// Added multi-cast RPC's to network checksum
-	HISTORY_REPLAY_RPC_CHECKSUMS	= 8,				// RPC's save checksums separately now
 };
 
-const uint32 FNetworkVersion::InternalProtocolVersion = HISTORY_REPLAY_RPC_CHECKSUMS;
+const uint32 FNetworkVersion::InternalProtocolVersion = HISTORY_SERIALIZE_DEFAULTS_CHECK;
 
 uint32 FNetworkVersion::GetLocalNetworkVersion(bool AllowOverrideDelegate /*=true*/)
 {
