Added LogRep and LogRepTraffic for replication specific UE_LOG messages

UE-2098

diff --git a/Engine/Plugins/NetcodeUnitTest/NetcodeUnitTest/Source/NetcodeUnitTest/Public/NUTUtilDebug.h b/Engine/Plugins/NetcodeUnitTest/NetcodeUnitTest/Source/NetcodeUnitTest/Public/NUTUtilDebug.h
--- a/Engine/Plugins/NetcodeUnitTest/NetcodeUnitTest/Source/NetcodeUnitTest/Public/NUTUtilDebug.h
+++ b/Engine/Plugins/NetcodeUnitTest/NetcodeUnitTest/Source/NetcodeUnitTest/Public/NUTUtilDebug.h
@@ -74,7 +74,9 @@ public:
 
 		// @todo JohnB: See if there are any other good net categories to add here
 		TempLogCategories.Add(TEXT("LogNet"));
+		TempLogCategories.Add(TEXT("LogRep"));
 		TempLogCategories.Add(TEXT("LogNetTraffic"));
+		TempLogCategories.Add(TEXT("LogRepTraffic"));
 		TempLogCategories.Add(TEXT("LogNetSerialization"));
 		TempLogCategories.Add(TEXT("LogNetPackageMap"));
 		TempLogCategories.Add(TEXT("LogNetPlayerMovement"));
diff --git a/Engine/Source/Runtime/Engine/Private/DataReplication.cpp b/Engine/Source/Runtime/Engine/Private/DataReplication.cpp
--- a/Engine/Source/Runtime/Engine/Private/DataReplication.cpp
+++ b/Engine/Source/Runtime/Engine/Private/DataReplication.cpp
@@ -36,7 +36,7 @@ public:
 
 			if (!bSuccess)
 			{
-				UE_LOG(LogNet, Warning, TEXT("NetSerializeStruct: Native NetSerialize %s failed."), *Struct->GetFullName());
+				UE_LOG(LogRep, Warning, TEXT("NetSerializeStruct: Native NetSerialize %s failed."), *Struct->GetFullName());
 			}
 		}
 		else
@@ -145,19 +145,19 @@ bool FObjectReplicator::ValidateAgainstState( const UObject* ObjectState )
 {
 	if ( !RepLayout.IsValid() )
 	{
-		UE_LOG( LogNet, Warning, TEXT( "ValidateAgainstState: RepLayout.IsValid() == false" ) );
+		UE_LOG(LogRep, Warning, TEXT("ValidateAgainstState: RepLayout.IsValid() == false"));
 		return false;
 	}
 
 	if ( RepState == NULL )
 	{
-		UE_LOG( LogNet, Warning, TEXT( "ValidateAgainstState: RepState == NULL" ) );
+		UE_LOG(LogRep, Warning, TEXT("ValidateAgainstState: RepState == NULL"));
 		return false;
 	}
 
 	if ( RepLayout->DiffProperties( RepState, ObjectState, false ) )
 	{
-		UE_LOG( LogNet, Warning, TEXT( "ValidateAgainstState: Properties changed for %s" ), *ObjectState->GetName() );
+		UE_LOG(LogRep, Warning, TEXT("ValidateAgainstState: Properties changed for %s"), *ObjectState->GetName());
 		return false;
 	}
 
@@ -181,7 +181,7 @@ void FObjectReplicator::InitWithObject( UObject* InObject, UNetConnection * InCo
 	{
 		// This may seem weird that we're checking for NULL, but the SetObject above will wrap this object with TWeakObjectPtr
 		// If the object is pending kill, it will switch to NULL, we're just making sure we handle this invalid edge case
-		UE_LOG( LogNet, Error, TEXT( "InitWithObject: Object == NULL" ) );
+		UE_LOG(LogRep, Error, TEXT("InitWithObject: Object == NULL"));
 		return;
 	}
 
@@ -237,7 +237,7 @@ void FObjectReplicator::StartReplicating( class UActorChannel * InActorChannel )
 
 	if ( GetObject() == NULL )
 	{
-		UE_LOG( LogNet, Error, TEXT( "StartReplicating: Object == NULL" ) );
+		UE_LOG(LogRep, Error, TEXT("StartReplicating: Object == NULL"));
 		return;
 	}
 
@@ -331,7 +331,7 @@ void FObjectReplicator::ReceivedNak( int32 NakPacketId )
 
 	if ( Object == NULL )
 	{
-		UE_LOG( LogNet, Verbose, TEXT( "ReceivedNak: Object == NULL" ) );
+		UE_LOG(LogNet, Verbose, TEXT("ReceivedNak: Object == NULL"));
 		return;
 	}
 
@@ -358,7 +358,7 @@ void FObjectReplicator::ReceivedNak( int32 NakPacketId )
 				}
 				else if ( NakPacketId >= Rec->OutPacketIdRange.First && NakPacketId <= Rec->OutPacketIdRange.Last )
 				{
-					UE_LOG( LogNet, Verbose, TEXT( "Restoring Previous Base State of dynamic property Channel %d. NakId: %d  (%d -%d)" ), OwningChannel->ChIndex, NakPacketId, Rec->OutPacketIdRange.First, Rec->OutPacketIdRange.Last  );
+					UE_LOG(LogNet, Verbose, TEXT("Restoring Previous Base State of dynamic property Channel %d. NakId: %d  (%d -%d)"), OwningChannel->ChIndex, NakPacketId, Rec->OutPacketIdRange.First, Rec->OutPacketIdRange.Last);
 
 					// The Nack'd packet did update this property, so we need to replace the buffer in RecentDynamic
 					// with the buffer we used to create this update (which was dropped), so that the update will be recreated on the next replicate actor
@@ -400,7 +400,7 @@ const FFieldNetCache* FObjectReplicator::ReadField( const FClassNetCache* ClassC
 
 		if ( Bunch.IsError() )
 		{
-			UE_LOG( LogNet, Error, TEXT( "ReadField: Error reading checksum: %s" ), *GetObject()->GetFullName() );
+			UE_LOG(LogNet, Error, TEXT("ReadField: Error reading checksum: %s"), *GetObject()->GetFullName());
 			return NULL;
 		}
 
@@ -413,7 +413,7 @@ const FFieldNetCache* FObjectReplicator::ReadField( const FClassNetCache* ClassC
 
 		if ( FieldNetCache == NULL )
 		{
-			UE_LOG( LogNet, Error, TEXT( "ReadField: GetFromChecksum failed: %s" ), *GetObject()->GetFullName() );
+			UE_LOG(LogNet, Error, TEXT("ReadField: GetFromChecksum failed: %s"), *GetObject()->GetFullName());
 			Bunch.SetError();
 			return NULL;
 		}
@@ -425,7 +425,7 @@ const FFieldNetCache* FObjectReplicator::ReadField( const FClassNetCache* ClassC
 
 	if ( Bunch.IsError() )
 	{
-		UE_LOG( LogNet, Error, TEXT( "ReadField: Error reading RepIndex: %s" ), *GetObject()->GetFullName() );
+		UE_LOG(LogRep, Error, TEXT("ReadField: Error reading RepIndex: %s"), *GetObject()->GetFullName());
 		return NULL;
 	}
 
@@ -437,7 +437,7 @@ const FFieldNetCache* FObjectReplicator::ReadField( const FClassNetCache* ClassC
 	if ( RepIndex > ClassCache->GetMaxIndex() )
 	{
 		// We shouldn't be receiving this bunch of this object has no properties or RPC functions to process
-		UE_LOG( LogNet, Error, TEXT( "ReadField: RepIndex too large: %s" ), *GetObject()->GetFullName() );
+		UE_LOG(LogRep, Error, TEXT("ReadField: RepIndex too large: %s"), *GetObject()->GetFullName());
 		Bunch.SetError();
 		return NULL;
 	}
@@ -446,7 +446,7 @@ const FFieldNetCache* FObjectReplicator::ReadField( const FClassNetCache* ClassC
 
 	if ( FieldNetCache == NULL )
 	{
-		UE_LOG( LogNet, Error, TEXT( "ReadField: GetFromIndex failed: %s" ), *GetObject()->GetFullName() );
+		UE_LOG(LogNet, Error, TEXT("ReadField: GetFromIndex failed: %s"), *GetObject()->GetFullName());
 		Bunch.SetError();
 		return NULL;
 	}
@@ -460,7 +460,7 @@ bool FObjectReplicator::ReceivedBunch( FInBunch& Bunch, const FReplicationFlags&
 
 	if ( Object == NULL )
 	{
-		UE_LOG( LogNet, Verbose, TEXT( "ReceivedBunch: Object == NULL" ) );
+		UE_LOG(LogNet, Verbose, TEXT("ReceivedBunch: Object == NULL"));
 		return false;
 	}
 
@@ -472,7 +472,7 @@ bool FObjectReplicator::ReceivedBunch( FInBunch& Bunch, const FReplicationFlags&
 
 	if ( ClassCache == NULL )
 	{
-		UE_LOG( LogNet, Error, TEXT( "ReceivedBunch: ClassCache == NULL: %s" ), *Object->GetFullName() );
+		UE_LOG(LogNet, Error, TEXT("ReceivedBunch: ClassCache == NULL: %s"), *Object->GetFullName());
 		return false;
 	}
 
@@ -483,7 +483,7 @@ bool FObjectReplicator::ReceivedBunch( FInBunch& Bunch, const FReplicationFlags&
 
 	if ( Bunch.IsError() )
 	{
-		UE_LOG( LogNet, Error, TEXT( "ReceivedBunch: Error reading field 1: %s" ), *Object->GetFullName() );
+		UE_LOG(LogNet, Error, TEXT("ReceivedBunch: Error reading field 1: %s"), *Object->GetFullName());
 		return false;
 	}
 
@@ -506,7 +506,7 @@ bool FObjectReplicator::ReceivedBunch( FInBunch& Bunch, const FReplicationFlags&
 			// Server shouldn't receive properties.
 			if ( bIsServer )
 			{
-				UE_LOG( LogNet, Error, TEXT( "Server received unwanted property value %s in %s" ), *ReplicatedProp->GetName(), *Object->GetFullName() );
+				UE_LOG(LogNet, Error, TEXT("Server received unwanted property value %s in %s"), *ReplicatedProp->GetName(), *Object->GetFullName());
 				return false;
 			}
 		
@@ -524,7 +524,7 @@ bool FObjectReplicator::ReceivedBunch( FInBunch& Bunch, const FReplicationFlags&
 				static IConsoleVariable* CVar = IConsoleManager::Get().FindConsoleVariable(TEXT("net.Replication.DebugProperty"));
 				if (CVar && !CVar->GetString().IsEmpty() && ReplicatedProp->GetName().Contains(CVar->GetString()) )
 				{
-					UE_LOG(LogNet, Log, TEXT("Replicating Property[%d] %s on %s"), ReplicatedProp->RepIndex, *ReplicatedProp->GetName(), *Object->GetName());
+					UE_LOG(LogRep, Log, TEXT("Replicating Property[%d] %s on %s"), ReplicatedProp->RepIndex, *ReplicatedProp->GetName(), *Object->GetName());
 					DebugProperty = true;
 				}
 			}
@@ -535,7 +535,7 @@ bool FObjectReplicator::ReceivedBunch( FInBunch& Bunch, const FReplicationFlags&
 				// We hijack a non custom delta property to signify we are using FRepLayout to read the entire property block
 				if ( !RepLayout->ReceiveProperties( ObjectClass, RepState, (void*)Object, Bunch, bLocalHasUnmapped ) )
 				{
-					UE_LOG( LogNet, Error, TEXT( "ReceiveProperties FAILED %s in %s" ), *ReplicatedProp->GetName(), *Object->GetFullName() );
+					UE_LOG(LogRep, Error, TEXT("ReceiveProperties FAILED %s in %s"), *ReplicatedProp->GetName(), *Object->GetFullName());
 					return false;
 				}
 
@@ -556,7 +556,7 @@ bool FObjectReplicator::ReceivedBunch( FInBunch& Bunch, const FReplicationFlags&
 
 					if ( Element >= (uint32)ReplicatedProp->ArrayDim )
 					{
-						UE_LOG( LogNet, Error, TEXT( "Element index too large %s in %s" ), *ReplicatedProp->GetName(), *Object->GetFullName() );
+						UE_LOG(LogRep, Error, TEXT("Element index too large %s in %s"), *ReplicatedProp->GetName(), *Object->GetFullName());
 						return false;
 					}
 				}
@@ -572,7 +572,7 @@ bool FObjectReplicator::ReceivedBunch( FInBunch& Bunch, const FReplicationFlags&
 				if ( StructProperty == NULL )
 				{
 					// This property isn't custom delta
-					UE_LOG( LogNetTraffic, Error, TEXT( "Property isn't custom delta %s" ), *ReplicatedProp->GetName() );
+					UE_LOG(LogRepTraffic, Error, TEXT("Property isn't custom delta %s"), *ReplicatedProp->GetName());
 					return false;
 				}
 
@@ -581,7 +581,7 @@ bool FObjectReplicator::ReceivedBunch( FInBunch& Bunch, const FReplicationFlags&
 				if ( !( InnerStruct->StructFlags & STRUCT_NetDeltaSerializeNative ) )
 				{
 					// This property isn't custom delta
-					UE_LOG( LogNetTraffic, Error, TEXT( "Property isn't custom delta %s" ), *ReplicatedProp->GetName() );
+					UE_LOG(LogRepTraffic, Error, TEXT("Property isn't custom delta %s"), *ReplicatedProp->GetName());
 					return false;
 				}
 
@@ -605,7 +605,7 @@ bool FObjectReplicator::ReceivedBunch( FInBunch& Bunch, const FReplicationFlags&
 
 				if ( Bunch.IsError() )
 				{
-					UE_LOG( LogNet, Error, TEXT( "ReceivedBunch: NetDeltaSerialize - Bunch.IsError() == true: %s" ), *Object->GetFullName() );
+					UE_LOG(LogNet, Error, TEXT("ReceivedBunch: NetDeltaSerialize - Bunch.IsError() == true: %s"), *Object->GetFullName());
 					return false;
 				}
 
@@ -616,7 +616,7 @@ bool FObjectReplicator::ReceivedBunch( FInBunch& Bunch, const FReplicationFlags&
 				}
 
 				// Successfully received it.
-				UE_LOG( LogNetTraffic, Log, TEXT( " %s - %s" ), *Object->GetName(), *ReplicatedProp->GetName());
+				UE_LOG(LogRepTraffic, Log, TEXT(" %s - %s"), *Object->GetName(), *ReplicatedProp->GetName());
 
 				// Notify the Object if this var is RepNotify
 				QueuePropertyRepNotify( Object, ReplicatedProp, Element, MetaData );
@@ -627,7 +627,7 @@ bool FObjectReplicator::ReceivedBunch( FInBunch& Bunch, const FReplicationFlags&
 
 			if ( Bunch.IsError() )
 			{
-				UE_LOG( LogNet, Error, TEXT( "ReceivedBunch: Error reading field 2: %s" ), *Object->GetFullName() );
+				UE_LOG(LogNet, Error, TEXT("ReceivedBunch: Error reading field 2: %s"), *Object->GetFullName());
 				return false;
 			}
 		}
@@ -640,23 +640,23 @@ bool FObjectReplicator::ReceivedBunch( FInBunch& Bunch, const FReplicationFlags&
 
 			if ( Function == NULL )
 			{
-				UE_LOG( LogNet, Error, TEXT( "ReceivedBunch: Function == NULL: %s" ), *Object->GetFullName() );
+				UE_LOG(LogNet, Error, TEXT("ReceivedBunch: Function == NULL: %s"), *Object->GetFullName());
 				return false;
 			}
 
 			if ( ( Function->FunctionFlags & FUNC_Net ) == 0 )
 			{
-				UE_LOG( LogNet, Error, TEXT( "Rejected non RPC function %s in %s" ), *Message.ToString(), *Object->GetFullName() );
+				UE_LOG(LogRep, Error, TEXT("Rejected non RPC function %s in %s"), *Message.ToString(), *Object->GetFullName());
 				return false;
 			}
 
 			if ( ( Function->FunctionFlags & ( bIsServer ? FUNC_NetServer : ( FUNC_NetClient | FUNC_NetMulticast ) ) ) == 0 )
 			{
-				UE_LOG( LogNet, Error, TEXT( "Rejected RPC function due to access rights %s in %s" ), *Message.ToString(), *Object->GetFullName() );
+				UE_LOG(LogRep, Error, TEXT("Rejected RPC function due to access rights %s in %s"), *Message.ToString(), *Object->GetFullName());
 				return false;
 			}
 
-			UE_LOG( LogNetTraffic, Log, TEXT( "      Received RPC: %s" ), *Message.ToString() );
+			UE_LOG(LogRepTraffic, Log, TEXT("      Received RPC: %s"), *Message.ToString());
 
 			// Get the parameters.
 			FMemMark Mark(FMemStack::Get());
@@ -669,7 +669,7 @@ bool FObjectReplicator::ReceivedBunch( FInBunch& Bunch, const FReplicationFlags&
 
 			if ( Bunch.IsError() )
 			{
-				UE_LOG( LogNet, Error, TEXT( "ReceivedBunch: ReceivePropertiesForRPC - Bunch.IsError() == true: Function: %s, Object: %s" ), *Message.ToString(), *Object->GetFullName() );
+				UE_LOG(LogRep, Error, TEXT("ReceivedBunch: ReceivePropertiesForRPC - Bunch.IsError() == true: Function: %s, Object: %s"), *Message.ToString(), *Object->GetFullName());
 				return false;
 			}
 
@@ -685,17 +685,17 @@ bool FObjectReplicator::ReceivedBunch( FInBunch& Bunch, const FReplicationFlags&
 
 				if ( RPC_GetLastFailedReason() != NULL )
 				{
-					UE_LOG( LogNet, Error, TEXT( "ReceivedBunch: RPC_GetLastFailedReason: %s" ), RPC_GetLastFailedReason() );
+					UE_LOG(LogRep, Error, TEXT("ReceivedBunch: RPC_GetLastFailedReason: %s"), RPC_GetLastFailedReason());
 					return false;
 				}
 			}
 			else
 			{
-				UE_LOG( LogNet, Verbose, TEXT( "Rejected unwanted function %s in %s" ), *Message.ToString(), *Object->GetFullName() );
+				UE_LOG(LogRep, Verbose, TEXT("Rejected unwanted function %s in %s"), *Message.ToString(), *Object->GetFullName());
 
 				if ( !OwningChannel->Connection->TrackLogsPerSecond() )	// This will disconnect the client if we get here too often
 				{
-					UE_LOG( LogNet, Error, TEXT( "Rejected too many unwanted functions %s in %s" ), *Message.ToString(), *Object->GetFullName() );
+					UE_LOG(LogRep, Error, TEXT("Rejected too many unwanted functions %s in %s"), *Message.ToString(), *Object->GetFullName());
 					return false;
 				}
 			}
@@ -723,13 +723,13 @@ bool FObjectReplicator::ReceivedBunch( FInBunch& Bunch, const FReplicationFlags&
 
 			if ( Bunch.IsError() )
 			{
-				UE_LOG( LogNet, Error, TEXT( "ReceivedBunch: Error reading field 3: %s" ), *Object->GetFullName() );
+				UE_LOG(LogNet, Error, TEXT("ReceivedBunch: Error reading field 3: %s"), *Object->GetFullName());
 				return false;
 			}
 		}
 		else if ( FieldCache )
 		{
-			UE_LOG( LogNet, Error, TEXT( "ReceivedBunch: Invalid replicated field %i in %s" ), FieldCache->FieldNetIndex, *Object->GetFullName() );
+			UE_LOG(LogRep, Error, TEXT("ReceivedBunch: Invalid replicated field %i in %s"), FieldCache->FieldNetIndex, *Object->GetFullName());
 			return false;
 		}
 	}
@@ -741,7 +741,7 @@ void FObjectReplicator::PostReceivedBunch()
 {
 	if ( GetObject() == NULL )
 	{
-		UE_LOG( LogNet, Verbose, TEXT( "PostReceivedBunch: Object == NULL" ) );
+		UE_LOG(LogNet, Verbose, TEXT("PostReceivedBunch: Object == NULL"));
 		return;
 	}
 
@@ -774,7 +774,7 @@ static FORCEINLINE FPropertyRetirement ** UpdateAckedRetirements( FPropertyRetir
 	{
 		if ( OutAckPacketId >= (*Rec)->OutPacketIdRange.Last )
 		{
-			UE_LOG( LogNetTraffic, Verbose, TEXT("Deleting Property Record (%d >= %d)"), OutAckPacketId, (*Rec)->OutPacketIdRange.Last );
+			UE_LOG(LogRepTraffic, Verbose, TEXT("Deleting Property Record (%d >= %d)"), OutAckPacketId, (*Rec)->OutPacketIdRange.Last);
 
 			// They've ack'd this packet so we can ditch this record (easier to do it here than look for these every Ack)
 			FPropertyRetirement * ToDelete = *Rec;
@@ -903,7 +903,7 @@ bool FObjectReplicator::ReplicateProperties( FOutBunch & Bunch, FReplicationFlag
 
 	if ( Object == NULL )
 	{
-		UE_LOG( LogNet, Verbose, TEXT( "ReplicateProperties: Object == NULL" ) );
+		UE_LOG(LogRep, Verbose, TEXT("ReplicateProperties: Object == NULL"));
 		return false;
 	}
 
@@ -936,7 +936,7 @@ bool FObjectReplicator::ReplicateProperties( FOutBunch & Bunch, FReplicationFlag
 
 		if ( CVar && CVar->GetValueOnGameThread() == 1 )
 		{
-			UE_LOG( LogNetTraffic, Warning,	TEXT("      Sending queued RPCs: %s. Channel[%d] [%.1f bytes]"), *Object->GetName(), OwningChannel->ChIndex, RemoteFunctions->GetNumBits() / 8.f );
+			UE_LOG( LogRepTraffic, Warning,	TEXT("      Sending queued RPCs: %s. Channel[%d] [%.1f bytes]"), *Object->GetName(), OwningChannel->ChIndex, RemoteFunctions->GetNumBits() / 8.f );
 		}
 
 		if ( !bContentBlockWritten )
@@ -969,7 +969,7 @@ void FObjectReplicator::ForceRefreshUnreliableProperties()
 {
 	if ( GetObject() == NULL )
 	{
-		UE_LOG( LogNet, Verbose, TEXT( "ForceRefreshUnreliableProperties: Object == NULL" ) );
+		UE_LOG( LogRep, Verbose, TEXT( "ForceRefreshUnreliableProperties: Object == NULL" ) );
 		return;
 	}
 
@@ -984,7 +984,7 @@ void FObjectReplicator::PostSendBunch( FPacketIdRange & PacketRange, uint8 bReli
 {
 	if ( GetObject() == NULL )
 	{
-		UE_LOG( LogNet, Verbose, TEXT( "PostSendBunch: Object == NULL" ) );
+		UE_LOG(LogNet, Verbose, TEXT("PostSendBunch: Object == NULL"));
 		return;
 	}
 
@@ -1050,7 +1050,7 @@ void FObjectReplicator::QueueRemoteFunctionBunch( UFunction* Func, FOutBunch &Bu
 	
 	if (++RemoteFuncInfo[InfoIdx].Calls > CVarMaxRPCPerNetUpdate.GetValueOnGameThread())
 	{
-		UE_LOG(LogNet, Log, TEXT("Too many calls to RPC %s within a single netupdate. Skipping. %s.  LastCallTime: %.2f. CurrentTime: %.2f. LastRelevantTime: %.2f. LastUpdateTime: %.2f "), 
+		UE_LOG(LogRep, Log, TEXT("Too many calls to RPC %s within a single netupdate. Skipping. %s.  LastCallTime: %.2f. CurrentTime: %.2f. LastRelevantTime: %.2f. LastUpdateTime: %.2f "), 
 			*Func->GetName(), *GetObject()->GetName(), RemoteFuncInfo[InfoIdx].LastCallTime, OwningChannel->Connection->Driver->Time, OwningChannel->RelevantTime, OwningChannel->LastUpdateTime );
 		return;
 	}
@@ -1084,7 +1084,7 @@ bool FObjectReplicator::ReadyForDormancy(bool suppressLogs)
 {
 	if ( GetObject() == NULL )
 	{
-		UE_LOG( LogNet, Verbose, TEXT( "ReadyForDormancy: Object == NULL" ) );
+		UE_LOG( LogRep, Verbose, TEXT( "ReadyForDormancy: Object == NULL" ) );
 		return true;		// Technically, we don't want to hold up dormancy, but the owner needs to clean us up, so we warn
 	}
 
@@ -1093,7 +1093,7 @@ bool FObjectReplicator::ReadyForDormancy(bool suppressLogs)
 	{
 		if ( !suppressLogs )
 		{
-			UE_LOG( LogNetTraffic, Verbose, TEXT( "    [%d] Not ready for dormancy. bLastUpdateEmpty = false" ), OwningChannel->ChIndex );
+			UE_LOG( LogRepTraffic, Verbose, TEXT( "    [%d] Not ready for dormancy. bLastUpdateEmpty = false" ), OwningChannel->ChIndex );
 		}
 
 		return false;
@@ -1106,7 +1106,7 @@ bool FObjectReplicator::ReadyForDormancy(bool suppressLogs)
 		{
 			if ( !suppressLogs )
 			{
-				UE_LOG( LogNetTraffic, Verbose, TEXT( "    [%d] OutAckPacketId: %d First: %d Last: %d " ), OwningChannel->ChIndex, OwningChannel->Connection->OutAckPacketId, Retirement[i].OutPacketIdRange.First, Retirement[i].OutPacketIdRange.Last );
+				UE_LOG( LogRepTraffic, Verbose, TEXT( "    [%d] OutAckPacketId: %d First: %d Last: %d " ), OwningChannel->ChIndex, OwningChannel->Connection->OutAckPacketId, Retirement[i].OutPacketIdRange.First, Retirement[i].OutPacketIdRange.Last );
 			}
 			return false;
 		}
@@ -1119,7 +1119,7 @@ void FObjectReplicator::StartBecomingDormant()
 {
 	if ( GetObject() == NULL )
 	{
-		UE_LOG( LogNet, Verbose, TEXT( "StartBecomingDormant: Object == NULL" ) );
+		UE_LOG( LogRep, Verbose, TEXT( "StartBecomingDormant: Object == NULL" ) );
 		return;
 	}
 
@@ -1151,7 +1151,7 @@ void FObjectReplicator::CallRepNotifies()
 	{
 		for (int32 RepNotifyIdx = 0; RepNotifyIdx < RepNotifies.Num(); RepNotifyIdx++)
 		{
-			//UE_LOG(LogNet, Log,  TEXT("Calling Object->%s with %s"), *RepNotifies(RepNotifyIdx)->RepNotifyFunc.ToString(), *RepNotifies(RepNotifyIdx)->GetName()); 						
+			//UE_LOG(LogRep, Log,  TEXT("Calling Object->%s with %s"), *RepNotifies(RepNotifyIdx)->RepNotifyFunc.ToString(), *RepNotifies(RepNotifyIdx)->GetName()); 						
 			UProperty* RepProperty = RepNotifies[RepNotifyIdx];
 			UFunction* RepNotifyFunc = Object->FindFunctionChecked(RepProperty->RepNotifyFunc);
 
@@ -1217,7 +1217,7 @@ void FObjectReplicator::UpdateUnmappedObjects( bool & bOutHasMoreUnmapped )
 
 	if ( Connection->State == USOCK_Closed )
 	{
-		UE_LOG( LogNet, Warning, TEXT( "FObjectReplicator::UpdateUnmappedObjects: Connection->State == USOCK_Closed" ) );
+		UE_LOG(LogNet, Warning, TEXT("FObjectReplicator::UpdateUnmappedObjects: Connection->State == USOCK_Closed"));
 		return;
 	}
 
@@ -1307,7 +1307,7 @@ void FObjectReplicator::QueuePropertyRepNotify( UObject* Object, UProperty * Pro
 		if ( Property->ArrayDim != 1 )
 		{
 			// For static arrays, we build the meta data here, but adding the Element index that was just read into the PropMetaData array.
-			UE_LOG( LogNetTraffic, Verbose, TEXT("Property %s had ArrayDim: %d change"), *Property->GetName(), ElementIndex );
+			UE_LOG( LogRepTraffic, Verbose, TEXT("Property %s had ArrayDim: %d change"), *Property->GetName(), ElementIndex );
 
 			// Property is multi dimensional, keep track of what elements changed
 			TArray< uint8 > & PropMetaData = RepNotifyMetaData.FindOrAdd( Property );
@@ -1317,9 +1317,9 @@ void FObjectReplicator::QueuePropertyRepNotify( UObject* Object, UProperty * Pro
 		{
 			// For other properties (TArrays only now) the MetaData array is build within ::NetSerialize. Just add it to the RepNotifyMetaData map here.
 
-			//UE_LOG(LogNetTraffic, Verbose, TEXT("Property %s had MetaData: "), *Property->GetName() );
+			//UE_LOG(LogRepTraffic, Verbose, TEXT("Property %s had MetaData: "), *Property->GetName() );
 			//for (auto MetaIt = MetaData.CreateIterator(); MetaIt; ++MetaIt)
-			//	UE_LOG(LogNetTraffic, Verbose, TEXT("   %d"), *MetaIt );
+			//	UE_LOG(LogRepTraffic, Verbose, TEXT("   %d"), *MetaIt );
 
 			// Property included some meta data about what was serialized. 
 			TArray< uint8 > & PropMetaData = RepNotifyMetaData.FindOrAdd( Property );
diff --git a/Engine/Source/Runtime/Engine/Private/RepLayout.cpp b/Engine/Source/Runtime/Engine/Private/RepLayout.cpp
--- a/Engine/Source/Runtime/Engine/Private/RepLayout.cpp
+++ b/Engine/Source/Runtime/Engine/Private/RepLayout.cpp
@@ -82,7 +82,7 @@ static FORCEINLINE bool PropertiesAreIdenticalNative( const FRepLayoutCmd& Cmd,
 		case REPCMD_PropertyString:			return CompareValue<FString>( A, B );
 		case REPCMD_Property:				return Cmd.Property->Identical( A, B );
 		default: 
-			UE_LOG( LogNet, Fatal, TEXT( "PropertiesAreIdentical: Unsupported type! %i (%s)" ), Cmd.Type, *Cmd.Property->GetName() );
+			UE_LOG( LogRep, Fatal, TEXT( "PropertiesAreIdentical: Unsupported type! %i (%s)" ), Cmd.Type, *Cmd.Property->GetName() );
 	}
 
 	return false;
@@ -95,7 +95,7 @@ static FORCEINLINE bool PropertiesAreIdentical( const FRepLayoutCmd& Cmd, const
 	// Sanity check result
 	if ( bIsIdentical != Cmd.Property->Identical( A, B ) )
 	{
-		UE_LOG( LogNet, Fatal, TEXT( "PropertiesAreIdentical: Result mismatch! (%s)" ), *Cmd.Property->GetFullName() );
+		UE_LOG( LogRep, Fatal, TEXT( "PropertiesAreIdentical: Result mismatch! (%s)" ), *Cmd.Property->GetFullName() );
 	}
 #endif
 	return bIsIdentical;
@@ -138,7 +138,7 @@ static void SerializeReadWritePropertyChecksum( const FRepLayoutCmd& Cmd, const
 	if ( MarkerChecksum != OriginalMarkerChecksum )
 	{
 		// This is fatal, as it means we are out of sync to the point we can't recover
-		UE_LOG( LogNet, Fatal, TEXT( "SerializeReadWritePropertyChecksum: Property checksum marker failed! [%s]" ), *Cmd.Property->GetFullName() );
+		UE_LOG( LogRep, Fatal, TEXT( "SerializeReadWritePropertyChecksum: Property checksum marker failed! [%s]" ), *Cmd.Property->GetFullName() );
 	}
 
 	if ( Cmd.Property->IsA( UObjectPropertyBase::StaticClass() ) )
@@ -202,7 +202,7 @@ static void SerializeReadWritePropertyChecksum( const FRepLayoutCmd& Cmd, const
 	if ( PropertyChecksum != OriginalPropertyChecksum )
 	{
 		// This is a warning, because for some reason, float rounding issues in the quantization functions cause this to return false positives
-		UE_LOG( LogNet, Warning, TEXT( "Property checksum failed! [%s]" ), *Cmd.Property->GetFullName() );
+		UE_LOG( LogRep, Warning, TEXT( "Property checksum failed! [%s]" ), *Cmd.Property->GetFullName() );
 	}
 }
 
@@ -497,9 +497,9 @@ void FRepLayout::LogChangeListMismatches(
 	UObject *						Object			= (UObject*)Data;
 	const UNetDriver *				NetDriver		= OwningChannel->Connection->Driver;
 
-	UE_LOG( LogNet, Warning, TEXT( "LogChangeListMismatches: %s" ), *Object->GetName() );
+	UE_LOG( LogRep, Warning, TEXT( "LogChangeListMismatches: %s" ), *Object->GetName() );
 
-	UE_LOG( LogNet, Warning, TEXT( "  %i, %i, %i, %i" ), NetDriver->ReplicationFrame, ChangeTracker->LastReplicationFrame, ChangeTracker->LastReplicationGroupFrame, RepState1->LastReplicationFrame );
+	UE_LOG( LogRep, Warning, TEXT( "  %i, %i, %i, %i" ), NetDriver->ReplicationFrame, ChangeTracker->LastReplicationFrame, ChangeTracker->LastReplicationGroupFrame, RepState1->LastReplicationFrame );
 
 	for ( int32 i = 0; i < PropertyList.Num(); i++ )
 	{
@@ -513,7 +513,7 @@ void FRepLayout::LogChangeListMismatches(
 			FString Changed2Str = Changed2 ? TEXT( "TRUE" ) : TEXT( "FALSE" );
 			FString ExtraStr	= Changed1 != Changed2 ? TEXT( "<--- MISMATCH!" ) : TEXT( "<--- SAME" );
 
-			UE_LOG( LogNet, Warning, TEXT( "    Property changed: %s (%s / %s) %s" ), *Parents[Index].Property->GetName(), *Changed1Str, *Changed2Str, *ExtraStr );
+			UE_LOG( LogRep, Warning, TEXT( "    Property changed: %s (%s / %s) %s" ), *Parents[Index].Property->GetName(), *Changed1Str, *Changed2Str, *ExtraStr );
 
 			FString StringValue1;
 			FString StringValue2;
@@ -523,7 +523,7 @@ void FRepLayout::LogChangeListMismatches(
 			Parents[Index].Property->ExportText_InContainer( Parents[Index].ArrayIndex, StringValue2, RepState1->StaticBuffer.GetData(), NULL, NULL, PPF_DebugDump );
 			Parents[Index].Property->ExportText_InContainer( Parents[Index].ArrayIndex, StringValue3, RepState2->StaticBuffer.GetData(), NULL, NULL, PPF_DebugDump );
 
-			UE_LOG( LogNet, Warning, TEXT( "      Values: %s, %s, %s" ), *StringValue1, *StringValue2, *StringValue3 );
+			UE_LOG( LogRep, Warning, TEXT( "      Values: %s, %s, %s" ), *StringValue1, *StringValue2, *StringValue3 );
 		}
 	}
 }
@@ -561,7 +561,7 @@ void FRepLayout::SanityCheckShadowStateAgainstChangeList(
 	if ( Result != ChangedParentsHasChanged( PropertyList, OtherChangedParents ) )
 	{
 		LogChangeListMismatches( Data, OwningChannel, PropertyList, RepState, OtherRepState, LocalChangedParents, OtherChangedParents );
-		UE_LOG( LogNet, Fatal, TEXT( "ReplicateProperties: Compare result mismatch: %s" ), *Object->GetName() );
+		UE_LOG( LogRep, Fatal, TEXT( "ReplicateProperties: Compare result mismatch: %s" ), *Object->GetName() );
 	}
 
 	for ( int32 i = 0; i < PropertyList.Num(); i++ )
@@ -571,7 +571,7 @@ void FRepLayout::SanityCheckShadowStateAgainstChangeList(
 		if ( OtherChangedParents[Index].Changed.Num() != LocalChangedParents[Index].Changed.Num() )
 		{
 			LogChangeListMismatches( Data, OwningChannel, PropertyList, RepState, OtherRepState, LocalChangedParents, OtherChangedParents );
-			UE_LOG( LogNet, Fatal, TEXT( "ReplicateProperties: Compare count mismatch: %s" ), *Object->GetName() );
+			UE_LOG( LogRep, Fatal, TEXT( "ReplicateProperties: Compare count mismatch: %s" ), *Object->GetName() );
 		}
 
 		for ( int32 j = 0; j < OtherChangedParents[Index].Changed.Num(); j++ )
@@ -579,7 +579,7 @@ void FRepLayout::SanityCheckShadowStateAgainstChangeList(
 			if ( OtherChangedParents[Index].Changed[j] != LocalChangedParents[Index].Changed[j] )
 			{
 				LogChangeListMismatches( Data, OwningChannel, PropertyList, RepState, OtherRepState, LocalChangedParents, OtherChangedParents );
-				UE_LOG( LogNet, Fatal, TEXT( "ReplicateProperties: Compare changelist value mismatch: %s" ), *Object->GetName() );
+				UE_LOG( LogRep, Fatal, TEXT( "ReplicateProperties: Compare changelist value mismatch: %s" ), *Object->GetName() );
 			}
 		}
 	}
@@ -782,7 +782,7 @@ void FRepLayout::UpdateChangelistHistory( FRepState * RepState, UClass * ObjectC
 	// If our buffer is currently full, forcibly send the entire history
 	if ( DumpHistory )
 	{
-		UE_LOG( LogNet, Warning, TEXT( "FRepLayout::UpdateChangelistHistory: History overflow, forcing history dump %s" ), *ObjectClass->GetName() );
+		UE_LOG( LogRep, Warning, TEXT( "FRepLayout::UpdateChangelistHistory: History overflow, forcing history dump %s" ), *ObjectClass->GetName() );
 	}
 
 	for ( int32 i = RepState->HistoryStart; i < RepState->HistoryEnd; i++ )
@@ -1335,12 +1335,12 @@ public:
 			// Note - If we already have an existing entry for this offset, we will replace it with this most recent data
 			StackState.UnmappedGuids->Map.Add( AbsCmdOffset, FUnmappedGuidMgrElement( Bunch, Mark, TrackedUnmappedGuids, Cmd.ParentIndex, CmdIndex ) );
 
-			UE_LOG( LogNet, Verbose, TEXT( "ADDED unmapped property: Offset: %i, Name: %s"), AbsCmdOffset, *Cmd.Property->GetName() );
+			UE_LOG( LogRep, Verbose, TEXT( "ADDED unmapped property: Offset: %i, Name: %s"), AbsCmdOffset, *Cmd.Property->GetName() );
 
 			// List all the guids that were unmapped
 			for ( int32 i = 0; i < TrackedUnmappedGuids.Num(); i++ )
 			{
-				UE_LOG( LogNet, Verbose, TEXT( "  Guid: %s"), *TrackedUnmappedGuids[i].ToString() );
+				UE_LOG( LogRep, Verbose, TEXT( "  Guid: %s"), *TrackedUnmappedGuids[i].ToString() );
 			}
 
 			bHasUnmapped = true;
@@ -1399,7 +1399,7 @@ bool FRepLayout::ReceiveProperties( UClass * InObjectClass, FRepState * RESTRICT
 	// Make sure we're waiting on the last NULL terminator
 	if ( ReceivePropertiesImpl.WaitingHandle != 0 )
 	{
-		UE_LOG( LogNet, Warning, TEXT( "Read out of sync." ) );
+		UE_LOG( LogRep, Warning, TEXT( "Read out of sync." ) );
 		return false;
 	}
 
@@ -1442,7 +1442,7 @@ void FRepLayout::UpdateUnmappedObjects_r(
 		if ( AbsOffset >= MaxAbsOffset )
 		{
 			// Array must have shrunk, we can remove this item
-			UE_LOG( LogNet, VeryVerbose, TEXT( "UpdateUnmappedObjects_r: REMOVED unmapped property: AbsOffset >= MaxAbsOffset. Offset: %i" ), AbsOffset );
+			UE_LOG( LogRep, VeryVerbose, TEXT( "UpdateUnmappedObjects_r: REMOVED unmapped property: AbsOffset >= MaxAbsOffset. Offset: %i" ), AbsOffset );
 			It.RemoveCurrent();
 			continue;
 		}
@@ -1472,7 +1472,7 @@ void FRepLayout::UpdateUnmappedObjects_r(
 
 			if ( PackageMap->IsGUIDBroken( GUID, false ) )
 			{
-				UE_LOG( LogNet, Warning, TEXT( "UpdateUnmappedObjects_r: Broken GUID. NetGuid: %s" ), *GUID.ToString() );
+				UE_LOG( LogRep, Warning, TEXT( "UpdateUnmappedObjects_r: Broken GUID. NetGuid: %s" ), *GUID.ToString() );
 				UnmappedProperty.UnmappedGUIDs.RemoveAt( i );
 				continue;
 			}
@@ -1481,7 +1481,7 @@ void FRepLayout::UpdateUnmappedObjects_r(
 
 			if ( Object != NULL )
 			{
-				UE_LOG( LogNet, VeryVerbose, TEXT( "UpdateUnmappedObjects_r: REMOVED unmapped property: Offset: %i, Guid: %s, PropName: %s, ObjName: %s" ), AbsOffset, *GUID.ToString(), *Cmd.Property->GetName(), *Object->GetName() );
+				UE_LOG( LogRep, VeryVerbose, TEXT( "UpdateUnmappedObjects_r: REMOVED unmapped property: Offset: %i, Guid: %s, PropName: %s, ObjName: %s" ), AbsOffset, *GUID.ToString(), *Cmd.Property->GetName(), *Object->GetName() );
 				UnmappedProperty.UnmappedGUIDs.RemoveAt( i );
 				bMappedSomeGUIDs = true;
 			}
@@ -1583,12 +1583,12 @@ void FRepLayout::ValidateWithChecksum_DynamicArray_r( const FRepLayoutCmd& Cmd,
 
 	if ( ArrayNum != Array->Num() )
 	{
-		UE_LOG( LogNet, Fatal, TEXT( "ValidateWithChecksum_AnyArray_r: Array sizes different! %s %i / %i" ), *Cmd.Property->GetFullName(), ArrayNum, Array->Num() );
+		UE_LOG( LogRep, Fatal, TEXT( "ValidateWithChecksum_AnyArray_r: Array sizes different! %s %i / %i" ), *Cmd.Property->GetFullName(), ArrayNum, Array->Num() );
 	}
 
 	if ( ElementSize != Cmd.ElementSize )
 	{
-		UE_LOG( LogNet, Fatal, TEXT( "ValidateWithChecksum_AnyArray_r: Array element sizes different! %s %i / %i" ), *Cmd.Property->GetFullName(), ElementSize, Cmd.ElementSize );
+		UE_LOG( LogRep, Fatal, TEXT( "ValidateWithChecksum_AnyArray_r: Array element sizes different! %s %i / %i" ), *Cmd.Property->GetFullName(), ElementSize, Cmd.ElementSize );
 	}
 
 	uint8* LocalData = (uint8*)Array->GetData();
@@ -1899,7 +1899,7 @@ public:
 
 			if ( !bSync )
 			{			
-				UE_LOG( LogNet, Warning, TEXT( "FDiffPropertiesImpl: Array sizes different: %s %i / %i" ), *Cmd.Property->GetFullName(), StackState.DataArray->Num(), StackState.ShadowArray->Num() );
+				UE_LOG( LogRep, Warning, TEXT( "FDiffPropertiesImpl: Array sizes different: %s %i / %i" ), *Cmd.Property->GetFullName(), StackState.DataArray->Num(), StackState.ShadowArray->Num() );
 				return;
 			}
 
@@ -1932,7 +1932,7 @@ public:
 
 			if ( !bSync )
 			{			
-				UE_LOG( LogNet, Warning, TEXT( "FDiffPropertiesImpl: Property different: %s" ), *Cmd.Property->GetFullName() );
+				UE_LOG( LogRep, Warning, TEXT( "FDiffPropertiesImpl: Property different: %s" ), *Cmd.Property->GetFullName() );
 				return;
 			}
 
@@ -2021,7 +2021,7 @@ void FRepLayout::AddPropertyCmd( UProperty * Property, int32 Offset, int32 Relat
 		}
 		else
 		{
-			UE_LOG( LogNet, VeryVerbose, TEXT( "AddPropertyCmd: Falling back to default type for property [%s]" ), *Cmd.Property->GetFullName() );
+			UE_LOG( LogRep, VeryVerbose, TEXT( "AddPropertyCmd: Falling back to default type for property [%s]" ), *Cmd.Property->GetFullName() );
 		}
 	}
 	else if ( Property->IsA( UBoolProperty::StaticClass() ) )
@@ -2062,7 +2062,7 @@ void FRepLayout::AddPropertyCmd( UProperty * Property, int32 Offset, int32 Relat
 	}
 	else
 	{
-		UE_LOG( LogNet, VeryVerbose, TEXT( "AddPropertyCmd: Falling back to default type for property [%s]" ), *Cmd.Property->GetFullName() );
+		UE_LOG( LogRep, VeryVerbose, TEXT( "AddPropertyCmd: Falling back to default type for property [%s]" ), *Cmd.Property->GetFullName() );
 	}
 }
 
@@ -2375,7 +2375,7 @@ void FRepLayout::SerializeProperties_DynamicArray_r(
 
 	if ( ArrayNum > MAX_ARRAY_SIZE )
 	{
-		UE_LOG( LogNetTraffic, Error, TEXT( "SerializeProperties_DynamicArray_r: ArrayNum > MAX_ARRAY_SIZE (%s)" ), *Cmd.Property->GetName() );
+		UE_LOG( LogRepTraffic, Error, TEXT( "SerializeProperties_DynamicArray_r: ArrayNum > MAX_ARRAY_SIZE (%s)" ), *Cmd.Property->GetName() );
 		Ar.SetError();
 		return;
 	}
@@ -2384,7 +2384,7 @@ void FRepLayout::SerializeProperties_DynamicArray_r(
 
 	if ( (int32)ArrayNum * Cmd.ElementSize > MAX_ARRAY_MEMORY )
 	{
-		UE_LOG( LogNetTraffic, Error, TEXT( "SerializeProperties_DynamicArray_r: ArrayNum * Cmd.ElementSize > MAX_ARRAY_MEMORY (%s)" ), *Cmd.Property->GetName() );
+		UE_LOG( LogRepTraffic, Error, TEXT( "SerializeProperties_DynamicArray_r: ArrayNum * Cmd.ElementSize > MAX_ARRAY_MEMORY (%s)" ), *Cmd.Property->GetName() );
 		Ar.SetError();
 		return;
 	}
@@ -2478,7 +2478,7 @@ void FRepLayout::SendPropertiesForRPC( UObject* Object, UFunction * Function, UA
 			if ( bHasUnmapped )
 			{
 				// RPC function is sending an unmapped object...
-				UE_LOG( LogNetTraffic, Log, TEXT( "Actor[%d] %s RPC %s parameter %s was sent while unmapped! This call may not be correctly handled on the receiving end." ),
+				UE_LOG( LogRepTraffic, Log, TEXT( "Actor[%d] %s RPC %s parameter %s was sent while unmapped! This call may not be correctly handled on the receiving end." ),
 					Channel->ChIndex, *Object->GetName(), *Function->GetName(), *Parents[i].Property->GetName() );
 			}
 		}
@@ -2510,7 +2510,7 @@ void FRepLayout::ReceivePropertiesForRPC( UObject* Object, UFunction * Function,
 			
 			if ( bHasUnmapped )
 			{
-				UE_LOG( LogNetTraffic, Log, TEXT( "Unable to resolve RPC parameter. Object[%d] %s. Function %s. Parameter %s." ), 
+				UE_LOG( LogRepTraffic, Log, TEXT( "Unable to resolve RPC parameter. Object[%d] %s. Function %s. Parameter %s." ), 
 					Channel->ChIndex, *Object->GetName(), *Function->GetName(), *Parents[i].Property->GetName() );
 			}
 		}
@@ -2536,7 +2536,7 @@ void FRepLayout::RebuildConditionalProperties( FRepState * RESTRICT	RepState, co
 {
 	SCOPE_CYCLE_COUNTER( STAT_NetRebuildConditionalTime );
 
-	//UE_LOG( LogNet, Warning, TEXT( "Rebuilding custom properties [%s]" ), *Owner->GetName() );
+	//UE_LOG( LogRep, Warning, TEXT( "Rebuilding custom properties [%s]" ), *Owner->GetName() );
 	
 	// Setup condition map
 	bool ConditionMap[COND_Max];
diff --git a/Engine/Source/Runtime/Engine/Public/EngineLogs.h b/Engine/Source/Runtime/Engine/Public/EngineLogs.h
--- a/Engine/Source/Runtime/Engine/Public/EngineLogs.h
+++ b/Engine/Source/Runtime/Engine/Public/EngineLogs.h
@@ -17,8 +17,10 @@ ENGINE_API DECLARE_LOG_CATEGORY_EXTERN(LogLevel, Log, All);
 ENGINE_API DECLARE_LOG_CATEGORY_EXTERN(LogSkeletalMesh, Log, All);
 ENGINE_API DECLARE_LOG_CATEGORY_EXTERN(LogStaticMesh, Log, All);
 ENGINE_API DECLARE_LOG_CATEGORY_EXTERN(LogNet, Log, All);
+ENGINE_API DECLARE_LOG_CATEGORY_EXTERN(LogRep, Log, All);
 ENGINE_API DECLARE_LOG_CATEGORY_EXTERN(LogNetPlayerMovement, Warning, All);
 ENGINE_API DECLARE_LOG_CATEGORY_EXTERN(LogNetTraffic, Warning, All);
+ENGINE_API DECLARE_LOG_CATEGORY_EXTERN(LogRepTraffic, Warning, All);
 ENGINE_API DECLARE_LOG_CATEGORY_EXTERN(LogNetDormancy, Warning, All);
 ENGINE_API DECLARE_LOG_CATEGORY_EXTERN(LogSkeletalControl, Warning, All);
 ENGINE_API DECLARE_LOG_CATEGORY_EXTERN(LogSubtitle, Log, All);
