SerializeNewActor does a 1 bit check whether to serialize Location, Rotation, Scale and Velocity

Changes:
1 bit check for each  Location, Rotation, Scale and Velocity whether they are default values and if they aren't then serialize the data, if they are default then don't serialize.

UE-3302

SpawnActor now replicates Scale

Changes:
UPackageMapClient::SerializeNewActor() replicates scale and sets it when an object has spawned

UE-14509

diff --git a/Engine/Source/Runtime/Engine/Private/PackageMapClient.cpp b/Engine/Source/Runtime/Engine/Private/PackageMapClient.cpp
--- a/Engine/Source/Runtime/Engine/Private/PackageMapClient.cpp
+++ b/Engine/Source/Runtime/Engine/Private/PackageMapClient.cpp
@@ -162,8 +162,8 @@ bool UPackageMapClient::WriteObject( FArchive& Ar, UObject* ObjOuter, FNetworkGU
  *	Standard method of serializing a new actor.
  *		For static actors, this will just be a single call to SerializeObject, since they can be referenced by their path name.
  *		For dynamic actors, first the actor's reference is serialized but will not resolve on clients since they haven't spawned the actor yet.
- *			The actor archetype is then serialized along with the starting location, rotation, and velocity.
- *			After reading this information, the client spawns this actor via GWorld and assigns it the NetGUID it read at the top of the function.
+ *		The actor archetype is then serialized along with the starting location, rotation, and velocity.
+ *		After reading this information, the client spawns this actor via GWorld and assigns it the NetGUID it read at the top of the function.
  *
  *		returns true if a new actor was spawned. false means an existing actor was found for the netguid.
  */
@@ -218,6 +218,7 @@ bool UPackageMapClient::SerializeNewActor(FArchive& Ar, class UActorChannel *Cha
 	{
 		UObject* Archetype = NULL;
 		FVector_NetQuantize10 Location;
+		FVector_NetQuantize10 Scale;
 		FVector_NetQuantize10 Velocity;
 		FRotator Rotation;
 		bool SerSuccess;
@@ -230,9 +231,12 @@ bool UPackageMapClient::SerializeNewActor(FArchive& Ar, class UActorChannel *Cha
 			check( Actor->NeedsLoadForClient() );			// We have no business sending this unless the client can load
 			check( Archetype->NeedsLoadForClient() );		// We have no business sending this unless the client can load
 
-			Location = Actor->GetRootComponent() ? Actor->GetActorLocation() : FVector::ZeroVector;
-			Rotation = Actor->GetRootComponent() ? Actor->GetActorRotation() : FRotator::ZeroRotator;
-			Velocity = Actor->GetVelocity();
+			const USceneComponent* RootComponent = Actor->GetRootComponent();
+
+			Location = RootComponent ? Actor->GetActorLocation() : FVector::ZeroVector;
+			Rotation = RootComponent ? Actor->GetActorRotation() : FRotator::ZeroRotator;
+			Scale = RootComponent ? Actor->GetActorScale() : FVector::ZeroVector;
+			Velocity = RootComponent ? Actor->GetVelocity() : FVector::ZeroVector;
 		}
 
 		FNetworkGUID ArchetypeNetGUID;
@@ -253,11 +257,92 @@ bool UPackageMapClient::SerializeNewActor(FArchive& Ar, class UActorChannel *Cha
 		}
 
 		// SerializeCompressedInitial
+		// only serialize the components that need to be serialized otherwise default them
+		bool bSerializeLocation = false;
+		bool bSerializeRotation = false;
+		bool bSerializeScale = false;
+		bool bSerializeVelocity = false;
+		const float epsilon = 0.001f;
 		{			
-			// read/write compressed location
-			Location.NetSerialize( Ar, this, SerSuccess );
-			Rotation.NetSerialize( Ar, this, SerSuccess );
-			Velocity.NetSerialize( Ar, this, SerSuccess );
+			// Server is serializing an object to be sent to a client
+			if (Ar.IsSaving())
+			{
+				const FVector DefaultScale(1.f, 1.f, 1.f);
+
+				// If the Location isn't the default Location
+				bSerializeLocation = !Location.Equals(FVector::ZeroVector, epsilon);
+				bSerializeRotation = !Rotation.Equals(FRotator::ZeroRotator, epsilon);
+				bSerializeScale = !Scale.Equals(DefaultScale, epsilon);
+				bSerializeVelocity = !Velocity.Equals(FVector::ZeroVector, epsilon);
+
+				Ar.SerializeBits(&bSerializeLocation, 1);
+				if (bSerializeLocation)
+				{
+					Location.NetSerialize(Ar, this, SerSuccess);
+				}
+
+				Ar.SerializeBits(&bSerializeRotation, 1);
+				if (bSerializeRotation)
+				{
+					Rotation.NetSerialize(Ar, this, SerSuccess);
+				}
+
+				Ar.SerializeBits(&bSerializeScale, 1);
+				if (bSerializeScale)
+				{
+					Scale.NetSerialize(Ar, this, SerSuccess);
+				}
+
+				Ar.SerializeBits(&bSerializeVelocity, 1);
+				if (bSerializeVelocity)
+				{
+					Velocity.NetSerialize(Ar, this, SerSuccess);
+				}
+			}
+			// Client is serializing from an object sent from the server
+			else if (Ar.IsLoading())
+			{
+
+				Ar.SerializeBits(&bSerializeLocation, 1);
+				if (bSerializeLocation)
+				{
+					Location.NetSerialize(Ar, this, SerSuccess);
+				}
+				else
+				{
+					Location = FVector::ZeroVector;
+				}
+
+				Ar.SerializeBits(&bSerializeRotation, 1);
+				if (bSerializeRotation)
+				{
+					Rotation.NetSerialize(Ar, this, SerSuccess);
+				}
+				else
+				{
+					Rotation = FRotator::ZeroRotator;
+				}
+
+				Ar.SerializeBits(&bSerializeScale, 1);
+				if (bSerializeScale)
+				{
+					Scale.NetSerialize(Ar, this, SerSuccess);
+				}
+				else
+				{
+					Scale = FVector(1, 1, 1);
+				}
+
+				Ar.SerializeBits(&bSerializeVelocity, 1);
+				if (bSerializeVelocity)
+				{
+					Velocity.NetSerialize(Ar, this, SerSuccess);
+				}
+				else
+				{
+					Velocity = FVector::ZeroVector;
+				}
+			}
 
 			if ( Channel && Ar.IsSaving() )
 			{
@@ -285,7 +370,18 @@ bool UPackageMapClient::SerializeNewActor(FArchive& Ar, class UActorChannel *Cha
 					SpawnInfo.bRemoteOwned = true;
 					SpawnInfo.bNoFail = true;
 					Actor = Connection->Driver->GetWorld()->SpawnActor(Archetype->GetClass(), &Location, &Rotation, SpawnInfo );
-					Actor->PostNetReceiveVelocity(Velocity);
+
+					// Velocity was serialized by the server
+					if (bSerializeVelocity)
+					{
+						Actor->PostNetReceiveVelocity(Velocity);
+					}
+
+					// Scale was serialized by the server
+					if (bSerializeScale)
+					{
+						Actor->SetActorScale3D(Scale);
+					}
 
 					GuidCache->RegisterNetGUID_Client( NetGUID, Actor );
 					bActorWasSpawned = true;
