Resolved compilation error from DialogueWave.cpp not including the header for UDialogueVoice.

diff --git a/Engine/Source/Runtime/Engine/Private/DialogueWave.cpp b/Engine/Source/Runtime/Engine/Private/DialogueWave.cpp
--- a/Engine/Source/Runtime/Engine/Private/DialogueWave.cpp
+++ b/Engine/Source/Runtime/Engine/Private/DialogueWave.cpp
@@ -8,6 +8,7 @@
 #include "SubtitleManager.h"
 #include "GatherableTextData.h"
 #include "InternationalizationMetadata.h"
+#include "Sound/DialogueVoice.h"
 
 #if WITH_EDITORONLY_DATA
 namespace
