Feature request: Added World<->View transforms to TransformPosition material node.

diff --git a/Engine/Shaders/MaterialTemplate.usf b/Engine/Shaders/MaterialTemplate.usf
--- a/Engine/Shaders/MaterialTemplate.usf
+++ b/Engine/Shaders/MaterialTemplate.usf
@@ -405,6 +405,18 @@ MaterialFloat3 TransformViewVectorToLocal(MaterialFloat3 InViewVector)
 	return mul(mul(InViewVector, (MaterialFloat3x3)View.ViewToTranslatedWorld), (MaterialFloat3x3)Primitive.WorldToLocal);
 }
 
+/** Transforms a vector from view space to world space */
+MaterialFloat3 TransformViewVectorToWorld(MaterialFloat3 InViewVector)
+{
+	return mul(InViewVector, (MaterialFloat3x3)View.ViewToTranslatedWorld);
+}
+
+/** Transforms a vector from world space to view space */
+MaterialFloat3 TransformWorldVectorToView(MaterialFloat3 InWorldVector)
+{
+	return mul(InWorldVector, (MaterialFloat3x3)View.TranslatedWorldToView);
+}
+
 /** Transforms a position from local space to absolute world space */
 float3 TransformLocalPositionToWorld(FMaterialPixelParameters Parameters,float3 InLocalPosition)
 {
@@ -427,6 +439,18 @@ float3 TransformWorldPositionToLocal(float3 InWorldPosition)
 	return mul(float4(InWorldPosition, 1), Primitive.WorldToLocal).xyz;
 }
 
+/** Transforms a position from view space to world space */
+MaterialFloat3 TransformViewPositionToWorld(MaterialFloat3 InViewPosition)
+{
+	return mul(MaterialFloat4(InViewPosition, 1), View.ViewToTranslatedWorld).xyz - View.PreViewTranslation.xyz;
+}
+
+/** Transforms a position from world space to view space */
+MaterialFloat3 TransformWorldPositionToView(MaterialFloat3 InWorldPosition)
+{
+	return mul(MaterialFloat4(InWorldPosition + View.PreViewTranslation.xyz, 1), View.TranslatedWorldToView).xyz;
+}
+
 /** Return the object's position in world space */
 float3 GetObjectWorldPosition(FMaterialPixelParameters Parameters)
 {
@@ -495,18 +519,6 @@ MaterialFloat GetDistanceCullFade()
 
 #endif // #if HAS_PRIMITIVE_UNIFORM_BUFFER
 
-/** Transforms a vector from view space to world space */
-MaterialFloat3 TransformViewVectorToWorld(MaterialFloat3 InViewVector)
-{
-	return mul(InViewVector, (MaterialFloat3x3)View.ViewToTranslatedWorld);
-}
-
-/** Transforms a vector from world space to view space */
-MaterialFloat3 TransformWorldVectorToView(MaterialFloat3 InWorldVector)
-{
-	return mul(InWorldVector, (MaterialFloat3x3)View.TranslatedWorldToView);
-}
-
 /** Rotates Position about the given axis by the given angle, in radians, and returns the offset to Position. */
 float3 RotateAboutAxis(float4 NormalizedRotationAxisAndAngle, float3 PositionOnAxis, float3 Position)
 {
diff --git a/Engine/Source/Runtime/Engine/Classes/Materials/MaterialExpressionTransformPosition.h b/Engine/Source/Runtime/Engine/Classes/Materials/MaterialExpressionTransformPosition.h
--- a/Engine/Source/Runtime/Engine/Classes/Materials/MaterialExpressionTransformPosition.h
+++ b/Engine/Source/Runtime/Engine/Classes/Materials/MaterialExpressionTransformPosition.h
@@ -10,6 +10,7 @@ enum EMaterialPositionTransformSource
 {
 	TRANSFORMPOSSOURCE_Local UMETA(DisplayName="Local"),
 	TRANSFORMPOSSOURCE_World UMETA(DisplayName="World"),
+	TRANSFORMPOSSOURCE_View  UMETA(DisplayName="View"),
 	TRANSFORMPOSSOURCE_MAX,
 };
 
diff --git a/Engine/Source/Runtime/Engine/Private/Materials/HLSLMaterialTranslator.h b/Engine/Source/Runtime/Engine/Private/Materials/HLSLMaterialTranslator.h
--- a/Engine/Source/Runtime/Engine/Private/Materials/HLSLMaterialTranslator.h
+++ b/Engine/Source/Runtime/Engine/Private/Materials/HLSLMaterialTranslator.h
@@ -3312,15 +3312,73 @@ protected:
 			// code string to transform the input vector
 			FString CodeStr;
 
+			if (SourceCoordinateSpace == TRANSFORMPOSSOURCE_Local ||
+				DestinationCoordinateSpace == TRANSFORMPOSSOURCE_Local)
+			{
+				if (ShaderFrequency != SF_Pixel && ShaderFrequency != SF_Compute && ShaderFrequency != SF_Vertex)
+				{
+					return Errorf(TEXT("Local space is only supported for vertex or pixel shader!"));
+				}
+			}
+			if (SourceCoordinateSpace == TRANSFORMPOSSOURCE_View ||
+				DestinationCoordinateSpace == TRANSFORMPOSSOURCE_View)
+			{
+				if (ShaderFrequency != SF_Pixel && ShaderFrequency != SF_Compute && ShaderFrequency != SF_Vertex)
+				{
+					return Errorf(TEXT("View space in only supported for vertex or pixel shader!"));
+				}
+			}
+
 			if (SourceCoordinateSpace == TRANSFORMPOSSOURCE_Local)
 			{
-				CodeStr = FString(TEXT("TransformLocalPositionToWorld(Parameters,%s)"));
+				if (DestinationCoordinateSpace == TRANSFORMPOSSOURCE_World)
+				{
+					CodeStr = FString(TEXT("TransformLocalPositionToWorld(Parameters,%s)"));
+				}
+				else if (DestinationCoordinateSpace == TRANSFORMPOSSOURCE_View)
+				{
+					CodeStr = FString(TEXT("TransformWorldPositionToView(TransformLocalPositionToWorld(Parameters,%s))"));
+				}
+				else
+				{
+					UE_LOG(LogMaterial, Fatal, TEXT("Invalid DestCoordType. See EMaterialPositionTransformSource"));
+				}
 			}
 			else if (SourceCoordinateSpace == TRANSFORMPOSSOURCE_World)
 			{
-				CodeStr = FString(TEXT("TransformWorldPositionToLocal(%s)"));
+				if (DestinationCoordinateSpace == TRANSFORMPOSSOURCE_Local)
+				{
+					CodeStr = FString(TEXT("TransformWorldPositionToLocal(%s)"));
+				}
+				else if (DestinationCoordinateSpace == TRANSFORMPOSSOURCE_View)
+				{
+					CodeStr = FString(TEXT("TransformWorldPositionToView(%s)"));
+				}
+				else
+				{
+					UE_LOG(LogMaterial, Fatal, TEXT("Invalid DestCoordType. See EMaterialPositionTransformSource"));
+				}
+			}
+			else if (SourceCoordinateSpace == TRANSFORMPOSSOURCE_View)
+			{
+				if (DestinationCoordinateSpace == TRANSFORMPOSSOURCE_Local)
+				{
+					CodeStr = FString(TEXT("TransformWorldPositionToLocal(TransformViewPositionToWorld(%s))"));
+				}
+				else if (DestinationCoordinateSpace == TRANSFORMPOSSOURCE_World)
+				{
+					CodeStr = FString(TEXT("TransformViewPositionToWorld(%s)"));
+				}
+				else
+				{
+					UE_LOG(LogMaterial, Fatal, TEXT("Invalid DestCoordType. See EMaterialPositionTransformSource"));
+				}
+			}
+			else
+			{
+				UE_LOG(LogMaterial, Fatal, TEXT("Invalid SourceCoordType. See EMaterialPositionTransformSource"));
 			}
-				
+
 			Result = AddCodeChunk(
 				MCT_Float3,
 				*CodeStr,
