UE-14157 State units (degrees) for MakeRot and BreakRot functions

diff --git a/Engine/Source/Runtime/Engine/Classes/Kismet/KismetMathLibrary.h b/Engine/Source/Runtime/Engine/Classes/Kismet/KismetMathLibrary.h
--- a/Engine/Source/Runtime/Engine/Classes/Kismet/KismetMathLibrary.h
+++ b/Engine/Source/Runtime/Engine/Classes/Kismet/KismetMathLibrary.h
@@ -1124,7 +1124,7 @@ class ENGINE_API UKismetMathLibrary : public UBlueprintFunctionLibrary
 	UFUNCTION(BlueprintPure, Category="Math|Vector", meta=(Keywords="rotation rotate"))
 	static FVector GetUpVector(FRotator InRot);
 
-	/** Makes a rotator {Pitch, Yaw, Roll} */
+	/** Makes a rotator {Pitch, Yaw, Roll} from rotation values supplied in degrees */
 	UFUNCTION(BlueprintPure, Category="Math|Rotator", meta=(Keywords="construct build rotation rotate rotator makerotator", NativeMakeFunc))
 	static FRotator MakeRot(float Pitch, float Yaw, float Roll);
 	
@@ -1168,7 +1168,7 @@ class ENGINE_API UKismetMathLibrary : public UBlueprintFunctionLibrary
 	UFUNCTION(BlueprintPure, Category="Math|Rotator", meta=(Keywords="construct build rotation rotate rotator makerotator"))
 	static FRotator MakeRotFromZY(const FVector& Z, const FVector& Y);
 
-	/** Breaks apart a rotator into Pitch, Yaw, Roll */
+	/** Breaks apart a rotator into Pitch, Yaw, Roll angles in degrees*/
 	UFUNCTION(BlueprintPure, Category="Math|Rotator", meta=(Keywords="rotation rotate rotator breakrotator", NativeBreakFunc))
 	static void BreakRot(FRotator InRot, float& Pitch, float& Yaw, float& Roll);
 
