UE-14753 Add ENGINE_API to entire TriggerSphere/Capsule/Box classes to let them be subclassed

diff --git a/Engine/Source/Runtime/Engine/Classes/Engine/TriggerBox.h b/Engine/Source/Runtime/Engine/Classes/Engine/TriggerBox.h
--- a/Engine/Source/Runtime/Engine/Classes/Engine/TriggerBox.h
+++ b/Engine/Source/Runtime/Engine/Classes/Engine/TriggerBox.h
@@ -6,15 +6,15 @@
 #include "TriggerBox.generated.h"
 
 /** A box shaped trigger, used to generate overlap events in the level */
-UCLASS(MinimalAPI)
-class ATriggerBox : public ATriggerBase
+UCLASS()
+class ENGINE_API ATriggerBox : public ATriggerBase
 {
 	GENERATED_UCLASS_BODY()
 
 
 #if WITH_EDITOR
 	// Begin AActor interface.
-	ENGINE_API virtual void EditorApplyScale(const FVector& DeltaScale, const FVector* PivotLocation, bool bAltDown, bool bShiftDown, bool bCtrlDown) override;
+	virtual void EditorApplyScale(const FVector& DeltaScale, const FVector* PivotLocation, bool bAltDown, bool bShiftDown, bool bCtrlDown) override;
 	// End AActor interface.
 #endif
 };
diff --git a/Engine/Source/Runtime/Engine/Classes/Engine/TriggerCapsule.h b/Engine/Source/Runtime/Engine/Classes/Engine/TriggerCapsule.h
--- a/Engine/Source/Runtime/Engine/Classes/Engine/TriggerCapsule.h
+++ b/Engine/Source/Runtime/Engine/Classes/Engine/TriggerCapsule.h
@@ -6,15 +6,15 @@
 #include "TriggerCapsule.generated.h"
 
 /** A capsule shaped trigger, used to generate overlap events in the level */
-UCLASS(MinimalAPI)
-class ATriggerCapsule : public ATriggerBase
+UCLASS()
+class ENGINE_API ATriggerCapsule : public ATriggerBase
 {
 	GENERATED_UCLASS_BODY()
 
 
 #if WITH_EDITOR
 	// Begin AActor interface.
-	ENGINE_API virtual void EditorApplyScale(const FVector& DeltaScale, const FVector* PivotLocation, bool bAltDown, bool bShiftDown, bool bCtrlDown) override;
+	virtual void EditorApplyScale(const FVector& DeltaScale, const FVector* PivotLocation, bool bAltDown, bool bShiftDown, bool bCtrlDown) override;
 	// End AActor interface.
 #endif
 };
diff --git a/Engine/Source/Runtime/Engine/Classes/Engine/TriggerSphere.h b/Engine/Source/Runtime/Engine/Classes/Engine/TriggerSphere.h
--- a/Engine/Source/Runtime/Engine/Classes/Engine/TriggerSphere.h
+++ b/Engine/Source/Runtime/Engine/Classes/Engine/TriggerSphere.h
@@ -6,15 +6,15 @@
 #include "TriggerSphere.generated.h"
 
 /** A sphere shaped trigger, used to generate overlap events in the level */
-UCLASS(MinimalAPI)
-class ATriggerSphere : public ATriggerBase
+UCLASS()
+class ENGINE_API ATriggerSphere : public ATriggerBase
 {
 	GENERATED_UCLASS_BODY()
 
 
 #if WITH_EDITOR
 	// Begin AActor interface.
-	ENGINE_API virtual void EditorApplyScale(const FVector& DeltaScale, const FVector* PivotLocation, bool bAltDown, bool bShiftDown, bool bCtrlDown) override;
+	virtual void EditorApplyScale(const FVector& DeltaScale, const FVector* PivotLocation, bool bAltDown, bool bShiftDown, bool bCtrlDown) override;
 	// End AActor interface.
 #endif
 };
