Refactored code to follow serializing conventions

diff --git a/Engine/Source/Runtime/Engine/Private/PackageMapClient.cpp b/Engine/Source/Runtime/Engine/Private/PackageMapClient.cpp
--- a/Engine/Source/Runtime/Engine/Private/PackageMapClient.cpp
+++ b/Engine/Source/Runtime/Engine/Private/PackageMapClient.cpp
@@ -274,74 +274,46 @@ bool UPackageMapClient::SerializeNewActor(FArchive& Ar, class UActorChannel *Cha
 				bSerializeRotation = !Rotation.Equals(FRotator::ZeroRotator, epsilon);
 				bSerializeScale = !Scale.Equals(DefaultScale, epsilon);
 				bSerializeVelocity = !Velocity.Equals(FVector::ZeroVector, epsilon);
+			}
 
-				Ar.SerializeBits(&bSerializeLocation, 1);
-				if (bSerializeLocation)
-				{
-					Location.NetSerialize(Ar, this, SerSuccess);
-				}
-
-				Ar.SerializeBits(&bSerializeRotation, 1);
-				if (bSerializeRotation)
-				{
-					Rotation.NetSerialize(Ar, this, SerSuccess);
-				}
-
-				Ar.SerializeBits(&bSerializeScale, 1);
-				if (bSerializeScale)
-				{
-					Scale.NetSerialize(Ar, this, SerSuccess);
-				}
-
-				Ar.SerializeBits(&bSerializeVelocity, 1);
-				if (bSerializeVelocity)
-				{
-					Velocity.NetSerialize(Ar, this, SerSuccess);
-				}
+			Ar.SerializeBits(&bSerializeLocation, 1);
+			if (bSerializeLocation)
+			{
+				Location.NetSerialize(Ar, this, SerSuccess);
 			}
-			// Client is serializing from an object sent from the server
-			else if (Ar.IsLoading())
+			else
 			{
+				Location = FVector::ZeroVector;
+			}
 
-				Ar.SerializeBits(&bSerializeLocation, 1);
-				if (bSerializeLocation)
-				{
-					Location.NetSerialize(Ar, this, SerSuccess);
-				}
-				else
-				{
-					Location = FVector::ZeroVector;
-				}
-
-				Ar.SerializeBits(&bSerializeRotation, 1);
-				if (bSerializeRotation)
-				{
-					Rotation.NetSerialize(Ar, this, SerSuccess);
-				}
-				else
-				{
-					Rotation = FRotator::ZeroRotator;
-				}
+			Ar.SerializeBits(&bSerializeRotation, 1);
+			if (bSerializeRotation)
+			{
+				Rotation.NetSerialize(Ar, this, SerSuccess);
+			}
+			else
+			{
+				Rotation = FRotator::ZeroRotator;
+			}
 
-				Ar.SerializeBits(&bSerializeScale, 1);
-				if (bSerializeScale)
-				{
-					Scale.NetSerialize(Ar, this, SerSuccess);
-				}
-				else
-				{
-					Scale = FVector(1, 1, 1);
-				}
+			Ar.SerializeBits(&bSerializeScale, 1);
+			if (bSerializeScale)
+			{
+				Scale.NetSerialize(Ar, this, SerSuccess);
+			}
+			else
+			{
+				Scale = FVector(1, 1, 1);
+			}
 
-				Ar.SerializeBits(&bSerializeVelocity, 1);
-				if (bSerializeVelocity)
-				{
-					Velocity.NetSerialize(Ar, this, SerSuccess);
-				}
-				else
-				{
-					Velocity = FVector::ZeroVector;
-				}
+			Ar.SerializeBits(&bSerializeVelocity, 1);
+			if (bSerializeVelocity)
+			{
+				Velocity.NetSerialize(Ar, this, SerSuccess);
+			}
+			else
+			{
+				Velocity = FVector::ZeroVector;
 			}
 
 			if ( Channel && Ar.IsSaving() )
