Framework - Adding the blueprint function GetSupportedFullscreenResolutions, which gets the list of supported fullscreen resolutions.

diff --git a/Engine/Source/Runtime/Engine/Classes/Kismet/KismetSystemLibrary.h b/Engine/Source/Runtime/Engine/Classes/Kismet/KismetSystemLibrary.h
--- a/Engine/Source/Runtime/Engine/Classes/Kismet/KismetSystemLibrary.h
+++ b/Engine/Source/Runtime/Engine/Classes/Kismet/KismetSystemLibrary.h
@@ -1238,6 +1238,13 @@ class ENGINE_API UKismetSystemLibrary : public UBlueprintFunctionLibrary
 	UFUNCTION(BlueprintPure, Category="Rendering|Material", meta=(UnsafeDuringActorConstruction = "true"))
 	static int32 GetRenderingMaterialQualityLevel();
 
+	/**
+	 * Gets the list of support fullscreen resolutions.
+	 * @return true if successfully queried the device for available resolutions.
+	 */
+	UFUNCTION(BlueprintCallable, Category="Rendering")
+	static bool GetSupportedFullscreenResolutions(TArray<FIntPoint>& Resolutions);
+
 	// Opens the specified URL in the platform's web browser of choice
 	UFUNCTION(BlueprintCallable, Category = "Utilities|Platform")
 	static void LaunchURL(const FString& URL);
diff --git a/Engine/Source/Runtime/Engine/Private/KismetSystemLibrary.cpp b/Engine/Source/Runtime/Engine/Private/KismetSystemLibrary.cpp
--- a/Engine/Source/Runtime/Engine/Private/KismetSystemLibrary.cpp
+++ b/Engine/Source/Runtime/Engine/Private/KismetSystemLibrary.cpp
@@ -2043,7 +2043,7 @@ bool UKismetSystemLibrary::SphereTraceSingleByObject_DEPRECATED(UObject* WorldCo
 		{
 			// no hit means all red
 			::DrawDebugSweptSphere(World, Start, End, Radius, FColor::Red, bPersistent, LifeTime);
-        }
+		}
 	}
 
 	return bHit;
@@ -2872,11 +2872,31 @@ int32 UKismetSystemLibrary::GetRenderingMaterialQualityLevel()
 	return Ret;
 }
 
+bool UKismetSystemLibrary::GetSupportedFullscreenResolutions(TArray<FIntPoint>& Resolutions)
+{
+	FScreenResolutionArray SupportedResolutions;
+	if ( RHIGetAvailableResolutions(SupportedResolutions, true) )
+	{
+		for ( const FScreenResolutionRHI& SupportedResolution : SupportedResolutions )
+		{
+			FIntPoint Resolution;
+			Resolution.X = SupportedResolution.Width;
+			Resolution.Y = SupportedResolution.Height;
+
+			Resolutions.Add(Resolution);
+		}
+
+		return true;
+	}
+
+	return false;
+}
+
 void UKismetSystemLibrary::LaunchURL(const FString& URL)
 {
 	if (!URL.IsEmpty())
 	{
-		FPlatformProcess::LaunchURL(*URL, NULL, NULL);
+		FPlatformProcess::LaunchURL(*URL, nullptr, nullptr);
 	}
 }
 
@@ -2966,4 +2986,3 @@ void UKismetSystemLibrary::SetSupressViewportTransitionMessage(UObject* WorldCon
 		World->GetFirstLocalPlayerFromController()->ViewportClient->SetSuppressTransitionMessage(bState);
 	}
 }
-
