GPU particle collision with distance fields
Global distance field can be accessed via new DistanceToNearestSurface and DistanceFieldGradient material nodes

diff --git a/Engine/Shaders/DistanceFieldLightingShared.usf b/Engine/Shaders/DistanceFieldLightingShared.usf
--- a/Engine/Shaders/DistanceFieldLightingShared.usf
+++ b/Engine/Shaders/DistanceFieldLightingShared.usf
@@ -312,19 +312,3 @@ float GetVPLOcclusion(float3 BentNormalAO, float3 NormalizedVectorToVPL, float V
 
 	return Visibility;
 }
-
-// Must match C++
-#define MAX_GLOBAL_DF_CLIPMAPS 4
-
-Texture3D GlobalDistanceFieldTexture[MAX_GLOBAL_DF_CLIPMAPS];
-SamplerState GlobalDistanceFieldSampler[MAX_GLOBAL_DF_CLIPMAPS];
-float4 GlobalVolumeCenterAndExtent[MAX_GLOBAL_DF_CLIPMAPS];
-float4 GlobalVolumeWorldToUVAddAndMul[MAX_GLOBAL_DF_CLIPMAPS];
-float GlobalVolumeDimension;
-
-float3 ComputeGlobalUV(float3 WorldPosition, uniform uint ClipmapIndex)
-{
-	//return ((WorldPosition - GlobalVolumeCenterAndExtent[ClipmapIndex].xyz + GlobalVolumeScollOffset[ClipmapIndex].xyz) / (GlobalVolumeCenterAndExtent[ClipmapIndex].w * 2) + .5f);
-	float4 WorldToUVAddAndMul = GlobalVolumeWorldToUVAddAndMul[ClipmapIndex];
-	return WorldPosition * WorldToUVAddAndMul.www + WorldToUVAddAndMul.xyz;
-}
diff --git a/Engine/Shaders/DistanceFieldScreenGridLighting.usf b/Engine/Shaders/DistanceFieldScreenGridLighting.usf
--- a/Engine/Shaders/DistanceFieldScreenGridLighting.usf
+++ b/Engine/Shaders/DistanceFieldScreenGridLighting.usf
@@ -8,6 +8,7 @@
 #include "DeferredShadingCommon.usf"
 #include "DistanceFieldLightingShared.usf"
 #include "DistanceFieldAOShared.usf"
+#include "GlobalDistanceFieldShared.usf"
 
 #define NUM_THREADS_PER_RECORD 4
 #define CONE_TRACE_DISPATCH_SIZEX (CULLED_TILE_SIZEX / TRACE_DOWNSAMPLE_FACTOR)
diff --git a/Engine/Shaders/DistanceFieldSurfaceCacheLighting.usf b/Engine/Shaders/DistanceFieldSurfaceCacheLighting.usf
--- a/Engine/Shaders/DistanceFieldSurfaceCacheLighting.usf
+++ b/Engine/Shaders/DistanceFieldSurfaceCacheLighting.usf
@@ -8,6 +8,7 @@
 #include "DeferredShadingCommon.usf"
 #include "DistanceFieldLightingShared.usf"
 #include "DistanceFieldAOShared.usf"
+#include "GlobalDistanceFieldShared.usf"
 
 struct FObjectCullVertexOutput
 {
diff --git a/Engine/Shaders/GlobalDistanceField.usf b/Engine/Shaders/GlobalDistanceField.usf
--- a/Engine/Shaders/GlobalDistanceField.usf
+++ b/Engine/Shaders/GlobalDistanceField.usf
@@ -8,6 +8,7 @@
 #include "DeferredShadingCommon.usf"
 #include "DistanceFieldLightingShared.usf"
 #include "DistanceFieldAOShared.usf"
+#include "GlobalDistanceFieldShared.usf"
 
 RWBuffer<float4> RWCulledObjectBounds;
 RWBuffer<float4> RWCulledObjectData;
@@ -161,11 +162,11 @@ void CullObjectsToGridCS(
 		for (uint ObjectIndex = ThreadIndex; ObjectIndex < NumCulledObjects; ObjectIndex += GRID_CULL_THREADGORUP_TOTALSIZE)
 		{
 			float4 ObjectPositionAndRadius = LoadObjectPositionAndRadius(ObjectIndex);
-			float BoxDistance = ComputeDistanceFromBoxToPoint(TileBoxCenter - TileBoxExtent, TileBoxCenter + TileBoxExtent, ObjectPositionAndRadius.xyz);
-			//float BoxDistanceSq = ComputeSquaredDistanceFromBoxToPoint(TileBoxCenter, TileBoxExtent, ObjectPositionAndRadius.xyz);
+			//float BoxDistance = ComputeDistanceFromBoxToPoint(TileBoxCenter - TileBoxExtent, TileBoxCenter + TileBoxExtent, ObjectPositionAndRadius.xyz);
+			float BoxDistanceSq = ComputeSquaredDistanceFromBoxToPoint(TileBoxCenter, TileBoxExtent, ObjectPositionAndRadius.xyz);
 			float CombinedDistance = ObjectPositionAndRadius.w + AOGlobalMaxSphereQueryRadius;
 
-			if (BoxDistance < CombinedDistance)
+			if (BoxDistanceSq < CombinedDistance * CombinedDistance)
 			{
 				uint DestIndex;
 				// Allocate space in the array for one more object
@@ -235,13 +236,13 @@ void CompositeObjectDistanceFieldsCS(
 		#endif
 
 		float4 ObjectPositionAndRadius = LoadObjectPositionAndRadius(ObjectIndex);
-		float BoxDistance = ComputeDistanceFromBoxToPoint(TileBoxCenter - TileBoxExtent, TileBoxCenter + TileBoxExtent, ObjectPositionAndRadius.xyz);
+		//float BoxDistance = ComputeDistanceFromBoxToPoint(TileBoxCenter - TileBoxExtent, TileBoxCenter + TileBoxExtent, ObjectPositionAndRadius.xyz);
 
 		// ComputeSquaredDistanceFromBoxToPoint is correct but adds a lot to the final cost
-		//float BoxDistanceSq = ComputeSquaredDistanceFromBoxToPoint(TileBoxCenter, TileBoxExtent, ObjectPositionAndRadius.xyz);
+		float BoxDistanceSq = ComputeSquaredDistanceFromBoxToPoint(TileBoxCenter, TileBoxExtent, ObjectPositionAndRadius.xyz);
 		float CombinedDistance = ObjectPositionAndRadius.w + AOGlobalMaxSphereQueryRadius;
 
-		if (BoxDistance < CombinedDistance)
+		if (BoxDistanceSq < CombinedDistance * CombinedDistance)
 		{
 			uint DestIndex;
 			InterlockedAdd(NumTileCulledObjects, 1U, DestIndex);
diff --git a/Engine/Shaders/GlobalDistanceFieldShared.usf b/Engine/Shaders/GlobalDistanceFieldShared.usf
new file mode 100644
--- /dev/null
+++ b/Engine/Shaders/GlobalDistanceFieldShared.usf
@@ -0,0 +1,126 @@
+// Copyright 1998-2015 Epic Games, Inc. All Rights Reserved.
+
+/*=============================================================================
+	GlobalDistanceFieldShared.usf
+=============================================================================*/
+
+// Must match C++
+#define MAX_GLOBAL_DF_CLIPMAPS 4
+
+Texture3D GlobalDistanceFieldTexture[MAX_GLOBAL_DF_CLIPMAPS];
+SamplerState GlobalDistanceFieldSampler[MAX_GLOBAL_DF_CLIPMAPS];
+float4 GlobalVolumeCenterAndExtent[MAX_GLOBAL_DF_CLIPMAPS];
+float4 GlobalVolumeWorldToUVAddAndMul[MAX_GLOBAL_DF_CLIPMAPS];
+float GlobalVolumeDimension;
+float GlobalVolumeTexelSize;
+float MaxGlobalDistance;
+
+float3 ComputeGlobalUV(float3 WorldPosition, uniform uint ClipmapIndex)
+{
+	//return ((WorldPosition - GlobalVolumeCenterAndExtent[ClipmapIndex].xyz + GlobalVolumeScollOffset[ClipmapIndex].xyz) / (GlobalVolumeCenterAndExtent[ClipmapIndex].w * 2) + .5f);
+	float4 WorldToUVAddAndMul = GlobalVolumeWorldToUVAddAndMul[ClipmapIndex];
+	return WorldPosition * WorldToUVAddAndMul.www + WorldToUVAddAndMul.xyz;
+}
+
+float GetDistanceToNearestSurfaceGlobalClipmap(float3 WorldPosition, uniform uint ClipmapIndex)
+{
+	float3 GlobalUV = ComputeGlobalUV(WorldPosition, ClipmapIndex);
+	float DistanceToSurface = Texture3DSampleLevel(GlobalDistanceFieldTexture[ClipmapIndex], GlobalDistanceFieldSampler[ClipmapIndex], GlobalUV, 0).x;
+	return DistanceToSurface;
+}
+
+float GetDistanceToNearestSurfaceGlobal(float3 WorldPosition)
+{
+	float DistanceToSurface = MaxGlobalDistance;
+	float DistanceFromClipmap = ComputeDistanceFromBoxToPointInside(GlobalVolumeCenterAndExtent[0].xyz, GlobalVolumeCenterAndExtent[0].www, WorldPosition);
+
+	//@todo - would be nice to atlas to get rid of branches and multiple samplers but the partial update scrolling relies on wrap addressing
+	BRANCH
+	if (DistanceFromClipmap > GlobalVolumeCenterAndExtent[0].w * GlobalVolumeTexelSize)
+	{
+		DistanceToSurface = GetDistanceToNearestSurfaceGlobalClipmap(WorldPosition, 0);
+	}
+	else
+	{
+		DistanceFromClipmap = ComputeDistanceFromBoxToPointInside(GlobalVolumeCenterAndExtent[1].xyz, GlobalVolumeCenterAndExtent[1].www, WorldPosition);
+
+		BRANCH
+		if (DistanceFromClipmap > GlobalVolumeCenterAndExtent[1].w * GlobalVolumeTexelSize)
+		{
+			DistanceToSurface = GetDistanceToNearestSurfaceGlobalClipmap(WorldPosition, 1);
+		}
+		else
+		{
+			DistanceFromClipmap = ComputeDistanceFromBoxToPointInside(GlobalVolumeCenterAndExtent[2].xyz, GlobalVolumeCenterAndExtent[2].www, WorldPosition);
+			float DistanceFromLastClipmap = ComputeDistanceFromBoxToPointInside(GlobalVolumeCenterAndExtent[3].xyz, GlobalVolumeCenterAndExtent[3].www, WorldPosition);
+
+			BRANCH
+			if (DistanceFromClipmap > GlobalVolumeCenterAndExtent[2].w * GlobalVolumeTexelSize)
+			{
+				DistanceToSurface = GetDistanceToNearestSurfaceGlobalClipmap(WorldPosition, 2);
+			}
+			else if (DistanceFromLastClipmap > GlobalVolumeCenterAndExtent[3].w * GlobalVolumeTexelSize)
+			{
+				DistanceToSurface = GetDistanceToNearestSurfaceGlobalClipmap(WorldPosition, 3);
+			}
+		}
+	}
+
+	return DistanceToSurface;
+}
+
+float3 GetDistanceFieldGradientGlobalClipmap(float3 WorldPosition, uniform uint ClipmapIndex)
+{
+	float3 GlobalUV = ComputeGlobalUV(WorldPosition, ClipmapIndex);
+	float R = Texture3DSampleLevel(GlobalDistanceFieldTexture[ClipmapIndex], GlobalDistanceFieldSampler[ClipmapIndex], float3(GlobalUV.x + GlobalVolumeTexelSize, GlobalUV.y, GlobalUV.z), 0).x;
+	float L = Texture3DSampleLevel(GlobalDistanceFieldTexture[ClipmapIndex], GlobalDistanceFieldSampler[ClipmapIndex], float3(GlobalUV.x - GlobalVolumeTexelSize, GlobalUV.y, GlobalUV.z), 0).x;
+	float F = Texture3DSampleLevel(GlobalDistanceFieldTexture[ClipmapIndex], GlobalDistanceFieldSampler[ClipmapIndex], float3(GlobalUV.x, GlobalUV.y + GlobalVolumeTexelSize, GlobalUV.z), 0).x;
+	float B = Texture3DSampleLevel(GlobalDistanceFieldTexture[ClipmapIndex], GlobalDistanceFieldSampler[ClipmapIndex], float3(GlobalUV.x, GlobalUV.y - GlobalVolumeTexelSize, GlobalUV.z), 0).x;
+	float U = Texture3DSampleLevel(GlobalDistanceFieldTexture[ClipmapIndex], GlobalDistanceFieldSampler[ClipmapIndex], float3(GlobalUV.x, GlobalUV.y, GlobalUV.z + GlobalVolumeTexelSize), 0).x;
+	float D = Texture3DSampleLevel(GlobalDistanceFieldTexture[ClipmapIndex], GlobalDistanceFieldSampler[ClipmapIndex], float3(GlobalUV.x, GlobalUV.y, GlobalUV.z - GlobalVolumeTexelSize), 0).x;
+		
+	float Extent = GlobalVolumeCenterAndExtent[ClipmapIndex].w;
+	float3 Gradient = .5f * float3(R - L, F - B, U - D) / Extent;
+	return Gradient;
+}
+
+float3 GetDistanceFieldGradientGlobal(float3 WorldPosition)
+{
+	float3 Gradient = 0;
+	float DistanceFromClipmap = ComputeDistanceFromBoxToPointInside(GlobalVolumeCenterAndExtent[0].xyz, GlobalVolumeCenterAndExtent[0].www, WorldPosition);
+	// Don't sample near the border to avoid sampling out of bounds during partial differencing
+	float BorderTexels = GlobalVolumeTexelSize * 4;
+
+	BRANCH
+	if (DistanceFromClipmap > GlobalVolumeCenterAndExtent[0].w * BorderTexels)
+	{
+		Gradient = GetDistanceFieldGradientGlobalClipmap(WorldPosition, 0);
+	}
+	else
+	{
+		DistanceFromClipmap = ComputeDistanceFromBoxToPointInside(GlobalVolumeCenterAndExtent[1].xyz, GlobalVolumeCenterAndExtent[1].www, WorldPosition);
+
+		BRANCH
+		if (DistanceFromClipmap > GlobalVolumeCenterAndExtent[1].w * BorderTexels)
+		{
+			Gradient = GetDistanceFieldGradientGlobalClipmap(WorldPosition, 1);
+		}
+		else
+		{
+			DistanceFromClipmap = ComputeDistanceFromBoxToPointInside(GlobalVolumeCenterAndExtent[2].xyz, GlobalVolumeCenterAndExtent[2].www, WorldPosition);
+			float DistanceFromLastClipmap = ComputeDistanceFromBoxToPointInside(GlobalVolumeCenterAndExtent[3].xyz, GlobalVolumeCenterAndExtent[3].www, WorldPosition);
+
+			BRANCH
+			if (DistanceFromClipmap > GlobalVolumeCenterAndExtent[2].w * BorderTexels)
+			{
+				Gradient = GetDistanceFieldGradientGlobalClipmap(WorldPosition, 2);
+			}
+			else if (DistanceFromLastClipmap > GlobalVolumeCenterAndExtent[3].w * BorderTexels)
+			{
+				Gradient = GetDistanceFieldGradientGlobalClipmap(WorldPosition, 3);
+			}
+		}
+	}
+
+	return Gradient;
+}
\ No newline at end of file
diff --git a/Engine/Shaders/MaterialTemplate.usf b/Engine/Shaders/MaterialTemplate.usf
--- a/Engine/Shaders/MaterialTemplate.usf
+++ b/Engine/Shaders/MaterialTemplate.usf
@@ -9,6 +9,7 @@
 
 // for MaterialExpressionDepthOfFieldFunction
 #include "DepthOfFieldCommon.usf"
+#include "GlobalDistanceFieldShared.usf"
 
 #if USES_SPEEDTREE
 	#include "SpeedTreeCommon.usf"
diff --git a/Engine/Shaders/ParticleSimulationShader.usf b/Engine/Shaders/ParticleSimulationShader.usf
--- a/Engine/Shaders/ParticleSimulationShader.usf
+++ b/Engine/Shaders/ParticleSimulationShader.usf
@@ -6,6 +6,10 @@
 
 #include "Common.usf"
 
+#if DISTANCE_FIELD_COLLISION
+	#include "GlobalDistanceFieldShared.usf"
+#endif
+
 /*------------------------------------------------------------------------------
 	Shared declarations and functions.
 ------------------------------------------------------------------------------*/
@@ -69,6 +73,8 @@ float2 TileCoord = Input.TexCoord.xy * float2(TILE_SIZE_X,TILE_SIZE_Y) + TileOff
 #if FORCE_NO_COLLISION
 	#undef DEPTH_BUFFER_COLLISION
 	#define DEPTH_BUFFER_COLLISION 0
+	#undef DISTANCE_FIELD_COLLISION
+	#define DISTANCE_FIELD_COLLISION 0
 #endif
 
 /** Position (XYZ) and squared radius (W) of the point attractor. */
@@ -216,15 +222,15 @@ float3 ComputeAttractionForce(float3 Position)
 	return Attraction * normalize(DirectionToPoint);
 }
 
+/** For retrieving the size of a particle. */
+Texture2D RenderAttributesTexture;
+SamplerState RenderAttributesTextureSampler;
+
 #if DEPTH_BUFFER_COLLISION
 /** For retrieving the world-space normal. */
 Texture2D GBufferATexture;
 SamplerState GBufferATextureSampler;
 
-/** For retrieving the size of a particle. */
-Texture2D RenderAttributesTexture;
-SamplerState RenderAttributesTextureSampler;
-
 /** Limits the depth bounds for which to search for a collision plane. */
 float CollisionDepthBounds;
 
@@ -316,6 +322,82 @@ void CollideWithDepthBuffer(
 }
 #endif // #if DEPTH_BUFFER_COLLISION
 
+#if DISTANCE_FIELD_COLLISION
+
+/**
+ * Compute collision with the global signed distance field
+ */
+void CollideWithDistanceField(
+	out float3 NewPosition,
+	out float3 NewVelocity,
+	inout float RelativeTime,
+	in float3 InPosition,
+	in float3 InVelocity,
+	in float3 Acceleration,
+	in float CollisionRadius,
+	in float Resilience
+	)
+{
+	// Integration assuming no collision.
+	float3 MidVelocity = InVelocity.xyz + 0.5f * Acceleration;
+	float3 DeltaPosition = DeltaSeconds * MidVelocity;
+	NewPosition = InPosition.xyz + DeltaPosition;
+	NewVelocity = InVelocity.xyz + Acceleration;
+
+	float DistanceToNearestSurface = GetDistanceToNearestSurfaceGlobal(InPosition);
+	float MaxCollisionDistance = CollisionRadius + length(DeltaPosition.xyz);
+
+	if (DistanceToNearestSurface < MaxCollisionDistance)
+	{
+		float3 CollisionWorldNormal = normalize(GetDistanceFieldGradientGlobal(InPosition));
+		float3 CollisionWorldPosition = InPosition - CollisionWorldNormal * DistanceToNearestSurface;
+
+		float4 CollisionPlane = float4(CollisionWorldNormal.xyz, dot(CollisionWorldPosition.xyz, CollisionWorldNormal.xyz));
+
+		// Compute the portion of velocity normal to the collision plane.
+		float VelocityDot = dot(CollisionPlane.xyz, DeltaPosition.xyz);
+		float InvVelocityDot = rcp(VelocityDot + 0.0001f); // Add a small amount to avoid division by zero.
+
+		// Distance to the plane from the center of the particle.
+		float DistanceToPlane = dot(CollisionPlane.xyz, InPosition.xyz) - CollisionPlane.w;
+
+		// Find out the time of intersection for both the front and back of the sphere.
+		float t_back = -(DistanceToPlane + CollisionRadius) * InvVelocityDot;
+		float t_front = -(DistanceToPlane - CollisionRadius) * InvVelocityDot;
+
+		//if (t_back >= 0 && t_front <= 1 && DistanceToPlane >= 0)
+		if (step(0, t_back) * step(t_front, 1) * step(0, DistanceToPlane))
+		{
+			// Separate velocity in to the components perpendicular and tangent to the collision plane.
+			float3 PerpVelocity = dot(MidVelocity,CollisionPlane.xyz) * CollisionPlane.xyz;
+			float3 TanVelocity = MidVelocity - PerpVelocity;
+
+			// Compute the new velocity accounting for resilience and friction.
+			NewVelocity = Simulation.OneMinusFriction * TanVelocity - Resilience * PerpVelocity;
+
+			// If the particle lies approximately on the collision plane, don't jump to the point of collision.
+			t_front *= step(VelocityDot,-1);
+
+			// Integrate position taking the collision in to account.
+			NewPosition = InPosition + DeltaPosition * t_front + NewVelocity * (1.0f - t_front) * DeltaSeconds;
+
+			// Update the relative time. Usually this does nothing, but if the
+			// user has elected to kill the particle upon collision this will do
+			// so.
+			RelativeTime += Simulation.CollisionTimeBias;
+		}
+		//else if (t_front > 0 && t_back < 1 && DistanceToPlane < 0)
+		else if (step(0, t_front) * step(t_back, 1) * step(DistanceToPlane,0))
+		{
+			// The particle has collided against a backface, kill it by setting
+			// relative time to a value > 1.0.
+			RelativeTime = 1.1f;
+		}
+	}
+}
+
+#endif
+
 void PixelMain(
 	in FShaderInterpolants Interpolants,
 	out float4 OutPosition : SV_Target0,
@@ -373,7 +455,7 @@ void PixelMain(
 	// Compute the acceleration to apply to the particle this frame.
 	float3 Acceleration = Force * DeltaSeconds;
 
-#if DEPTH_BUFFER_COLLISION
+#if DEPTH_BUFFER_COLLISION || DISTANCE_FIELD_COLLISION
 	// We need to look up render attributes for this particle to figure out how big it is.
 	float4 RenderAttributeSample = Texture2DSampleLevel(RenderAttributesTexture, RenderAttributesTextureSampler, Interpolants.TexCoord.xy, 0);
 	
@@ -389,10 +471,13 @@ void PixelMain(
 
 	// Compute the radius with which to perform collision checks.
 	float CollisionRadius = min(Size.x,Size.y) * Simulation.CollisionRadiusScale + Simulation.CollisionRadiusBias;
+#endif
 
+#if DEPTH_BUFFER_COLLISION
+	
 	// Compute the new position and velocity of the particle by colliding against
 	// the scene's depth buffer.
-	float3 NewPosition,NewVelocity;
+	float3 NewPosition, NewVelocity;
 	CollideWithDepthBuffer(
 		NewPosition.xyz,
 		NewVelocity.xyz,
@@ -403,12 +488,27 @@ void PixelMain(
 		CollisionRadius,
 		Resilience
 		);
-#else // #if DEPTH_BUFFER_COLLISION
+#elif DISTANCE_FIELD_COLLISION
+
+	float3 NewPosition, NewVelocity;
+
+	CollideWithDistanceField(
+		NewPosition.xyz,
+		NewVelocity.xyz,
+		RelativeTime,
+		Position.xyz,
+		Velocity.xyz,
+		Acceleration,
+		CollisionRadius,
+		Resilience
+		);
+
+#else 
 	// Integrate position and velocity forward.
 	float3 DeltaPosition = DeltaSeconds * (Velocity.xyz + 0.5f * Acceleration);
 	float3 NewPosition = Position.xyz + DeltaPosition;
 	float3 NewVelocity = Velocity.xyz + Acceleration;
-#endif // #if DEPTH_BUFFER_COLLISION
+#endif
 
 	
 	// Apply orbit.
diff --git a/Engine/Source/Runtime/Engine/Classes/Engine/EngineTypes.h b/Engine/Source/Runtime/Engine/Classes/Engine/EngineTypes.h
--- a/Engine/Source/Runtime/Engine/Classes/Engine/EngineTypes.h
+++ b/Engine/Source/Runtime/Engine/Classes/Engine/EngineTypes.h
@@ -183,6 +183,15 @@ enum ETrailWidthMode
 	ETrailWidthMode_FromSecond UMETA(DisplayName = "From Second Socket"),
 };
 
+UENUM()
+namespace EParticleCollisionMode
+{
+	enum Type
+	{
+		SceneDepth UMETA(DisplayName="Scene Depth"),
+		DistanceField UMETA(DisplayName="Distance Field")
+	};
+}
 
 // Note: Check UMaterialInstance::Serialize if changed!!
 UENUM()
diff --git a/Engine/Source/Runtime/Engine/Classes/Materials/MaterialExpressionDistanceFieldGradient.h b/Engine/Source/Runtime/Engine/Classes/Materials/MaterialExpressionDistanceFieldGradient.h
new file mode 100644
--- /dev/null
+++ b/Engine/Source/Runtime/Engine/Classes/Materials/MaterialExpressionDistanceFieldGradient.h
@@ -0,0 +1,17 @@
+// Copyright 1998-2015 Epic Games, Inc. All Rights Reserved.
+
+
+#pragma once
+#include "Materials/MaterialExpression.h"
+#include "MaterialExpressionDistanceFieldGradient.generated.h"
+
+UCLASS()
+class UMaterialExpressionDistanceFieldGradient : public UMaterialExpression
+{
+	GENERATED_UCLASS_BODY()
+
+	// Begin UMaterialExpression Interface
+	virtual int32 Compile(class FMaterialCompiler* Compiler, int32 OutputIndex, int32 MultiplexIndex) override;
+	virtual void GetCaption(TArray<FString>& OutCaptions) const override;
+	// End UMaterialExpression Interface
+};
diff --git a/Engine/Source/Runtime/Engine/Classes/Materials/MaterialExpressionDistanceToNearestSurface.h b/Engine/Source/Runtime/Engine/Classes/Materials/MaterialExpressionDistanceToNearestSurface.h
new file mode 100644
--- /dev/null
+++ b/Engine/Source/Runtime/Engine/Classes/Materials/MaterialExpressionDistanceToNearestSurface.h
@@ -0,0 +1,17 @@
+// Copyright 1998-2015 Epic Games, Inc. All Rights Reserved.
+
+
+#pragma once
+#include "Materials/MaterialExpression.h"
+#include "MaterialExpressionDistanceToNearestSurface.generated.h"
+
+UCLASS()
+class UMaterialExpressionDistanceToNearestSurface : public UMaterialExpression
+{
+	GENERATED_UCLASS_BODY()
+
+	// Begin UMaterialExpression Interface
+	virtual int32 Compile(class FMaterialCompiler* Compiler, int32 OutputIndex, int32 MultiplexIndex) override;
+	virtual void GetCaption(TArray<FString>& OutCaptions) const override;
+	// End UMaterialExpression Interface
+};
diff --git a/Engine/Source/Runtime/Engine/Classes/Materials/MaterialInterface.h b/Engine/Source/Runtime/Engine/Classes/Materials/MaterialInterface.h
--- a/Engine/Source/Runtime/Engine/Classes/Materials/MaterialInterface.h
+++ b/Engine/Source/Runtime/Engine/Classes/Materials/MaterialInterface.h
@@ -57,6 +57,9 @@ struct ENGINE_API FMaterialRelevance
 	UPROPERTY()
 	uint32 bOutputsVelocityInBasePass : 1;
 
+	UPROPERTY()
+	uint32 bUsesGlobalDistanceField : 1;
+
 	UPROPERTY()
 	uint16 ShadingModelMask;
 
@@ -69,6 +72,7 @@ struct ENGINE_API FMaterialRelevance
 		, bNormalTranslucency(false)
 		, bDisableDepthTest(false)		
 		, bOutputsVelocityInBasePass(true)
+		, bUsesGlobalDistanceField(false)
 		, ShadingModelMask(0)
 	{}
 
@@ -83,6 +87,7 @@ struct ENGINE_API FMaterialRelevance
 		bDisableDepthTest |= B.bDisableDepthTest;
 		ShadingModelMask |= B.ShadingModelMask;
 		bOutputsVelocityInBasePass |= B.bOutputsVelocityInBasePass;
+		bUsesGlobalDistanceField |= B.bUsesGlobalDistanceField;
 		return *this;
 	}
 
diff --git a/Engine/Source/Runtime/Engine/Classes/Particles/Collision/ParticleModuleCollisionGPU.h b/Engine/Source/Runtime/Engine/Classes/Particles/Collision/ParticleModuleCollisionGPU.h
--- a/Engine/Source/Runtime/Engine/Classes/Particles/Collision/ParticleModuleCollisionGPU.h
+++ b/Engine/Source/Runtime/Engine/Classes/Particles/Collision/ParticleModuleCollisionGPU.h
@@ -21,7 +21,7 @@ namespace EParticleCollisionResponse
 	};
 }
 
-UCLASS(editinlinenew, hidecategories=Object, MinimalAPI, meta=(DisplayName = "Collision (Scene Depth)"))
+UCLASS(editinlinenew, hidecategories=Object, MinimalAPI, meta=(DisplayName = "Collision"))
 class UParticleModuleCollisionGPU : public UParticleModuleCollisionBase
 {
 	GENERATED_UCLASS_BODY()
@@ -64,6 +64,9 @@ class UParticleModuleCollisionGPU : public UParticleModuleCollisionBase
 	UPROPERTY(EditAnywhere, Category=Collision)
 	TEnumAsByte<EParticleCollisionResponse::Type> Response;
 
+	UPROPERTY(EditAnywhere, Category=Collision)
+	TEnumAsByte<EParticleCollisionMode::Type> CollisionMode;
+
 	/** Initializes the default values for this property */
 	void InitializeDefaults();
 
diff --git a/Engine/Source/Runtime/Engine/Classes/Particles/TypeData/ParticleModuleTypeDataGpu.h b/Engine/Source/Runtime/Engine/Classes/Particles/TypeData/ParticleModuleTypeDataGpu.h
--- a/Engine/Source/Runtime/Engine/Classes/Particles/TypeData/ParticleModuleTypeDataGpu.h
+++ b/Engine/Source/Runtime/Engine/Classes/Particles/TypeData/ParticleModuleTypeDataGpu.h
@@ -169,6 +169,9 @@ struct FGPUSpriteEmitterInfo
 	UPROPERTY()
 	uint32 bEnableCollision : 1;
 
+	UPROPERTY()
+	TEnumAsByte<EParticleCollisionMode::Type> CollisionMode;
+
 	/** Dynamic color scale from the ColorOverLife module. */
 	UPROPERTY()
 	FRawDistributionVector DynamicColor;
@@ -201,6 +204,7 @@ struct FGPUSpriteEmitterInfo
 		, ScreenAlignment(0)
 		, LockAxisFlag(0)
 		, bEnableCollision(false)
+		, CollisionMode(EParticleCollisionMode::SceneDepth)
 	{
 	}
 
diff --git a/Engine/Source/Runtime/Engine/Private/Materials/HLSLMaterialTranslator.h b/Engine/Source/Runtime/Engine/Private/Materials/HLSLMaterialTranslator.h
--- a/Engine/Source/Runtime/Engine/Private/Materials/HLSLMaterialTranslator.h
+++ b/Engine/Source/Runtime/Engine/Private/Materials/HLSLMaterialTranslator.h
@@ -262,10 +262,6 @@ public:
 			Material->CompileErrors.Empty();
 			Material->ErrorExpressions.Empty();
 
-			MaterialCompilationOutput.bRequiresSceneColorCopy = false;
-			check(!MaterialCompilationOutput.bNeedsSceneTextures);
-			MaterialCompilationOutput.bUsesEyeAdaptation = false;
-
 			bCompileForComputeShader = Material->IsLightFunction();
 
 			// Generate code
@@ -3641,6 +3637,30 @@ protected:
 		return AddCodeChunk( MCT_Float3, TEXT("MaterialExpressionBlackBody(%s)"), *GetParameterCode(Temp) );
 	}
 
+	virtual int32 DistanceToNearestSurface() override
+	{
+		if (ErrorUnlessFeatureLevelSupported(ERHIFeatureLevel::SM5) == INDEX_NONE)
+		{
+			return INDEX_NONE;
+		}
+
+		MaterialCompilationOutput.bUsesGlobalDistanceField = true;
+
+		return AddCodeChunk(MCT_Float, TEXT("GetDistanceToNearestSurfaceGlobal(Parameters.WorldPosition)"));
+	}
+
+	virtual int32 DistanceFieldGradient() override
+	{
+		if (ErrorUnlessFeatureLevelSupported(ERHIFeatureLevel::SM5) == INDEX_NONE)
+		{
+			return INDEX_NONE;
+		}
+
+		MaterialCompilationOutput.bUsesGlobalDistanceField = true;
+
+		return AddCodeChunk(MCT_Float3, TEXT("GetDistanceFieldGradientGlobal(Parameters.WorldPosition)"));
+	}
+
 	virtual int32 AtmosphericFogColor( int32 WorldPosition ) override
 	{
 		if (ErrorUnlessFeatureLevelSupported(ERHIFeatureLevel::SM4) == INDEX_NONE)
diff --git a/Engine/Source/Runtime/Engine/Private/Materials/MaterialExpressions.cpp b/Engine/Source/Runtime/Engine/Private/Materials/MaterialExpressions.cpp
--- a/Engine/Source/Runtime/Engine/Private/Materials/MaterialExpressions.cpp
+++ b/Engine/Source/Runtime/Engine/Private/Materials/MaterialExpressions.cpp
@@ -128,6 +128,8 @@
 #include "Materials/MaterialExpressionVertexNormalWS.h"
 #include "Materials/MaterialExpressionViewSize.h"
 #include "Materials/MaterialExpressionWorldPosition.h"
+#include "Materials/MaterialExpressionDistanceToNearestSurface.h"
+#include "Materials/MaterialExpressionDistanceFieldGradient.h"
 #include "Materials/MaterialFunction.h"
 #include "Materials/MaterialParameterCollection.h"
 
@@ -8250,6 +8252,67 @@ void UMaterialExpressionBlackBody::GetCaption(TArray<FString>& OutCaptions) cons
 }
 
 
+///////////////////////////////////////////////////////////////////////////////
+// UMaterialExpressionDistanceToNearestSurface
+///////////////////////////////////////////////////////////////////////////////
+UMaterialExpressionDistanceToNearestSurface::UMaterialExpressionDistanceToNearestSurface(const FObjectInitializer& ObjectInitializer)
+	: Super(ObjectInitializer)
+{
+	// Structure to hold one-time initialization
+	struct FConstructorStatics
+	{
+		FText NAME_Utility;
+		FConstructorStatics()
+			: NAME_Utility(LOCTEXT( "Utility", "Utility" ))
+		{
+		}
+	};
+	static FConstructorStatics ConstructorStatics;
+
+	MenuCategories.Add(ConstructorStatics.NAME_Utility);
+}
+
+int32 UMaterialExpressionDistanceToNearestSurface::Compile(class FMaterialCompiler* Compiler, int32 OutputIndex, int32 MultiplexIndex)
+{
+	return Compiler->DistanceToNearestSurface();
+}
+
+void UMaterialExpressionDistanceToNearestSurface::GetCaption(TArray<FString>& OutCaptions) const
+{
+	OutCaptions.Add(TEXT("DistanceToNearestSurface"));
+}
+
+
+///////////////////////////////////////////////////////////////////////////////
+// UMaterialExpressionDistanceFieldGradient
+///////////////////////////////////////////////////////////////////////////////
+UMaterialExpressionDistanceFieldGradient::UMaterialExpressionDistanceFieldGradient(const FObjectInitializer& ObjectInitializer)
+	: Super(ObjectInitializer)
+{
+	// Structure to hold one-time initialization
+	struct FConstructorStatics
+	{
+		FText NAME_Utility;
+		FConstructorStatics()
+			: NAME_Utility(LOCTEXT( "Utility", "Utility" ))
+		{
+		}
+	};
+	static FConstructorStatics ConstructorStatics;
+
+	MenuCategories.Add(ConstructorStatics.NAME_Utility);
+}
+
+int32 UMaterialExpressionDistanceFieldGradient::Compile(class FMaterialCompiler* Compiler, int32 OutputIndex, int32 MultiplexIndex)
+{
+	return Compiler->DistanceFieldGradient();
+}
+
+void UMaterialExpressionDistanceFieldGradient::GetCaption(TArray<FString>& OutCaptions) const
+{
+	OutCaptions.Add(TEXT("DistanceFieldGradient"));
+}
+
 ///////////////////////////////////////////////////////////////////////////////
 // UMaterialExpressionDistance
 ///////////////////////////////////////////////////////////////////////////////
diff --git a/Engine/Source/Runtime/Engine/Private/Materials/MaterialInterface.cpp b/Engine/Source/Runtime/Engine/Private/Materials/MaterialInterface.cpp
--- a/Engine/Source/Runtime/Engine/Private/Materials/MaterialInterface.cpp
+++ b/Engine/Source/Runtime/Engine/Private/Materials/MaterialInterface.cpp
@@ -22,6 +22,7 @@ void FMaterialRelevance::SetPrimitiveViewRelevance(FPrimitiveViewRelevance& OutV
 	OutViewRelevance.bSeparateTranslucencyRelevance = bSeparateTranslucency;
 	OutViewRelevance.bNormalTranslucencyRelevance = bNormalTranslucency;
 	OutViewRelevance.ShadingModelMaskRelevance = ShadingModelMask;
+	OutViewRelevance.bUsesGlobalDistanceField = bUsesGlobalDistanceField;
 }
 
 //////////////////////////////////////////////////////////////////////////
@@ -64,7 +65,8 @@ FMaterialRelevance UMaterialInterface::GetRelevance_Internal(const UMaterial* Ma
 		MaterialRelevance.bSeparateTranslucency = bIsTranslucent && Material->bEnableSeparateTranslucency;
 		MaterialRelevance.bNormalTranslucency = bIsTranslucent && !Material->bEnableSeparateTranslucency;
 		MaterialRelevance.bDisableDepthTest = bIsTranslucent && Material->bDisableDepthTest;		
-		MaterialRelevance.bOutputsVelocityInBasePass = Material->bOutputVelocityOnBasePass;		
+		MaterialRelevance.bOutputsVelocityInBasePass = Material->bOutputVelocityOnBasePass;	
+		MaterialRelevance.bUsesGlobalDistanceField = MaterialResource->GetGameThreadShaderMap() ? MaterialResource->GetGameThreadShaderMap()->UsesGlobalDistanceField() : false;
 		return MaterialRelevance;
 	}
 	else
diff --git a/Engine/Source/Runtime/Engine/Private/Materials/MaterialShared.cpp b/Engine/Source/Runtime/Engine/Private/Materials/MaterialShared.cpp
--- a/Engine/Source/Runtime/Engine/Private/Materials/MaterialShared.cpp
+++ b/Engine/Source/Runtime/Engine/Private/Materials/MaterialShared.cpp
@@ -248,15 +248,12 @@ void FMaterialCompilationOutput::Serialize(FArchive& Ar)
 {
 	UniformExpressionSet.Serialize(Ar);
 
-	if (Ar.UE4Ver() >= VER_UE4_PURGED_FMATERIAL_COMPILE_OUTPUTS)
-	{
-		Ar << bRequiresSceneColorCopy;
-	}
-
+	Ar << bRequiresSceneColorCopy;
 	Ar << bNeedsSceneTextures;
 	Ar << bUsesEyeAdaptation;
 	Ar << bModifiesMeshPosition;
 	Ar << bNeedsGBuffer;
+	Ar << bUsesGlobalDistanceField;
 }
 
 void FMaterial::GetShaderMapId(EShaderPlatform Platform, FMaterialShaderMapId& OutId) const
diff --git a/Engine/Source/Runtime/Engine/Private/Particles/FXSystem.cpp b/Engine/Source/Runtime/Engine/Private/Particles/FXSystem.cpp
--- a/Engine/Source/Runtime/Engine/Private/Particles/FXSystem.cpp
+++ b/Engine/Source/Runtime/Engine/Private/Particles/FXSystem.cpp
@@ -318,11 +318,22 @@ void FFXSystem::PreInitViews()
 	}
 }
 
-void FFXSystem::PreRender(FRHICommandListImmediate& RHICmdList)
+bool FFXSystem::UsesGlobalDistanceField() const
 {
 	if (RHISupportsGPUParticles(FeatureLevel))
 	{
-		SimulateGPUParticles(RHICmdList, EParticleSimulatePhase::Main, NULL, FTexture2DRHIParamRef(), FTexture2DRHIParamRef());
+		return UsesGlobalDistanceFieldInternal();
+	}
+
+	return false;
+}
+
+void FFXSystem::PreRender(FRHICommandListImmediate& RHICmdList, const FGlobalDistanceFieldParameterData* GlobalDistanceFieldParameterData)
+{
+	if (RHISupportsGPUParticles(FeatureLevel))
+	{
+		SimulateGPUParticles(RHICmdList, EParticleSimulatePhase::Main, NULL, NULL, FTexture2DRHIParamRef(), FTexture2DRHIParamRef());
+		SimulateGPUParticles(RHICmdList, EParticleSimulatePhase::CollisionDistanceField, NULL, GlobalDistanceFieldParameterData, FTexture2DRHIParamRef(), FTexture2DRHIParamRef());
 	}
 }
 
@@ -330,7 +341,7 @@ void FFXSystem::PostRenderOpaque(FRHICommandListImmediate& RHICmdList, const cla
 {
 	if (RHISupportsGPUParticles(FeatureLevel))
 	{
-		SimulateGPUParticles(RHICmdList, EParticleSimulatePhase::Collision, CollisionView, SceneDepthTexture, GBufferATexture);
+		SimulateGPUParticles(RHICmdList, EParticleSimulatePhase::CollisionDepthBuffer, CollisionView, NULL, SceneDepthTexture, GBufferATexture);
 		SortGPUParticles(RHICmdList);
 	}
 }
diff --git a/Engine/Source/Runtime/Engine/Private/Particles/FXSystemPrivate.h b/Engine/Source/Runtime/Engine/Private/Particles/FXSystemPrivate.h
--- a/Engine/Source/Runtime/Engine/Private/Particles/FXSystemPrivate.h
+++ b/Engine/Source/Runtime/Engine/Private/Particles/FXSystemPrivate.h
@@ -24,15 +24,16 @@ namespace EParticleSimulatePhase
 	{
 		/** The main simulation pass is for standard particles. */
 		Main,
+		CollisionDistanceField,
 		/** The collision pass is used by these that collide against the scene depth buffer. */
-		Collision,
+		CollisionDepthBuffer,
 
 		/**********************************************************************/
 
 		/** The first simulation phase that is run each frame. */
 		First = Main,
 		/** The final simulation phase that is run each frame. */
-		Last = Collision
+		Last = CollisionDepthBuffer
 	};
 };
 
@@ -65,7 +66,8 @@ public:
 	virtual void UpdateVectorField(UVectorFieldComponent* VectorFieldComponent) override;
 	virtual FParticleEmitterInstance* CreateGPUSpriteEmitterInstance(FGPUSpriteEmitterInfo& EmitterInfo) override;
 	virtual void PreInitViews() override;
-	virtual void PreRender(FRHICommandListImmediate& RHICmdList) override;
+	virtual bool UsesGlobalDistanceField() const override;
+	virtual void PreRender(FRHICommandListImmediate& RHICmdList, const FGlobalDistanceFieldParameterData* GlobalDistanceFieldParameterData) override;
 	virtual void PostRenderOpaque(FRHICommandListImmediate& RHICmdList, const class FSceneView* CollisionView, FTexture2DRHIParamRef SceneDepthTexture, FTexture2DRHIParamRef GBufferATexture) override;
 	// End FFXSystemInterface.
 
@@ -147,6 +149,8 @@ private:
 	 */
 	void SortGPUParticles(FRHICommandListImmediate& RHICmdList);
 
+	bool UsesGlobalDistanceFieldInternal() const;
+
 	/**
 	 * Update particles simulated on the GPU.
 	 * @param Phase				Which emitters are being simulated.
@@ -158,6 +162,7 @@ private:
 		FRHICommandListImmediate& RHICmdList,
 		EParticleSimulatePhase::Type Phase,
 		const class FSceneView* CollisionView,
+		const FGlobalDistanceFieldParameterData* GlobalDistanceFieldParameterData,
 		FTexture2DRHIParamRef SceneDepthTexture,
 		FTexture2DRHIParamRef GBufferATexture
 		);
diff --git a/Engine/Source/Runtime/Engine/Private/Particles/ParticleGpuSimulation.cpp b/Engine/Source/Runtime/Engine/Private/Particles/ParticleGpuSimulation.cpp
--- a/Engine/Source/Runtime/Engine/Private/Particles/ParticleGpuSimulation.cpp
+++ b/Engine/Source/Runtime/Engine/Private/Particles/ParticleGpuSimulation.cpp
@@ -30,6 +30,7 @@
 #include "VectorField/VectorField.h"
 #include "SceneUtils.h"
 #include "MeshBatch.h"
+#include "GlobalDistanceFieldParameters.h"
 
 /*------------------------------------------------------------------------------
 	Constants to tune memory and performance for GPU particle simulation.
@@ -831,10 +832,17 @@ private:
 	FShaderResourceParameter TileOffsets;
 };
 
+enum EParticleCollisionShaderMode
+{
+	PCM_None,
+	PCM_DepthBuffer,
+	PCM_DistanceField
+};
+
 /**
  * Pixel shader for simulating particles on the GPU.
  */
-template <bool bUseDepthBufferCollision>
+template <EParticleCollisionShaderMode CollisionMode>
 class TParticleSimulationPS : public FGlobalShader
 {
 	DECLARE_SHADER_TYPE(TParticleSimulationPS,Global);
@@ -851,7 +859,8 @@ public:
 		FGlobalShader::ModifyCompilationEnvironment(Platform, OutEnvironment);
 		OutEnvironment.SetDefine(TEXT("PARTICLE_SIMULATION_PIXELSHADER"), 1);
 		OutEnvironment.SetDefine(TEXT("MAX_VECTOR_FIELDS"), MAX_VECTOR_FIELDS);
-		OutEnvironment.SetDefine(TEXT("DEPTH_BUFFER_COLLISION"), (uint32)(bUseDepthBufferCollision ? 1 : 0));
+		OutEnvironment.SetDefine(TEXT("DEPTH_BUFFER_COLLISION"), (uint32)(CollisionMode == PCM_DepthBuffer ? 1 : 0));
+		OutEnvironment.SetDefine(TEXT("DISTANCE_FIELD_COLLISION"), (uint32)(CollisionMode == PCM_DistanceField ? 1 : 0));
 		OutEnvironment.SetRenderTargetOutputFormat(0, PF_A32B32G32R32F);
 	}
 
@@ -888,6 +897,7 @@ public:
 		GBufferATextureParameterSampler.Bind(Initializer.ParameterMap,TEXT("GBufferATextureSampler"));
 		CollisionDepthBounds.Bind(Initializer.ParameterMap,TEXT("CollisionDepthBounds"));
 		PerFrameParameters.Bind(Initializer.ParameterMap);
+		GlobalDistanceFieldParameters.Bind(Initializer.ParameterMap);
 	}
 
 	/** Serialization. */
@@ -918,6 +928,7 @@ public:
 		Ar << GBufferATextureParameterSampler;
 		Ar << CollisionDepthBounds;
 		Ar << PerFrameParameters;
+		Ar << GlobalDistanceFieldParameters;
 		return bShaderHasOutdatedParameters;
 	}
 
@@ -930,6 +941,7 @@ public:
 		const FParticleAttributesTexture& InAttributesTexture,
 		const FParticleAttributesTexture& InRenderAttributesTexture,
 		const FSceneView* CollisionView,
+		const FGlobalDistanceFieldParameterData* GlobalDistanceFieldParameterData,
 		FTexture2DRHIParamRef SceneDepthTexture,
 		FTexture2DRHIParamRef GBufferATexture
 		)
@@ -941,7 +953,8 @@ public:
 		SetTextureParameter(RHICmdList, PixelShaderRHI, VelocityTexture, VelocityTextureSampler, SamplerStatePoint, TextureResources.VelocityTextureRHI);
 		SetTextureParameter(RHICmdList, PixelShaderRHI, AttributesTexture, AttributesTextureSampler, SamplerStatePoint, InAttributesTexture.TextureRHI);
 		SetTextureParameter(RHICmdList, PixelShaderRHI, CurveTexture, CurveTextureSampler, SamplerStateLinear, GParticleCurveTexture.GetCurveTexture());
-		if (bUseDepthBufferCollision)
+
+		if (CollisionMode == PCM_DepthBuffer)
 		{
 			check(CollisionView != NULL);
 			FGlobalShader::SetParameters(RHICmdList, PixelShaderRHI,*CollisionView);
@@ -971,6 +984,19 @@ public:
 				);
 			SetShaderValue(RHICmdList, PixelShaderRHI, CollisionDepthBounds, FXConsoleVariables::GPUCollisionDepthBounds);
 		}
+		else if (CollisionMode == PCM_DistanceField)
+		{
+			GlobalDistanceFieldParameters.Set(RHICmdList, PixelShaderRHI, *GlobalDistanceFieldParameterData);
+
+			SetTextureParameter(
+				RHICmdList, 
+				PixelShaderRHI,
+				RenderAttributesTexture,
+				RenderAttributesTextureSampler,
+				SamplerStatePoint,
+				InRenderAttributesTexture.TextureRHI
+				);
+		}
 	}
 
 	/**
@@ -1059,6 +1085,7 @@ private:
 	FParticlePerFrameSimulationShaderParameters PerFrameParameters;
 	/** Collision depth bounds. */
 	FShaderParameter CollisionDepthBounds;
+	FGlobalDistanceFieldParameters GlobalDistanceFieldParameters;
 };
 
 /**
@@ -1103,8 +1130,9 @@ public:
 
 /** Implementation for all shaders used for simulation. */
 IMPLEMENT_SHADER_TYPE(,FParticleTileVS,TEXT("ParticleSimulationShader"),TEXT("VertexMain"),SF_Vertex);
-IMPLEMENT_SHADER_TYPE(template<>,TParticleSimulationPS<false>,TEXT("ParticleSimulationShader"),TEXT("PixelMain"),SF_Pixel);
-IMPLEMENT_SHADER_TYPE(template<>,TParticleSimulationPS<true>,TEXT("ParticleSimulationShader"),TEXT("PixelMain"),SF_Pixel);
+IMPLEMENT_SHADER_TYPE(template<>,TParticleSimulationPS<PCM_None>,TEXT("ParticleSimulationShader"),TEXT("PixelMain"),SF_Pixel);
+IMPLEMENT_SHADER_TYPE(template<>,TParticleSimulationPS<PCM_DepthBuffer>,TEXT("ParticleSimulationShader"),TEXT("PixelMain"),SF_Pixel);
+IMPLEMENT_SHADER_TYPE(template<>,TParticleSimulationPS<PCM_DistanceField>,TEXT("ParticleSimulationShader"),TEXT("PixelMain"),SF_Pixel);
 IMPLEMENT_SHADER_TYPE(,FParticleSimulationClearPS,TEXT("ParticleSimulationShader"),TEXT("PixelMain"),SF_Pixel);
 
 /**
@@ -1248,7 +1276,7 @@ struct FSimulationCommandGPU
  * @param CollisionView		The view to use for collision, if any.
  * @param SceneDepthTexture The depth texture to use for collision, if any.
  */
-template <bool bUseDepthBufferCollision>
+template <EParticleCollisionShaderMode CollisionMode>
 void ExecuteSimulationCommands(
 	FRHICommandList& RHICmdList,
 	ERHIFeatureLevel::Type FeatureLevel,
@@ -1257,13 +1285,14 @@ void ExecuteSimulationCommands(
 	const FParticleAttributesTexture& AttributeTexture,
 	const FParticleAttributesTexture& RenderAttributeTexture,
 	const FSceneView* CollisionView,
+	const FGlobalDistanceFieldParameterData* GlobalDistanceFieldParameterData,
 	FTexture2DRHIParamRef SceneDepthTexture,
 	FTexture2DRHIParamRef GBufferATexture
 	)
 {
 	// Grab shaders.
 	TShaderMapRef<FParticleTileVS> VertexShader(GetGlobalShaderMap(FeatureLevel));
-	TShaderMapRef<TParticleSimulationPS<bUseDepthBufferCollision> > PixelShader(GetGlobalShaderMap(FeatureLevel));
+	TShaderMapRef<TParticleSimulationPS<CollisionMode> > PixelShader(GetGlobalShaderMap(FeatureLevel));
 
 	// Bound shader state.
 	
@@ -1278,7 +1307,7 @@ void ExecuteSimulationCommands(
 		0
 		);
 
-	PixelShader->SetParameters(RHICmdList, TextureResources, AttributeTexture, RenderAttributeTexture, CollisionView, SceneDepthTexture, GBufferATexture);
+	PixelShader->SetParameters(RHICmdList, TextureResources, AttributeTexture, RenderAttributeTexture, CollisionView, GlobalDistanceFieldParameterData, SceneDepthTexture, GBufferATexture);
 
 	// Draw tiles to perform the simulation step.
 	const int32 CommandCount = SimulationCommands.Num();
@@ -2293,6 +2322,8 @@ public:
 	/** True if the simulation wants collision enabled. */
 	bool bWantsCollision;
 
+	EParticleCollisionMode::Type CollisionMode;
+
 	/** Flag that specifies the simulation's resources are dirty and need to be updated. */
 	bool bDirty_GameThread;
 	bool bReleased_GameThread;
@@ -2305,6 +2336,7 @@ public:
 		, SimulationIndex(INDEX_NONE)
 		, SimulationPhase(EParticleSimulatePhase::Main)
 		, bWantsCollision(false)
+		, CollisionMode(EParticleCollisionMode::SceneDepth)
 		, bDirty_GameThread(true)
 		, bReleased_GameThread(true)
 		, bDestroyed_GameThread(false)
@@ -2580,9 +2612,13 @@ struct FGPUSpriteDynamicEmitterData : FDynamicEmitterDataBase
 		// If the simulation wants to collide against the depth buffer
 		// and we're not rendering with an opaque material put the 
 		// simulation in the collision phase.
-		if (bTranslucent && Simulation->bWantsCollision)
+		if (bTranslucent && Simulation->bWantsCollision && Simulation->CollisionMode == EParticleCollisionMode::SceneDepth)
 		{
-			Simulation->SimulationPhase = EParticleSimulatePhase::Collision;
+			Simulation->SimulationPhase = EParticleSimulatePhase::CollisionDepthBuffer;
+		}
+		else if (Simulation->bWantsCollision && Simulation->CollisionMode == EParticleCollisionMode::DistanceField)
+		{
+			Simulation->SimulationPhase = EParticleSimulatePhase::CollisionDistanceField;
 		}
 	}
 
@@ -2813,6 +2849,7 @@ public:
 			EmitterInfo.LocalVectorField.Field->InitInstance(&Simulation->LocalVectorField, /*bPreviewInstance=*/ false);
 		}
 		Simulation->bWantsCollision = InEmitterInfo.bEnableCollision;
+		Simulation->CollisionMode = InEmitterInfo.CollisionMode;
 
 #if TRACK_TILE_ALLOCATIONS
 		TSet<class FGPUSpriteParticleEmitterInstance*>* EmitterSet = GPUSpriteParticleEmitterInstances.Find(FXSystem);
@@ -4027,10 +4064,27 @@ static void SetParametersForVectorField(FVectorFieldUniformParameters& OutParame
 	OutParameters.TilingAxes[Index].Z = VectorFieldInstance->bTileZ ? 1.0f : 0.0f;
 }
 
+bool FFXSystem::UsesGlobalDistanceFieldInternal() const
+{
+	for (TSparseArray<FParticleSimulationGPU*>::TConstIterator It(GPUSimulations); It; ++It)
+	{
+		const FParticleSimulationGPU* Simulation = *It;
+
+		if (Simulation->SimulationPhase == EParticleSimulatePhase::CollisionDistanceField
+			&& Simulation->TileVertexBuffer.AlignedTileCount > 0)
+		{
+			return true;
+		}
+	}
+
+	return false;
+}
+
 void FFXSystem::SimulateGPUParticles(
 	FRHICommandListImmediate& RHICmdList,
 	EParticleSimulatePhase::Type Phase,
 	const class FSceneView* CollisionView,
+	const FGlobalDistanceFieldParameterData* GlobalDistanceFieldParameterData,
 	FTexture2DRHIParamRef SceneDepthTexture,
 	FTexture2DRHIParamRef GBufferATexture
 	)
@@ -4184,10 +4238,26 @@ void FFXSystem::SimulateGPUParticles(
 	{
 		SCOPED_DRAW_EVENT(RHICmdList, ParticleSimulation);
 
-		if (Phase == EParticleSimulatePhase::Collision && CollisionView)
+		if (Phase == EParticleSimulatePhase::CollisionDepthBuffer && CollisionView)
+		{
+			/// ?
+			ExecuteSimulationCommands<PCM_DepthBuffer>(
+				RHICmdList,
+				FeatureLevel,
+				SimulationCommands,
+				PrevStateTextures,
+				ParticleSimulationResources->SimulationAttributesTexture,
+				ParticleSimulationResources->RenderAttributesTexture,
+				CollisionView,
+				GlobalDistanceFieldParameterData,
+				SceneDepthTexture,
+				GBufferATexture
+				);
+		}
+		else if (Phase == EParticleSimulatePhase::CollisionDistanceField && GlobalDistanceFieldParameterData)
 		{
 			/// ?
-			ExecuteSimulationCommands<true>(
+			ExecuteSimulationCommands<PCM_DistanceField>(
 				RHICmdList,
 				FeatureLevel,
 				SimulationCommands,
@@ -4195,6 +4265,7 @@ void FFXSystem::SimulateGPUParticles(
 				ParticleSimulationResources->SimulationAttributesTexture,
 				ParticleSimulationResources->RenderAttributesTexture,
 				CollisionView,
+				GlobalDistanceFieldParameterData,
 				SceneDepthTexture,
 				GBufferATexture
 				);
@@ -4202,7 +4273,7 @@ void FFXSystem::SimulateGPUParticles(
 		else
 		{
 			/// ?
-			ExecuteSimulationCommands<false>(
+			ExecuteSimulationCommands<PCM_None>(
 				RHICmdList,
 				FeatureLevel,
 				SimulationCommands,
@@ -4210,6 +4281,7 @@ void FFXSystem::SimulateGPUParticles(
 				ParticleSimulationResources->SimulationAttributesTexture,
 				ParticleSimulationResources->RenderAttributesTexture,
 				NULL,
+				GlobalDistanceFieldParameterData,
 				FTexture2DRHIParamRef(),
 				FTexture2DRHIParamRef()
 				);
diff --git a/Engine/Source/Runtime/Engine/Private/Particles/ParticleModules.cpp b/Engine/Source/Runtime/Engine/Private/Particles/ParticleModules.cpp
--- a/Engine/Source/Runtime/Engine/Private/Particles/ParticleModules.cpp
+++ b/Engine/Source/Runtime/Engine/Private/Particles/ParticleModules.cpp
@@ -4341,6 +4341,7 @@ void UParticleModuleTypeDataGpu::Build( FParticleEmitterBuildInfo& EmitterBuildI
 
 	// Collision flag.
 	EmitterInfo.bEnableCollision = EmitterBuildInfo.bEnableCollision;
+	EmitterInfo.CollisionMode = (EParticleCollisionMode::Type)EmitterBuildInfo.CollisionMode;
 
 	// Create or update GPU resources.
 	if ( EmitterInfo.Resources )
diff --git a/Engine/Source/Runtime/Engine/Private/Particles/ParticleModules_Collision.cpp b/Engine/Source/Runtime/Engine/Private/Particles/ParticleModules_Collision.cpp
--- a/Engine/Source/Runtime/Engine/Private/Particles/ParticleModules_Collision.cpp
+++ b/Engine/Source/Runtime/Engine/Private/Particles/ParticleModules_Collision.cpp
@@ -637,6 +637,7 @@ void UParticleModuleCollisionGPU::SetToSensibleDefaults(UParticleEmitter* Owner)
 void UParticleModuleCollisionGPU::CompileModule(struct FParticleEmitterBuildInfo& EmitterInfo)
 {
 	EmitterInfo.bEnableCollision = true;
+	EmitterInfo.CollisionMode = CollisionMode;
 	EmitterInfo.CollisionResponse = Response;
 	EmitterInfo.CollisionRadiusScale = RadiusScale;
 	EmitterInfo.CollisionRadiusBias = RadiusBias;
@@ -666,12 +667,23 @@ bool UParticleModuleCollisionGPU::IsValidForLODLevel(UParticleLODLevel* LODLevel
 		BlendMode = MaterialResource->GetBlendMode();
 	}
 
-	if (BlendMode == BLEND_Opaque || BlendMode == BLEND_Masked)
+	if (CollisionMode == EParticleCollisionMode::SceneDepth && (BlendMode == BLEND_Opaque || BlendMode == BLEND_Masked))
 	{
 		OutErrorString = NSLOCTEXT("UnrealEd", "CollisionOnOpaqueEmitter", "Scene depth collision cannot be used on emitters with an opaque material.").ToString();
 		return false;
 	}
 
+	if (CollisionMode == EParticleCollisionMode::DistanceField)
+	{
+		static const auto CVar = IConsoleManager::Get().FindTConsoleVariableDataInt(TEXT("r.GenerateMeshDistanceFields"));
+
+		if (CVar->GetValueOnGameThread() == 0)
+		{
+			OutErrorString = NSLOCTEXT("UnrealEd", "CollisionWithoutDistanceField", "Distance Field collision requires the 'Generate Mesh Distance Fields' Renderer project setting to be enabled.").ToString();
+			return false;
+		}
+	}
+
 	if (LODLevel->TypeDataModule && LODLevel->TypeDataModule->IsA(UParticleModuleTypeDataGpu::StaticClass()))
 	{
 		if(!IsDistributionAllowedOnGPU(ResilienceScaleOverLife.Distribution))
diff --git a/Engine/Source/Runtime/Engine/Private/ShaderDerivedDataVersion.h b/Engine/Source/Runtime/Engine/Private/ShaderDerivedDataVersion.h
--- a/Engine/Source/Runtime/Engine/Private/ShaderDerivedDataVersion.h
+++ b/Engine/Source/Runtime/Engine/Private/ShaderDerivedDataVersion.h
@@ -10,4 +10,4 @@
 // guid as version
 
 #define GLOBALSHADERMAP_DERIVEDDATA_VER			TEXT("af1606fcec214373844d2b6cf9df71c8")
-#define MATERIALSHADERMAP_DERIVEDDATA_VER		TEXT("983afa3849534faa8491540b994bd32e")
+#define MATERIALSHADERMAP_DERIVEDDATA_VER		TEXT("984afa3849534faa8491540b994bd32e")
diff --git a/Engine/Source/Runtime/Engine/Public/FXSystem.h b/Engine/Source/Runtime/Engine/Public/FXSystem.h
--- a/Engine/Source/Runtime/Engine/Public/FXSystem.h
+++ b/Engine/Source/Runtime/Engine/Public/FXSystem.h
@@ -161,11 +161,13 @@ public:
 	 */
 	virtual void PreInitViews() = 0;
 
+	virtual bool UsesGlobalDistanceField() const = 0;
+
 	/**
 	 * Notification from the renderer that it is about to draw FX belonging to
 	 * this system.
 	 */
-	virtual void PreRender(FRHICommandListImmediate& RHICmdList) = 0;
+	virtual void PreRender(FRHICommandListImmediate& RHICmdList, const class FGlobalDistanceFieldParameterData* GlobalDistanceFieldParameterData) = 0;
 
 	/**
 	 * Notification from the renderer that opaque primitives have rendered.
diff --git a/Engine/Source/Runtime/Engine/Public/MaterialCompiler.h b/Engine/Source/Runtime/Engine/Public/MaterialCompiler.h
--- a/Engine/Source/Runtime/Engine/Public/MaterialCompiler.h
+++ b/Engine/Source/Runtime/Engine/Public/MaterialCompiler.h
@@ -173,6 +173,8 @@ public:
 	virtual int32 AntialiasedTextureMask(int32 Tex, int32 UV, float Threshold, uint8 Channel) = 0;
 	virtual int32 Noise(int32 Position, float Scale, int32 Quality, uint8 NoiseFunction, bool bTurbulence, int32 Levels, float OutputMin, float OutputMax, float LevelScale, int32 FilterWidth) = 0;
 	virtual int32 BlackBody( int32 Temp ) = 0;
+	virtual int32 DistanceToNearestSurface() = 0;
+	virtual int32 DistanceFieldGradient() = 0;
 	virtual int32 DepthOfFieldFunction(int32 Depth, int32 FunctionValueIndex) = 0;
 	virtual int32 AtmosphericFogColor(int32 WorldPosition) = 0;
 	virtual int32 SpeedTree(ESpeedTreeGeometryType GeometryType, ESpeedTreeWindType WindType, ESpeedTreeLODType LODType, float BillboardThreshold, bool bAccurateWindVelocities) = 0;
@@ -334,6 +336,8 @@ public:
 		return Compiler->Noise(Position, Scale, Quality, NoiseFunction, bTurbulence, Levels, OutputMin, OutputMax, LevelScale, FilterWidth);
 	}
 	virtual int32 BlackBody( int32 Temp ) override { return Compiler->BlackBody(Temp); }
+	virtual int32 DistanceToNearestSurface() override { return Compiler->DistanceToNearestSurface(); }
+	virtual int32 DistanceFieldGradient() override { return Compiler->DistanceFieldGradient(); }
 	virtual int32 PerInstanceRandom() override { return Compiler->PerInstanceRandom(); }
 	virtual int32 PerInstanceFadeAmount() override { return Compiler->PerInstanceFadeAmount(); }
 	virtual int32 DepthOfFieldFunction(int32 Depth, int32 FunctionValueIndex) override
diff --git a/Engine/Source/Runtime/Engine/Public/MaterialShared.h b/Engine/Source/Runtime/Engine/Public/MaterialShared.h
--- a/Engine/Source/Runtime/Engine/Public/MaterialShared.h
+++ b/Engine/Source/Runtime/Engine/Public/MaterialShared.h
@@ -280,7 +280,8 @@ public:
 		bNeedsSceneTextures(false),
 		bUsesEyeAdaptation(false),
 		bModifiesMeshPosition(false),
-		bNeedsGBuffer(false)
+		bNeedsGBuffer(false),
+		bUsesGlobalDistanceField(false)
 	{}
 
 	ENGINE_API void Serialize(FArchive& Ar);
@@ -303,6 +304,9 @@ public:
 
 	/** true if the material uses any GBuffer textures */
 	bool bNeedsGBuffer;
+
+	/** true if material uses the global distance field */
+	bool bUsesGlobalDistanceField;
 };
 
 
@@ -656,6 +660,7 @@ public:
 	const FString& GetDebugDescription() const { return DebugDescription; }
 	bool RequiresSceneColorCopy() const { return MaterialCompilationOutput.bRequiresSceneColorCopy; }
 	bool NeedsSceneTextures() const { return MaterialCompilationOutput.bNeedsSceneTextures; }
+	bool UsesGlobalDistanceField() const { return MaterialCompilationOutput.bUsesGlobalDistanceField; }
 	bool NeedsGBuffer() const { return MaterialCompilationOutput.bNeedsGBuffer; }
 	bool UsesEyeAdaptation() const { return MaterialCompilationOutput.bUsesEyeAdaptation; }
 	bool ModifiesMeshPosition() const { return MaterialCompilationOutput.bModifiesMeshPosition; }
diff --git a/Engine/Source/Runtime/Engine/Public/ParticleEmitterInstances.h b/Engine/Source/Runtime/Engine/Public/ParticleEmitterInstances.h
--- a/Engine/Source/Runtime/Engine/Public/ParticleEmitterInstances.h
+++ b/Engine/Source/Runtime/Engine/Public/ParticleEmitterInstances.h
@@ -103,8 +103,8 @@ struct FParticleEmitterBuildInfo
 	/** Enable collision? */
 	bool bEnableCollision;
 	/** How particles respond to collision. */
-	//EParticleCollisionResponse::Type CollisionResponse;
 	uint8 CollisionResponse;
+	uint8 CollisionMode;
 	/** Radius scale applied to friction. */
 	float CollisionRadiusScale;
 	/** Bias applied to the collision radius. */
diff --git a/Engine/Source/Runtime/Engine/Public/PrimitiveViewRelevance.h b/Engine/Source/Runtime/Engine/Public/PrimitiveViewRelevance.h
--- a/Engine/Source/Runtime/Engine/Public/PrimitiveViewRelevance.h
+++ b/Engine/Source/Runtime/Engine/Public/PrimitiveViewRelevance.h
@@ -41,6 +41,7 @@ struct FPrimitiveViewRelevance
 	uint32 bNormalTranslucencyRelevance : 1;
 	// The primitive has one or more elements that have World Position Offset.
 	uint32 bHasWorldPositionOffset : 1;
+	uint32 bUsesGlobalDistanceField : 1;
 
 	/** 
 	 * Whether this primitive view relevance has been initialized this frame.  
@@ -72,6 +73,7 @@ struct FPrimitiveViewRelevance
 		bSeparateTranslucencyRelevance(false),
 		bNormalTranslucencyRelevance(false),		
 		bHasWorldPositionOffset(false),
+		bUsesGlobalDistanceField(false),
 		bInitializedThisFrame(false)
 	{}
 
@@ -95,6 +97,7 @@ struct FPrimitiveViewRelevance
 		bNormalTranslucencyRelevance |= B.bNormalTranslucencyRelevance != 0;
 		bInitializedThisFrame |= B.bInitializedThisFrame;		
 		bHasWorldPositionOffset |= B.bHasWorldPositionOffset != 0;
+		bUsesGlobalDistanceField |= B.bUsesGlobalDistanceField;
 		return *this;
 	}
 
diff --git a/Engine/Source/Runtime/Renderer/Private/DeferredShadingRenderer.cpp b/Engine/Source/Runtime/Renderer/Private/DeferredShadingRenderer.cpp
--- a/Engine/Source/Runtime/Renderer/Private/DeferredShadingRenderer.cpp
+++ b/Engine/Source/Runtime/Renderer/Private/DeferredShadingRenderer.cpp
@@ -715,7 +715,7 @@ void FDeferredShadingSceneRenderer::Render(FRHICommandListImmediate& RHICmdList)
 		FRHICommandListExecutor::GetImmediateCommandList().ImmediateFlush(EImmediateFlushType::FlushRHIThreadFlushResources);
 	}
 
-	if (ShouldPrepareForDistanceFieldAO() || ShouldPrepareForDistanceFieldShadows())
+	if (ShouldPrepareDistanceFields())
 	{
 		QUICK_SCOPE_CYCLE_COUNTER(STAT_FDeferredShadingSceneRenderer_DistanceFieldAO_Init);
 		GDistanceFieldVolumeTextureAtlas.UpdateAllocations();
@@ -815,7 +815,7 @@ void FDeferredShadingSceneRenderer::Render(FRHICommandListImmediate& RHICmdList)
 	// Notify the FX system that the scene is about to be rendered.
 	if (Scene->FXSystem)
 	{
-		Scene->FXSystem->PreRender(RHICmdList);
+		Scene->FXSystem->PreRender(RHICmdList, &Views[0].GlobalDistanceFieldInfo.ParameterData);
 	}
 
 	GRenderTargetPool.AddPhaseEvent(TEXT("EarlyZPass"));
diff --git a/Engine/Source/Runtime/Renderer/Private/DeferredShadingRenderer.h b/Engine/Source/Runtime/Renderer/Private/DeferredShadingRenderer.h
--- a/Engine/Source/Runtime/Renderer/Private/DeferredShadingRenderer.h
+++ b/Engine/Source/Runtime/Renderer/Private/DeferredShadingRenderer.h
@@ -361,6 +361,7 @@ private:
 	/** Whether distance field global data structures should be prepared for features that use it. */
 	bool ShouldPrepareForDistanceFieldShadows() const;
 	bool ShouldPrepareForDistanceFieldAO() const;
+	bool ShouldPrepareDistanceFields() const;
 
 	void UpdateGlobalDistanceFieldObjectBuffers(FRHICommandListImmediate& RHICmdList);
 
diff --git a/Engine/Source/Runtime/Renderer/Private/DistanceFieldScreenGridLighting.cpp b/Engine/Source/Runtime/Renderer/Private/DistanceFieldScreenGridLighting.cpp
--- a/Engine/Source/Runtime/Renderer/Private/DistanceFieldScreenGridLighting.cpp
+++ b/Engine/Source/Runtime/Renderer/Private/DistanceFieldScreenGridLighting.cpp
@@ -135,7 +135,7 @@ public:
 
 		if (bUseGlobalDistanceField)
 		{
-			GlobalDistanceFieldParameters.Set(RHICmdList, ShaderRHI, GlobalDistanceFieldInfo);
+			GlobalDistanceFieldParameters.Set(RHICmdList, ShaderRHI, GlobalDistanceFieldInfo.ParameterData);
 		}
 
 		FAOSampleData2 AOSampleData;
@@ -297,7 +297,7 @@ public:
 		ObjectParameters.Set(RHICmdList, ShaderRHI, GAOCulledObjectBuffers.Buffers);
 		AOParameters.Set(RHICmdList, ShaderRHI, Parameters);
 		ScreenGridParameters.Set(RHICmdList, ShaderRHI, View, DistanceFieldNormal);
-		GlobalDistanceFieldParameters.Set(RHICmdList, ShaderRHI, GlobalDistanceFieldInfo);
+		GlobalDistanceFieldParameters.Set(RHICmdList, ShaderRHI, GlobalDistanceFieldInfo.ParameterData);
 
 		FAOSampleData2 AOSampleData;
 
diff --git a/Engine/Source/Runtime/Renderer/Private/DistanceFieldSurfaceCacheLighting.cpp b/Engine/Source/Runtime/Renderer/Private/DistanceFieldSurfaceCacheLighting.cpp
--- a/Engine/Source/Runtime/Renderer/Private/DistanceFieldSurfaceCacheLighting.cpp
+++ b/Engine/Source/Runtime/Renderer/Private/DistanceFieldSurfaceCacheLighting.cpp
@@ -21,6 +21,7 @@
 #include "BasePassRendering.h"
 #include "HeightfieldLighting.h"
 #include "GlobalDistanceField.h"
+#include "FXSystem.h"
 
 int32 GDistanceFieldAO = 1;
 FAutoConsoleVariableRef CVarDistanceFieldAO(
@@ -2637,6 +2638,15 @@ bool FDeferredShadingSceneRenderer::ShouldPrepareForDistanceFieldAO() const
 			|| ViewFamily.EngineShowFlags.VisualizeDistanceFieldGI);
 }
 
+bool FDeferredShadingSceneRenderer::ShouldPrepareDistanceFields() const
+{
+	return SupportsDistanceFieldAO(Scene->GetFeatureLevel(), Scene->GetShaderPlatform())
+		&& (ShouldPrepareForDistanceFieldAO() 
+			|| ShouldPrepareForDistanceFieldShadows() 
+			|| Views[0].bUsesGlobalDistanceField
+			|| Scene->FXSystem->UsesGlobalDistanceField());
+}
+
 void RenderDistanceFieldAOSurfaceCache(
 	FRHICommandListImmediate& RHICmdList, 
 	const FViewInfo& View,
@@ -3304,7 +3314,7 @@ public:
 
 		if (bUseGlobalDistanceField)
 		{
-			GlobalDistanceFieldParameters.Set(RHICmdList, ShaderRHI, GlobalDistanceFieldInfo);
+			GlobalDistanceFieldParameters.Set(RHICmdList, ShaderRHI, GlobalDistanceFieldInfo.ParameterData);
 		}
 
 		SetShaderValue(RHICmdList, ShaderRHI, NumGroups, NumGroupsValue);
diff --git a/Engine/Source/Runtime/Renderer/Private/ForwardShadingRenderer.cpp b/Engine/Source/Runtime/Renderer/Private/ForwardShadingRenderer.cpp
--- a/Engine/Source/Runtime/Renderer/Private/ForwardShadingRenderer.cpp
+++ b/Engine/Source/Runtime/Renderer/Private/ForwardShadingRenderer.cpp
@@ -100,7 +100,7 @@ void FForwardShadingSceneRenderer::Render(FRHICommandListImmediate& RHICmdList)
 	// Notify the FX system that the scene is about to be rendered.
 	if (Scene->FXSystem)
 	{
-		Scene->FXSystem->PreRender(RHICmdList);
+		Scene->FXSystem->PreRender(RHICmdList, NULL);
 	}
 
 	GRenderTargetPool.VisualizeTexture.OnStartFrame(Views[0]);
diff --git a/Engine/Source/Runtime/Renderer/Private/GlobalDistanceField.cpp b/Engine/Source/Runtime/Renderer/Private/GlobalDistanceField.cpp
--- a/Engine/Source/Runtime/Renderer/Private/GlobalDistanceField.cpp
+++ b/Engine/Source/Runtime/Renderer/Private/GlobalDistanceField.cpp
@@ -82,6 +82,45 @@ FAutoConsoleVariableRef CVarAOGlobalDFStartDistance(
 	ECVF_Cheat | ECVF_RenderThreadSafe
 	);
 
+void FGlobalDistanceFieldInfo::UpdateParameterData(float MaxOcclusionDistance)
+{
+	if (Clipmaps.Num() > 0)
+	{
+		for (int32 ClipmapIndex = 0; ClipmapIndex < GMaxGlobalDistanceFieldClipmaps; ClipmapIndex++)
+		{
+			FTextureRHIParamRef TextureValue = ClipmapIndex < Clipmaps.Num() 
+				? Clipmaps[ClipmapIndex].RenderTarget->GetRenderTargetItem().ShaderResourceTexture 
+				: NULL;
+
+			ParameterData.Textures[ClipmapIndex] = TextureValue;
+
+			if (ClipmapIndex < Clipmaps.Num())
+			{
+				const FGlobalDistanceFieldClipmap& Clipmap = Clipmaps[ClipmapIndex];
+				ParameterData.CenterAndExtent[ClipmapIndex] = FVector4(Clipmap.Bounds.GetCenter(), Clipmap.Bounds.GetExtent().X);
+
+				// GlobalUV = (WorldPosition - GlobalVolumeCenterAndExtent[ClipmapIndex].xyz + GlobalVolumeScollOffset[ClipmapIndex].xyz) / (GlobalVolumeCenterAndExtent[ClipmapIndex].w * 2) + .5f;
+				// WorldToUVMul = 1.0f / (GlobalVolumeCenterAndExtent[ClipmapIndex].w * 2)
+				// WorldToUVAdd = (GlobalVolumeScollOffset[ClipmapIndex].xyz - GlobalVolumeCenterAndExtent[ClipmapIndex].xyz) / (GlobalVolumeCenterAndExtent[ClipmapIndex].w * 2) + .5f
+				const FVector WorldToUVAdd = (Clipmap.ScrollOffset - Clipmap.Bounds.GetCenter()) / (Clipmap.Bounds.GetExtent().X * 2) + FVector(.5f);
+				ParameterData.WorldToUVAddAndMul[ClipmapIndex] = FVector4(WorldToUVAdd, 1.0f / (Clipmap.Bounds.GetExtent().X * 2));
+			}
+			else
+			{
+				ParameterData.CenterAndExtent[ClipmapIndex] = FVector4(0, 0, 0, 0);
+				ParameterData.WorldToUVAddAndMul[ClipmapIndex] = FVector4(0, 0, 0, 0);
+			}
+		}
+
+		ParameterData.GlobalDFResolution = GAOGlobalDFResolution;
+		ParameterData.MaxDistance = MaxOcclusionDistance;
+	}
+	else
+	{
+		FPlatformMemory::Memzero(&ParameterData, sizeof(ParameterData));
+	}
+}
+
 TGlobalResource<FDistanceFieldObjectBufferResource> GGlobalDistanceFieldCulledObjectBuffers;
 
 uint32 CullObjectsGroupSize = 64;
@@ -253,7 +292,7 @@ public:
 		FComputeShaderRHIParamRef ShaderRHI = GetComputeShader();
 		FGlobalShader::SetParameters(RHICmdList, ShaderRHI, View);
 		CulledObjectBufferParameters.Set(RHICmdList, ShaderRHI, GGlobalDistanceFieldCulledObjectBuffers.Buffers);
-		GlobalDistanceFieldParameters.Set(RHICmdList, ShaderRHI, GlobalDistanceFieldInfo);
+		GlobalDistanceFieldParameters.Set(RHICmdList, ShaderRHI, GlobalDistanceFieldInfo.ParameterData);
 
 		CulledObjectGrid.SetBuffer(RHICmdList, ShaderRHI, GObjectGridBuffers.CulledObjectGrid);
 
@@ -356,7 +395,7 @@ public:
 		FComputeShaderRHIParamRef ShaderRHI = GetComputeShader();
 		FGlobalShader::SetParameters(RHICmdList, ShaderRHI, View);
 		CulledObjectBufferParameters.Set(RHICmdList, ShaderRHI, GGlobalDistanceFieldCulledObjectBuffers.Buffers);
-		GlobalDistanceFieldParameters.Set(RHICmdList, ShaderRHI, GlobalDistanceFieldInfo);
+		GlobalDistanceFieldParameters.Set(RHICmdList, ShaderRHI, GlobalDistanceFieldInfo.ParameterData);
 
 		GlobalDistanceFieldTexture.SetTexture(RHICmdList, ShaderRHI, GlobalDistanceFieldInfo.Clipmaps[ClipmapIndexValue].RenderTarget->GetRenderTargetItem().ShaderResourceTexture, GlobalDistanceFieldInfo.Clipmaps[ClipmapIndexValue].RenderTarget->GetRenderTargetItem().UAV);
 
@@ -530,7 +569,7 @@ float ComputeClipmapExtent(int32 ClipmapIndex)
 	return GAOInnerGlobalDFClipmapDistance * FMath::Pow(GAOGlobalDFClipmapDistanceExponent, ClipmapIndex);
 }
 
-void ComputeUpdateRegionsAndUpdateViewState(const FViewInfo& View, FGlobalDistanceFieldInfo& GlobalDistanceFieldInfo, int32 NumClipmaps, const TArray<FVector4>& PrimitiveModifiedBounds)
+void ComputeUpdateRegionsAndUpdateViewState(const FViewInfo& View, FGlobalDistanceFieldInfo& GlobalDistanceFieldInfo, int32 NumClipmaps, const TArray<FVector4>& PrimitiveModifiedBounds, float MaxOcclusionDistance)
 {
 	GlobalDistanceFieldInfo.Clipmaps.AddZeroed(NumClipmaps);
 
@@ -695,6 +734,8 @@ void ComputeUpdateRegionsAndUpdateViewState(const FViewInfo& View, FGlobalDistan
 			Clipmap.UpdateRegions.Add(UpdateRegion);
 		}
 	}
+
+	GlobalDistanceFieldInfo.UpdateParameterData(MaxOcclusionDistance);
 }
 
 /** 
@@ -711,7 +752,7 @@ void UpdateGlobalDistanceFieldVolume(
 {
 	if (Scene->DistanceFieldSceneData.NumObjectsInBuffer > 0)
 	{
-		ComputeUpdateRegionsAndUpdateViewState(View, GlobalDistanceFieldInfo, GMaxGlobalDistanceFieldClipmaps, Scene->DistanceFieldSceneData.PrimitiveModifiedBounds);
+		ComputeUpdateRegionsAndUpdateViewState(View, GlobalDistanceFieldInfo, GMaxGlobalDistanceFieldClipmaps, Scene->DistanceFieldSceneData.PrimitiveModifiedBounds, MaxOcclusionDistance);
 
 		bool bHasUpdateRegions = false;
 
diff --git a/Engine/Source/Runtime/Renderer/Private/GlobalDistanceField.h b/Engine/Source/Runtime/Renderer/Private/GlobalDistanceField.h
--- a/Engine/Source/Runtime/Renderer/Private/GlobalDistanceField.h
+++ b/Engine/Source/Runtime/Renderer/Private/GlobalDistanceField.h
@@ -6,88 +6,7 @@
 
 #pragma once
 
-class FGlobalDistanceFieldParameters
-{
-public:
-	void Bind(const FShaderParameterMap& ParameterMap)
-	{
-		GlobalDistanceFieldTexture.Bind(ParameterMap,TEXT("GlobalDistanceFieldTexture"));
-		GlobalDistanceFieldSampler.Bind(ParameterMap,TEXT("GlobalDistanceFieldSampler"));
-		GlobalVolumeCenterAndExtent.Bind(ParameterMap,TEXT("GlobalVolumeCenterAndExtent"));
-		GlobalVolumeDimension.Bind(ParameterMap,TEXT("GlobalVolumeDimension"));
-		GlobalVolumeWorldToUVAddAndMul.Bind(ParameterMap,TEXT("GlobalVolumeWorldToUVAddAndMul"));
-	}
-
-	friend FArchive& operator<<(FArchive& Ar,FGlobalDistanceFieldParameters& Parameters)
-	{
-		Ar << Parameters.GlobalDistanceFieldTexture;
-		Ar << Parameters.GlobalDistanceFieldSampler;
-		Ar << Parameters.GlobalVolumeCenterAndExtent;
-		Ar << Parameters.GlobalVolumeDimension;
-		Ar << Parameters.GlobalVolumeWorldToUVAddAndMul;
-		return Ar;
-	}
-
-	template<typename ShaderRHIParamRef>
-	void Set(FRHICommandList& RHICmdList, const ShaderRHIParamRef ShaderRHI, const FGlobalDistanceFieldInfo& Info)
-	{
-		for (int32 ClipmapIndex = 0; ClipmapIndex < (int32)GlobalDistanceFieldTexture.GetNumResources(); ClipmapIndex++)
-		{
-			FTextureRHIParamRef TextureValue = ClipmapIndex < Info.Clipmaps.Num() 
-				? Info.Clipmaps[ClipmapIndex].RenderTarget->GetRenderTargetItem().ShaderResourceTexture 
-				: NULL;
-
-			SetTextureParameter(
-				RHICmdList,
-				ShaderRHI,
-				GlobalDistanceFieldTexture,
-				GlobalDistanceFieldSampler,
-				TStaticSamplerState<SF_Bilinear,AM_Wrap,AM_Wrap,AM_Wrap>::GetRHI(),
-				TextureValue,
-				ClipmapIndex
-				);
-		}
-
-		FVector4 CenterAndExtent[GMaxGlobalDistanceFieldClipmaps];
-		FVector4 WorldToUVAddAndMul[GMaxGlobalDistanceFieldClipmaps];
-
-		const int32 NumToSet = GlobalVolumeCenterAndExtent.GetNumBytes() / sizeof(FVector4);
-
-		for (int32 ClipmapIndex = 0; ClipmapIndex < NumToSet; ClipmapIndex++)
-		{
-			if (ClipmapIndex < Info.Clipmaps.Num())
-			{
-				const FGlobalDistanceFieldClipmap& Clipmap = Info.Clipmaps[ClipmapIndex];
-				CenterAndExtent[ClipmapIndex] = FVector4(Clipmap.Bounds.GetCenter(), Clipmap.Bounds.GetExtent().X);
-
-				// GlobalUV = (WorldPosition - GlobalVolumeCenterAndExtent[ClipmapIndex].xyz + GlobalVolumeScollOffset[ClipmapIndex].xyz) / (GlobalVolumeCenterAndExtent[ClipmapIndex].w * 2) + .5f;
-				// WorldToUVMul = 1.0f / (GlobalVolumeCenterAndExtent[ClipmapIndex].w * 2)
-				// WorldToUVAdd = (GlobalVolumeScollOffset[ClipmapIndex].xyz - GlobalVolumeCenterAndExtent[ClipmapIndex].xyz) / (GlobalVolumeCenterAndExtent[ClipmapIndex].w * 2) + .5f
-				const FVector WorldToUVAdd = (Clipmap.ScrollOffset - Clipmap.Bounds.GetCenter()) / (Clipmap.Bounds.GetExtent().X * 2) + FVector(.5f);
-				WorldToUVAddAndMul[ClipmapIndex] = FVector4(WorldToUVAdd, 1.0f / (Clipmap.Bounds.GetExtent().X * 2));
-			}
-			else
-			{
-				CenterAndExtent[ClipmapIndex] = FVector4(0);
-				WorldToUVAddAndMul[ClipmapIndex] = FVector4(0);
-			}
-		}
-
-		SetShaderValueArray(RHICmdList, ShaderRHI, GlobalVolumeCenterAndExtent, CenterAndExtent, GMaxGlobalDistanceFieldClipmaps);
-		SetShaderValueArray(RHICmdList, ShaderRHI, GlobalVolumeWorldToUVAddAndMul, WorldToUVAddAndMul, GMaxGlobalDistanceFieldClipmaps);
-
-		extern int32 GAOGlobalDFResolution;
-		SetShaderValue(RHICmdList, ShaderRHI, GlobalVolumeDimension, (float)GAOGlobalDFResolution);
-	}
-
-private:
-	FShaderResourceParameter GlobalDistanceFieldTexture;
-	FShaderResourceParameter GlobalDistanceFieldSampler;
-	FShaderParameter GlobalVolumeCenterAndExtent;
-	FShaderParameter GlobalVolumeDimension;
-	FShaderParameter GlobalVolumeScollOffset;
-	FShaderParameter GlobalVolumeWorldToUVAddAndMul;
-};
+#include "GlobalDistanceFieldParameters.h"
 
 extern int32 GAOGlobalDistanceField;
 extern int32 GAOVisualizeGlobalDistanceField;
diff --git a/Engine/Source/Runtime/Renderer/Private/ScenePrivate.h b/Engine/Source/Runtime/Renderer/Private/ScenePrivate.h
--- a/Engine/Source/Runtime/Renderer/Private/ScenePrivate.h
+++ b/Engine/Source/Runtime/Renderer/Private/ScenePrivate.h
@@ -69,6 +69,7 @@ bool IsMobileHDR32bpp();
 #include "ScopedPointer.h"
 #include "ClearQuad.h"
 #include "AtmosphereRendering.h"
+#include "GlobalDistanceFieldParameters.h"
 
 #if WITH_SLI || PLATFORM_SHOULD_BUFFER_QUERIES
 #define BUFFERED_OCCLUSION_QUERIES 1
@@ -398,9 +399,6 @@ private:
 /** Random table for occlusion **/
 extern FOcclusionRandomStream GOcclusionRandomStream;
 
-/** Must match global distance field shaders. */
-const int32 GMaxGlobalDistanceFieldClipmaps = 4;
-
 /**
  * The scene manager's private implementation of persistent view state.
  * This class is associated with a particular camera across multiple frames by the game thread.
diff --git a/Engine/Source/Runtime/Renderer/Private/SceneRendering.h b/Engine/Source/Runtime/Renderer/Private/SceneRendering.h
--- a/Engine/Source/Runtime/Renderer/Private/SceneRendering.h
+++ b/Engine/Source/Runtime/Renderer/Private/SceneRendering.h
@@ -10,6 +10,7 @@
 #include "DistortionRendering.h"
 #include "CustomDepthRendering.h"
 #include "HeightfieldLighting.h"
+#include "GlobalDistanceFieldParameters.h"
 
 // Forward declarations.
 class FPostprocessContext;
@@ -377,6 +378,9 @@ class FGlobalDistanceFieldInfo
 public:
 
 	TArray<FGlobalDistanceFieldClipmap> Clipmaps;
+	FGlobalDistanceFieldParameterData ParameterData;
+
+	void UpdateParameterData(float MaxOcclusionDistance);
 };
 
 /** A FSceneView with additional state used by the scene renderer. */
@@ -499,6 +503,8 @@ public:
 	uint32 bDisableQuerySubmissions : 1;
 	/** Whether we should disable distance-based fade transitions for this frame (usually after a large camera movement.) */
 	uint32 bDisableDistanceBasedFadeTransitions : 1;
+	/** Whether the view has any materials that use the global distance field. */
+	uint32 bUsesGlobalDistanceField : 1;
 	/** Bitmask of all shading models used by primitives in this view */
 	uint16 ShadingModelMaskInView;
 
diff --git a/Engine/Source/Runtime/Renderer/Private/SceneVisibility.cpp b/Engine/Source/Runtime/Renderer/Private/SceneVisibility.cpp
--- a/Engine/Source/Runtime/Renderer/Private/SceneVisibility.cpp
+++ b/Engine/Source/Runtime/Renderer/Private/SceneVisibility.cpp
@@ -803,6 +803,7 @@ struct FRelevancePacket
 	FRelevancePrimSet<FPrimitiveSceneInfo*> UpdateStaticMeshes;
 	FRelevancePrimSet<FPrimitiveSceneInfo*> VisibleEditorPrimitives;
 	uint16 CombinedShadingModelMask;
+	bool bUsesGlobalDistanceField;
 
 	FRelevancePacket(
 		FRHICommandListImmediate& InRHICmdList,
@@ -825,6 +826,7 @@ struct FRelevancePacket
 		, OutHasDynamicEditorMeshElementsMasks(InOutHasDynamicEditorMeshElementsMasks)
 		, MarkMasks(InMarkMasks)
 		, CombinedShadingModelMask(0)
+		, bUsesGlobalDistanceField(false)
 	{
 	}
 
@@ -837,6 +839,8 @@ struct FRelevancePacket
 	void ComputeRelevance()
 	{
 		CombinedShadingModelMask = 0;
+		bUsesGlobalDistanceField = false;
+
 		SCOPE_CYCLE_COUNTER(STAT_ComputeViewRelevance);
 		for (int32 Index = 0; Index < Input.NumPrims; Index++)
 		{
@@ -897,6 +901,7 @@ struct FRelevancePacket
 			}
 
 			CombinedShadingModelMask |= ViewRelevance.ShadingModelMaskRelevance;			
+			bUsesGlobalDistanceField |= ViewRelevance.bUsesGlobalDistanceField;
 
 			if (ViewRelevance.bRenderCustomDepth)
 			{
@@ -1033,6 +1038,7 @@ struct FRelevancePacket
 			WriteView.PrimitiveVisibilityMap[NotDrawRelevant.Prims[Index]] = false;
 		}
 		WriteView.ShadingModelMaskInView |= CombinedShadingModelMask;
+		WriteView.bUsesGlobalDistanceField |= bUsesGlobalDistanceField;
 		VisibleEditorPrimitives.AppendTo(WriteView.VisibleEditorPrimitives);
 		VisibleDynamicPrimitives.AppendTo(WriteView.VisibleDynamicPrimitives);
 		WriteView.TranslucentPrimSet.AppendScenePrimitives(SortedTranslucencyPrims.Prims, SortedTranslucencyPrims.NumPrims, SortedSeparateTranslucencyPrims.Prims, SortedSeparateTranslucencyPrims.NumPrims);
diff --git a/Engine/Source/Runtime/Renderer/Private/ShaderBaseClasses.cpp b/Engine/Source/Runtime/Renderer/Private/ShaderBaseClasses.cpp
--- a/Engine/Source/Runtime/Renderer/Private/ShaderBaseClasses.cpp
+++ b/Engine/Source/Runtime/Renderer/Private/ShaderBaseClasses.cpp
@@ -73,6 +73,8 @@ FMaterialShader::FMaterialShader(const FMaterialShaderType::CompiledShaderInitia
 	// only used it Material has expression that requires PerlinNoise3DTexture
 	PerlinNoise3DTexture.Bind(Initializer.ParameterMap,TEXT("PerlinNoise3DTexture"));
 	PerlinNoise3DTextureSampler.Bind(Initializer.ParameterMap,TEXT("PerlinNoise3DTextureSampler"));
+
+	GlobalDistanceFieldParameters.Bind(Initializer.ParameterMap);
 }
 
 FUniformBufferRHIParamRef FMaterialShader::GetParameterCollectionBuffer(const FGuid& Id, const FSceneInterface* SceneInterface) const
@@ -300,6 +302,8 @@ void FMaterialShader::SetParameters(
 			Texture
 			);
 	}
+
+	GlobalDistanceFieldParameters.Set(RHICmdList, ShaderRHI, static_cast<const FViewInfo&>(View).GlobalDistanceFieldInfo.ParameterData);
 }
 
 // Doxygen struggles to parse these explicit specializations. Just ignore them for now.
@@ -333,19 +337,8 @@ bool FMaterialShader::Serialize(FArchive& Ar)
 	Ar << DeferredParameters;
 	Ar << LightAttenuation;
 	Ar << LightAttenuationSampler;
-	if (Ar.UE4Ver() >= VER_UE4_SMALLER_DEBUG_MATERIALSHADER_UNIFORM_EXPRESSIONS)
-	{
-		Ar << DebugUniformExpressionSet;
-		Ar << DebugDescription;
-	}
-	else if (Ar.IsLoading())
-	{
-		FUniformExpressionSet TempExpressionSet;
-		TempExpressionSet.Serialize(Ar);
-		DebugUniformExpressionSet.InitFromExpressionSet(TempExpressionSet);
-
-		Ar << DebugDescription;
-	}
+	Ar << DebugUniformExpressionSet;
+	Ar << DebugDescription;
 	Ar << AtmosphericFogTextureParameters;
 	Ar << EyeAdaptation;
 	Ar << PerlinNoiseGradientTexture;
@@ -357,6 +350,7 @@ bool FMaterialShader::Serialize(FArchive& Ar)
 	Ar << PerFrameVectorExpressions;
 	Ar << PerFramePrevScalarExpressions;
 	Ar << PerFramePrevVectorExpressions;
+	Ar << GlobalDistanceFieldParameters;
 
 	return bShaderHasOutdatedParameters;
 }
diff --git a/Engine/Source/Runtime/Renderer/Public/GlobalDistanceFieldParameters.h b/Engine/Source/Runtime/Renderer/Public/GlobalDistanceFieldParameters.h
new file mode 100644
--- /dev/null
+++ b/Engine/Source/Runtime/Renderer/Public/GlobalDistanceFieldParameters.h
@@ -0,0 +1,87 @@
+// Copyright 1998-2015 Epic Games, Inc. All Rights Reserved.
+
+/*=============================================================================
+	GlobalDistanceFieldParameters.h
+=============================================================================*/
+
+#pragma once
+
+/** Must match global distance field shaders. */
+const int32 GMaxGlobalDistanceFieldClipmaps = 4;
+
+class FGlobalDistanceFieldParameterData
+{
+public:
+
+	FGlobalDistanceFieldParameterData()
+	{
+		FPlatformMemory::Memzero(this, sizeof(FGlobalDistanceFieldParameterData));
+	}
+
+	FVector4 CenterAndExtent[GMaxGlobalDistanceFieldClipmaps];
+	FVector4 WorldToUVAddAndMul[GMaxGlobalDistanceFieldClipmaps];
+	FTextureRHIParamRef Textures[GMaxGlobalDistanceFieldClipmaps];
+	float GlobalDFResolution;
+	float MaxDistance;
+};
+
+class FGlobalDistanceFieldParameters
+{
+public:
+	void Bind(const FShaderParameterMap& ParameterMap)
+	{
+		GlobalDistanceFieldTexture.Bind(ParameterMap,TEXT("GlobalDistanceFieldTexture"));
+		GlobalDistanceFieldSampler.Bind(ParameterMap,TEXT("GlobalDistanceFieldSampler"));
+		GlobalVolumeCenterAndExtent.Bind(ParameterMap,TEXT("GlobalVolumeCenterAndExtent"));
+		GlobalVolumeWorldToUVAddAndMul.Bind(ParameterMap,TEXT("GlobalVolumeWorldToUVAddAndMul"));
+		GlobalVolumeDimension.Bind(ParameterMap,TEXT("GlobalVolumeDimension"));
+		GlobalVolumeTexelSize.Bind(ParameterMap,TEXT("GlobalVolumeTexelSize"));
+		MaxGlobalDistance.Bind(ParameterMap,TEXT("MaxGlobalDistance"));
+	}
+
+	friend FArchive& operator<<(FArchive& Ar,FGlobalDistanceFieldParameters& Parameters)
+	{
+		Ar << Parameters.GlobalDistanceFieldTexture;
+		Ar << Parameters.GlobalDistanceFieldSampler;
+		Ar << Parameters.GlobalVolumeCenterAndExtent;
+		Ar << Parameters.GlobalVolumeWorldToUVAddAndMul;
+		Ar << Parameters.GlobalVolumeDimension;
+		Ar << Parameters.GlobalVolumeTexelSize;
+		Ar << Parameters.MaxGlobalDistance;
+		return Ar;
+	}
+
+	template<typename ShaderRHIParamRef>
+	void Set(FRHICommandList& RHICmdList, const ShaderRHIParamRef ShaderRHI, const FGlobalDistanceFieldParameterData& ParameterData)
+	{
+		if (GlobalVolumeCenterAndExtent.IsBound() || GlobalVolumeWorldToUVAddAndMul.IsBound())
+		{
+			for (int32 ClipmapIndex = 0; ClipmapIndex < (int32)GlobalDistanceFieldTexture.GetNumResources(); ClipmapIndex++)
+			{
+				SetTextureParameter(
+					RHICmdList,
+					ShaderRHI,
+					GlobalDistanceFieldTexture,
+					GlobalDistanceFieldSampler,
+					TStaticSamplerState<SF_Bilinear,AM_Wrap,AM_Wrap,AM_Wrap>::GetRHI(),
+					ParameterData.Textures[ClipmapIndex],
+					ClipmapIndex);
+			}
+
+			SetShaderValueArray(RHICmdList, ShaderRHI, GlobalVolumeCenterAndExtent, ParameterData.CenterAndExtent, GMaxGlobalDistanceFieldClipmaps);
+			SetShaderValueArray(RHICmdList, ShaderRHI, GlobalVolumeWorldToUVAddAndMul, ParameterData.WorldToUVAddAndMul, GMaxGlobalDistanceFieldClipmaps);
+			SetShaderValue(RHICmdList, ShaderRHI, GlobalVolumeDimension, ParameterData.GlobalDFResolution);
+			SetShaderValue(RHICmdList, ShaderRHI, GlobalVolumeTexelSize, 1.0f / ParameterData.GlobalDFResolution);
+			SetShaderValue(RHICmdList, ShaderRHI, MaxGlobalDistance, ParameterData.MaxDistance);
+		}
+	}
+
+private:
+	FShaderResourceParameter GlobalDistanceFieldTexture;
+	FShaderResourceParameter GlobalDistanceFieldSampler;
+	FShaderParameter GlobalVolumeCenterAndExtent;
+	FShaderParameter GlobalVolumeWorldToUVAddAndMul;
+	FShaderParameter GlobalVolumeDimension;
+	FShaderParameter GlobalVolumeTexelSize;
+	FShaderParameter MaxGlobalDistance;
+};
diff --git a/Engine/Source/Runtime/Renderer/Public/MaterialShader.h b/Engine/Source/Runtime/Renderer/Public/MaterialShader.h
--- a/Engine/Source/Runtime/Renderer/Public/MaterialShader.h
+++ b/Engine/Source/Runtime/Renderer/Public/MaterialShader.h
@@ -10,6 +10,7 @@
 #include "AtmosphereTextureParameters.h"
 #include "SceneRenderTargetParameters.h"
 #include "PostProcessParameters.h"
+#include "GlobalDistanceFieldParameters.h"
 
 template<typename ParameterType> 
 struct TUniformParameter
@@ -159,6 +160,8 @@ private:
 	FShaderResourceParameter PerlinNoise3DTexture;
 	FShaderResourceParameter PerlinNoise3DTextureSampler;
 
+	FGlobalDistanceFieldParameters GlobalDistanceFieldParameters;
+
 	FDebugUniformExpressionSet DebugUniformExpressionSet;
 	FString DebugDescription;
 
