CanEditChange for OcclusionTint

diff --git a/Engine/Source/Runtime/Engine/Private/Components/SkyLightComponent.cpp b/Engine/Source/Runtime/Engine/Private/Components/SkyLightComponent.cpp
--- a/Engine/Source/Runtime/Engine/Private/Components/SkyLightComponent.cpp
+++ b/Engine/Source/Runtime/Engine/Private/Components/SkyLightComponent.cpp
@@ -280,7 +280,8 @@ bool USkyLightComponent::CanEditChange(const UProperty* InProperty) const
 
 		if (FCString::Strcmp(*PropertyName, TEXT("Contrast")) == 0
 			|| FCString::Strcmp(*PropertyName, TEXT("OcclusionMaxDistance")) == 0
-			|| FCString::Strcmp(*PropertyName, TEXT("MinOcclusion")) == 0)
+			|| FCString::Strcmp(*PropertyName, TEXT("MinOcclusion")) == 0
+			|| FCString::Strcmp(*PropertyName, TEXT("OcclusionTint")) == 0)
 		{
 			static const auto CVar = IConsoleManager::Get().FindTConsoleVariableDataInt(TEXT("r.GenerateMeshDistanceFields"));
 			return Mobility == EComponentMobility::Movable && CastShadows && CVar->GetValueOnGameThread() != 0;
