GetNetMode for listen server when PIE will now return correct value and not default to NM_Standalone

UE-12241

GetNetMode will return correctly for Servers and Clients when Loading Maps. Previously returned NM_Standalone due to net driver being null

UE-4792

diff --git a/Engine/Source/Runtime/Engine/Classes/Engine/World.h b/Engine/Source/Runtime/Engine/Classes/Engine/World.h
--- a/Engine/Source/Runtime/Engine/Classes/Engine/World.h
+++ b/Engine/Source/Runtime/Engine/Classes/Engine/World.h
@@ -2937,6 +2937,12 @@ public:
 	/** Returns the net mode this world is running under */
 	ENetMode GetNetMode() const;
 
+	/** Attempts to derive the net mode from PlayInSettings for PIE*/
+	ENetMode AttemptDeriveFromPlayInSettings() const;
+
+	/** Attempts to derive the net mode from URL */
+	ENetMode AttemptDeriveFromURL() const;
+
 	/**
 	 * Sets the net driver to use for this world
 	 * @param NewDriver the new net driver to use
diff --git a/Engine/Source/Runtime/Engine/Private/World.cpp b/Engine/Source/Runtime/Engine/Private/World.cpp
--- a/Engine/Source/Runtime/Engine/Private/World.cpp
+++ b/Engine/Source/Runtime/Engine/Private/World.cpp
@@ -51,6 +51,7 @@
 	#include "SlateBasics.h"
 	#include "Editor/Kismet/Public/FindInBlueprintManager.h"
 	#include "Editor/UnrealEd/Classes/Editor/UnrealEdTypes.h"
+	#include "Editor/UnrealEd/Classes/Settings/LevelEditorPlaySettings.h"
 #endif
 
 #include "MallocProfiler.h"
@@ -5362,9 +5363,89 @@ ENetMode UWorld::GetNetMode() const
 		return DemoNetDriver->GetNetMode();
 	}
 
+// PIE: NetDriver is not initialized so use PlayInSettings
+// to determine the Net Mode
+#if WITH_EDITOR
+	return AttemptDeriveFromPlayInSettings();
+#endif
+
+	// Use NextURL or PendingNetURL to derive NetMode
+	return AttemptDeriveFromURL();
+
+	//return NM_Standalone;
+}
+
+ENetMode UWorld::AttemptDeriveFromPlayInSettings() const
+{
+	if (WorldType == EWorldType::PIE)
+	{
+		const ULevelEditorPlaySettings* PlayInSettings = GetDefault<ULevelEditorPlaySettings>();
+		if (PlayInSettings)
+		{
+			EPlayNetMode PlayNetMode;
+			PlayInSettings->GetPlayNetMode(PlayNetMode);
+
+			switch (PlayNetMode)
+			{
+			case EPlayNetMode::PIE_Client:
+				return NM_Client;
+			case EPlayNetMode::PIE_ListenServer:
+				return NM_ListenServer;
+			case EPlayNetMode::PIE_Standalone:
+				return NM_Standalone;
+			default:
+				break;
+			}
+		}
+	}
+
+	return NM_Standalone;
+}
+
+ENetMode UWorld::AttemptDeriveFromURL() const
+{
+	if (GEngine != nullptr)
+	{
+		FWorldContext* WorldContext = GEngine->GetWorldContextFromWorld(this);
+
+		if (WorldContext != nullptr)
+		{
+			// NetMode can be derived from the NextURL if it exists
+			if (NextURL.Len() > 0)
+			{
+				FURL NextLevelURL(&WorldContext->LastURL, *NextURL, NextTravelType);
+
+				if (NextLevelURL.Valid)
+				{
+					if (NextLevelURL.HasOption(TEXT("listen")))
+					{
+						return NM_ListenServer;
+					}
+					else if (NextLevelURL.Host.Len() > 0)
+					{
+						return NM_Client;
+					}
+				}
+			}
+			// NetMode can be derived from the PendingNetURL if it exists
+			else if (WorldContext->PendingNetGame != nullptr && WorldContext->PendingNetGame->URL.Valid)
+			{
+				if (WorldContext->PendingNetGame->URL.HasOption(TEXT("listen")))
+				{
+					return NM_ListenServer;
+				}
+				else if (WorldContext->PendingNetGame->URL.Host.Len() > 0)
+				{
+					return NM_Client;
+				}
+			}
+		}
+	}
+
 	return NM_Standalone;
 }
 
+
 void UWorld::CopyGameState(AGameMode* FromGameMode, AGameState* FromGameState)
 {
 	AuthorityGameMode = FromGameMode;
