Fix sound in cooked builds

diff --git a/Engine/Source/Runtime/Engine/Private/SoundCue.cpp b/Engine/Source/Runtime/Engine/Private/SoundCue.cpp
--- a/Engine/Source/Runtime/Engine/Private/SoundCue.cpp
+++ b/Engine/Source/Runtime/Engine/Private/SoundCue.cpp
@@ -7,6 +7,7 @@
 #include "Sound/SoundNodeAttenuation.h"
 #include "Sound/SoundNodeQualityLevel.h"
 #include "Sound/SoundWave.h"
+#include "GameFramework/GameUserSettings.h"
 #if WITH_EDITOR
 #include "Kismet2/BlueprintEditorUtils.h"
 #include "UnrealEd.h"
@@ -57,6 +58,7 @@ void USoundCue::AddReferencedObjects(UObject* InThis, FReferenceCollector& Colle
 
 	Super::AddReferencedObjects(InThis, Collector);
 }
+#endif // WITH_EDITOR
 
 void USoundCue::PostLoad()
 {
@@ -118,6 +120,8 @@ void USoundCue::PostLoad()
 	}
 }
 
+#if WITH_EDITOR
+
 void USoundCue::PostEditChangeProperty(struct FPropertyChangedEvent& PropertyChangedEvent)
 {
 	Super::PostEditChangeProperty(PropertyChangedEvent);
