Workaround for GetGameThreadShaderMap during concurrent render state creation

diff --git a/Engine/Source/Runtime/Engine/Private/Materials/MaterialInterface.cpp b/Engine/Source/Runtime/Engine/Private/Materials/MaterialInterface.cpp
--- a/Engine/Source/Runtime/Engine/Private/Materials/MaterialInterface.cpp
+++ b/Engine/Source/Runtime/Engine/Private/Materials/MaterialInterface.cpp
@@ -66,7 +66,7 @@ FMaterialRelevance UMaterialInterface::GetRelevance_Internal(const UMaterial* Ma
 		MaterialRelevance.bNormalTranslucency = bIsTranslucent && !Material->bEnableSeparateTranslucency;
 		MaterialRelevance.bDisableDepthTest = bIsTranslucent && Material->bDisableDepthTest;		
 		MaterialRelevance.bOutputsVelocityInBasePass = Material->bOutputVelocityOnBasePass;	
-		MaterialRelevance.bUsesGlobalDistanceField = MaterialResource->GetGameThreadShaderMap() ? MaterialResource->GetGameThreadShaderMap()->UsesGlobalDistanceField() : false;
+		MaterialRelevance.bUsesGlobalDistanceField = false;//MaterialResource->GetGameThreadShaderMap() ? MaterialResource->GetGameThreadShaderMap()->UsesGlobalDistanceField() : false;
 		return MaterialRelevance;
 	}
 	else
