added help text

diff --git a/Engine/Source/Runtime/Engine/Private/Components/StaticMeshComponent.cpp b/Engine/Source/Runtime/Engine/Private/Components/StaticMeshComponent.cpp
--- a/Engine/Source/Runtime/Engine/Private/Components/StaticMeshComponent.cpp
+++ b/Engine/Source/Runtime/Engine/Private/Components/StaticMeshComponent.cpp
@@ -1151,7 +1151,7 @@ void UStaticMeshComponent::PostLoad()
 			Arguments.Add(TEXT("MeshName"), FText::FromString(GetName()));
 			FMessageLog("MapCheck").Info()
 				->AddToken(FUObjectToken::Create(GetOuter()))
-				->AddToken(FTextToken::Create(FText::Format( LOCTEXT( "MapCheck_Message_RepairedPaintedVertexColors", "{MeshName} : Repaired painted vertex colors" ), Arguments ) ))
+				->AddToken(FTextToken::Create(FText::Format( LOCTEXT( "MapCheck_Message_RepairedPaintedVertexColors", "{MeshName} : Repaired painted vertex colors (slow loading, can be fixed by saving asset)" ), Arguments ) ))
 				->AddToken(FMapErrorToken::Create(FMapErrors::RepairedPaintedVertexColors));
 #endif
 		}
