Fixing previous submit with checks for with editor

diff --git a/Engine/Source/Runtime/Engine/Private/World.cpp b/Engine/Source/Runtime/Engine/Private/World.cpp
--- a/Engine/Source/Runtime/Engine/Private/World.cpp
+++ b/Engine/Source/Runtime/Engine/Private/World.cpp
@@ -5377,6 +5377,7 @@ ENetMode UWorld::GetNetMode() const
 
 ENetMode UWorld::AttemptDeriveFromPlayInSettings() const
 {
+#if WITH_EDITOR
 	if (WorldType == EWorldType::PIE)
 	{
 		const ULevelEditorPlaySettings* PlayInSettings = GetDefault<ULevelEditorPlaySettings>();
@@ -5398,7 +5399,7 @@ ENetMode UWorld::AttemptDeriveFromPlayInSettings() const
 			}
 		}
 	}
-
+#endif
 	return NM_Standalone;
 }
 
