Fixed global distance field relevance with concurrent creates

diff --git a/Engine/Source/Runtime/Engine/Private/Materials/MaterialInterface.cpp b/Engine/Source/Runtime/Engine/Private/Materials/MaterialInterface.cpp
--- a/Engine/Source/Runtime/Engine/Private/Materials/MaterialInterface.cpp
+++ b/Engine/Source/Runtime/Engine/Private/Materials/MaterialInterface.cpp
@@ -66,7 +66,7 @@ FMaterialRelevance UMaterialInterface::GetRelevance_Internal(const UMaterial* Ma
 		MaterialRelevance.bNormalTranslucency = bIsTranslucent && !Material->bEnableSeparateTranslucency;
 		MaterialRelevance.bDisableDepthTest = bIsTranslucent && Material->bDisableDepthTest;		
 		MaterialRelevance.bOutputsVelocityInBasePass = Material->bOutputVelocityOnBasePass;	
-		MaterialRelevance.bUsesGlobalDistanceField = false;//MaterialResource->GetGameThreadShaderMap() ? MaterialResource->GetGameThreadShaderMap()->UsesGlobalDistanceField() : false;
+		MaterialRelevance.bUsesGlobalDistanceField = MaterialResource->UsesGlobalDistanceField_GameThread();
 		return MaterialRelevance;
 	}
 	else
diff --git a/Engine/Source/Runtime/Engine/Private/Materials/MaterialShared.cpp b/Engine/Source/Runtime/Engine/Private/Materials/MaterialShared.cpp
--- a/Engine/Source/Runtime/Engine/Private/Materials/MaterialShared.cpp
+++ b/Engine/Source/Runtime/Engine/Private/Materials/MaterialShared.cpp
@@ -500,6 +500,11 @@ bool FMaterial::UsesEyeAdaptation() const
 	return false;
 }
 
+bool FMaterial::UsesGlobalDistanceField_GameThread() const 
+{ 
+	return GameThreadShaderMap.GetReference() ? GameThreadShaderMap->UsesGlobalDistanceField() : false; 
+}
+
 bool FMaterial::MaterialModifiesMeshPosition_RenderThread() const
 { 
 	check(IsInParallelRenderingThread());
diff --git a/Engine/Source/Runtime/Engine/Public/MaterialShared.h b/Engine/Source/Runtime/Engine/Public/MaterialShared.h
--- a/Engine/Source/Runtime/Engine/Public/MaterialShared.h
+++ b/Engine/Source/Runtime/Engine/Public/MaterialShared.h
@@ -1024,6 +1024,7 @@ public:
 	ENGINE_API bool NeedsSceneTextures() const;
 	ENGINE_API bool NeedsGBuffer() const;
 	ENGINE_API bool UsesEyeAdaptation() const;	
+	ENGINE_API bool UsesGlobalDistanceField_GameThread() const;
 
 	/** Does the material modify the mesh position. */
 	ENGINE_API bool MaterialModifiesMeshPosition_RenderThread() const;
