#ue4 - Fix regression causing SkeletalMeshComponents not to trigger overlap events.

UE-15766

diff --git a/Engine/Source/Runtime/Engine/Classes/Components/SkeletalMeshComponent.h b/Engine/Source/Runtime/Engine/Classes/Components/SkeletalMeshComponent.h
--- a/Engine/Source/Runtime/Engine/Classes/Components/SkeletalMeshComponent.h
+++ b/Engine/Source/Runtime/Engine/Classes/Components/SkeletalMeshComponent.h
@@ -838,6 +838,7 @@ public:
 	virtual FBoxSphereBounds CalcBounds(const FTransform& LocalToWorld) const override;
 	virtual bool IsAnySimulatingPhysics() const override;
 	virtual void OnUpdateTransform(bool bSkipPhysicsMove, bool bTeleport = false) override;
+	virtual void UpdateOverlaps(TArray<FOverlapInfo> const* PendingOverlaps=NULL, bool bDoNotifies=true, const TArray<FOverlapInfo>* OverlapsAtEndLocation=NULL) override;
 	// End USceneComponent interface.
 
 	// Begin UPrimitiveComponent interface.
diff --git a/Engine/Source/Runtime/Engine/Private/SkeletalMeshComponentPhysics.cpp b/Engine/Source/Runtime/Engine/Private/SkeletalMeshComponentPhysics.cpp
--- a/Engine/Source/Runtime/Engine/Private/SkeletalMeshComponentPhysics.cpp
+++ b/Engine/Source/Runtime/Engine/Private/SkeletalMeshComponentPhysics.cpp
@@ -1356,6 +1356,12 @@ void USkeletalMeshComponent::OnUpdateTransform(bool bSkipPhysicsMove, bool bTele
 #endif //#if WITH_APEX_CLOTHING
 }
 
+void USkeletalMeshComponent::UpdateOverlaps(TArray<FOverlapInfo> const* PendingOverlaps, bool bDoNotifies, const TArray<FOverlapInfo>* OverlapsAtEndLocation)
+{
+	// Parent class (USkinnedMeshComponent) routes only to children, but we really do want to test our own bodies for overlaps.
+	UPrimitiveComponent::UpdateOverlaps(PendingOverlaps, bDoNotifies, OverlapsAtEndLocation);
+}
+
 void USkeletalMeshComponent::CreatePhysicsState()
 {
 	//	UE_LOG(LogSkeletalMesh, Warning, TEXT("Creating Physics State (%s : %s)"), *GetNameSafe(GetOuter()),  *GetName());
