PR# 1155 - Contributed by [REDACTED_PERSON]
UE-15727 - Fix uninitialized variable FHLSLMaterialTranslator::bUsesAtmosphericFog

diff --git a/Engine/Source/Runtime/Engine/Private/Materials/HLSLMaterialTranslator.h b/Engine/Source/Runtime/Engine/Private/Materials/HLSLMaterialTranslator.h
--- a/Engine/Source/Runtime/Engine/Private/Materials/HLSLMaterialTranslator.h
+++ b/Engine/Source/Runtime/Engine/Private/Materials/HLSLMaterialTranslator.h
@@ -239,6 +239,7 @@ public:
 	,	bNeedsWorldPositionExcludingShaderOffsets(false)
 	,	bNeedsParticleSize(false)
 	,	bNeedsSceneTexturePostProcessInputs(false)
+	,	bUsesAtmosphericFog(false)
 	,	bUsesVertexColor(false)
 	,	bUsesParticleColor(false)
 	,	bUsesTransformVector(false)
