Accessor function for locked fov setting on player camera manager

diff --git a/Engine/Source/Runtime/Engine/Classes/Camera/PlayerCameraManager.h b/Engine/Source/Runtime/Engine/Classes/Camera/PlayerCameraManager.h
--- a/Engine/Source/Runtime/Engine/Classes/Camera/PlayerCameraManager.h
+++ b/Engine/Source/Runtime/Engine/Classes/Camera/PlayerCameraManager.h
@@ -183,6 +183,8 @@ public:
 	UPROPERTY(EditAnywhere, BlueprintReadWrite, Category=PlayerCameraManager)
 	float DefaultFOV;
 
+	float GetLockedFOV() const { return LockedFOV; }
+
 protected:
 	/** Value to lock FOV to, in degrees. Ignored if <= 0, utilized if > 0. */
 	float LockedFOV;
