Added FCompactPose, contains only the bone transforms needed by RequiredBones. Allows us to do animation evaluation without having to refer to required bones array all the time

API Breaking Changes

Removed Functions:
FA2CSPose::SafeSetCSBoneTransforms
FA2CSPose::LocalBlendCSBoneTransforms

Member Type Changes:
- int32 -> FCompactPoseBoneIndex
--- FBoneTransform::BoneIndex
--- FABRIKChainLink::BoneIndex
- FA2CSPose -> FCSPose<FCompactPose>
--- FAnimNode_SkeletalControlBase::ForwardedPose
--- FComponentSpacePoseContext::Pose
- FA2Pose -> FCompactPose
--- FSlotEvaluationPose::Pose
--- FAnimNode_TransitionPoseEvaluator::CachedPose
--- FAnimNode_SaveCachedPose::CachedPose
--- FPoseContext::Pose

Paramater Changes:
- FA2Pose& -> FCompactPose&
--- UAnimSingleNodeInstance::InternalBlendSpaceEvaluatePose (FA2Pose& Pose -> FCompactPose&)
--- UAnimInstance::SequenceEvaluatePose
--- UAnimInstance::BlendSequences
--- UAnimInstance::CopyPose
--- UAnimInstance::ApplyAdditiveSequence
--- UAnimInstance::BlendSpaceEvaluatePose
--- UAnimInstance::BlendRotationOffset
--- UAnimInstance::GetSlotWeight
--- UAnimInstance::SlotEvaluatePose
- FA2CSPose& -> FCSPose<FCompactPose>&
--- FAnimNode_SkeletalControlBase::EvaluateBoneTransforms
--- UAnimGraphNode_SkeletalControlBase::ConvertCSVectorToBoneSpace
--- UAnimGraphNode_SkeletalControlBase::ConvertCSRotationToBoneSpace
--- UAnimGraphNode_SkeletalControlBase::ConvertWidgetLocation
--- UAnimPreviewInstance::ApplyBoneControllers
- TArray<FTransform> -> FCompactPose
--- UAnimPreviewInstance::SetKeyImplementation
--- UAnimSequence::GetAnimationPose
--- UAnimSequence::GetBonePose
--- UAnimSequence::GetBonePose_Additive
--- UAnimSequence::GetAdditiveBasePose
--- UAnimSequence::GetBonePose_AdditiveMeshRotationOnly
--- FAnimationRuntime::BlendPosesTogether ( + Removed NumPoses and RequiredBones )
--- FAnimationRuntime::BlendPosesTogetherPerBone  ( + Removed NumPoses and RequiredBones )
--- FAnimationRuntime::BlendPosesTogetherPerBoneInMeshSpace ( + Removed NumPoses and RequiredBones )
--- FAnimationRuntime::BlendPosesPerBoneFilter ( + Removed Skeleton and RequiredBones )
--- FAnimationRuntime::GetPoseFromSequence ( + Removed  RequiredBones )
--- FAnimationRuntime::GetPoseFromAnimTrack ( + Removed  RequiredBones )
--- FAnimationRuntime::FillWithRetargetBaseRefPose ( + Removed  RequiredBones )
--- FAnimationRuntime::ConvertPoseToAdditive ( + Removed  RequiredBones )
--- FAnimationRuntime::ConvertPoseToMeshRotation ( + Removed  RequiredBones )
--- FAnimationRuntime::BlendPosesAccumulate ( + Removed  RequiredBones )
--- FAnimationRuntime::BlendAdditivePose ( + Removed  RequiredBones )
--- FAnimationRuntime::GetPoseFromBlendSpace ( + Removed  RequiredBones )
--- FAnimationRuntime::ConvertCSTransformToBoneSpace
--- FAnimationRuntime::ConvertBoneSpaceTransformToCS
--- FAnimationRuntime::BlendMeshPosesPerBoneWeights ( + Removed Skeleton and RequiredBones )
--- FAnimationRuntime::BlendLocalPosesPerBoneWeights ( + Removed Skeleton and RequiredBones )
- TArray<FTransform> -> FTransform
--- UAnimSequence::ResetRootBoneForRootMotion
- int32 -> FCompactPoseBoneIndex
--- UAnimSequence::RetargetBoneTransform

FAnimationRuntime: Many changes to animation operation functions, removal of NumPoses and RequiredBones, conversion of TArray<FTransform> to FCompactPose

Removed USTRUCT markup from FBoneTransform
Removed UPROPERTY markup from FAnimNode_TransitionPoseEvaluator::CachedPose

diff --git a/Engine/Source/Editor/AnimGraph/Classes/AnimGraphNode_SkeletalControlBase.h b/Engine/Source/Editor/AnimGraph/Classes/AnimGraphNode_SkeletalControlBase.h
--- a/Engine/Source/Editor/AnimGraph/Classes/AnimGraphNode_SkeletalControlBase.h
+++ b/Engine/Source/Editor/AnimGraph/Classes/AnimGraphNode_SkeletalControlBase.h
@@ -75,11 +75,11 @@ protected:
 	// local conversion function for drawing
 	void ConvertToComponentSpaceTransform(const USkeletalMeshComponent* SkelComp, const FTransform & InTransform, FTransform & OutCSTransform, int32 BoneIndex, EBoneControlSpace Space) const;
 	// convert drag vector in component space to bone space 
-	FVector ConvertCSVectorToBoneSpace(const USkeletalMeshComponent* SkelComp, FVector& InCSVector, FA2CSPose& MeshBases, const FName& BoneName, const EBoneControlSpace Space);
+	FVector ConvertCSVectorToBoneSpace(const USkeletalMeshComponent* SkelComp, FVector& InCSVector, FCSPose<FCompactPose>& MeshBases, const FName& BoneName, const EBoneControlSpace Space);
 	// convert rotator in component space to bone space 
-	FQuat ConvertCSRotationToBoneSpace(const USkeletalMeshComponent* SkelComp, FRotator& InCSRotator, FA2CSPose& MeshBases, const FName& BoneName, const EBoneControlSpace Space);
+	FQuat ConvertCSRotationToBoneSpace(const USkeletalMeshComponent* SkelComp, FRotator& InCSRotator, FCSPose<FCompactPose>& MeshBases, const FName& BoneName, const EBoneControlSpace Space);
 	// convert widget location according to bone control space
-	FVector ConvertWidgetLocation(const USkeletalMeshComponent* InSkelComp, FA2CSPose& InMeshBases, const FName& BoneName, const FVector& InLocation, const EBoneControlSpace Space);
+	FVector ConvertWidgetLocation(const USkeletalMeshComponent* InSkelComp, FCSPose<FCompactPose>& InMeshBases, const FName& BoneName, const FVector& InLocation, const EBoneControlSpace Space);
 	// set literal value for FVector
 	void SetDefaultValue(const FString& InDefaultValueName, const FVector& InValue);
 	// get literal value for vector
diff --git a/Engine/Source/Editor/AnimGraph/Classes/AnimPreviewInstance.h b/Engine/Source/Editor/AnimGraph/Classes/AnimPreviewInstance.h
--- a/Engine/Source/Editor/AnimGraph/Classes/AnimPreviewInstance.h
+++ b/Engine/Source/Editor/AnimGraph/Classes/AnimPreviewInstance.h
@@ -150,7 +150,7 @@ private:
 	 * @param	BoneControllers	 List of Bone Controllers to apply
 	 * @param 	OutMeshPose	Outpose in Mesh Space once applied
 	 */
-	void ApplyBoneControllers(USkeletalMeshComponent* Component, TArray<FAnimNode_ModifyBone> &BoneControllers, FA2CSPose& OutMeshPose);
+	void ApplyBoneControllers(USkeletalMeshComponent* Component, TArray<FAnimNode_ModifyBone> &BoneControllers, FCSPose<FCompactPose>& OutMeshPose);
 	/** 
 	 * Update CurveControllers based on TransformCurves of Animation
 	 */
@@ -160,7 +160,7 @@ private:
 	 * Set Key Implementation function
 	 * It gets Pre Controller Local Space and gets Post Controller Local Space, and add the key to the curve 
 	 */
-	void SetKeyImplementation(const TArray<FTransform>& PreControllerInLocalSpace, const TArray<FTransform>& PostControllerInLocalSpace);
+	void SetKeyImplementation(const FCompactPose& PreControllerInLocalSpace, const FCompactPose& PostControllerInLocalSpace);
 	/** 
 	 * Add Key to the Sequence
 	 * Now Additive Key is generated, add to the curves
diff --git a/Engine/Source/Editor/AnimGraph/Private/AnimGraphNode_ModifyBone.cpp b/Engine/Source/Editor/AnimGraph/Private/AnimGraphNode_ModifyBone.cpp
--- a/Engine/Source/Editor/AnimGraph/Private/AnimGraphNode_ModifyBone.cpp
+++ b/Engine/Source/Editor/AnimGraph/Private/AnimGraphNode_ModifyBone.cpp
@@ -59,18 +59,20 @@ FVector UAnimGraphNode_ModifyBone::GetWidgetLocation(const USkeletalMeshComponen
 	// if the current widget mode is translate, then shows the widget according to translation space
 	if (CurWidgetMode == FWidget::WM_Translate)
 	{
-		FA2CSPose& MeshBases = AnimNode->ForwardedPose;
+		FCSPose<FCompactPose>& MeshBases = AnimNode->ForwardedPose;
 		WidgetLoc = ConvertWidgetLocation(SkelComp, MeshBases, Node.BoneToModify.BoneName, GetNodeValue(FString("Translation"), Node.Translation), Node.TranslationSpace);
 
-		if (MeshBases.IsValid() && Node.TranslationMode == BMM_Additive)
+		if (MeshBases.GetPose().IsValid() && Node.TranslationMode == BMM_Additive)
 		{
 			if(Node.TranslationSpace == EBoneControlSpace::BCS_WorldSpace ||
 				Node.TranslationSpace == EBoneControlSpace::BCS_ComponentSpace)
 			{
-				int32 MeshBoneIndex = SkelComp->GetBoneIndex(Node.BoneToModify.BoneName);
-				if(MeshBoneIndex != INDEX_NONE)
+				const FMeshPoseBoneIndex MeshBoneIndex(SkelComp->GetBoneIndex(Node.BoneToModify.BoneName));
+				const FCompactPoseBoneIndex BoneIndex = MeshBases.GetPose().GetBoneContainer().MakeCompactPoseIndex(MeshBoneIndex);
+
+				if (BoneIndex != INDEX_NONE)
 				{
-					FTransform BoneTM = MeshBases.GetComponentSpaceTransform(MeshBoneIndex);
+					const FTransform& BoneTM = MeshBases.GetComponentSpaceTransform(BoneIndex);
 					WidgetLoc += BoneTM.GetLocation();
 				}
 			}
diff --git a/Engine/Source/Editor/AnimGraph/Private/AnimGraphNode_SkeletalControlBase.cpp b/Engine/Source/Editor/AnimGraph/Private/AnimGraphNode_SkeletalControlBase.cpp
--- a/Engine/Source/Editor/AnimGraph/Private/AnimGraphNode_SkeletalControlBase.cpp
+++ b/Engine/Source/Editor/AnimGraph/Private/AnimGraphNode_SkeletalControlBase.cpp
@@ -126,13 +126,14 @@ void UAnimGraphNode_SkeletalControlBase::ConvertToComponentSpaceTransform(const
 }
 
 
-FVector UAnimGraphNode_SkeletalControlBase::ConvertCSVectorToBoneSpace(const USkeletalMeshComponent* SkelComp, FVector& InCSVector, FA2CSPose& MeshBases, const FName& BoneName, const EBoneControlSpace Space)
+FVector UAnimGraphNode_SkeletalControlBase::ConvertCSVectorToBoneSpace(const USkeletalMeshComponent* SkelComp, FVector& InCSVector, FCSPose<FCompactPose>& MeshBases, const FName& BoneName, const EBoneControlSpace Space)
 {
 	FVector OutVector = InCSVector;
 
-	if (MeshBases.IsValid())
+	if (MeshBases.GetPose().IsValid())
 	{
-		int32 MeshBoneIndex = SkelComp->GetBoneIndex(BoneName);
+		const FMeshPoseBoneIndex MeshBoneIndex(SkelComp->GetBoneIndex(BoneName));
+		const FCompactPoseBoneIndex BoneIndex = MeshBases.GetPose().GetBoneContainer().MakeCompactPoseIndex(MeshBoneIndex);
 
 		switch (Space)
 		{
@@ -144,10 +145,10 @@ FVector UAnimGraphNode_SkeletalControlBase::ConvertCSVectorToBoneSpace(const USk
 
 		case BCS_ParentBoneSpace:
 		{
-			const int32 ParentIndex = MeshBases.GetParentBoneIndex(MeshBoneIndex);
+			const FCompactPoseBoneIndex ParentIndex = MeshBases.GetPose().GetParentBoneIndex(BoneIndex);
 			if (ParentIndex != INDEX_NONE)
 			{
-				FTransform ParentTM = MeshBases.GetComponentSpaceTransform(ParentIndex);
+				const FTransform& ParentTM = MeshBases.GetComponentSpaceTransform(ParentIndex);
 				OutVector = ParentTM.InverseTransformVector(InCSVector);
 			}
 		}
@@ -155,7 +156,7 @@ FVector UAnimGraphNode_SkeletalControlBase::ConvertCSVectorToBoneSpace(const USk
 
 		case BCS_BoneSpace:
 		{
-			FTransform BoneTM = MeshBases.GetComponentSpaceTransform(MeshBoneIndex);
+			const FTransform& BoneTM = MeshBases.GetComponentSpaceTransform(BoneIndex);
 			OutVector = BoneTM.InverseTransformVector(InCSVector);
 		}
 			break;
@@ -165,13 +166,14 @@ FVector UAnimGraphNode_SkeletalControlBase::ConvertCSVectorToBoneSpace(const USk
 	return OutVector;
 }
 
-FQuat UAnimGraphNode_SkeletalControlBase::ConvertCSRotationToBoneSpace(const USkeletalMeshComponent* SkelComp, FRotator& InCSRotator, FA2CSPose& MeshBases, const FName& BoneName, const EBoneControlSpace Space)
+FQuat UAnimGraphNode_SkeletalControlBase::ConvertCSRotationToBoneSpace(const USkeletalMeshComponent* SkelComp, FRotator& InCSRotator, FCSPose<FCompactPose>& MeshBases, const FName& BoneName, const EBoneControlSpace Space)
 {
 	FQuat OutQuat = FQuat::Identity;
 
-	if (MeshBases.IsValid())
+	if (MeshBases.GetPose().IsValid())
 	{
-		int32 MeshBoneIndex = SkelComp->GetBoneIndex(BoneName);
+		const FMeshPoseBoneIndex MeshBoneIndex(SkelComp->GetBoneIndex(BoneName));
+		const FCompactPoseBoneIndex BoneIndex = MeshBases.GetPose().GetBoneContainer().MakeCompactPoseIndex(MeshBoneIndex);
 
 		FVector RotAxis;
 		float RotAngle;
@@ -188,13 +190,13 @@ FQuat UAnimGraphNode_SkeletalControlBase::ConvertCSRotationToBoneSpace(const USk
 
 		case BCS_ParentBoneSpace:
 		{
-			const int32 ParentIndex = MeshBases.GetParentBoneIndex(MeshBoneIndex);
+			const FCompactPoseBoneIndex ParentIndex = MeshBases.GetPose().GetParentBoneIndex(BoneIndex);
 			if (ParentIndex != INDEX_NONE)
 			{
-				FTransform ParentTM = MeshBases.GetComponentSpaceTransform(ParentIndex);
-				ParentTM = ParentTM.Inverse();
+				const FTransform& ParentTM = MeshBases.GetComponentSpaceTransform(ParentIndex);
+				FTransform InverseParentTM = ParentTM.Inverse();
 				//Calculate the new delta rotation
-				FVector4 BoneSpaceAxis = ParentTM.TransformVector(RotAxis);
+				FVector4 BoneSpaceAxis = InverseParentTM.TransformVector(RotAxis);
 				FQuat DeltaQuat(BoneSpaceAxis, RotAngle);
 				DeltaQuat.Normalize();
 				OutQuat = DeltaQuat;
@@ -204,9 +206,9 @@ FQuat UAnimGraphNode_SkeletalControlBase::ConvertCSRotationToBoneSpace(const USk
 
 		case BCS_BoneSpace:
 		{
-			FTransform BoneTM = MeshBases.GetComponentSpaceTransform(MeshBoneIndex);
-			BoneTM = BoneTM.Inverse();
-			FVector4 BoneSpaceAxis = BoneTM.TransformVector(RotAxis);
+			const FTransform& BoneTM = MeshBases.GetComponentSpaceTransform(BoneIndex);
+			FTransform InverseBoneTM = BoneTM.Inverse();
+			FVector4 BoneSpaceAxis = InverseBoneTM.TransformVector(RotAxis);
 			//Calculate the new delta rotation
 			FQuat DeltaQuat(BoneSpaceAxis, RotAngle);
 			DeltaQuat.Normalize();
@@ -219,14 +221,15 @@ FQuat UAnimGraphNode_SkeletalControlBase::ConvertCSRotationToBoneSpace(const USk
 	return OutQuat;
 }
 
-FVector UAnimGraphNode_SkeletalControlBase::ConvertWidgetLocation(const USkeletalMeshComponent* SkelComp, FA2CSPose& MeshBases, const FName& BoneName, const FVector& Location, const EBoneControlSpace Space)
+FVector UAnimGraphNode_SkeletalControlBase::ConvertWidgetLocation(const USkeletalMeshComponent* SkelComp, FCSPose<FCompactPose>& MeshBases, const FName& BoneName, const FVector& Location, const EBoneControlSpace Space)
 {
 	FVector WidgetLoc = FVector::ZeroVector;
 
-	if (MeshBases.IsValid())
+	if (MeshBases.GetPose().IsValid())
 	{
 		USkeleton * Skeleton = SkelComp->SkeletalMesh->Skeleton;
-		int32 MeshBoneIndex = SkelComp->GetBoneIndex(BoneName);
+		const FMeshPoseBoneIndex MeshBoneIndex(SkelComp->GetBoneIndex(BoneName));
+		const FCompactPoseBoneIndex CompactBoneIndex = MeshBases.GetPose().GetBoneContainer().MakeCompactPoseIndex(MeshBoneIndex);
 
 		switch (Space)
 		{
@@ -241,12 +244,12 @@ FVector UAnimGraphNode_SkeletalControlBase::ConvertWidgetLocation(const USkeleta
 
 		case BCS_ParentBoneSpace:
 
-			if (MeshBoneIndex != INDEX_NONE)
+			if (CompactBoneIndex != INDEX_NONE)
 			{
-				const int32 MeshParentIndex = MeshBases.GetParentBoneIndex(MeshBoneIndex);
-				if (MeshParentIndex != INDEX_NONE)
+				const FCompactPoseBoneIndex CompactParentIndex = MeshBases.GetPose().GetParentBoneIndex(CompactBoneIndex);
+				if (CompactParentIndex != INDEX_NONE)
 				{
-					const FTransform ParentTM = MeshBases.GetComponentSpaceTransform(MeshParentIndex);
+					const FTransform& ParentTM = MeshBases.GetComponentSpaceTransform(CompactParentIndex);
 					WidgetLoc = ParentTM.TransformPosition(Location);
 				}
 			}
@@ -254,9 +257,9 @@ FVector UAnimGraphNode_SkeletalControlBase::ConvertWidgetLocation(const USkeleta
 
 		case BCS_BoneSpace:
 
-			if (MeshBoneIndex != INDEX_NONE)
+			if (CompactBoneIndex != INDEX_NONE)
 			{
-				FTransform BoneTM = MeshBases.GetComponentSpaceTransform(MeshBoneIndex);
+				const FTransform& BoneTM = MeshBases.GetComponentSpaceTransform(CompactBoneIndex);
 				WidgetLoc = BoneTM.TransformPosition(Location);
 			}
 		}
diff --git a/Engine/Source/Editor/AnimGraph/Private/AnimPreviewInstance.cpp b/Engine/Source/Editor/AnimGraph/Private/AnimPreviewInstance.cpp
--- a/Engine/Source/Editor/AnimGraph/Private/AnimPreviewInstance.cpp
+++ b/Engine/Source/Editor/AnimGraph/Private/AnimPreviewInstance.cpp
@@ -123,12 +123,12 @@ bool UAnimPreviewInstance::NativeEvaluateAnimation(FPoseContext& Output)
 		USkeletalMeshComponent* MeshComponent = GetSkelMeshComponent();
 		if(MeshComponent && MeshComponent->SkeletalMesh)
 		{
-			FAnimationRuntime::FillWithRetargetBaseRefPose(Output.Pose.Bones, GetSkelMeshComponent()->SkeletalMesh, RequiredBones);
+			FAnimationRuntime::FillWithRetargetBaseRefPose(Output.Pose, GetSkelMeshComponent()->SkeletalMesh);
 		}
 		else
 		{
 			// ideally we'll return just ref pose, but not sure if this will work with LODs
-			FAnimationRuntime::FillWithRefPose(Output.Pose.Bones, RequiredBones);
+			Output.Pose.ResetToRefPose();
 		}
 	}
 	else
@@ -149,16 +149,16 @@ bool UAnimPreviewInstance::NativeEvaluateAnimation(FPoseContext& Output)
 			// create bone controllers from 
 			if(BoneControllers.Num() > 0 || CurveBoneControllers.Num() > 0)
 			{
-				TArray<FTransform> PreController, PostController;
+				FCompactPose PreController, PostController;
 				// if set key is true, we should save pre controller local space transform 
 				// so that we can calculate the delta correctly
 				if(bSetKey)
 				{
-					PreController = Output.Pose.Bones;
+					PreController = Output.Pose;
 				}
 
-				FA2CSPose OutMeshPose;
-				OutMeshPose.AllocateLocalPoses(RequiredBones, Output.Pose);
+				FCSPose<FCompactPose> OutMeshPose;
+				OutMeshPose.InitPose(&RequiredBones);
 
 				// apply curve data first
 				ApplyBoneControllers(Component, CurveBoneControllers, OutMeshPose);
@@ -173,7 +173,7 @@ bool UAnimPreviewInstance::NativeEvaluateAnimation(FPoseContext& Output)
 				if(bSetKey)
 				{
 					// now we have post controller, and calculate delta now
-					PostController = Output.Pose.Bones;
+					PostController = Output.Pose;
 					SetKeyImplementation(PreController, PostController);
 				}
 			}
@@ -183,7 +183,7 @@ bool UAnimPreviewInstance::NativeEvaluateAnimation(FPoseContext& Output)
 				if(bSetKey)
 				{
 					// in this case, pose is same
-					SetKeyImplementation(Output.Pose.Bones, Output.Pose.Bones);
+					SetKeyImplementation(Output.Pose, Output.Pose);
 				}
 			}
 		}
@@ -198,7 +198,7 @@ bool UAnimPreviewInstance::NativeEvaluateAnimation(FPoseContext& Output)
 	return true;
 }
 
-void UAnimPreviewInstance::ApplyBoneControllers(USkeletalMeshComponent* Component, TArray<FAnimNode_ModifyBone> &InBoneControllers, FA2CSPose& OutMeshPose)
+void UAnimPreviewInstance::ApplyBoneControllers(USkeletalMeshComponent* Component, TArray<FAnimNode_ModifyBone> &InBoneControllers, FCSPose<FCompactPose>& OutMeshPose)
 {
 	for(auto& SingleBoneController : InBoneControllers)
 	{
@@ -206,7 +206,7 @@ void UAnimPreviewInstance::ApplyBoneControllers(USkeletalMeshComponent* Componen
 		if(SingleBoneController.BoneToModify.BoneIndex != INDEX_NONE)
 		{
 			TArray<FBoneTransform> BoneTransforms;
-			SingleBoneController.EvaluateBoneTransforms(Component, RequiredBones, OutMeshPose, BoneTransforms);
+			SingleBoneController.EvaluateBoneTransforms(Component, OutMeshPose, BoneTransforms);
 			if(BoneTransforms.Num() > 0)
 			{
 				OutMeshPose.LocalBlendCSBoneTransforms(BoneTransforms, 1.0f);
@@ -215,7 +215,7 @@ void UAnimPreviewInstance::ApplyBoneControllers(USkeletalMeshComponent* Componen
 	}
 }
 
-void UAnimPreviewInstance::SetKeyImplementation(const TArray<FTransform>& PreControllerInLocalSpace, const TArray<FTransform>& PostControllerInLocalSpace)
+void UAnimPreviewInstance::SetKeyImplementation(const FCompactPose& PreControllerInLocalSpace, const FCompactPose& PostControllerInLocalSpace)
 {
 #if WITH_EDITOR
 	// evaluate the curve data first
@@ -240,7 +240,8 @@ void UAnimPreviewInstance::SetKeyImplementation(const TArray<FTransform>& PreCon
 			// find if this already exists, then just add curve data only
 			FName BoneName = SingleBoneController.BoneToModify.BoneName;
 			// now convert data
-			const int32 BoneIndex = Component->GetBoneIndex(BoneName);
+			const FMeshPoseBoneIndex MeshBoneIndex(Component->GetBoneIndex(BoneName));
+			const FCompactPoseBoneIndex BoneIndex = RequiredBones.MakeCompactPoseIndex(MeshBoneIndex);
 			FTransform  LocalTransform = PostControllerInLocalSpace[BoneIndex];
 
 			// now we have LocalTransform and get additive data
diff --git a/Engine/Source/Runtime/Engine/Classes/Animation/AnimInstance.h b/Engine/Source/Runtime/Engine/Classes/Animation/AnimInstance.h
--- a/Engine/Source/Runtime/Engine/Classes/Animation/AnimInstance.h
+++ b/Engine/Source/Runtime/Engine/Classes/Animation/AnimInstance.h
@@ -7,6 +7,7 @@
 #include "Components/SkeletalMeshComponent.h"
 #include "Components/SkinnedMeshComponent.h"
 #include "AnimStateMachineTypes.h"
+#include "BonePose.h"
 #include "AnimInstance.generated.h"
 
 struct FAnimMontageInstance;
@@ -135,23 +136,6 @@ public:
 	/** convert to local poses **/
 	void ConvertToLocalPoses(FA2Pose & LocalPoses) const;
 
-	/** 
-	 * Set a bunch of Component Space Bone Transforms.
-	 * Do this safely by insuring that Parents are already in Component Space,
-	 * and any Component Space children are converted back to Local Space before hand.
-	 */
-	void SafeSetCSBoneTransforms(const TArray<struct FBoneTransform> & BoneTransforms);
-
-	/** 
-	 * Blends Component Space transforms to MeshPose in Local Space. 
-	 * Used by SkelControls to apply their transforms.
-	 *
-	 * The tricky bit is that SkelControls deliver their transforms in Component Space,
-	 * But the blending is done in Local Space. Also we need to refresh any Children they have
-	 * that has been previously converted to Component Space.
-	 */
-	void LocalBlendCSBoneTransforms(const TArray<struct FBoneTransform> & BoneTransforms, float Alpha);
-
 private:
 	/** Calculate all transform till parent **/
 	void CalculateComponentSpaceTransform(int32 Index);
@@ -170,24 +154,19 @@ private:
 	friend class FAnimationRuntime;
 };
 
-USTRUCT(BlueprintType)
 struct FBoneTransform
 {
-	GENERATED_USTRUCT_BODY()
-
-	/** @todo anim: should be Skeleton bone index in the future, but right now it's Mesh BoneIndex **/
-	UPROPERTY()
-	int32 BoneIndex;
+	/** @todo anim: should be Skeleton bone index in the future, but right now it's CompactBoneIndex **/
+	FCompactPoseBoneIndex BoneIndex;
 
 	/** Transform to apply **/
-	UPROPERTY()
 	FTransform Transform;
 
 	FBoneTransform() 
 		: BoneIndex(INDEX_NONE)
 	{}
 
-	FBoneTransform( int32 InBoneIndex, const FTransform& InTransform) 
+	FBoneTransform(FCompactPoseBoneIndex InBoneIndex, const FTransform& InTransform)
 		: BoneIndex(InBoneIndex)
 		, Transform(InTransform)
 	{}
@@ -237,8 +216,7 @@ struct FSlotEvaluationPose
 	float Weight;
 
 	/** Pose */
-	UPROPERTY()
-	FA2Pose Pose;
+	FCompactPose Pose;
 	
 	FSlotEvaluationPose()
 	{
@@ -405,22 +383,22 @@ public:
 	// Creates an uninitialized tick record in the list for the correct group or the ungrouped array.  If the group is valid, OutSyncGroupPtr will point to the group.
 	FAnimTickRecord& CreateUninitializedTickRecord(int32 GroupIndex, FAnimGroupInstance*& OutSyncGroupPtr);
 
-	void SequenceEvaluatePose(UAnimSequenceBase* Sequence, struct FA2Pose& Pose, const FAnimExtractContext& ExtractionContext);
+	void SequenceEvaluatePose(UAnimSequenceBase* Sequence, struct FCompactPose& Pose, const FAnimExtractContext& ExtractionContext);
 
-	void BlendSequences(const struct FA2Pose& Pose1, const struct FA2Pose& Pose2, float Alpha, struct FA2Pose& Blended);
+	void BlendSequences(const FCompactPose& Pose1, const FCompactPose& Pose2, float Alpha, FCompactPose& Result);
 
-	static void CopyPose(const struct FA2Pose& Source, struct FA2Pose& Destination);
+	static void CopyPose(const struct FCompactPose& Source, struct FCompactPose& Destination);
 
-	void ApplyAdditiveSequence(const struct FA2Pose& BasePose, const struct FA2Pose& AdditivePose, float Alpha, struct FA2Pose& Blended);
+	void ApplyAdditiveSequence(const struct FCompactPose& BasePose, const struct FCompactPose& AdditivePose, float Alpha, struct FCompactPose& Blended);
 
-	void BlendSpaceEvaluatePose(UBlendSpaceBase* BlendSpace, TArray<FBlendSampleData>& BlendSampleDataCache, struct FA2Pose& Pose);
+	void BlendSpaceEvaluatePose(UBlendSpaceBase* BlendSpace, TArray<FBlendSampleData>& BlendSampleDataCache, struct FCompactPose& OutPose);
 
 	// skeletal control related functions
-	void BlendRotationOffset(const struct FA2Pose& BasePose/* local space base pose */, struct FA2Pose const & RotationOffsetPose/* mesh space rotation only additive **/, float Alpha/*0 means no additive, 1 means whole additive */, struct FA2Pose& Pose /** local space blended pose **/);
+	void BlendRotationOffset(const FCompactPose& BasePose/* local space base pose */, const FCompactPose& RotationOffsetPose/* mesh space rotation only additive **/, float Alpha/*0 means no additive, 1 means whole additive */, struct FCompactPose& OutPose /** local space blended pose **/);
 
 	// slotnode interfaces
 	void GetSlotWeight(FName const & SlotNodeName, float& out_SlotNodeWeight, float& out_SourceWeight) const;
-	void SlotEvaluatePose(FName SlotNodeName, const struct FA2Pose & SourcePose, struct FA2Pose & BlendedPose, float SlotNodeWeight);
+	void SlotEvaluatePose(FName SlotNodeName, const FCompactPose& SourcePose, FCompactPose& BlendedPose, float SlotNodeWeight);
 
 	// slot node run-time functions
 	void ReinitializeSlotNodes();
diff --git a/Engine/Source/Runtime/Engine/Classes/Animation/AnimNodeBase.h b/Engine/Source/Runtime/Engine/Classes/Animation/AnimNodeBase.h
--- a/Engine/Source/Runtime/Engine/Classes/Animation/AnimNodeBase.h
+++ b/Engine/Source/Runtime/Engine/Classes/Animation/AnimNodeBase.h
@@ -6,9 +6,9 @@
 #include "AnimBlueprintGeneratedClass.h"
 #include "AnimInstance.h"
 #include "AnimationRuntime.h"
+#include "BonePose.h"
 #include "AnimNodeBase.generated.h"
 
-
 // Base class for update/evaluate contexts
 struct FAnimationBaseContext
 {
@@ -104,7 +104,7 @@ public:
 struct FPoseContext : public FAnimationBaseContext
 {
 public:
-	FA2Pose Pose;
+	FCompactPose Pose;
 
 public:
 	// This constructor allocates a new uninitialized pose for the specified anim instance
@@ -112,9 +112,7 @@ public:
 		: FAnimationBaseContext(InAnimInstance)
 	{
 		checkSlow( AnimInstance && AnimInstance->RequiredBones.IsValid() );
-		const int32 NumBones = AnimInstance->RequiredBones.GetNumBones();
-		Pose.Bones.Empty(NumBones);
-		Pose.Bones.AddUninitialized(NumBones);
+		Pose.SetBoneContainer(&AnimInstance->RequiredBones);
 	}
 
 	// This constructor allocates a new uninitialized pose, copying non-pose state from the source context
@@ -122,25 +120,28 @@ public:
 		: FAnimationBaseContext(SourceContext.AnimInstance)
 	{
 		checkSlow( AnimInstance && AnimInstance->RequiredBones.IsValid() );
-		const int32 NumBones = AnimInstance->RequiredBones.GetNumBones();
-		Pose.Bones.Empty(NumBones);
-		Pose.Bones.AddUninitialized(NumBones);
+		Pose.SetBoneContainer(&AnimInstance->RequiredBones);
 	}
 
 	void ResetToRefPose()
 	{
-		checkSlow( AnimInstance && AnimInstance->RequiredBones.IsValid() );
-		FAnimationRuntime::FillWithRefPose(Pose.Bones, AnimInstance->RequiredBones);
+		Pose.ResetToRefPose();	
 	}
 
 	void ResetToIdentity()
 	{
-		checkSlow( AnimInstance && AnimInstance->RequiredBones.IsValid() );
-		FAnimationRuntime::InitializeTransform(AnimInstance->RequiredBones, Pose.Bones);
+		Pose.ResetToIdentity();
 	}
 
-	bool ContainsNaN() const;
-	bool IsNormalized() const;
+	bool ContainsNaN() const
+	{
+		return Pose.ContainsNaN();
+	}
+
+	bool IsNormalized() const
+	{
+		return Pose.IsNormalized();
+	}
 };
 
 
@@ -148,7 +149,7 @@ public:
 struct FComponentSpacePoseContext : public FAnimationBaseContext
 {
 public:
-	FA2CSPose Pose;
+	FCSPose<FCompactPose> Pose;
 
 public:
 	// This constructor allocates a new uninitialized pose for the specified anim instance
@@ -168,7 +169,7 @@ public:
 	void ResetToRefPose()
 	{
 		checkSlow( AnimInstance && AnimInstance->RequiredBones.IsValid() );
-		Pose.AllocateLocalPoses(AnimInstance->RequiredBones, AnimInstance->RequiredBones.GetRefPoseArray());
+		Pose.InitPose(&AnimInstance->RequiredBones);
 	}
 
 	bool ContainsNaN() const;
diff --git a/Engine/Source/Runtime/Engine/Classes/Animation/AnimNode_BlendListBase.h b/Engine/Source/Runtime/Engine/Classes/Animation/AnimNode_BlendListBase.h
--- a/Engine/Source/Runtime/Engine/Classes/Animation/AnimNode_BlendListBase.h
+++ b/Engine/Source/Runtime/Engine/Classes/Animation/AnimNode_BlendListBase.h
@@ -26,6 +26,10 @@ protected:
 
 	UPROPERTY()
 	int32 LastActiveChildIndex;
+
+	//Store which poses we need to evaluate
+	TArray<int32> PosesToEvaluate;
+
 public:	
 	FAnimNode_BlendListBase()
 		: LastActiveChildIndex(0)
diff --git a/Engine/Source/Runtime/Engine/Classes/Animation/AnimNode_SaveCachedPose.h b/Engine/Source/Runtime/Engine/Classes/Animation/AnimNode_SaveCachedPose.h
--- a/Engine/Source/Runtime/Engine/Classes/Animation/AnimNode_SaveCachedPose.h
+++ b/Engine/Source/Runtime/Engine/Classes/Animation/AnimNode_SaveCachedPose.h
@@ -25,7 +25,8 @@ struct ENGINE_API FAnimNode_SaveCachedPose : public FAnimNode_Base
 	FPoseLink Pose;
 
 protected:
-	FA2Pose CachedPose;
+	FCompactPose CachedPose;
+
 public:	
 	FAnimNode_SaveCachedPose();
 
diff --git a/Engine/Source/Runtime/Engine/Classes/Animation/AnimNode_TransitionPoseEvaluator.h b/Engine/Source/Runtime/Engine/Classes/Animation/AnimNode_TransitionPoseEvaluator.h
--- a/Engine/Source/Runtime/Engine/Classes/Animation/AnimNode_TransitionPoseEvaluator.h
+++ b/Engine/Source/Runtime/Engine/Classes/Animation/AnimNode_TransitionPoseEvaluator.h
@@ -2,6 +2,7 @@
 
 #pragma once
 #include "Animation/AnimNodeBase.h"
+#include "BonePose.h"
 #include "AnimNode_TransitionPoseEvaluator.generated.h"
 
 // Indicates which state is being evaluated by this node (source or destination)
@@ -49,8 +50,7 @@ struct ENGINE_API FAnimNode_TransitionPoseEvaluator : public FAnimNode_Base
 	UPROPERTY(EditAnywhere, BlueprintReadWrite, Category=Pose, meta=(NeverAsPin, ClampMin="1", UIMin="1"))
 	int32 FramesToCachePose;
 
-	UPROPERTY(transient)
-	FA2Pose CachedPose;
+	FCompactPose CachedPose;
 
 	UPROPERTY(transient)
 	int32 CacheFramesRemaining;
@@ -68,5 +68,5 @@ public:
 
 	bool InputNodeNeedsUpdate() const;
 	bool InputNodeNeedsEvaluate() const;
-	void CachePose(FPoseContext& Output, FA2Pose& PoseToCache);
+	void CachePose(FPoseContext& Output, FCompactPose& PoseToCache);
 };
diff --git a/Engine/Source/Runtime/Engine/Classes/Animation/AnimSequence.h b/Engine/Source/Runtime/Engine/Classes/Animation/AnimSequence.h
--- a/Engine/Source/Runtime/Engine/Classes/Animation/AnimSequence.h
+++ b/Engine/Source/Runtime/Engine/Classes/Animation/AnimSequence.h
@@ -13,6 +13,8 @@
 typedef TArray<FTransform> FTransformArrayA2;
 
 class USkeletalMesh;
+struct FCompactPose;
+
 
 /**
  * Indicates animation data compression format.
@@ -581,64 +583,61 @@ public:
 
 	// Begin Transform related functions 
 
-	/** 
-	 * Get Bone Transform of the Time given, relative to Parent for all RequiredBones
-	 * This returns different transform based on additive or not. Or what kind of additive.
-	 * 
-	 * @param	OutAtoms			[out] Array of output bone transforms
-	 * @param	RequiredBones		Array of Desired Tracks
-	 * @param	ExtractionContext	Extraction Context (position, looping, root motion, etc.)
-	 */
-	void GetAnimationPose( FTransformArrayA2& OutAtoms, const FBoneContainer& RequiredBones, const FAnimExtractContext& ExtractionContext) const;
+	/**
+	* Get Bone Transform of the Time given, relative to Parent for all RequiredBones
+	* This returns different transform based on additive or not. Or what kind of additive.
+	*
+	* @param	OutPose				Pose object to fill
+	* @param	ExtractionContext	Extraction Context (position, looping, root motion, etc.)
+	*/
+	void GetAnimationPose(FCompactPose& OutPose, const FAnimExtractContext& ExtractionContext) const;
 
-	/** 
-	 * Get Bone Transform of the animation for the Time given, relative to Parent for all RequiredBones 
-	 * 
-	 * @param	OutAtoms			[out] Array of output bone transforms
-	 * @param	RequiredBones		Array of Desired Tracks
-	 * @param	ExtractionContext	Extraction Context (position, looping, root motion, etc.)
-	 */
-	ENGINE_API void GetBonePose( FTransformArrayA2& OutAtoms, const FBoneContainer& RequiredBones, const FAnimExtractContext& ExtractionContext) const;
+	/**
+	* Get Bone Transform of the animation for the Time given, relative to Parent for all RequiredBones
+	*
+	* @param	OutPose				[out] Array of output bone transforms
+	* @param	RequiredBones		Array of Desired Tracks
+	* @param	ExtractionContext	Extraction Context (position, looping, root motion, etc.)
+	*/
+	ENGINE_API void GetBonePose(FCompactPose& OutPose, const FAnimExtractContext& ExtractionContext) const;
 
 private:
 	/** 
 	 * Utility function.
 	 * When extracting root motion, this function resets the root bone to the first frame of the animation.
 	 * 
-	 * @param	BoneTransforms		Bone Transforms to read/write root bone transform.
+	 * @param	BoneTransform		Root Bone Transform
 	 * @param	RequiredBones		BoneContainer
 	 * @param	ExtractionContext	Extraction Context to access root motion extraction information.
 	 */
-	void ResetRootBoneForRootMotion(FTransformArrayA2 & BoneTransforms, const FBoneContainer & RequiredBones, ERootMotionRootLock::Type RootMotionRootLock) const;
-	
-	/** 
-	 * Retarget a single bone transform, to apply right after extraction.
-	 * 
-	 * @param	BoneTransform		BoneTransform to read/write from.
-	 * @param	SkeletonBoneIndex	Bone Index in USkeleton.
-	 * @param	PoseBoneIndex		Bone Index in Bone Transform array.
-	 * @param	RequiredBones		BoneContainer
-	 */	
-	void RetargetBoneTransform(FTransform& BoneTransform, const int32& SkeletonBoneIndex, const int32& PoseBoneIndex, const FBoneContainer& RequiredBones) const;
+	void ResetRootBoneForRootMotion(FTransform& BoneTransform, const FBoneContainer & RequiredBones, ERootMotionRootLock::Type RootMotionRootLock) const;
+
+	/**
+	* Retarget a single bone transform, to apply right after extraction.
+	*
+	* @param	BoneTransform		BoneTransform to read/write from.
+	* @param	SkeletonBoneIndex	Bone Index in USkeleton.
+	* @param	BoneIndex			Bone Index in Bone Transform array.
+	* @param	RequiredBones		BoneContainer
+	*/
+	void RetargetBoneTransform(FTransform& BoneTransform, const int32& SkeletonBoneIndex, const FCompactPoseBoneIndex& BoneIndex, const FBoneContainer& RequiredBones) const;
 
 public:
-	/** 
-	 * Get Bone Transform of the additive animation for the Time given, relative to Parent for all RequiredBones 
-	 * 
-	 * @param	OutAtoms			[out] Array of output bone transforms
-	 * @param	RequiredBones		Array of Desired Tracks
-	 * @param	ExtractionContext	Extraction Context (position, looping, root motion, etc.)
-	 */
-	ENGINE_API void GetBonePose_Additive( FTransformArrayA2& OutAtoms, const FBoneContainer& RequiredBones, const FAnimExtractContext& ExtractionContext) const;
+	/**
+	* Get Bone Transform of the additive animation for the Time given, relative to Parent for all RequiredBones
+	*
+	* @param	OutPose				[out] Output bone transforms
+	* @param	ExtractionContext	Extraction Context (position, looping, root motion, etc.)
+	*/
+	ENGINE_API void GetBonePose_Additive(FCompactPose& OutPose, const FAnimExtractContext& ExtractionContext) const;
 
-	/** 
-	 * Get Bone Transform of the base (reference) pose of the additive animation for the Time given, relative to Parent for all RequiredBones 
-	 * 
-	 * @param	OutAtoms			[out] Array of output bone transforms
-	 * @param	RequiredBones		Array of Desired Tracks
-	 * @param	ExtractionContext	Extraction Context (position, looping, root motion, etc.)
-	 */
-	ENGINE_API void GetAdditiveBasePose(FTransformArrayA2 & OutAtoms, const FBoneContainer& RequiredBones, const FAnimExtractContext& ExtractionContext) const;
+	/**
+	* Get Bone Transform of the base (reference) pose of the additive animation for the Time given, relative to Parent for all RequiredBones
+	*
+	* @param	OutPose			[out] Output bone transforms
+	* @param	ExtractionContext	Extraction Context (position, looping, root motion, etc.)
+	*/
+	ENGINE_API void GetAdditiveBasePose(FCompactPose& OutPose, const FAnimExtractContext& ExtractionContext) const;
 
 	/**
 	 * Get Bone Transform of the Time given, relative to Parent for the Track Given
@@ -808,15 +807,14 @@ public:
 #endif
 
 private:
-	/** 
-	 * Get Bone Transform of the animation for the Time given, relative to Parent for all RequiredBones 
-	 * This return mesh rotation only additive pose
-	 * 
-	 * @param	OutAtoms			[out] Array of output bone transforms
-	 * @param	RequiredBones		Array of Desired Tracks
-	 * @param	ExtractionContext	Extraction Context (position, looping, root motion, etc.)
-	 */
-	void GetBonePose_AdditiveMeshRotationOnly( FTransformArrayA2& OutAtoms, const FBoneContainer& RequiredBones, const FAnimExtractContext& ExtractionContext) const;
+	/**
+	* Get Bone Transform of the animation for the Time given, relative to Parent for all RequiredBones
+	* This return mesh rotation only additive pose
+	*
+	* @param	OutPose				[out] Output bone transforms
+	* @param	ExtractionContext	Extraction Context (position, looping, root motion, etc.)
+	*/
+	void GetBonePose_AdditiveMeshRotationOnly(FCompactPose& OutPose, const FAnimExtractContext& ExtractionContext) const;
 
 #if WITH_EDITOR
 	/**
diff --git a/Engine/Source/Runtime/Engine/Classes/Animation/AnimSingleNodeInstance.h b/Engine/Source/Runtime/Engine/Classes/Animation/AnimSingleNodeInstance.h
--- a/Engine/Source/Runtime/Engine/Classes/Animation/AnimSingleNodeInstance.h
+++ b/Engine/Source/Runtime/Engine/Classes/Animation/AnimSingleNodeInstance.h
@@ -63,7 +63,7 @@ class ENGINE_API UAnimSingleNodeInstance : public UAnimInstance
 
 protected:
 	virtual void Montage_Advance(float DeltaTime) override;
-	void InternalBlendSpaceEvaluatePose(class UBlendSpaceBase* BlendSpace, TArray<FBlendSampleData>& BlendSampleDataCache, struct FA2Pose& Pose);
+	void InternalBlendSpaceEvaluatePose(class UBlendSpaceBase* BlendSpace, TArray<FBlendSampleData>& BlendSampleDataCache, FCompactPose& OutPose);
 	// End UAnimInstance interface
 public:
 
diff --git a/Engine/Source/Runtime/Engine/Classes/Animation/BoneControllers/AnimNode_BoneDrivenController.h b/Engine/Source/Runtime/Engine/Classes/Animation/BoneControllers/AnimNode_BoneDrivenController.h
--- a/Engine/Source/Runtime/Engine/Classes/Animation/BoneControllers/AnimNode_BoneDrivenController.h
+++ b/Engine/Source/Runtime/Engine/Classes/Animation/BoneControllers/AnimNode_BoneDrivenController.h
@@ -67,7 +67,7 @@ struct ENGINE_API FAnimNode_BoneDrivenController : public FAnimNode_SkeletalCont
 	//////////////////////////////////////////////////////////////////////////
 
 	// FAnimNode_SkeletalControlBase interface
-	virtual void EvaluateBoneTransforms(USkeletalMeshComponent* SkelComp, const FBoneContainer& RequiredBones, FA2CSPose& MeshBases, TArray<FBoneTransform>& OutBoneTransforms) override;
+	virtual void EvaluateBoneTransforms(USkeletalMeshComponent* SkelComp, FCSPose<FCompactPose>& MeshBases, TArray<FBoneTransform>& OutBoneTransforms) override;
 	virtual bool IsValidToEvaluate(const USkeleton* Skeleton, const FBoneContainer& RequiredBones) override;
 	//////////////////////////////////////////////////////////////////////////
 
diff --git a/Engine/Source/Runtime/Engine/Classes/Animation/BoneControllers/AnimNode_CopyBone.h b/Engine/Source/Runtime/Engine/Classes/Animation/BoneControllers/AnimNode_CopyBone.h
--- a/Engine/Source/Runtime/Engine/Classes/Animation/BoneControllers/AnimNode_CopyBone.h
+++ b/Engine/Source/Runtime/Engine/Classes/Animation/BoneControllers/AnimNode_CopyBone.h
@@ -41,7 +41,7 @@ struct ENGINE_API FAnimNode_CopyBone : public FAnimNode_SkeletalControlBase
 	// End of FAnimNode_Base interface
 
 	// FAnimNode_SkeletalControlBase interface
-	virtual void EvaluateBoneTransforms(USkeletalMeshComponent* SkelComp, const FBoneContainer& RequiredBones, FA2CSPose& MeshBases, TArray<FBoneTransform>& OutBoneTransforms) override;
+	virtual void EvaluateBoneTransforms(USkeletalMeshComponent* SkelComp, FCSPose<FCompactPose>& MeshBases, TArray<FBoneTransform>& OutBoneTransforms) override;
 	virtual bool IsValidToEvaluate(const USkeleton* Skeleton, const FBoneContainer& RequiredBones) override;
 	// End of FAnimNode_SkeletalControlBase interface
 
diff --git a/Engine/Source/Runtime/Engine/Classes/Animation/BoneControllers/AnimNode_Fabrik.h b/Engine/Source/Runtime/Engine/Classes/Animation/BoneControllers/AnimNode_Fabrik.h
--- a/Engine/Source/Runtime/Engine/Classes/Animation/BoneControllers/AnimNode_Fabrik.h
+++ b/Engine/Source/Runtime/Engine/Classes/Animation/BoneControllers/AnimNode_Fabrik.h
@@ -23,7 +23,7 @@ public:
 	float Length;
 
 	/** Bone Index in SkeletalMesh */
-	int32 BoneIndex;
+	FCompactPoseBoneIndex BoneIndex;
 
 	/** Transform Index that this control will output */
 	int32 TransformIndex;
@@ -41,7 +41,7 @@ public:
 	{
 	}
 
-	FABRIKChainLink(FVector const & InPosition, float const & InLength, int32 const & InBoneIndex, int32 const & InTransformIndex)
+	FABRIKChainLink(FVector const & InPosition, float const & InLength, const FCompactPoseBoneIndex& InBoneIndex, int32 const & InTransformIndex)
 		: Position(InPosition)
 		, Length(InLength)
 		, BoneIndex(InBoneIndex)
@@ -94,7 +94,7 @@ public:
 	FAnimNode_Fabrik();
 
 	// FAnimNode_SkeletalControlBase interface
-	virtual void EvaluateBoneTransforms(USkeletalMeshComponent* SkelComp, const FBoneContainer& RequiredBones, FA2CSPose& MeshBases, TArray<FBoneTransform>& OutBoneTransforms) override;
+	virtual void EvaluateBoneTransforms(USkeletalMeshComponent* SkelComp, FCSPose<FCompactPose>& MeshBases, TArray<FBoneTransform>& OutBoneTransforms) override;
 	virtual bool IsValidToEvaluate(const USkeleton* Skeleton, const FBoneContainer& RequiredBones) override;
 	// End of FAnimNode_SkeletalControlBase interface
 
@@ -105,5 +105,5 @@ private:
 	// End of FAnimNode_SkeletalControlBase interface
 
 	// Convenience function to get current (pre-translation iteration) component space location of bone by bone index
-	FVector GetCurrentLocation(FA2CSPose & MeshBases, int32 const & BoneIndex);
+	FVector GetCurrentLocation(FCSPose<FCompactPose>& MeshBases, const FCompactPoseBoneIndex& BoneIndex);
 };
\ No newline at end of file
diff --git a/Engine/Source/Runtime/Engine/Classes/Animation/BoneControllers/AnimNode_HandIKRetargeting.h b/Engine/Source/Runtime/Engine/Classes/Animation/BoneControllers/AnimNode_HandIKRetargeting.h
--- a/Engine/Source/Runtime/Engine/Classes/Animation/BoneControllers/AnimNode_HandIKRetargeting.h
+++ b/Engine/Source/Runtime/Engine/Classes/Animation/BoneControllers/AnimNode_HandIKRetargeting.h
@@ -48,7 +48,7 @@ struct ENGINE_API FAnimNode_HandIKRetargeting : public FAnimNode_SkeletalControl
 	// End of FAnimNode_Base interface
 
 	// FAnimNode_SkeletalControlBase interface
-	virtual void EvaluateBoneTransforms(USkeletalMeshComponent* SkelComp, const FBoneContainer& RequiredBones, FA2CSPose& MeshBases, TArray<FBoneTransform>& OutBoneTransforms) override;
+	virtual void EvaluateBoneTransforms(USkeletalMeshComponent* SkelComp, FCSPose<FCompactPose>& MeshBases, TArray<FBoneTransform>& OutBoneTransforms) override;
 	virtual bool IsValidToEvaluate(const USkeleton* Skeleton, const FBoneContainer& RequiredBones) override;
 	// End of FAnimNode_SkeletalControlBase interface
 
diff --git a/Engine/Source/Runtime/Engine/Classes/Animation/BoneControllers/AnimNode_LookAt.h b/Engine/Source/Runtime/Engine/Classes/Animation/BoneControllers/AnimNode_LookAt.h
--- a/Engine/Source/Runtime/Engine/Classes/Animation/BoneControllers/AnimNode_LookAt.h
+++ b/Engine/Source/Runtime/Engine/Classes/Animation/BoneControllers/AnimNode_LookAt.h
@@ -96,7 +96,7 @@ struct ENGINE_API FAnimNode_LookAt : public FAnimNode_SkeletalControlBase
 	// End of FAnimNode_Base interface
 
 	// FAnimNode_SkeletalControlBase interface
-	virtual void EvaluateBoneTransforms(USkeletalMeshComponent* SkelComp, const FBoneContainer& RequiredBones, FA2CSPose& MeshBases, TArray<FBoneTransform>& OutBoneTransforms) override;
+	virtual void EvaluateBoneTransforms(USkeletalMeshComponent* SkelComp, FCSPose<FCompactPose>& MeshBases, TArray<FBoneTransform>& OutBoneTransforms) override;
 	virtual bool IsValidToEvaluate(const USkeleton* Skeleton, const FBoneContainer& RequiredBones) override;
 	// End of FAnimNode_SkeletalControlBase interface
 
@@ -105,7 +105,7 @@ private:
 	virtual void InitializeBoneReferences(const FBoneContainer& RequiredBones) override;
 	// End of FAnimNode_SkeletalControlBase interface
 
-	FVector GetAlignVector(FTransform& Transform, EAxisOption::Type AxisOption);
+	FVector GetAlignVector(const FTransform& Transform, EAxisOption::Type AxisOption);
 
 	EAlphaBlendType GetInterpolationType()
 	{
diff --git a/Engine/Source/Runtime/Engine/Classes/Animation/BoneControllers/AnimNode_ModifyBone.h b/Engine/Source/Runtime/Engine/Classes/Animation/BoneControllers/AnimNode_ModifyBone.h
--- a/Engine/Source/Runtime/Engine/Classes/Animation/BoneControllers/AnimNode_ModifyBone.h
+++ b/Engine/Source/Runtime/Engine/Classes/Animation/BoneControllers/AnimNode_ModifyBone.h
@@ -73,7 +73,7 @@ struct ENGINE_API FAnimNode_ModifyBone : public FAnimNode_SkeletalControlBase
 	// End of FAnimNode_Base interface
 
 	// FAnimNode_SkeletalControlBase interface
-	virtual void EvaluateBoneTransforms(USkeletalMeshComponent* SkelComp, const FBoneContainer& RequiredBones, FA2CSPose& MeshBases, TArray<FBoneTransform>& OutBoneTransforms) override;
+	virtual void EvaluateBoneTransforms(USkeletalMeshComponent* SkelComp, FCSPose<FCompactPose>& MeshBases, TArray<FBoneTransform>& OutBoneTransforms) override;
 	virtual bool IsValidToEvaluate(const USkeleton* Skeleton, const FBoneContainer& RequiredBones) override;
 	// End of FAnimNode_SkeletalControlBase interface
 
diff --git a/Engine/Source/Runtime/Engine/Classes/Animation/BoneControllers/AnimNode_RotationMultiplier.h b/Engine/Source/Runtime/Engine/Classes/Animation/BoneControllers/AnimNode_RotationMultiplier.h
--- a/Engine/Source/Runtime/Engine/Classes/Animation/BoneControllers/AnimNode_RotationMultiplier.h
+++ b/Engine/Source/Runtime/Engine/Classes/Animation/BoneControllers/AnimNode_RotationMultiplier.h
@@ -48,7 +48,7 @@ struct ENGINE_API FAnimNode_RotationMultiplier : public FAnimNode_SkeletalContro
 	// End of FAnimNode_Base interface
 
 	// FAnimNode_SkeletalControlBase interface
-	virtual void EvaluateBoneTransforms(USkeletalMeshComponent* SkelComp, const FBoneContainer& RequiredBones, FA2CSPose& MeshBases, TArray<FBoneTransform>& OutBoneTransforms) override;
+	virtual void EvaluateBoneTransforms(USkeletalMeshComponent* SkelComp, FCSPose<FCompactPose>& MeshBases, TArray<FBoneTransform>& OutBoneTransforms) override;
 	virtual bool IsValidToEvaluate(const USkeleton* Skeleton, const FBoneContainer& RequiredBones) override;
 	// End of FAnimNode_SkeletalControlBase interface
 
diff --git a/Engine/Source/Runtime/Engine/Classes/Animation/BoneControllers/AnimNode_SkeletalControlBase.h b/Engine/Source/Runtime/Engine/Classes/Animation/BoneControllers/AnimNode_SkeletalControlBase.h
--- a/Engine/Source/Runtime/Engine/Classes/Animation/BoneControllers/AnimNode_SkeletalControlBase.h
+++ b/Engine/Source/Runtime/Engine/Classes/Animation/BoneControllers/AnimNode_SkeletalControlBase.h
@@ -38,7 +38,7 @@ public:
 public:
 #if WITH_EDITORONLY_DATA
 	// forwarded pose data from the wired node which current node's skeletal control is not applied yet
-	FA2CSPose ForwardedPose;
+	FCSPose<FCompactPose> ForwardedPose;
 #endif //#if WITH_EDITORONLY_DATA
 
 	// FAnimNode_Base interface
@@ -52,7 +52,7 @@ protected:
 	// Interface for derived skeletal controls to implement
 
 	// Evaluate the new component-space transforms for the affected bones.
-	virtual void EvaluateBoneTransforms(USkeletalMeshComponent* SkelComp, const FBoneContainer& RequiredBones, FA2CSPose& MeshBases, TArray<FBoneTransform>& OutBoneTransforms) {}
+	virtual void EvaluateBoneTransforms(USkeletalMeshComponent* SkelComp, FCSPose<FCompactPose>& MeshBases, TArray<FBoneTransform>& OutBoneTransforms) {}
 	// return true if it is valid to Evaluate
 	virtual bool IsValidToEvaluate(const USkeleton* Skeleton, const FBoneContainer& RequiredBones) { return false; }
 	// initialize any bone references you have
diff --git a/Engine/Source/Runtime/Engine/Classes/Animation/BoneControllers/AnimNode_SpringBone.h b/Engine/Source/Runtime/Engine/Classes/Animation/BoneControllers/AnimNode_SpringBone.h
--- a/Engine/Source/Runtime/Engine/Classes/Animation/BoneControllers/AnimNode_SpringBone.h
+++ b/Engine/Source/Runtime/Engine/Classes/Animation/BoneControllers/AnimNode_SpringBone.h
@@ -68,7 +68,7 @@ public:
 	// End of FAnimNode_Base interface
 
 	// FAnimNode_SkeletalControlBase interface
-	virtual void EvaluateBoneTransforms(USkeletalMeshComponent* SkelComp, const FBoneContainer& RequiredBones, FA2CSPose& MeshBases, TArray<FBoneTransform>& OutBoneTransforms) override;
+	virtual void EvaluateBoneTransforms(USkeletalMeshComponent* SkelComp, FCSPose<FCompactPose>& MeshBases, TArray<FBoneTransform>& OutBoneTransforms) override;
 	virtual bool IsValidToEvaluate(const USkeleton* Skeleton, const FBoneContainer& RequiredBones) override;
 	// End of FAnimNode_SkeletalControlBase interface
 
diff --git a/Engine/Source/Runtime/Engine/Classes/Animation/BoneControllers/AnimNode_Trail.h b/Engine/Source/Runtime/Engine/Classes/Animation/BoneControllers/AnimNode_Trail.h
--- a/Engine/Source/Runtime/Engine/Classes/Animation/BoneControllers/AnimNode_Trail.h
+++ b/Engine/Source/Runtime/Engine/Classes/Animation/BoneControllers/AnimNode_Trail.h
@@ -71,7 +71,7 @@ struct ENGINE_API FAnimNode_Trail : public FAnimNode_SkeletalControlBase
 	// End of FAnimNode_Base interface
 
 	// FAnimNode_SkeletalControlBase interface
-	virtual void EvaluateBoneTransforms(USkeletalMeshComponent* SkelComp, const FBoneContainer& RequiredBones, FA2CSPose& MeshBases, TArray<FBoneTransform>& OutBoneTransforms) override;
+	virtual void EvaluateBoneTransforms(USkeletalMeshComponent* SkelComp, FCSPose<FCompactPose>& MeshBases, TArray<FBoneTransform>& OutBoneTransforms) override;
 	virtual bool IsValidToEvaluate(const USkeleton* Skeleton, const FBoneContainer& RequiredBones) override;
 	// End of FAnimNode_SkeletalControlBase interface
 
diff --git a/Engine/Source/Runtime/Engine/Classes/Animation/BoneControllers/AnimNode_TwoBoneIK.h b/Engine/Source/Runtime/Engine/Classes/Animation/BoneControllers/AnimNode_TwoBoneIK.h
--- a/Engine/Source/Runtime/Engine/Classes/Animation/BoneControllers/AnimNode_TwoBoneIK.h
+++ b/Engine/Source/Runtime/Engine/Classes/Animation/BoneControllers/AnimNode_TwoBoneIK.h
@@ -65,7 +65,7 @@ struct ENGINE_API FAnimNode_TwoBoneIK : public FAnimNode_SkeletalControlBase
 	// End of FAnimNode_Base interface
 
 	// FAnimNode_SkeletalControlBase interface
-	virtual void EvaluateBoneTransforms(USkeletalMeshComponent* SkelComp, const FBoneContainer& RequiredBones, FA2CSPose& MeshBases, TArray<FBoneTransform>& OutBoneTransforms) override;
+	virtual void EvaluateBoneTransforms(USkeletalMeshComponent* SkelComp, FCSPose<FCompactPose>& MeshBases, TArray<FBoneTransform>& OutBoneTransforms) override;
 	virtual bool IsValidToEvaluate(const USkeleton* Skeleton, const FBoneContainer& RequiredBones) override;
 	// End of FAnimNode_SkeletalControlBase interface
 
diff --git a/Engine/Source/Runtime/Engine/Classes/Animation/BoneControllers/AnimNode_WheelHandler.h b/Engine/Source/Runtime/Engine/Classes/Animation/BoneControllers/AnimNode_WheelHandler.h
--- a/Engine/Source/Runtime/Engine/Classes/Animation/BoneControllers/AnimNode_WheelHandler.h
+++ b/Engine/Source/Runtime/Engine/Classes/Animation/BoneControllers/AnimNode_WheelHandler.h
@@ -34,7 +34,7 @@ struct ENGINE_API FAnimNode_WheelHandler : public FAnimNode_SkeletalControlBase
 	// End of FAnimNode_Base interface
 
 	// FAnimNode_SkeletalControlBase interface
-	virtual void EvaluateBoneTransforms(USkeletalMeshComponent* SkelComp, const FBoneContainer& RequiredBones, FA2CSPose& MeshBases, TArray<FBoneTransform>& OutBoneTransforms) override;
+	virtual void EvaluateBoneTransforms(USkeletalMeshComponent* SkelComp, FCSPose<FCompactPose>& MeshBases, TArray<FBoneTransform>& OutBoneTransforms) override;
 	virtual bool IsValidToEvaluate(const USkeleton* Skeleton, const FBoneContainer& RequiredBones) override;
 	virtual void Update(const FAnimationUpdateContext& Context) override;
 	virtual void Initialize(const FAnimationInitializeContext& Context) override;
diff --git a/Engine/Source/Runtime/Engine/Classes/Engine/SkeletalMesh.h b/Engine/Source/Runtime/Engine/Classes/Engine/SkeletalMesh.h
--- a/Engine/Source/Runtime/Engine/Classes/Engine/SkeletalMesh.h
+++ b/Engine/Source/Runtime/Engine/Classes/Engine/SkeletalMesh.h
@@ -11,6 +11,7 @@
 #include "BoneContainer.h"
 #include "Interfaces/Interface_CollisionDataProvider.h"
 #include "Interfaces/Interface_AssetUserData.h"
+#include "BoneIndices.h"
 #include "SkeletalMesh.generated.h"
 
 /** The maximum number of skeletal mesh LODs allowed. */
@@ -152,6 +153,9 @@ struct FBoneReference
 
 	/** return true if valid. Otherwise return false **/
 	ENGINE_API bool IsValid(const FBoneContainer& RequiredBones) const;
+
+	FMeshPoseBoneIndex GetMeshPoseIndex() const { return FMeshPoseBoneIndex(BoneIndex); }
+	FCompactPoseBoneIndex GetCompactPoseIndex(const FBoneContainer& RequiredBones) const { return RequiredBones.MakeCompactPoseIndex(GetMeshPoseIndex()); }
 };
 
 /**
diff --git a/Engine/Source/Runtime/Engine/Private/Animation/AnimInstance.cpp b/Engine/Source/Runtime/Engine/Private/Animation/AnimInstance.cpp
--- a/Engine/Source/Runtime/Engine/Private/Animation/AnimInstance.cpp
+++ b/Engine/Source/Runtime/Engine/Private/Animation/AnimInstance.cpp
@@ -135,30 +135,28 @@ FAnimTickRecord& UAnimInstance::CreateUninitializedTickRecord(int32 GroupIndex,
 	return *TickRecord;
 }
 
-void UAnimInstance::SequenceEvaluatePose(UAnimSequenceBase* Sequence, FA2Pose& Pose, const FAnimExtractContext& ExtractionContext)
+void UAnimInstance::SequenceEvaluatePose(UAnimSequenceBase* Sequence, struct FCompactPose& Pose, const FAnimExtractContext& ExtractionContext)
 {
 	SCOPE_CYCLE_COUNTER(STAT_AnimNativeEvaluatePoses);
-	checkSlow( RequiredBones.IsValid() );
+	checkSlow(RequiredBones.IsValid());
 
 	USkeletalMeshComponent* Component = GetSkelMeshComponent();
 
 	FAnimExtractContext NewExtractContext(ExtractionContext);
 	NewExtractContext.bExtractRootMotion = RootMotionMode == ERootMotionMode::RootMotionFromEverything || RootMotionMode == ERootMotionMode::IgnoreRootMotion;
 
-	if(const UAnimSequence* AnimSequence = Cast<const UAnimSequence>(Sequence))
+	if (const UAnimSequence* AnimSequence = Cast<const UAnimSequence>(Sequence))
 	{
 		FAnimationRuntime::GetPoseFromSequence(
 			AnimSequence,
-			RequiredBones,
-			/*out*/ Pose.Bones, 
+			/*out*/ Pose,
 			NewExtractContext);
 	}
-	else if(const UAnimComposite* Composite = Cast<const UAnimComposite>(Sequence))
+	else if (const UAnimComposite* Composite = Cast<const UAnimComposite>(Sequence))
 	{
 		FAnimationRuntime::GetPoseFromAnimTrack(
-			Composite->AnimationTrack, 
-			RequiredBones, 
-			/*out*/ Pose.Bones,
+			Composite->AnimationTrack,
+			/*out*/ Pose,
 			NewExtractContext);
 	}
 	else
@@ -166,39 +164,32 @@ void UAnimInstance::SequenceEvaluatePose(UAnimSequenceBase* Sequence, FA2Pose& P
 #if 0
 		UE_LOG(LogAnimation, Log, TEXT("FAnimationRuntime::GetPoseFromSequence - %s - No animation data!"), *GetFName());
 #endif
-		FAnimationRuntime::FillWithRefPose(Pose.Bones, RequiredBones);
+		Pose.ResetToRefPose();
 	}
 }
 
-void UAnimInstance::BlendSequences(const FA2Pose& Pose1, const FA2Pose& Pose2, float Alpha, FA2Pose& Result)
+void UAnimInstance::BlendSequences(const FCompactPose& Pose1, const FCompactPose& Pose2, float Alpha, FCompactPose& Result)
 {
 	SCOPE_CYCLE_COUNTER(STAT_AnimNativeBlendPoses);
 
-	const FTransformArrayA2* Children[2];
-	float Weights[2];
-
-	Children[0] = &(Pose1.Bones);
-	Children[1] = &(Pose2.Bones);
-
 	Alpha = FMath::Clamp<float>(Alpha, 0.0f, 1.0f);
-	Weights[0] = 1.0f - Alpha;
-	Weights[1] = Alpha;
 
-	if (Result.Bones.Num() < Pose1.Bones.Num())
+	if (Result.GetNumBones() < Pose1.GetNumBones())
 	{
-		ensureMsg (false, TEXT("Source Pose has more bones than Target Pose"));
+		ensureMsg(false, TEXT("Source Pose has more bones than Target Pose"));
 		//@hack
-		Result.Bones.AddUninitialized(Pose1.Bones.Num() - Result.Bones.Num());
+		Result.SetBoneContainer(&Pose1.GetBoneContainer());
 	}
-	FAnimationRuntime::BlendPosesTogether(2, Children, Weights, RequiredBones, /*out*/ Result.Bones);
+
+	FAnimationRuntime::BlendTwoPosesTogether(Pose1, Pose2, (1.0f - Alpha), Result);
 }
 
-void UAnimInstance::CopyPose(const FA2Pose& Source, FA2Pose& Destination)
+void UAnimInstance::CopyPose(const FCompactPose& Source, FCompactPose& Destination)
 {
 	if (&Destination != &Source)
 	{
 		SCOPE_CYCLE_COUNTER(STAT_AnimNativeCopyPoses);
-		Destination.Bones = Source.Bones;
+		Destination = Source;
 	}
 }
 
@@ -1044,108 +1035,98 @@ void UAnimInstance::DisplayDebug(class UCanvas* Canvas, const FDebugDisplayInfo&
 	}
 }
 
-void UAnimInstance::BlendSpaceEvaluatePose(class UBlendSpaceBase* BlendSpace, TArray<FBlendSampleData>& BlendSampleDataCache, struct FA2Pose& Pose)
+void UAnimInstance::BlendSpaceEvaluatePose(class UBlendSpaceBase* BlendSpace, TArray<FBlendSampleData>& BlendSampleDataCache, struct FCompactPose& OutPose)
 {
 	SCOPE_CYCLE_COUNTER(STAT_AnimNativeEvaluatePoses);
 
 	FAnimationRuntime::GetPoseFromBlendSpace(
 		BlendSpace,
 		BlendSampleDataCache, 
-		RequiredBones,
-		/*out*/ Pose.Bones);
+		/*out*/ OutPose);
 }
 
-void UAnimInstance::BlendRotationOffset(const struct FA2Pose& BasePose/* local space base pose */, struct FA2Pose const & RotationOffsetPose/* mesh space rotation only additive **/, float Alpha, struct FA2Pose& Pose /** local space blended pose **/)
+void UAnimInstance::BlendRotationOffset(const FCompactPose& BasePose/* local space base pose */, const FCompactPose& RotationOffsetPose/* mesh space rotation only additive **/, float Alpha/*0 means no additive, 1 means whole additive */, struct FCompactPose& OutPose /** local space blended pose **/)
 {
+	//MDW TODO: Can this be improved? Lots of pose creation
 	SCOPE_CYCLE_COUNTER(STAT_AnimNativeBlendPoses);
 
-	check ( RotationOffsetPose.Bones.Num() == RequiredBones.GetNumBones() );
-	check ( BasePose.Bones.Num() == RotationOffsetPose.Bones.Num() );
-	check ( Pose.Bones.Num() == RotationOffsetPose.Bones.Num() );
-
-	FA2Pose BlendedPose;
-	BlendedPose.Bones.AddUninitialized(Pose.Bones.Num());
+	check(RotationOffsetPose.GetNumBones() == RequiredBones.GetCompactPoseNumBones());
+	check(BasePose.GetNumBones() == RotationOffsetPose.GetNumBones());
+	check(OutPose.GetNumBones() == RotationOffsetPose.GetNumBones());
 
 	// now Pose has Mesh based BasePose
 	// apply additive
 	if (Alpha > ZERO_ANIMWEIGHT_THRESH)
 	{
-		FA2Pose MeshBasePose;
-		MeshBasePose.Bones.AddUninitialized(Pose.Bones.Num());
+		FCompactPose BlendedPose;
+		BlendedPose.SetBoneContainer(&RequiredBones);
+
+		FCompactPose MeshBasePose;
+		MeshBasePose.SetBoneContainer(&RequiredBones);
 
 		// note that RotationOffsetPose has MeshSpaceRotation additive but everything else (translation/scale) is local space
 		// First calculate Mesh space for Base Pose
-		const TArray<FBoneIndexType> & RequiredBoneIndices = RequiredBones.GetBoneIndicesArray();
-
-		for (int32 I=0; I<RequiredBoneIndices.Num(); ++I)
+		for (const FCompactPoseBoneIndex BoneIndex : BasePose.ForEachBoneIndex())
 		{
-			int32 BoneIndex = RequiredBoneIndices[I];
-			int32 ParentIndex = RequiredBones.GetParentBoneIndex(BoneIndex);
-			if ( ParentIndex != INDEX_NONE )
+			const FCompactPoseBoneIndex ParentIndex = RequiredBones.GetParentBoneIndex(BoneIndex);
+
+			if (ParentIndex != INDEX_NONE)
 			{
-				MeshBasePose.Bones[BoneIndex] = BasePose.Bones[BoneIndex] * MeshBasePose.Bones[ParentIndex];
+				MeshBasePose[BoneIndex] = BasePose[BoneIndex] * MeshBasePose[ParentIndex];
 			}
 			else
 			{
-				MeshBasePose.Bones[BoneIndex] = BasePose.Bones[BoneIndex];
+				MeshBasePose[BoneIndex] = BasePose[BoneIndex];
 			}
-		}	
-		
+		}
+
 		const ScalarRegister VBlendWeight(Alpha);
-		for (int32 I=0; I<RequiredBoneIndices.Num(); ++I)
+		for (const FCompactPoseBoneIndex BoneIndex : BlendedPose.ForEachBoneIndex())
 		{
-			int32 BoneIndex = RequiredBoneIndices[I];
-			
-			FTransform& Result = BlendedPose.Bones[BoneIndex];
+			FTransform& Result = BlendedPose[BoneIndex];
 
 			// We want Base pose (local Pose)
-			Result = BasePose.Bones[BoneIndex];
+			Result = BasePose[BoneIndex];
 
 			// set result rotation to be mesh space rotation, so that it applys to mesh space rotation
-			Result.SetRotation(MeshBasePose.Bones[BoneIndex].GetRotation());
+			Result.SetRotation(MeshBasePose[BoneIndex].GetRotation());
 
 			// @fixme laurent - we should make a read only version so we can avoid the copy.
-			FTransform Additive = RotationOffsetPose.Bones[BoneIndex];
+			FTransform Additive = RotationOffsetPose[BoneIndex];
 			FTransform::BlendFromIdentityAndAccumulate(Result, Additive, VBlendWeight);
 		}
 
 		// Ensure that all of the resulting rotations are normalized
-		FAnimationRuntime::NormalizeRotations(RequiredBones, BlendedPose.Bones);
+		BlendedPose.NormalizeRotations();
 
 		// now convert back to Local
-		for(int32 I=0; I<RequiredBoneIndices.Num(); ++I)
+		for (const FCompactPoseBoneIndex BoneIndex : BlendedPose.ForEachBoneIndex())
 		{
-			int32 BoneIndex = RequiredBoneIndices[I];
-			int32 ParentIndex = RequiredBones.GetParentBoneIndex(BoneIndex);
+			FCompactPoseBoneIndex ParentIndex = RequiredBones.GetParentBoneIndex(BoneIndex);
+
+			OutPose[BoneIndex] = BlendedPose[BoneIndex];
 
-			Pose.Bones[BoneIndex] = BlendedPose.Bones[BoneIndex];
-			if(ParentIndex != INDEX_NONE)
+			if (ParentIndex != INDEX_NONE)
 			{
 				// convert to local space first
-				FQuat Rotation = BlendedPose.Bones[ParentIndex].GetRotation().Inverse() * BlendedPose.Bones[BoneIndex].GetRotation();
-				Pose.Bones[BoneIndex].SetRotation(Rotation);
+				FQuat Rotation = BlendedPose[ParentIndex].GetRotation().Inverse() * BlendedPose[BoneIndex].GetRotation();
+				OutPose[BoneIndex].SetRotation(Rotation);
 			}
 		}
 	}
 	else
 	{
-		BlendedPose = BasePose;
+		OutPose = BasePose;
 	}
 }
 
-void UAnimInstance::ApplyAdditiveSequence(const struct FA2Pose& BasePose,const struct FA2Pose& AdditivePose,float Alpha,struct FA2Pose& Blended)
+void UAnimInstance::ApplyAdditiveSequence(const struct FCompactPose& BasePose, const struct FCompactPose& AdditivePose, float Alpha, struct FCompactPose& Blended)
 {
-	if (Blended.Bones.Num() < BasePose.Bones.Num())
-	{
-		// see if this happens
-		ensureMsg (false, TEXT("BasePose has more bones than Blended pose"));
-		//@hack
-		Blended.Bones.AddUninitialized(BasePose.Bones.Num() - Blended.Bones.Num());
-	}
-
+	check(BasePose.GetNumBones() == Blended.GetNumBones());
+	
 	float BlendWeight = FMath::Clamp<float>(Alpha, 0.f, 1.f);
 
-	FAnimationRuntime::BlendAdditivePose(BasePose.Bones, AdditivePose.Bones, BlendWeight, RequiredBones, Blended.Bones);
+	FAnimationRuntime::BlendAdditivePose(BasePose, AdditivePose, BlendWeight, Blended);
 }
 
 void UAnimInstance::RecalcRequiredBones()
@@ -1478,11 +1459,12 @@ void UAnimInstance::GetSlotWeight(FName const & SlotNodeName, float& out_SlotNod
 	out_SourceWeight = 1.f - NonAdditiveTotalWeight;
 }
 
-void UAnimInstance::SlotEvaluatePose(FName SlotNodeName, const FA2Pose & SourcePose, FA2Pose & BlendedPose, float SlotNodeWeight)
+void UAnimInstance::SlotEvaluatePose(FName SlotNodeName, const FCompactPose& SourcePose, FCompactPose& BlendedPose, float SlotNodeWeight)
 {
 	//Accessing MontageInstances from this function is not safe (as this can be called during Parallel Anim Evaluation!
 	//Any montage data you need to add should be part of MontageEvaluationData
 
+	//MDW TODO
 	SCOPE_CYCLE_COUNTER(STAT_AnimNativeEvaluatePoses);
 	if (SlotNodeWeight <= ZERO_ANIMWEIGHT_THRESH)
 	{
@@ -1511,11 +1493,11 @@ void UAnimInstance::SlotEvaluatePose(FName SlotNodeName, const FA2Pose & SourceP
 			FSlotEvaluationPose NewPose(EvalState.MontageWeight, AdditiveAnimType);
 			
 			// Bone array has to be allocated prior to calling GetPoseFromAnimTrack
-			NewPose.Pose.Bones.AddUninitialized(RequiredBones.GetNumBones());
+			NewPose.Pose.SetBoneContainer(&RequiredBones);
 
 			// Extract pose from Track
 			FAnimExtractContext ExtractionContext(EvalState.MontagePosition, EvalState.Montage->HasRootMotion() && RootMotionMode != ERootMotionMode::NoRootMotionExtraction);
-			FAnimationRuntime::GetPoseFromAnimTrack(*AnimTrack, RequiredBones, NewPose.Pose.Bones, ExtractionContext);
+			FAnimationRuntime::GetPoseFromAnimTrack(*AnimTrack, NewPose.Pose, ExtractionContext);
 
 			TotalWeight += EvalState.MontageWeight;
 			if (AdditiveAnimType == AAT_None)
@@ -1572,27 +1554,26 @@ void UAnimInstance::SlotEvaluatePose(FName SlotNodeName, const FA2Pose & SourceP
 			float const SourceWeight = FMath::Clamp<float>(1.f - NonAdditiveWeight, 0.f, 1.f);
 			int32 const NumPoses = NonAdditivePoses.Num() + ((SourceWeight > ZERO_ANIMWEIGHT_THRESH) ? 1 : 0);
 
-			FTransformArrayA2 const ** BlendingPoses = new FTransformArrayA2 const *[NumPoses];
+			TArray<const FCompactPose*> BlendingPoses;
+			BlendingPoses.AddUninitialized(NumPoses);
+
 			TArray<float> BlendWeights;
 			BlendWeights.AddUninitialized(NumPoses);
 			for (int32 Index = 0; Index < NonAdditivePoses.Num(); Index++)
 			{
-				BlendingPoses[Index] = &NonAdditivePoses[Index].Pose.Bones;
+				BlendingPoses[Index] = &NonAdditivePoses[Index].Pose;
 				BlendWeights[Index] = NonAdditivePoses[Index].Weight;
 			}
 
 			if (SourceWeight > ZERO_ANIMWEIGHT_THRESH)
 			{
 				int32 const SourceIndex = BlendWeights.Num() - 1;
-				BlendingPoses[SourceIndex] = &SourcePose.Bones;
+				BlendingPoses[SourceIndex] = &SourcePose;
 				BlendWeights[SourceIndex] = SourceWeight;
 			}
 
 			// now time to blend all montages
-			FAnimationRuntime::BlendPosesTogether(BlendWeights.Num(), (const FTransformArrayA2**)BlendingPoses, (const float*)BlendWeights.GetData(), RequiredBones, BlendedPose.Bones);
-
-			// clean up memory
-			delete[] BlendingPoses;
+			FAnimationRuntime::BlendPosesTogether(BlendingPoses, BlendWeights, BlendedPose);
 		}
 	}
 
diff --git a/Engine/Source/Runtime/Engine/Private/Animation/AnimNodeBase.cpp b/Engine/Source/Runtime/Engine/Private/Animation/AnimNodeBase.cpp
--- a/Engine/Source/Runtime/Engine/Private/Animation/AnimNodeBase.cpp
+++ b/Engine/Source/Runtime/Engine/Private/Animation/AnimNodeBase.cpp
@@ -151,74 +151,17 @@ void FComponentSpacePoseLink::EvaluateComponentSpace(FComponentSpacePoseContext&
 	checkSlow( Output.IsNormalized() );
 }
 
-/////////////////////////////////////////////////////
-// FPoseContext
-
-bool FPoseContext::ContainsNaN() const
-{
-	checkSlow( AnimInstance && AnimInstance->RequiredBones.IsValid() );
-	const TArray<FBoneIndexType> & RequiredBoneIndices = AnimInstance->RequiredBones.GetBoneIndicesArray();
-	for (int32 Iter = 0; Iter < RequiredBoneIndices.Num(); ++Iter)
-	{
-		const int32 BoneIndex = RequiredBoneIndices[Iter];
-		if (Pose.Bones[BoneIndex].ContainsNaN())
-		{
-			return true;
-		}
-	}
-
-	return false;
-}
-
-bool FPoseContext::IsNormalized() const
-{
-	checkSlow( AnimInstance && AnimInstance->RequiredBones.IsValid() );
-	const TArray<FBoneIndexType> & RequiredBoneIndices = AnimInstance->RequiredBones.GetBoneIndicesArray();
-	for (int32 Iter = 0; Iter < RequiredBoneIndices.Num(); ++Iter)
-	{
-		const int32 BoneIndex = RequiredBoneIndices[Iter];
-		if( !Pose.Bones[BoneIndex].IsRotationNormalized() )
-		{
-			return false;
-		}
-	}
-
-	return true;
-}
-
 /////////////////////////////////////////////////////
 // FComponentSpacePoseContext
 
 bool FComponentSpacePoseContext::ContainsNaN() const
 {
-	checkSlow( AnimInstance && AnimInstance->RequiredBones.IsValid() );
-	const TArray<FBoneIndexType> & RequiredBoneIndices = AnimInstance->RequiredBones.GetBoneIndicesArray();
-	for (int32 Iter = 0; Iter < RequiredBoneIndices.Num(); ++Iter)
-	{
-		const int32 BoneIndex = RequiredBoneIndices[Iter];
-		if (Pose.Bones[BoneIndex].ContainsNaN())
-		{
-			return true;
-		}
-	}
-
-	return false;
+	return Pose.GetPose().ContainsNaN();
 }
 
 bool FComponentSpacePoseContext::IsNormalized() const
 {
-	checkSlow( AnimInstance && AnimInstance->RequiredBones.IsValid() );
-	const TArray<FBoneIndexType> & RequiredBoneIndices = AnimInstance->RequiredBones.GetBoneIndicesArray();
-	for (int32 Iter = 0; Iter < RequiredBoneIndices.Num(); ++Iter)
-	{
-		const int32 BoneIndex = RequiredBoneIndices[Iter];
-		if( !Pose.Bones[BoneIndex].IsRotationNormalized() )
-		{
-			return false;
-		}
-	}
-
-	return true;
+	return Pose.GetPose().IsNormalized();
 }
 
 /////////////////////////////////////////////////////
diff --git a/Engine/Source/Runtime/Engine/Private/Animation/AnimNodeRotationMultiplier.cpp b/Engine/Source/Runtime/Engine/Private/Animation/AnimNodeRotationMultiplier.cpp
--- a/Engine/Source/Runtime/Engine/Private/Animation/AnimNodeRotationMultiplier.cpp
+++ b/Engine/Source/Runtime/Engine/Private/Animation/AnimNodeRotationMultiplier.cpp
@@ -120,18 +120,22 @@ FQuat FAnimNode_RotationMultiplier::MultiplyQuatBasedOnSourceIndex(const FTransf
 	return OutQuat;
 }
 
-void FAnimNode_RotationMultiplier::EvaluateBoneTransforms(USkeletalMeshComponent* SkelComp, const FBoneContainer& RequiredBones, FA2CSPose& MeshBases, TArray<FBoneTransform>& OutBoneTransforms)
+void FAnimNode_RotationMultiplier::EvaluateBoneTransforms(USkeletalMeshComponent* SkelComp, FCSPose<FCompactPose>& MeshBases, TArray<FBoneTransform>& OutBoneTransforms)
 {
 	check(OutBoneTransforms.Num() == 0);
 
 	if ( Multiplier != 0.f )
 	{
 		// Reference bone
-		const TArray<FTransform> & LocalRefPose = RequiredBones.GetRefPoseArray();
-		const FQuat RefQuat = LocalRefPose[TargetBone.BoneIndex].GetRotation();
-		FQuat NewQuat = MultiplyQuatBasedOnSourceIndex(LocalRefPose[SourceBone.BoneIndex], MeshBases.GetLocalSpaceTransform(SourceBone.BoneIndex), RotationAxisToRefer, Multiplier, RefQuat);
+		const FBoneContainer& BoneContainer = MeshBases.GetPose().GetBoneContainer();
+		const FCompactPoseBoneIndex TargetBoneIndex = TargetBone.GetCompactPoseIndex(BoneContainer);
+		const FCompactPoseBoneIndex SourceBoneIndex = SourceBone.GetCompactPoseIndex(BoneContainer);
 
-		FTransform NewLocalTransform = MeshBases.GetLocalSpaceTransform(TargetBone.BoneIndex);
+		const FQuat RefQuat = MeshBases.GetPose().GetRefPose(TargetBoneIndex).GetRotation();
+		const FTransform& SourceRefPose = MeshBases.GetPose().GetRefPose(SourceBoneIndex);
+		FQuat NewQuat = MultiplyQuatBasedOnSourceIndex(SourceRefPose, MeshBases.GetLocalSpaceTransform(SourceBoneIndex), RotationAxisToRefer, Multiplier, RefQuat);
+
+		FTransform NewLocalTransform = MeshBases.GetLocalSpaceTransform(TargetBoneIndex);
 		
 		if (bIsAdditive)
 		{
@@ -140,16 +144,16 @@ void FAnimNode_RotationMultiplier::EvaluateBoneTransforms(USkeletalMeshComponent
 		
 		NewLocalTransform.SetRotation(NewQuat);
 
-		const int32 ParentIndex = RequiredBones.GetParentBoneIndex(TargetBone.BoneIndex);
+		const FCompactPoseBoneIndex ParentIndex = MeshBases.GetPose().GetParentBoneIndex(TargetBoneIndex);
 		if( ParentIndex != INDEX_NONE )
 		{
-			const FTransform ParentTM = MeshBases.GetComponentSpaceTransform(ParentIndex);
+			const FTransform& ParentTM = MeshBases.GetComponentSpaceTransform(ParentIndex);
 			FTransform NewTransform = NewLocalTransform * ParentTM;
-			OutBoneTransforms.Add( FBoneTransform(TargetBone.BoneIndex, NewTransform) );
+			OutBoneTransforms.Add( FBoneTransform(TargetBoneIndex, NewTransform) );
 		}
 		else
 		{
-			OutBoneTransforms.Add( FBoneTransform(TargetBone.BoneIndex, NewLocalTransform) );
+			OutBoneTransforms.Add( FBoneTransform(TargetBoneIndex, NewLocalTransform) );
 		}
 	}
 }
diff --git a/Engine/Source/Runtime/Engine/Private/Animation/AnimNodeSkeletalControlBase.cpp b/Engine/Source/Runtime/Engine/Private/Animation/AnimNodeSkeletalControlBase.cpp
--- a/Engine/Source/Runtime/Engine/Private/Animation/AnimNodeSkeletalControlBase.cpp
+++ b/Engine/Source/Runtime/Engine/Private/Animation/AnimNodeSkeletalControlBase.cpp
@@ -55,7 +55,7 @@ void FAnimNode_SkeletalControlBase::EvaluateComponentSpace(FComponentSpacePoseCo
 #endif // #if WITH_EDITORONLY_DATA
 
 		TArray<FBoneTransform> BoneTransforms;
-		EvaluateBoneTransforms(Component, Output.AnimInstance->RequiredBones, Output.Pose, BoneTransforms);
+		EvaluateBoneTransforms(Component, Output.Pose, BoneTransforms);
 
 		checkSlow(!ContainsNaN(BoneTransforms));
 
diff --git a/Engine/Source/Runtime/Engine/Private/Animation/AnimNodeSpaceConversions.cpp b/Engine/Source/Runtime/Engine/Private/Animation/AnimNodeSpaceConversions.cpp
--- a/Engine/Source/Runtime/Engine/Private/Animation/AnimNodeSpaceConversions.cpp
+++ b/Engine/Source/Runtime/Engine/Private/Animation/AnimNodeSpaceConversions.cpp
@@ -31,7 +31,7 @@ void FAnimNode_ConvertComponentToLocalSpace::Evaluate(FPoseContext & Output)
 	FComponentSpacePoseContext InputCSPose(Output.AnimInstance);
 	ComponentPose.EvaluateComponentSpace(InputCSPose);
 
-	checkSlow( InputCSPose.Pose.IsValid() );
+	checkSlow( InputCSPose.Pose.GetPose().IsValid() );
 	InputCSPose.Pose.ConvertToLocalPoses(Output.Pose);
 }
 
@@ -77,5 +77,5 @@ void FAnimNode_ConvertLocalToComponentSpace::EvaluateComponentSpace(FComponentSp
 	FPoseContext InputPose(OutputCSPose.AnimInstance);
 	LocalPose.Evaluate(InputPose);
 
-	OutputCSPose.Pose.AllocateLocalPoses(OutputCSPose.AnimInstance->RequiredBones, InputPose.Pose);
+	OutputCSPose.Pose.InitPose(InputPose.Pose);
 }
diff --git a/Engine/Source/Runtime/Engine/Private/Animation/AnimNode_BlendListBase.cpp b/Engine/Source/Runtime/Engine/Private/Animation/AnimNode_BlendListBase.cpp
--- a/Engine/Source/Runtime/Engine/Private/Animation/AnimNode_BlendListBase.cpp
+++ b/Engine/Source/Runtime/Engine/Private/Animation/AnimNode_BlendListBase.cpp
@@ -46,6 +46,8 @@ void FAnimNode_BlendListBase::Update(const FAnimationUpdateContext& Context)
 	const int NumPoses = BlendPose.Num();
 	checkSlow((BlendTime.Num() == NumPoses) && (BlendWeights.Num() == NumPoses));
 
+	PosesToEvaluate.Empty(NumPoses);
+
 	if (NumPoses > 0)
 	{
 		// Handle a change in the active child index; adjusting the target weights
@@ -103,6 +105,7 @@ void FAnimNode_BlendListBase::Update(const FAnimationUpdateContext& Context)
 			if (BlendWeight > ZERO_ANIMWEIGHT_THRESH)
 			{
 				BlendPose[i].Update(Context.FractionalWeight(BlendWeight));
+				PosesToEvaluate.Add(i);
 			}
 		}
 	}
@@ -112,38 +115,32 @@ void FAnimNode_BlendListBase::Evaluate(FPoseContext& Output)
 {
 	SCOPE_CYCLE_COUNTER(STAT_AnimNativeBlendPoses);
 
-	const int32 MaxNumPoses = BlendPose.Num();
+	const int32 NumPoses = PosesToEvaluate.Num();
 
-	//@TODO: This is currently using O(NumPoses) memory but doesn't need to
-	if ((MaxNumPoses > 0) && (BlendPose.Num() == BlendWeights.Num()))
+	if ((NumPoses > 0) && (BlendPose.Num() == BlendWeights.Num()))
 	{
-		TArray<FPoseContext> FilteredPoseContexts;
-		FilteredPoseContexts.Empty(MaxNumPoses);
+		TArray<FCompactPose> FilteredPoses;
+		FilteredPoses.SetNum(NumPoses);
 
-		FTransformArrayA2** FilteredPoses = new FTransformArrayA2*[MaxNumPoses];
-		float* FilteredWeights = new float[MaxNumPoses];
+		TArray<float> FilteredWeights;
+		FilteredWeights.AddUninitialized(NumPoses);
 
 		int32 NumActivePoses = 0;
-		for (int32 i = 0; i < BlendPose.Num(); ++i)
+		for (int32 i = 0; i < PosesToEvaluate.Num(); ++i)
 		{
-			const float BlendWeight = BlendWeights[i];
-			if (BlendWeight > ZERO_ANIMWEIGHT_THRESH)
-			{
-				FPoseContext& CurrentPoseContext = *(new (FilteredPoseContexts) FPoseContext(Output));
+			int32 PoseIndex = PosesToEvaluate[i];
 
-				FPoseLink& CurrentPose = BlendPose[i];
-				CurrentPose.Evaluate(CurrentPoseContext);
+			const float BlendWeight = BlendWeights[PoseIndex];
+			FPoseContext EvaluateContext(Output);
 
-				FilteredPoses[NumActivePoses] = &(CurrentPoseContext.Pose.Bones);
-				FilteredWeights[NumActivePoses] = BlendWeight;
-				NumActivePoses++;
-			}
-		}
+			FPoseLink& CurrentPose = BlendPose[PoseIndex];
+			CurrentPose.Evaluate(EvaluateContext);
 
-		FAnimationRuntime::BlendPosesTogether(NumActivePoses, (const FTransformArrayA2**)FilteredPoses, (const float*)FilteredWeights, Output.AnimInstance->RequiredBones, Output.Pose.Bones);
+			FilteredPoses[i].MoveBonesFrom(EvaluateContext.Pose);
+			FilteredWeights[i] = BlendWeight;
+		}
 
-		delete[] FilteredPoses;
-		delete[] FilteredWeights;
+		FAnimationRuntime::BlendPosesTogether(FilteredPoses, FilteredWeights, Output.Pose);
 	}
 	else
 	{
diff --git a/Engine/Source/Runtime/Engine/Private/Animation/AnimNode_BoneDrivenController.cpp b/Engine/Source/Runtime/Engine/Private/Animation/AnimNode_BoneDrivenController.cpp
--- a/Engine/Source/Runtime/Engine/Private/Animation/AnimNode_BoneDrivenController.cpp
+++ b/Engine/Source/Runtime/Engine/Private/Animation/AnimNode_BoneDrivenController.cpp
@@ -25,7 +25,7 @@ void FAnimNode_BoneDrivenController::GatherDebugData(FNodeDebugData& DebugData)
 	ComponentPose.GatherDebugData(DebugData);
 }
 
-void FAnimNode_BoneDrivenController::EvaluateBoneTransforms(USkeletalMeshComponent* SkelComp, const FBoneContainer& RequiredBones, FA2CSPose& MeshBases, TArray<FBoneTransform>& OutBoneTransforms)
+void FAnimNode_BoneDrivenController::EvaluateBoneTransforms(USkeletalMeshComponent* SkelComp, FCSPose<FCompactPose>& MeshBases, TArray<FBoneTransform>& OutBoneTransforms)
 {
 	check(OutBoneTransforms.Num() == 0);
 	
@@ -36,9 +36,9 @@ void FAnimNode_BoneDrivenController::EvaluateBoneTransforms(USkeletalMeshCompone
 	}
 
 	// Get the Local space transform and the ref pose transform to see how the transform for the source bone has changed
-	const TArray<FTransform>& RefPose = RequiredBones.GetRefPoseArray();
-	FTransform SourceOrigRef = RefPose[SourceBone.BoneIndex];
-	FTransform SourceCurr = MeshBases.GetLocalSpaceTransform(SourceBone.BoneIndex);
+	const FBoneContainer& BoneContainer = MeshBases.GetPose().GetBoneContainer();
+	const FTransform& SourceOrigRef = BoneContainer.GetRefPoseArray()[SourceBone.BoneIndex];
+	FTransform SourceCurr = MeshBases.GetLocalSpaceTransform(SourceBone.GetCompactPoseIndex(BoneContainer));
 
 	// Difference from the ref pose
 	FVector RotationDiff = (SourceCurr.GetRotation() * SourceOrigRef.GetRotation().Inverse()).Euler();
@@ -46,9 +46,6 @@ void FAnimNode_BoneDrivenController::EvaluateBoneTransforms(USkeletalMeshCompone
 	FVector CurrentScale = SourceCurr.GetScale3D();
 	FVector RefScale = SourceOrigRef.GetScale3D();
 	float ScaleDiff = FMath::Max3(CurrentScale[0], CurrentScale[1], CurrentScale[2]) - FMath::Max3(RefScale[0], RefScale[1], RefScale[2]);
-
-	// Starting point for the new transform
-	FTransform NewLocal = MeshBases.GetLocalSpaceTransform(TargetBone.BoneIndex);
 	
 	// Difference to apply after processing
 	FVector NewRotDiff(FVector::ZeroVector);
@@ -98,19 +95,25 @@ void FAnimNode_BoneDrivenController::EvaluateBoneTransforms(USkeletalMeshCompone
 
 	// Build final transform difference
 	FTransform FinalDiff(FQuat::MakeFromEuler(NewRotDiff), NewTransDiff, FVector(NewScaleDiff));
+
+	FCompactPoseBoneIndex TargetBoneIndex = TargetBone.GetCompactPoseIndex(BoneContainer);
+
+	// Starting point for the new transform
+	FTransform NewLocal = MeshBases.GetLocalSpaceTransform(TargetBoneIndex);
 	NewLocal.AccumulateWithAdditiveScale3D(FinalDiff);
 
 	// If we have a parent, concatenate the transform, otherwise just take the new transform
-	int32 PIdx = RequiredBones.GetParentBoneIndex(TargetBone.BoneIndex);
+	FCompactPoseBoneIndex PIdx = MeshBases.GetPose().GetParentBoneIndex(TargetBoneIndex);
+
 	if(PIdx != INDEX_NONE)
 	{
-		FTransform ParentTM = MeshBases.GetComponentSpaceTransform(PIdx);
+		const FTransform& ParentTM = MeshBases.GetComponentSpaceTransform(PIdx);
 		
-		OutBoneTransforms.Add(FBoneTransform(TargetBone.BoneIndex, NewLocal * ParentTM));
+		OutBoneTransforms.Add(FBoneTransform(TargetBoneIndex, NewLocal * ParentTM));
 	}
 	else
 	{
-		OutBoneTransforms.Add(FBoneTransform(TargetBone.BoneIndex, NewLocal));
+		OutBoneTransforms.Add(FBoneTransform(TargetBoneIndex, NewLocal));
 	}
 
 }
diff --git a/Engine/Source/Runtime/Engine/Private/Animation/AnimNode_CopyBone.cpp b/Engine/Source/Runtime/Engine/Private/Animation/AnimNode_CopyBone.cpp
--- a/Engine/Source/Runtime/Engine/Private/Animation/AnimNode_CopyBone.cpp
+++ b/Engine/Source/Runtime/Engine/Private/Animation/AnimNode_CopyBone.cpp
@@ -25,7 +25,7 @@ void FAnimNode_CopyBone::GatherDebugData(FNodeDebugData& DebugData)
 	ComponentPose.GatherDebugData(DebugData);
 }
 
-void FAnimNode_CopyBone::EvaluateBoneTransforms(USkeletalMeshComponent* SkelComp, const FBoneContainer& RequiredBones, FA2CSPose& MeshBases, TArray<FBoneTransform>& OutBoneTransforms)
+void FAnimNode_CopyBone::EvaluateBoneTransforms(USkeletalMeshComponent* SkelComp, FCSPose<FCompactPose>& MeshBases, TArray<FBoneTransform>& OutBoneTransforms)
 {
 	check(OutBoneTransforms.Num() == 0);
 
@@ -36,8 +36,11 @@ void FAnimNode_CopyBone::EvaluateBoneTransforms(USkeletalMeshComponent* SkelComp
 	}
 
 	// Get component space transform for source and current bone.
-	FTransform SourceBoneTM = MeshBases.GetComponentSpaceTransform(SourceBone.BoneIndex);
-	FTransform CurrentBoneTM = MeshBases.GetComponentSpaceTransform(TargetBone.BoneIndex);
+	const FBoneContainer& BoneContainer = MeshBases.GetPose().GetBoneContainer();
+	FCompactPoseBoneIndex TargetBoneIndex = TargetBone.GetCompactPoseIndex(BoneContainer);
+
+	const FTransform& SourceBoneTM = MeshBases.GetComponentSpaceTransform(SourceBone.GetCompactPoseIndex(BoneContainer));
+	FTransform CurrentBoneTM = MeshBases.GetComponentSpaceTransform(TargetBoneIndex);
 
 	// Copy individual components
 	if (bCopyTranslation)
@@ -56,7 +59,7 @@ void FAnimNode_CopyBone::EvaluateBoneTransforms(USkeletalMeshComponent* SkelComp
 	}
 
 	// Output new transform for current bone.
-	OutBoneTransforms.Add( FBoneTransform(TargetBone.BoneIndex, CurrentBoneTM) );
+	OutBoneTransforms.Add(FBoneTransform(TargetBoneIndex, CurrentBoneTM));
 }
 
 bool FAnimNode_CopyBone::IsValidToEvaluate(const USkeleton* Skeleton, const FBoneContainer& RequiredBones) 
diff --git a/Engine/Source/Runtime/Engine/Private/Animation/AnimNode_Fabrik.cpp b/Engine/Source/Runtime/Engine/Private/Animation/AnimNode_Fabrik.cpp
--- a/Engine/Source/Runtime/Engine/Private/Animation/AnimNode_Fabrik.cpp
+++ b/Engine/Source/Runtime/Engine/Private/Animation/AnimNode_Fabrik.cpp
@@ -18,19 +18,19 @@ FAnimNode_Fabrik::FAnimNode_Fabrik()
 {
 }
 
-FVector FAnimNode_Fabrik::GetCurrentLocation(FA2CSPose & MeshBases, int32 const & BoneIndex)
+FVector FAnimNode_Fabrik::GetCurrentLocation(FCSPose<FCompactPose>& MeshBases, const FCompactPoseBoneIndex& BoneIndex)
 {
 	return MeshBases.GetComponentSpaceTransform(BoneIndex).GetLocation();
 }
 
-void FAnimNode_Fabrik::EvaluateBoneTransforms(USkeletalMeshComponent* SkelComp, const FBoneContainer& RequiredBones, FA2CSPose& MeshBases, TArray<FBoneTransform>& OutBoneTransforms)
+void FAnimNode_Fabrik::EvaluateBoneTransforms(USkeletalMeshComponent* SkelComp, FCSPose<FCompactPose>& MeshBases, TArray<FBoneTransform>& OutBoneTransforms)
 {
-	// IsValidToEvaluate validated our inputs, we don't need to check pre-requisites again.
-	int32 const RootIndex = RootBone.BoneIndex;
+	const FBoneContainer& BoneContainer = MeshBases.GetPose().GetBoneContainer();
 
 	// Update EffectorLocation if it is based off a bone position
-	FTransform CSEffectorTransform = FTransform(EffectorTransform);
-	FAnimationRuntime::ConvertBoneSpaceTransformToCS(SkelComp, MeshBases, CSEffectorTransform, EffectorTransformBone.BoneIndex, EffectorTransformSpace);
+	FTransform CSEffectorTransform = EffectorTransform;
+	FAnimationRuntime::ConvertBoneSpaceTransformToCS(SkelComp, MeshBases, CSEffectorTransform, EffectorTransformBone.GetCompactPoseIndex(BoneContainer), EffectorTransformSpace);
+	
 	FVector const CSEffectorLocation = CSEffectorTransform.GetLocation();
 
 	// @fixme - we need better to draw widgets and debug information in editor.
@@ -44,13 +44,15 @@ void FAnimNode_Fabrik::EvaluateBoneTransforms(USkeletalMeshComponent* SkelComp,
 #endif
 
 	// Gather all bone indices between root and tip.
-	TArray<int32> BoneIndices;
+	TArray<FCompactPoseBoneIndex> BoneIndices;
+	
 	{
-		int32 BoneIndex = TipBone.BoneIndex;
-		do 
+		const FCompactPoseBoneIndex RootIndex = RootBone.GetCompactPoseIndex(BoneContainer);
+		FCompactPoseBoneIndex BoneIndex = TipBone.GetCompactPoseIndex(BoneContainer);
+		do
 		{
 			BoneIndices.Insert(BoneIndex, 0);
-			BoneIndex = RequiredBones.GetParentBoneIndex(BoneIndex);
+			BoneIndex = MeshBases.GetPose().GetParentBoneIndex(BoneIndex);
 		} while (BoneIndex != RootIndex);
 		BoneIndices.Insert(BoneIndex, 0);
 	}
@@ -68,20 +70,21 @@ void FAnimNode_Fabrik::EvaluateBoneTransforms(USkeletalMeshComponent* SkelComp,
 
 	// Start with Root Bone
 	{
-		int32 const & RootBoneIndex = BoneIndices[0];
-		FTransform const BoneCSTransform = MeshBases.GetComponentSpaceTransform(RootBoneIndex);
-		OutBoneTransforms[0] = FBoneTransform(RootBoneIndex, BoneCSTransform);
+		const FCompactPoseBoneIndex& RootBoneIndex = BoneIndices[0];
+		const FTransform& BoneCSTransform = MeshBases.GetComponentSpaceTransform(RootBoneIndex);
 
+		OutBoneTransforms[0] = FBoneTransform(RootBoneIndex, BoneCSTransform);
 		Chain.Add(FABRIKChainLink(BoneCSTransform.GetLocation(), 0.f, RootBoneIndex, 0));
 	}
 
 	// Go through remaining transforms
 	for (int32 TransformIndex = 1; TransformIndex < NumTransforms; TransformIndex++)
 	{
-		int32 const & BoneIndex = BoneIndices[TransformIndex];
+		const FCompactPoseBoneIndex& BoneIndex = BoneIndices[TransformIndex];
 
-		FTransform const BoneCSTransform = MeshBases.GetComponentSpaceTransform(BoneIndex);
+		const FTransform& BoneCSTransform = MeshBases.GetComponentSpaceTransform(BoneIndex);
 		FVector const BoneCSPosition = BoneCSTransform.GetLocation();
+
 		OutBoneTransforms[TransformIndex] = FBoneTransform(BoneIndex, BoneCSTransform);
 
 		// Calculate the combined length of this segment of skeleton
diff --git a/Engine/Source/Runtime/Engine/Private/Animation/AnimNode_HandIKRetargeting.cpp b/Engine/Source/Runtime/Engine/Private/Animation/AnimNode_HandIKRetargeting.cpp
--- a/Engine/Source/Runtime/Engine/Private/Animation/AnimNode_HandIKRetargeting.cpp
+++ b/Engine/Source/Runtime/Engine/Private/Animation/AnimNode_HandIKRetargeting.cpp
@@ -28,15 +28,16 @@ void FAnimNode_HandIKRetargeting::GatherDebugData(FNodeDebugData& DebugData)
 	ComponentPose.GatherDebugData(DebugData);
 }
 
-void FAnimNode_HandIKRetargeting::EvaluateBoneTransforms(USkeletalMeshComponent* SkelComp, const FBoneContainer& RequiredBones, FA2CSPose& MeshBases, TArray<FBoneTransform>& OutBoneTransforms)
+void FAnimNode_HandIKRetargeting::EvaluateBoneTransforms(USkeletalMeshComponent* SkelComp, FCSPose<FCompactPose>& MeshBases, TArray<FBoneTransform>& OutBoneTransforms)
 {
 	checkSlow(OutBoneTransforms.Num() == 0);
 
+	const FBoneContainer& BoneContainer = MeshBases.GetPose().GetBoneContainer();
 	// Get component space transforms for all of our IK and FK bones.
-	FTransform const RightHandFKTM = MeshBases.GetComponentSpaceTransform(RightHandFK.BoneIndex);
-	FTransform const LeftHandFKTM = MeshBases.GetComponentSpaceTransform(LeftHandFK.BoneIndex);
-	FTransform const RightHandIKTM = MeshBases.GetComponentSpaceTransform(RightHandIK.BoneIndex);
-	FTransform const LeftHandIKTM = MeshBases.GetComponentSpaceTransform(LeftHandIK.BoneIndex);
+	const FTransform& RightHandFKTM = MeshBases.GetComponentSpaceTransform(RightHandFK.GetCompactPoseIndex(BoneContainer));
+	const FTransform& LeftHandFKTM = MeshBases.GetComponentSpaceTransform(LeftHandFK.GetCompactPoseIndex(BoneContainer));
+	const FTransform& RightHandIKTM = MeshBases.GetComponentSpaceTransform(RightHandIK.GetCompactPoseIndex(BoneContainer));
+	const FTransform& LeftHandIKTM = MeshBases.GetComponentSpaceTransform(LeftHandIK.GetCompactPoseIndex(BoneContainer));
 	
 	// Compute weight FK and IK hand location. And translation from IK to FK.
 	FVector const FKLocation = FMath::Lerp<FVector>(LeftHandFKTM.GetTranslation(), RightHandFKTM.GetTranslation(), HandFKWeight);
@@ -50,12 +51,13 @@ void FAnimNode_HandIKRetargeting::EvaluateBoneTransforms(USkeletalMeshComponent*
 		for (int32 BoneIndex = 0; BoneIndex < IKBonesToMove.Num(); BoneIndex++)
 		{
 			FBoneReference const & BoneReference = IKBonesToMove[BoneIndex];
-			if (BoneReference.IsValid(RequiredBones))
+			if (BoneReference.IsValid(BoneContainer))
 			{
-				FTransform BoneTransform = MeshBases.GetComponentSpaceTransform(BoneReference.BoneIndex);
+				FCompactPoseBoneIndex BoneIndex = BoneReference.GetCompactPoseIndex(BoneContainer);
+				FTransform BoneTransform = MeshBases.GetComponentSpaceTransform(BoneIndex);
 				BoneTransform.AddToTranslation(IK_To_FK_Translation);
 
-				OutBoneTransforms.Add(FBoneTransform(BoneReference.BoneIndex, BoneTransform));
+				OutBoneTransforms.Add(FBoneTransform(BoneIndex, BoneTransform));
 			}
 		}
 	}
diff --git a/Engine/Source/Runtime/Engine/Private/Animation/AnimNode_LayeredBoneBlend.cpp b/Engine/Source/Runtime/Engine/Private/Animation/AnimNode_LayeredBoneBlend.cpp
--- a/Engine/Source/Runtime/Engine/Private/Animation/AnimNode_LayeredBoneBlend.cpp
+++ b/Engine/Source/Runtime/Engine/Private/Animation/AnimNode_LayeredBoneBlend.cpp
@@ -87,7 +87,9 @@ void FAnimNode_LayeredBoneBlend::Evaluate(FPoseContext& Output)
 		// evaluate children
 		BasePose.Evaluate(BasePoseContext);
 
-		TArray<FA2Pose> TargetBlendPoses;
+		TArray<FCompactPose> TargetBlendPoses;
+		TargetBlendPoses.SetNum(NumPoses);
+
 		for (int32 ChildIndex = 0; ChildIndex < NumPoses; ++ChildIndex)
 		{
 			if (BlendWeights[ChildIndex] > ZERO_ANIMWEIGHT_THRESH)
@@ -95,19 +97,15 @@ void FAnimNode_LayeredBoneBlend::Evaluate(FPoseContext& Output)
 				FPoseContext CurrentPoseContext(Output);
 				BlendPoses[ChildIndex].Evaluate(CurrentPoseContext);
 
-				// @todo fixme : change this to not copy
-				TargetBlendPoses.Add(CurrentPoseContext.Pose);
+				TargetBlendPoses[ChildIndex].MoveBonesFrom(CurrentPoseContext.Pose);
 			}
 			else
 			{
-				//Add something here so that array ordering is maintained
-				FA2Pose Pose;
-				FAnimationRuntime::FillWithRefPose(Pose.Bones, Output.AnimInstance->RequiredBones);
-				TargetBlendPoses.Add(Pose);
+				TargetBlendPoses[ChildIndex].ResetToRefPose(BasePoseContext.Pose.GetBoneContainer());
 			}
 		}
 
-		FAnimationRuntime::BlendPosesPerBoneFilter(BasePoseContext.Pose, TargetBlendPoses, Output.Pose, CurrentBoneBlendWeights, bMeshSpaceRotationBlend, Output.AnimInstance->RequiredBones, Output.AnimInstance->CurrentSkeleton);
+		FAnimationRuntime::BlendPosesPerBoneFilter(BasePoseContext.Pose, TargetBlendPoses, Output.Pose, CurrentBoneBlendWeights, bMeshSpaceRotationBlend);
 	}
 }
 
diff --git a/Engine/Source/Runtime/Engine/Private/Animation/AnimNode_LookAt.cpp b/Engine/Source/Runtime/Engine/Private/Animation/AnimNode_LookAt.cpp
--- a/Engine/Source/Runtime/Engine/Private/Animation/AnimNode_LookAt.cpp
+++ b/Engine/Source/Runtime/Engine/Private/Animation/AnimNode_LookAt.cpp
@@ -46,17 +46,19 @@ void DrawDebugData(UWorld* World, const FTransform& TransformVector, const FVect
 	DrawDebugLine(World, Start, End, Color);
 }
 
-void FAnimNode_LookAt::EvaluateBoneTransforms(USkeletalMeshComponent* SkelComp, const FBoneContainer& RequiredBones, FA2CSPose& MeshBases, TArray<FBoneTransform>& OutBoneTransforms)
+void FAnimNode_LookAt::EvaluateBoneTransforms(USkeletalMeshComponent* SkelComp, FCSPose<FCompactPose>& MeshBases, TArray<FBoneTransform>& OutBoneTransforms)
 {
 	check(OutBoneTransforms.Num() == 0);
 
-	FTransform ComponentBoneTransform = MeshBases.GetComponentSpaceTransform(BoneToModify.BoneIndex);
+	const FBoneContainer& BoneContainer = MeshBases.GetPose().GetBoneContainer();
+	const FCompactPoseBoneIndex ModifyBoneIndex = BoneToModify.GetCompactPoseIndex(BoneContainer);
+	FTransform ComponentBoneTransform = MeshBases.GetComponentSpaceTransform(ModifyBoneIndex);
 
 	// get target location
 	FVector TargetLocationInComponentSpace;
-	if (LookAtBone.IsValid(RequiredBones))
+	if (LookAtBone.IsValid(BoneContainer))
 	{
-		FTransform LookAtTransform  = MeshBases.GetComponentSpaceTransform(LookAtBone.BoneIndex);
+		const FTransform& LookAtTransform  = MeshBases.GetComponentSpaceTransform(LookAtBone.GetCompactPoseIndex(BoneContainer));
 		TargetLocationInComponentSpace = LookAtTransform.GetLocation();
 	}
 	else
@@ -145,7 +147,7 @@ void FAnimNode_LookAt::EvaluateBoneTransforms(USkeletalMeshComponent* SkelComp,
 	FQuat NewRotation = DeltaRot * CurrentRot;
 	ComponentBoneTransform.SetRotation(NewRotation);
 
-	OutBoneTransforms.Add( FBoneTransform(BoneToModify.BoneIndex, ComponentBoneTransform) );
+	OutBoneTransforms.Add(FBoneTransform(ModifyBoneIndex, ComponentBoneTransform));
 }
 
 bool FAnimNode_LookAt::IsValidToEvaluate(const USkeleton* Skeleton, const FBoneContainer& RequiredBones) 
@@ -164,7 +166,7 @@ void FAnimNode_LookAt::InitializeBoneReferences(const FBoneContainer& RequiredBo
 	LookAtBone.Initialize(RequiredBones);
 }
 
-FVector FAnimNode_LookAt::GetAlignVector(FTransform& Transform, EAxisOption::Type AxisOption)
+FVector FAnimNode_LookAt::GetAlignVector(const FTransform& Transform, EAxisOption::Type AxisOption)
 {
 	switch (AxisOption)
 	{
diff --git a/Engine/Source/Runtime/Engine/Private/Animation/AnimNode_ModifyBone.cpp b/Engine/Source/Runtime/Engine/Private/Animation/AnimNode_ModifyBone.cpp
--- a/Engine/Source/Runtime/Engine/Private/Animation/AnimNode_ModifyBone.cpp
+++ b/Engine/Source/Runtime/Engine/Private/Animation/AnimNode_ModifyBone.cpp
@@ -32,18 +32,22 @@ void FAnimNode_ModifyBone::GatherDebugData(FNodeDebugData& DebugData)
 	ComponentPose.GatherDebugData(DebugData);
 }
 
-void FAnimNode_ModifyBone::EvaluateBoneTransforms(USkeletalMeshComponent* SkelComp, const FBoneContainer& RequiredBones, FA2CSPose& MeshBases, TArray<FBoneTransform>& OutBoneTransforms)
+void FAnimNode_ModifyBone::EvaluateBoneTransforms(USkeletalMeshComponent* SkelComp, FCSPose<FCompactPose>& MeshBases, TArray<FBoneTransform>& OutBoneTransforms)
 {
 	check(OutBoneTransforms.Num() == 0);
 
 	// the way we apply transform is same as FMatrix or FTransform
 	// we apply scale first, and rotation, and translation
-	// if you'd like to translate first, you'll need two nodes that first node does translate and second nodes to rotate. 
-	FTransform NewBoneTM = MeshBases.GetComponentSpaceTransform(BoneToModify.BoneIndex);
+	// if you'd like to translate first, you'll need two nodes that first node does translate and second nodes to rotate.
+	const FBoneContainer BoneContainer = MeshBases.GetPose().GetBoneContainer();
+
+	FCompactPoseBoneIndex CompactPoseBoneToModify = BoneToModify.GetCompactPoseIndex(BoneContainer);
+	FTransform NewBoneTM = MeshBases.GetComponentSpaceTransform(CompactPoseBoneToModify);
+	
 	if (ScaleMode != BMM_Ignore)
 	{
 		// Convert to Bone Space.
-		FAnimationRuntime::ConvertCSTransformToBoneSpace(SkelComp, MeshBases, NewBoneTM, BoneToModify.BoneIndex, ScaleSpace);
+		FAnimationRuntime::ConvertCSTransformToBoneSpace(SkelComp, MeshBases, NewBoneTM, CompactPoseBoneToModify, ScaleSpace);
 
 		if (ScaleMode == BMM_Additive)
 		{
@@ -55,13 +59,13 @@ void FAnimNode_ModifyBone::EvaluateBoneTransforms(USkeletalMeshComponent* SkelCo
 		}
 
 		// Convert back to Component Space.
-		FAnimationRuntime::ConvertBoneSpaceTransformToCS(SkelComp, MeshBases, NewBoneTM, BoneToModify.BoneIndex, ScaleSpace);
+		FAnimationRuntime::ConvertBoneSpaceTransformToCS(SkelComp, MeshBases, NewBoneTM, CompactPoseBoneToModify, ScaleSpace);
 	}
 
 	if (RotationMode != BMM_Ignore)
 	{
 		// Convert to Bone Space.
-		FAnimationRuntime::ConvertCSTransformToBoneSpace(SkelComp, MeshBases, NewBoneTM, BoneToModify.BoneIndex, RotationSpace);
+		FAnimationRuntime::ConvertCSTransformToBoneSpace(SkelComp, MeshBases, NewBoneTM, CompactPoseBoneToModify, RotationSpace);
 
 		const FQuat BoneQuat(Rotation);
 		if (RotationMode == BMM_Additive)
@@ -74,13 +78,13 @@ void FAnimNode_ModifyBone::EvaluateBoneTransforms(USkeletalMeshComponent* SkelCo
 		}
 
 		// Convert back to Component Space.
-		FAnimationRuntime::ConvertBoneSpaceTransformToCS(SkelComp, MeshBases, NewBoneTM, BoneToModify.BoneIndex, RotationSpace);
+		FAnimationRuntime::ConvertBoneSpaceTransformToCS(SkelComp, MeshBases, NewBoneTM, CompactPoseBoneToModify, RotationSpace);
 	}
-
+	
 	if (TranslationMode != BMM_Ignore)
 	{
 		// Convert to Bone Space.
-		FAnimationRuntime::ConvertCSTransformToBoneSpace(SkelComp, MeshBases, NewBoneTM, BoneToModify.BoneIndex, TranslationSpace);
+		FAnimationRuntime::ConvertCSTransformToBoneSpace(SkelComp, MeshBases, NewBoneTM, CompactPoseBoneToModify, TranslationSpace);
 
 		if (TranslationMode == BMM_Additive)
 		{
@@ -92,10 +96,10 @@ void FAnimNode_ModifyBone::EvaluateBoneTransforms(USkeletalMeshComponent* SkelCo
 		}
 
 		// Convert back to Component Space.
-		FAnimationRuntime::ConvertBoneSpaceTransformToCS(SkelComp, MeshBases, NewBoneTM, BoneToModify.BoneIndex, TranslationSpace);
+		FAnimationRuntime::ConvertBoneSpaceTransformToCS(SkelComp, MeshBases, NewBoneTM, CompactPoseBoneToModify, TranslationSpace);
 	}
-
-	OutBoneTransforms.Add( FBoneTransform(BoneToModify.BoneIndex, NewBoneTM) );
+	
+	OutBoneTransforms.Add( FBoneTransform(BoneToModify.GetCompactPoseIndex(BoneContainer), NewBoneTM) );
 }
 
 bool FAnimNode_ModifyBone::IsValidToEvaluate(const USkeleton* Skeleton, const FBoneContainer& RequiredBones) 
diff --git a/Engine/Source/Runtime/Engine/Private/Animation/AnimNode_RotateRootBone.cpp b/Engine/Source/Runtime/Engine/Private/Animation/AnimNode_RotateRootBone.cpp
--- a/Engine/Source/Runtime/Engine/Private/Animation/AnimNode_RotateRootBone.cpp
+++ b/Engine/Source/Runtime/Engine/Private/Animation/AnimNode_RotateRootBone.cpp
@@ -39,7 +39,8 @@ void FAnimNode_RotateRootBone::Evaluate(FPoseContext& Output)
 	const FQuat MeshSpaceDeltaQuat = MeshToComponentQuat.Inverse() * DeltaQuat * MeshToComponentQuat;
 
 	// Apply rotation to root bone.
-	Output.Pose.Bones[0].SetRotation(Output.Pose.Bones[0].GetRotation() * MeshSpaceDeltaQuat);
+	FCompactPoseBoneIndex RootBone(0);
+	Output.Pose[RootBone].SetRotation(Output.Pose[RootBone].GetRotation() * MeshSpaceDeltaQuat);
 }
 
 
diff --git a/Engine/Source/Runtime/Engine/Private/Animation/AnimNode_SaveCachedPose.cpp b/Engine/Source/Runtime/Engine/Private/Animation/AnimNode_SaveCachedPose.cpp
--- a/Engine/Source/Runtime/Engine/Private/Animation/AnimNode_SaveCachedPose.cpp
+++ b/Engine/Source/Runtime/Engine/Private/Animation/AnimNode_SaveCachedPose.cpp
@@ -55,7 +55,7 @@ void FAnimNode_SaveCachedPose::Evaluate(FPoseContext& Output)
 
 		FPoseContext CachingContext(Output);
 		Pose.Evaluate(CachingContext);
-		CachedPose = CachingContext.Pose;
+		CachedPose.MoveBonesFrom(CachingContext.Pose);
 	}
 
 	// Return the cached result
diff --git a/Engine/Source/Runtime/Engine/Private/Animation/AnimNode_SpringBone.cpp b/Engine/Source/Runtime/Engine/Private/Animation/AnimNode_SpringBone.cpp
--- a/Engine/Source/Runtime/Engine/Private/Animation/AnimNode_SpringBone.cpp
+++ b/Engine/Source/Runtime/Engine/Private/Animation/AnimNode_SpringBone.cpp
@@ -57,13 +57,16 @@ void FAnimNode_SpringBone::GatherDebugData(FNodeDebugData& DebugData)
 	ComponentPose.GatherDebugData(DebugData);
 }
 
-void FAnimNode_SpringBone::EvaluateBoneTransforms(USkeletalMeshComponent* SkelComp, const FBoneContainer& RequiredBones, FA2CSPose& MeshBases, TArray<FBoneTransform>& OutBoneTransforms)
+void FAnimNode_SpringBone::EvaluateBoneTransforms(USkeletalMeshComponent* SkelComp, FCSPose<FCompactPose>& MeshBases, TArray<FBoneTransform>& OutBoneTransforms)
 {
 	check(OutBoneTransforms.Num() == 0);
 
 	// Location of our bone in world space
-	FTransform SpaceBase = MeshBases.GetComponentSpaceTransform(SpringBone.BoneIndex);
-	FTransform  BoneTransformInWorldSpace = (SkelComp != NULL) ? SpaceBase * SkelComp->GetComponentToWorld() : SpaceBase;
+	const FBoneContainer& BoneContainer = MeshBases.GetPose().GetBoneContainer();
+
+	const FCompactPoseBoneIndex SpringBoneIndex = SpringBone.GetCompactPoseIndex(BoneContainer);
+	const FTransform& SpaceBase = MeshBases.GetComponentSpaceTransform(SpringBoneIndex);
+	FTransform BoneTransformInWorldSpace = (SkelComp != NULL) ? SpaceBase * SkelComp->GetComponentToWorld() : SpaceBase;
 
 	FVector const TargetPos = BoneTransformInWorldSpace.GetLocation();
 
@@ -158,7 +161,7 @@ void FAnimNode_SpringBone::EvaluateBoneTransforms(USkeletalMeshComponent* SkelCo
 	OutBoneTM.SetLocation( SkelComp->GetComponentToWorld().InverseTransformPosition(BoneLocation) );
 
 	// Output new transform for current bone.
-	OutBoneTransforms.Add( FBoneTransform(SpringBone.BoneIndex, OutBoneTM) );
+	OutBoneTransforms.Add(FBoneTransform(SpringBoneIndex, OutBoneTM));
 }
 
 
diff --git a/Engine/Source/Runtime/Engine/Private/Animation/AnimNode_StateMachine.cpp b/Engine/Source/Runtime/Engine/Private/Animation/AnimNode_StateMachine.cpp
--- a/Engine/Source/Runtime/Engine/Private/Animation/AnimNode_StateMachine.cpp
+++ b/Engine/Source/Runtime/Engine/Private/Animation/AnimNode_StateMachine.cpp
@@ -559,13 +559,13 @@ void FAnimNode_StateMachine::Evaluate(FPoseContext& Output)
 	{
 		if (Machine->States.Num() == 0)
 		{
-			FAnimationRuntime::FillWithRefPose(Output.Pose.Bones, Output.AnimInstance->RequiredBones);
+			Output.Pose.ResetToRefPose();
 			return;
 		}
 	}
 	else
 	{
-		FAnimationRuntime::FillWithRefPose(Output.Pose.Bones, Output.AnimInstance->RequiredBones);
+		Output.Pose.ResetToRefPose();
 		return;
 	}
 
@@ -575,12 +575,6 @@ void FAnimNode_StateMachine::Evaluate(FPoseContext& Output)
 	{
 		check(Output.AnimInstance->CurrentSkeleton);
 
-		// Create an accumulator pose
-		const int32 NumBones = Output.AnimInstance->RequiredBones.GetNumBones();
-
-		Output.Pose.Bones.Empty(NumBones);
-		Output.Pose.Bones.AddUninitialized(NumBones);
-
 		//each transition stomps over the last because they will already include the output from the transition before it
 		for (int32 Index = 0; Index < ActiveTransitionArray.Num(); ++Index)
 		{
@@ -608,7 +602,7 @@ void FAnimNode_StateMachine::Evaluate(FPoseContext& Output)
 		}
 
 		// Ensure that all of the resulting rotations are normalized
-		FAnimationRuntime::NormalizeRotations(Output.AnimInstance->RequiredBones, Output.Pose.Bones);
+		Output.Pose.NormalizeRotations();
 	}
 	else if (!IsAConduitState(CurrentState))
 	{
@@ -637,12 +631,10 @@ void FAnimNode_StateMachine::EvaluateTransitionStandardBlend(FPoseContext& Outpu
 	// Blend it in
 	const ScalarRegister VPreviousWeight(1.0f - Transition.Alpha);
 	const ScalarRegister VWeight(Transition.Alpha);
-	const TArray<FBoneIndexType> & RequiredBoneIndices = Output.AnimInstance->RequiredBones.GetBoneIndicesArray();
-	for (int32 j = 0; j < RequiredBoneIndices.Num(); ++j)
+	for (FCompactPoseBoneIndex BoneIndex : Output.Pose.ForEachBoneIndex())
 	{
-		const int32 BoneIndex = RequiredBoneIndices[j];
-		Output.Pose.Bones[BoneIndex] = PreviouseStateResult.Pose.Bones[BoneIndex] * VPreviousWeight;
-		Output.Pose.Bones[BoneIndex].AccumulateWithShortestRotation(NextStateResult.Pose.Bones[BoneIndex], VWeight);
+		Output.Pose[BoneIndex] = PreviouseStateResult.Pose[BoneIndex] * VPreviousWeight;
+		Output.Pose[BoneIndex].AccumulateWithShortestRotation(NextStateResult.Pose[BoneIndex], VWeight);
 	}
 }
 
@@ -682,11 +674,9 @@ void FAnimNode_StateMachine::EvaluateTransitionCustomBlend(FPoseContext& Output,
 		Transition.CustomTransitionGraph.Evaluate(StatePoseResult);
 
 		// First pose will just overwrite the destination
-		const TArray<FBoneIndexType> & RequiredBoneIndices = Output.AnimInstance->RequiredBones.GetBoneIndicesArray();
-		for (int32 j = 0; j < RequiredBoneIndices.Num(); ++j)
+		for (const FCompactPoseBoneIndex BoneIndex : Output.Pose.ForEachBoneIndex())
 		{
-			const int32 BoneIndex = RequiredBoneIndices[j];
-			Output.Pose.Bones[BoneIndex] = StatePoseResult.Pose.Bones[BoneIndex];
+			Output.Pose[BoneIndex] = StatePoseResult.Pose[BoneIndex];
 		}
 	}
 }
@@ -872,9 +862,7 @@ void FAnimNode_StateMachine::UpdateState(int32 StateIndex, const FAnimationUpdat
 
 void FAnimNode_StateMachine::EvaluateState(int32 StateIndex, FPoseContext& Output)
 {
-	const int32 NumBones = Output.AnimInstance->RequiredBones.GetNumBones();
-	Output.Pose.Bones.Empty(NumBones);
-	Output.Pose.Bones.AddUninitialized(NumBones);
+	Output.Pose.SetBoneContainer(&Output.AnimInstance->RequiredBones);
 
 	if (!IsAConduitState(StateIndex))
 	{
diff --git a/Engine/Source/Runtime/Engine/Private/Animation/AnimNode_Trail.cpp b/Engine/Source/Runtime/Engine/Private/Animation/AnimNode_Trail.cpp
--- a/Engine/Source/Runtime/Engine/Private/Animation/AnimNode_Trail.cpp
+++ b/Engine/Source/Runtime/Engine/Private/Animation/AnimNode_Trail.cpp
@@ -38,7 +38,7 @@ void FAnimNode_Trail::GatherDebugData(FNodeDebugData& DebugData)
 	ComponentPose.GatherDebugData(DebugData);
 }
 
-void FAnimNode_Trail::EvaluateBoneTransforms(USkeletalMeshComponent* SkelComp, const FBoneContainer& RequiredBones, FA2CSPose& MeshBases, TArray<FBoneTransform>& OutBoneTransforms)
+void FAnimNode_Trail::EvaluateBoneTransforms(USkeletalMeshComponent* SkelComp, FCSPose<FCompactPose>& MeshBases, TArray<FBoneTransform>& OutBoneTransforms)
 {
 	check(OutBoneTransforms.Num() == 0);
 
@@ -50,9 +50,10 @@ void FAnimNode_Trail::EvaluateBoneTransforms(USkeletalMeshComponent* SkelComp, c
 	// The incoming BoneIndex is the 'end' of the spline chain. We need to find the 'start' by walking SplineLength bones up hierarchy.
 	// Fail if we walk past the root bone.
 
-	int32 WalkBoneIndex = TrailBone.BoneIndex;
+	const FBoneContainer& BoneContainer = MeshBases.GetPose().GetBoneContainer();
+	FCompactPoseBoneIndex WalkBoneIndex = TrailBone.GetCompactPoseIndex(BoneContainer);
 
-	TArray<int32> ChainBoneIndices;
+	TArray<FCompactPoseBoneIndex> ChainBoneIndices;
 	ChainBoneIndices.AddZeroed(ChainLength);
 
 	ChainBoneIndices[ChainLength - 1] = WalkBoneIndex;
@@ -67,7 +68,7 @@ void FAnimNode_Trail::EvaluateBoneTransforms(USkeletalMeshComponent* SkelComp, c
 		}
 
 		// Get parent bone.
-		WalkBoneIndex = RequiredBones.GetParentBoneIndex(WalkBoneIndex);
+		WalkBoneIndex = BoneContainer.GetParentBoneIndex(WalkBoneIndex);
 
 		//Insert indices at the start of array, so that parents are before children in the array.
 		int32 TransformIndex = ChainLength - (i + 1);
@@ -86,8 +87,8 @@ void FAnimNode_Trail::EvaluateBoneTransforms(USkeletalMeshComponent* SkelComp, c
 
 		for(int32 i=0; i<ChainBoneIndices.Num(); i++)
 		{
-			int32 ChildIndex = ChainBoneIndices[i];
-			FTransform ChainTransform = MeshBases.GetComponentSpaceTransform(ChildIndex);
+			FCompactPoseBoneIndex ChildIndex = ChainBoneIndices[i];
+			const FTransform& ChainTransform = MeshBases.GetComponentSpaceTransform(ChildIndex);
 			TrailBoneLocations[i] = ChainTransform.GetTranslation();
 		}
 		OldLocalToWorld = SkelComp->GetTransformMatrix();
@@ -117,8 +118,8 @@ void FAnimNode_Trail::EvaluateBoneTransforms(USkeletalMeshComponent* SkelComp, c
 	}
 
 	// Root bone of trail is not modified.
-	int32 RootIndex = ChainBoneIndices[0]; 
-	FTransform ChainTransform = MeshBases.GetComponentSpaceTransform(RootIndex);
+	FCompactPoseBoneIndex RootIndex = ChainBoneIndices[0]; 
+	const FTransform& ChainTransform = MeshBases.GetComponentSpaceTransform(RootIndex);
 	OutBoneTransforms[0] = FBoneTransform(RootIndex, ChainTransform);
 	TrailBoneLocations[0] = ChainTransform.GetTranslation();
 
@@ -126,12 +127,12 @@ void FAnimNode_Trail::EvaluateBoneTransforms(USkeletalMeshComponent* SkelComp, c
 	for(int32 i=1; i<ChainBoneIndices.Num(); i++)
 	{
 		// Parent bone position in component space.
-		int32 ParentIndex = ChainBoneIndices[i-1];
+		FCompactPoseBoneIndex ParentIndex = ChainBoneIndices[i - 1];
 		FVector ParentPos = TrailBoneLocations[i-1];
 		FVector ParentAnimPos = MeshBases.GetComponentSpaceTransform(ParentIndex).GetTranslation();
 
 		// Child bone position in component space.
-		int32 ChildIndex = ChainBoneIndices[i];
+		FCompactPoseBoneIndex ChildIndex = ChainBoneIndices[i];
 		FVector ChildPos = OldToNewTM.TransformPosition(TrailBoneLocations[i]); // move from 'last frames component' frame to 'this frames component' frame
 		FVector ChildAnimPos = MeshBases.GetComponentSpaceTransform(ChildIndex).GetTranslation();
 
diff --git a/Engine/Source/Runtime/Engine/Private/Animation/AnimNode_TransitionPoseEvaluator.cpp b/Engine/Source/Runtime/Engine/Private/Animation/AnimNode_TransitionPoseEvaluator.cpp
--- a/Engine/Source/Runtime/Engine/Private/Animation/AnimNode_TransitionPoseEvaluator.cpp
+++ b/Engine/Source/Runtime/Engine/Private/Animation/AnimNode_TransitionPoseEvaluator.cpp
@@ -30,8 +30,7 @@ void FAnimNode_TransitionPoseEvaluator::Initialize(const FAnimationInitializeCon
 
 void FAnimNode_TransitionPoseEvaluator::CacheBones(const FAnimationCacheBonesContext& Context) 
 {
-	const int32 NumBones = Context.AnimInstance->RequiredBones.GetNumBones();
-	CachedPose.Bones.Empty(NumBones);
+	CachedPose.SetBoneContainer(&Context.AnimInstance->RequiredBones);
 }
 
 void FAnimNode_TransitionPoseEvaluator::Update(const FAnimationUpdateContext& Context)
@@ -70,7 +69,7 @@ bool FAnimNode_TransitionPoseEvaluator::InputNodeNeedsEvaluate() const
 	return (EvaluatorMode == EEvaluatorMode::EM_Standard) || (CacheFramesRemaining > 0);
 }
 
-void FAnimNode_TransitionPoseEvaluator::CachePose(FPoseContext& Output, FA2Pose& PoseToCache)
+void FAnimNode_TransitionPoseEvaluator::CachePose(FPoseContext& Output, FCompactPose& PoseToCache)
 {
 	Output.AnimInstance->CopyPose(PoseToCache, CachedPose);
 }
\ No newline at end of file
diff --git a/Engine/Source/Runtime/Engine/Private/Animation/AnimNode_TwoBoneIK.cpp b/Engine/Source/Runtime/Engine/Private/Animation/AnimNode_TwoBoneIK.cpp
--- a/Engine/Source/Runtime/Engine/Private/Animation/AnimNode_TwoBoneIK.cpp
+++ b/Engine/Source/Runtime/Engine/Private/Animation/AnimNode_TwoBoneIK.cpp
@@ -31,30 +31,34 @@ void FAnimNode_TwoBoneIK::GatherDebugData(FNodeDebugData& DebugData)
 	ComponentPose.GatherDebugData(DebugData);
 }
 
-void FAnimNode_TwoBoneIK::EvaluateBoneTransforms(USkeletalMeshComponent* SkelComp, const FBoneContainer& RequiredBones, FA2CSPose& MeshBases, TArray<FBoneTransform>& OutBoneTransforms)
+void FAnimNode_TwoBoneIK::EvaluateBoneTransforms(USkeletalMeshComponent* SkelComp, FCSPose<FCompactPose>& MeshBases, TArray<FBoneTransform>& OutBoneTransforms)
 {
 	check(OutBoneTransforms.Num() == 0);
 
+	const FBoneContainer& BoneContainer = MeshBases.GetPose().GetBoneContainer();
+
 	// Get indices of the lower and upper limb bones and check validity.
 	bool bInvalidLimb = false;
 
-	const int32 EndBoneIndex = IKBone.BoneIndex;
-	const int32 LowerLimbIndex = RequiredBones.GetParentBoneIndex(EndBoneIndex);
+	FCompactPoseBoneIndex IKBoneCompactPoseIndex = IKBone.GetCompactPoseIndex(BoneContainer);
+
+	const FCompactPoseBoneIndex LowerLimbIndex = BoneContainer.GetParentBoneIndex(IKBoneCompactPoseIndex);
 	if (LowerLimbIndex == INDEX_NONE)
 	{
 		bInvalidLimb = true;
 	}
 
-	const int32 UpperLimbIndex = RequiredBones.GetParentBoneIndex(LowerLimbIndex);
+	const FCompactPoseBoneIndex UpperLimbIndex = BoneContainer.GetParentBoneIndex(LowerLimbIndex);
 	if (UpperLimbIndex == INDEX_NONE)
 	{
 		bInvalidLimb = true;
 	}
 
 	const bool bInBoneSpace = (EffectorLocationSpace == BCS_ParentBoneSpace) || (EffectorLocationSpace == BCS_BoneSpace);
-	const int32 EffectorSpaceBoneIndex = bInBoneSpace ? RequiredBones.GetPoseBoneIndexForBoneName(EffectorSpaceBoneName) : INDEX_NONE;
+	const int32 EffectorBoneIndex = bInBoneSpace ? BoneContainer.GetPoseBoneIndexForBoneName(EffectorSpaceBoneName) : INDEX_NONE;
+	const FCompactPoseBoneIndex EffectorSpaceBoneIndex = BoneContainer.MakeCompactPoseIndex(FMeshPoseBoneIndex(EffectorBoneIndex));
 
-	if (bInBoneSpace && ((EffectorSpaceBoneIndex == INDEX_NONE) || !RequiredBones.Contains(EffectorSpaceBoneIndex)))
+	if (bInBoneSpace && (EffectorSpaceBoneIndex == INDEX_NONE))
 	{
 		bInvalidLimb = true;
 	}
@@ -68,12 +72,12 @@ void FAnimNode_TwoBoneIK::EvaluateBoneTransforms(USkeletalMeshComponent* SkelCom
 	// Get Local Space transforms for our bones. We do this first in case they already are local.
 	// As right after we get them in component space. (And that does the auto conversion).
 	// We might save one transform by doing local first...
-	const FTransform EndBoneLocalTransform = MeshBases.GetLocalSpaceTransform(IKBone.BoneIndex);
+	const FTransform EndBoneLocalTransform = MeshBases.GetLocalSpaceTransform(IKBoneCompactPoseIndex);
 
 	// Now get those in component space...
-	FTransform UpperLimbCSTransform = MeshBases.GetComponentSpaceTransform(UpperLimbIndex);
 	FTransform LowerLimbCSTransform = MeshBases.GetComponentSpaceTransform(LowerLimbIndex);
-	FTransform EndBoneCSTransform = MeshBases.GetComponentSpaceTransform(IKBone.BoneIndex);
+	FTransform UpperLimbCSTransform = MeshBases.GetComponentSpaceTransform(UpperLimbIndex);
+	FTransform EndBoneCSTransform = MeshBases.GetComponentSpaceTransform(IKBoneCompactPoseIndex);
 
 	// Get current position of root of limb.
 	// All position are in Component space.
@@ -104,7 +108,14 @@ void FAnimNode_TwoBoneIK::EvaluateBoneTransforms(USkeletalMeshComponent* SkelCom
 
 	// Get joint target (used for defining plane that joint should be in).
 	FTransform JointTargetTransform(JointTargetLocation);
-	const int32 JointTargetSpaceBoneIndex = (JointTargetLocationSpace == BCS_ParentBoneSpace || JointTargetLocationSpace == BCS_BoneSpace) ? RequiredBones.GetPoseBoneIndexForBoneName(JointTargetSpaceBoneName) : INDEX_NONE;
+	FCompactPoseBoneIndex JointTargetSpaceBoneIndex(INDEX_NONE);
+
+	if (JointTargetLocationSpace == BCS_ParentBoneSpace || JointTargetLocationSpace == BCS_BoneSpace)
+	{
+		int32 Index = BoneContainer.GetPoseBoneIndexForBoneName(JointTargetSpaceBoneName);
+		JointTargetSpaceBoneIndex = BoneContainer.MakeCompactPoseIndex(FMeshPoseBoneIndex(Index));
+	}
+	
 	FAnimationRuntime::ConvertBoneSpaceTransformToCS(SkelComp, MeshBases, JointTargetTransform, JointTargetSpaceBoneIndex, JointTargetLocationSpace);
 
 	FVector	JointTargetPos = JointTargetTransform.GetTranslation();
@@ -270,7 +281,7 @@ void FAnimNode_TwoBoneIK::EvaluateBoneTransforms(USkeletalMeshComponent* SkelCom
 		EndBoneCSTransform.SetTranslation(OutEndPos);
 
 		// Order important. Third bone is End Bone.
-		OutBoneTransforms.Add( FBoneTransform(IKBone.BoneIndex, EndBoneCSTransform) );
+		OutBoneTransforms.Add(FBoneTransform(IKBoneCompactPoseIndex, EndBoneCSTransform));
 	}
 
 	// Make sure we have correct number of bones
diff --git a/Engine/Source/Runtime/Engine/Private/Animation/AnimNode_WheelHandler.cpp b/Engine/Source/Runtime/Engine/Private/Animation/AnimNode_WheelHandler.cpp
--- a/Engine/Source/Runtime/Engine/Private/Animation/AnimNode_WheelHandler.cpp
+++ b/Engine/Source/Runtime/Engine/Private/Animation/AnimNode_WheelHandler.cpp
@@ -42,20 +42,23 @@ void FAnimNode_WheelHandler::GatherDebugData(FNodeDebugData& DebugData)
 	ComponentPose.GatherDebugData(DebugData);
 }
 
-void FAnimNode_WheelHandler::EvaluateBoneTransforms(USkeletalMeshComponent* SkelComp, const FBoneContainer& RequiredBones, FA2CSPose& MeshBases, TArray<FBoneTransform>& OutBoneTransforms)
+void FAnimNode_WheelHandler::EvaluateBoneTransforms(USkeletalMeshComponent* SkelComp, FCSPose<FCompactPose>& MeshBases, TArray<FBoneTransform>& OutBoneTransforms)
 {
 	check(OutBoneTransforms.Num() == 0);
 
+	const FBoneContainer& BoneContainer = MeshBases.GetPose().GetBoneContainer();
 	for(const auto & WheelSim : WheelSimulators)
 	{
-		if(WheelSim.BoneReference.IsValid(RequiredBones))
+		if (WheelSim.BoneReference.IsValid(BoneContainer))
 		{
+			FCompactPoseBoneIndex WheelSimBoneIndex = WheelSim.BoneReference.GetCompactPoseIndex(BoneContainer);
+
 			// the way we apply transform is same as FMatrix or FTransform
 			// we apply scale first, and rotation, and translation
 			// if you'd like to translate first, you'll need two nodes that first node does translate and second nodes to rotate. 
-			FTransform NewBoneTM = MeshBases.GetComponentSpaceTransform(WheelSim.BoneReference.BoneIndex);
+			FTransform NewBoneTM = MeshBases.GetComponentSpaceTransform(WheelSimBoneIndex);
 
-			FAnimationRuntime::ConvertCSTransformToBoneSpace(SkelComp, MeshBases, NewBoneTM, WheelSim.BoneReference.BoneIndex, BCS_ComponentSpace);
+			FAnimationRuntime::ConvertCSTransformToBoneSpace(SkelComp, MeshBases, NewBoneTM, WheelSimBoneIndex, BCS_ComponentSpace);
 			
 			// Apply rotation offset
 			const FQuat BoneQuat(WheelSim.RotOffset);
@@ -65,10 +68,10 @@ void FAnimNode_WheelHandler::EvaluateBoneTransforms(USkeletalMeshComponent* Skel
 			NewBoneTM.AddToTranslation(WheelSim.LocOffset);
 
 			// Convert back to Component Space.
-			FAnimationRuntime::ConvertBoneSpaceTransformToCS(SkelComp, MeshBases, NewBoneTM, WheelSim.BoneReference.BoneIndex, BCS_ComponentSpace);
+			FAnimationRuntime::ConvertBoneSpaceTransformToCS(SkelComp, MeshBases, NewBoneTM, WheelSimBoneIndex, BCS_ComponentSpace);
 
 			// add back to it
-			OutBoneTransforms.Add( FBoneTransform(WheelSim.BoneReference.BoneIndex, NewBoneTM) );
+			OutBoneTransforms.Add(FBoneTransform(WheelSimBoneIndex, NewBoneTM));
 		}
 	}
 }
diff --git a/Engine/Source/Runtime/Engine/Private/Animation/AnimSequence.cpp b/Engine/Source/Runtime/Engine/Private/Animation/AnimSequence.cpp
--- a/Engine/Source/Runtime/Engine/Private/Animation/AnimSequence.cpp
+++ b/Engine/Source/Runtime/Engine/Private/Animation/AnimSequence.cpp
@@ -843,74 +843,71 @@ FTransform UAnimSequence::ExtractRootMotionFromRange(float StartTrackPosition, f
 	return EndTransform.GetRelativeTransform(StartTransform);
 }
 
-void UAnimSequence::GetAnimationPose(FTransformArrayA2& OutAtoms, const FBoneContainer& RequiredBones, const FAnimExtractContext& ExtractionContext) const
+void UAnimSequence::GetAnimationPose(FCompactPose& OutPose, const FAnimExtractContext& ExtractionContext) const
 {
 	// @todo anim: if compressed and baked in the future, we don't have to do this 
-	if( IsValidAdditive() ) 
+	if (IsValidAdditive())
 	{
 		if (AdditiveAnimType == AAT_LocalSpaceBase)
 		{
-			GetBonePose_Additive(OutAtoms, RequiredBones, ExtractionContext);
+			GetBonePose_Additive(OutPose, ExtractionContext);
 		}
-		else if ( AdditiveAnimType == AAT_RotationOffsetMeshSpace )
+		else if (AdditiveAnimType == AAT_RotationOffsetMeshSpace)
 		{
-			GetBonePose_AdditiveMeshRotationOnly(OutAtoms, RequiredBones, ExtractionContext);
+			GetBonePose_AdditiveMeshRotationOnly(OutPose, ExtractionContext);
 		}
 	}
 	else
 	{
-		GetBonePose(OutAtoms, RequiredBones, ExtractionContext);
+		GetBonePose(OutPose, ExtractionContext);
 	}
 }
 
-void UAnimSequence::ResetRootBoneForRootMotion(FTransformArrayA2& BoneTransforms, const FBoneContainer& RequiredBones, ERootMotionRootLock::Type InRootMotionRootLock) const
+void UAnimSequence::ResetRootBoneForRootMotion(FTransform& BoneTransform, const FBoneContainer& RequiredBones, ERootMotionRootLock::Type InRootMotionRootLock) const
 {
 	switch (InRootMotionRootLock)
 	{
-		case ERootMotionRootLock::AnimFirstFrame: BoneTransforms[0] = ExtractRootTrackTransform(0.f, &RequiredBones); break;
-		case ERootMotionRootLock::Zero: BoneTransforms[0] = FTransform::Identity; break;
+		case ERootMotionRootLock::AnimFirstFrame: BoneTransform = ExtractRootTrackTransform(0.f, &RequiredBones); break;
+		case ERootMotionRootLock::Zero: BoneTransform = FTransform::Identity; break;
 		default:
-		case ERootMotionRootLock::RefPose: BoneTransforms[0] = RequiredBones.GetRefPoseArray()[0]; break;
+		case ERootMotionRootLock::RefPose: BoneTransform = RequiredBones.GetRefPoseArray()[0]; break;
 	}
 }
 
-void UAnimSequence::GetBonePose(FTransformArrayA2& OutAtoms, const FBoneContainer& RequiredBones, const FAnimExtractContext& ExtractionContext) const
+void UAnimSequence::GetBonePose(FCompactPose& OutPose, const FAnimExtractContext& ExtractionContext) const
 {
+	QUICK_SCOPE_CYCLE_COUNTER(STAT_GETBONEPOSE);
+
 	USkeleton* MySkeleton = GetSkeleton();
 	if (!MySkeleton)
 	{
-		FAnimationRuntime::FillWithRefPose(OutAtoms, RequiredBones);
+		OutPose.ResetToRefPose();
 		return;
 	}
 
+	const FBoneContainer& RequiredBones = OutPose.GetBoneContainer();
+
 	// if retargeting is disabled, we initialize pose with 'Retargeting Source' ref pose.
 	bool const bDisableRetargeting = RequiredBones.GetDisableRetargeting();
 	if (bDisableRetargeting)
 	{
 		TArray<FTransform> const& AuthoredOnRefSkeleton = MySkeleton->GetRefLocalPoses(RetargetSource);
 		TArray<FBoneIndexType> const& RequireBonesIndexArray = RequiredBones.GetBoneIndicesArray();
-		TArray<int32> const& PoseToSkeletonBoneIndexArray = RequiredBones.GetPoseToSkeletonBoneIndexArray();
-	
-		// Allocate output and make sure it's the right size.
-		int32 const NumPoseBones = RequiredBones.GetNumBones();
-		OutAtoms.Empty(NumPoseBones);
-		OutAtoms.AddUninitialized(NumPoseBones);
 
 		int32 const NumRequiredBones = RequireBonesIndexArray.Num();
-		for (int32 ArrayIndex = 0; ArrayIndex < NumRequiredBones; ArrayIndex++)
+		for (FCompactPoseBoneIndex PoseBoneIndex : OutPose.ForEachBoneIndex())
 		{
-			int32 const& PoseBoneIndex = RequireBonesIndexArray[ArrayIndex];
-			int32 const& SkeletonBoneIndex = PoseToSkeletonBoneIndexArray[PoseBoneIndex];
+			int32 const& SkeletonBoneIndex = RequiredBones.GetSkeletonIndex(PoseBoneIndex);
 
 			// Pose bone index should always exist in Skeleton
 			checkSlow(SkeletonBoneIndex != INDEX_NONE);
-			OutAtoms[PoseBoneIndex] = AuthoredOnRefSkeleton[SkeletonBoneIndex];
+			OutPose[PoseBoneIndex] = AuthoredOnRefSkeleton[SkeletonBoneIndex];
 		}
 	}
 	else
 	{
 		// initialize with ref-pose
-		FAnimationRuntime::FillWithRefPose(OutAtoms, RequiredBones);
+		OutPose.ResetToRefPose();
 	}
 
 	int32 const NumTracks = GetNumberOfTracks();
@@ -924,9 +921,8 @@ void UAnimSequence::GetBonePose(FTransformArrayA2& OutAtoms, const FBoneContaine
 	// Slower path for disable retargeting, that's only used in editor and for debugging.
 	if (bDisableRetargeting || RequiredBones.ShouldUseRawData() || RequiredBones.ShouldUseSourceData())
 	{
-		TArray<int32> const& SkeletonToPoseBoneIndexArray = RequiredBones.GetSkeletonToPoseBoneIndexArray();
 		const TArray<FRawAnimSequenceTrack>* AnimationData = NULL;
-		
+
 		if (RequiredBones.ShouldUseSourceData() && SourceRawAnimationData.Num() > 0)
 		{
 			AnimationData = &SourceRawAnimationData;
@@ -936,29 +932,29 @@ void UAnimSequence::GetBonePose(FTransformArrayA2& OutAtoms, const FBoneContaine
 			AnimationData = &RawAnimationData;
 		}
 
-		for(int32 TrackIndex=0; TrackIndex<NumTracks; TrackIndex++)
+		for (int32 TrackIndex = 0; TrackIndex < NumTracks; TrackIndex++)
 		{
 			const int32 SkeletonBoneIndex = GetSkeletonIndexFromTrackIndex(TrackIndex);
 			// not sure it's safe to assume that SkeletonBoneIndex can never be INDEX_NONE
-			if( (SkeletonBoneIndex != INDEX_NONE) && (SkeletonBoneIndex < MAX_BONES) )
+			if ((SkeletonBoneIndex != INDEX_NONE) && (SkeletonBoneIndex < MAX_BONES))
 			{
-				const int32 PoseBoneIndex = SkeletonToPoseBoneIndexArray[SkeletonBoneIndex];
-				if(PoseBoneIndex != INDEX_NONE)
+				const FCompactPoseBoneIndex PoseBoneIndex = RequiredBones.GetCompactPoseIndexFromSkeletonIndex(SkeletonBoneIndex);
+				if (PoseBoneIndex != INDEX_NONE)
 				{
 					// extract animation
-					ExtractBoneTransform(*AnimationData, OutAtoms[PoseBoneIndex], TrackIndex, ExtractionContext.CurrentTime);
+					ExtractBoneTransform(*AnimationData, OutPose[PoseBoneIndex], TrackIndex, ExtractionContext.CurrentTime);
 
-					if(!bDisableRetargeting)
+					if (!bDisableRetargeting)
 					{
-						RetargetBoneTransform(OutAtoms[PoseBoneIndex], SkeletonBoneIndex, PoseBoneIndex, RequiredBones);
+						RetargetBoneTransform(OutPose[PoseBoneIndex], SkeletonBoneIndex, PoseBoneIndex, RequiredBones);
 					}
 				}
 			}
 		}
 
-		if( ExtractionContext.bExtractRootMotion && bEnableRootMotion )
+		if (ExtractionContext.bExtractRootMotion && bEnableRootMotion)
 		{
-			ResetRootBoneForRootMotion(OutAtoms, RequiredBones, RootMotionRootLock);
+			ResetRootBoneForRootMotion(OutPose[FCompactPoseBoneIndex(0)], RequiredBones, RootMotionRootLock);
 		}
 		return;
 	}
@@ -980,74 +976,84 @@ void UAnimSequence::GetBonePose(FTransformArrayA2& OutAtoms, const FBoneContaine
 	AnimRelativeRetargetingPairs.Empty(NumTracks);
 
 	// Optimization: assuming first index is root bone. That should always be the case in Skeletons.
-	checkSlow( (SkeletonToPoseBoneIndexArray[0] == 0) );
+	checkSlow((SkeletonToPoseBoneIndexArray[0] == 0));
 	// this is not guaranteed for AnimSequences though... If Root is not animated, Track will not exist.
 	const bool bFirstTrackIsRootBone = (GetSkeletonIndexFromTrackIndex(0) == 0);
 
 	// Handle root bone separately if it is track 0. so we start w/ Index 1.
-	for(int32 TrackIndex=(bFirstTrackIsRootBone ? 1 : 0); TrackIndex<NumTracks; TrackIndex++)
+	for (int32 TrackIndex = (bFirstTrackIsRootBone ? 1 : 0); TrackIndex < NumTracks; TrackIndex++)
 	{
 		const int32 SkeletonBoneIndex = GetSkeletonIndexFromTrackIndex(TrackIndex);
 		// not sure it's safe to assume that SkeletonBoneIndex can never be INDEX_NONE
-		if( SkeletonBoneIndex != INDEX_NONE )
+		if (SkeletonBoneIndex != INDEX_NONE)
 		{
-			const int32 PoseBoneIndex = SkeletonToPoseBoneIndexArray[SkeletonBoneIndex];
-			if(PoseBoneIndex != INDEX_NONE)
+			const FCompactPoseBoneIndex BoneIndex = RequiredBones.GetCompactPoseIndexFromSkeletonIndex(SkeletonBoneIndex);
+			//Nasty, we break our type safety, code in the lower levels should be adjusted for this
+			const int32 CompactPoseBoneIndex = BoneIndex.GetInt();
+			if (CompactPoseBoneIndex != INDEX_NONE)
 			{
-				checkSlow(PoseBoneIndex < RequiredBones.GetNumBones());
-				RotationScalePairs.Add(BoneTrackPair(PoseBoneIndex, TrackIndex));
+				RotationScalePairs.Add(BoneTrackPair(CompactPoseBoneIndex, TrackIndex));
 
 				// Skip extracting translation component for EBoneTranslationRetargetingMode::Skeleton.
-				switch(BoneTree[SkeletonBoneIndex].TranslationRetargetingMode)
+				switch (BoneTree[SkeletonBoneIndex].TranslationRetargetingMode)
 				{
 				case EBoneTranslationRetargetingMode::Animation:
-				TranslationPairs.Add(BoneTrackPair(PoseBoneIndex, TrackIndex));
-				break;
+					TranslationPairs.Add(BoneTrackPair(CompactPoseBoneIndex, TrackIndex));
+					break;
 				case EBoneTranslationRetargetingMode::AnimationScaled:
-				TranslationPairs.Add(BoneTrackPair(PoseBoneIndex, TrackIndex));
-				AnimScaleRetargetingPairs.Add(BoneTrackPair(PoseBoneIndex, SkeletonBoneIndex));
-				break;
+					TranslationPairs.Add(BoneTrackPair(CompactPoseBoneIndex, TrackIndex));
+					AnimScaleRetargetingPairs.Add(BoneTrackPair(CompactPoseBoneIndex, SkeletonBoneIndex));
+					break;
 				case EBoneTranslationRetargetingMode::AnimationRelative:
-				TranslationPairs.Add(BoneTrackPair(PoseBoneIndex, TrackIndex));
-				AnimRelativeRetargetingPairs.Add(BoneTrackPair(PoseBoneIndex, SkeletonBoneIndex));
-				break;
+					TranslationPairs.Add(BoneTrackPair(CompactPoseBoneIndex, TrackIndex));
+					AnimRelativeRetargetingPairs.Add(BoneTrackPair(CompactPoseBoneIndex, SkeletonBoneIndex));
+					break;
 				}
 			}
 		}
 	}
 
+	QUICK_SCOPE_CYCLE_COUNTER(STAT_GETBONEPOSE_FORREALZ);
 	// Handle Root Bone separately
-	if( bFirstTrackIsRootBone )
+	if (bFirstTrackIsRootBone)
 	{
 		const int32 TrackIndex = 0;
-		FTransform& RootAtom = OutAtoms[0];
+		FCompactPoseBoneIndex RootBone(0);
+		FTransform& RootAtom = OutPose[RootBone];
 
-		AnimationFormat_GetBoneAtom(	
+		AnimationFormat_GetBoneAtom(
 			RootAtom,
 			*this,
 			TrackIndex,
 			ExtractionContext.CurrentTime);
 
 		// @laurent - we should look into splitting rotation and translation tracks, so we don't have to process translation twice.
-		RetargetBoneTransform(RootAtom, 0, 0, RequiredBones);
+		RetargetBoneTransform(RootAtom, 0, RootBone, RequiredBones);
 	}
 
 	if (RotationScalePairs.Num() > 0)
 	{
 		// get the remaining bone atoms
-		AnimationFormat_GetAnimationPose(
-			*((FTransformArray*)&OutAtoms), //@TODO:@ANIMATION: Nasty hack
+		OutPose.PopulateFromAnimation( //@TODO:@ANIMATION: Nasty hack, be good to not have this function on the pose
+			*this,
 			RotationScalePairs,
 			TranslationPairs,
 			RotationScalePairs,
-			*this,
 			ExtractionContext.CurrentTime);
+
+		/*AnimationFormat_GetAnimationPose(
+			*((FTransformArray*)&OutAtoms), 
+			RotationScalePairs,
+			TranslationPairs,
+			RotationScalePairs,
+			*this,
+			ExtractionContext.CurrentTime);*/
 	}
 
 	// Once pose has been extracted, snap root bone back to first frame if we are extracting root motion.
-	if( ExtractionContext.bExtractRootMotion && bEnableRootMotion)
+	if (ExtractionContext.bExtractRootMotion && bEnableRootMotion)
 	{
-		ResetRootBoneForRootMotion(OutAtoms, RequiredBones, RootMotionRootLock);
+		ResetRootBoneForRootMotion(OutPose[FCompactPoseBoneIndex(0)], RequiredBones, RootMotionRootLock);
 	}
 
 	// Anim Scale Retargeting
@@ -1055,20 +1061,18 @@ void UAnimSequence::GetBonePose(FTransformArrayA2& OutAtoms, const FBoneContaine
 	if (NumBonesToScaleRetarget > 0)
 	{
 		TArray<FTransform> const& AuthoredOnRefSkeleton = MySkeleton->GetRefLocalPoses(RetargetSource);
-		TArray<FTransform> const& PlayingOnRefSkeleton = RequiredBones.GetRefPoseArray();
 
-		for (int32 Index = 0; Index<NumBonesToScaleRetarget; Index++)
+		for (const BoneTrackPair& BonePair : AnimScaleRetargetingPairs)
 		{
-			BoneTrackPair const& BonePair = AnimScaleRetargetingPairs[Index];
-			int32 const& PoseBoneIndex = BonePair.AtomIndex;
-			int32 const& SkeletonBoneIndex = BonePair.TrackIndex; 
+			const FCompactPoseBoneIndex BoneIndex(BonePair.AtomIndex); //Nasty, we break our type safety, code in the lower levels should be adjusted for this
+			int32 const& SkeletonBoneIndex = BonePair.TrackIndex;
 
 			// @todo - precache that in FBoneContainer when we have SkeletonIndex->TrackIndex mapping. So we can just apply scale right away.
 			float const SourceTranslationLength = AuthoredOnRefSkeleton[SkeletonBoneIndex].GetTranslation().Size();
 			if (FMath::Abs(SourceTranslationLength) > KINDA_SMALL_NUMBER)
 			{
-				float const TargetTranslationLength = PlayingOnRefSkeleton[PoseBoneIndex].GetTranslation().Size();
-				OutAtoms[PoseBoneIndex].ScaleTranslation(TargetTranslationLength / SourceTranslationLength);
+				float const TargetTranslationLength = RequiredBones.GetRefPoseTransform(BoneIndex).GetTranslation().Size();
+				OutPose[BoneIndex].ScaleTranslation(TargetTranslationLength / SourceTranslationLength);
 			}
 		}
 	}
@@ -1078,119 +1082,120 @@ void UAnimSequence::GetBonePose(FTransformArrayA2& OutAtoms, const FBoneContaine
 	if (NumBonesToRelativeRetarget > 0)
 	{
 		TArray<FTransform> const& AuthoredOnRefSkeleton = MySkeleton->GetRefLocalPoses(RetargetSource);
-		TArray<FTransform> const& PlayingOnRefSkeleton = RequiredBones.GetRefPoseArray();
-
-		for (int32 Index = 0; Index < NumBonesToRelativeRetarget; Index++)
+	
+		for (const BoneTrackPair& BonePair : AnimRelativeRetargetingPairs)
 		{
-			BoneTrackPair const& BonePair = AnimRelativeRetargetingPairs[Index];
-			int32 const& PoseBoneIndex = BonePair.AtomIndex;
+			const FCompactPoseBoneIndex BoneIndex(BonePair.AtomIndex); //Nasty, we break our type safety, code in the lower levels should be adjusted for this
 			int32 const& SkeletonBoneIndex = BonePair.TrackIndex;
 
+			const FTransform& RefPose = RequiredBones.GetRefPoseTransform(BoneIndex);
+
 			// Apply the retargeting as if it were an additive difference between the current skeleton and the retarget skeleton. 
-			OutAtoms[PoseBoneIndex].SetRotation(OutAtoms[PoseBoneIndex].GetRotation() * AuthoredOnRefSkeleton[SkeletonBoneIndex].GetRotation().Inverse() * PlayingOnRefSkeleton[PoseBoneIndex].GetRotation() );
-			OutAtoms[PoseBoneIndex].SetTranslation(OutAtoms[PoseBoneIndex].GetTranslation() + (PlayingOnRefSkeleton[PoseBoneIndex].GetTranslation() - AuthoredOnRefSkeleton[SkeletonBoneIndex].GetTranslation()));
-			OutAtoms[PoseBoneIndex].SetScale3D(OutAtoms[PoseBoneIndex].GetScale3D() * (PlayingOnRefSkeleton[PoseBoneIndex].GetScale3D() * AuthoredOnRefSkeleton[SkeletonBoneIndex].GetSafeScaleReciprocal(AuthoredOnRefSkeleton[SkeletonBoneIndex].GetScale3D())));
-			OutAtoms[PoseBoneIndex].NormalizeRotation();
+			OutPose[BoneIndex].SetRotation(OutPose[BoneIndex].GetRotation() * AuthoredOnRefSkeleton[SkeletonBoneIndex].GetRotation().Inverse() * RefPose.GetRotation());
+			OutPose[BoneIndex].SetTranslation(OutPose[BoneIndex].GetTranslation() + (RefPose.GetTranslation() - AuthoredOnRefSkeleton[SkeletonBoneIndex].GetTranslation()));
+			OutPose[BoneIndex].SetScale3D(OutPose[BoneIndex].GetScale3D() * (RefPose.GetScale3D() * AuthoredOnRefSkeleton[SkeletonBoneIndex].GetSafeScaleReciprocal(AuthoredOnRefSkeleton[SkeletonBoneIndex].GetScale3D())));
+			OutPose[BoneIndex].NormalizeRotation();
 		}
 	}
 }
 
-void UAnimSequence::GetBonePose_Additive(FTransformArrayA2& OutAtoms, const FBoneContainer& RequiredBones, const FAnimExtractContext& ExtractionContext) const
+void UAnimSequence::GetBonePose_Additive(FCompactPose& OutPose, const FAnimExtractContext& ExtractionContext) const
 {
-	if( !IsValidAdditive() )
+	if (!IsValidAdditive())
 	{
-		FAnimationRuntime::InitializeTransform(RequiredBones, OutAtoms);
+		OutPose.ResetToIdentity();
 		return;
 	}
 
 	// Extract target pose
-	GetBonePose(OutAtoms, RequiredBones, ExtractionContext);
+	GetBonePose(OutPose, ExtractionContext);
 
 	// Extract base pose
-	FTransformArrayA2 BasePoseAtoms;
-	GetAdditiveBasePose(BasePoseAtoms, RequiredBones, ExtractionContext);
+	FCompactPose BasePose;
+	BasePose.SetBoneContainer(&OutPose.GetBoneContainer());
+
+	GetAdditiveBasePose(BasePose, ExtractionContext);
 
 	// Create Additive animation
-	FAnimationRuntime::ConvertPoseToAdditive(OutAtoms, BasePoseAtoms, RequiredBones);
+	FAnimationRuntime::ConvertPoseToAdditive(OutPose, BasePose);
 }
 
-void UAnimSequence::GetAdditiveBasePose(FTransformArrayA2& OutAtoms, const FBoneContainer& RequiredBones, const FAnimExtractContext& ExtractionContext) const
+void UAnimSequence::GetAdditiveBasePose(FCompactPose& OutPose, const FAnimExtractContext& ExtractionContext) const
 {
 	switch (RefPoseType)
 	{
 		// use whole animation as a base pose. Need BasePoseSeq.
-	case ABPT_AnimScaled:	
-		{
-			// normalize time to fit base seq
-			const float Fraction = FMath::Clamp<float>(ExtractionContext.CurrentTime / SequenceLength, 0.f, 1.f);
-			const float BasePoseTime = RefPoseSeq->SequenceLength * Fraction;
+	case ABPT_AnimScaled:
+	{
+		// normalize time to fit base seq
+		const float Fraction = FMath::Clamp<float>(ExtractionContext.CurrentTime / SequenceLength, 0.f, 1.f);
+		const float BasePoseTime = RefPoseSeq->SequenceLength * Fraction;
 
-			FAnimExtractContext BasePoseExtractionContext(ExtractionContext);
-			BasePoseExtractionContext.CurrentTime = BasePoseTime;
-			RefPoseSeq->GetBonePose(OutAtoms, RequiredBones, BasePoseExtractionContext);
-			break;
-		}
+		FAnimExtractContext BasePoseExtractionContext(ExtractionContext);
+		BasePoseExtractionContext.CurrentTime = BasePoseTime;
+		RefPoseSeq->GetBonePose(OutPose, BasePoseExtractionContext);
+		break;
+	}
 		// use animation as a base pose. Need BasePoseSeq and RefFrameIndex (will clamp if outside).
-	case ABPT_AnimFrame:		
-		{
-			const float Fraction = (RefPoseSeq->NumFrames > 0) ? FMath::Clamp<float>((float)RefFrameIndex/(float)RefPoseSeq->NumFrames, 0.f, 1.f) : 0.f;
-			const float BasePoseTime = RefPoseSeq->SequenceLength * Fraction;
+	case ABPT_AnimFrame:
+	{
+		const float Fraction = (RefPoseSeq->NumFrames > 0) ? FMath::Clamp<float>((float)RefFrameIndex / (float)RefPoseSeq->NumFrames, 0.f, 1.f) : 0.f;
+		const float BasePoseTime = RefPoseSeq->SequenceLength * Fraction;
 
-			FAnimExtractContext BasePoseExtractionContext(ExtractionContext);
-			BasePoseExtractionContext.CurrentTime = BasePoseTime;
-			RefPoseSeq->GetBonePose(OutAtoms, RequiredBones, BasePoseExtractionContext);
-			break;
-		}
+		FAnimExtractContext BasePoseExtractionContext(ExtractionContext);
+		BasePoseExtractionContext.CurrentTime = BasePoseTime;
+		RefPoseSeq->GetBonePose(OutPose, BasePoseExtractionContext);
+		break;
+	}
 		// use ref pose of Skeleton as base
-	case ABPT_RefPose:		
+	case ABPT_RefPose:
 	default:
-		FAnimationRuntime::FillWithRefPose(OutAtoms, RequiredBones);
+		OutPose.ResetToRefPose();
 		break;
 	}
 }
 
-void UAnimSequence::GetBonePose_AdditiveMeshRotationOnly(FTransformArrayA2& OutAtoms, const FBoneContainer& RequiredBones, const FAnimExtractContext& ExtractionContext) const
+void UAnimSequence::GetBonePose_AdditiveMeshRotationOnly(FCompactPose& OutPose, const FAnimExtractContext& ExtractionContext) const
 {
-	if( !IsValidAdditive() )
+	if (!IsValidAdditive())
 	{
 		// since this is additive, need to initialize to identity
-		FAnimationRuntime::InitializeTransform(RequiredBones, OutAtoms);
+		OutPose.ResetToIdentity();
 		return;
 	}
 
 	// Get target pose
-	GetBonePose(OutAtoms, RequiredBones, ExtractionContext);
+	GetBonePose(OutPose, ExtractionContext);
 
 	// get base pose
-	FTransformArrayA2 BasePoseAtoms;
-	GetAdditiveBasePose(BasePoseAtoms, RequiredBones, ExtractionContext);
+	FCompactPose BasePose;
+	BasePose.SetBoneContainer(&OutPose.GetBoneContainer());
+	GetAdditiveBasePose(BasePose, ExtractionContext);
 
 	// Convert them to mesh rotation.
-	FAnimationRuntime::ConvertPoseToMeshRotation(OutAtoms, RequiredBones);
-	FAnimationRuntime::ConvertPoseToMeshRotation(BasePoseAtoms, RequiredBones);
+	FAnimationRuntime::ConvertPoseToMeshRotation(OutPose);
+	FAnimationRuntime::ConvertPoseToMeshRotation(BasePose);
 
 	// Turn into Additive
-	FAnimationRuntime::ConvertPoseToAdditive(OutAtoms, BasePoseAtoms, RequiredBones);
+	FAnimationRuntime::ConvertPoseToAdditive(OutPose, BasePose);
 }
 
-void UAnimSequence::RetargetBoneTransform(FTransform& BoneTransform, const int32& SkeletonBoneIndex, const int32& PoseBoneIndex, const FBoneContainer& RequiredBones) const
+void UAnimSequence::RetargetBoneTransform(FTransform& BoneTransform, const int32& SkeletonBoneIndex, const FCompactPoseBoneIndex& BoneIndex, const FBoneContainer& RequiredBones) const
 {
 	const USkeleton* MySkeleton = GetSkeleton();
 	const TArray<FBoneNode>& BoneTree = MySkeleton->GetBoneTree();
-	if( BoneTree[SkeletonBoneIndex].TranslationRetargetingMode == EBoneTranslationRetargetingMode::Skeleton )
+	if (BoneTree[SkeletonBoneIndex].TranslationRetargetingMode == EBoneTranslationRetargetingMode::Skeleton)
 	{
-		const TArray<FTransform>& TargetRefPoseArray = RequiredBones.GetRefPoseArray();
-		BoneTransform.SetTranslation( TargetRefPoseArray[PoseBoneIndex].GetTranslation() );
+		BoneTransform.SetTranslation(RequiredBones.GetRefPoseTransform(BoneIndex).GetTranslation());
 	}
-	else if( BoneTree[SkeletonBoneIndex].TranslationRetargetingMode == EBoneTranslationRetargetingMode::AnimationScaled )
+	else if (BoneTree[SkeletonBoneIndex].TranslationRetargetingMode == EBoneTranslationRetargetingMode::AnimationScaled)
 	{
 		// @todo - precache that in FBoneContainer when we have SkeletonIndex->TrackIndex mapping. So we can just apply scale right away.
 		const TArray<FTransform>& SkeletonRefPoseArray = GetSkeleton()->GetRefLocalPoses(RetargetSource);
 		const float SourceTranslationLength = SkeletonRefPoseArray[SkeletonBoneIndex].GetTranslation().Size();
-		if( FMath::Abs(SourceTranslationLength) > KINDA_SMALL_NUMBER )
+		if (FMath::Abs(SourceTranslationLength) > KINDA_SMALL_NUMBER)
 		{
-			const TArray<FTransform>& TargetRefPoseArray = RequiredBones.GetRefPoseArray();
-			const float TargetTranslationLength = TargetRefPoseArray[PoseBoneIndex].GetTranslation().Size();
+			const float TargetTranslationLength = RequiredBones.GetRefPoseTransform(BoneIndex).GetTranslation().Size();
 			BoneTransform.ScaleTranslation(TargetTranslationLength / SourceTranslationLength);
 		}
 	}
@@ -1199,10 +1204,12 @@ void UAnimSequence::RetargetBoneTransform(FTransform& BoneTransform, const int32
 		const TArray<FTransform>& AuthoredOnRefSkeleton = GetSkeleton()->GetRefLocalPoses(RetargetSource);
 		const TArray<FTransform>& PlayingOnRefSkeleton = RequiredBones.GetRefPoseArray();
 
+		const FTransform& RefPoseTransform = RequiredBones.GetRefPoseTransform(BoneIndex);
+
 		// Apply the retargeting as if it were an additive difference between the current skeleton and the retarget skeleton. 
-		BoneTransform.SetRotation(BoneTransform.GetRotation() * AuthoredOnRefSkeleton[SkeletonBoneIndex].GetRotation().Inverse() * PlayingOnRefSkeleton[PoseBoneIndex].GetRotation());
-		BoneTransform.SetTranslation(BoneTransform.GetTranslation() + (PlayingOnRefSkeleton[PoseBoneIndex].GetTranslation() - AuthoredOnRefSkeleton[SkeletonBoneIndex].GetTranslation()));
-		BoneTransform.SetScale3D(BoneTransform.GetScale3D() * (PlayingOnRefSkeleton[PoseBoneIndex].GetScale3D() * AuthoredOnRefSkeleton[SkeletonBoneIndex].GetSafeScaleReciprocal(AuthoredOnRefSkeleton[SkeletonBoneIndex].GetScale3D())));
+		BoneTransform.SetRotation(BoneTransform.GetRotation() * AuthoredOnRefSkeleton[SkeletonBoneIndex].GetRotation().Inverse() * RefPoseTransform.GetRotation());
+		BoneTransform.SetTranslation(BoneTransform.GetTranslation() + (RefPoseTransform.GetTranslation() - AuthoredOnRefSkeleton[SkeletonBoneIndex].GetTranslation()));
+		BoneTransform.SetScale3D(BoneTransform.GetScale3D() * (RefPoseTransform.GetScale3D() * AuthoredOnRefSkeleton[SkeletonBoneIndex].GetSafeScaleReciprocal(AuthoredOnRefSkeleton[SkeletonBoneIndex].GetScale3D())));
 		BoneTransform.NormalizeRotation();
 	}
 }
diff --git a/Engine/Source/Runtime/Engine/Private/Animation/AnimSingleNodeInstance.cpp b/Engine/Source/Runtime/Engine/Private/Animation/AnimSingleNodeInstance.cpp
--- a/Engine/Source/Runtime/Engine/Private/Animation/AnimSingleNodeInstance.cpp
+++ b/Engine/Source/Runtime/Engine/Private/Animation/AnimSingleNodeInstance.cpp
@@ -319,15 +319,14 @@ bool UAnimSingleNodeInstance::NativeEvaluateAnimation(FPoseContext& Output)
 		{
 			if (Sequence->IsValidAdditive())
 			{
-				FA2Pose BasePose;
-				FA2Pose AdditivePose;
-
-				BasePose.Bones.AddUninitialized(Output.Pose.Bones.Num());
-				AdditivePose.Bones.AddUninitialized(Output.Pose.Bones.Num());
+				FCompactPose BasePose;
+				BasePose.SetBoneContainer(&Output.Pose.GetBoneContainer());
+				FCompactPose AdditivePose;
+				AdditivePose.SetBoneContainer(&Output.Pose.GetBoneContainer());
 
 				FAnimExtractContext ExtractionContext(CurrentTime);
-				Sequence->GetAdditiveBasePose(BasePose.Bones, RequiredBones, ExtractionContext);
-				Sequence->GetAnimationPose(AdditivePose.Bones, RequiredBones, ExtractionContext);
+				Sequence->GetAdditiveBasePose(BasePose, ExtractionContext);
+				Sequence->GetAnimationPose(AdditivePose, ExtractionContext);
 				if (Sequence->AdditiveAnimType == AAT_LocalSpaceBase)
 				{
 					ApplyAdditiveSequence(BasePose, AdditivePose, 1.0f, Output.Pose);
@@ -340,7 +339,7 @@ bool UAnimSingleNodeInstance::NativeEvaluateAnimation(FPoseContext& Output)
 			else
 			{
 				// if sekeltalmesh isn't there, we'll need to use skeleton
-				FAnimationRuntime::GetPoseFromSequence(Sequence, RequiredBones, Output.Pose.Bones, FAnimExtractContext(CurrentTime, Sequence->bEnableRootMotion));
+				FAnimationRuntime::GetPoseFromSequence(Sequence, Output.Pose, FAnimExtractContext(CurrentTime, Sequence->bEnableRootMotion));
 			}
 		}
 		else if (UAnimComposite* Composite = Cast<UAnimComposite>(CurrentAsset))
@@ -351,18 +350,18 @@ bool UAnimSingleNodeInstance::NativeEvaluateAnimation(FPoseContext& Output)
 			EAdditiveAnimationType AdditiveAnimType = AAT_None;
 			if (AnimTrack.IsAdditive())
 			{
-				FA2Pose AdditivePose;
-				FA2Pose BasePose;
+				FCompactPose BasePose;
+				BasePose.SetBoneContainer(&Output.Pose.GetBoneContainer());
+				FCompactPose AdditivePose;
+				AdditivePose.SetBoneContainer(&Output.Pose.GetBoneContainer());
 
-				BasePose.Bones.AddUninitialized(Output.Pose.Bones.Num());
-				AdditivePose.Bones.AddUninitialized(Output.Pose.Bones.Num());
 				AdditiveAnimType = AnimTrack.IsRotationOffsetAdditive()? AAT_RotationOffsetMeshSpace : AAT_LocalSpaceBase;
 				
 				// get base pose - for now we only support ref pose as base
-				FAnimationRuntime::FillWithRefPose(BasePose.Bones, RequiredBones);
+				BasePose.ResetToRefPose();
 				
 				//get the additive pose
-				FAnimationRuntime::GetPoseFromAnimTrack(AnimTrack, RequiredBones, AdditivePose.Bones, FAnimExtractContext(CurrentTime, RootMotionMode == ERootMotionMode::RootMotionFromEverything));
+				FAnimationRuntime::GetPoseFromAnimTrack(AnimTrack, AdditivePose, FAnimExtractContext(CurrentTime, RootMotionMode == ERootMotionMode::RootMotionFromEverything));
 
 				// if additive, we should blend with source to make it fullbody
 				if (AdditiveAnimType == AAT_LocalSpaceBase)
@@ -377,7 +376,7 @@ bool UAnimSingleNodeInstance::NativeEvaluateAnimation(FPoseContext& Output)
 			else
 			{
 				//doesn't handle additive yet
-				FAnimationRuntime::GetPoseFromAnimTrack(AnimTrack, RequiredBones, Output.Pose.Bones, FAnimExtractContext(CurrentTime, RootMotionMode == ERootMotionMode::RootMotionFromEverything));
+				FAnimationRuntime::GetPoseFromAnimTrack(AnimTrack, Output.Pose, FAnimExtractContext(CurrentTime, RootMotionMode == ERootMotionMode::RootMotionFromEverything));
 			}
 		}
 		else if (UAnimMontage* Montage = Cast<UAnimMontage>(CurrentAsset))
@@ -386,25 +385,26 @@ bool UAnimSingleNodeInstance::NativeEvaluateAnimation(FPoseContext& Output)
 			// in the future, add option to see which slot to see
 			if (Montage->SlotAnimTracks.Num() > 0)
 			{
-				FA2Pose SourcePose;
-				SourcePose.Bones.AddUninitialized(Output.Pose.Bones.Num());
+				FCompactPose SourcePose;
+				SourcePose.SetBoneContainer(&Output.Pose.GetBoneContainer());
+			
 				if (Montage->IsValidAdditive())
 				{
 #if WITH_EDITORONLY_DATA
 					// if montage is additive, we need to have base pose for the slot pose evaluate
 					if (Montage->PreviewBasePose && Montage->SequenceLength > 0.f)
 					{
-						Montage->PreviewBasePose->GetBonePose(SourcePose.Bones, RequiredBones, FAnimExtractContext(CurrentTime));
+						Montage->PreviewBasePose->GetBonePose(SourcePose, FAnimExtractContext(CurrentTime));
 					}
 					else
 #endif // WITH_EDITORONLY_DATA
 					{
-						FAnimationRuntime::FillWithRefPose(SourcePose.Bones, RequiredBones);				
+						SourcePose.ResetToRefPose();				
 					}
 				}
 				else
 				{
-					FAnimationRuntime::FillWithRefPose(SourcePose.Bones, RequiredBones);				
+					SourcePose.ResetToRefPose();			
 				}
 
 				SlotEvaluatePose(Montage->SlotAnimTracks[0].SlotName, SourcePose, Output.Pose, 1.f);
@@ -562,52 +562,50 @@ void UAnimSingleNodeInstance::SetBlendSpaceInput(const FVector& InBlendInput)
 	BlendSpaceInput = InBlendInput;
 }
 
-void UAnimSingleNodeInstance::InternalBlendSpaceEvaluatePose(class UBlendSpaceBase* BlendSpace, TArray<FBlendSampleData>& BlendSampleDataCache, struct FA2Pose& Pose)
+void UAnimSingleNodeInstance::InternalBlendSpaceEvaluatePose(class UBlendSpaceBase* BlendSpace, TArray<FBlendSampleData>& BlendSampleDataCache, FCompactPose& OutPose)
 {
 	USkeletalMeshComponent* Component = GetSkelMeshComponent();
 
 	if (BlendSpace->IsValidAdditive())
 	{
-		FA2Pose BasePose;
-		FA2Pose AdditivePose;
-		BasePose.Bones.AddUninitialized(Pose.Bones.Num());
-		AdditivePose.Bones.AddUninitialized(Pose.Bones.Num());
+		FCompactPose BasePose;
+		BasePose.SetBoneContainer(&OutPose.GetBoneContainer());
+		FCompactPose AdditivePose;
+		AdditivePose.SetBoneContainer(&OutPose.GetBoneContainer());
 
 #if WITH_EDITORONLY_DATA
 		if (BlendSpace->PreviewBasePose)
 		{
-			BlendSpace->PreviewBasePose->GetBonePose(/*out*/ BasePose.Bones, RequiredBones, FAnimExtractContext(PreviewPoseCurrentTime));
+			BlendSpace->PreviewBasePose->GetBonePose(/*out*/ BasePose, FAnimExtractContext(PreviewPoseCurrentTime));
 		}
 		else
 #endif // WITH_EDITORONLY_DATA
 		{
 			// otherwise, get ref pose
-			FAnimationRuntime::FillWithRefPose(BasePose.Bones, RequiredBones);
+			BasePose.ResetToRefPose();
 		}
 
 		FAnimationRuntime::GetPoseFromBlendSpace(
 			BlendSpace,
-			BlendSampleDataCache, 
-			RequiredBones,
-			/*out*/ AdditivePose.Bones);
+			BlendSampleDataCache,
+			/*out*/ AdditivePose);
 
 		if (BlendSpace->IsA(UAimOffsetBlendSpace::StaticClass()) ||
 			BlendSpace->IsA(UAimOffsetBlendSpace1D::StaticClass()) )
 		{
-			BlendRotationOffset(BasePose, AdditivePose, 1.f, Pose);
+			BlendRotationOffset(BasePose, AdditivePose, 1.f, OutPose);
 		}
 		else
 		{
-			ApplyAdditiveSequence(BasePose, AdditivePose, 1.f, Pose);
+			ApplyAdditiveSequence(BasePose, AdditivePose, 1.f, OutPose);
 		}
 	}
 	else
 	{
 		FAnimationRuntime::GetPoseFromBlendSpace(
 			BlendSpace,
-			BlendSampleDataCache, 
-			RequiredBones,
-			/*out*/ Pose.Bones);
+			BlendSampleDataCache,
+			/*out*/ OutPose);
 	}
 }
 
diff --git a/Engine/Source/Runtime/Engine/Private/Animation/AnimationRuntime.cpp b/Engine/Source/Runtime/Engine/Private/Animation/AnimationRuntime.cpp
--- a/Engine/Source/Runtime/Engine/Private/Animation/AnimationRuntime.cpp
+++ b/Engine/Source/Runtime/Engine/Private/Animation/AnimationRuntime.cpp
@@ -13,6 +13,7 @@
 #include "Animation/BlendSpaceBase.h"
 #include "Animation/AnimCompositeBase.h"
 #include "Animation/AnimInstance.h"
+#include "BonePose.h"
 
 DEFINE_LOG_CATEGORY(LogAnimation);
 DEFINE_LOG_CATEGORY(LogRootMotion);
@@ -91,76 +92,70 @@ FORCEINLINE void BlendPose(const FTransformArrayA2& SourcePoses, FTransformArray
 	}
 }
 
-/**
- * Blends together a set of poses, each with a given weight.
- * This function is lightweight, it does not cull out nearly zero weights or check to make sure weights sum to 1.0, the caller should take care of that if needed.
- *
- * The blend is done by taking a weighted sum of each atom, and re-normalizing the quaternion part at the end, not using SLERP.
- * This allows n-way blends, and makes the code much faster, though the angular velocity will not be constant across the blend.
- *
- * @param	ResultAtoms		Output array of relative bone transforms.
- * @param	RequiredBones	Indices of bones that we want to return. Note that bones not in this array will not be modified, so are not safe to access! 
- *							This array must be in strictly increasing order.
- */
-void FAnimationRuntime::BlendPosesTogether(int32 NumPoses, const FTransformArrayA2** SourcePoses, const float* SourceWeights, const FBoneContainer& RequiredBones, /*out*/ FTransformArrayA2& ResultAtoms)
+template <int32 TRANSFORM_BLEND_MODE>
+FORCEINLINE void BlendPose(const FCompactPose& SourcePose, FCompactPose& ResultPose, const float BlendWeight)
 {
-	check(NumPoses > 0);
-	// verification process so that it doesn't crash
-	const TArray<FBoneIndexType> & RequiredBoneIndices = RequiredBones.GetBoneIndicesArray();
-	if (ensure(SourcePoses && ResultAtoms.Num() >= RequiredBoneIndices.Num()))
+	for (FCompactPoseBoneIndex BoneIndex : SourcePose.ForEachBoneIndex())
 	{
-		// debug purpose for now, but this can cause 0 bone transform, so we'd like to catch it 
-		float WeightSum = SourceWeights[0];
-		BlendPose<ETransformBlendMode::Overwrite>(*SourcePoses[0], ResultAtoms, RequiredBoneIndices, SourceWeights[0]);
+		BlendTransform<TRANSFORM_BLEND_MODE>(SourcePose[BoneIndex], ResultPose[BoneIndex], BlendWeight);
+	}
+}
 
-		for (int32 i = 1; i < NumPoses; ++i)
-		{
-			WeightSum += SourceWeights[i];
+void FAnimationRuntime::BlendPosesTogether(
+	const TArray<FCompactPose>& SourcePoses,
+	const TArray<float>&	SourceWeights,
+	/*out*/ FCompactPose& ResultPose)
+{
+	check(SourcePoses.Num() > 0);
 
-			BlendPose<ETransformBlendMode::Accumulate>(*SourcePoses[i], ResultAtoms, RequiredBoneIndices, SourceWeights[i]);
-		}
+	BlendPose<ETransformBlendMode::Overwrite>(SourcePoses[0], ResultPose, SourceWeights[0]);
 
-		ensure (WeightSum != 0.f);
+	for (int32 PoseIndex = 1; PoseIndex < SourcePoses.Num(); ++PoseIndex)
+	{
+		BlendPose<ETransformBlendMode::Accumulate>(SourcePoses[PoseIndex], ResultPose, SourceWeights[PoseIndex]);
+	}
 
-		// Ensure that all of the resulting rotations are normalized
-		if( NumPoses > 1 )
-		{
-			FAnimationRuntime::NormalizeRotations(RequiredBones, ResultAtoms);
-		}
+	// Ensure that all of the resulting rotations are normalized
+	if (SourcePoses.Num() > 1)
+	{
+		ResultPose.NormalizeRotations();
 	}
 }
 
-/**
- * Blends together a set of poses, each with a given weight.
- * This function is lightweight, it does not cull out nearly zero weights or check to make sure weights sum to 1.0, the caller should take care of that if needed.
- *
- * The blend is done by taking a weighted sum of each atom, and re-normalizing the quaternion part at the end, not using SLERP.
- * This allows n-way blends, and makes the code much faster, though the angular velocity will not be constant across the blend.
- *
- * @param	ResultAtoms		Output array of relative bone transforms.
- * @param	RequiredBones	Indices of bones that we want to return. Note that bones not in this array will not be modified, so are not safe to access! 
- *							This array must be in strictly increasing order.
- */
-void FAnimationRuntime::BlendPosesTogether(int32 NumPoses, const TArray<FTransformArrayA2>& SourcePoses, const TArray<float>&	SourceWeights, const FBoneContainer& RequiredBones, /*out*/ FTransformArrayA2& ResultAtoms)
+void FAnimationRuntime::BlendPosesTogether(
+	const TArray<const FCompactPose*>& SourcePoses,
+	const TArray<float>&	SourceWeights,
+	/*out*/ FCompactPose& ResultPose)
 {
-	check(NumPoses > 0);
-
-	const TArray<FBoneIndexType> & RequiredBoneIndices = RequiredBones.GetBoneIndicesArray();
+	check(SourcePoses.Num() > 0);
 
-	BlendPose<ETransformBlendMode::Overwrite>(SourcePoses[0], ResultAtoms, RequiredBoneIndices, SourceWeights[0]);
+	BlendPose<ETransformBlendMode::Overwrite>(*SourcePoses[0], ResultPose, SourceWeights[0]);
 
-	for (int32 i = 1; i < NumPoses; ++i)
+	for (int32 PoseIndex = 1; PoseIndex < SourcePoses.Num(); ++PoseIndex)
 	{
-		BlendPose<ETransformBlendMode::Accumulate>(SourcePoses[i], ResultAtoms, RequiredBoneIndices, SourceWeights[i]);
+		BlendPose<ETransformBlendMode::Accumulate>(*SourcePoses[PoseIndex], ResultPose, SourceWeights[PoseIndex]);
 	}
 
 	// Ensure that all of the resulting rotations are normalized
-	if( NumPoses > 1 )
+	if (SourcePoses.Num() > 1)
 	{
-		FAnimationRuntime::NormalizeRotations(RequiredBones, ResultAtoms);
+		ResultPose.NormalizeRotations();
 	}
 }
 
+void FAnimationRuntime::BlendTwoPosesTogether(
+	const FCompactPose& SourcePose1,
+	const FCompactPose& SourcePose2,
+	const float			WeightOfPose1,
+	/*out*/ FCompactPose& ResultPose)
+{
+	BlendPose<ETransformBlendMode::Overwrite>(SourcePose1, ResultPose, WeightOfPose1);
+	BlendPose<ETransformBlendMode::Accumulate>(SourcePose2, ResultPose, 1.f - WeightOfPose1);
+
+	// Ensure that all of the resulting rotations are normalized
+	ResultPose.NormalizeRotations();
+}
+
 template <int32 TRANSFORM_BLEND_MODE>
 void BlendPosePerBone(const TArray<FBoneIndexType> &RequiredBoneIndices, const TArray<int32>& PerBoneIndices, const FBlendSampleData& BlendSampleDataCache, FTransformArrayA2 &ResultAtoms, const FTransformArrayA2& SourceAtoms)
 {
@@ -186,6 +181,31 @@ void BlendPosePerBone(const TArray<FBoneIndexType> &RequiredBoneIndices, const T
 	}
 }
 
+template <int32 TRANSFORM_BLEND_MODE>
+void BlendPosePerBone(const TArray<int32>& PerBoneIndices, const FBlendSampleData& BlendSampleDataCache, FCompactPose& ResultPose, const FCompactPose& SourcePose)
+{
+	const float BlendWeight = BlendSampleDataCache.GetWeight();
+	TArray<float> PerBoneBlends;
+	for (int32 i = 0; i < BlendSampleDataCache.PerBoneBlendData.Num(); ++i)
+	{
+		PerBoneBlends.Add(FMath::Clamp<float>(BlendSampleDataCache.PerBoneBlendData[i], 0.f, 1.f));
+	}
+
+	for (FCompactPoseBoneIndex BoneIndex : SourcePose.ForEachBoneIndex())
+	{
+		int32 PerBoneIndex = PerBoneIndices[BoneIndex.GetInt()];
+
+		if (PerBoneIndex == INDEX_NONE || !BlendSampleDataCache.PerBoneBlendData.IsValidIndex(PerBoneIndex))
+		{
+			BlendTransform<TRANSFORM_BLEND_MODE>(SourcePose[BoneIndex], ResultPose[BoneIndex], BlendWeight);
+		}
+		else
+		{
+			BlendTransform<TRANSFORM_BLEND_MODE>(SourcePose[BoneIndex], ResultPose[BoneIndex], PerBoneBlends[PerBoneIndex]);
+		}
+	}
+}
+
 /**
  * Blends together a set of poses, each with a given weight.
  * This function is lightweight, it does not cull out nearly zero weights or check to make sure weights sum to 1.0, the caller should take care of that if needed.
@@ -197,71 +217,69 @@ void BlendPosePerBone(const TArray<FBoneIndexType> &RequiredBoneIndices, const T
  * @param	RequiredBones	Indices of bones that we want to return. Note that bones not in this array will not be modified, so are not safe to access! 
  *							This array must be in strictly increasing order.
  */
-void FAnimationRuntime::BlendPosesTogetherPerBone(int32 NumPoses, const TArray<FTransformArrayA2>& SourcePoses, const UBlendSpaceBase* BlendSpace, const TArray<FBlendSampleData>& BlendSampleDataCache, const FBoneContainer& RequiredBones, /*out*/ FTransformArrayA2& ResultAtoms)
+
+void FAnimationRuntime::BlendPosesTogetherPerBone(const TArray<FCompactPose>& SourcePoses, const UBlendSpaceBase* BlendSpace, const TArray<FBlendSampleData>& BlendSampleDataCache, /*out*/ FCompactPose& ResultPose)
 {
-	check(NumPoses > 0);
+	check(SourcePoses.Num() > 0);
 
-	const TArray<FBoneIndexType> & RequiredBoneIndices = RequiredBones.GetBoneIndicesArray();
+	const TArray<FBoneIndexType> & RequiredBoneIndices = ResultPose.GetBoneContainer().GetBoneIndicesArray();
 
 	TArray<int32> PerBoneIndices;
-	PerBoneIndices.AddUninitialized(ResultAtoms.Num());
+	PerBoneIndices.AddUninitialized(ResultPose.GetNumBones());
 	for (int BoneIndex = 0; BoneIndex < PerBoneIndices.Num(); ++BoneIndex)
 	{
-		PerBoneIndices[BoneIndex] = BlendSpace->GetPerBoneInterpolationIndex(RequiredBoneIndices[BoneIndex], RequiredBones);
+		PerBoneIndices[BoneIndex] = BlendSpace->GetPerBoneInterpolationIndex(RequiredBoneIndices[BoneIndex], ResultPose.GetBoneContainer());
 	}
 
-	BlendPosePerBone<ETransformBlendMode::Overwrite>(RequiredBoneIndices, PerBoneIndices, BlendSampleDataCache[0], ResultAtoms, SourcePoses[0]);
+	BlendPosePerBone<ETransformBlendMode::Overwrite>(PerBoneIndices, BlendSampleDataCache[0], ResultPose, SourcePoses[0]);
 
-	for (int32 i = 1; i < NumPoses; ++i)
+	for (int32 i = 1; i < SourcePoses.Num(); ++i)
 	{
-		BlendPosePerBone<ETransformBlendMode::Accumulate>(RequiredBoneIndices, PerBoneIndices, BlendSampleDataCache[i], ResultAtoms, SourcePoses[i]);
+		BlendPosePerBone<ETransformBlendMode::Accumulate>(PerBoneIndices, BlendSampleDataCache[i], ResultPose, SourcePoses[i]);
 	}
 
 	// Ensure that all of the resulting rotations are normalized
-	FAnimationRuntime::NormalizeRotations(RequiredBones, ResultAtoms);
+	ResultPose.NormalizeRotations();
 }
 
-void FAnimationRuntime::BlendPosesTogetherPerBoneInMeshSpace(int32 NumPoses, TArray<FTransformArrayA2>& SourcePoses, const UBlendSpaceBase* BlendSpace, const TArray<FBlendSampleData>& BlendSampleDataCache, const FBoneContainer& RequiredBones, /*out*/ FTransformArrayA2& ResultAtoms)
+void FAnimationRuntime::BlendPosesTogetherPerBoneInMeshSpace(TArray<FCompactPose>& SourcePoses, const UBlendSpaceBase* BlendSpace, const TArray<FBlendSampleData>& BlendSampleDataCache, FCompactPose& ResultPose)
 {
 	FQuat NewRotation;
 	USkeleton* Skeleton = BlendSpace->GetSkeleton();
-	
-	const TArray<FBoneIndexType> & RequiredBoneIndices = RequiredBones.GetBoneIndicesArray();
+
 	// all this is going to do is to convert SourcePoses.Rotation to be mesh space, and then once it goes through BlendPosesTogetherPerBone, convert back to local
-	for (int32 Pose=0; Pose<NumPoses; ++Pose)
+	for (FCompactPose& Pose : SourcePoses)
 	{
-		for (int32 I=0; I<RequiredBoneIndices.Num(); ++I)
+		for (const FCompactPoseBoneIndex BoneIndex : Pose.ForEachBoneIndex())
 		{
-			int32 const & BoneIndex = RequiredBoneIndices[I];
-			const int32 ParentIndex = RequiredBones.GetParentBoneIndex(BoneIndex);
-			if ( ParentIndex!=INDEX_NONE )
+			const FCompactPoseBoneIndex ParentIndex = Pose.GetParentBoneIndex(BoneIndex);
+			if (ParentIndex != INDEX_NONE)
 			{
-				NewRotation = SourcePoses[Pose][ParentIndex].GetRotation()*SourcePoses[Pose][BoneIndex].GetRotation();
+				NewRotation = Pose[ParentIndex].GetRotation()*Pose[BoneIndex].GetRotation();
 				NewRotation.Normalize();
 			}
 			else
 			{
-				NewRotation = SourcePoses[Pose][BoneIndex].GetRotation();
+				NewRotation = Pose[BoneIndex].GetRotation();
 			}
 
 			// now copy back to SourcePoses
-			SourcePoses[Pose][BoneIndex].SetRotation(NewRotation);
+			Pose[BoneIndex].SetRotation(NewRotation);
 		}
 	}
 
 	// now we have mesh space rotation, call BlendPosesTogetherPerBone
-	BlendPosesTogetherPerBone(NumPoses, SourcePoses, BlendSpace, BlendSampleDataCache, RequiredBones, ResultAtoms);
+	BlendPosesTogetherPerBone(SourcePoses, BlendSpace, BlendSampleDataCache, ResultPose);
 
-	// now result atoms has the output with mesh space roatation. Convert back to local space, start from back
-	for (int32 I=RequiredBoneIndices.Num()-1; I>=0; --I)
+	// now result atoms has the output with mesh space rotation. Convert back to local space, start from back
+	for (const FCompactPoseBoneIndex BoneIndex : ResultPose.ForEachBoneIndex())
 	{
-		int32 const & BoneIndex = RequiredBoneIndices[I];
-		const int32 ParentIndex = RequiredBones.GetParentBoneIndex(BoneIndex);
-		if ( ParentIndex!=INDEX_NONE )
+		const FCompactPoseBoneIndex ParentIndex = ResultPose.GetParentBoneIndex(BoneIndex);
+		if (ParentIndex != INDEX_NONE)
 		{
-			FQuat LocalBlendQuat = ResultAtoms[ParentIndex].GetRotation().Inverse()*ResultAtoms[BoneIndex].GetRotation();
-			ResultAtoms[BoneIndex].SetRotation(LocalBlendQuat);
-			ResultAtoms[BoneIndex].NormalizeRotation();
+			FQuat LocalBlendQuat = ResultPose[ParentIndex].GetRotation().Inverse()*ResultPose[BoneIndex].GetRotation();
+			ResultPose[BoneIndex].SetRotation(LocalBlendQuat);
+			ResultPose[BoneIndex].NormalizeRotation();
 		}
 	}
 }
@@ -275,9 +293,9 @@ void FAnimationRuntime::BlendPosesTogetherPerBoneInMeshSpace(int32 NumPoses, TAr
  * @param	RequiredBones	Indices of bones to process, in strictly increasing order.
  * @param	ResultAtoms		Array of BoneTransforms to accumulate to.
  */
-void FAnimationRuntime::BlendPosesAccumulate(const FTransformArrayA2& BlendPoses, const float BlendWeight, const FBoneContainer& RequiredBones, /*inout*/ FTransformArrayA2& ResultAtoms)
+void FAnimationRuntime::BlendPosesAccumulate(const FCompactPose& SourcePose, const float BlendWeight, /*inout*/ FCompactPose& ResultPose)
 {
-	BlendPose<ETransformBlendMode::Accumulate>(BlendPoses, ResultAtoms, RequiredBones.GetBoneIndicesArray(), BlendWeight);
+	BlendPose<ETransformBlendMode::Accumulate>(SourcePose, ResultPose, BlendWeight);
 }
 
 void FAnimationRuntime::LerpBoneTransforms(TArray<FTransform> & A, const TArray<FTransform> & B, float Alpha, const TArray<FBoneIndexType> & RequiredBonesArray)
@@ -308,23 +326,23 @@ void FAnimationRuntime::LerpBoneTransforms(TArray<FTransform> & A, const TArray<
 	}
 }
 
-void FAnimationRuntime::BlendAdditivePose(const FTransformArrayA2& SourcePoses, const FTransformArrayA2& AdditiveBlendPoses, const float BlendWeight, const FBoneContainer& RequiredBones, /*out*/ FTransformArrayA2& ResultAtoms)
+void FAnimationRuntime::BlendAdditivePose(const FCompactPose& SourcePose, const FCompactPose& AdditiveBlendPose, const float BlendWeight,/*out*/ FCompactPose& ResultPose)
 {
 	const ScalarRegister VBlendWeight(BlendWeight);
 	// Subsequent poses need to be blended in
-	const TArray<FBoneIndexType> & RequiredBoneIndices = RequiredBones.GetBoneIndicesArray();
-	for (int32 j = 0; j < RequiredBoneIndices.Num(); ++j)
+
+	for (FCompactPoseBoneIndex BoneIndex : SourcePose.ForEachBoneIndex())
 	{
-		const int32 BoneIndex = RequiredBoneIndices[j];
-		FTransform Additive = AdditiveBlendPoses[BoneIndex];
-		ResultAtoms[BoneIndex] = SourcePoses[BoneIndex];
-		FTransform::BlendFromIdentityAndAccumulate(ResultAtoms[BoneIndex], Additive, VBlendWeight);
+		FTransform Additive = AdditiveBlendPose[BoneIndex];
+		ResultPose[BoneIndex] = SourcePose[BoneIndex];
+		FTransform::BlendFromIdentityAndAccumulate(ResultPose[BoneIndex], Additive, VBlendWeight);
 	}
 
 	// Ensure that all of the resulting rotations are normalized
-	FAnimationRuntime::NormalizeRotations(RequiredBones, ResultAtoms);
+	ResultPose.NormalizeRotations();
 }
 
+
 void FAnimationRuntime::CombineWithAdditiveAnimations(int32 NumAdditivePoses, const FTransformArrayA2** SourceAdditivePoses, const float* SourceAdditiveWeights, const FBoneContainer& RequiredBones, /*inout*/ FTransformArrayA2& Atoms)
 {
 	const TArray<FBoneIndexType> & RequiredBoneIndices = RequiredBones.GetBoneIndicesArray();
@@ -343,69 +361,61 @@ void FAnimationRuntime::CombineWithAdditiveAnimations(int32 NumAdditivePoses, co
 }
 
 void FAnimationRuntime::GetPoseFromSequence(
-	const UAnimSequenceBase* InSequence,
-	const FBoneContainer& RequiredBones,
-	/*out*/ FTransformArrayA2 & ResultAtoms,
-	const FAnimExtractContext& ExtractionContext)
+		const UAnimSequenceBase* InSequence,
+		/*out*/ FCompactPose& ResultPose,
+		const FAnimExtractContext& ExtractionContext)
 {
 	//@TODO: ANIM: ANIMCOMPOSITE: Some of the code in here cares about UAnimSequence members and won't work on something derived from a different branch of UAnimSequenceBase
 	const UAnimSequence* Sequence = Cast<const UAnimSequence>(InSequence);
 
-	if( !Sequence )
+	if (!Sequence)
 	{
 #if 0
 		UE_LOG(LogAnimation, Log, TEXT("FAnimationRuntime::GetPoseFromSequence - %s - No animation data!"), *GetFName());
 #endif
-		FAnimationRuntime::FillWithRefPose(ResultAtoms, RequiredBones);
+		ResultPose.ResetToRefPose();
 		return;
 	}
 
-	Sequence->GetAnimationPose(ResultAtoms, RequiredBones, ExtractionContext);
+	Sequence->GetAnimationPose(ResultPose, ExtractionContext);
 
 	// Check that all bone atoms coming from animation are normalized
 #if DO_CHECK && WITH_EDITORONLY_DATA
-	const TArray<FBoneIndexType> & RequiredBoneIndices = RequiredBones.GetBoneIndicesArray();
-	for (int32 i = 0; i < RequiredBoneIndices.Num(); ++i)
-	{
-		const int32 BoneIndex = RequiredBoneIndices[i];
-		check( ResultAtoms[BoneIndex].IsRotationNormalized() );
-	}
+	check(ResultPose.IsNormalized());
 #endif
 }
 
 void FAnimationRuntime::GetPoseFromAnimTrack(
 	const FAnimTrack& Track,
-	const FBoneContainer& RequiredBones,
-	/*out*/ FTransformArrayA2& ResultAtoms, 
+	/*out*/ FCompactPose& ResultPose,
 	const FAnimExtractContext& ExtractionContext)
 {
-	const int32 NumBones = RequiredBones.GetNumBones();
-
-	TArray<FTransformArrayA2> SourcePoses;
+	TArray<FCompactPose> SourcePoses;
 	TArray<float> SourceWeights;
 	float TotalWeight = 0.f;
 
 	float CurrentTime = FMath::Clamp(ExtractionContext.CurrentTime, 0.f, Track.GetLength());
 
 	// first get all the montage instance weight this slot node has
-	for (int32 I=0; I<Track.AnimSegments.Num(); ++I)
+	for (int32 I = 0; I<Track.AnimSegments.Num(); ++I)
 	{
 		const FAnimSegment& AnimSegment = Track.AnimSegments[I];
 
 		float PositionInAnim = 0.f;
 		float Weight = 0.f;
 		UAnimSequenceBase* AnimRef = AnimSegment.GetAnimationData(CurrentTime, PositionInAnim, Weight);
+
 		// make this to be 1 function
-		if( AnimRef && (Weight > ZERO_ANIMWEIGHT_THRESH) )
+		if (AnimRef && (Weight > ZERO_ANIMWEIGHT_THRESH))
 		{
 			// todo anim: hack - until we fix animcomposite
 			UAnimSequence * Sequence = Cast<UAnimSequence>(AnimRef);
-			if( Sequence )
+			if (Sequence)
 			{
 				const int32 NewIndex = SourceWeights.AddUninitialized(1);
-				SourcePoses.AddZeroed(1);
+				SourcePoses.Add(FCompactPose());
 
-				SourcePoses[NewIndex].AddUninitialized(NumBones);
+				SourcePoses[NewIndex].SetBoneContainer(&ResultPose.GetBoneContainer());
 				SourceWeights[NewIndex] = Weight;
 				TotalWeight += Weight;
 
@@ -413,69 +423,60 @@ void FAnimationRuntime::GetPoseFromAnimTrack(
 				FAnimExtractContext SequenceExtractionContext(ExtractionContext);
 				SequenceExtractionContext.CurrentTime = PositionInAnim;
 				SequenceExtractionContext.bExtractRootMotion &= Sequence->bEnableRootMotion;
-				FAnimationRuntime::GetPoseFromSequence(Sequence, RequiredBones, SourcePoses[NewIndex], SequenceExtractionContext);
+				FAnimationRuntime::GetPoseFromSequence(Sequence, SourcePoses[NewIndex], SequenceExtractionContext);
 			}
 		}
 	}
 
-	if( SourcePoses.Num() == 0 )
+	if (SourcePoses.Num() == 0)
 	{
-		FAnimationRuntime::FillWithRefPose(ResultAtoms, RequiredBones);
+		ResultPose.ResetToRefPose();
 	}
-	else if( SourcePoses.Num() == 1 )
+	else if (SourcePoses.Num() == 1)
 	{
-		ResultAtoms = SourcePoses[0];
+		ResultPose = SourcePoses[0];
 	}
 	else
 	{
 		// If we have SourcePoses.Num() > 0, then we will have a non zero weight.
 		check(TotalWeight >= ZERO_ANIMWEIGHT_THRESH);
-		for (int32 I=0; I<SourceWeights.Num(); ++I)
+		for (int32 I = 0; I < SourceWeights.Num(); ++I)
 		{
 			// normalize I
 			SourceWeights[I] /= TotalWeight;
 		}
-		FAnimationRuntime::BlendPosesTogether(SourcePoses.Num(), SourcePoses, SourceWeights, RequiredBones, ResultAtoms);
+		FAnimationRuntime::BlendPosesTogether(SourcePoses, SourceWeights, ResultPose);
 	}
 }
 
-void FAnimationRuntime::ConvertPoseToAdditive(FTransformArrayA2 & TargetPose, const FTransformArrayA2 & BasePose, const FBoneContainer& RequiredBones)
+void FAnimationRuntime::ConvertPoseToAdditive(FCompactPose& TargetPose, const FCompactPose& BasePose)
 {
-	const TArray<FBoneIndexType> & RequiredBoneIndices = RequiredBones.GetBoneIndicesArray();
-	const int32 NumRequiredBones = RequiredBoneIndices.Num();
-
-	for(int32 Index=0; Index<NumRequiredBones; Index++)
+	for (FCompactPoseBoneIndex BoneIndex : BasePose.ForEachBoneIndex())
 	{
-		const int32 BoneIndex = RequiredBoneIndices[Index];
-
 		FTransform& TargetTransform = TargetPose[BoneIndex];
 		const FTransform& BaseTransform = BasePose[BoneIndex];
 
-		TargetTransform.SetRotation( TargetTransform.GetRotation() * BaseTransform.GetRotation().Inverse() );
-		TargetTransform.SetTranslation( TargetTransform.GetTranslation() - BaseTransform.GetTranslation() );
-		TargetTransform.SetScale3D( TargetTransform.GetScale3D() * BaseTransform.GetSafeScaleReciprocal( BaseTransform.GetScale3D() ) );
+		TargetTransform.SetRotation(TargetTransform.GetRotation() * BaseTransform.GetRotation().Inverse());
+		TargetTransform.SetTranslation(TargetTransform.GetTranslation() - BaseTransform.GetTranslation());
+		TargetTransform.SetScale3D(TargetTransform.GetScale3D() * BaseTransform.GetSafeScaleReciprocal(BaseTransform.GetScale3D()));
 		TargetTransform.NormalizeRotation();
 	}
 }
 
-void FAnimationRuntime::ConvertPoseToMeshRotation(FTransformArrayA2 & LocalPose, const FBoneContainer& RequiredBones)
+void FAnimationRuntime::ConvertPoseToMeshRotation(FCompactPose& LocalPose)
 {
-	const TArray<FBoneIndexType> & RequiredBoneIndices = RequiredBones.GetBoneIndicesArray();
-	const int32 NumRequiredBones = RequiredBoneIndices.Num();
-
-	// only the root bone doesn't have a parent. So skip it to save a branch in the iteration.
-	check( RequiredBoneIndices[0] == 0 );
-
 	// Convert all rotations to mesh space
-	for (int32 Index=1; Index<NumRequiredBones; Index++)
+	for (FCompactPoseBoneIndex BoneIndex : LocalPose.ForEachBoneIndex())
 	{
-		const int32 BoneIndex = RequiredBoneIndices[Index];
-		const int32 ParentIndex = RequiredBones.GetParentBoneIndex(BoneIndex);
+		if (!BoneIndex.IsRootBone()) // Skip root
+		{
+			const FCompactPoseBoneIndex ParentIndex = LocalPose.GetParentBoneIndex(BoneIndex);
 
-		FTransform& LocalTransform = LocalPose[BoneIndex];
-		const FTransform& LocalParentTransform = LocalPose[ParentIndex];
+			FTransform& LocalTransform = LocalPose[BoneIndex];
+			const FTransform& LocalParentTransform = LocalPose[ParentIndex];
 
-		LocalTransform.SetRotation( LocalParentTransform.GetRotation() * LocalTransform.GetRotation() );
+			LocalTransform.SetRotation(LocalParentTransform.GetRotation() * LocalTransform.GetRotation());
+		}
 	}
 }
 
@@ -533,36 +534,36 @@ void FAnimationRuntime::ApplyWeightToTransform(const FBoneContainer& RequiredBon
 void FAnimationRuntime::GetPoseFromBlendSpace(
 	UBlendSpaceBase * BlendSpace,
 	TArray<FBlendSampleData>& BlendSampleDataCache,
-	FBoneContainer & RequiredBones,
-	/*out*/ FTransformArrayA2& ResultAtoms)
+	/*out*/ FCompactPose& ResultPose)
 {
 	// if not valid parameters
 	if (BlendSpace == NULL || BlendSampleDataCache.Num() == 0)
 	{
-		FillWithRefPose(ResultAtoms, RequiredBones);
+		ResultPose.ResetToRefPose();
 		return;
 	}
 
 	const int32 NumPoses = BlendSampleDataCache.Num();
 
-	TArray<FTransformArrayA2> ChildrenTransform;
-	ChildrenTransform.AddZeroed(NumPoses);
+	TArray<FCompactPose> ChildrenPoses;
+	ChildrenPoses.AddZeroed(NumPoses);
 
-	TArray<float>				ChildrenWeights;
+	TArray<float>	ChildrenWeights;
 	ChildrenWeights.AddZeroed(NumPoses);
 
-	const int32 NumBones = RequiredBones.GetNumBones();
-	for (auto ChIter=ChildrenTransform.CreateIterator(); ChIter; ++ChIter)
+	for (FCompactPose& Pose : ChildrenPoses)
 	{
-		(*ChIter).AddUninitialized(NumBones);
+		Pose.SetBoneContainer(&ResultPose.GetBoneContainer());
 	}
 
 	const TArray<struct FBlendSample> & SampleData = BlendSpace->GetBlendSamples();
 
 	// get all child atoms we interested in
-	for (int32 I=0; I<BlendSampleDataCache.Num(); ++I)
+	for (int32 I = 0; I < BlendSampleDataCache.Num(); ++I)
 	{
-		if ( SampleData.IsValidIndex(BlendSampleDataCache[I].SampleDataIndex) )
+		FCompactPose& Pose = ChildrenPoses[I];
+
+		if (SampleData.IsValidIndex(BlendSampleDataCache[I].SampleDataIndex))
 		{
 			const FBlendSample& Sample = SampleData[BlendSampleDataCache[I].SampleDataIndex];
 			ChildrenWeights[I] = BlendSampleDataCache[I].GetWeight();
@@ -572,44 +573,44 @@ void FAnimationRuntime::GetPoseFromBlendSpace(
 				const float Time = FMath::Clamp<float>(BlendSampleDataCache[I].Time, 0.f, Sample.Animation->SequenceLength);
 
 				// first one always fills up the source one
-				FAnimationRuntime::GetPoseFromSequence(Sample.Animation, RequiredBones, ChildrenTransform[I], FAnimExtractContext(Time, true));
+				FAnimationRuntime::GetPoseFromSequence(Sample.Animation, Pose, FAnimExtractContext(Time, true));
 			}
 			else
 			{
-				FAnimationRuntime::FillWithRefPose(ChildrenTransform[I], RequiredBones);
+				Pose.ResetToRefPose();
 			}
 		}
 		else
 		{
-			FAnimationRuntime::FillWithRefPose(ChildrenTransform[I], RequiredBones);
+			Pose.ResetToRefPose();
 		}
 	}
 
-	if ( BlendSpace->PerBoneBlend.Num() > 0 )
+	if (BlendSpace->PerBoneBlend.Num() > 0)
 	{
-		if ( BlendSpace->IsValidAdditive() )
+		if (BlendSpace->IsValidAdditive())
 		{
-			FAnimationRuntime::BlendPosesTogetherPerBone(NumPoses, ChildrenTransform, BlendSpace, BlendSampleDataCache, RequiredBones, ResultAtoms);
+			FAnimationRuntime::BlendPosesTogetherPerBone(ChildrenPoses, BlendSpace, BlendSampleDataCache, ResultPose);
 		}
 		else
 		{
-			if ( BlendSpace->bRotationBlendInMeshSpace )
+			if (BlendSpace->bRotationBlendInMeshSpace)
 			{
-				FAnimationRuntime::BlendPosesTogetherPerBoneInMeshSpace(NumPoses, ChildrenTransform, BlendSpace, BlendSampleDataCache, RequiredBones, ResultAtoms);
+				FAnimationRuntime::BlendPosesTogetherPerBoneInMeshSpace(ChildrenPoses, BlendSpace, BlendSampleDataCache, ResultPose);
 			}
 			else
 			{
-				FAnimationRuntime::BlendPosesTogetherPerBone(NumPoses, ChildrenTransform, BlendSpace, BlendSampleDataCache, RequiredBones, ResultAtoms);
+				FAnimationRuntime::BlendPosesTogetherPerBone(ChildrenPoses, BlendSpace, BlendSampleDataCache, ResultPose);
 			}
 		}
 	}
 	else
 	{
-		FAnimationRuntime::BlendPosesTogether(NumPoses, ChildrenTransform, ChildrenWeights, RequiredBones, ResultAtoms);
+		FAnimationRuntime::BlendPosesTogether(ChildrenPoses, ChildrenWeights, ResultPose);
 	}
 
 	// Once all the accumulation and blending has been done, normalize rotations.
-	FAnimationRuntime::NormalizeRotations(RequiredBones, ResultAtoms);
+	ResultPose.NormalizeRotations();
 }
 
 
@@ -688,25 +689,17 @@ void FAnimationRuntime::FillWithRefPose(TArray<FTransform> & OutAtoms, const FBo
 }
 
 #if WITH_EDITOR
-void FAnimationRuntime::FillWithRetargetBaseRefPose( TArray<FTransform> & OutAtoms, const USkeletalMesh* Mesh, const FBoneContainer& RequiredBones )
+void FAnimationRuntime::FillWithRetargetBaseRefPose(FCompactPose& OutPose, const USkeletalMesh* Mesh)
 {
 	// Copy Target Asset's ref pose.
 	if (Mesh)
 	{
-		const int32 NumBones = RequiredBones.GetNumBones();
-		const TArray<FBoneIndexType> & BoneIndices = RequiredBones.GetBoneIndicesArray();
-		OutAtoms.Empty(NumBones);
-		OutAtoms.AddUninitialized(NumBones);
-		for (int32 Idx = 0; Idx < BoneIndices.Num(); ++Idx)
+		for (FCompactPoseBoneIndex BoneIndex : OutPose.ForEachBoneIndex())
 		{
-			const int32 PoseIdx = BoneIndices[Idx];
-			OutAtoms[PoseIdx] = Mesh->RetargetBasePose[PoseIdx];
+			FMeshPoseBoneIndex PoseIndex = OutPose.GetBoneContainer().MakeMeshPoseIndex(BoneIndex);
+			OutPose[BoneIndex] = Mesh->RetargetBasePose[PoseIndex.GetInt()];
 		}
 	}
-	else
-	{
-		OutAtoms.Empty();
-	}
 }
 #endif // WITH_EDITOR
 
@@ -825,62 +818,59 @@ void FAnimationRuntime::ExcludeBonesWithNoParents(const TArray<int32> & BoneIndi
 	}
 }
 
-void FAnimationRuntime::BlendMeshPosesPerBoneWeights( 
-		struct FA2Pose& BasePose, 
-		TArray<struct FA2Pose>& BlendPoses, 
-		const TArray<FPerBoneBlendWeight> & BoneBlendWeights, 
-		const FBoneContainer& RequiredBones, 
-		USkeleton* Skeleton,
-		/*out*/ FA2Pose& OutPose)
+void FAnimationRuntime::BlendMeshPosesPerBoneWeights(
+		struct FCompactPose& BasePose,
+		const TArray<struct FCompactPose>& BlendPoses,
+		const TArray<FPerBoneBlendWeight> & BoneBlendWeights,
+		/*out*/ FCompactPose& OutPose)
 {
-	check ( BasePose.Bones.Num() == BoneBlendWeights.Num() );
+	check(BasePose.GetNumBones() == BoneBlendWeights.Num());
 
-	float BlendWeight;
+	const FBoneContainer& BoneContainer = BasePose.GetBoneContainer();
 
-	TArray<FQuat> SourceRotations;
-	TArray<FQuat> BlendRotations;
-	TArray<FQuat> TargetRotations;
-	SourceRotations.AddUninitialized(BasePose.Bones.Num());
-	BlendRotations.AddUninitialized(BasePose.Bones.Num());
-	TargetRotations.AddUninitialized(BasePose.Bones.Num());
-	
-	const TArray<FBoneIndexType> & RequiredBoneIndices = RequiredBones.GetBoneIndicesArray();
-	for (int32 I=0; I<RequiredBoneIndices.Num(); ++I)
+	TCustomBoneIndexArray<FQuat, FCompactPoseBoneIndex> SourceRotations;
+	TCustomBoneIndexArray<FQuat, FCompactPoseBoneIndex> BlendRotations;
+	TCustomBoneIndexArray<FQuat, FCompactPoseBoneIndex> TargetRotations;
+
+	SourceRotations.AddUninitialized(BasePose.GetNumBones());
+	BlendRotations.AddUninitialized(BasePose.GetNumBones());
+	TargetRotations.AddUninitialized(BasePose.GetNumBones());
+
+	for (FCompactPoseBoneIndex BoneIndex : BasePose.ForEachBoneIndex())
 	{
-		int32 const & BoneIndex = RequiredBoneIndices[I];
-		int32 const & PoseIndex = BoneBlendWeights[BoneIndex].SourceIndex;
-		const int32 ParentIndex = RequiredBones.GetParentBoneIndex(BoneIndex);
+		const int32 PoseIndex = BoneBlendWeights[BoneIndex.GetInt()].SourceIndex;
+		const FCompactPoseBoneIndex ParentIndex = BoneContainer.GetParentBoneIndex(BoneIndex);
 
 		FQuat SrcRotationInMesh;
 		FQuat TargetRotationInMesh;
 
-		if (ParentIndex!=INDEX_NONE)
+		if (ParentIndex != INDEX_NONE)
 		{
-			SrcRotationInMesh = SourceRotations[ParentIndex]*BasePose.Bones[BoneIndex].GetRotation();
-			TargetRotationInMesh = TargetRotations[ParentIndex]*BlendPoses[PoseIndex].Bones[BoneIndex].GetRotation();
+			SrcRotationInMesh = SourceRotations[ParentIndex] * BasePose[BoneIndex].GetRotation();
+			TargetRotationInMesh = TargetRotations[ParentIndex] * BlendPoses[PoseIndex][BoneIndex].GetRotation();
 		}
 		else
 		{
-			SrcRotationInMesh = BasePose.Bones[BoneIndex].GetRotation();
-			TargetRotationInMesh = BlendPoses[PoseIndex].Bones[BoneIndex].GetRotation();
+			SrcRotationInMesh = BasePose[BoneIndex].GetRotation();
+			TargetRotationInMesh = BlendPoses[PoseIndex][BoneIndex].GetRotation();
 		}
-	
+
 		// update mesh based rotations
 		SourceRotations[BoneIndex] = SrcRotationInMesh;
 		TargetRotations[BoneIndex] = TargetRotationInMesh;
 
 		// now update outer
-		FTransform BaseAtom = BasePose.Bones[BoneIndex];
-		FTransform TargetAtom = BlendPoses[PoseIndex].Bones[BoneIndex];
+		FTransform BaseAtom = BasePose[BoneIndex];
+		FTransform TargetAtom = BlendPoses[PoseIndex][BoneIndex];
 		FTransform BlendAtom;
 
-		BlendWeight = FMath::Clamp(BoneBlendWeights[BoneIndex].BlendWeight, 0.f, 1.f);
-		if ( BlendWeight < ZERO_ANIMWEIGHT_THRESH )
+		const float BlendWeight = FMath::Clamp(BoneBlendWeights[BoneIndex.GetInt()].BlendWeight, 0.f, 1.f);
+		if (BlendWeight < ZERO_ANIMWEIGHT_THRESH)
 		{
 			BlendAtom = BaseAtom;
 			BlendRotations[BoneIndex] = SourceRotations[BoneIndex];
 		}
-		else if ( (1.0 - BlendWeight) < ZERO_ANIMWEIGHT_THRESH )
+		else if ((1.0 - BlendWeight) < ZERO_ANIMWEIGHT_THRESH)
 		{
 			BlendAtom = TargetAtom;
 			BlendRotations[BoneIndex] = TargetRotations[BoneIndex];
@@ -892,57 +882,52 @@ void FAnimationRuntime::BlendMeshPosesPerBoneWeights(
 
 			// blend rotation in mesh space
 			BlendRotations[BoneIndex] = FQuat::FastLerp(SourceRotations[BoneIndex], TargetRotations[BoneIndex], BlendWeight);
-			
+
 			// Fast lerp produces un-normalized quaternions, re-normalize.
 			BlendRotations[BoneIndex].Normalize();
 		}
 
-		OutPose.Bones[BoneIndex] = BlendAtom;
-		if (ParentIndex!=INDEX_NONE)
+		OutPose[BoneIndex] = BlendAtom;
+		if (ParentIndex != INDEX_NONE)
 		{
 			FQuat LocalBlendQuat = BlendRotations[ParentIndex].Inverse() * BlendRotations[BoneIndex];
-			
+
 			// local -> mesh -> local transformations can cause loss of precision for long bone chains, we have to normalize rotation there.
 			LocalBlendQuat.Normalize();
-			OutPose.Bones[BoneIndex].SetRotation(LocalBlendQuat);
+			OutPose[BoneIndex].SetRotation(LocalBlendQuat);
 		}
 	}
 }
 
-void FAnimationRuntime::BlendLocalPosesPerBoneWeights( 
-	FA2Pose& BasePose, 
-	TArray<FA2Pose>& BlendPoses, 
-	const TArray<FPerBoneBlendWeight> & BoneBlendWeights, 
-	const FBoneContainer& RequiredBones, 
-	/*out*/ FA2Pose& OutPose)
+void FAnimationRuntime::BlendLocalPosesPerBoneWeights(
+	FCompactPose& BasePose,
+	const TArray<FCompactPose>& BlendPoses,
+	const TArray<FPerBoneBlendWeight> & BoneBlendWeights,
+	/*out*/ FCompactPose& OutPose)
 {
-	check ( BasePose.Bones.Num() == BoneBlendWeights.Num() );
+	check(BasePose.GetNumBones() == BoneBlendWeights.Num());
 
-	float BlendWeight;
-
-	const TArray<FBoneIndexType> & RequiredBoneIndices = RequiredBones.GetBoneIndicesArray();
-	for (int32 I=0; I<RequiredBoneIndices.Num(); ++I)
+	for (FCompactPoseBoneIndex BoneIndex : BasePose.ForEachBoneIndex())
 	{
-		int32 const & BoneIndex = RequiredBoneIndices[I];
-		int32 const & PoseIndex = BoneBlendWeights[BoneIndex].SourceIndex;
-		FTransform BaseAtom = GetSpaceTransform(BasePose, BoneIndex);
+		const int32 PoseIndex = BoneBlendWeights[BoneIndex.GetInt()].SourceIndex;
+		const FTransform& BaseAtom = BasePose[BoneIndex];
 
-		BlendWeight = FMath::Clamp(BoneBlendWeights[BoneIndex].BlendWeight, 0.f, 1.f);
-		if ( BlendWeight < ZERO_ANIMWEIGHT_THRESH )
+		const float BlendWeight = FMath::Clamp(BoneBlendWeights[BoneIndex.GetInt()].BlendWeight, 0.f, 1.f);
+
+		if (BlendWeight < ZERO_ANIMWEIGHT_THRESH)
 		{
-			SetSpaceTransform(OutPose, BoneIndex, BaseAtom);
+			OutPose[BoneIndex] = BaseAtom;
 		}
-		else if ( (1.0 - BlendWeight) < ZERO_ANIMWEIGHT_THRESH )
+		else if ((1.0 - BlendWeight) < ZERO_ANIMWEIGHT_THRESH)
 		{
-			FTransform BlendAtom = GetSpaceTransform(BlendPoses[PoseIndex], BoneIndex);
-			SetSpaceTransform(OutPose, BoneIndex, BlendAtom);
+			OutPose[BoneIndex] = BlendPoses[PoseIndex][BoneIndex];
 		}
 		else // we want blend here
 		{
-			FTransform BlendAtom = GetSpaceTransform(BasePose, BoneIndex);
-			FTransform TargetAtom = GetSpaceTransform(BlendPoses[PoseIndex], BoneIndex);
+			FTransform BlendAtom = BaseAtom;
+			const FTransform& TargetAtom = BlendPoses[PoseIndex][BoneIndex];
 			BlendAtom.BlendWith(TargetAtom, BlendWeight);
-			SetSpaceTransform(OutPose, BoneIndex, BlendAtom);
+			OutPose[BoneIndex] = BlendAtom;
 		}
 	}
 }
@@ -970,19 +955,19 @@ void FAnimationRuntime::UpdateDesiredBoneWeight(const TArray<FPerBoneBlendWeight
 	}
 }
 
-void FAnimationRuntime::BlendPosesPerBoneFilter(struct FA2Pose& BasePose, const TArray<struct FA2Pose>& BlendPoses, struct FA2Pose& OutPose, TArray<FPerBoneBlendWeight>& BoneBlendWeights, bool bMeshSpaceRotationBlending, const FBoneContainer& RequiredBones, USkeleton* Skeleton)
+void FAnimationRuntime::BlendPosesPerBoneFilter(struct FCompactPose& BasePose, const TArray<struct FCompactPose>& BlendPoses, struct FCompactPose& OutPose, TArray<FPerBoneBlendWeight>& BoneBlendWeights, bool bMeshSpaceRotationBlending)
 {
-	ensure (OutPose.Bones.Num() == BasePose.Bones.Num());
+	ensure(OutPose.GetNumBones() == BasePose.GetNumBones());
+
 	if (BlendPoses.Num() != 0)
 	{
-		TArray<FA2Pose>& NewBlendPoses = (TArray<FA2Pose>&)BlendPoses;
 		if (bMeshSpaceRotationBlending)
 		{
-			BlendMeshPosesPerBoneWeights(BasePose, NewBlendPoses, BoneBlendWeights, RequiredBones, Skeleton, OutPose);
+			BlendMeshPosesPerBoneWeights(BasePose, BlendPoses, BoneBlendWeights, OutPose);
 		}
 		else
 		{
-			BlendLocalPosesPerBoneWeights(BasePose, NewBlendPoses, BoneBlendWeights, RequiredBones, OutPose);
+			BlendLocalPosesPerBoneWeights(BasePose, BlendPoses, BoneBlendWeights, OutPose);
 		}
 	}
 	else // if no blendpose, outpose = basepose
@@ -1039,9 +1024,9 @@ void FAnimationRuntime::CreateMaskWeights(int32 NumOfBones, TArray<FPerBoneBlend
 void FAnimationRuntime::ConvertCSTransformToBoneSpace
 (
 	USkeletalMeshComponent * SkelComp,  
-	FA2CSPose & MeshBases, 
+	FCSPose<FCompactPose> & MeshBases,
 	/**inout*/ FTransform& CSBoneTM, 
-	int32 BoneIndex, 
+	FCompactPoseBoneIndex BoneIndex,
 	uint8 Space
 )
 {
@@ -1058,10 +1043,10 @@ void FAnimationRuntime::ConvertCSTransformToBoneSpace
 
 		case BCS_ParentBoneSpace :
 			{
-				const int32 ParentIndex = MeshBases.GetParentBoneIndex(BoneIndex);
+				const FCompactPoseBoneIndex ParentIndex = MeshBases.GetPose().GetParentBoneIndex(BoneIndex);
 				if (ParentIndex != INDEX_NONE)
 				{
-					const FTransform ParentTM = MeshBases.GetComponentSpaceTransform(ParentIndex);
+					const FTransform& ParentTM = MeshBases.GetComponentSpaceTransform(ParentIndex);
 					CSBoneTM.SetToRelativeTransform(ParentTM);
 				}
 			}
@@ -1069,7 +1054,7 @@ void FAnimationRuntime::ConvertCSTransformToBoneSpace
 
 		case BCS_BoneSpace :
 			{
-				const FTransform BoneTM = MeshBases.GetComponentSpaceTransform(BoneIndex);
+				const FTransform& BoneTM = MeshBases.GetComponentSpaceTransform(BoneIndex);
 				CSBoneTM.SetToRelativeTransform(BoneTM);
 			}
 			break;
@@ -1084,9 +1069,9 @@ void FAnimationRuntime::ConvertCSTransformToBoneSpace
 void FAnimationRuntime::ConvertBoneSpaceTransformToCS
 (
 	USkeletalMeshComponent * SkelComp,  
-	FA2CSPose & MeshBases, 
+	FCSPose<FCompactPose> & MeshBases,
 	/*inout*/ FTransform& BoneSpaceTM, 
-	int32 BoneIndex, 
+	FCompactPoseBoneIndex BoneIndex,
 	uint8 Space
 )
 {
@@ -1103,10 +1088,10 @@ void FAnimationRuntime::ConvertBoneSpaceTransformToCS
 		case BCS_ParentBoneSpace :
 			if( BoneIndex != INDEX_NONE )
 			{
-				const int32 ParentIndex = MeshBases.GetParentBoneIndex(BoneIndex);
+				const FCompactPoseBoneIndex ParentIndex = MeshBases.GetPose().GetParentBoneIndex(BoneIndex);
 				if( ParentIndex != INDEX_NONE )
 				{
-					const FTransform ParentTM = MeshBases.GetComponentSpaceTransform(ParentIndex);
+					const FTransform& ParentTM = MeshBases.GetComponentSpaceTransform(ParentIndex);
 					BoneSpaceTM *= ParentTM;
 				}
 			}
@@ -1115,7 +1100,7 @@ void FAnimationRuntime::ConvertBoneSpaceTransformToCS
 		case BCS_BoneSpace :
 			if( BoneIndex != INDEX_NONE )
 			{
-				const FTransform BoneTM = MeshBases.GetComponentSpaceTransform(BoneIndex);
+				const FTransform& BoneTM = MeshBases.GetComponentSpaceTransform(BoneIndex);
 				BoneSpaceTM *= BoneTM;
 			}
 			break;
@@ -1420,215 +1405,3 @@ void FA2CSPose::ConvertToLocalPoses(FA2Pose & LocalPoses)  const
 	}
 }
 
-/** 
- * Set a bunch of Component Space Bone Transforms.
- * Do this safely by insuring that Parents are already in Component Space,
- * and any Component Space children are converted back to Local Space before hand.
- */
-void FA2CSPose::SafeSetCSBoneTransforms(const TArray<struct FBoneTransform> & BoneTransforms)
-{
-	checkSlow( IsValid() );
-
-	// Bone Mask to keep track of which bones have to be converted to local space.
-	// This is basically BoneTransforms's children.
-	TArray<uint8> BoneMask;
-	BoneMask.AddZeroed( Bones.Num() );
-
-	// First build our BoneMask
-	for (int32 Index=0; Index<BoneTransforms.Num(); Index++)	
-	{
-		const int32 BoneIndex = BoneTransforms[Index].BoneIndex;
-		// Mark those bones in Mesh Pose as being required to be in Local Space.
-		BoneMask[BoneIndex] = 1;
-	}
-
-	// First Flag children
-	const TArray<FBoneIndexType> & RequiredBoneIndices = BoneContainer->GetBoneIndicesArray();
-	for(int32 Index=0; Index<RequiredBoneIndices.Num(); ++Index)
-	{
-		const int32 BoneIndex = RequiredBoneIndices[Index];
-		const int32 ParentIndex = BoneContainer->GetParentBoneIndex(BoneIndex);
-		// Propagate our BoneMask to children.
-		if( ParentIndex != INDEX_NONE )
-		{
-			BoneMask[BoneIndex] |= BoneMask[ParentIndex];
-		}
-	}
-
-	// now iterate from children to parent to calculate back to local space
-	// Sadly this has to iterate from back because you'll need parent to be component space
-	// if you do this from parent to child, you'll see parent turned to local before children applies it
-	int32 BoneTransformIndex = BoneTransforms.Num() - 1;
-	for(int32 Index=RequiredBoneIndices.Num()-1; Index>=0; --Index)
-	{
-		const int32 BoneIndex = RequiredBoneIndices[Index];
-		// If this bone has to be converted to Local Space...
-		if( BoneMask[BoneIndex] != 0 )
-		{
-			// If this is one of the original BoneTransforms list bones
-			// Then we don't actually want to convert this one to local space, since we're going to overwrite it.
-			// So skip it.
-			if( BoneTransformIndex >=0 && BoneIndex == BoneTransforms[BoneTransformIndex].BoneIndex )
-			{
-				BoneTransformIndex--;
-			}
-			// If this is a children bone then we want it to be in local space!
-			else
-			{
-				// .. If it is not currently in Local Space, then convert it.
-				ConvertBoneToLocalSpace(BoneIndex);
-			}
-		}
-	}
-
-	// Finally copy our Component Space transforms
-	for (int32 Index=0; Index<BoneTransforms.Num(); Index++)	
-	{
-		const int32 BoneIndex = BoneTransforms[Index].BoneIndex;
-		
-		// Make sure our BoneTransforms were in Component Space in the first place, before we overwrite them
-		// Only check their parent to do minimal work needed.
-		const int32 ParentBoneIndex = BoneContainer->GetParentBoneIndex(BoneIndex);
-		if( ParentBoneIndex != INDEX_NONE && ComponentSpaceFlags[ParentBoneIndex] == 0 )
-		{
-			CalculateComponentSpaceTransform(ParentBoneIndex);
-		}
-
-		// Set new Component Space transform.
-		SetComponentSpaceTransform( BoneIndex, BoneTransforms[Index].Transform );
-	}
-}
-
-/** 
- * Blends Component Space transforms to MeshPose in Local Space. 
- * Used by SkelControls to apply their transforms.
- *
- * The tricky bit is that SkelControls deliver their transforms in Component Space,
- * But the blending is done in Local Space. Also we need to refresh any Children they have
- * that has been previously converted to Component Space.
- */
-void FA2CSPose::LocalBlendCSBoneTransforms
-(
-	const TArray<struct FBoneTransform>& BoneTransforms,
-	float Alpha
-)
-{
-	// if Alpha is small enough, skip
-	if( Alpha < ZERO_ANIMWEIGHT_THRESH )
-	{
-		return;
-	}
-
-#if DO_CHECK
-	{
-		int32 LastIndex = INDEX_NONE;
-		// Make sure bones are sorted in "Parents before Children" order.
-		for (int32 I=0; I<BoneTransforms.Num(); ++I)
-		{
-			check(BoneTransforms[I].BoneIndex != INDEX_NONE);
-			check(BoneTransforms[I].BoneIndex >= LastIndex);
-			LastIndex = BoneTransforms[I].BoneIndex;
-		}
-	}
-#endif
-
-	// If we are not doing any blending, do a faster path.
-	// Set transforms directly in Component Space. But still refresh children.
-	if( Alpha >= 1.f - ZERO_ANIMWEIGHT_THRESH )
-	{
-		SafeSetCSBoneTransforms(BoneTransforms);
-	}
-	// Slower blending path.
-	else
-	{
-		// Bone Mask to keep track of which bones have to be converted to local space.
-		// This is basically BoneTransforms bones and their children.
-		TArray<uint8> BoneMask;
-		BoneMask.AddZeroed( Bones.Num() );
-
-		TArray<struct FBoneTransform> LocalBoneTransforms;
-		LocalBoneTransforms.SetNumUninitialized( BoneTransforms.Num() );
-
-		const TArray<FBoneIndexType> & RequiredBoneIndices = BoneContainer->GetBoneIndicesArray();
-		// First, convert BoneTransforms to local space for blending.
-		for (int32 Index=0; Index<BoneTransforms.Num(); Index++)	
-		{
-			const int32 BoneIndex = BoneTransforms[Index].BoneIndex;
-			const int32 ParentIndex = BoneContainer->GetParentBoneIndex(BoneIndex);
-			if( ParentIndex != INDEX_NONE )
-			{
-				// if BoneTransforms(modified by controllers) contains ParentIndex, it should use that as ParentTransform, not the one from input
-				int32 LocalParentIndex = INDEX_NONE;
-				for ( int32 LocalIndex = 0; LocalIndex<BoneTransforms.Num(); ++LocalIndex )
-				{
-					if (ParentIndex == BoneTransforms[LocalIndex].BoneIndex)
-					{
-						LocalParentIndex = LocalIndex;
-						break;
-					}
-				}
-
-				// saves Parent Transform
-				FTransform ParentTransform;
-				if ( LocalParentIndex == INDEX_NONE )
-				{
-					ParentTransform = GetComponentSpaceTransform(ParentIndex);
-				}
-				else
-				{
-					ParentTransform = BoneTransforms[LocalParentIndex].Transform;
-				}
-
-				LocalBoneTransforms[Index].Transform = BoneTransforms[Index].Transform.GetRelativeTransform(ParentTransform);
-				LocalBoneTransforms[Index].BoneIndex = BoneIndex;
-
-				// Mark those bones in Mesh Pose as being required to be in Local Space.
-				BoneMask[BoneIndex] = 1;
-			}
-			else
-			{
-				// when root is entered as to modify, we don't need to adjust parent index, just clear it
-				LocalBoneTransforms[Index].Transform =  BoneTransforms[Index].Transform;
-				LocalBoneTransforms[Index].BoneIndex = BoneIndex;
-
-				BoneMask[BoneIndex] = 1;
-			}
-		}
-
-		// Then, convert MeshPose Bones from BoneTransforms list, and their children, to local space if they are not already.
-		for(int32 Index=0; Index<RequiredBoneIndices.Num(); ++Index)
-		{
-			const int32 BoneIndex = RequiredBoneIndices[Index];
-			const int32 ParentIndex = BoneContainer->GetParentBoneIndex(BoneIndex);
-			// Propagate our BoneMask to children.
-			if( ParentIndex != INDEX_NONE )
-			{
-				BoneMask[BoneIndex] |= BoneMask[ParentIndex];
-			}
-		}
-
-		for(int32 Index=RequiredBoneIndices.Num()-1; Index>=0; --Index)		
-		{
-			const int32 BoneIndex = RequiredBoneIndices[Index];
-			// If this bone has to be converted to Local Space...
-			if( BoneMask[BoneIndex] != 0 )
-			{
-				// .. And is not currently in Local Space, then convert it.
-				ConvertBoneToLocalSpace(BoneIndex);
-			}
-		}
-
-		// Lastly, do the blending in local space.
-		for (int32 Index=0; Index<LocalBoneTransforms.Num(); Index++)	
-		{
-			const int32 BoneIndex = LocalBoneTransforms[Index].BoneIndex;
-			// Make sure this transform is in local space, because we are writing a local space one to it.
-			// If we are not in local space, this could mean trouble for our children.
-			check( (ComponentSpaceFlags[BoneIndex] == 0) || (BoneIndex == 0) );
-
-			// No need to normalize rotation since BlendWith() does it.
-			Bones[BoneIndex].BlendWith(LocalBoneTransforms[Index].Transform, Alpha);
-		}
-	}
-}
-
diff --git a/Engine/Source/Runtime/Engine/Private/Animation/BonePose.cpp b/Engine/Source/Runtime/Engine/Private/Animation/BonePose.cpp
new file mode 100644
--- /dev/null
+++ b/Engine/Source/Runtime/Engine/Private/Animation/BonePose.cpp
@@ -0,0 +1,123 @@
+// Copyright 1998-2015 Epic Games, Inc. All Rights Reserved.
+
+#include "EnginePrivate.h"
+#include "BonePose.h"
+
+void FMeshPose::ResetToRefPose()
+{
+	FAnimationRuntime::FillWithRefPose(Bones, *BoneContainer);
+}
+
+void FMeshPose::ResetToIdentity()
+{
+	FAnimationRuntime::InitializeTransform(*BoneContainer, Bones);
+}
+
+void FMeshPose::InitFrom(FCompactPose& CompactPose)
+{
+	SetBoneContainer(&CompactPose.GetBoneContainer());
+	for (const FCompactPoseBoneIndex BoneIndex : CompactPose.ForEachBoneIndex())
+	{
+		FMeshPoseBoneIndex MeshPoseBoneIndex = GetBoneContainer().MakeMeshPoseIndex(BoneIndex);
+		Bones[MeshPoseBoneIndex.GetInt()] = CompactPose[BoneIndex];
+	}
+}
+
+bool FMeshPose::ContainsNaN() const
+{
+	const TArray<FBoneIndexType> & RequiredBoneIndices = BoneContainer->GetBoneIndicesArray();
+	for (int32 Iter = 0; Iter < RequiredBoneIndices.Num(); ++Iter)
+	{
+		const int32 BoneIndex = RequiredBoneIndices[Iter];
+		if (Bones[BoneIndex].ContainsNaN())
+		{
+			return true;
+		}
+	}
+
+	return false;
+}
+
+bool FMeshPose::IsNormalized() const
+{
+	const TArray<FBoneIndexType> & RequiredBoneIndices = BoneContainer->GetBoneIndicesArray();
+	for (int32 Iter = 0; Iter < RequiredBoneIndices.Num(); ++Iter)
+	{
+		int32 BoneIndex = RequiredBoneIndices[Iter];
+		const FTransform& Trans = Bones[BoneIndex];
+		if (!Bones[BoneIndex].IsRotationNormalized())
+		{
+			return false;
+		}
+	}
+
+	return true;
+}
+
+void FCompactPose::ResetToRefPose(const FBoneContainer& RequiredBones)
+{
+	((TArray<FTransform>&)Bones) = RequiredBones.GetRefPoseCompactArray();
+
+	// If retargeting is disabled, copy ref pose from Skeleton, rather than mesh.
+	// this is only used in editor and for debugging.
+	if (RequiredBones.GetDisableRetargeting())
+	{
+		checkSlow(RequiredBones.IsValid());
+		// Only do this if we have a mesh. otherwise we're not retargeting animations.
+		if (RequiredBones.GetSkeletalMeshAsset())
+		{
+			TArray<FTransform> const & SkeletonRefPose = RequiredBones.GetSkeletonAsset()->GetRefLocalPoses();
+
+			for (const FCompactPoseBoneIndex BoneIndex : ForEachBoneIndex())
+			{
+				const int32 SkeletonBoneIndex = GetBoneContainer().GetSkeletonIndex(BoneIndex);
+
+				// Pose bone index should always exist in Skeleton
+				checkSlow(SkeletonBoneIndex != INDEX_NONE);
+				Bones[BoneIndex.GetInt()] = SkeletonRefPose[SkeletonBoneIndex];
+			}
+		}
+	}
+}
+
+void FCompactPose::ResetToIdentity()
+{
+	for (FTransform& Bone : Bones)
+	{
+		Bone.SetIdentity();
+	}
+}
+
+bool FCompactPose::IsNormalized() const
+{
+	for (const FTransform& Bone : Bones)
+	{
+		if (!Bone.IsRotationNormalized())
+		{
+			return false;
+		}
+	}
+
+	return true;
+}
+
+bool FCompactPose::ContainsNaN() const
+{
+	for (const FTransform& Bone : Bones)
+	{
+		if (Bone.ContainsNaN())
+		{
+			return true;
+		}
+	}
+
+	return false;
+}
+
+void FCompactPose::NormalizeRotations()
+{
+	for (FTransform& Bone : Bones)
+	{
+		Bone.NormalizeRotation();
+	}
+}
diff --git a/Engine/Source/Runtime/Engine/Private/AnimationAsset.cpp b/Engine/Source/Runtime/Engine/Private/AnimationAsset.cpp
--- a/Engine/Source/Runtime/Engine/Private/AnimationAsset.cpp
+++ b/Engine/Source/Runtime/Engine/Private/AnimationAsset.cpp
@@ -225,6 +225,43 @@ void FBoneContainer::Initialize()
 		// Right now we only support a single Skeleton. Skeleton hierarchy coming soon!
 		RemapFromSkeleton(*AssetSkeleton.Get());
 	}
+
+	//Set up compact pose data
+	int32 NumReqBones = BoneIndicesArray.Num();
+	CompactPoseParentBones.Empty(NumReqBones);
+	
+	CompactPoseRefPoseBones.Empty(NumReqBones);
+	CompactPoseRefPoseBones.AddUninitialized(NumReqBones);
+	
+	CompactPoseToSkeletonIndex.Empty(NumReqBones);
+	CompactPoseToSkeletonIndex.AddUninitialized(NumReqBones);
+	
+	SkeletonToCompactPose.Empty(SkeletonToPoseBoneIndexArray.Num());
+
+	const TArray<FTransform>& RefPoseArray = RefSkeleton->GetRefBonePose();
+
+	for (int32 CompactBoneIndex = 0; CompactBoneIndex < NumReqBones; ++CompactBoneIndex)
+	{
+		FBoneIndexType MeshPoseIndex = BoneIndicesArray[CompactBoneIndex];
+
+		//Parent Bone
+		const int32 ParentIndex = GetParentBoneIndex(MeshPoseIndex);
+		const int32 CompactParentIndex = ParentIndex == INDEX_NONE ? INDEX_NONE : BoneIndicesArray.IndexOfByKey(ParentIndex);
+
+		CompactPoseParentBones.Add(FCompactPoseBoneIndex(CompactParentIndex));
+
+		//Ref Pose
+		CompactPoseRefPoseBones[CompactBoneIndex] = RefPoseArray[MeshPoseIndex];
+
+		CompactPoseToSkeletonIndex[CompactBoneIndex] = PoseToSkeletonBoneIndexArray[MeshPoseIndex];
+	}
+
+	for (int32 SkeletonBoneIndex = 0; SkeletonBoneIndex < SkeletonToPoseBoneIndexArray.Num(); ++SkeletonBoneIndex)
+	{
+		int32 PoseBoneIndex = SkeletonToPoseBoneIndexArray[SkeletonBoneIndex];
+		int32 CompactIndex  = (PoseBoneIndex != INDEX_NONE) ? (BoneIndicesArray.IndexOfByKey(PoseBoneIndex)) : INDEX_NONE;
+		SkeletonToCompactPose.Add(FCompactPoseBoneIndex(CompactIndex));
+	}
 }
 
 int32 FBoneContainer::GetPoseBoneIndexForBoneName(const FName& BoneName) const
@@ -240,6 +277,13 @@ int32 FBoneContainer::GetParentBoneIndex(const int32& BoneIndex) const
 	return RefSkeleton->GetParentIndex(BoneIndex);
 }
 
+FCompactPoseBoneIndex FBoneContainer::GetParentBoneIndex(const FCompactPoseBoneIndex& BoneIndex) const
+{
+	checkSlow(IsValid());
+	checkSlow(BoneIndex != INDEX_NONE);
+	return CompactPoseParentBones[BoneIndex.GetInt()];
+}
+
 int32 FBoneContainer::GetDepthBetweenBones(const int32& BoneIndex, const int32& ParentBoneIndex) const
 {
 	checkSlow( IsValid() );
diff --git a/Engine/Source/Runtime/Engine/Private/Components/SkeletalMeshComponent.cpp b/Engine/Source/Runtime/Engine/Private/Components/SkeletalMeshComponent.cpp
--- a/Engine/Source/Runtime/Engine/Private/Components/SkeletalMeshComponent.cpp
+++ b/Engine/Source/Runtime/Engine/Private/Components/SkeletalMeshComponent.cpp
@@ -956,12 +956,16 @@ void USkeletalMeshComponent::EvaluateAnimation(const USkeletalMesh* InSkeletalMe
 			InAnimInstance->EvaluateAnimation(EvaluationContext);
 
 			// can we avoid that copy?
-			if( EvaluationContext.Pose.Bones.Num() > 0 )
+			if( EvaluationContext.Pose.GetNumBones() > 0 )
 			{
-				OutLocalAtoms = EvaluationContext.Pose.Bones;
-
 				// Make sure rotations are normalized to account for accumulation of errors.
-				FAnimationRuntime::NormalizeRotations(OutLocalAtoms);
+				EvaluationContext.Pose.NormalizeRotations();
+
+				for (const FCompactPoseBoneIndex BoneIndex : EvaluationContext.Pose.ForEachBoneIndex())
+				{
+					FMeshPoseBoneIndex MeshPoseBoneIndex = EvaluationContext.Pose.GetBoneContainer().MakeMeshPoseIndex(BoneIndex);
+					OutLocalAtoms[MeshPoseBoneIndex.GetInt()] = EvaluationContext.Pose[BoneIndex];
+				}
 			}
 			else
 			{
diff --git a/Engine/Source/Runtime/Engine/Public/AnimationRuntime.h b/Engine/Source/Runtime/Engine/Public/AnimationRuntime.h
--- a/Engine/Source/Runtime/Engine/Public/AnimationRuntime.h
+++ b/Engine/Source/Runtime/Engine/Public/AnimationRuntime.h
@@ -7,6 +7,7 @@
 
 #pragma once
 #include "Animation/AnimSequenceBase.h"
+#include "BonePose.h"
 
 struct FInputBlendPose;
 struct FA2CSPose;
@@ -29,90 +30,89 @@ public:
 #if DO_GUARD_SLOW
 	static bool ContainsNaN(TArray<FBoneIndexType> & RequiredBoneIndices, FA2Pose & Pose);
 #endif
+
 	/**
-	 * Blends together a set of poses, each with a given weight.
-	 * This function is lightweight, it does not cull out nearly zero weights or check to make sure weights sum to 1.0, the caller should take care of that if needed.
-	 *
-	 * The blend is done by taking a weighted sum of each atom, and re-normalizing the quaternion part at the end, not using SLERP.
-	 * This allows n-way blends, and makes the code much faster, though the angular velocity will not be constant across the blend.
-	 *
-	 * @param	ResultAtoms		Output array of relative bone transforms.
-	 * @param	RequiredBones	Indices of bones that we want to return. Note that bones not in this array will not be modified, so are not safe to access! 
-	 *							This array must be in strictly increasing order.
-	 */
+	* Blends together a set of poses, each with a given weight.
+	* This function is lightweight, it does not cull out nearly zero weights or check to make sure weights sum to 1.0, the caller should take care of that if needed.
+	*
+	* The blend is done by taking a weighted sum of each atom, and re-normalizing the quaternion part at the end, not using SLERP.
+	* This allows n-way blends, and makes the code much faster, though the angular velocity will not be constant across the blend.
+	*
+	* @param	ResultPose		Output pose of relative bone transforms.
+	*/
 	static void BlendPosesTogether(
-		int32 NumPoses,
-		const FTransformArrayA2** SourcePoses,
-		const float* SourceWeights,
-		const FBoneContainer& RequiredBones,
-		/*out*/ FTransformArrayA2& ResultAtoms);
+		const TArray<FCompactPose>& SourcePoses,
+		const TArray<float>&	SourceWeights,
+		/*out*/ FCompactPose& ResultPose);
 
 	/**
-	 * Blends together a set of poses, each with a given weight.
-	 * This function is lightweight, it does not cull out nearly zero weights or check to make sure weights sum to 1.0, the caller should take care of that if needed.
-	 *
-	 * The blend is done by taking a weighted sum of each atom, and re-normalizing the quaternion part at the end, not using SLERP.
-	 * This allows n-way blends, and makes the code much faster, though the angular velocity will not be constant across the blend.
-	 *
-	 * @param	ResultAtoms		Output array of relative bone transforms.
-	 * @param	RequiredBones	Indices of bones that we want to return. Note that bones not in this array will not be modified, so are not safe to access! 
-	 *							This array must be in strictly increasing order.
-	 */
+	* Blends together a set of poses, each with a given weight.
+	* This function is lightweight, it does not cull out nearly zero weights or check to make sure weights sum to 1.0, the caller should take care of that if needed.
+	*
+	* The blend is done by taking a weighted sum of each atom, and re-normalizing the quaternion part at the end, not using SLERP.
+	* This allows n-way blends, and makes the code much faster, though the angular velocity will not be constant across the blend.
+	*
+	* @param	ResultPose		Output pose of relative bone transforms.
+	*/
 	static void BlendPosesTogether(
-		int32 NumPoses,
-		const TArray<FTransformArrayA2>& SourcePoses,
+		const TArray<const FCompactPose*>& SourcePoses,
 		const TArray<float>&	SourceWeights,
-		const FBoneContainer& RequiredBones,
-		/*out*/ FTransformArrayA2& ResultAtoms);
+		/*out*/ FCompactPose& ResultPose);
 
 	/**
-	 * Blends together a set of poses, each with a given weight.
-	 * This function is for BlendSpace per bone blending. BlendSampleDataCache contains the weight information
-	 *
-	 * This blends in local space
-	 *
-	 * @param	ResultAtoms		Output array of relative bone transforms.
-	 * @param	RequiredBones	Indices of bones that we want to return. Note that bones not in this array will not be modified, so are not safe to access! 
-	 *							This array must be in strictly increasing order.
-	 */
+	* Blends together two poses.
+	* This function is lightweight
+	*
+	* The blend is done by taking a weighted sum of each atom, and re-normalizing the quaternion part at the end, not using SLERP.
+	*
+	* @param	ResultPose		Output pose of relative bone transforms.
+	*/
+	static void BlendTwoPosesTogether(
+		const FCompactPose& SourcePose1,
+		const FCompactPose& SourcePose2,
+		const float			WeightOfPose1,
+		/*out*/ FCompactPose& ResultPose);
+
+	/**
+	* Blends together a set of poses, each with a given weight.
+	* This function is for BlendSpace per bone blending. BlendSampleDataCache contains the weight information
+	*
+	* This blends in local space
+	*
+	* @param	ResultPose		Output pose of relative bone transforms.
+	*/
 	static void BlendPosesTogetherPerBone(
-		int32 NumPoses,
-		const TArray<FTransformArrayA2>& SourcePoses,
-		const UBlendSpaceBase* BlendSpace, 
-		const TArray<FBlendSampleData>& BlendSampleDataCache, 
-		const FBoneContainer& RequiredBones,
-		/*out*/ FTransformArrayA2& ResultAtoms);
+		const TArray<FCompactPose>& SourcePoses,
+		const UBlendSpaceBase* BlendSpace,
+		const TArray<FBlendSampleData>& BlendSampleDataCache,
+		/*out*/ FCompactPose& ResultPose);
 
 	/**
-	 * Blends together a set of poses, each with a given weight.
-	 * This function is for BlendSpace per bone blending. BlendSampleDataCache contains the weight information
-	 *
-	 * This blends rotation in mesh space, so performance intensive. 
-	 *
-	 * @param	ResultAtoms		Output array of relative bone transforms.
-	 * @param	RequiredBones	Indices of bones that we want to return. Note that bones not in this array will not be modified, so are not safe to access! 
-	 *							This array must be in strictly increasing order.
-	 */
+	* Blends together a set of poses, each with a given weight.
+	* This function is for BlendSpace per bone blending. BlendSampleDataCache contains the weight information
+	*
+	* This blends rotation in mesh space, so performance intensive.
+	*
+	* @param	ResultPose		Output pose of relative bone transforms.
+	*/
 	static void BlendPosesTogetherPerBoneInMeshSpace(
-		int32 NumPoses,
-		TArray<FTransformArrayA2>& SourcePoses,
-		const UBlendSpaceBase* BlendSpace, 
-		const TArray<FBlendSampleData>& BlendSampleDataCache, 
-		const FBoneContainer& RequiredBones,
-		/*out*/ FTransformArrayA2& ResultAtoms);
+		TArray<FCompactPose>& SourcePoses,
+		const UBlendSpaceBase* BlendSpace,
+		const TArray<FBlendSampleData>& BlendSampleDataCache,
+		/*out*/ FCompactPose& ResultPose);
 
-	/** 
-	 * Blend Poses per bone weights : The BasePoses + BlendPoses(SourceIndex) * Blend Weights(BoneIndex)
-	 * Please note BlendWeights are array, so you can define per bone base 
-	 * This supports multi per bone blending, but only one pose as blend at a time per track
-	 * PerBoneBlendWeights.Num() == Atoms.Num()
-	 *
-	 * I had multiple debates about having PerBoneBlendWeights array(for memory reasons),  
-	 * but it so much simplifies multiple purpose for this function instead of searching bonenames or 
-	 * having multiple bone names with multiple weights, and filtering through which one is correct one
-	 * I assume all those things should be determined before coming here and this only cares about weights
-	 **/
-	static void BlendPosesPerBoneFilter(struct FA2Pose& BasePose, const TArray<struct FA2Pose>& BlendPoses, struct FA2Pose& OutPose, TArray<FPerBoneBlendWeight>& BoneBlendWeights, bool bMeshSpaceRotationBlending, const FBoneContainer& RequiredBones, USkeleton* Skeleton);
+	/**
+	* Blend Poses per bone weights : The BasePoses + BlendPoses(SourceIndex) * Blend Weights(BoneIndex)
+	* Please note BlendWeights are array, so you can define per bone base
+	* This supports multi per bone blending, but only one pose as blend at a time per track
+	* PerBoneBlendWeights.Num() == Atoms.Num()
+	*
+	* I had multiple debates about having PerBoneBlendWeights array(for memory reasons),
+	* but it so much simplifies multiple purpose for this function instead of searching bonenames or
+	* having multiple bone names with multiple weights, and filtering through which one is correct one
+	* I assume all those things should be determined before coming here and this only cares about weights
+	**/
+	static void BlendPosesPerBoneFilter(struct FCompactPose& BasePose, const TArray<struct FCompactPose>& BlendPoses, struct FCompactPose& OutPose, TArray<FPerBoneBlendWeight>& BoneBlendWeights, bool bMeshSpaceRotationBlending);
 
 	static void UpdateDesiredBoneWeight(const TArray<FPerBoneBlendWeight>& SrcBoneBlendWeights, TArray<FPerBoneBlendWeight>& TargetBoneBlendWeights, const TArray<float>& BlendWeights, const FBoneContainer& RequiredBones, USkeleton* Skeleton);
 
@@ -130,15 +130,13 @@ public:
 		/*inout*/ FTransformArrayA2& Atoms);
 
 	static void GetPoseFromSequence(
-		const UAnimSequenceBase* Sequence,
-		const FBoneContainer& RequiredBones,
-		/*out*/ FTransformArrayA2& ResultAtoms,
+		const UAnimSequenceBase* InSequence,
+		/*out*/ FCompactPose& ResultPose,
 		const FAnimExtractContext& ExtractionContext);
 
 	static void GetPoseFromAnimTrack(
 		const FAnimTrack& Track,
-		const FBoneContainer& RequiredBones,
-		/*out*/ FTransformArrayA2& ResultAtoms,
+		/*out*/ FCompactPose& ResultPose,
 		const FAnimExtractContext& ExtractionContext);
 
 	/** Fill ref pose **/
@@ -146,49 +144,44 @@ public:
 
 #if WITH_EDITOR
 	/** fill with retarget base ref pose but this isn't used during run-time, so it always copies all of them */
-	static void FillWithRetargetBaseRefPose( TArray<FTransform> & OutAtoms, const USkeletalMesh* Mesh, const FBoneContainer& RequiredBones );
+	static void FillWithRetargetBaseRefPose(FCompactPose& OutPose, const USkeletalMesh* Mesh);
 #endif
 
 	/** Convert LocalTransforms into MeshSpaceTransforms over RequiredBones. */
 	static void ConvertPoseToMeshSpace(const TArray<FTransform> & LocalTransforms, TArray<FTransform> & MeshSpaceTransforms, const FBoneContainer& RequiredBones);
 
-	/** Convert TargetPose into an AdditivePose, by doing TargetPose = TargetPose - BasePose over RequiredBones. */
-	static void ConvertPoseToAdditive(FTransformArrayA2 & TargetPose, const FTransformArrayA2 & BasePose, const FBoneContainer& RequiredBones);
+	/** Convert TargetPose into an AdditivePose, by doing TargetPose = TargetPose - BasePose */
+	static void ConvertPoseToAdditive(FCompactPose& TargetPose, const FCompactPose& BasePose);
 
-	/** Convert LocalPose into MeshSpaceRotations over RequiredBones. */
-	static void ConvertPoseToMeshRotation(FTransformArrayA2 & LocalPose, const FBoneContainer& RequiredBones);
+	/** Convert LocalPose into MeshSpaceRotations. */
+	static void ConvertPoseToMeshRotation(FCompactPose& LocalPose);
 
 	/**
 	 * Accumulates BlendPoses to ResultAtoms with BlendWeight. 
 	 * ResultAtoms += BlendPose*BlendWeight.
 	 * Result is NOT Normalized, as this is expected to be called several times with each 'BlendPoses' transform array to be accumulated.
 	 * 
-	 * @param	BlendPoses		Array to BoneTransforms to scale and accumulate.
+	 * @param	SourcePose		BoneTransforms to scale and accumulate.
 	 * @param	BlendWeight		Weight to apply to 'BlendPoses'
-	 * @param	RequiredBones	Indices of bones to process, in strictly increasing order.
-	 * @param	ResultAtoms		Array of BoneTransforms to accumulate to.
+	 * @param	ResultPose		BoneTransforms to accumulate to.
 	 */
 	static void BlendPosesAccumulate(
-		const FTransformArrayA2& BlendPoses,
-		const float BlendWeight, 
-		const FBoneContainer& RequiredBones, 
-		/*inout*/ FTransformArrayA2& ResultAtoms);
+		const FCompactPose& SourcePose,
+		const float BlendWeight,
+		/*inout*/ FCompactPose& ResultPose);
 
 	/**
-	 * Blends SourcePose + BlendPose*BlendWeight = ResultAtoms. As it states, this accumulates to the SourcePose
-	 *
-	 * Right now this came from BlendPosesTogether, light weight blend, no slerp or anything. 
-	 * 
-	 * @param	ResultAtoms		Output array of relative bone transforms.
-	 * @param	RequiredBones	Indices of bones that we want to return. Note that bones not in this array will not be modified, so are not safe to access! 
-	 *							This array must be in strictly increasing order.
-	 */
+	* Blends SourcePose + BlendPose*BlendWeight = ResultAtoms. As it states, this accumulates to the SourcePose
+	*
+	* Right now this came from BlendPosesTogether, light weight blend, no slerp or anything.
+	*
+	* @param	ResultPose		Output array of relative bone transforms.
+	*/
 	static void BlendAdditivePose(
-		const FTransformArrayA2& SourcePoses, 
-		const FTransformArrayA2& AdditiveBlendPoses, 
-		const float BlendWeight, 
-		const FBoneContainer& RequiredBones, 
-		/*out*/ FTransformArrayA2& ResultAtoms);
+		const FCompactPose& SourcePose,
+		const FCompactPose& AdditiveBlendPose,
+		const float BlendWeight,
+		/*out*/ FCompactPose& ResultPose);
 
 	/** Lerp for BoneTransforms. Stores results in A. Performs A = Lerp(A, B, Alpha);
 	 * @param A : In/Out transform array.
@@ -198,14 +191,13 @@ public:
 	 */
 	static void LerpBoneTransforms(TArray<FTransform> & A, const TArray<FTransform> & B, float Alpha, const TArray<FBoneIndexType> & RequiredBonesArray);
 
-	/** 
-	 * Get Pose From BlendSpace
-	 */
+	/**
+	* Get Pose From BlendSpace
+	*/
 	static void GetPoseFromBlendSpace(
 		UBlendSpaceBase * BlendSpace,
 		TArray<FBlendSampleData>& BlendSampleDataCache,
-		FBoneContainer & RequiredBones,
-		/*out*/ FTransformArrayA2& ResultAtoms);
+		/*out*/ FCompactPose& ResultPose);
 
 	/** 
 	 * Advance CurrentTime to CurrentTime + MoveDelta. 
@@ -242,9 +234,9 @@ public:
 	static void ConvertCSTransformToBoneSpace
 	(
 		USkeletalMeshComponent * SkelComp,  
-		FA2CSPose & MeshBases, 
+		FCSPose<FCompactPose> & MeshBases,
 		/*inout*/ FTransform& CSBoneTM, 
-		int32 BoneIndex, 
+		FCompactPoseBoneIndex BoneIndex,
 		uint8 Space
 	);
 
@@ -252,9 +244,9 @@ public:
 	static void ConvertBoneSpaceTransformToCS
 	(
 		USkeletalMeshComponent * SkelComp,  
-		FA2CSPose & MeshBases, 
+		FCSPose<FCompactPose>& MeshBases,
 		/*inout*/ FTransform& BoneSpaceTM, 
-		int32 BoneIndex, 
+		FCompactPoseBoneIndex BoneIndex,
 		uint8 Space
 	);
 
@@ -272,7 +264,7 @@ public:
 
 private:
 	/** 
-	* Blend Poses per bone weights : The BasePoses + BlendPoses(SourceIndex) * Blend Weights(BoneIndex)
+	* Blend Poses per bone weights : The BasePose + BlendPoses(SourceIndex) * Blend Weights(BoneIndex)
 	* Please note BlendWeights are array, so you can define per bone base 
 	* This supports multi per bone blending, but only one pose as blend at a time per track
 	* PerBoneBlendWeights.Num() == Atoms.Num()
@@ -284,13 +276,11 @@ private:
 	* having multiple bone names with multiple weights, and filtering through which one is correct one
 	* I assume all those things should be determined before coming here and this only cares about weights
 	**/
-	static void BlendMeshPosesPerBoneWeights( 
-			struct FA2Pose& BasePose, 
-			TArray<struct FA2Pose>& BlendPoses, 
-			const TArray<FPerBoneBlendWeight> & BoneBlendWeights, 
-			const FBoneContainer& RequiredBones, 
-			USkeleton* Skeleton,
-			/*out*/ FA2Pose& OutPose);
+	static void BlendMeshPosesPerBoneWeights(
+				struct FCompactPose& BasePose,
+				const TArray<struct FCompactPose>& BlendPoses,
+				const TArray<FPerBoneBlendWeight> & BoneBlendWeights,
+				/*out*/ FCompactPose& OutPose);
 
 	/** 
 	* Blend Poses per bone weights : The BasePoses + BlendPoses(SourceIndex) * Blend Weights(BoneIndex)
@@ -305,12 +295,10 @@ private:
 	* having multiple bone names with multiple weights, and filtering through which one is correct one
 	* I assume all those things should be determined before coming here and this only cares about weights
 	**/
-	static void BlendLocalPosesPerBoneWeights( 
-			FA2Pose& BasePose, 
-			TArray<FA2Pose>& BlendPoses, 
-			const TArray<FPerBoneBlendWeight> & BoneBlendWeights, 
-			const FBoneContainer& RequiredBones, 
-			/*out*/ FA2Pose& OutPose);
-
+	static void BlendLocalPosesPerBoneWeights(
+			FCompactPose& BasePose,
+			const TArray<FCompactPose>& BlendPoses,
+			const TArray<FPerBoneBlendWeight> & BoneBlendWeights,
+			/*out*/ FCompactPose& OutPose);
 };
 
diff --git a/Engine/Source/Runtime/Engine/Public/BoneContainer.h b/Engine/Source/Runtime/Engine/Public/BoneContainer.h
--- a/Engine/Source/Runtime/Engine/Public/BoneContainer.h
+++ b/Engine/Source/Runtime/Engine/Public/BoneContainer.h
@@ -4,6 +4,9 @@
 
 class USkeleton;
 
+#include "BoneIndices.h"
+#include "CustomBoneIndexArray.h"
+
 /**
 * This is a native transient structure.
 * Contains:
@@ -34,6 +37,18 @@ private:
 	/** Mapping table between Pose Bone Indices and Skeleton Bone Indices. */
 	TArray<int32> PoseToSkeletonBoneIndexArray;
 
+	// Look up from skeleton to compact pose format
+	TArray<int32> CompactPoseToSkeletonIndex;
+
+	// Look up from compact pose format to skeleton
+	TArray<FCompactPoseBoneIndex> SkeletonToCompactPose;
+
+	// Compact pose format of Parent Bones (to save us converting to mesh space and back)
+	TArray<FCompactPoseBoneIndex> CompactPoseParentBones;
+
+	// Compact pose format of Ref Pose Bones (to save us converting to mesh space and back)
+	TArray<FTransform>    CompactPoseRefPoseBones;
+
 	/** For debugging. */
 	/** Disable Retargeting. Extract animation, but do not retarget it. */
 	bool bDisableRetargeting;
@@ -133,6 +148,16 @@ public:
 		return RefSkeleton->GetRefBonePose();
 	}
 
+	const FTransform& GetRefPoseTransform(const FCompactPoseBoneIndex& BoneIndex) const
+	{
+		return CompactPoseRefPoseBones[BoneIndex.GetInt()];
+	}
+
+	const TArray<FTransform>& GetRefPoseCompactArray() const
+	{
+		return CompactPoseRefPoseBones;
+	}
+
 	/** Access to Asset's RefSkeleton. */
 	const FReferenceSkeleton& GetReferenceSkeleton() const
 	{
@@ -145,12 +170,20 @@ public:
 		return RefSkeleton->GetNum();
 	}
 
+	const int32 GetCompactPoseNumBones() const
+	{
+		return BoneIndicesArray.Num();
+	}
+
 	/** Get BoneIndex for BoneName for current Asset. */
 	ENGINE_API int32 GetPoseBoneIndexForBoneName(const FName& BoneName) const;
 
 	/** Get ParentBoneIndex for current Asset. */
 	ENGINE_API int32 GetParentBoneIndex(const int32& BoneIndex) const;
 
+	/** Get ParentBoneIndex for current Asset. */
+	ENGINE_API FCompactPoseBoneIndex GetParentBoneIndex(const FCompactPoseBoneIndex& BoneIndex) const;
+
 	/** Get Depth between bones for current asset. */
 	int32 GetDepthBetweenBones(const int32& BoneIndex, const int32& ParentBoneIndex) const;
 
@@ -201,6 +234,26 @@ public:
 		return PoseToSkeletonBoneIndexArray;
 	}
 
+	int32 GetSkeletonIndex(const FCompactPoseBoneIndex& BoneIndex) const
+	{
+		return CompactPoseToSkeletonIndex[BoneIndex.GetInt()];
+	}
+
+	FCompactPoseBoneIndex GetCompactPoseIndexFromSkeletonIndex(const int32 SkeletonIndex) const
+	{
+		return SkeletonToCompactPose[SkeletonIndex];
+	}
+
+	FMeshPoseBoneIndex MakeMeshPoseIndex(const FCompactPoseBoneIndex& BoneIndex) const
+	{
+		return FMeshPoseBoneIndex(GetBoneIndicesArray()[BoneIndex.GetInt()]);
+	}
+
+	FCompactPoseBoneIndex MakeCompactPoseIndex(const FMeshPoseBoneIndex& BoneIndex) const
+	{
+		return FCompactPoseBoneIndex(GetBoneIndicesArray().IndexOfByKey(BoneIndex.GetInt()));
+	}
+
 private:
 	/** Initialize FBoneContainer. */
 	ENGINE_API void Initialize();
diff --git a/Engine/Source/Runtime/Engine/Public/BoneIndices.h b/Engine/Source/Runtime/Engine/Public/BoneIndices.h
new file mode 100644
--- /dev/null
+++ b/Engine/Source/Runtime/Engine/Public/BoneIndices.h
@@ -0,0 +1,127 @@
+// Copyright 1998-2015 Epic Games, Inc. All Rights Reserved.
+
+#pragma once
+
+struct FBoneIndexBase
+{
+	FBoneIndexBase() : BoneIndex(INDEX_NONE) {}
+
+	FORCEINLINE int32 GetInt() const { return BoneIndex; }
+
+	FORCEINLINE bool IsRootBone() const { return BoneIndex == 0; }
+
+protected:
+	int32 BoneIndex;
+};
+
+FORCEINLINE int32 GetIntFromComp(const int32& InComp)
+{
+	return InComp;
+}
+
+FORCEINLINE FBoneIndexType GetIntFromComp(const FBoneIndexBase& InComp)
+{
+	return InComp.GetInt();
+}
+
+template<class RealBoneIndexType>
+struct FBoneIndexWithOperators : public FBoneIndexBase
+{
+	// BoneIndexType
+	FORCEINLINE bool operator==(const RealBoneIndexType& Rhs) const
+	{
+		return BoneIndex == GetIntFromComp(Rhs);
+	}
+
+	FORCEINLINE bool operator!=(const RealBoneIndexType& Rhs) const
+	{
+		return BoneIndex != GetIntFromComp(Rhs);
+	}
+
+	FORCEINLINE bool operator>(const RealBoneIndexType& Rhs) const
+	{
+		return BoneIndex > GetIntFromComp(Rhs);
+	}
+
+	FORCEINLINE bool operator>=(const RealBoneIndexType& Rhs) const
+	{
+		return BoneIndex >= GetIntFromComp(Rhs);
+	}
+
+	FORCEINLINE bool operator<(const RealBoneIndexType& Rhs) const
+	{
+		return BoneIndex < GetIntFromComp(Rhs);
+	}
+
+	FORCEINLINE bool operator<=(const RealBoneIndexType& Rhs) const
+	{
+		return BoneIndex <= GetIntFromComp(Rhs);
+	}
+
+	// FBoneIndexType
+	FORCEINLINE bool operator==(const int32& Rhs) const
+	{
+		return BoneIndex == GetIntFromComp(Rhs);
+	}
+
+	FORCEINLINE bool operator!=(const int32& Rhs) const
+	{
+		return BoneIndex != GetIntFromComp(Rhs);
+	}
+
+	FORCEINLINE bool operator>(const int32& Rhs) const
+	{
+		return BoneIndex > GetIntFromComp(Rhs);
+	}
+
+	FORCEINLINE bool operator>=(const int32& Rhs) const
+	{
+		return BoneIndex >= GetIntFromComp(Rhs);
+	}
+
+	FORCEINLINE bool operator<(const int32& Rhs) const
+	{
+		return BoneIndex < GetIntFromComp(Rhs);
+	}
+
+	FORCEINLINE bool operator<=(const int32& Rhs) const
+	{
+		return BoneIndex <= GetIntFromComp(Rhs);
+	}
+
+	RealBoneIndexType& operator++()
+	{
+		++BoneIndex;
+		return *((RealBoneIndexType*)this);
+	}
+
+	RealBoneIndexType& operator--()
+	{
+		--BoneIndex;
+		return *((RealBoneIndexType*)this);
+	}
+
+	const RealBoneIndexType& operator=(const RealBoneIndexType& Rhs)
+	{
+		BoneIndex = Rhs.BoneIndex;
+		return Rhs;
+	}
+};
+
+struct FCompactPoseBoneIndex : public FBoneIndexWithOperators < FCompactPoseBoneIndex >
+{
+public:
+	explicit FCompactPoseBoneIndex(int32 InBoneIndex) { BoneIndex = InBoneIndex; }
+};
+
+struct FMeshPoseBoneIndex : public FBoneIndexWithOperators < FMeshPoseBoneIndex >
+{
+public:
+	explicit FMeshPoseBoneIndex(int32 InBoneIndex) { BoneIndex = InBoneIndex; }
+};
+
+struct FSkeletonPoseBoneIndex : public FBoneIndexWithOperators < FSkeletonPoseBoneIndex >
+{
+public:
+	explicit FSkeletonPoseBoneIndex(int32 InBoneIndex) { BoneIndex = InBoneIndex; }
+};
\ No newline at end of file
diff --git a/Engine/Source/Runtime/Engine/Public/BonePose.h b/Engine/Source/Runtime/Engine/Public/BonePose.h
new file mode 100644
--- /dev/null
+++ b/Engine/Source/Runtime/Engine/Public/BonePose.h
@@ -0,0 +1,646 @@
+// Copyright 1998-2015 Epic Games, Inc. All Rights Reserved.
+
+#pragma once
+
+#include "SkeletalMeshTypes.h"
+#include "BoneContainer.h"
+#include "BoneIndices.h"
+#include "CustomBoneIndexArray.h"
+#include "AnimEncoding.h"
+
+struct FCompactPose;
+
+template<class BoneIndexType>
+struct FBasePose
+{
+public:
+	FORCEINLINE void InitBones(int32 NumBones) { Bones.Empty(NumBones); Bones.AddUninitialized(NumBones); }
+
+	FORCEINLINE int32 GetNumBones() const { return Bones.Num(); }
+
+	FORCEINLINE bool IsValidIndex(const BoneIndexType& BoneIndex)
+	{
+		return Bones.IsValidIndex(BoneIndex.GetInt());
+	}
+
+	FORCEINLINE FTransform& operator[](const BoneIndexType& BoneIndex)
+	{ 
+		return Bones[BoneIndex.GetInt()];
+	}
+
+	FORCEINLINE const FTransform& operator[] (const BoneIndexType& BoneIndex) const
+	{
+		return Bones[BoneIndex.GetInt()];
+	}
+
+	//Bone Index Iteration
+	template<typename PoseType, typename IterType>
+	struct FRangedForSupport
+	{
+		const PoseType& Pose;
+
+		FRangedForSupport(const PoseType& Pose) : Pose(Pose) {};
+		
+		IterType begin() { return Pose.MakeBeginIter<IterType>(); }
+		IterType end() { return Pose.MakeEndIter<IterType>(); }
+	};
+
+protected:
+	TArray<FTransform> Bones;
+};
+
+
+struct FMeshPose : public FBasePose<FMeshPoseBoneIndex>
+{
+public:
+	typedef FMeshPoseBoneIndex BoneIndexType;
+
+	const FBoneContainer& GetBoneContainer() const
+	{
+		checkSlow(BoneContainer && BoneContainer->IsValid());
+		return *BoneContainer;
+	}
+
+	void SetBoneContainer(const FBoneContainer* InBoneContainer)
+	{
+		check(InBoneContainer && InBoneContainer->IsValid());
+		BoneContainer = InBoneContainer;
+		InitBones(BoneContainer->GetNumBones());
+	}
+
+	void InitFrom(FCompactPose& CompactPose);
+
+	void MoveBonesTo(TArray<FTransform>& OutTransforms)
+	{
+		OutTransforms = MoveTemp(Bones);
+		BoneContainer = NULL;
+	}
+
+	// Sets this pose to its ref pose
+	ENGINE_API void ResetToRefPose();
+
+	// Sets every bone transform to Identity
+	void ResetToIdentity();
+
+	// returns true if all bone rotations are normalized
+	bool IsNormalized() const;
+
+	// Returns true if any bone rotation contains NaN
+	bool ContainsNaN() const;
+
+
+	FORCEINLINE BoneIndexType GetParentBone(const BoneIndexType& BoneIndex)
+	{
+		return BoneIndexType(BoneContainer->GetParentBoneIndex(BoneIndex.GetInt()));
+	}
+
+protected:
+
+	//Reference to our BoneContainer
+	const FBoneContainer* BoneContainer;
+};
+
+struct FCompactPoseBoneIndexIterator
+{
+	int32 Index;
+
+	FCompactPoseBoneIndexIterator(int32 Index) : Index(Index) {}
+
+	FCompactPoseBoneIndexIterator& operator++() { ++Index; return (*this); }
+	bool operator==(FCompactPoseBoneIndexIterator& Rhs) { return Index == Rhs.Index; }
+	bool operator!=(FCompactPoseBoneIndexIterator& Rhs) { return Index != Rhs.Index; }
+	FCompactPoseBoneIndex operator*() const { return FCompactPoseBoneIndex(Index); }
+};
+
+struct FCompactPoseBoneIndexReverseIterator
+{
+	int32 Index;
+
+	FCompactPoseBoneIndexReverseIterator(int32 Index) : Index(Index) {}
+
+	FCompactPoseBoneIndexReverseIterator& operator++() { --Index; return (*this); }
+	bool operator==(FCompactPoseBoneIndexReverseIterator& Rhs) { return Index == Rhs.Index; }
+	bool operator!=(FCompactPoseBoneIndexReverseIterator& Rhs) { return Index != Rhs.Index; }
+	FCompactPoseBoneIndex operator*() const { return FCompactPoseBoneIndex(Index); }
+};
+
+struct FCompactPose : FBasePose<FCompactPoseBoneIndex>
+{
+public:
+	typedef FCompactPoseBoneIndex BoneIndexType;
+
+	//--------------------------------------------------------------------------
+	//Bone Index Iteration
+		typedef FBasePose<FCompactPoseBoneIndex>::FRangedForSupport<FCompactPose, FCompactPoseBoneIndexIterator> RangedForBoneIndexFwd;
+		typedef FBasePose<FCompactPoseBoneIndex>::FRangedForSupport<FCompactPose, FCompactPoseBoneIndexReverseIterator> RangedForBoneIndexBwd;
+
+		RangedForBoneIndexFwd ForEachBoneIndex() const
+		{
+			return RangedForBoneIndexFwd(*this);
+		}
+
+		RangedForBoneIndexBwd ForEachBoneIndexReverse() const
+		{
+			return RangedForBoneIndexBwd(*this);
+		}
+
+		// plain template implementations 
+		template<typename Type>
+		Type MakeBeginIter() const { check(false); return Type(); }
+
+		template<typename Type>
+		Type MakeEndIter() const { check(false); return Type(); }
+
+		template<>
+		FCompactPoseBoneIndexIterator MakeBeginIter<FCompactPoseBoneIndexIterator>() const { return FCompactPoseBoneIndexIterator(0); }
+
+		template<>
+		FCompactPoseBoneIndexIterator MakeEndIter<FCompactPoseBoneIndexIterator>() const { return FCompactPoseBoneIndexIterator(GetNumBones()); }
+
+		template<>
+		FCompactPoseBoneIndexReverseIterator MakeBeginIter<FCompactPoseBoneIndexReverseIterator>() const { return FCompactPoseBoneIndexReverseIterator(GetNumBones() - 1); }
+
+		template<>
+		FCompactPoseBoneIndexReverseIterator MakeEndIter<FCompactPoseBoneIndexReverseIterator>() const { return FCompactPoseBoneIndexReverseIterator(-1); }
+	//--------------------------------------------------------------------------
+
+	const FBoneContainer& GetBoneContainer() const 
+	{
+		checkSlow(BoneContainer && BoneContainer->IsValid());
+		return *BoneContainer;
+	}
+
+	void SetBoneContainer(const FBoneContainer* InBoneContainer)
+	{
+		check(InBoneContainer && InBoneContainer->IsValid());
+		BoneContainer = InBoneContainer;
+		InitBones(BoneContainer->GetBoneIndicesArray().Num());
+	}
+
+	void InitFrom(const FCompactPose& SrcPose)
+	{
+		SetBoneContainer(SrcPose.BoneContainer);
+		Bones = SrcPose.Bones;
+	}
+
+	// Moves transform data out of the supplied SrcPose (SrcPose will be left empty)
+	void MoveBonesFrom(FCompactPose& SrcPose)
+	{
+		Bones = MoveTemp(SrcPose.Bones);
+		BoneContainer = SrcPose.BoneContainer;
+		SrcPose.BoneContainer = NULL;
+	}
+
+	//Moves transform data to supplied array (pose will be left empty)
+	void MoveBonesTo(TArray<FTransform>& OutTransforms)
+	{
+		OutTransforms = MoveTemp(Bones);
+		BoneContainer = NULL;
+	}
+
+	// Sets this pose to its ref pose
+	void ResetToRefPose()
+	{
+		ResetToRefPose(GetBoneContainer());
+	}
+
+	// Sets this pose to the supplied BoneContainers ref pose
+	ENGINE_API void ResetToRefPose(const FBoneContainer& RequiredBones);
+
+	// Sets every bone transform to Identity
+	void ResetToIdentity();
+
+	// returns true if all bone rotations are normalized
+	bool IsNormalized() const;
+
+	// Returns true if any bone rotation contains NaN
+	bool ContainsNaN() const;
+
+	// Normalizes all rotations in this pose
+	void NormalizeRotations();
+
+	bool IsValid() const
+	{
+		return (BoneContainer && BoneContainer->IsValid());
+	}
+
+	// Returns the bone index for the parent bone
+	BoneIndexType GetParentBoneIndex(const BoneIndexType& BoneIndex) const
+	{
+		return GetBoneContainer().GetParentBoneIndex(BoneIndex);
+	}
+
+	// Returns the ref pose for the supplied bone
+	const FTransform& GetRefPose(const BoneIndexType& BoneIndex) const
+	{
+		return GetBoneContainer().GetRefPoseTransform(BoneIndex);
+	}
+
+	// Populates this pose from the supplied animation and track data
+	void PopulateFromAnimation(
+		const UAnimSequence& Seq,
+		const BoneTrackArray& RotationTracks,
+		const BoneTrackArray& TranslationTracks,
+		const BoneTrackArray& ScaleTracks,
+		float Time)
+	{
+		AnimationFormat_GetAnimationPose(
+			*((FTransformArray*)&Bones), //@TODO:@ANIMATION: Nasty hack
+			RotationTracks,
+			TranslationTracks,
+			ScaleTracks,
+			Seq,
+			Time);
+	}
+protected:
+
+	//Reference to our BoneContainer
+	const FBoneContainer* BoneContainer;
+};
+
+template<class PoseType>
+struct FCSPose
+{
+	// Set up our index type based on the type of pose we are manipulating
+	typedef typename PoseType::BoneIndexType BoneIndexType;
+
+	// Init Pose
+	void InitPose(const FBoneContainer* InBoneContainer)
+	{
+		Pose.SetBoneContainer(InBoneContainer);
+		Pose.ResetToRefPose();
+		ComponentSpaceFlags.Empty(Pose.GetNumBones());
+		ComponentSpaceFlags.AddZeroed(Pose.GetNumBones());
+		ComponentSpaceFlags[0] = 1;
+	}
+
+	// Init Pose
+	void InitPose(const PoseType& SrcPose)
+	{
+		Pose.InitFrom(SrcPose);
+		ComponentSpaceFlags.Empty(Pose.GetNumBones());
+		ComponentSpaceFlags.AddZeroed(Pose.GetNumBones());
+		ComponentSpaceFlags[0] = 1;
+	}
+	const PoseType& GetPose() const { return Pose; }
+
+	// Get transform for supplied bone in local space
+	FTransform GetLocalSpaceTransform(BoneIndexType BoneIndex);
+
+	// Get Transform for supplied bone in component space
+	const FTransform& GetComponentSpaceTransform(BoneIndexType BoneIndex);
+
+	// Set the transform for the supplied bone 
+	void SetComponentSpaceTransform(BoneIndexType BoneIndex, const FTransform& NewTransform);
+
+	// Calculate the component space transform for the supplied bone
+	void CalculateComponentSpaceTransform(BoneIndexType BoneIndex);
+
+	/**
+	* Convert Bone to Local Space.
+	*/
+	void ConvertBoneToLocalSpace(BoneIndexType BoneIndex);
+
+	/**
+	* Set a bunch of Component Space Bone Transforms.
+	* Do this safely by insuring that Parents are already in Component Space,
+	* and any Component Space children are converted back to Local Space before hand.
+	*/
+	void SafeSetCSBoneTransforms(const TArray<struct FBoneTransform> & BoneTransforms);
+
+	/**
+	* Blends Component Space transforms to MeshPose in Local Space.
+	* Used by SkelControls to apply their transforms.
+	*
+	* The tricky bit is that SkelControls deliver their transforms in Component Space,
+	* But the blending is done in Local Space. Also we need to refresh any Children they have
+	* that has been previously converted to Component Space.
+	*/
+	void LocalBlendCSBoneTransforms(const TArray<struct FBoneTransform>& BoneTransforms, float Alpha);
+
+	// Convert any component space transforms back to local space
+	void ConvertToLocalPoses(PoseType& OutPose);
+
+protected:
+	PoseType Pose;
+
+	// Flags to track each bones current state (0 means local pose, 1 means component space pose)
+	TCustomBoneIndexArray<uint8, BoneIndexType> ComponentSpaceFlags;
+};
+
+template<class PoseType>
+FTransform FCSPose<PoseType>::GetLocalSpaceTransform(BoneIndexType BoneIndex)
+{
+	checkSlow(Pose.IsValid());
+	check(Pose.IsValidIndex(BoneIndex));
+
+	// if evaluated, calculate it
+	if (ComponentSpaceFlags[BoneIndex])
+	{
+		const BoneIndexType ParentIndex = Pose.GetParentBoneIndex(BoneIndex);
+
+		if (ParentIndex != INDEX_NONE)
+		{
+			const FTransform& ParentTransform = GetComponentSpaceTransform(ParentIndex);
+			const FTransform& BoneTransform = Pose[BoneIndex];
+			// calculate local space
+			return BoneTransform.GetRelativeTransform(ParentTransform);
+		}
+	}
+
+	return Pose[BoneIndex];
+}
+
+template<class PoseType>
+const FTransform& FCSPose<PoseType>::GetComponentSpaceTransform(BoneIndexType BoneIndex)
+{
+	checkSlow(Pose.IsValid());
+	check(Pose.IsValidIndex(BoneIndex));
+
+	check(!Pose[BoneIndex].ContainsNaN());
+	// if not evaluate, calculate it
+	if (ComponentSpaceFlags[BoneIndex] == 0)
+	{
+		CalculateComponentSpaceTransform(BoneIndex);
+	}
+	check(!Pose[BoneIndex].ContainsNaN());
+	return Pose[BoneIndex];
+}
+
+template<class PoseType>
+void FCSPose<PoseType>::SetComponentSpaceTransform(BoneIndexType BoneIndex, const FTransform& NewTransform)
+{
+	checkSlow(Pose.IsValid());
+	check(Pose.IsValidIndex(BoneIndex));
+
+	// this one forcefully sets component space transform
+	Pose[BoneIndex] = NewTransform;
+	ComponentSpaceFlags[BoneIndex] = 1;
+}
+
+template<class PoseType>
+void FCSPose<PoseType>::CalculateComponentSpaceTransform(BoneIndexType BoneIndex)
+{
+	checkSlow(Pose.IsValid());
+	check(ComponentSpaceFlags[BoneIndex] == 0);
+
+	// root is already verified, so root should not come here
+	// check AllocateLocalPoses
+	const BoneIndexType ParentIndex = Pose.GetParentBoneIndex(BoneIndex);
+
+	// if Parent already has been calculated, use it
+	if (ComponentSpaceFlags[ParentIndex] == 0)
+	{
+		// if Parent hasn't been calculated, also calculate parents
+		CalculateComponentSpaceTransform(ParentIndex);
+	}
+
+	// current Bones(Index) should contain LocalPoses.
+	FTransform& Bone = Pose[BoneIndex];
+	FTransform& ParentBone = Pose[ParentIndex];
+	check(!Pose[BoneIndex].ContainsNaN());
+	check(!Pose[ParentIndex].ContainsNaN());
+	FTransform Suspect = Pose[BoneIndex] * Pose[ParentIndex];
+	check(!Suspect.ContainsNaN());
+	Pose[BoneIndex] = Pose[BoneIndex] * Pose[ParentIndex];
+	Pose[BoneIndex].NormalizeRotation();
+	check(!Pose[BoneIndex].ContainsNaN());
+	ComponentSpaceFlags[BoneIndex] = 1;
+}
+
+template<class PoseType>
+void FCSPose<PoseType>::ConvertBoneToLocalSpace(BoneIndexType BoneIndex)
+{
+	checkSlow(Pose.IsValid());
+
+	// If BoneTransform is in Component Space, then convert it.
+	// Never convert Root to Local Space.
+	if (!BoneIndex.IsRootBone() && ComponentSpaceFlags[BoneIndex] == 1)
+	{
+		const BoneIndexType ParentIndex = Pose.GetParentBoneIndex(BoneIndex);
+
+		// Verify that our Parent is also in Component Space. That should always be the case.
+		check(ComponentSpaceFlags[ParentIndex] == 1);
+
+		// Convert to local space.
+		Pose[BoneIndex].SetToRelativeTransform(Pose[ParentIndex]);
+		ComponentSpaceFlags[BoneIndex] = 0;
+	}
+}
+
+template<class PoseType>
+void FCSPose<PoseType>::SafeSetCSBoneTransforms(const TArray<struct FBoneTransform> & BoneTransforms)
+{
+	checkSlow(Pose.IsValid());
+
+	// Bone Mask to keep track of which bones have to be converted to local space.
+	// This is basically BoneTransforms's children.
+	TCustomBoneIndexArray<uint8, FCompactPoseBoneIndex> BoneMask;
+	BoneMask.AddZeroed(Pose.GetNumBones());
+
+	// First build our BoneMask
+	for (const FBoneTransform& BoneTransform : BoneTransforms)
+	{
+		// Mark those bones in Mesh Pose as being required to be in Local Space.
+		BoneMask[BoneTransform.BoneIndex] = 1;
+	}
+
+
+	//Flag children
+	for (const FCompactPoseBoneIndex BoneIndex : Pose.ForEachBoneIndex())
+	{
+		const FCompactPoseBoneIndex ParentIndex = Pose.GetParentBoneIndex(BoneIndex);
+		// Propagate our BoneMask to children.
+		if (ParentIndex != INDEX_NONE)
+		{
+			BoneMask[BoneIndex] |= BoneMask[ParentIndex];
+		}
+	}
+
+	// now iterate from children to parent to calculate back to local space
+	// Sadly this has to iterate from back because you'll need parent to be component space
+	// if you do this from parent to child, you'll see parent turned to local before children applies it
+	int32 BoneTransformIndex = BoneTransforms.Num() - 1;
+	for (FCompactPoseBoneIndex BoneIndex : Pose.ForEachBoneIndexReverse())
+	{
+		// If this bone has to be converted to Local Space...
+		if (BoneMask[BoneIndex] != 0)
+		{
+			// If this is one of the original BoneTransforms list bones
+			// Then we don't actually want to convert this one to local space, since we're going to overwrite it.
+			// So skip it.
+			if (BoneTransformIndex >= 0 && BoneIndex == BoneTransforms[BoneTransformIndex].BoneIndex)
+			{
+				BoneTransformIndex--;
+			}
+			// If this is a children bone then we want it to be in local space!
+			else
+			{
+				// .. If it is not currently in Local Space, then convert it.
+				ConvertBoneToLocalSpace(BoneIndex);
+			}
+		}
+	}
+
+	// Finally copy our Component Space transforms
+	for (const FBoneTransform& BoneTransform : BoneTransforms)
+	{
+		const FCompactPoseBoneIndex BoneIndex = BoneTransform.BoneIndex;
+
+		// Make sure our BoneTransforms were in Component Space in the first place, before we overwrite them
+		// Only check their parent to do minimal work needed.
+		const FCompactPoseBoneIndex ParentBoneIndex = Pose.GetParentBoneIndex(BoneIndex);
+		if (ParentBoneIndex != INDEX_NONE && ComponentSpaceFlags[ParentBoneIndex] == 0)
+		{
+			CalculateComponentSpaceTransform(ParentBoneIndex);
+		}
+
+		// Set new Component Space transform.
+		SetComponentSpaceTransform(BoneIndex, BoneTransform.Transform);
+	}
+}
+
+template<class PoseType>
+void FCSPose<PoseType>::LocalBlendCSBoneTransforms(const TArray<struct FBoneTransform>& BoneTransforms, float Alpha)
+{
+	// if Alpha is small enough, skip
+	if (Alpha < ZERO_ANIMWEIGHT_THRESH)
+	{
+		return;
+	}
+
+#if DO_CHECK
+	if (BoneTransforms.Num() > 0)
+	{
+		FCompactPoseBoneIndex LastIndex(BoneTransforms[0].BoneIndex);
+		// Make sure bones are sorted in "Parents before Children" order.
+		for (int32 I = 1; I < BoneTransforms.Num(); ++I)
+		{
+			check(BoneTransforms[I].BoneIndex >= LastIndex);
+			LastIndex = BoneTransforms[I].BoneIndex;
+		}
+	}
+#endif
+
+	// If we are not doing any blending, do a faster path.
+	// Set transforms directly in Component Space. But still refresh children.
+	if (Alpha >= 1.f - ZERO_ANIMWEIGHT_THRESH)
+	{
+		SafeSetCSBoneTransforms(BoneTransforms);
+	}
+	// Slower blending path.
+	else
+	{
+		// Bone Mask to keep track of which bones have to be converted to local space.
+		// This is basically BoneTransforms bones and their children.
+		TCustomBoneIndexArray<uint8, BoneIndexType> BoneMask;
+		BoneMask.AddZeroed(Pose.GetNumBones());
+
+		TArray<struct FBoneTransform> LocalBoneTransforms;
+		LocalBoneTransforms.SetNumUninitialized(BoneTransforms.Num());
+
+		// First, convert BoneTransforms to local space for blending.
+		for (int32 Index = 0; Index < BoneTransforms.Num(); Index++)
+		{
+			const BoneIndexType BoneIndex = BoneTransforms[Index].BoneIndex;
+			const BoneIndexType ParentIndex = Pose.GetParentBoneIndex(BoneIndex);
+
+			if (ParentIndex != INDEX_NONE)
+			{
+				// if BoneTransforms(modified by controllers) contains ParentIndex, it should use that as ParentTransform, not the one from input
+				int32 LocalParentIndex = INDEX_NONE;
+				for (int32 LocalIndex = 0; LocalIndex < BoneTransforms.Num(); ++LocalIndex)
+				{
+					if (ParentIndex == BoneTransforms[LocalIndex].BoneIndex)
+					{
+						LocalParentIndex = LocalIndex;
+						break;
+					}
+				}
+
+				// saves Parent Transform
+				const bool bNoParent = LocalParentIndex == INDEX_NONE;
+				const FTransform& ParentTransform = bNoParent ? GetComponentSpaceTransform(ParentIndex) : BoneTransforms[LocalParentIndex].Transform;
+
+				LocalBoneTransforms[Index].Transform = BoneTransforms[Index].Transform.GetRelativeTransform(ParentTransform);
+				LocalBoneTransforms[Index].BoneIndex = BoneIndex;
+
+				// Mark those bones in Mesh Pose as being required to be in Local Space.
+				BoneMask[BoneIndex] = 1;
+			}
+			else
+			{
+				// when root is entered as to modify, we don't need to adjust parent index, just clear it
+				LocalBoneTransforms[Index].Transform = BoneTransforms[Index].Transform;
+				LocalBoneTransforms[Index].BoneIndex = BoneIndex;
+
+				BoneMask[BoneIndex] = 1;
+			}
+		}
+
+		// Then, convert MeshPose Bones from BoneTransforms list, and their children, to local space if they are not already.
+		for (const BoneIndexType BoneIndex : Pose.ForEachBoneIndex())
+		{
+			const BoneIndexType ParentIndex = Pose.GetParentBoneIndex(BoneIndex);
+			// Propagate our BoneMask to children.
+			if (ParentIndex != INDEX_NONE)
+			{
+				BoneMask[BoneIndex] |= BoneMask[ParentIndex];
+			}
+		}
+
+		for (const BoneIndexType BoneIndex : Pose.ForEachBoneIndexReverse())
+		{
+			if (!BoneIndex.IsRootBone())
+			{
+				// If this bone has to be converted to Local Space...
+				if (BoneMask[BoneIndex] != 0)
+				{
+					// .. And is not currently in Local Space, then convert it.
+					ConvertBoneToLocalSpace(BoneIndex);
+				}
+			}
+		}
+
+		// Lastly, do the blending in local space.
+		for (int32 Index = 0; Index < LocalBoneTransforms.Num(); Index++)
+		{
+			const FCompactPoseBoneIndex BoneIndex = LocalBoneTransforms[Index].BoneIndex;
+			// Make sure this transform is in local space, because we are writing a local space one to it.
+			// If we are not in local space, this could mean trouble for our children.
+			check((ComponentSpaceFlags[BoneIndex] == 0) || (BoneIndex == 0));
+
+			// No need to normalize rotation since BlendWith() does it.
+			Pose[BoneIndex].BlendWith(LocalBoneTransforms[Index].Transform, Alpha);
+		}
+	}
+}
+
+template<class PoseType>
+void FCSPose<PoseType>::ConvertToLocalPoses(PoseType& OutPose)
+{
+	checkSlow(Pose.IsValid());
+	OutPose = Pose;
+
+	// now we need to convert back to local bases
+	// only convert back that has been converted to mesh base
+	// if it was local base, and if it hasn't been modified
+	// that is still okay even if parent is changed, 
+	// that doesn't mean this local has to change
+	// go from child to parent since I need parent inverse to go back to local
+	// root is same, so no need to do Index == 0
+	for (const BoneIndexType BoneIndex : Pose.ForEachBoneIndexReverse())
+	{
+		if (!BoneIndex.IsRootBone())
+		{
+			// root is already verified, so root should not come here
+			if (ComponentSpaceFlags[BoneIndex])
+			{
+				const FCompactPoseBoneIndex ParentIndex = Pose.GetParentBoneIndex(BoneIndex);
+
+				OutPose[BoneIndex].SetToRelativeTransform(OutPose[ParentIndex]);
+				OutPose[BoneIndex].NormalizeRotation();
+			}
+		}
+	}
+}
\ No newline at end of file
diff --git a/Engine/Source/Runtime/Engine/Public/CustomBoneIndexArray.h b/Engine/Source/Runtime/Engine/Public/CustomBoneIndexArray.h
new file mode 100644
--- /dev/null
+++ b/Engine/Source/Runtime/Engine/Public/CustomBoneIndexArray.h
@@ -0,0 +1,25 @@
+// Copyright 1998-2015 Epic Games, Inc. All Rights Reserved.
+
+#pragma once
+
+template<typename ArrayType, typename BoneIndexType>
+class TCustomBoneIndexArray : public TArray < ArrayType >
+{
+public:
+	FORCEINLINE ArrayType& operator[](int32 BoneIndex)
+	{ 
+		return TArray<ArrayType>::operator [](BoneIndex); 
+	}
+	FORCEINLINE const ArrayType& operator[](int32 BoneIndex) const
+	{ 
+		return TArray<ArrayType>::operator [](BoneIndex); 
+	}
+	FORCEINLINE ArrayType& operator[](const BoneIndexType& BoneIndex) 
+	{ 
+		return TArray<ArrayType>::operator [] (BoneIndex.GetInt()); 
+	}
+	FORCEINLINE const ArrayType& operator[](const BoneIndexType& BoneIndex) const 
+	{ 
+		return TArray<ArrayType>::operator [] (BoneIndex.GetInt()); 
+	}
+};
\ No newline at end of file
