Fix for shader compilation errors when trying to compile gpu particle shaders on SM4 using the distance field collision stuff.

diff --git a/Engine/Source/Runtime/Engine/Private/Particles/ParticleGpuSimulation.cpp b/Engine/Source/Runtime/Engine/Private/Particles/ParticleGpuSimulation.cpp
--- a/Engine/Source/Runtime/Engine/Private/Particles/ParticleGpuSimulation.cpp
+++ b/Engine/Source/Runtime/Engine/Private/Particles/ParticleGpuSimulation.cpp
@@ -839,6 +839,12 @@ enum EParticleCollisionShaderMode
 	PCM_DistanceField
 };
 
+/** Helper function to determine whether the given particle collision shader mode is supported on the given shader platform */
+inline bool IsParticleCollisionModeSupported(EShaderPlatform InPlatform, EParticleCollisionShaderMode InCollisionShaderMode)
+{
+	return IsFeatureLevelSupported(InPlatform, ERHIFeatureLevel::SM5) || InCollisionShaderMode != PCM_DistanceField;
+}
+
 /**
  * Pixel shader for simulating particles on the GPU.
  */
@@ -851,7 +857,7 @@ public:
 
 	static bool ShouldCache(EShaderPlatform Platform)
 	{
-		return SupportsGPUParticles(Platform);
+		return SupportsGPUParticles(Platform) && IsParticleCollisionModeSupported(Platform, CollisionMode);
 	}
 
 	static void ModifyCompilationEnvironment(EShaderPlatform Platform, FShaderCompilerEnvironment& OutEnvironment)
