CIS Fixes

diff --git a/Engine/Source/Runtime/Engine/Classes/Animation/AnimInstance.h b/Engine/Source/Runtime/Engine/Classes/Animation/AnimInstance.h
--- a/Engine/Source/Runtime/Engine/Classes/Animation/AnimInstance.h
+++ b/Engine/Source/Runtime/Engine/Classes/Animation/AnimInstance.h
@@ -154,24 +154,6 @@ private:
 	friend class FAnimationRuntime;
 };
 
-struct FBoneTransform
-{
-	/** @todo anim: should be Skeleton bone index in the future, but right now it's CompactBoneIndex **/
-	FCompactPoseBoneIndex BoneIndex;
-
-	/** Transform to apply **/
-	FTransform Transform;
-
-	FBoneTransform() 
-		: BoneIndex(INDEX_NONE)
-	{}
-
-	FBoneTransform(FCompactPoseBoneIndex InBoneIndex, const FTransform& InTransform)
-		: BoneIndex(InBoneIndex)
-		, Transform(InTransform)
-	{}
-};
-
 USTRUCT(BlueprintType)
 struct FPerBoneBlendWeight
 {
diff --git a/Engine/Source/Runtime/Engine/Public/BonePose.h b/Engine/Source/Runtime/Engine/Public/BonePose.h
--- a/Engine/Source/Runtime/Engine/Public/BonePose.h
+++ b/Engine/Source/Runtime/Engine/Public/BonePose.h
@@ -8,8 +8,27 @@
 #include "CustomBoneIndexArray.h"
 #include "AnimEncoding.h"
 
+
 struct FCompactPose;
 
+struct FBoneTransform
+{
+	/** @todo anim: should be Skeleton bone index in the future, but right now it's CompactBoneIndex **/
+	FCompactPoseBoneIndex BoneIndex;
+
+	/** Transform to apply **/
+	FTransform Transform;
+
+	FBoneTransform()
+		: BoneIndex(INDEX_NONE)
+	{}
+
+	FBoneTransform(FCompactPoseBoneIndex InBoneIndex, const FTransform& InTransform)
+		: BoneIndex(InBoneIndex)
+		, Transform(InTransform)
+	{}
+};
+
 template<class BoneIndexType>
 struct FBasePose
 {
@@ -39,7 +58,7 @@ public:
 	{
 		const PoseType& Pose;
 
-		FRangedForSupport(const PoseType& Pose) : Pose(Pose) {};
+		FRangedForSupport(const PoseType& InPose) : Pose(InPose) {};
 		
 		IterType begin() { return Pose.MakeBeginIter<IterType>(); }
 		IterType end() { return Pose.MakeEndIter<IterType>(); }
