Hopefully fix for explicit specialization of 'MakeEndIter' in class scope

diff --git a/Engine/Source/Runtime/Engine/Public/BonePose.h b/Engine/Source/Runtime/Engine/Public/BonePose.h
--- a/Engine/Source/Runtime/Engine/Public/BonePose.h
+++ b/Engine/Source/Runtime/Engine/Public/BonePose.h
@@ -123,7 +123,7 @@ struct FCompactPoseBoneIndexIterator
 {
 	int32 Index;
 
-	FCompactPoseBoneIndexIterator(int32 Index) : Index(Index) {}
+	FCompactPoseBoneIndexIterator(int32 InIndex) : Index(InIndex) {}
 
 	FCompactPoseBoneIndexIterator& operator++() { ++Index; return (*this); }
 	bool operator==(FCompactPoseBoneIndexIterator& Rhs) { return Index == Rhs.Index; }
@@ -135,7 +135,7 @@ struct FCompactPoseBoneIndexReverseIterator
 {
 	int32 Index;
 
-	FCompactPoseBoneIndexReverseIterator(int32 Index) : Index(Index) {}
+	FCompactPoseBoneIndexReverseIterator(int32 InIndex) : Index(InIndex) {}
 
 	FCompactPoseBoneIndexReverseIterator& operator++() { --Index; return (*this); }
 	bool operator==(FCompactPoseBoneIndexReverseIterator& Rhs) { return Index == Rhs.Index; }
@@ -171,16 +171,16 @@ public:
 		Type MakeEndIter() const { check(false); return Type(); }
 
 		template<>
-		FCompactPoseBoneIndexIterator MakeBeginIter<FCompactPoseBoneIndexIterator>() const { return FCompactPoseBoneIndexIterator(0); }
+		FCompactPoseBoneIndexIterator MakeBeginIter<FCompactPoseBoneIndexIterator>() const;
 
 		template<>
-		FCompactPoseBoneIndexIterator MakeEndIter<FCompactPoseBoneIndexIterator>() const { return FCompactPoseBoneIndexIterator(GetNumBones()); }
+		FCompactPoseBoneIndexIterator MakeEndIter<FCompactPoseBoneIndexIterator>() const;
 
 		template<>
-		FCompactPoseBoneIndexReverseIterator MakeBeginIter<FCompactPoseBoneIndexReverseIterator>() const { return FCompactPoseBoneIndexReverseIterator(GetNumBones() - 1); }
+		FCompactPoseBoneIndexReverseIterator MakeBeginIter<FCompactPoseBoneIndexReverseIterator>() const;
 
 		template<>
-		FCompactPoseBoneIndexReverseIterator MakeEndIter<FCompactPoseBoneIndexReverseIterator>() const { return FCompactPoseBoneIndexReverseIterator(-1); }
+		FCompactPoseBoneIndexReverseIterator MakeEndIter<FCompactPoseBoneIndexReverseIterator>() const;
 	//--------------------------------------------------------------------------
 
 	const FBoneContainer& GetBoneContainer() const 
@@ -277,6 +277,18 @@ protected:
 	const FBoneContainer* BoneContainer;
 };
 
+template<>
+FCompactPoseBoneIndexIterator FCompactPose::MakeBeginIter<FCompactPoseBoneIndexIterator>() const { return FCompactPoseBoneIndexIterator(0); }
+
+template<>
+FCompactPoseBoneIndexIterator FCompactPose::MakeEndIter<FCompactPoseBoneIndexIterator>() const { return FCompactPoseBoneIndexIterator(GetNumBones()); }
+
+template<>
+FCompactPoseBoneIndexReverseIterator FCompactPose::MakeBeginIter<FCompactPoseBoneIndexReverseIterator>() const { return FCompactPoseBoneIndexReverseIterator(GetNumBones() - 1); }
+
+template<>
+FCompactPoseBoneIndexReverseIterator FCompactPose::MakeEndIter<FCompactPoseBoneIndexReverseIterator>() const { return FCompactPoseBoneIndexReverseIterator(-1); }
+
 template<class PoseType>
 struct FCSPose
 {
