Fix Use 'template' keyword to treat 'MakeBeginIter' as a dependent template name (cheers Steve.R)

diff --git a/Engine/Source/Runtime/Engine/Public/BonePose.h b/Engine/Source/Runtime/Engine/Public/BonePose.h
--- a/Engine/Source/Runtime/Engine/Public/BonePose.h
+++ b/Engine/Source/Runtime/Engine/Public/BonePose.h
@@ -60,8 +60,8 @@ public:
 
 		FRangedForSupport(const PoseType& InPose) : Pose(InPose) {};
 		
-		IterType begin() { return Pose.MakeBeginIter<IterType>(); }
-		IterType end() { return Pose.MakeEndIter<IterType>(); }
+		IterType begin() { return Pose.template MakeBeginIter<IterType>(); }
+		IterType end() { return Pose.template MakeEndIter<IterType>(); }
 	};
 
 protected:
