Fixes for PlatformGame

diff --git a/Engine/Source/Runtime/Engine/Public/BonePose.h b/Engine/Source/Runtime/Engine/Public/BonePose.h
--- a/Engine/Source/Runtime/Engine/Public/BonePose.h
+++ b/Engine/Source/Runtime/Engine/Public/BonePose.h
@@ -37,7 +37,7 @@ public:
 
 	FORCEINLINE int32 GetNumBones() const { return Bones.Num(); }
 
-	FORCEINLINE bool IsValidIndex(const BoneIndexType& BoneIndex)
+	FORCEINLINE bool IsValidIndex(const BoneIndexType& BoneIndex) const
 	{
 		return Bones.IsValidIndex(BoneIndex.GetInt());
 	}
