Fix for explicit specialization on Mac

diff --git a/Engine/Source/Runtime/Engine/Public/BonePose.h b/Engine/Source/Runtime/Engine/Public/BonePose.h
--- a/Engine/Source/Runtime/Engine/Public/BonePose.h
+++ b/Engine/Source/Runtime/Engine/Public/BonePose.h
@@ -60,8 +60,19 @@ public:
 
 		FRangedForSupport(const PoseType& InPose) : Pose(InPose) {};
 		
-		IterType begin() { return Pose.template MakeBeginIter<IterType>(); }
-		IterType end() { return Pose.template MakeEndIter<IterType>(); }
+		IterType begin() { return Pose.MakeBeginIter(); }
+		IterType end() { return Pose.MakeEndIter(); }
+	};
+
+	template<typename PoseType, typename IterType>
+	struct FRangedForReverseSupport
+	{
+		const PoseType& Pose;
+
+		FRangedForReverseSupport(const PoseType& InPose) : Pose(InPose) {};
+
+		IterType begin() { return Pose.MakeBeginIterReverse(); }
+		IterType end() { return Pose.MakeEndIterReverse(); }
 	};
 
 protected:
@@ -151,7 +162,7 @@ public:
 	//--------------------------------------------------------------------------
 	//Bone Index Iteration
 		typedef FBasePose<FCompactPoseBoneIndex>::FRangedForSupport<FCompactPose, FCompactPoseBoneIndexIterator> RangedForBoneIndexFwd;
-		typedef FBasePose<FCompactPoseBoneIndex>::FRangedForSupport<FCompactPose, FCompactPoseBoneIndexReverseIterator> RangedForBoneIndexBwd;
+		typedef FBasePose<FCompactPoseBoneIndex>::FRangedForReverseSupport<FCompactPose, FCompactPoseBoneIndexReverseIterator> RangedForBoneIndexBwd;
 
 		RangedForBoneIndexFwd ForEachBoneIndex() const
 		{
@@ -163,24 +174,13 @@ public:
 			return RangedForBoneIndexBwd(*this);
 		}
 
-		// plain template implementations 
-		template<typename Type>
-		Type MakeBeginIter() const { check(false); return Type(); }
+		FCompactPoseBoneIndexIterator MakeBeginIter() const { return FCompactPoseBoneIndexIterator(0); }
 
-		template<typename Type>
-		Type MakeEndIter() const { check(false); return Type(); }
+		FCompactPoseBoneIndexIterator MakeEndIter() const { return FCompactPoseBoneIndexIterator(GetNumBones()); }
 
-		template<>
-		FCompactPoseBoneIndexIterator MakeBeginIter<FCompactPoseBoneIndexIterator>() const;
+		FCompactPoseBoneIndexReverseIterator MakeBeginIterReverse() const { return FCompactPoseBoneIndexReverseIterator(GetNumBones() - 1); }
 
-		template<>
-		FCompactPoseBoneIndexIterator MakeEndIter<FCompactPoseBoneIndexIterator>() const;
-
-		template<>
-		FCompactPoseBoneIndexReverseIterator MakeBeginIter<FCompactPoseBoneIndexReverseIterator>() const;
-
-		template<>
-		FCompactPoseBoneIndexReverseIterator MakeEndIter<FCompactPoseBoneIndexReverseIterator>() const;
+		FCompactPoseBoneIndexReverseIterator MakeEndIterReverse() const { return FCompactPoseBoneIndexReverseIterator(-1); }
 	//--------------------------------------------------------------------------
 
 	const FBoneContainer& GetBoneContainer() const 
@@ -277,18 +277,6 @@ protected:
 	const FBoneContainer* BoneContainer;
 };
 
-template<>
-FCompactPoseBoneIndexIterator FCompactPose::MakeBeginIter<FCompactPoseBoneIndexIterator>() const { return FCompactPoseBoneIndexIterator(0); }
-
-template<>
-FCompactPoseBoneIndexIterator FCompactPose::MakeEndIter<FCompactPoseBoneIndexIterator>() const { return FCompactPoseBoneIndexIterator(GetNumBones()); }
-
-template<>
-FCompactPoseBoneIndexReverseIterator FCompactPose::MakeBeginIter<FCompactPoseBoneIndexReverseIterator>() const { return FCompactPoseBoneIndexReverseIterator(GetNumBones() - 1); }
-
-template<>
-FCompactPoseBoneIndexReverseIterator FCompactPose::MakeEndIter<FCompactPoseBoneIndexReverseIterator>() const { return FCompactPoseBoneIndexReverseIterator(-1); }
-
 template<class PoseType>
 struct FCSPose
 {
