Fix for MTU size changes, merged NetConstants.h with CoreNet.h and removed dependancy on Engine in OnlineSubSystem

diff --git a/Engine/Source/Runtime/CoreUObject/Public/UObject/CoreNet.h b/Engine/Source/Runtime/CoreUObject/Public/UObject/CoreNet.h
--- a/Engine/Source/Runtime/CoreUObject/Public/UObject/CoreNet.h
+++ b/Engine/Source/Runtime/CoreUObject/Public/UObject/CoreNet.h
@@ -430,6 +430,14 @@ COREUOBJECT_API void SerializeChecksum(FArchive &Ar, uint32 x, bool ErrorOK);
 
 #endif
 
+// Up to this many reliable channel bunches may be buffered.
+enum { RELIABLE_BUFFER = 256 }; // Power of 2 >= 1.
+enum { MAX_PACKETID = 16384 }; // Power of 2 >= 1, covering guaranteed loss/misorder time.
+enum { MAX_CHSEQUENCE = 1024 }; // Power of 2 >RELIABLE_BUFFER, covering loss/misorder time.
+enum { MAX_BUNCH_HEADER_BITS = 64 };
+enum { MAX_PACKET_HEADER_BITS = 16 };
+enum { MAX_PACKET_TRAILER_BITS = 1 };
+enum { MAX_PACKET_SIZE = 1000 }; // MTU for the connection
 
 /**
  * Functions to assist in detecting errors during RPC calls
diff --git a/Engine/Source/Runtime/Engine/Classes/Engine/NetConnection.h b/Engine/Source/Runtime/Engine/Classes/Engine/NetConnection.h
--- a/Engine/Source/Runtime/Engine/Classes/Engine/NetConnection.h
+++ b/Engine/Source/Runtime/Engine/Classes/Engine/NetConnection.h
@@ -6,7 +6,6 @@
 
 #pragma once
 
-#include "Net/NetConstants.h"
 #include "Net/DataBunch.h"
 #include "Engine/Channel.h"
 #include "Engine/Player.h"
@@ -91,6 +90,13 @@ struct DelayedPacket
 #define PING_ACK_DELAY 0.5
 
 
+class SecurityLog
+{
+	void SecurityEvent(const char* string);
+};
+
+
+
 UCLASS(customConstructor, Abstract, MinimalAPI, transient, config=Engine)
 class UNetConnection : public UPlayer
 {
@@ -102,7 +108,9 @@ class UNetConnection : public UPlayer
 
 	/** Owning net driver */
 	UPROPERTY()
-	class UNetDriver* Driver;					
+	class UNetDriver* Driver;	
+
+	class SecurityLog Log;
 
 	UPROPERTY()
 	/** Package map between local and remote. (negotiates net serialization) */
diff --git a/Engine/Source/Runtime/Engine/Private/NetConnection.cpp b/Engine/Source/Runtime/Engine/Private/NetConnection.cpp
--- a/Engine/Source/Runtime/Engine/Private/NetConnection.cpp
+++ b/Engine/Source/Runtime/Engine/Private/NetConnection.cpp
@@ -1897,4 +1897,11 @@ void UNetConnection::ResetPacketBitCounts()
 	NumBunchBits = 0;
 	NumAckBits = 0;
 	NumPaddingBits = 0;
+}
+
+
+
+void SecurityLog::SecurityEvent(const char* string)
+{
+	UE_LOG(LogNet, Warning, TEXT("GHI"));
 }
\ No newline at end of file
diff --git a/Engine/Source/Runtime/Engine/Public/Net/NetConstants.h b/Engine/Source/Runtime/Engine/Public/Net/NetConstants.h
deleted file mode 100644
diff --git a/Engine/Source/Runtime/Online/OnlineSubsystem/OnlineSubsystem.Build.cs b/Engine/Source/Runtime/Online/OnlineSubsystem/OnlineSubsystem.Build.cs
--- a/Engine/Source/Runtime/Online/OnlineSubsystem/OnlineSubsystem.Build.cs
+++ b/Engine/Source/Runtime/Online/OnlineSubsystem/OnlineSubsystem.Build.cs
@@ -28,8 +28,6 @@ public class OnlineSubsystem : ModuleRules
 			}
 		);
 
-        PublicIncludePathModuleNames.Add("Engine"); 
-
 
 	}
 
diff --git a/Engine/Source/Runtime/Online/OnlineSubsystem/Public/LANBeacon.h b/Engine/Source/Runtime/Online/OnlineSubsystem/Public/LANBeacon.h
--- a/Engine/Source/Runtime/Online/OnlineSubsystem/Public/LANBeacon.h
+++ b/Engine/Source/Runtime/Online/OnlineSubsystem/Public/LANBeacon.h
@@ -2,7 +2,7 @@
 
 #pragma once
 
-#include "Net/NetConstants.h"
+#include "CoreNet.h"
 
 /**
  * This value indicates which packet version the server is sending. Clients with
