Fix for issue with getting Skeletal Mesh vertex positions highlighted in https://udn.unrealengine.com/questions/237832/particle-module-location-skel-vert-surface-spawn-p.html

diff --git a/Engine/Source/Runtime/Engine/Private/Components/SkinnedMeshComponent.cpp b/Engine/Source/Runtime/Engine/Private/Components/SkinnedMeshComponent.cpp
--- a/Engine/Source/Runtime/Engine/Private/Components/SkinnedMeshComponent.cpp
+++ b/Engine/Source/Runtime/Engine/Private/Components/SkinnedMeshComponent.cpp
@@ -1545,7 +1545,7 @@ FORCEINLINE FVector USkinnedMeshComponent::GetTypedSkinnedVertexPosition(const F
 	// Do rigid (one-influence) skinning for this vertex.
 	else
 	{
-		const TGPUSkinVertexBase<false>* SrcRigidVertex = VertexBufferGPUSkin.GetVertexPtr<false>(Chunk.GetRigidVertexBufferIndex()+VertIndex);
+		const TGPUSkinVertexBase<bExtraBoneInfluencesT>* SrcRigidVertex = VertexBufferGPUSkin.GetVertexPtr<bExtraBoneInfluencesT>(Chunk.GetRigidVertexBufferIndex() + VertIndex);
 		const int32 RigidInfluenceIndex = SkinningTools::GetRigidInfluenceIndex();
 		int32 BoneIndex = Chunk.BoneMap[SrcRigidVertex->InfluenceBones[RigidInfluenceIndex]];
 		if(MasterPoseComponentInst)
