Shadowing fix

diff --git a/Engine/Source/Runtime/Engine/Private/Animation/AnimNode_HandIKRetargeting.cpp b/Engine/Source/Runtime/Engine/Private/Animation/AnimNode_HandIKRetargeting.cpp
--- a/Engine/Source/Runtime/Engine/Private/Animation/AnimNode_HandIKRetargeting.cpp
+++ b/Engine/Source/Runtime/Engine/Private/Animation/AnimNode_HandIKRetargeting.cpp
@@ -48,9 +48,8 @@ void FAnimNode_HandIKRetargeting::EvaluateBoneTransforms(USkeletalMeshComponent*
 	if (!IK_To_FK_Translation.IsNearlyZero())
 	{
 		// Move desired bones
-		for (int32 BoneIndex = 0; BoneIndex < IKBonesToMove.Num(); BoneIndex++)
+		for (const FBoneReference& BoneReference : IKBonesToMove)
 		{
-			FBoneReference const & BoneReference = IKBonesToMove[BoneIndex];
 			if (BoneReference.IsValid(BoneContainer))
 			{
 				FCompactPoseBoneIndex BoneIndex = BoneReference.GetCompactPoseIndex(BoneContainer);
