Added skip code if nan is found in physics animation.

diff --git a/Engine/Source/Runtime/Engine/Private/PhysicsEngine/PhysAnim.cpp b/Engine/Source/Runtime/Engine/Private/PhysicsEngine/PhysAnim.cpp
--- a/Engine/Source/Runtime/Engine/Private/PhysicsEngine/PhysAnim.cpp
+++ b/Engine/Source/Runtime/Engine/Private/PhysicsEngine/PhysAnim.cpp
@@ -357,7 +357,12 @@ void USkeletalMeshComponent::UpdateKinematicBonesToAnim(const TArray<FTransform>
 #if WITH_PHYSX
 						// update bone transform to world
 						const FTransform BoneTransform = InSpaceBases[BoneIndex] * CurrentLocalToWorld;
-						ensure(!BoneTransform.ContainsNaN());
+						if(BoneTransform.ContainsNaN())
+						{
+							const FName BodyName = PhysicsAsset->BodySetup[i]->BoneName;
+							UE_LOG(LogPhysics, Warning, TEXT("UpdateKinematicBonesToAnim: Trying to set transform with bad data %s on PhysicsAsset '%s' in SkeletalMesh '%s' for bone '%s'"), *BoneTransform.ToHumanReadableString(), *PhysicsAsset->GetName(), *SkeletalMesh->GetName(), *BodyName.ToString());
+							continue;
+						}					
 
 						// If kinematic and not teleporting, set kinematic target
 						PxRigidDynamic* PRigidDynamic = BodyInst->GetPxRigidDynamic_AssumesLocked();
