Particles using distance field collision fallback to scene depth in SM4

diff --git a/Engine/Source/Runtime/Engine/Private/Particles/ParticleGpuSimulation.cpp b/Engine/Source/Runtime/Engine/Private/Particles/ParticleGpuSimulation.cpp
--- a/Engine/Source/Runtime/Engine/Private/Particles/ParticleGpuSimulation.cpp
+++ b/Engine/Source/Runtime/Engine/Private/Particles/ParticleGpuSimulation.cpp
@@ -2624,7 +2624,19 @@ struct FGPUSpriteDynamicEmitterData : FDynamicEmitterDataBase
 		}
 		else if (Simulation->bWantsCollision && Simulation->CollisionMode == EParticleCollisionMode::DistanceField)
 		{
-			Simulation->SimulationPhase = EParticleSimulatePhase::CollisionDistanceField;
+			if (IsParticleCollisionModeSupported(FXSystem->GetShaderPlatform(), PCM_DistanceField))
+			{
+				Simulation->SimulationPhase = EParticleSimulatePhase::CollisionDistanceField;
+			}
+			else if (bTranslucent)
+			{
+				// Fall back to scene depth collision if translucent
+				Simulation->SimulationPhase = EParticleSimulatePhase::CollisionDepthBuffer;
+			}
+			else
+			{
+				Simulation->SimulationPhase = EParticleSimulatePhase::Main;
+			}
 		}
 	}
 
