[UE-15768] Restore ability to trigger overlap events prior to World->bHasBegunPlay == true.

change summary:
- modified AActor::PostActorConstruction() to remove the check for GetWorld()->HasBegunPlay() prior to calling UpdateOverlaps(). this was introduced in CL# 2535852.

notes:
- GetWorld()->HasBegunPlay() will not return 'true' until all Actors have received BeginPlay().
- Actors spawned in response to a BeginPlay() event on other Actors need to be able to trigger overlap events on spawn (some existing projects are dependent on this ability).
- this will cause a regression on UE-10170 for now. we will revisit that for the next UE4 release.

diff --git a/Engine/Source/Runtime/Engine/Private/Actor.cpp b/Engine/Source/Runtime/Engine/Private/Actor.cpp
--- a/Engine/Source/Runtime/Engine/Private/Actor.cpp
+++ b/Engine/Source/Runtime/Engine/Private/Actor.cpp
@@ -2494,7 +2494,7 @@ void AActor::PostActorConstruction()
 	}
 
 	// Components are all there and we've begun play, init overlapping state
-	if (GetWorld()->HasBegunPlay() && !deferBeginPlayAndUpdateOverlaps)
+	if (!deferBeginPlayAndUpdateOverlaps)
 	{
 		UpdateOverlaps();
 	}
