#ue4 - Fix bad logic in UCharacterMovementComponent::PhysicsRotation() that was ignoring rotation rates of zero.

diff --git a/Engine/Source/Runtime/Engine/Private/Components/CharacterMovementComponent.cpp b/Engine/Source/Runtime/Engine/Private/Components/CharacterMovementComponent.cpp
--- a/Engine/Source/Runtime/Engine/Private/Components/CharacterMovementComponent.cpp
+++ b/Engine/Source/Runtime/Engine/Private/Components/CharacterMovementComponent.cpp
@@ -4382,31 +4382,27 @@ void UCharacterMovementComponent::PhysicsRotation(float DeltaTime)
 	
 	// Accumulate a desired new rotation.
 	const float AngleTolerance = 1e-3f;
-	bool bRotationChanged = false;
 
 	// PITCH
-	if ((DeltaRot.Pitch != 0.f) && !FMath::IsNearlyEqual(CurrentRotation.Pitch, DesiredRotation.Pitch, AngleTolerance))
+	if (!FMath::IsNearlyEqual(CurrentRotation.Pitch, DesiredRotation.Pitch, AngleTolerance))
 	{
 		DesiredRotation.Pitch = FMath::FixedTurn(CurrentRotation.Pitch, DesiredRotation.Pitch, DeltaRot.Pitch);
-		bRotationChanged = true;
 	}
 
 	// YAW
-	if ((DeltaRot.Yaw != 0.f) && !FMath::IsNearlyEqual(CurrentRotation.Yaw, DesiredRotation.Yaw, AngleTolerance))
+	if (!FMath::IsNearlyEqual(CurrentRotation.Yaw, DesiredRotation.Yaw, AngleTolerance))
 	{
 		DesiredRotation.Yaw = FMath::FixedTurn(CurrentRotation.Yaw, DesiredRotation.Yaw, DeltaRot.Yaw);
-		bRotationChanged = true;
 	}
 
 	// ROLL
-	if ((DeltaRot.Roll != 0.f) && !FMath::IsNearlyEqual(CurrentRotation.Roll, DesiredRotation.Roll, AngleTolerance))
+	if (!FMath::IsNearlyEqual(CurrentRotation.Roll, DesiredRotation.Roll, AngleTolerance))
 	{
 		DesiredRotation.Roll = FMath::FixedTurn(CurrentRotation.Roll, DesiredRotation.Roll, DeltaRot.Roll);
-		bRotationChanged = true;
 	}
 
 	// Set the new rotation.
-	if (bRotationChanged)
+	if (!CurrentRotation.Equals(DesiredRotation, AngleTolerance))
 	{
 		MoveUpdatedComponent( FVector::ZeroVector, DesiredRotation, true );
 	}
