Add some checks to UActorChannel::CleanUp before a potential crash.

diff --git a/Engine/Source/Runtime/Engine/Private/DataChannel.cpp b/Engine/Source/Runtime/Engine/Private/DataChannel.cpp
--- a/Engine/Source/Runtime/Engine/Private/DataChannel.cpp
+++ b/Engine/Source/Runtime/Engine/Private/DataChannel.cpp
@@ -1487,6 +1487,9 @@ void UActorChannel::DestroyActorAndComponents()
 
 bool UActorChannel::CleanUp( const bool bForDestroy )
 {
+	checkf(Connection != nullptr, TEXT("UActorChannel::CleanUp: Connection is null!"));
+	checkf(Connection->Driver != nullptr, TEXT("UActorChannel::CleanUp: Connection->Driver is null!"));
+
 	const bool bIsServer = Connection->Driver->IsServer();
 
 	UE_LOG( LogNetTraffic, Log, TEXT( "UActorChannel::CleanUp: Channel: %i, IsServer: %s" ), ChIndex, bIsServer ? TEXT( "YES" ) : TEXT( "NO" ) );
@@ -1495,9 +1498,7 @@ bool UActorChannel::CleanUp( const bool bForDestroy )
 	if (!bIsServer)
 	{
 		check(Actor == NULL || Actor->IsValidLowLevel());
-		checkSlow(Connection != NULL);
 		checkSlow(Connection->IsValidLowLevel());
-		checkSlow(Connection->Driver != NULL);
 		checkSlow(Connection->Driver->IsValidLowLevel());
 		if (Actor != NULL)
 		{
