'Light Types' show menu
'Lighting Components' -> 'Ambient Occlusion' affects both DFAO and SSAO now
Added 'Screen Space Ambient Occlusion' show flag under 'Lighting Features'

diff --git a/Engine/Source/Editor/DetailCustomizations/Private/SceneCaptureDetails.cpp b/Engine/Source/Editor/DetailCustomizations/Private/SceneCaptureDetails.cpp
--- a/Engine/Source/Editor/DetailCustomizations/Private/SceneCaptureDetails.cpp
+++ b/Engine/Source/Editor/DetailCustomizations/Private/SceneCaptureDetails.cpp
@@ -148,6 +148,9 @@ void FSceneCaptureDetails::CustomizeDetails( IDetailLayoutBuilder& DetailLayout
 			case SFG_Visualize:
 				GroupName = LOCTEXT("VisualizeShowFlagsMenu", "Visualize Show Flags");
 				break;
+			case SFG_LightTypes:
+				GroupName = LOCTEXT("LightTypesShowFlagsMenu", "Light Types Show Flags");
+				break;
 			case SFG_LightingComponents:
 				GroupName = LOCTEXT("LightingComponentsShowFlagsMenu", "Lighting Components Show Flags");
 				break;
diff --git a/Engine/Source/Editor/LevelEditor/Private/SLevelViewportToolBar.cpp b/Engine/Source/Editor/LevelEditor/Private/SLevelViewportToolBar.cpp
--- a/Engine/Source/Editor/LevelEditor/Private/SLevelViewportToolBar.cpp
+++ b/Engine/Source/Editor/LevelEditor/Private/SLevelViewportToolBar.cpp
@@ -894,6 +894,9 @@ TSharedRef<SWidget> SLevelViewportToolBar::GenerateShowMenu() const
 			ShowMenuBuilder.AddSubMenu( LOCTEXT("PostProcessShowFlagsMenu", "Post Processing"), LOCTEXT("PostProcessShowFlagsMenu_ToolTip", "Post process show flags"),
 				FNewMenuDelegate::CreateStatic( &FillShowMenu, ShowMenu[SFG_PostProcess], 0 ) );
 
+			ShowMenuBuilder.AddSubMenu( LOCTEXT("LightTypesShowFlagsMenu", "Light Types"), LOCTEXT("LightTypesShowFlagsMenu_ToolTip", "Light Types show flags"),
+				FNewMenuDelegate::CreateStatic( &FillShowMenu, ShowMenu[SFG_LightTypes], 0 ) );
+
 			ShowMenuBuilder.AddSubMenu( LOCTEXT("LightingComponentsShowFlagsMenu", "Lighting Components"), LOCTEXT("LightingComponentsShowFlagsMenu_ToolTip", "Lighting Components show flags"),
 				FNewMenuDelegate::CreateStatic( &FillShowMenu, ShowMenu[SFG_LightingComponents], 0 ) );
 
diff --git a/Engine/Source/Editor/UnrealEd/Private/SCommonEditorViewportToolbarBase.cpp b/Engine/Source/Editor/UnrealEd/Private/SCommonEditorViewportToolbarBase.cpp
--- a/Engine/Source/Editor/UnrealEd/Private/SCommonEditorViewportToolbarBase.cpp
+++ b/Engine/Source/Editor/UnrealEd/Private/SCommonEditorViewportToolbarBase.cpp
@@ -285,6 +285,9 @@ TSharedRef<SWidget> SCommonEditorViewportToolbarBase::GenerateShowMenu() const
 			ShowMenuBuilder.AddSubMenu( LOCTEXT("PostProcessShowFlagsMenu", "Post Processing"), LOCTEXT("PostProcessShowFlagsMenu_ToolTip", "Post process show flags"),
 				FNewMenuDelegate::CreateStatic(&CommonEditorViewportUtils::FillShowMenu, ShowMenu[SFG_PostProcess], 0));
 
+			ShowMenuBuilder.AddSubMenu( LOCTEXT("LightTypesShowFlagsMenu", "Light Types"), LOCTEXT("LightTypesShowFlagsMenu_ToolTip", "Light Types show flags"),
+				FNewMenuDelegate::CreateStatic(&CommonEditorViewportUtils::FillShowMenu, ShowMenu[SFG_LightTypes], 0));
+
 			ShowMenuBuilder.AddSubMenu( LOCTEXT("LightingComponentsShowFlagsMenu", "Lighting Components"), LOCTEXT("LightingComponentsShowFlagsMenu_ToolTip", "Lighting Components show flags"),
 				FNewMenuDelegate::CreateStatic(&CommonEditorViewportUtils::FillShowMenu, ShowMenu[SFG_LightingComponents], 0));
 
diff --git a/Engine/Source/Runtime/Engine/Private/SceneView.cpp b/Engine/Source/Runtime/Engine/Private/SceneView.cpp
--- a/Engine/Source/Runtime/Engine/Private/SceneView.cpp
+++ b/Engine/Source/Runtime/Engine/Private/SceneView.cpp
@@ -1219,7 +1219,7 @@ void FSceneView::EndFinalPostprocessSettings(const FSceneViewInitOptions& ViewIn
 		FinalPostProcessSettings.ScreenPercentage = 100;
 	}
 
-	if(!Family->EngineShowFlags.AmbientOcclusion)
+	if(!Family->EngineShowFlags.AmbientOcclusion || !Family->EngineShowFlags.ScreenSpaceAO)
 	{
 		FinalPostProcessSettings.AmbientOcclusionIntensity = 0;
 	}
diff --git a/Engine/Source/Runtime/Engine/Private/ShowFlags.cpp b/Engine/Source/Runtime/Engine/Private/ShowFlags.cpp
--- a/Engine/Source/Runtime/Engine/Private/ShowFlags.cpp
+++ b/Engine/Source/Runtime/Engine/Private/ShowFlags.cpp
@@ -345,6 +345,7 @@ void EngineShowFlagOverride(EShowFlagInitMode ShowFlagInitMode, EViewModeIndex V
 		if(ICVar->GetValueOnGameThread() == 0)
 		{
 			EngineShowFlags.AmbientOcclusion = 0;
+			EngineShowFlags.ScreenSpaceAO = 0;
 			EngineShowFlags.Decals = 0;
 			EngineShowFlags.DynamicShadows = 0;
 			EngineShowFlags.GlobalIllumination = 0;
diff --git a/Engine/Source/Runtime/Engine/Public/ShowFlags.h b/Engine/Source/Runtime/Engine/Public/ShowFlags.h
--- a/Engine/Source/Runtime/Engine/Public/ShowFlags.h
+++ b/Engine/Source/Runtime/Engine/Public/ShowFlags.h
@@ -17,6 +17,7 @@ enum EShowFlagGroup
 	SFG_CollisionModes,
 	SFG_Developer,
 	SFG_Visualize,
+	SFG_LightTypes,
 	SFG_LightingComponents,
 	SFG_LightingFeatures,
 	SFG_Hidden,
diff --git a/Engine/Source/Runtime/Engine/Public/ShowFlagsValues.inl b/Engine/Source/Runtime/Engine/Public/ShowFlagsValues.inl
--- a/Engine/Source/Runtime/Engine/Public/ShowFlagsValues.inl
+++ b/Engine/Source/Runtime/Engine/Public/ShowFlagsValues.inl
@@ -57,11 +57,11 @@ SHOWFLAG_ALWAYS_ACCESSIBLE(VisualizeBuffer, SFG_Hidden, LOCTEXT("VisualizeBuffer
 /** Allows to disable all direct lighting (does not affect indirect light) */
 SHOWFLAG_FIXED_IN_SHIPPING(DirectLighting, 1, SFG_LightingComponents, LOCTEXT("DirectLightingSF", "Direct Lighting"))
 /** Allows to disable lighting from Directional Lights */
-SHOWFLAG_FIXED_IN_SHIPPING(DirectionalLights, 1, SFG_LightingComponents, LOCTEXT("DirectionalLightsSF", "Directional Lights"))
+SHOWFLAG_FIXED_IN_SHIPPING(DirectionalLights, 1, SFG_LightTypes, LOCTEXT("DirectionalLightsSF", "Directional Lights"))
 /** Allows to disable lighting from Point Lights */
-SHOWFLAG_FIXED_IN_SHIPPING(PointLights, 1, SFG_LightingComponents, LOCTEXT("PointLightsSF", "Point Lights"))
+SHOWFLAG_FIXED_IN_SHIPPING(PointLights, 1, SFG_LightTypes, LOCTEXT("PointLightsSF", "Point Lights"))
 /** Allows to disable lighting from Spot Lights */
-SHOWFLAG_FIXED_IN_SHIPPING(SpotLights, 1, SFG_LightingComponents, LOCTEXT("SpotLightsSF", "Spot Lights"))
+SHOWFLAG_FIXED_IN_SHIPPING(SpotLights, 1, SFG_LightTypes, LOCTEXT("SpotLightsSF", "Spot Lights"))
 /** Color correction after tone mapping */
 SHOWFLAG_ALWAYS_ACCESSIBLE(ColorGrading, SFG_PostProcess, LOCTEXT("ColorGradingSF", "Color Grading"))
 /** Visualize vector fields. */
@@ -257,7 +257,7 @@ SHOWFLAG_ALWAYS_ACCESSIBLE(VisualizeLightCulling, SFG_Hidden, LOCTEXT("Visualize
 /** To disable precomputed visibility */
 SHOWFLAG_ALWAYS_ACCESSIBLE(PrecomputedVisibility, SFG_Advanced, LOCTEXT("PrecomputedVisibilitySF", "Precomputed Visibility"))
 /** Contribution from sky light */
-SHOWFLAG_ALWAYS_ACCESSIBLE(SkyLighting, SFG_LightingComponents, LOCTEXT("SkyLightingSF", "Sky Lighting"))
+SHOWFLAG_ALWAYS_ACCESSIBLE(SkyLighting, SFG_LightTypes, LOCTEXT("SkyLightingSF", "Sky Lighting"))
 /** Visualize Light Propagation Volume, for developer (by default off): */
 SHOWFLAG_ALWAYS_ACCESSIBLE(VisualizeLPV, SFG_Visualize, LOCTEXT("VisualizeLPVSF", "Light Propagation Volume"))
 /** Visualize preview shadow indicator */
@@ -274,6 +274,8 @@ SHOWFLAG_ALWAYS_ACCESSIBLE(VisualizeDistanceFieldAO, SFG_Visualize, LOCTEXT("Vis
 SHOWFLAG_ALWAYS_ACCESSIBLE(VisualizeDistanceFieldGI, SFG_Visualize, LOCTEXT("VisualizeDistanceFieldGISF", "Distance Field Global Illumination"))
 /** Mesh Distance fields */
 SHOWFLAG_ALWAYS_ACCESSIBLE(VisualizeMeshDistanceFields, SFG_Visualize, LOCTEXT("MeshDistanceFieldsSF", "Mesh DistanceFields"))
+/** Screen space AO */
+SHOWFLAG_ALWAYS_ACCESSIBLE(ScreenSpaceAO, SFG_LightingFeatures, LOCTEXT("ScreenSpaceAOSF", "Screen Space Ambient Occlusion"))
 /** Distance field AO */
 SHOWFLAG_ALWAYS_ACCESSIBLE(DistanceFieldAO, SFG_LightingFeatures, LOCTEXT("DistanceFieldAOSF", "Distance Field Ambient Occlusion"))
 /** Distance field GI */
diff --git a/Engine/Source/Runtime/Renderer/Private/DistanceFieldSurfaceCacheLighting.cpp b/Engine/Source/Runtime/Renderer/Private/DistanceFieldSurfaceCacheLighting.cpp
--- a/Engine/Source/Runtime/Renderer/Private/DistanceFieldSurfaceCacheLighting.cpp
+++ b/Engine/Source/Runtime/Renderer/Private/DistanceFieldSurfaceCacheLighting.cpp
@@ -3171,7 +3171,7 @@ void FDeferredShadingSceneRenderer::RenderDynamicSkyLighting(FRHICommandListImme
 		FDistanceFieldAOParameters Parameters(Scene->SkyLight->OcclusionMaxDistance, Scene->SkyLight->Contrast);
 		TRefCountPtr<IPooledRenderTarget> DynamicIrradiance;
 
-		if (ShouldRenderDistanceFieldAO())
+		if (ShouldRenderDistanceFieldAO() && ViewFamily.EngineShowFlags.AmbientOcclusion)
 		{
 			bApplyShadowing = RenderDistanceFieldLighting(RHICmdList, Parameters, VelocityTexture, DynamicBentNormalAO, DynamicIrradiance, false, false);
 		}
