Proper Thumbnail for blueprint with inheritable blueprint components
https://udn.unrealengine.com/questions/246332/thumbnail-not-working-for-blueprint-with-inheritab.html

diff --git a/Engine/Source/Editor/UnrealEd/Private/ThumbnailHelpers.cpp b/Engine/Source/Editor/UnrealEd/Private/ThumbnailHelpers.cpp
--- a/Engine/Source/Editor/UnrealEd/Private/ThumbnailHelpers.cpp
+++ b/Engine/Source/Editor/UnrealEd/Private/ThumbnailHelpers.cpp
@@ -1132,13 +1132,13 @@ void FBlueprintThumbnailScene::BlueprintChanged(class UBlueprint* Blueprint)
 	ClearComponentsPool();
 }
 
-void FBlueprintThumbnailScene::InstanceComponents(USCS_Node* CurrentNode, USceneComponent* ParentComponent, const TMap<UActorComponent*, UActorComponent*>& NativeInstanceMap, TArray<UActorComponent*>& OutComponents)
+void FBlueprintThumbnailScene::InstanceComponents(USCS_Node* CurrentNode, USceneComponent* ParentComponent, const TMap<UActorComponent*, UActorComponent*>& NativeInstanceMap, TArray<UActorComponent*>& OutComponents, UBlueprintGeneratedClass* ActualBPGC)
 {
 	check(CurrentNode != NULL);
 
 	// Get the instanced actor component for this node. This is either an instance made from the native components, or a new instance we create using the current node's template.
 	UActorComponent* CurrentActorComp = NULL;
-	UActorComponent* ComponentTemplate = CurrentNode->ComponentTemplate;
+	UActorComponent* ComponentTemplate = CurrentNode->GetActualComponentTemplate(ActualBPGC);
 	if ( ComponentTemplate != NULL )
 	{
 		// Try to find the template in the list of native components we processed. If we find it, use the corresponding instance instead of making a new one.
@@ -1187,7 +1187,7 @@ void FBlueprintThumbnailScene::InstanceComponents(USCS_Node* CurrentNode, UScene
 		{
 			USCS_Node* Node = CurrentNode->ChildNodes[NodeIdx];
 			check(Node != NULL);
-			InstanceComponents(Node, ParentSceneComponentOfChildren, NativeInstanceMap, OutComponents);
+			InstanceComponents(Node, ParentSceneComponentOfChildren, NativeInstanceMap, OutComponents, ActualBPGC);
 		}
 	}
 }
@@ -1306,7 +1306,7 @@ TArray<UPrimitiveComponent*> FBlueprintThumbnailScene::GetPooledVisualizableComp
 								}
 							}
 
-							InstanceComponents(RootNode, ParentComponent, NativeInstanceMap, AllCreatedActorComponents);
+							InstanceComponents(RootNode, ParentComponent, NativeInstanceMap, AllCreatedActorComponents, Cast<UBlueprintGeneratedClass>(Blueprint->GeneratedClass));
 						}
 					}
 				}
diff --git a/Engine/Source/Editor/UnrealEd/Public/ThumbnailHelpers.h b/Engine/Source/Editor/UnrealEd/Public/ThumbnailHelpers.h
--- a/Engine/Source/Editor/UnrealEd/Public/ThumbnailHelpers.h
+++ b/Engine/Source/Editor/UnrealEd/Public/ThumbnailHelpers.h
@@ -260,7 +260,7 @@ public:
 
 protected:
 	/** Creates instances of template components found in a blueprint's simple construction script */
-	void InstanceComponents(USCS_Node* CurrentNode, USceneComponent* ParentComponent, const TMap<UActorComponent*, UActorComponent*>& NativeInstanceMap, TArray<UActorComponent*>& OutComponents);
+	void InstanceComponents(USCS_Node* CurrentNode, USceneComponent* ParentComponent, const TMap<UActorComponent*, UActorComponent*>& NativeInstanceMap, TArray<UActorComponent*>& OutComponents, UBlueprintGeneratedClass* ActualBPGC);
 
 	/** Get/Release for the component pool */
 	virtual TArray<UPrimitiveComponent*> GetPooledVisualizableComponents(UObject* Obj) override;
diff --git a/Engine/Source/Runtime/Engine/Classes/Engine/SCS_Node.h b/Engine/Source/Runtime/Engine/Classes/Engine/SCS_Node.h
--- a/Engine/Source/Runtime/Engine/Classes/Engine/SCS_Node.h
+++ b/Engine/Source/Runtime/Engine/Classes/Engine/SCS_Node.h
@@ -92,6 +92,9 @@ class USCS_Node : public UObject
 	 */
 	UActorComponent* ExecuteNodeOnActor(AActor* Actor, USceneComponent* ParentComponent, const FTransform* RootTransform, bool bIsDefaultTransform);
 
+	/** Return the actual component template used in the BPGC. The template can be overridden in a child. */
+	ENGINE_API UActorComponent* GetActualComponentTemplate(UBlueprintGeneratedClass* ActualBPGC);
+
 	/** Returns an array containing this node and all children below it */
 	TArray<USCS_Node*> GetAllNodes();
 	
diff --git a/Engine/Source/Runtime/Engine/Private/SCS_Node.cpp b/Engine/Source/Runtime/Engine/Private/SCS_Node.cpp
--- a/Engine/Source/Runtime/Engine/Private/SCS_Node.cpp
+++ b/Engine/Source/Runtime/Engine/Private/SCS_Node.cpp
@@ -26,17 +26,13 @@ USCS_Node::USCS_Node(const FObjectInitializer& ObjectInitializer)
 #endif
 }
 
-UActorComponent* USCS_Node::ExecuteNodeOnActor(AActor* Actor, USceneComponent* ParentComponent, const FTransform* RootTransform, bool bIsDefaultTransform)
+UActorComponent* USCS_Node::GetActualComponentTemplate(UBlueprintGeneratedClass* ActualBPGC)
 {
-	check(Actor != nullptr);
-	check((ParentComponent != nullptr && !ParentComponent->IsPendingKill()) || (RootTransform != nullptr)); // must specify either a parent component or a world transform
-
 	UActorComponent* OverridenComponentTemplate = nullptr;
 	static const FBoolConfigValueHelper EnableInheritableComponents(TEXT("Kismet"), TEXT("bEnableInheritableComponents"), GEngineIni);
 	if (EnableInheritableComponents)
 	{
 		const FComponentKey ComponentKey(this);
-		auto ActualBPGC = Cast<UBlueprintGeneratedClass>(Actor->GetClass());
 		while (!OverridenComponentTemplate && ActualBPGC)
 		{
 			auto InheritableComponentHandler = ActualBPGC->GetInheritableComponentHandler();
@@ -47,7 +43,16 @@ UActorComponent* USCS_Node::ExecuteNodeOnActor(AActor* Actor, USceneComponent* P
 			ActualBPGC = Cast<UBlueprintGeneratedClass>(ActualBPGC->GetSuperClass());
 		}
 	}
-	UActorComponent* ActualComponentTemplate = OverridenComponentTemplate ? OverridenComponentTemplate : ComponentTemplate;
+	return OverridenComponentTemplate ? OverridenComponentTemplate : ComponentTemplate;
+}
+
+UActorComponent* USCS_Node::ExecuteNodeOnActor(AActor* Actor, USceneComponent* ParentComponent, const FTransform* RootTransform, bool bIsDefaultTransform)
+{
+	check(Actor != nullptr);
+	check((ParentComponent != nullptr && !ParentComponent->IsPendingKill()) || (RootTransform != nullptr)); // must specify either a parent component or a world transform
+
+	auto ActualBPGC = Cast<UBlueprintGeneratedClass>(Actor->GetClass());
+	UActorComponent* ActualComponentTemplate = GetActualComponentTemplate(ActualBPGC);
 
 	// Create a new component instance based on the template
 	UActorComponent* NewActorComp = Actor->CreateComponentFromTemplate(ActualComponentTemplate, VariableName.ToString());
