Added new shapes to log with Visual Logger: mesh, convex polygon, navarea / extruded convex and navoctree dump from given bounding box.

diff --git a/Engine/Source/Developer/LogVisualizer/Private/LogVisualizerSettings.cpp b/Engine/Source/Developer/LogVisualizer/Private/LogVisualizerSettings.cpp
--- a/Engine/Source/Developer/LogVisualizer/Private/LogVisualizerSettings.cpp
+++ b/Engine/Source/Developer/LogVisualizer/Private/LogVisualizerSettings.cpp
@@ -12,7 +12,9 @@
 
 FCategoryFiltersManager FCategoryFiltersManager::StaticManager;
 
-ULogVisualizerSettings::ULogVisualizerSettings(const FObjectInitializer& ObjectInitializer) : Super(ObjectInitializer)
+ULogVisualizerSettings::ULogVisualizerSettings(const FObjectInitializer& ObjectInitializer) 
+	: Super(ObjectInitializer)
+	, DebugMeshMaterialFakeLightName(TEXT("/Engine/EngineDebugMaterials/DebugMeshMaterialFakeLight.DebugMeshMaterialFakeLight"))
 {
 	TrivialLogsThreshold = 1;
 	DefaultCameraDistance = 150;
@@ -23,6 +25,16 @@ ULogVisualizerSettings::ULogVisualizerSettings(const FObjectInitializer& ObjectI
 	bUsePlayersOnlyForPause = true;
 }
 
+class UMaterial* ULogVisualizerSettings::GetDebugMeshMaterial()
+{
+	if (DebugMeshMaterialFakeLight == nullptr)
+	{
+		DebugMeshMaterialFakeLight = LoadObject<UMaterial>(NULL, *DebugMeshMaterialFakeLightName, NULL, LOAD_None, NULL);
+	}
+
+	return DebugMeshMaterialFakeLight;
+}
+
 #if WITH_EDITOR
 void ULogVisualizerSettings::PostEditChangeProperty(struct FPropertyChangedEvent& PropertyChangedEvent)
 {
diff --git a/Engine/Source/Developer/LogVisualizer/Private/SVisualLogger.cpp b/Engine/Source/Developer/LogVisualizer/Private/SVisualLogger.cpp
--- a/Engine/Source/Developer/LogVisualizer/Private/SVisualLogger.cpp
+++ b/Engine/Source/Developer/LogVisualizer/Private/SVisualLogger.cpp
@@ -780,6 +780,11 @@ void SVisualLogger::OnFiltersChanged()
 {
 	LogsList->OnFiltersChanged();
 	MainView->OnFiltersChanged();
+	if (AVisualLoggerRenderingActor* HelperActor = Cast<AVisualLoggerRenderingActor>(FLogVisualizer::Get().GetVisualLoggerHelperActor()))
+	{
+		HelperActor->OnItemSelectionChanged(LogsList->GetCurrentLogEntry());
+		HelperActor->MarkComponentsRenderStateDirty();
+	}
 }
 
 void SVisualLogger::OnObjectSelectionChanged(TSharedPtr<class STimeline> TimeLine)
@@ -814,8 +819,8 @@ void SVisualLogger::OnFiltersSearchChanged(const FText& Filter)
 
 void SVisualLogger::OnLogLineSelectionChanged(TSharedPtr<struct FLogEntryItem> SelectedItem, int64 UserData, FName TagName)
 {
-	TMap<FName, FVisualLogExtensionInterface*>& AllExtensions = FVisualLogger::Get().GetAllExtensions();
-	for (auto Iterator = AllExtensions.CreateIterator(); Iterator; ++Iterator)
+	const TMap<FName, FVisualLogExtensionInterface*>& AllExtensions = FVisualLogger::Get().GetAllExtensions();
+	for (auto Iterator = AllExtensions.CreateConstIterator(); Iterator; ++Iterator)
 	{
 		FVisualLogExtensionInterface* Extension = (*Iterator).Value;
 		if (Extension != NULL)
diff --git a/Engine/Source/Developer/LogVisualizer/Private/VisualLoggerRenderingActor.cpp b/Engine/Source/Developer/LogVisualizer/Private/VisualLoggerRenderingActor.cpp
--- a/Engine/Source/Developer/LogVisualizer/Private/VisualLoggerRenderingActor.cpp
+++ b/Engine/Source/Developer/LogVisualizer/Private/VisualLoggerRenderingActor.cpp
@@ -4,6 +4,7 @@
 #include "Engine/GameInstance.h"
 #include "Debug/DebugDrawService.h"
 #include "GameFramework/HUD.h"
+#include "AI/Navigation/RecastHelpers.h"
 #if WITH_EDITOR
 #include "Editor/EditorEngine.h"
 #include "GeomTools.h"
@@ -54,7 +55,9 @@ FPrimitiveSceneProxy* UVisualLoggerRenderingComponent::CreateSceneProxy()
 		return NULL;
 	}
 
+	ULogVisualizerSettings *Settings = ULogVisualizerSettings::StaticClass()->GetDefaultObject<ULogVisualizerSettings>();
 	FVisualLoggerSceneProxy *VLogSceneProxy = new FVisualLoggerSceneProxy(this);
+	VLogSceneProxy->SolidMeshMaterial = Settings->GetDebugMeshMaterial();
 	VLogSceneProxy->Spheres = RenderingActor->Points;
 	VLogSceneProxy->Lines = RenderingActor->Lines;
 	VLogSceneProxy->Cones = RenderingActor->Cones;
@@ -184,6 +187,39 @@ void AVisualLoggerRenderingActor::ObjectSelectionChanged(TSharedPtr<class STimel
 	}
 }
 
+namespace
+{
+	static void GetPolygonMesh(const FVisualLogShapeElement* ElementToDraw, FDebugRenderSceneProxy::FMesh& TestMesh, const FVector& VertexOffset = FVector::ZeroVector)
+	{
+		TestMesh.Color = ElementToDraw->GetFColor();
+
+		FClipSMPolygon InPoly(ElementToDraw->Points.Num());
+		for (int32 Index = 0; Index < ElementToDraw->Points.Num(); Index++)
+		{
+			FClipSMVertex v1;
+			v1.Pos = ElementToDraw->Points[Index];
+			InPoly.Vertices.Add(v1);
+		}
+
+		TArray<FClipSMTriangle> OutTris;
+		if (TriangulatePoly(OutTris, InPoly, false))
+		{
+			int32 LastIndex = 0;
+
+			RemoveRedundantTriangles(OutTris);
+			for (const auto& CurrentTri : OutTris)
+			{
+				TestMesh.Vertices.Add(FDynamicMeshVertex(CurrentTri.Vertices[0].Pos + VertexOffset));
+				TestMesh.Vertices.Add(FDynamicMeshVertex(CurrentTri.Vertices[1].Pos + VertexOffset));
+				TestMesh.Vertices.Add(FDynamicMeshVertex(CurrentTri.Vertices[2].Pos + VertexOffset));
+				TestMesh.Indices.Add(LastIndex++);
+				TestMesh.Indices.Add(LastIndex++);
+				TestMesh.Indices.Add(LastIndex++);
+			}
+		}
+	}
+}
+
 void AVisualLoggerRenderingActor::OnItemSelectionChanged(const FVisualLogDevice::FVisualLogEntryItem& EntryItem)
 {
 	const FVisualLogEntry* Entry = &EntryItem.Entry;
@@ -222,7 +258,7 @@ void AVisualLoggerRenderingActor::OnItemSelectionChanged(const FVisualLogDevice:
 		}
 	}
 
-	TMap<FName, FVisualLogExtensionInterface*>& AllExtensions = FVisualLogger::Get().GetAllExtensions();
+	const TMap<FName, FVisualLogExtensionInterface*>& AllExtensions = FVisualLogger::Get().GetAllExtensions();
 	for (auto& Extension : AllExtensions)
 	{
 		Extension.Value->OnTimestampChange(Entry->TimeStamp, GetWorld(), this);
@@ -265,7 +301,7 @@ void AVisualLoggerRenderingActor::OnItemSelectionChanged(const FVisualLogDevice:
 		{
 			const float Radius = float(ElementToDraw->Radius);
 			const bool bDrawLabel = ElementToDraw->Description.IsEmpty() == false;
-			for (int32 Index=0; Index < ElementToDraw->Points.Num(); ++Index)
+			for (int32 Index = 0; Index < ElementToDraw->Points.Num(); ++Index)
 			{
 				const FVector& Point = ElementToDraw->Points[Index];
 				Points.Add(FDebugRenderSceneProxy::FSphere(Radius, Point, Color));
@@ -276,47 +312,53 @@ void AVisualLoggerRenderingActor::OnItemSelectionChanged(const FVisualLogDevice:
 				}
 			}
 		}
-		break;
+			break;
 		case EVisualLoggerShapeElement::Polygon:
 		{
-			FClipSMPolygon InPoly(ElementToDraw->Points.Num());
-			for (int32 Index = 0; Index < ElementToDraw->Points.Num(); Index++)
-			{
-				FClipSMVertex v1;
-				v1.Pos = ElementToDraw->Points[Index];
-				InPoly.Vertices.Add(v1);
-			}
-
-			TArray<FClipSMTriangle> OutTris;
-			if (TriangulatePoly(OutTris, InPoly, false))
-			{
-				int32 LastIndex = 0;
-				FDebugRenderSceneProxy::FMesh TestMesh;
-				TestMesh.Color = Color;
-
-				RemoveRedundantTriangles(OutTris);
-				for (const auto& CurrentTri : OutTris)
-				{
-					TestMesh.Vertices.Add(FDynamicMeshVertex(CurrentTri.Vertices[0].Pos + CorridorOffset));
-					TestMesh.Vertices.Add(FDynamicMeshVertex(CurrentTri.Vertices[1].Pos + CorridorOffset));
-					TestMesh.Vertices.Add(FDynamicMeshVertex(CurrentTri.Vertices[2].Pos + CorridorOffset));
-					TestMesh.Indices.Add(LastIndex++);
-					TestMesh.Indices.Add(LastIndex++);
-					TestMesh.Indices.Add(LastIndex++);
-				}
-				Meshes.Add(TestMesh);
-			}
+			FDebugRenderSceneProxy::FMesh TestMesh;
+			GetPolygonMesh(ElementToDraw, TestMesh, CorridorOffset);
+			Meshes.Add(TestMesh);
 
 			for (int32 VIdx = 0; VIdx < ElementToDraw->Points.Num(); VIdx++)
 			{
 				Lines.Add(FDebugRenderSceneProxy::FDebugLine(
 					ElementToDraw->Points[VIdx] + CorridorOffset,
 					ElementToDraw->Points[(VIdx + 1) % ElementToDraw->Points.Num()] + CorridorOffset,
-					FColor::Cyan, 
+					FColor::Cyan,
 					2)
-				);
+					);
+			}
+		}
+			break;
+		case EVisualLoggerShapeElement::Mesh:
+		{
+			struct FHeaderData
+			{
+				float VerticesNum, FacesNum;
+				FHeaderData(const FVector& InVector) : VerticesNum(InVector.X), FacesNum(InVector.Y) {}
+			};
+			const FHeaderData HeaderData(ElementToDraw->Points[0]);
+
+			FDebugRenderSceneProxy::FMesh TestMesh;
+			TestMesh.Color = ElementToDraw->GetFColor();
+			int32 StartIndex = 1;
+			int32 EndIndex = StartIndex + HeaderData.VerticesNum;
+			for (int32 VIdx = StartIndex; VIdx < EndIndex; VIdx++)
+			{
+				TestMesh.Vertices.Add(ElementToDraw->Points[VIdx]);
 			}
 
+
+			StartIndex = EndIndex;
+			EndIndex = StartIndex + HeaderData.FacesNum;
+			for (int32 VIdx = StartIndex; VIdx < EndIndex; VIdx++)
+			{
+				const FVector &CurrentFace = ElementToDraw->Points[VIdx];
+				TestMesh.Indices.Add(CurrentFace.X);
+				TestMesh.Indices.Add(CurrentFace.Y);
+				TestMesh.Indices.Add(CurrentFace.Z);
+			}
+			Meshes.Add(TestMesh);
 		}
 			break;
 		case EVisualLoggerShapeElement::Segment:
@@ -351,7 +393,7 @@ void AVisualLoggerRenderingActor::OnItemSelectionChanged(const FVisualLogDevice:
 				Location = CurrentLocation;
 			}
 		}
-		break;
+			break;
 		case EVisualLoggerShapeElement::Box:
 		{
 			const float Thickness = float(ElementToDraw->Thicknes);
@@ -373,7 +415,7 @@ void AVisualLoggerRenderingActor::OnItemSelectionChanged(const FVisualLogDevice:
 				Texts.Add(FDebugRenderSceneProxy::FText3d(ElementToDraw->Description, ElementToDraw->Points[0] + (ElementToDraw->Points[1] - ElementToDraw->Points[0]) / 2, Color));
 			}
 		}
-		break;
+			break;
 		case EVisualLoggerShapeElement::Cone:
 		{
 			const float Thickness = float(ElementToDraw->Thicknes);
@@ -425,7 +467,7 @@ void AVisualLoggerRenderingActor::OnItemSelectionChanged(const FVisualLogDevice:
 				const float HalfHeight = FirstData.X;
 				const float Radius = FirstData.Y;
 				const FQuat Rotation = FQuat(FirstData.Z, SecondData.X, SecondData.Y, SecondData.Z);
-				
+
 				const FMatrix Axes = FQuatRotationTranslationMatrix(Rotation, FVector::ZeroVector);
 				const FVector XAxis = Axes.GetScaledAxis(EAxis::X);
 				const FVector YAxis = Axes.GetScaledAxis(EAxis::Y);
@@ -438,7 +480,73 @@ void AVisualLoggerRenderingActor::OnItemSelectionChanged(const FVisualLogDevice:
 				}
 			}
 		}
-			break;
+		break;
+		case EVisualLoggerShapeElement::NavAreaMesh:
+		{
+			if (ElementToDraw->Points.Num() == 0)
+				continue;
+
+			struct FHeaderData
+			{
+				float MinZ, MaxZ;
+				FHeaderData(const FVector& InVector) : MinZ(InVector.X), MaxZ(InVector.Y) {}
+			};
+			const FHeaderData HeaderData(ElementToDraw->Points[0]);
+
+			TArray<FVector> AreaMeshPoints = ElementToDraw->Points;
+			AreaMeshPoints.RemoveAt(0, 1, false);
+			AreaMeshPoints.Add(ElementToDraw->Points[1]);
+			TArray<FVector> Vertices;
+			TNavStatArray<FVector> Faces;
+			int32 CurrentIndex = 0;
+			FDebugRenderSceneProxy::FMesh TestMesh;
+			TestMesh.Color = ElementToDraw->GetFColor();
+
+			for (int32 PointIndex = 0; PointIndex < AreaMeshPoints.Num() - 1; PointIndex++)
+			{
+				FVector Point = AreaMeshPoints[PointIndex];
+				FVector NextPoint = AreaMeshPoints[PointIndex + 1];
+
+				FVector P1(Point.X, Point.Y, HeaderData.MinZ);
+				FVector P2(Point.X, Point.Y, HeaderData.MaxZ);
+				FVector P3(NextPoint.X, NextPoint.Y, HeaderData.MinZ);
+				FVector P4(NextPoint.X, NextPoint.Y, HeaderData.MaxZ);
+
+				TestMesh.Vertices.Add(P1); TestMesh.Vertices.Add(P2); TestMesh.Vertices.Add(P3);
+				TestMesh.Indices.Add(CurrentIndex + 0);
+				TestMesh.Indices.Add(CurrentIndex + 1);
+				TestMesh.Indices.Add(CurrentIndex + 2);
+				CurrentIndex += 3;
+				TestMesh.Vertices.Add(P3); TestMesh.Vertices.Add(P2); TestMesh.Vertices.Add(P4);
+				TestMesh.Indices.Add(CurrentIndex + 0);
+				TestMesh.Indices.Add(CurrentIndex + 1);
+				TestMesh.Indices.Add(CurrentIndex + 2);
+				CurrentIndex += 3;
+			}
+			Meshes.Add(TestMesh);
+
+			{
+				FDebugRenderSceneProxy::FMesh PolygonMesh;
+				FVisualLogShapeElement PolygonToDraw(EVisualLoggerShapeElement::Polygon);
+				PolygonToDraw.SetColor(ElementToDraw->GetFColor());
+				PolygonToDraw.Points.Reserve(AreaMeshPoints.Num());
+				PolygonToDraw.Points = AreaMeshPoints;
+				GetPolygonMesh(&PolygonToDraw, PolygonMesh, FVector(0,0,HeaderData.MaxZ));
+				Meshes.Add(PolygonMesh);
+			}
+
+			for (int32 VIdx = 0; VIdx < AreaMeshPoints.Num(); VIdx++)
+			{
+				Lines.Add(FDebugRenderSceneProxy::FDebugLine(
+					AreaMeshPoints[VIdx] + FVector(0, 0, HeaderData.MaxZ),
+					AreaMeshPoints[(VIdx + 1) % AreaMeshPoints.Num()] + FVector(0, 0, HeaderData.MaxZ),
+					ElementToDraw->GetFColor(),
+					2)
+				);
+			}
+
+		}
+		break;
 		}
 	}
 
diff --git a/Engine/Source/Developer/LogVisualizer/Public/LogVisualizerSettings.h b/Engine/Source/Developer/LogVisualizer/Public/LogVisualizerSettings.h
--- a/Engine/Source/Developer/LogVisualizer/Public/LogVisualizerSettings.h
+++ b/Engine/Source/Developer/LogVisualizer/Public/LogVisualizerSettings.h
@@ -102,12 +102,21 @@ public:
 	virtual void PostEditChangeProperty(struct FPropertyChangedEvent& PropertyChangedEvent) override;
 #endif
 
+	class UMaterial* GetDebugMeshMaterial();
+
 protected:
 	FVisualLoggerFilters CurrentFilters;
 	
 	UPROPERTY(config)
 	FVisualLoggerFilters PresistentFilters;
 
+	/** A material used to render debug meshes with kind of flat shading, mostly used by Visual Logger tool. */
+	UPROPERTY()
+	class UMaterial* DebugMeshMaterialFakeLight;
+
+	/** @todo document */
+	UPROPERTY(config)
+	FString DebugMeshMaterialFakeLightName;
 private:
 
 	// Holds an event delegate that is executed when a setting has changed.
diff --git a/Engine/Source/Runtime/Engine/Private/AI/Navigation/NavigationPath.cpp b/Engine/Source/Runtime/Engine/Private/AI/Navigation/NavigationPath.cpp
--- a/Engine/Source/Runtime/Engine/Private/AI/Navigation/NavigationPath.cpp
+++ b/Engine/Source/Runtime/Engine/Private/AI/Navigation/NavigationPath.cpp
@@ -1129,8 +1129,9 @@ void FNavMeshPath::DescribeSelfToVisLog(FVisualLogEntry* Snapshot) const
 	// draw corridor
 #if WITH_RECAST
 	FVisualLogShapeElement CorridorPoly(EVisualLoggerShapeElement::Polygon);
-	CorridorPoly.SetColor(FColorList::Cyan);
+	CorridorPoly.SetColor(FColorList::Cyan.WithAlpha(100));
 	CorridorPoly.Category = LogNavigation.GetCategoryName();
+	CorridorPoly.Verbosity = ELogVerbosity::Verbose;
 	CorridorPoly.Points.Reserve(PathCorridor.Num() * 6);
 
 	const FVector CorridorOffset = NavigationDebugDrawing::PathOffset * 1.25f;
@@ -1151,7 +1152,7 @@ void FNavMeshPath::DescribeSelfToVisLog(FVisualLogEntry* Snapshot) const
 		const UNavArea* DefArea = AreaClass ? ((UClass*)AreaClass)->GetDefaultObject<UNavArea>() : NULL;
 		const FColor PolygonColor = AreaClass != UNavigationSystem::GetDefaultWalkableArea() ? (DefArea ? DefArea->DrawColor : NavMesh->GetConfig().Color) : FColorList::Cyan;
 
-		CorridorPoly.SetColor(PolygonColor);
+		CorridorPoly.SetColor(PolygonColor.WithAlpha(100));
 		CorridorPoly.Points.Reset();
 		CorridorPoly.Points.Append(Verts);
 		Snapshot->ElementsToDraw.Add(CorridorPoly);
@@ -1168,6 +1169,7 @@ void FNavMeshPath::DescribeSelfToVisLog(FVisualLogEntry* Snapshot) const
 			FVisualLogShapeElement AreaMarkElem(EVisualLoggerShapeElement::Segment);
 			AreaMarkElem.SetColor(FColorList::Orange);
 			AreaMarkElem.Category = LogNavigation.GetCategoryName();
+			AreaMarkElem.Verbosity = ELogVerbosity::Verbose;
 			AreaMarkElem.Thicknes = 2;
 			AreaMarkElem.Description = AreaClass->GetName();
 
diff --git a/Engine/Source/Runtime/Engine/Private/AI/Navigation/RecastNavMeshGenerator.cpp b/Engine/Source/Runtime/Engine/Private/AI/Navigation/RecastNavMeshGenerator.cpp
--- a/Engine/Source/Runtime/Engine/Private/AI/Navigation/RecastNavMeshGenerator.cpp
+++ b/Engine/Source/Runtime/Engine/Private/AI/Navigation/RecastNavMeshGenerator.cpp
@@ -23,6 +23,9 @@
 #include "DetourTileCacheBuilder.h"
 #include "RecastHelpers.h"
 #include "NavigationSystemHelpers.h"
+#include "VisualLogger/VisualLogger.h"
+#include "NavmeshRenderingHelpers.h"
+
 
 #define SEAMLESS_REBUILDING_ENABLED 1
 
@@ -4146,6 +4149,98 @@ uint32 FRecastNavMeshGenerator::LogMemUsed() const
 	return GeneratorsMem + sizeof(FRecastNavMeshGenerator) + PendingDirtyTiles.GetAllocatedSize() + PendingDirtyTiles.GetAllocatedSize();
 }
 
+#if !(UE_BUILD_SHIPPING || UE_BUILD_TEST)
+void FRecastNavMeshGenerator::GrabDebugSnapshot(struct FVisualLogEntry* Snapshot, const FBox& BoundingBox, const struct FLogCategoryBase& LogCategory, ELogVerbosity::Type LogVerbosity) const
+{
+	const UNavigationSystem* NavSys = UNavigationSystem::GetCurrent(GetWorld());
+	const FNavigationOctree* NavOctree = NavSys ? NavSys->GetNavOctree() : NULL;
+	if (Snapshot == nullptr)
+	{
+		return;
+	}
+
+	if (NavOctree == NULL)
+	{
+		UE_LOG(LogNavigation, Error, TEXT("Failed to vlog navigation data due to %s being NULL"), NavSys == NULL ? TEXT("NavigationSystem") : TEXT("NavOctree"));
+		return;
+	}
+
+	ELogVerbosity::Type NavAreaVerbosity = FMath::Clamp(ELogVerbosity::Type(LogVerbosity + 1), ELogVerbosity::NoLogging, ELogVerbosity::VeryVerbose);
+
+	for (int32 Index = 0; Index < NavSys->NavDataSet.Num(); ++Index)
+	{
+		TArray<FVector> CoordBuffer;
+		TArray<int32> Indices;
+		TNavStatArray<FVector> Faces;
+		const ARecastNavMesh* NavData = Cast<const ARecastNavMesh>(NavSys->NavDataSet[Index]);
+		if (NavData)
+		{
+			for (FNavigationOctree::TConstElementBoxIterator<FNavigationOctree::DefaultStackAllocator> It(*NavOctree, BoundingBox);
+				It.HasPendingElements();
+				It.Advance())
+			{
+				const FNavigationOctreeElement& Element = It.GetCurrentElement();
+				const bool bExportGeometry = Element.Data->HasGeometry() && Element.ShouldUseGeometry(DestNavMesh->GetConfig());
+
+				if (bExportGeometry && Element.Data->CollisionData.Num())
+				{
+					FRecastGeometryCache CachedGeometry(Element.Data->CollisionData.GetData());
+					AppendGeometry(CoordBuffer, Indices, CachedGeometry.Verts, CachedGeometry.Header.NumVerts, CachedGeometry.Indices, CachedGeometry.Header.NumFaces);
+					Snapshot->AddElement(CoordBuffer, Indices, LogCategory.GetCategoryName(), LogVerbosity, FColorList::LightGrey.WithAlpha(255));
+				}
+				else
+				{
+					TArray<FVector> Verts;
+					const TArray<FAreaNavModifier>& AreaMods = Element.Data->Modifiers.GetAreas();
+					for (int32 i = 0; i < AreaMods.Num(); i++)
+					{
+						if (AreaMods[i].GetShapeType() == ENavigationShapeType::Unknown)
+						{
+							continue;
+						}
+
+						const uint8 AreaId = NavData->GetAreaID(AreaMods[i].GetAreaClass());
+						const UClass* AreaClass = NavData->GetAreaClass(AreaId);
+						const UNavArea* DefArea = AreaClass ? ((UClass*)AreaClass)->GetDefaultObject<UNavArea>() : NULL;
+						const FColor PolygonColor = AreaClass != UNavigationSystem::GetDefaultWalkableArea() ? (DefArea ? DefArea->DrawColor : NavData->GetConfig().Color) : FColorList::Cyan;
+
+						if (AreaMods[i].GetShapeType() == ENavigationShapeType::Box)
+						{
+							FBoxNavAreaData Box;
+							AreaMods[i].GetBox(Box);
+
+							Snapshot->AddElement(FBox::BuildAABB(Box.Origin, Box.Extent), FMatrix::Identity, LogCategory.GetCategoryName(), NavAreaVerbosity, PolygonColor.WithAlpha(255));
+						}
+						else if (AreaMods[i].GetShapeType() == ENavigationShapeType::Cylinder)
+						{
+							FCylinderNavAreaData Cylinder;
+							AreaMods[i].GetCylinder(Cylinder);
+
+							Snapshot->AddElement(Cylinder.Origin, Cylinder.Origin + FVector(0, 0, Cylinder.Height), Cylinder.Radius, LogCategory.GetCategoryName(), NavAreaVerbosity, PolygonColor.WithAlpha(255));
+						}
+						else
+						{
+							FConvexNavAreaData Convex;
+							AreaMods[i].GetConvex(Convex);
+							Verts.Reset();
+							GrowConvexHull(NavData->AgentRadius, Convex.Points, Verts);
+
+							Snapshot->AddElement(
+								Verts,
+								Convex.MinZ - NavData->CellHeight, 
+								Convex.MaxZ + NavData->CellHeight, 
+								LogCategory.GetCategoryName(), NavAreaVerbosity, PolygonColor.WithAlpha(255));
+						}
+					}
+				}
+			}
+
+		}
+
+	}
+}
+#endif
+
 #if !(UE_BUILD_SHIPPING || UE_BUILD_TEST)
 void FRecastNavMeshGenerator::ExportNavigationData(const FString& FileName) const
 {
diff --git a/Engine/Source/Runtime/Engine/Private/DebugRenderSceneProxy.cpp b/Engine/Source/Runtime/Engine/Private/DebugRenderSceneProxy.cpp
--- a/Engine/Source/Runtime/Engine/Private/DebugRenderSceneProxy.cpp
+++ b/Engine/Source/Runtime/Engine/Private/DebugRenderSceneProxy.cpp
@@ -13,7 +13,7 @@
 
 // FPrimitiveSceneProxy interface.
 
-FDebugRenderSceneProxy::FDebugRenderSceneProxy(const UPrimitiveComponent* InComponent) 
+FDebugRenderSceneProxy::FDebugRenderSceneProxy(const UPrimitiveComponent* InComponent)
 	: FPrimitiveSceneProxy(InComponent)
 	, ViewFlagIndex(uint32(FEngineShowFlags::FindIndexByName(TEXT("Game"))))
 	, ViewFlagName(TEXT("Game"))
@@ -21,7 +21,6 @@ FDebugRenderSceneProxy::FDebugRenderSceneProxy(const UPrimitiveComponent* InComp
 	, DrawType(WireMesh)
 	, DrawAlpha(100)
 {
-
 }
 
 void FDebugRenderSceneProxy::RegisterDebugDrawDelgate()
@@ -83,9 +82,9 @@ void FDebugRenderSceneProxy::GetDynamicMeshElements(const TArray<const FSceneVie
 	// Draw solid spheres
 	struct FMaterialCache
 	{
-		FMaterialCache() {}
+		FMaterialCache() : bUseFakeLight(false) {}
 
-		FMaterialRenderProxy* operator[](FLinearColor Color) const
+		FMaterialRenderProxy* operator[](FLinearColor Color)
 		{
 			FMaterialRenderProxy* MeshColor = NULL;
 			const uint32 HashKey = GetTypeHash(Color);
@@ -95,16 +94,35 @@ void FDebugRenderSceneProxy::GetDynamicMeshElements(const TArray<const FSceneVie
 			}
 			else
 			{
-				MeshColor = new(FMemStack::Get()) FColoredMaterialRenderProxy(GEngine->DebugMeshMaterial->GetRenderProxy(false), Color);
+				if (bUseFakeLight && SolidMeshMaterial.IsValid())
+				{
+					
+					MeshColor = new(FMemStack::Get())  FColoredMaterialRenderProxy(
+						SolidMeshMaterial->GetRenderProxy(false, false/*IsSelected(), IsHovered()*/),
+						Color,
+						"GizmoColor"
+						);
+				}
+				else
+				{
+					MeshColor = new(FMemStack::Get()) FColoredMaterialRenderProxy(GEngine->DebugMeshMaterial->GetRenderProxy(true), Color);
+				}
+
 				MeshColorInstances.Add(HashKey, MeshColor);
 			}
 
 			return MeshColor;
 		}
 
-		mutable TMap<uint32, FMaterialRenderProxy*> MeshColorInstances;
+		void UseFakeLight(bool UseLight, class UMaterial* InMaterial) { bUseFakeLight = UseLight; SolidMeshMaterial = InMaterial; }
+
+		TMap<uint32, FMaterialRenderProxy*> MeshColorInstances;
+		TWeakObjectPtr<class UMaterial> SolidMeshMaterial;
+		bool bUseFakeLight;
 	};
-	const FMaterialCache MaterialCache;
+
+	FMaterialCache MaterialCache[2];
+	MaterialCache[1].UseFakeLight(true, SolidMeshMaterial.Get());
 
 	for (int32 ViewIndex = 0; ViewIndex < Views.Num(); ViewIndex++)
 	{
@@ -155,7 +173,7 @@ void FDebugRenderSceneProxy::GetDynamicMeshElements(const TArray<const FSceneVie
 
 				if (DrawType == SolidAndWireMeshes || DrawType == SolidMesh)
 				{
-					GetCylinderMesh(Cylinder.Base, FVector(1, 0, 0), FVector(0, 1, 0), FVector(0, 0, 1), Cylinder.Radius, Cylinder.HalfHeight, 16, MaterialCache[Cylinder.Color.WithAlpha(DrawAlpha)], SDPG_World, ViewIndex, Collector);
+					GetCylinderMesh(Cylinder.Base, FVector(1, 0, 0), FVector(0, 1, 0), FVector(0, 0, 1), Cylinder.Radius, Cylinder.HalfHeight, 16, MaterialCache[0][Cylinder.Color.WithAlpha(DrawAlpha)], SDPG_World, ViewIndex, Collector);
 				}
 			}
 
@@ -168,7 +186,7 @@ void FDebugRenderSceneProxy::GetDynamicMeshElements(const TArray<const FSceneVie
 				}
 				if (DrawType == SolidAndWireMeshes || DrawType == SolidMesh)
 				{
-					GetBoxMesh(FTransform(Box.Box.GetCenter()).ToMatrixNoScale() * Box.Transform.ToMatrixWithScale(), Box.Box.GetExtent(), MaterialCache[Box.Color.WithAlpha(DrawAlpha)], SDPG_World, ViewIndex, Collector);
+					GetBoxMesh(FTransform(Box.Box.GetCenter()).ToMatrixNoScale() * Box.Transform.ToMatrixWithScale(), Box.Box.GetExtent(), MaterialCache[0][Box.Color.WithAlpha(DrawAlpha)], SDPG_World, ViewIndex, Collector);
 				}
 			}
 
@@ -182,7 +200,7 @@ void FDebugRenderSceneProxy::GetDynamicMeshElements(const TArray<const FSceneVie
 				}
 				if (DrawType == SolidAndWireMeshes || DrawType == SolidMesh)
 				{
-					GetConeMesh(CurrentCone.ConeToWorld, CurrentCone.Angle1, CurrentCone.Angle2, 16, MaterialCache[CurrentCone.Color.WithAlpha(DrawAlpha)], SDPG_World, ViewIndex, Collector);
+					GetConeMesh(CurrentCone.ConeToWorld, CurrentCone.Angle1, CurrentCone.Angle2, 16, MaterialCache[0][CurrentCone.Color.WithAlpha(DrawAlpha)], SDPG_World, ViewIndex, Collector);
 				}
 			}
 
@@ -196,7 +214,7 @@ void FDebugRenderSceneProxy::GetDynamicMeshElements(const TArray<const FSceneVie
 					}
 					if (DrawType == SolidAndWireMeshes || DrawType == SolidMesh)
 					{
-						GetSphereMesh(It->Location, FVector(It->Radius), 20, 7, MaterialCache[It->Color.WithAlpha(DrawAlpha)], SDPG_World, false, ViewIndex, Collector);
+						GetSphereMesh(It->Location, FVector(It->Radius), 20, 7, MaterialCache[0][It->Color.WithAlpha(DrawAlpha)], SDPG_World, false, ViewIndex, Collector);
 					}
 				}
 			}
@@ -216,7 +234,7 @@ void FDebugRenderSceneProxy::GetDynamicMeshElements(const TArray<const FSceneVie
 					}
 					if (DrawType == SolidAndWireMeshes || DrawType == SolidMesh)
 					{
-						GetCapsuleMesh(It->Location, It->X, It->Y, It->Z, It->Color, It->Radius, It->HalfHeight, 16, MaterialCache[It->Color.WithAlpha(DrawAlpha)], SDPG_World, false, ViewIndex, Collector);
+						GetCapsuleMesh(It->Location, It->X, It->Y, It->Z, It->Color, It->Radius, It->HalfHeight, 16, MaterialCache[0][It->Color.WithAlpha(DrawAlpha)], SDPG_World, false, ViewIndex, Collector);
 					}
 				}
 			}
@@ -227,7 +245,7 @@ void FDebugRenderSceneProxy::GetDynamicMeshElements(const TArray<const FSceneVie
 				MeshBuilder.AddVertices(Mesh.Vertices);
 				MeshBuilder.AddTriangles(Mesh.Indices);
 
-				MeshBuilder.GetMesh(FMatrix::Identity, MaterialCache[Mesh.Color.WithAlpha(DrawAlpha)], SDPG_World, false, false, ViewIndex, Collector);
+				MeshBuilder.GetMesh(FMatrix::Identity, MaterialCache[Mesh.Color.A == 255 ? 1 : 0][Mesh.Color.WithAlpha(DrawAlpha)], SDPG_World, false, false, ViewIndex, Collector);
 			}
 
 		}
diff --git a/Engine/Source/Runtime/Engine/Private/VisualLogger/VisualLogger.cpp b/Engine/Source/Runtime/Engine/Private/VisualLogger/VisualLogger.cpp
--- a/Engine/Source/Runtime/Engine/Private/VisualLogger/VisualLogger.cpp
+++ b/Engine/Source/Runtime/Engine/Private/VisualLogger/VisualLogger.cpp
@@ -249,6 +249,24 @@ void FVisualLogger::EventLog(const class UObject* Object, const FName EventTag1,
 	CurrentEntry->Events[Index].EventTags.Remove(NAME_None);
 }
 
+void FVisualLogger::NavOctreeShapeLog(const class UObject* Object, const struct FLogCategoryBase& Category, ELogVerbosity::Type Verbosity, int32 UniqueLogId, const FBox& Box)
+{
+	SCOPE_CYCLE_COUNTER(STAT_VisualLog);
+	UWorld *World = nullptr;
+	FVisualLogEntry *CurrentEntry = nullptr;
+	if (CheckVisualLogInputInternal(Object, Category, Verbosity, &World, &CurrentEntry) == false)
+	{
+		return;
+	}
+
+	const ANavigationData* MainNavData = World ? UNavigationSystem::GetNavigationSystem(World)->GetMainNavData(FNavigationSystem::ECreateIfEmpty::DontCreate) : nullptr;
+	const FNavDataGenerator* Generator = MainNavData ? MainNavData->GetGenerator() : nullptr;
+	if (Generator)
+	{
+		Generator->GrabDebugSnapshot(CurrentEntry, Box, Category, Verbosity);
+	}
+}
+
 
 FVisualLogger::FVisualLogger()
 {
@@ -448,7 +466,7 @@ public:
 				{
 					FString Category = FParse::Token(Cmd, 1);
 					FVisualLogger::Get().BlockAllCategories(true);
-					FVisualLogger::Get().GetWhiteList().AddUnique(*Category);
+					FVisualLogger::Get().AddCategortyToWhiteList(*Category);
 					return true;
 				}
 #if WITH_EDITOR
diff --git a/Engine/Source/Runtime/Engine/Private/VisualLogger/VisualLoggerTypes.cpp b/Engine/Source/Runtime/Engine/Private/VisualLogger/VisualLoggerTypes.cpp
--- a/Engine/Source/Runtime/Engine/Private/VisualLogger/VisualLoggerTypes.cpp
+++ b/Engine/Source/Runtime/Engine/Private/VisualLogger/VisualLoggerTypes.cpp
@@ -193,6 +193,47 @@ void FVisualLogEntry::AddElement(const FVector& Center, float HalfHeight, float
 	ElementsToDraw.Add(Element);
 }
 
+void FVisualLogEntry::AddElement(const TArray<FVector>& ConvexPoints, float MinZ, float MaxZ, const FName& CategoryName, ELogVerbosity::Type Verbosity, const FColor& Color, const FString& Description)
+{
+	FVisualLogShapeElement Element(Description, Color, 0, CategoryName);
+	Element.Points.Reserve(1 + ConvexPoints.Num());
+	Element.Points.Add(FVector(MinZ, MaxZ, 0));
+	Element.Points.Append(ConvexPoints);
+	Element.Type = EVisualLoggerShapeElement::NavAreaMesh;
+	Element.Verbosity = Verbosity;
+	ElementsToDraw.Add(Element);
+}
+
+void FVisualLogEntry::AddElement(const TArray<FVector>& Vertices, const TArray<int32>& Indices, const FName& CategoryName, ELogVerbosity::Type Verbosity, const FColor& Color, const FString& Description)
+{
+	FVisualLogShapeElement Element(Description, Color, 0, CategoryName);
+	uint32 FacesNum = Indices.Num() / 3;
+	Element.Points.Reserve(1 + Vertices.Num() + FacesNum);
+	Element.Points.Add(FVector(Vertices.Num(), FacesNum, 0)); //add header data
+	Element.Points.Append(Vertices);
+	TArray<FVector> Faces;
+	Faces.Reserve(FacesNum);
+	for (int32 i = 0; i < Indices.Num(); i += 3)
+	{
+		Faces.Add(FVector(Indices[i + 0], Indices[i + 1], Indices[i + 2]));
+	}
+	Element.Points.Append(Faces);
+
+	Element.Type = EVisualLoggerShapeElement::Mesh;
+	Element.Verbosity = Verbosity;
+	ElementsToDraw.Add(Element);
+}
+
+void FVisualLogEntry::AddConvexElement(const TArray<FVector>& Points, const FName& CategoryName, ELogVerbosity::Type Verbosity, const FColor& Color, const FString& Description)
+{
+	FVisualLogShapeElement Element(Description, Color, 0, CategoryName);
+	Element.Points = Points;
+	Element.Verbosity = Verbosity;
+	Element.Type = EVisualLoggerShapeElement::Polygon;
+	ElementsToDraw.Add(Element);
+}
+
+
 void FVisualLogEntry::AddHistogramData(const FVector2D& DataSample, const FName& CategoryName, ELogVerbosity::Type Verbosity, const FName& GraphName, const FName& DataName)
 {
 	FVisualLogHistogramSample Sample;
diff --git a/Engine/Source/Runtime/Engine/Public/AI/NavDataGenerator.h b/Engine/Source/Runtime/Engine/Public/AI/NavDataGenerator.h
--- a/Engine/Source/Runtime/Engine/Public/AI/NavDataGenerator.h
+++ b/Engine/Source/Runtime/Engine/Public/AI/NavDataGenerator.h
@@ -50,6 +50,7 @@ public:
 
 #if !(UE_BUILD_SHIPPING || UE_BUILD_TEST)
 	virtual void ExportNavigationData(const FString& FileName) const {}
+	virtual void GrabDebugSnapshot(struct FVisualLogEntry* Snapshot, const FBox& BoundingBox, const struct FLogCategoryBase& Category, ELogVerbosity::Type Verbosity) const {}
 #endif
 
 };
\ No newline at end of file
diff --git a/Engine/Source/Runtime/Engine/Public/AI/Navigation/RecastNavMeshGenerator.h b/Engine/Source/Runtime/Engine/Public/AI/Navigation/RecastNavMeshGenerator.h
--- a/Engine/Source/Runtime/Engine/Public/AI/Navigation/RecastNavMeshGenerator.h
+++ b/Engine/Source/Runtime/Engine/Public/AI/Navigation/RecastNavMeshGenerator.h
@@ -407,6 +407,7 @@ public:
 	
 #if !(UE_BUILD_SHIPPING || UE_BUILD_TEST)
 	virtual void ExportNavigationData(const FString& FileName) const override;
+	virtual void GrabDebugSnapshot(struct FVisualLogEntry* Snapshot, const FBox& BoundingBox, const struct FLogCategoryBase& Category, ELogVerbosity::Type Verbosity) const override;
 #endif
 
 	/** 
diff --git a/Engine/Source/Runtime/Engine/Public/DebugRenderSceneProxy.h b/Engine/Source/Runtime/Engine/Public/DebugRenderSceneProxy.h
--- a/Engine/Source/Runtime/Engine/Public/DebugRenderSceneProxy.h
+++ b/Engine/Source/Runtime/Engine/Public/DebugRenderSceneProxy.h
@@ -244,6 +244,8 @@ public:
 	FDelegateHandle DebugTextDrawingDelegateHandle;
 	EDrawType DrawType;
 	uint32 DrawAlpha;
+
+	TWeakObjectPtr<class UMaterial> SolidMeshMaterial;
 };
 
 #endif
diff --git a/Engine/Source/Runtime/Engine/Public/VisualLogger/VisualLogger.h b/Engine/Source/Runtime/Engine/Public/VisualLogger/VisualLogger.h
--- a/Engine/Source/Runtime/Engine/Public/VisualLogger/VisualLogger.h
+++ b/Engine/Source/Runtime/Engine/Public/VisualLogger/VisualLogger.h
@@ -45,6 +45,19 @@
 // Histogram data for 2d graphs 
 #define UE_VLOG_HISTOGRAM(LogOwner, CategoryName, Verbosity, GraphName, DataName, Data) FVisualLogger::HistogramDataLogf(LogOwner, CategoryName, ELogVerbosity::Verbosity, INDEX_NONE, GraphName, DataName, Data, FColor::White, TEXT(""))
 #define UE_CVLOG_HISTOGRAM(Condition, LogOwner, CategoryName, Verbosity, GraphName, DataName, Data) if(Condition) {UE_VLOG_HISTOGRAM(LogOwner, CategoryName, Verbosity, GraphName, DataName, Data);} 
+// NavArea or extruded convex shape
+#define UE_VLOG_NAVAREA(LogOwner, CategoryName, Verbosity, ConvexPoints, MinZ, MaxZ, Color, Format, ...) FVisualLogger::NavAreaShapeLogf(LogOwner, CategoryName, ELogVerbosity::Verbosity, INDEX_NONE, ConvexPoints, MinZ, MaxZ, Color, Format, ##__VA_ARGS__)
+#define UE_CVLOG_NAVAREA(Condition, LogOwner, CategoryName, Verbosity, ConvexPoints, MinZ, MaxZ, Color, Format, ...) if(Condition) {UE_VLOG_NAVAREA(LogOwner, CategoryName, Verbosity, ConvexPoints, MinZ, MaxZ, Color, Format, ##__VA_ARGS__);}
+// regular 3d mesh shape to log
+#define UE_VLOG_MESH(LogOwner, CategoryName, Verbosity, Vertices, Indices, Color, Format, ...) FVisualLogger::GeometryShapeLogf(LogOwner, CategoryName, ELogVerbosity::Verbosity, INDEX_NONE, Vertices, Indices, Format, ##__VA_ARGS__)
+#define UE_CVLOG_MESH(Condition, LogOwner, CategoryName, Verbosity, Vertices, Indexes, Color, Format, ...) if(Condition) {UE_VLOG_MESH(LogOwner, CategoryName, Verbosity, Vertices, Indices, Color, Format, ##__VA_ARGS__);}
+// 2d convex poly shape
+#define UE_VLOG_CONVEX(LogOwner, CategoryName, Verbosity, Points, Color, Format, ...)  FVisualLogger::GeometryConvexLogf(LogOwner, CategoryName, ELogVerbosity::Verbosity, INDEX_NONE, Points, Color, Format, ##__VA_ARGS__)
+#define UE_CVLOG_CONVEX(Condition, LogOwner, CategoryName, Verbosity, Points, Color, Format, ...)  if(Condition) {UE_VLOG_CONVEX(LogOwner, CategoryName, Verbosity, Points, Color, Format, ##__VA_ARGS__);}
+// Navigation Octree dump for given bounding box
+#define UE_VLOG_NAVOCTREE(LogOwner, CategoryName, Verbosity, Box, Color)  FVisualLogger::NavOctreeShapeLog(LogOwner, CategoryName, ELogVerbosity::Verbosity, INDEX_NONE, Box)
+#define UE_CVLOG_NAVOCTREE(Condition, LogOwner, CategoryName, Verbosity, Box, Color)  if(Condition) {UE_VLOG_NAVOCTREE(LogOwner, CategoryName, Verbosity, Box, Color, Format);}
+
 
 #define DECLARE_VLOG_EVENT(EventName) extern FVisualLogEventBase EventName;
 #define DEFINE_VLOG_EVENT(EventName, Verbosity, UserFriendlyDesc) FVisualLogEventBase EventName(TEXT(#EventName), TEXT(UserFriendlyDesc), ELogVerbosity::Verbosity); 
@@ -82,6 +95,14 @@
 #define UE_CVLOG_CAPSULE(Condition, Object, CategoryName, Verbosity, Center, HalfHeight, Radius, Rotation, Color, DescriptionFormat, ...)
 #define UE_VLOG_HISTOGRAM(Actor, CategoryName, Verbosity, GraphName, DataName, Data)
 #define UE_CVLOG_HISTOGRAM(Condition, Actor, CategoryName, Verbosity, GraphName, DataName, Data)
+#define UE_VLOG_NAVAREA(LogOwner, CategoryName, Verbosity, ConvexPoints, MinZ, MaxZ, Color, Format, ...)
+#define UE_CVLOG_NAVAREA(Condition, LogOwner, CategoryName, Verbosity, ConvexPoints, MinZ, MaxZ, Color, Format, ...)
+#define UE_VLOG_MESH(LogOwner, CategoryName, Verbosity, Vertices, Indexes, Color, Format, ...) 
+#define UE_CVLOG_MESH(Condition, LogOwner, CategoryName, Verbosity, Vertices, Indexes, Color, Format, ...) 
+#define UE_VLOG_CONVEX(LogOwner, CategoryName, Verbosity, Points, Color, Format, ...) 
+#define UE_CVLOG_CONVEX(Condition, LogOwner, CategoryName, Verbosity, Points, Color, Format, ...) 
+#define UE_VLOG_NAVOCTREE(LogOwner, CategoryName, Verbosity, Box, Color)
+#define UE_CVLOG_NAVOCTREE(Condition, LogOwner, CategoryName, Verbosity, Box, Color)
 
 #define DECLARE_VLOG_EVENT(EventName)
 #define DEFINE_VLOG_EVENT(EventName, Verbosity, UserFriendlyDesc)
@@ -126,8 +147,33 @@ public:
 	VARARG_DECL(static void, static void, return, GeometryShapeLogf, VARARG_NONE, const TCHAR*, VARARG_EXTRA(const class UObject* LogOwner) VARARG_EXTRA(const struct FLogCategoryBase& Category) VARARG_EXTRA(ELogVerbosity::Type Verbosity) VARARG_EXTRA(int32 UniqueLogId) VARARG_EXTRA(const FVector& Start) VARARG_EXTRA(const FVector& End) VARARG_EXTRA(const float Radius) VARARG_EXTRA(const FColor& Color), VARARG_EXTRA(LogOwner) VARARG_EXTRA(Category) VARARG_EXTRA(Verbosity) VARARG_EXTRA(UniqueLogId) VARARG_EXTRA(Start) VARARG_EXTRA(End) VARARG_EXTRA(Radius) VARARG_EXTRA(Color));
 	// Capsule log
 	VARARG_DECL(static void, static void, return, GeometryShapeLogf, VARARG_NONE, const TCHAR*, VARARG_EXTRA(const class UObject* LogOwner) VARARG_EXTRA(const struct FLogCategoryBase& Category) VARARG_EXTRA(ELogVerbosity::Type Verbosity) VARARG_EXTRA(int32 UniqueLogId) VARARG_EXTRA(const FVector& Center) VARARG_EXTRA(float HalfHeight) VARARG_EXTRA(float Radius) VARARG_EXTRA(const FQuat & Rotation) VARARG_EXTRA(const FColor& Color), VARARG_EXTRA(LogOwner) VARARG_EXTRA(Category) VARARG_EXTRA(Verbosity) VARARG_EXTRA(UniqueLogId) VARARG_EXTRA(Center) VARARG_EXTRA(HalfHeight) VARARG_EXTRA(Radius) VARARG_EXTRA(Rotation) VARARG_EXTRA(Color));
+	// NavArea/Extruded convex log
+	VARARG_DECL(static void, static void, return, NavAreaShapeLogf, VARARG_NONE, const TCHAR*, 
+		VARARG_EXTRA(const class UObject* LogOwner) 
+		VARARG_EXTRA(const struct FLogCategoryBase& Category) 
+		VARARG_EXTRA(ELogVerbosity::Type Verbosity) 
+		VARARG_EXTRA(int32 UniqueLogId) 
+		VARARG_EXTRA(const TArray<FVector> &ConvexPoints) 
+		VARARG_EXTRA(float MinZ) 
+		VARARG_EXTRA(float MaxZ) 
+		VARARG_EXTRA(const FColor& Color), 
+		VARARG_EXTRA(LogOwner) VARARG_EXTRA(Category) VARARG_EXTRA(Verbosity) VARARG_EXTRA(UniqueLogId) VARARG_EXTRA(ConvexPoints) VARARG_EXTRA(MinZ) VARARG_EXTRA(MaxZ) VARARG_EXTRA(Color));
+	// 3d Mesh log
+	VARARG_DECL(static void, static void, return, GeometryShapeLogf, VARARG_NONE, const TCHAR*, VARARG_EXTRA(const class UObject* LogOwner) VARARG_EXTRA(const struct FLogCategoryBase& Category) VARARG_EXTRA(ELogVerbosity::Type Verbosity) VARARG_EXTRA(int32 UniqueLogId) VARARG_EXTRA(const TArray<FVector> &Vertices) VARARG_EXTRA(const TArray<int32> &Indices) VARARG_EXTRA(const FColor& Color), VARARG_EXTRA(LogOwner) VARARG_EXTRA(Category) VARARG_EXTRA(Verbosity) VARARG_EXTRA(UniqueLogId) VARARG_EXTRA(Vertices) VARARG_EXTRA(Indices) VARARG_EXTRA(Color));
+	// 2d Convex shape
+	VARARG_DECL(static void, static void, return, GeometryConvexLogf, VARARG_NONE, const TCHAR*, 
+		VARARG_EXTRA(const class UObject* LogOwner) 
+		VARARG_EXTRA(const struct FLogCategoryBase& Category) 
+		VARARG_EXTRA(ELogVerbosity::Type Verbosity) 
+		VARARG_EXTRA(int32 UniqueLogId) 
+		VARARG_EXTRA(const TArray<FVector> &Points) 
+		VARARG_EXTRA(const FColor& Color), 
+		VARARG_EXTRA(LogOwner) VARARG_EXTRA(Category) VARARG_EXTRA(Verbosity) VARARG_EXTRA(UniqueLogId) VARARG_EXTRA(Points) VARARG_EXTRA(Color));
+
 	//Histogram data
 	VARARG_DECL(static void, static void, return, HistogramDataLogf, VARARG_NONE, const TCHAR*, VARARG_EXTRA(const class UObject* LogOwner) VARARG_EXTRA(const struct FLogCategoryBase& Category) VARARG_EXTRA(ELogVerbosity::Type Verbosity) VARARG_EXTRA(int32 UniqueLogId) VARARG_EXTRA(FName GraphName) VARARG_EXTRA(FName DataName) VARARG_EXTRA(const FVector2D& Data) VARARG_EXTRA(const FColor& Color), VARARG_EXTRA(LogOwner) VARARG_EXTRA(Category) VARARG_EXTRA(Verbosity) VARARG_EXTRA(UniqueLogId) VARARG_EXTRA(GraphName) VARARG_EXTRA(DataName) VARARG_EXTRA(Data) VARARG_EXTRA(Color));
+	// Navigation octree data
+	static void NavOctreeShapeLog(const class UObject* LogOwner, const struct FLogCategoryBase& Category, ELogVerbosity::Type Verbosity, int32 UniqueLogId, const FBox& Box);
 
 	/** Log events */
 	static void EventLog(const class UObject* LogOwner, const FName EventTag1, const FVisualLogEventBase& Event1, const FName EventTag2 = NAME_None, const FName EventTag3 = NAME_None, const FName EventTag4 = NAME_None, const FName EventTag5 = NAME_None, const FName EventTag6 = NAME_None);
@@ -167,7 +213,7 @@ public:
 	bool IsBlockedForAllCategories() const { return !!bBlockedAllCategories; }
 
 	/** Returns white list for modifications */
-	TArray<class FName>& GetWhiteList() { return CategoriesWhiteList; }
+	const TArray<class FName>& GetWhiteList() const { return CategoriesWhiteList; }
 
 	bool IsWhiteListed(const FName& Name) const { return CategoriesWhiteList.Find(Name) != INDEX_NONE; }
 
@@ -181,22 +227,22 @@ public:
 	/** Starts visual log collecting and recording */
 	void SetIsRecording(bool InIsRecording);
 	/** return information is vlog recording is enabled or not */
-	bool IsRecording() { return !!bIsRecording; }
+	bool IsRecording() const { return !!bIsRecording; }
 
 	/** Starts visual log collecting and recording */
 	void SetIsRecordingToFile(bool InIsRecording);
 	/** return information is vlog recording is enabled or not */
-	bool IsRecordingToFile() { return !!bIsRecordingToFile; }
+	bool IsRecordingToFile() const { return !!bIsRecordingToFile; }
 
 	void SetIsRecordingOnServer(bool IsRecording) { bIsRecordingOnServer = IsRecording; }
-	bool IsRecordingOnServer() { return !!bIsRecordingOnServer; }
+	bool IsRecordingOnServer() const { return !!bIsRecordingOnServer; }
 
 	/** Add visual logger output device */
 	void AddDevice(FVisualLogDevice* InDevice) { OutputDevices.AddUnique(InDevice); }
 	/** Remove visual logger output device */
 	void RemoveDevice(FVisualLogDevice* InDevice) { OutputDevices.RemoveSwap(InDevice); }
 	/** Remove visual logger output device */
-	TArray<FVisualLogDevice*>& GetDevices() { return OutputDevices; }
+	const TArray<FVisualLogDevice*>& GetDevices() const { return OutputDevices; }
 	/** Returns  current entry for given TimeStap or creates another one  but first it serialize previous entry as completed to vislog devices. Use VisualLogger::DontCreate to get current entry without serialization*/
 	FVisualLogEntry* GetEntryToWrite(const class UObject* Object, float TimeStamp, ECreateIfNeeded ShouldCreate = ECreateIfNeeded::Create);
 	/** Retrieves last used entry for given UObject*/
@@ -212,9 +258,9 @@ public:
 	/**  Removes previously registered extension */
 	void UnregisterExtension(FName TagName, FVisualLogExtensionInterface* ExtensionInterface) { AllExtensions.Remove(TagName); }
 	/** returns extension identified by given tag */
-	FVisualLogExtensionInterface* GetExtensionForTag(const FName TagName) { return AllExtensions.Contains(TagName) ? AllExtensions[TagName] : nullptr; }
+	FVisualLogExtensionInterface* GetExtensionForTag(const FName TagName) const { return AllExtensions.Contains(TagName) ? AllExtensions[TagName] : nullptr; }
 	/** Returns reference to map with all registered extension */
-	TMap<FName, FVisualLogExtensionInterface*>& GetAllExtensions() { return AllExtensions; }
+	const TMap<FName, FVisualLogExtensionInterface*>& GetAllExtensions() const { return AllExtensions; }
 
 	/** internal check for each usage of visual logger */
 	static bool CheckVisualLogInputInternal(const class UObject* Object, const struct FLogCategoryBase& Category, ELogVerbosity::Type Verbosity, UWorld **World, FVisualLogEntry **CurrentEntry);
@@ -402,6 +448,55 @@ VARARG_BODY(void, FVisualLogger::GeometryShapeLogf, const TCHAR*, VARARG_EXTRA(c
 	);
 }
 
+inline
+VARARG_BODY(void, FVisualLogger::NavAreaShapeLogf, const TCHAR*, VARARG_EXTRA(const class UObject* Object) VARARG_EXTRA(const struct FLogCategoryBase& Category) VARARG_EXTRA(ELogVerbosity::Type Verbosity) VARARG_EXTRA(int32 UniqueLogId) VARARG_EXTRA(const TArray<FVector> &ConvexPoints) VARARG_EXTRA(float MinZ) VARARG_EXTRA(float MaxZ) VARARG_EXTRA(const FColor& Color))
+{
+	SCOPE_CYCLE_COUNTER(STAT_VisualLog);
+	UWorld *World = nullptr;
+	FVisualLogEntry *CurrentEntry = nullptr;
+	if (CheckVisualLogInputInternal(Object, Category, Verbosity, &World, &CurrentEntry) == false)
+	{
+		return;
+	}
+
+	COLLAPSED_LOGF(
+		CurrentEntry->AddElement(ConvexPoints, MinZ, MaxZ, Category.GetCategoryName(), Verbosity, Color, Buffer);
+	);
+}
+
+inline
+VARARG_BODY(void, FVisualLogger::GeometryShapeLogf, const TCHAR*, VARARG_EXTRA(const class UObject* Object) VARARG_EXTRA(const struct FLogCategoryBase& Category) VARARG_EXTRA(ELogVerbosity::Type Verbosity) VARARG_EXTRA(int32 UniqueLogId) VARARG_EXTRA(const TArray<FVector> &Vertices) VARARG_EXTRA(const TArray<int32> &Indices) VARARG_EXTRA(const FColor& Color))
+{
+	SCOPE_CYCLE_COUNTER(STAT_VisualLog);
+	UWorld *World = nullptr;
+	FVisualLogEntry *CurrentEntry = nullptr;
+	if (CheckVisualLogInputInternal(Object, Category, Verbosity, &World, &CurrentEntry) == false)
+	{
+		return;
+	}
+
+	COLLAPSED_LOGF(
+		CurrentEntry->AddElement(Vertices, Indices, Category.GetCategoryName(), Verbosity, Color, Buffer);
+	);
+}
+
+inline
+VARARG_BODY(void, FVisualLogger::GeometryConvexLogf, const TCHAR*, VARARG_EXTRA(const class UObject* Object) VARARG_EXTRA(const struct FLogCategoryBase& Category) VARARG_EXTRA(ELogVerbosity::Type Verbosity) VARARG_EXTRA(int32 UniqueLogId) VARARG_EXTRA(const TArray<FVector> &Points) VARARG_EXTRA(const FColor& Color))
+{
+	SCOPE_CYCLE_COUNTER(STAT_VisualLog);
+	UWorld *World = nullptr;
+	FVisualLogEntry *CurrentEntry = nullptr;
+	if (CheckVisualLogInputInternal(Object, Category, Verbosity, &World, &CurrentEntry) == false)
+	{
+		return;
+	}
+
+	COLLAPSED_LOGF(
+		CurrentEntry->AddConvexElement(Points, Category.GetCategoryName(), Verbosity, Color, Buffer);
+	);
+}
+
+
 inline
 VARARG_BODY(void, FVisualLogger::HistogramDataLogf, const TCHAR*, VARARG_EXTRA(const class UObject* Object) VARARG_EXTRA(const struct FLogCategoryBase& Category) VARARG_EXTRA(ELogVerbosity::Type Verbosity) VARARG_EXTRA(int32 UniqueLogId) VARARG_EXTRA(FName GraphName) VARARG_EXTRA(FName DataName) VARARG_EXTRA(const FVector2D& Data) VARARG_EXTRA(const FColor& Color))
 {
diff --git a/Engine/Source/Runtime/Engine/Public/VisualLogger/VisualLoggerTypes.h b/Engine/Source/Runtime/Engine/Public/VisualLogger/VisualLoggerTypes.h
--- a/Engine/Source/Runtime/Engine/Public/VisualLogger/VisualLoggerTypes.h
+++ b/Engine/Source/Runtime/Engine/Public/VisualLogger/VisualLoggerTypes.h
@@ -52,6 +52,8 @@ enum class EVisualLoggerShapeElement : uint8
 	Cylinder,
 	Capsule,
 	Polygon,
+	Mesh,
+	NavAreaMesh, // convex based mesh with min and max Z values
 	// note that in order to remain backward compatibility in terms of log
 	// serialization new enum values need to be added at the end
 };
@@ -191,6 +193,12 @@ struct ENGINE_API FVisualLogEntry
 	void AddElement(const FVector& Center, float HalfHeight, float Radius, const FQuat & Rotation, const FName& CategoryName, ELogVerbosity::Type Verbosity, const FColor& Color = FColor::White, const FString& Description = TEXT(""));
 	// custom element
 	void AddElement(const FVisualLogShapeElement& Element);
+	// NavAreaMesh
+	void AddElement(const TArray<FVector>& ConvexPoints, float MinZ, float MaxZ, const FName& CategoryName, ELogVerbosity::Type Verbosity, const FColor& Color = FColor::White, const FString& Description = TEXT(""));
+	// 3d Mesh
+	void AddElement(const TArray<FVector>& Vertices, const TArray<int32>& Indices, const FName& CategoryName, ELogVerbosity::Type Verbosity, const FColor& Color = FColor::White, const FString& Description = TEXT(""));
+	// 2d convex
+	void AddConvexElement(const TArray<FVector>& Points, const FName& CategoryName, ELogVerbosity::Type Verbosity, const FColor& Color = FColor::White, const FString& Description = TEXT(""));
 	// histogram sample
 	void AddHistogramData(const FVector2D& DataSample, const FName& CategoryName, ELogVerbosity::Type Verbosity, const FName& GraphName, const FName& DataName);
 	// Custom data block
@@ -371,7 +379,9 @@ void FVisualLogShapeElement::SetType(EVisualLoggerShapeElement InType)
 inline
 FColor FVisualLogShapeElement::GetFColor() const
 {
-	return FColor(((Color & 0xc0) << 24) | ((Color & 0x30) << 18) | ((Color & 0x0c) << 12) | ((Color & 0x03) << 6));
+	FColor RetColor(((Color & 0xc0) << 24) | ((Color & 0x30) << 18) | ((Color & 0x0c) << 12) | ((Color & 0x03) << 6));
+	RetColor.A = (RetColor.A * 255) / 192; // convert alpha to 0-255 range
+	return RetColor;
 }
 
 
