Changed few names in Visual Logger to be more generic.

diff --git a/Engine/Source/Runtime/Engine/Private/VisualLogger/VisualLogger.cpp b/Engine/Source/Runtime/Engine/Private/VisualLogger/VisualLogger.cpp
--- a/Engine/Source/Runtime/Engine/Private/VisualLogger/VisualLogger.cpp
+++ b/Engine/Source/Runtime/Engine/Private/VisualLogger/VisualLogger.cpp
@@ -249,7 +249,7 @@ void FVisualLogger::EventLog(const class UObject* Object, const FName EventTag1,
 	CurrentEntry->Events[Index].EventTags.Remove(NAME_None);
 }
 
-void FVisualLogger::NavOctreeShapeLog(const class UObject* Object, const struct FLogCategoryBase& Category, ELogVerbosity::Type Verbosity, int32 UniqueLogId, const FBox& Box)
+void FVisualLogger::NavigationDataDump(const class UObject* Object, const struct FLogCategoryBase& Category, ELogVerbosity::Type Verbosity, int32 UniqueLogId, const FBox& Box)
 {
 	SCOPE_CYCLE_COUNTER(STAT_VisualLog);
 	UWorld *World = nullptr;
diff --git a/Engine/Source/Runtime/Engine/Public/VisualLogger/VisualLogger.h b/Engine/Source/Runtime/Engine/Public/VisualLogger/VisualLogger.h
--- a/Engine/Source/Runtime/Engine/Public/VisualLogger/VisualLogger.h
+++ b/Engine/Source/Runtime/Engine/Public/VisualLogger/VisualLogger.h
@@ -45,18 +45,15 @@
 // Histogram data for 2d graphs 
 #define UE_VLOG_HISTOGRAM(LogOwner, CategoryName, Verbosity, GraphName, DataName, Data) FVisualLogger::HistogramDataLogf(LogOwner, CategoryName, ELogVerbosity::Verbosity, INDEX_NONE, GraphName, DataName, Data, FColor::White, TEXT(""))
 #define UE_CVLOG_HISTOGRAM(Condition, LogOwner, CategoryName, Verbosity, GraphName, DataName, Data) if(Condition) {UE_VLOG_HISTOGRAM(LogOwner, CategoryName, Verbosity, GraphName, DataName, Data);} 
-// NavArea or extruded convex shape
-#define UE_VLOG_NAVAREA(LogOwner, CategoryName, Verbosity, ConvexPoints, MinZ, MaxZ, Color, Format, ...) FVisualLogger::NavAreaShapeLogf(LogOwner, CategoryName, ELogVerbosity::Verbosity, INDEX_NONE, ConvexPoints, MinZ, MaxZ, Color, Format, ##__VA_ARGS__)
-#define UE_CVLOG_NAVAREA(Condition, LogOwner, CategoryName, Verbosity, ConvexPoints, MinZ, MaxZ, Color, Format, ...) if(Condition) {UE_VLOG_NAVAREA(LogOwner, CategoryName, Verbosity, ConvexPoints, MinZ, MaxZ, Color, Format, ##__VA_ARGS__);}
+// NavArea or vertically pulled convex shape
+#define UE_VLOG_PULLEDCONVEX(LogOwner, CategoryName, Verbosity, ConvexPoints, MinZ, MaxZ, Color, Format, ...) FVisualLogger::NavAreaShapeLogf(LogOwner, CategoryName, ELogVerbosity::Verbosity, INDEX_NONE, ConvexPoints, MinZ, MaxZ, Color, Format, ##__VA_ARGS__)
+#define UE_CVLOG_PULLEDCONVEX(Condition, LogOwner, CategoryName, Verbosity, ConvexPoints, MinZ, MaxZ, Color, Format, ...) if(Condition) {UE_VLOG_PULLEDCONVEX(LogOwner, CategoryName, Verbosity, ConvexPoints, MinZ, MaxZ, Color, Format, ##__VA_ARGS__);}
 // regular 3d mesh shape to log
 #define UE_VLOG_MESH(LogOwner, CategoryName, Verbosity, Vertices, Indices, Color, Format, ...) FVisualLogger::GeometryShapeLogf(LogOwner, CategoryName, ELogVerbosity::Verbosity, INDEX_NONE, Vertices, Indices, Format, ##__VA_ARGS__)
 #define UE_CVLOG_MESH(Condition, LogOwner, CategoryName, Verbosity, Vertices, Indexes, Color, Format, ...) if(Condition) {UE_VLOG_MESH(LogOwner, CategoryName, Verbosity, Vertices, Indices, Color, Format, ##__VA_ARGS__);}
 // 2d convex poly shape
-#define UE_VLOG_CONVEX(LogOwner, CategoryName, Verbosity, Points, Color, Format, ...)  FVisualLogger::GeometryConvexLogf(LogOwner, CategoryName, ELogVerbosity::Verbosity, INDEX_NONE, Points, Color, Format, ##__VA_ARGS__)
-#define UE_CVLOG_CONVEX(Condition, LogOwner, CategoryName, Verbosity, Points, Color, Format, ...)  if(Condition) {UE_VLOG_CONVEX(LogOwner, CategoryName, Verbosity, Points, Color, Format, ##__VA_ARGS__);}
-// Navigation Octree dump for given bounding box
-#define UE_VLOG_NAVOCTREE(LogOwner, CategoryName, Verbosity, Box, Color)  FVisualLogger::NavOctreeShapeLog(LogOwner, CategoryName, ELogVerbosity::Verbosity, INDEX_NONE, Box)
-#define UE_CVLOG_NAVOCTREE(Condition, LogOwner, CategoryName, Verbosity, Box, Color)  if(Condition) {UE_VLOG_NAVOCTREE(LogOwner, CategoryName, Verbosity, Box, Color, Format);}
+#define UE_VLOG_CONVEXPOLY(LogOwner, CategoryName, Verbosity, Points, Color, Format, ...)  FVisualLogger::GeometryConvexLogf(LogOwner, CategoryName, ELogVerbosity::Verbosity, INDEX_NONE, Points, Color, Format, ##__VA_ARGS__)
+#define UE_CVLOG_CONVEXPOLY(Condition, LogOwner, CategoryName, Verbosity, Points, Color, Format, ...)  if(Condition) {UE_VLOG_CONVEXPOLY(LogOwner, CategoryName, Verbosity, Points, Color, Format, ##__VA_ARGS__);}
 
 
 #define DECLARE_VLOG_EVENT(EventName) extern FVisualLogEventBase EventName;
@@ -67,6 +64,8 @@
 #define UE_VLOG_EVENT_WITH_DATA(LogOwner, LogEvent, ...) FVisualLogger::EventLog(LogOwner, LogEvent, ##__VA_ARGS__)
 #define UE_CVLOG_EVENT_WITH_DATA(Condition, LogOwner, LogEvent, ...) if(Condition) {UE_VLOG_EVENT_WITH_DATA(LogOwner, LogEvent, ##__VA_ARGS__);}
 
+#define UE_IFVLOG(__code_block__) __code_block__;
+
 #else
 #define REDIRECT_TO_VLOG(Dest)
 #define REDIRECT_OBJECT_TO_VLOG(Src, Dest)
@@ -95,14 +94,12 @@
 #define UE_CVLOG_CAPSULE(Condition, Object, CategoryName, Verbosity, Center, HalfHeight, Radius, Rotation, Color, DescriptionFormat, ...)
 #define UE_VLOG_HISTOGRAM(Actor, CategoryName, Verbosity, GraphName, DataName, Data)
 #define UE_CVLOG_HISTOGRAM(Condition, Actor, CategoryName, Verbosity, GraphName, DataName, Data)
-#define UE_VLOG_NAVAREA(LogOwner, CategoryName, Verbosity, ConvexPoints, MinZ, MaxZ, Color, Format, ...)
-#define UE_CVLOG_NAVAREA(Condition, LogOwner, CategoryName, Verbosity, ConvexPoints, MinZ, MaxZ, Color, Format, ...)
+#define UE_VLOG_PULLEDCONVEX(LogOwner, CategoryName, Verbosity, ConvexPoints, MinZ, MaxZ, Color, Format, ...)
+#define UE_CVLOG_PULLEDCONVEX(Condition, LogOwner, CategoryName, Verbosity, ConvexPoints, MinZ, MaxZ, Color, Format, ...)
 #define UE_VLOG_MESH(LogOwner, CategoryName, Verbosity, Vertices, Indexes, Color, Format, ...) 
 #define UE_CVLOG_MESH(Condition, LogOwner, CategoryName, Verbosity, Vertices, Indexes, Color, Format, ...) 
-#define UE_VLOG_CONVEX(LogOwner, CategoryName, Verbosity, Points, Color, Format, ...) 
-#define UE_CVLOG_CONVEX(Condition, LogOwner, CategoryName, Verbosity, Points, Color, Format, ...) 
-#define UE_VLOG_NAVOCTREE(LogOwner, CategoryName, Verbosity, Box, Color)
-#define UE_CVLOG_NAVOCTREE(Condition, LogOwner, CategoryName, Verbosity, Box, Color)
+#define UE_VLOG_CONVEXPOLY(LogOwner, CategoryName, Verbosity, Points, Color, Format, ...) 
+#define UE_CVLOG_CONVEXPOLY(Condition, LogOwner, CategoryName, Verbosity, Points, Color, Format, ...) 
 
 #define DECLARE_VLOG_EVENT(EventName)
 #define DEFINE_VLOG_EVENT(EventName, Verbosity, UserFriendlyDesc)
@@ -111,6 +108,7 @@
 #define UE_VLOG_EVENT_WITH_DATA(LogOwner, LogEvent, ...)
 #define UE_CVLOG_EVENT_WITH_DATA(Condition, LogOwner, LogEvent, ...)
 
+#define UE_IFVLOG(__code_block__)
 
 #endif //ENABLE_VISUAL_LOG
 
@@ -148,32 +146,15 @@ public:
 	// Capsule log
 	VARARG_DECL(static void, static void, return, GeometryShapeLogf, VARARG_NONE, const TCHAR*, VARARG_EXTRA(const class UObject* LogOwner) VARARG_EXTRA(const struct FLogCategoryBase& Category) VARARG_EXTRA(ELogVerbosity::Type Verbosity) VARARG_EXTRA(int32 UniqueLogId) VARARG_EXTRA(const FVector& Center) VARARG_EXTRA(float HalfHeight) VARARG_EXTRA(float Radius) VARARG_EXTRA(const FQuat & Rotation) VARARG_EXTRA(const FColor& Color), VARARG_EXTRA(LogOwner) VARARG_EXTRA(Category) VARARG_EXTRA(Verbosity) VARARG_EXTRA(UniqueLogId) VARARG_EXTRA(Center) VARARG_EXTRA(HalfHeight) VARARG_EXTRA(Radius) VARARG_EXTRA(Rotation) VARARG_EXTRA(Color));
 	// NavArea/Extruded convex log
-	VARARG_DECL(static void, static void, return, NavAreaShapeLogf, VARARG_NONE, const TCHAR*, 
-		VARARG_EXTRA(const class UObject* LogOwner) 
-		VARARG_EXTRA(const struct FLogCategoryBase& Category) 
-		VARARG_EXTRA(ELogVerbosity::Type Verbosity) 
-		VARARG_EXTRA(int32 UniqueLogId) 
-		VARARG_EXTRA(const TArray<FVector> &ConvexPoints) 
-		VARARG_EXTRA(float MinZ) 
-		VARARG_EXTRA(float MaxZ) 
-		VARARG_EXTRA(const FColor& Color), 
-		VARARG_EXTRA(LogOwner) VARARG_EXTRA(Category) VARARG_EXTRA(Verbosity) VARARG_EXTRA(UniqueLogId) VARARG_EXTRA(ConvexPoints) VARARG_EXTRA(MinZ) VARARG_EXTRA(MaxZ) VARARG_EXTRA(Color));
+	VARARG_DECL(static void, static void, return, NavAreaShapeLogf, VARARG_NONE, const TCHAR*, VARARG_EXTRA(const class UObject* LogOwner) VARARG_EXTRA(const struct FLogCategoryBase& Category) VARARG_EXTRA(ELogVerbosity::Type Verbosity) VARARG_EXTRA(int32 UniqueLogId) VARARG_EXTRA(const TArray<FVector> &ConvexPoints) VARARG_EXTRA(float MinZ) VARARG_EXTRA(float MaxZ) VARARG_EXTRA(const FColor& Color), VARARG_EXTRA(LogOwner) VARARG_EXTRA(Category) VARARG_EXTRA(Verbosity) VARARG_EXTRA(UniqueLogId) VARARG_EXTRA(ConvexPoints) VARARG_EXTRA(MinZ) VARARG_EXTRA(MaxZ) VARARG_EXTRA(Color));
 	// 3d Mesh log
 	VARARG_DECL(static void, static void, return, GeometryShapeLogf, VARARG_NONE, const TCHAR*, VARARG_EXTRA(const class UObject* LogOwner) VARARG_EXTRA(const struct FLogCategoryBase& Category) VARARG_EXTRA(ELogVerbosity::Type Verbosity) VARARG_EXTRA(int32 UniqueLogId) VARARG_EXTRA(const TArray<FVector> &Vertices) VARARG_EXTRA(const TArray<int32> &Indices) VARARG_EXTRA(const FColor& Color), VARARG_EXTRA(LogOwner) VARARG_EXTRA(Category) VARARG_EXTRA(Verbosity) VARARG_EXTRA(UniqueLogId) VARARG_EXTRA(Vertices) VARARG_EXTRA(Indices) VARARG_EXTRA(Color));
 	// 2d Convex shape
-	VARARG_DECL(static void, static void, return, GeometryConvexLogf, VARARG_NONE, const TCHAR*, 
-		VARARG_EXTRA(const class UObject* LogOwner) 
-		VARARG_EXTRA(const struct FLogCategoryBase& Category) 
-		VARARG_EXTRA(ELogVerbosity::Type Verbosity) 
-		VARARG_EXTRA(int32 UniqueLogId) 
-		VARARG_EXTRA(const TArray<FVector> &Points) 
-		VARARG_EXTRA(const FColor& Color), 
-		VARARG_EXTRA(LogOwner) VARARG_EXTRA(Category) VARARG_EXTRA(Verbosity) VARARG_EXTRA(UniqueLogId) VARARG_EXTRA(Points) VARARG_EXTRA(Color));
-
+	VARARG_DECL(static void, static void, return, GeometryConvexLogf, VARARG_NONE, const TCHAR*, VARARG_EXTRA(const class UObject* LogOwner) VARARG_EXTRA(const struct FLogCategoryBase& Category) VARARG_EXTRA(ELogVerbosity::Type Verbosity) VARARG_EXTRA(int32 UniqueLogId) VARARG_EXTRA(const TArray<FVector> &Points) VARARG_EXTRA(const FColor& Color), VARARG_EXTRA(LogOwner) VARARG_EXTRA(Category) VARARG_EXTRA(Verbosity) VARARG_EXTRA(UniqueLogId) VARARG_EXTRA(Points) VARARG_EXTRA(Color));
 	//Histogram data
 	VARARG_DECL(static void, static void, return, HistogramDataLogf, VARARG_NONE, const TCHAR*, VARARG_EXTRA(const class UObject* LogOwner) VARARG_EXTRA(const struct FLogCategoryBase& Category) VARARG_EXTRA(ELogVerbosity::Type Verbosity) VARARG_EXTRA(int32 UniqueLogId) VARARG_EXTRA(FName GraphName) VARARG_EXTRA(FName DataName) VARARG_EXTRA(const FVector2D& Data) VARARG_EXTRA(const FColor& Color), VARARG_EXTRA(LogOwner) VARARG_EXTRA(Category) VARARG_EXTRA(Verbosity) VARARG_EXTRA(UniqueLogId) VARARG_EXTRA(GraphName) VARARG_EXTRA(DataName) VARARG_EXTRA(Data) VARARG_EXTRA(Color));
-	// Navigation octree data
-	static void NavOctreeShapeLog(const class UObject* LogOwner, const struct FLogCategoryBase& Category, ELogVerbosity::Type Verbosity, int32 UniqueLogId, const FBox& Box);
+	// Navigation data debug snapshot
+	static void NavigationDataDump(const class UObject* LogOwner, const struct FLogCategoryBase& Category, ELogVerbosity::Type Verbosity, int32 UniqueLogId, const FBox& Box);
 
 	/** Log events */
 	static void EventLog(const class UObject* LogOwner, const FName EventTag1, const FVisualLogEventBase& Event1, const FName EventTag2 = NAME_None, const FName EventTag3 = NAME_None, const FName EventTag4 = NAME_None, const FName EventTag5 = NAME_None, const FName EventTag6 = NAME_None);
