UE4 - fix cooking crash with atmosphere rendering

diff --git a/Engine/Source/Runtime/Engine/Private/Atmosphere/Atmosphere.cpp b/Engine/Source/Runtime/Engine/Private/Atmosphere/Atmosphere.cpp
--- a/Engine/Source/Runtime/Engine/Private/Atmosphere/Atmosphere.cpp
+++ b/Engine/Source/Runtime/Engine/Private/Atmosphere/Atmosphere.cpp
@@ -552,7 +552,7 @@ void UAtmosphericFogComponent::PostInterpChange(UProperty* PropertyThatChanged)
 void UAtmosphericFogComponent::StartPrecompute()
 {
 #if WITH_EDITOR
-	if (GIsEditor && !IsTemplate())
+	if (GIsEditor && !IsTemplate() && !GUsingNullRHI)
 	{
 		FSceneInterface* AtmosphericFogScene = GetScene();
 		if (AtmosphericFogScene)
@@ -572,7 +572,7 @@ void UAtmosphericFogComponent::StartPrecompute()
 			{
 				FAtmosphericFogSceneInfo* AtmosphericFogSceneInfo = AtmosphericFogScene->GetAtmosphericFogSceneInfo();
 
-				if (AtmosphericFogSceneInfo->Component == Component)
+				if (AtmosphericFogSceneInfo && AtmosphericFogSceneInfo->Component == Component)
 				{
 					AtmosphericFogSceneInfo->bNeedRecompute = true;
 				}
