Fix for crash in blendspaces when referenced anim sequences have been deleted.

diff --git a/Engine/Source/Editor/Persona/Private/SAnimationBlendSpaceBase.cpp b/Engine/Source/Editor/Persona/Private/SAnimationBlendSpaceBase.cpp
--- a/Engine/Source/Editor/Persona/Private/SAnimationBlendSpaceBase.cpp
+++ b/Engine/Source/Editor/Persona/Private/SAnimationBlendSpaceBase.cpp
@@ -893,6 +893,7 @@ void SBlendSpaceWidget::HandleSampleDrop(const FGeometry& MyGeometry, const FVec
 
 void SBlendSpaceWidget::ResampleData()
 {
+	BlendSpace->ValidateSampleData();
 	CachedSamples.Empty();
 	BlendSpace->GetBlendSamplePoints(CachedSamples);
 	ValidateSamplePositions();
diff --git a/Engine/Source/Runtime/Engine/Classes/Animation/BlendSpaceBase.h b/Engine/Source/Runtime/Engine/Classes/Animation/BlendSpaceBase.h
--- a/Engine/Source/Runtime/Engine/Classes/Animation/BlendSpaceBase.h
+++ b/Engine/Source/Runtime/Engine/Classes/Animation/BlendSpaceBase.h
@@ -355,6 +355,9 @@ public:
 	/** return true if all sample data is additive **/
 	virtual bool IsValidAdditive() const {check(false); return false;}
 
+	/** 1) Remove data if redundant 2) Remove data if animation isn't found **/
+	ENGINE_API void ValidateSampleData();
+
 protected:
 	/** Initialize Per Bone Blend **/
 	void InitializePerBoneBlend();
@@ -368,9 +371,6 @@ protected:
 	/** Validates supplied blend sample against current contents of blendspace */
 	virtual bool ValidateSampleInput(FBlendSample & BlendSample, int32 OriginalIndex=INDEX_NONE) const;
 
-	/** 1) Remove data if redundant 2) Remove data if animation isn't found **/
-	void ValidateSampleData();
-
 	/** Returns where blend space sample animations are valid for the supplied AdditiveType */
 	bool IsValidAdditiveInternal(EAdditiveAnimationType AdditiveType) const;
 
diff --git a/Engine/Source/Runtime/Engine/Private/Animation/BlendSpaceBase.cpp b/Engine/Source/Runtime/Engine/Private/Animation/BlendSpaceBase.cpp
--- a/Engine/Source/Runtime/Engine/Private/Animation/BlendSpaceBase.cpp
+++ b/Engine/Source/Runtime/Engine/Private/Animation/BlendSpaceBase.cpp
@@ -644,7 +644,7 @@ bool UBlendSpaceBase::ValidateSampleInput(FBlendSample & BlendSample, int32 Orig
 
 void UBlendSpaceBase::ValidateSampleData()
 {
-	bool bMarkPackageDirty=false;
+	bool bSampleDataChanged=false;
 	AnimLength = 0.f;
 
 	for (int32 I=0; I<SampleData.Num(); ++I)
@@ -654,7 +654,7 @@ void UBlendSpaceBase::ValidateSampleData()
 			SampleData.RemoveAt(I);
 			--I;
 
-			bMarkPackageDirty = true;
+			bSampleDataChanged = true;
 			continue;
 		}
 
@@ -672,7 +672,7 @@ void UBlendSpaceBase::ValidateSampleData()
 				SampleData.RemoveAt(J);
 				--J;
 
-				bMarkPackageDirty = true;
+				bSampleDataChanged = true;
 			}
 		}
 
@@ -683,8 +683,9 @@ void UBlendSpaceBase::ValidateSampleData()
 		}
 	}
 
-	if (bMarkPackageDirty)
+	if (bSampleDataChanged)
 	{
+		GridSamples.Empty();
 		MarkPackageDirty();
 	}
 }
