Properly offsetting burst spawned particles based on initial velocity.

diff --git a/Engine/Source/Runtime/Engine/Private/Particles/ParticleEmitterInstances.cpp b/Engine/Source/Runtime/Engine/Private/Particles/ParticleEmitterInstances.cpp
--- a/Engine/Source/Runtime/Engine/Private/Particles/ParticleEmitterInstances.cpp
+++ b/Engine/Source/Runtime/Engine/Private/Particles/ParticleEmitterInstances.cpp
@@ -1791,7 +1791,7 @@ float FParticleEmitterInstance::Spawn(float DeltaTime)
 		float	NewLeftover = OldLeftover + DeltaTime * SpawnRate;
 		int32		Number		= FMath::FloorToInt(NewLeftover);
 		float	Increment	= (SpawnRate > 0.0f) ? (1.f / SpawnRate) : 0.0f;
-		float	StartTime	= DeltaTime + OldLeftover * Increment - Increment;
+		float	StartTime = DeltaTime + OldLeftover * Increment - Increment;
 		NewLeftover			= NewLeftover - Number;
 
 		// Handle growing arrays.
@@ -1807,6 +1807,10 @@ float FParticleEmitterInstance::Spawn(float DeltaTime)
 			NewCount = ActiveParticles + Number + BurstCount;
 		}
 
+		float	BurstIncrement = (BurstCount > 0.0f) ? (1.f / BurstCount) : 0.0f;
+		float	BurstStartTime = DeltaTime * BurstIncrement;
+
+
 		if (NewCount >= MaxActiveParticles)
 		{
 			if (DeltaTime < PeakActiveParticleUpdateDelta)
@@ -1837,7 +1841,7 @@ float FParticleEmitterInstance::Spawn(float DeltaTime)
 			SpawnParticles( Number, StartTime, Increment, InitialLocation, FVector::ZeroVector, EventPayload );
 
 			// Burst particles.
-			SpawnParticles( BurstCount, StartTime, 0.0f, InitialLocation, FVector::ZeroVector, EventPayload );
+			SpawnParticles(BurstCount, BurstStartTime, BurstIncrement, InitialLocation, FVector::ZeroVector, EventPayload);
 
 			return NewLeftover;
 		}
