- DemoNetDriver for replays is destroyed when a seamless travel copies the world data to ensure it always gets destroyed.

Merging

//depot/UE4-Orion/Engine/Source/Runtime/Engine/Private/World.cpp

to //depot/UE4/Engine/Source/Runtime/Engine/Private/World.cpp

diff --git a/Engine/Source/Runtime/Engine/Private/World.cpp b/Engine/Source/Runtime/Engine/Private/World.cpp
--- a/Engine/Source/Runtime/Engine/Private/World.cpp
+++ b/Engine/Source/Runtime/Engine/Private/World.cpp
@@ -4537,6 +4537,8 @@ void FSeamlessTravelHandler::StartLoadingDestination()
 
 void FSeamlessTravelHandler::CopyWorldData()
 {
+	CurrentWorld->DestroyDemoNetDriver();
+
 	UNetDriver* const NetDriver = CurrentWorld->GetNetDriver();
 	LoadedWorld->SetNetDriver(NetDriver);
 	if (NetDriver != NULL)
