Make physx shared data singleton shutdown explicitly

diff --git a/Engine/Source/Runtime/Engine/Private/PhysicsEngine/PhysLevel.cpp b/Engine/Source/Runtime/Engine/Private/PhysicsEngine/PhysLevel.cpp
--- a/Engine/Source/Runtime/Engine/Private/PhysicsEngine/PhysLevel.cpp
+++ b/Engine/Source/Runtime/Engine/Private/PhysicsEngine/PhysLevel.cpp
@@ -347,6 +347,8 @@ void InitGamePhys()
 	GPhysXSDK = PxCreatePhysics(PX_PHYSICS_VERSION, *GPhysXFoundation, PScale, false, GPhysXProfileZoneManager);
 	check(GPhysXSDK);
 
+	FPhysxSharedData::Initialize();
+
 	GPhysCommandHandler = new FPhysCommandHandler();
 
 	FCoreUObjectDelegates::PreGarbageCollect.AddRaw(GPhysCommandHandler, &FPhysCommandHandler::Flush);
@@ -464,6 +466,8 @@ void TermGamePhys()
 #endif
 
 #if WITH_PHYSX
+	FPhysxSharedData::Terminate();
+
 	// Do nothing if they were never initialized
 	if(GPhysXFoundation == NULL)
 	{
diff --git a/Engine/Source/Runtime/Engine/Private/PhysicsEngine/PhysXSupport.cpp b/Engine/Source/Runtime/Engine/Private/PhysicsEngine/PhysXSupport.cpp
--- a/Engine/Source/Runtime/Engine/Private/PhysicsEngine/PhysXSupport.cpp
+++ b/Engine/Source/Runtime/Engine/Private/PhysicsEngine/PhysXSupport.cpp
@@ -786,11 +786,21 @@ physx::PxPairFlags FApexPhysX3Interface::getContactReportFlags(const physx::PxSh
 
 #endif	// WITH_APEX
 
+FPhysxSharedData* FPhysxSharedData::Singleton = nullptr;
 
-FPhysxSharedData& FPhysxSharedData::Get()
+void FPhysxSharedData::Initialize()
 {
-	static FPhysxSharedData s_SharedData; 
-	return s_SharedData;
+	check(Singleton == nullptr);
+	Singleton = new FPhysxSharedData();
+}
+
+void FPhysxSharedData::Terminate()
+{
+	if (Singleton)
+	{
+		delete Singleton;
+		Singleton = nullptr;
+	}
 }
 
 void FPhysxSharedData::Add( PxBase* Obj )
