missing header

diff --git a/Engine/Source/Runtime/Engine/Private/PhysicsEngine/PhysXSupport.h b/Engine/Source/Runtime/Engine/Private/PhysicsEngine/PhysXSupport.h
--- a/Engine/Source/Runtime/Engine/Private/PhysicsEngine/PhysXSupport.h
+++ b/Engine/Source/Runtime/Engine/Private/PhysicsEngine/PhysXSupport.h
@@ -412,17 +412,9 @@ extern TArray<PxMaterial*>		GPhysXPendingKillMaterial;
 class FPhysxSharedData
 {
 public:
-	static FPhysxSharedData& Get();
-
-	FPhysxSharedData()
-	{
-		SharedObjects = PxCreateCollection();
-	}
-
-	~FPhysxSharedData()
-	{
-		SharedObjects->release();
-	}
+	static FPhysxSharedData& Get(){ return *Singleton; }
+	static void FPhysxSharedData::Initialize();
+	static void FPhysxSharedData::Terminate();
 
 	void Add(PxBase* Obj);
 	void Remove(PxBase* Obj)	{ if(Obj) { SharedObjects->remove(*Obj); } }
@@ -433,6 +425,18 @@ public:
 private:
 	/** Collection of shared physx objects */
 	PxCollection* SharedObjects;
+	
+	static FPhysxSharedData* Singleton;
+
+	FPhysxSharedData()
+	{
+		SharedObjects = PxCreateCollection();
+	}
+
+	~FPhysxSharedData()
+	{
+		SharedObjects->release();
+	}
 
 };
 
