Fix crash in blend node with lodded pose

diff --git a/Engine/Source/Runtime/Engine/Private/Animation/AnimNode_LayeredBoneBlend.cpp b/Engine/Source/Runtime/Engine/Private/Animation/AnimNode_LayeredBoneBlend.cpp
--- a/Engine/Source/Runtime/Engine/Private/Animation/AnimNode_LayeredBoneBlend.cpp
+++ b/Engine/Source/Runtime/Engine/Private/Animation/AnimNode_LayeredBoneBlend.cpp
@@ -31,8 +31,7 @@ void FAnimNode_LayeredBoneBlend::Initialize(const FAnimationInitializeContext& C
 
 void FAnimNode_LayeredBoneBlend::ReinitializeBoneBlendWeights(const FBoneContainer& RequiredBones, const USkeleton* Skeleton)
 {
-	const int32 NumBones = RequiredBones.GetNumBones();
-	FAnimationRuntime::CreateMaskWeights(NumBones, DesiredBoneBlendWeights, LayerSetup, RequiredBones, Skeleton);
+	FAnimationRuntime::CreateMaskWeights(DesiredBoneBlendWeights, LayerSetup, RequiredBones, Skeleton);
 
 	CurrentBoneBlendWeights.Empty(DesiredBoneBlendWeights.Num());
 	CurrentBoneBlendWeights.AddZeroed(DesiredBoneBlendWeights.Num());
@@ -45,6 +44,11 @@ void FAnimNode_LayeredBoneBlend::CacheBones(const FAnimationCacheBonesContext& C
 	{
 		BlendPoses[ChildIndex].CacheBones(Context);
 	}
+
+	if (Context.AnimInstance->RequiredBones.GetBoneIndicesArray().Num() != DesiredBoneBlendWeights.Num())
+	{
+		ReinitializeBoneBlendWeights(Context.AnimInstance->RequiredBones, Context.AnimInstance->CurrentSkeleton);
+	}
 }
 
 void FAnimNode_LayeredBoneBlend::Update(const FAnimationUpdateContext& Context)
@@ -61,13 +65,6 @@ void FAnimNode_LayeredBoneBlend::Update(const FAnimationUpdateContext& Context)
 		{
 			BlendPoses[ChildIndex].Update(Context);
 		}
-
-		if ( Context.AnimInstance->RequiredBones.GetNumBones() != DesiredBoneBlendWeights.Num() )
-		{
-			ReinitializeBoneBlendWeights(Context.AnimInstance->RequiredBones, Context.AnimInstance->CurrentSkeleton);
-		}
-
-		FAnimationRuntime::UpdateDesiredBoneWeight(DesiredBoneBlendWeights, CurrentBoneBlendWeights, BlendWeights, Context.AnimInstance->RequiredBones, Context.AnimInstance->CurrentSkeleton);
 	}
 }
 
@@ -82,6 +79,8 @@ void FAnimNode_LayeredBoneBlend::Evaluate(FPoseContext& Output)
 	}
 	else
 	{
+		FAnimationRuntime::UpdateDesiredBoneWeight(DesiredBoneBlendWeights, CurrentBoneBlendWeights, BlendWeights);
+
 		FPoseContext BasePoseContext(Output);
 
 		// evaluate children
diff --git a/Engine/Source/Runtime/Engine/Private/Animation/AnimationRuntime.cpp b/Engine/Source/Runtime/Engine/Private/Animation/AnimationRuntime.cpp
--- a/Engine/Source/Runtime/Engine/Private/Animation/AnimationRuntime.cpp
+++ b/Engine/Source/Runtime/Engine/Private/Animation/AnimationRuntime.cpp
@@ -932,18 +932,15 @@ void FAnimationRuntime::BlendLocalPosesPerBoneWeights(
 	}
 }
 
-void FAnimationRuntime::UpdateDesiredBoneWeight(const TArray<FPerBoneBlendWeight>& SrcBoneBlendWeights, TArray<FPerBoneBlendWeight>& TargetBoneBlendWeights, const TArray<float>& BlendWeights, const FBoneContainer& RequiredBones, USkeleton* Skeleton)
+void FAnimationRuntime::UpdateDesiredBoneWeight(const TArray<FPerBoneBlendWeight>& SrcBoneBlendWeights, TArray<FPerBoneBlendWeight>& TargetBoneBlendWeights, const TArray<float>& BlendWeights)
 {
 	// in the future, cache this outside
 	ensure (TargetBoneBlendWeights.Num() == SrcBoneBlendWeights.Num());
 
 	FMemory::Memset(TargetBoneBlendWeights.GetData(), 0, TargetBoneBlendWeights.Num() * sizeof(FPerBoneBlendWeight));
 
-	// go through skeleton tree requiredboneindices
-	const TArray<FBoneIndexType> & RequiredBoneIndices = RequiredBones.GetBoneIndicesArray();
-	for (int32 I=0; I<RequiredBoneIndices.Num(); ++I)
+	for (int32 BoneIndex = 0; BoneIndex < SrcBoneBlendWeights.Num(); ++BoneIndex)
 	{
-		int32 BoneIndex = RequiredBoneIndices[I];
 		int32 PoseIndex = SrcBoneBlendWeights[BoneIndex].SourceIndex;
 		float TargetBlendWeight = BlendWeights[PoseIndex]*SrcBoneBlendWeights[BoneIndex].BlendWeight;
 		// if relevant, otherwise all initialized as zero
@@ -976,14 +973,15 @@ void FAnimationRuntime::BlendPosesPerBoneFilter(struct FCompactPose& BasePose, c
 	}
 }
 
-void FAnimationRuntime::CreateMaskWeights(int32 NumOfBones, TArray<FPerBoneBlendWeight> & BoneBlendWeights, const TArray<FInputBlendPose>	&BlendFilters, const FBoneContainer& RequiredBones, const USkeleton* Skeleton)
+void FAnimationRuntime::CreateMaskWeights(TArray<FPerBoneBlendWeight> & BoneBlendWeights, const TArray<FInputBlendPose>	&BlendFilters, const FBoneContainer& RequiredBones, const USkeleton* Skeleton)
 {
 	if ( Skeleton )
 	{
-		BoneBlendWeights.Empty(NumOfBones);
-		BoneBlendWeights.AddZeroed(NumOfBones);
-
 		const TArray<FBoneIndexType> & RequiredBoneIndices = RequiredBones.GetBoneIndicesArray();
+		
+		BoneBlendWeights.Empty(RequiredBoneIndices.Num());
+		BoneBlendWeights.AddZeroed(RequiredBoneIndices.Num());
+
 		// base mask bone
 		for (int32 PoseIndex=0; PoseIndex<BlendFilters.Num(); ++PoseIndex)
 		{
@@ -999,10 +997,10 @@ void FAnimationRuntime::CreateMaskWeights(int32 NumOfBones, TArray<FPerBoneBlend
 				float IncreaseWeightPerDepth = (BranchFilter.BlendDepth != 0) ? (1.f/((float)BranchFilter.BlendDepth)) : 1.f;
 	
 				// go through skeleton tree requiredboneindices
-				for (int32 I=0; I<RequiredBoneIndices.Num(); ++I)
+				for (int32 BoneIndex = 0; BoneIndex<RequiredBoneIndices.Num(); ++BoneIndex)
 				{
-					int32 BoneIndex = RequiredBoneIndices[I];
-					int32 Depth = RequiredBones.GetDepthBetweenBones(BoneIndex, MaskBoneIndex);
+					int32 MeshBoneIndex = RequiredBoneIndices[BoneIndex];
+					int32 Depth = RequiredBones.GetDepthBetweenBones(MeshBoneIndex, MaskBoneIndex);
 
 					// if Depth == -1, it's not a child
 					if( Depth != -1 )
diff --git a/Engine/Source/Runtime/Engine/Public/AnimationRuntime.h b/Engine/Source/Runtime/Engine/Public/AnimationRuntime.h
--- a/Engine/Source/Runtime/Engine/Public/AnimationRuntime.h
+++ b/Engine/Source/Runtime/Engine/Public/AnimationRuntime.h
@@ -114,9 +114,9 @@ public:
 	**/
 	static void BlendPosesPerBoneFilter(struct FCompactPose& BasePose, const TArray<struct FCompactPose>& BlendPoses, struct FCompactPose& OutPose, TArray<FPerBoneBlendWeight>& BoneBlendWeights, bool bMeshSpaceRotationBlending);
 
-	static void UpdateDesiredBoneWeight(const TArray<FPerBoneBlendWeight>& SrcBoneBlendWeights, TArray<FPerBoneBlendWeight>& TargetBoneBlendWeights, const TArray<float>& BlendWeights, const FBoneContainer& RequiredBones, USkeleton* Skeleton);
+	static void UpdateDesiredBoneWeight(const TArray<FPerBoneBlendWeight>& SrcBoneBlendWeights, TArray<FPerBoneBlendWeight>& TargetBoneBlendWeights, const TArray<float>& BlendWeights);
 
-	static void CreateMaskWeights(int32 NumOfBones, 
+	static void CreateMaskWeights(
 			TArray<FPerBoneBlendWeight> & BoneBlendWeights, 
 			const TArray<FInputBlendPose>	&BlendFilters, 
 			const FBoneContainer& RequiredBones, 
