Additive blending optimizations

diff --git a/Engine/Source/Runtime/Engine/Classes/Animation/AnimInstance.h b/Engine/Source/Runtime/Engine/Classes/Animation/AnimInstance.h
--- a/Engine/Source/Runtime/Engine/Classes/Animation/AnimInstance.h
+++ b/Engine/Source/Runtime/Engine/Classes/Animation/AnimInstance.h
@@ -371,13 +371,8 @@ public:
 
 	static void CopyPose(const struct FCompactPose& Source, struct FCompactPose& Destination);
 
-	void ApplyAdditiveSequence(const struct FCompactPose& BasePose, const struct FCompactPose& AdditivePose, float Alpha, struct FCompactPose& Blended);
-
 	void BlendSpaceEvaluatePose(UBlendSpaceBase* BlendSpace, TArray<FBlendSampleData>& BlendSampleDataCache, struct FCompactPose& OutPose);
 
-	// skeletal control related functions
-	void BlendRotationOffset(const FCompactPose& BasePose/* local space base pose */, const FCompactPose& RotationOffsetPose/* mesh space rotation only additive **/, float Alpha/*0 means no additive, 1 means whole additive */, struct FCompactPose& OutPose /** local space blended pose **/);
-
 	// slotnode interfaces
 	void GetSlotWeight(FName const & SlotNodeName, float& out_SlotNodeWeight, float& out_SourceWeight) const;
 	void SlotEvaluatePose(FName SlotNodeName, const FCompactPose& SourcePose, FCompactPose& BlendedPose, float SlotNodeWeight);
diff --git a/Engine/Source/Runtime/Engine/Private/Animation/AnimInstance.cpp b/Engine/Source/Runtime/Engine/Private/Animation/AnimInstance.cpp
--- a/Engine/Source/Runtime/Engine/Private/Animation/AnimInstance.cpp
+++ b/Engine/Source/Runtime/Engine/Private/Animation/AnimInstance.cpp
@@ -1045,90 +1045,6 @@ void UAnimInstance::BlendSpaceEvaluatePose(class UBlendSpaceBase* BlendSpace, TA
 		/*out*/ OutPose);
 }
 
-void UAnimInstance::BlendRotationOffset(const FCompactPose& BasePose/* local space base pose */, const FCompactPose& RotationOffsetPose/* mesh space rotation only additive **/, float Alpha/*0 means no additive, 1 means whole additive */, struct FCompactPose& OutPose /** local space blended pose **/)
-{
-	//MDW TODO: Can this be improved? Lots of pose creation
-	SCOPE_CYCLE_COUNTER(STAT_AnimNativeBlendPoses);
-
-	check(RotationOffsetPose.GetNumBones() == RequiredBones.GetCompactPoseNumBones());
-	check(BasePose.GetNumBones() == RotationOffsetPose.GetNumBones());
-	check(OutPose.GetNumBones() == RotationOffsetPose.GetNumBones());
-
-	// now Pose has Mesh based BasePose
-	// apply additive
-	if (Alpha > ZERO_ANIMWEIGHT_THRESH)
-	{
-		FCompactPose BlendedPose;
-		BlendedPose.SetBoneContainer(&RequiredBones);
-
-		FCompactPose MeshBasePose;
-		MeshBasePose.SetBoneContainer(&RequiredBones);
-
-		// note that RotationOffsetPose has MeshSpaceRotation additive but everything else (translation/scale) is local space
-		// First calculate Mesh space for Base Pose
-		for (const FCompactPoseBoneIndex BoneIndex : BasePose.ForEachBoneIndex())
-		{
-			const FCompactPoseBoneIndex ParentIndex = RequiredBones.GetParentBoneIndex(BoneIndex);
-
-			if (ParentIndex != INDEX_NONE)
-			{
-				MeshBasePose[BoneIndex] = BasePose[BoneIndex] * MeshBasePose[ParentIndex];
-			}
-			else
-			{
-				MeshBasePose[BoneIndex] = BasePose[BoneIndex];
-			}
-		}
-
-		const ScalarRegister VBlendWeight(Alpha);
-		for (const FCompactPoseBoneIndex BoneIndex : BlendedPose.ForEachBoneIndex())
-		{
-			FTransform& Result = BlendedPose[BoneIndex];
-
-			// We want Base pose (local Pose)
-			Result = BasePose[BoneIndex];
-
-			// set result rotation to be mesh space rotation, so that it applys to mesh space rotation
-			Result.SetRotation(MeshBasePose[BoneIndex].GetRotation());
-
-			// @fixme laurent - we should make a read only version so we can avoid the copy.
-			FTransform Additive = RotationOffsetPose[BoneIndex];
-			FTransform::BlendFromIdentityAndAccumulate(Result, Additive, VBlendWeight);
-		}
-
-		// Ensure that all of the resulting rotations are normalized
-		BlendedPose.NormalizeRotations();
-
-		// now convert back to Local
-		for (const FCompactPoseBoneIndex BoneIndex : BlendedPose.ForEachBoneIndex())
-		{
-			FCompactPoseBoneIndex ParentIndex = RequiredBones.GetParentBoneIndex(BoneIndex);
-
-			OutPose[BoneIndex] = BlendedPose[BoneIndex];
-
-			if (ParentIndex != INDEX_NONE)
-			{
-				// convert to local space first
-				FQuat Rotation = BlendedPose[ParentIndex].GetRotation().Inverse() * BlendedPose[BoneIndex].GetRotation();
-				OutPose[BoneIndex].SetRotation(Rotation);
-			}
-		}
-	}
-	else
-	{
-		OutPose = BasePose;
-	}
-}
-
-void UAnimInstance::ApplyAdditiveSequence(const struct FCompactPose& BasePose, const struct FCompactPose& AdditivePose, float Alpha, struct FCompactPose& Blended)
-{
-	check(BasePose.GetNumBones() == Blended.GetNumBones());
-	
-	float BlendWeight = FMath::Clamp<float>(Alpha, 0.f, 1.f);
-
-	FAnimationRuntime::BlendAdditivePose(BasePose, AdditivePose, BlendWeight, Blended);
-}
-
 void UAnimInstance::RecalcRequiredBones()
 {
 	USkeletalMeshComponent * SkelMeshComp = GetSkelMeshComponent();
@@ -1578,6 +1494,7 @@ void UAnimInstance::SlotEvaluatePose(FName SlotNodeName, const FCompactPose& Sou
 	}
 
 	// Third pass, layer on weighted additive poses.
+	if (AdditivePoses.Num() > 0)
 	{
 		for (int32 Index = 0; Index < AdditivePoses.Num(); Index++)
 		{
@@ -1585,11 +1502,11 @@ void UAnimInstance::SlotEvaluatePose(FName SlotNodeName, const FCompactPose& Sou
 			// if additive, we should blend with source to make it full body
 			if (AdditivePose.AdditiveType == AAT_LocalSpaceBase)
 			{
-				ApplyAdditiveSequence(BlendedPose, AdditivePose.Pose, AdditivePose.Weight, BlendedPose);
+				FAnimationRuntime::AccumulateAdditivePose(BlendedPose, AdditivePose.Pose, AdditivePose.Weight);
 			}
 			else if (AdditivePose.AdditiveType == AAT_RotationOffsetMeshSpace)
 			{
-				BlendRotationOffset(BlendedPose, AdditivePose.Pose, AdditivePose.Weight, BlendedPose);
+				FAnimationRuntime::AccumulateMeshSpaceRotationAdditiveToLocalPose(BlendedPose, AdditivePose.Pose, AdditivePose.Weight);
 			}
 			else
 			{
@@ -1597,6 +1514,9 @@ void UAnimInstance::SlotEvaluatePose(FName SlotNodeName, const FCompactPose& Sou
 			}
 		}
 	}
+
+	// Normalize rotations after blending/accumulation
+	BlendedPose.NormalizeRotations();
 }
 
 void UAnimInstance::ReinitializeSlotNodes()
diff --git a/Engine/Source/Runtime/Engine/Private/Animation/AnimNode_ApplyAdditive.cpp b/Engine/Source/Runtime/Engine/Private/Animation/AnimNode_ApplyAdditive.cpp
--- a/Engine/Source/Runtime/Engine/Private/Animation/AnimNode_ApplyAdditive.cpp
+++ b/Engine/Source/Runtime/Engine/Private/Animation/AnimNode_ApplyAdditive.cpp
@@ -36,13 +36,13 @@ void FAnimNode_ApplyAdditive::Evaluate(FPoseContext& Output)
 	const float ActualAlpha = AlphaScaleBias.ApplyTo(Alpha);
 	if (ActualAlpha > ZERO_ANIMWEIGHT_THRESH)
 	{
-		FPoseContext BaseEvalContext(Output);
 		FPoseContext AdditiveEvalContext(Output);
 
-		Base.Evaluate(BaseEvalContext);
+		Base.Evaluate(Output);
 		Additive.Evaluate(AdditiveEvalContext);
 
-		Output.AnimInstance->ApplyAdditiveSequence(BaseEvalContext.Pose, AdditiveEvalContext.Pose, ActualAlpha, Output.Pose);
+		FAnimationRuntime::AccumulateAdditivePose(Output.Pose, AdditiveEvalContext.Pose, ActualAlpha);
+		Output.Pose.NormalizeRotations();
 	}
 	else
 	{
diff --git a/Engine/Source/Runtime/Engine/Private/Animation/AnimNode_RotationOffsetBlendSpace.cpp b/Engine/Source/Runtime/Engine/Private/Animation/AnimNode_RotationOffsetBlendSpace.cpp
--- a/Engine/Source/Runtime/Engine/Private/Animation/AnimNode_RotationOffsetBlendSpace.cpp
+++ b/Engine/Source/Runtime/Engine/Private/Animation/AnimNode_RotationOffsetBlendSpace.cpp
@@ -26,14 +26,18 @@ void FAnimNode_RotationOffsetBlendSpace::Update(const FAnimationUpdateContext& C
 
 void FAnimNode_RotationOffsetBlendSpace::Evaluate(FPoseContext& Context)
 {
-	FPoseContext BlendSpacePose(Context);
-	FAnimNode_BlendSpacePlayer::Evaluate(BlendSpacePose);
+	// Evaluate base pose
+	BasePose.Evaluate(Context);
 
-	FPoseContext BasePoseContext(Context);
-	BasePose.Evaluate(BasePoseContext);
+	// Evaluate MeshSpaceRotation additive blendspace
+	FPoseContext MeshSpaceRotationAdditivePoseContext(Context);
+	FAnimNode_BlendSpacePlayer::Evaluate(MeshSpaceRotationAdditivePoseContext);
 
-	const float Alpha = 1.0f;
-	Context.AnimInstance->BlendRotationOffset(BasePoseContext.Pose, BlendSpacePose.Pose, Alpha, Context.Pose);
+	// Accumulate poses together
+	FAnimationRuntime::AccumulateMeshSpaceRotationAdditiveToLocalPose(Context.Pose, MeshSpaceRotationAdditivePoseContext.Pose, 1.f);
+
+	// Resulting rotations are not normalized, so normalize here.
+	Context.Pose.NormalizeRotations();
 }
 
 void FAnimNode_RotationOffsetBlendSpace::GatherDebugData(FNodeDebugData& DebugData)
diff --git a/Engine/Source/Runtime/Engine/Private/Animation/AnimSingleNodeInstance.cpp b/Engine/Source/Runtime/Engine/Private/Animation/AnimSingleNodeInstance.cpp
--- a/Engine/Source/Runtime/Engine/Private/Animation/AnimSingleNodeInstance.cpp
+++ b/Engine/Source/Runtime/Engine/Private/Animation/AnimSingleNodeInstance.cpp
@@ -319,26 +319,26 @@ bool UAnimSingleNodeInstance::NativeEvaluateAnimation(FPoseContext& Output)
 		{
 			if (Sequence->IsValidAdditive())
 			{
-				FCompactPose BasePose;
-				BasePose.SetBoneContainer(&Output.Pose.GetBoneContainer());
 				FCompactPose AdditivePose;
 				AdditivePose.SetBoneContainer(&Output.Pose.GetBoneContainer());
 
 				FAnimExtractContext ExtractionContext(CurrentTime);
-				Sequence->GetAdditiveBasePose(BasePose, ExtractionContext);
-				Sequence->GetAnimationPose(AdditivePose, ExtractionContext);
+				Sequence->GetAdditiveBasePose(AdditivePose, ExtractionContext);
+				Sequence->GetAnimationPose(Output.Pose, ExtractionContext);
 				if (Sequence->AdditiveAnimType == AAT_LocalSpaceBase)
 				{
-					ApplyAdditiveSequence(BasePose, AdditivePose, 1.0f, Output.Pose);
+					FAnimationRuntime::AccumulateAdditivePose(Output.Pose, AdditivePose, 1.f);
+					Output.Pose.NormalizeRotations();
 				}
 				else
 				{
-					BlendRotationOffset(BasePose, AdditivePose, 1.0f, Output.Pose);
+					FAnimationRuntime::AccumulateMeshSpaceRotationAdditiveToLocalPose(Output.Pose, AdditivePose, 1.f);
+					Output.Pose.NormalizeRotations();
 				}
 			}
 			else
 			{
-				// if sekeltalmesh isn't there, we'll need to use skeleton
+				// if SkeletalMesh isn't there, we'll need to use skeleton
 				FAnimationRuntime::GetPoseFromSequence(Sequence, Output.Pose, FAnimExtractContext(CurrentTime, Sequence->bEnableRootMotion));
 			}
 		}
@@ -350,15 +350,13 @@ bool UAnimSingleNodeInstance::NativeEvaluateAnimation(FPoseContext& Output)
 			EAdditiveAnimationType AdditiveAnimType = AAT_None;
 			if (AnimTrack.IsAdditive())
 			{
-				FCompactPose BasePose;
-				BasePose.SetBoneContainer(&Output.Pose.GetBoneContainer());
 				FCompactPose AdditivePose;
 				AdditivePose.SetBoneContainer(&Output.Pose.GetBoneContainer());
 
 				AdditiveAnimType = AnimTrack.IsRotationOffsetAdditive()? AAT_RotationOffsetMeshSpace : AAT_LocalSpaceBase;
 				
 				// get base pose - for now we only support ref pose as base
-				BasePose.ResetToRefPose();
+				Output.Pose.ResetToRefPose();
 				
 				//get the additive pose
 				FAnimationRuntime::GetPoseFromAnimTrack(AnimTrack, AdditivePose, FAnimExtractContext(CurrentTime, RootMotionMode == ERootMotionMode::RootMotionFromEverything));
@@ -366,11 +364,13 @@ bool UAnimSingleNodeInstance::NativeEvaluateAnimation(FPoseContext& Output)
 				// if additive, we should blend with source to make it fullbody
 				if (AdditiveAnimType == AAT_LocalSpaceBase)
 				{
-					ApplyAdditiveSequence(BasePose,AdditivePose,1.f,Output.Pose);
+					FAnimationRuntime::AccumulateAdditivePose(Output.Pose, AdditivePose, 1.f);
+					Output.Pose.NormalizeRotations();
 				}
 				else if (AdditiveAnimType == AAT_RotationOffsetMeshSpace)
 				{
-					BlendRotationOffset(BasePose,AdditivePose,1.f,Output.Pose);
+					FAnimationRuntime::AccumulateMeshSpaceRotationAdditiveToLocalPose(Output.Pose, AdditivePose, 1.f);
+					Output.Pose.NormalizeRotations();
 				}
 			}
 			else
@@ -568,21 +568,19 @@ void UAnimSingleNodeInstance::InternalBlendSpaceEvaluatePose(class UBlendSpaceBa
 
 	if (BlendSpace->IsValidAdditive())
 	{
-		FCompactPose BasePose;
-		BasePose.SetBoneContainer(&OutPose.GetBoneContainer());
 		FCompactPose AdditivePose;
 		AdditivePose.SetBoneContainer(&OutPose.GetBoneContainer());
 
 #if WITH_EDITORONLY_DATA
 		if (BlendSpace->PreviewBasePose)
 		{
-			BlendSpace->PreviewBasePose->GetBonePose(/*out*/ BasePose, FAnimExtractContext(PreviewPoseCurrentTime));
+			BlendSpace->PreviewBasePose->GetBonePose(/*out*/ OutPose, FAnimExtractContext(PreviewPoseCurrentTime));
 		}
 		else
 #endif // WITH_EDITORONLY_DATA
 		{
 			// otherwise, get ref pose
-			BasePose.ResetToRefPose();
+			OutPose.ResetToRefPose();
 		}
 
 		FAnimationRuntime::GetPoseFromBlendSpace(
@@ -593,11 +591,13 @@ void UAnimSingleNodeInstance::InternalBlendSpaceEvaluatePose(class UBlendSpaceBa
 		if (BlendSpace->IsA(UAimOffsetBlendSpace::StaticClass()) ||
 			BlendSpace->IsA(UAimOffsetBlendSpace1D::StaticClass()) )
 		{
-			BlendRotationOffset(BasePose, AdditivePose, 1.f, OutPose);
+			FAnimationRuntime::AccumulateMeshSpaceRotationAdditiveToLocalPose(OutPose, AdditivePose, 1.f);
+			OutPose.NormalizeRotations();
 		}
 		else
 		{
-			ApplyAdditiveSequence(BasePose, AdditivePose, 1.f, OutPose);
+			FAnimationRuntime::AccumulateAdditivePose(OutPose, AdditivePose, 1.f);
+			OutPose.NormalizeRotations();
 		}
 	}
 	else
diff --git a/Engine/Source/Runtime/Engine/Private/Animation/AnimationRuntime.cpp b/Engine/Source/Runtime/Engine/Private/Animation/AnimationRuntime.cpp
--- a/Engine/Source/Runtime/Engine/Private/Animation/AnimationRuntime.cpp
+++ b/Engine/Source/Runtime/Engine/Private/Animation/AnimationRuntime.cpp
@@ -326,23 +326,6 @@ void FAnimationRuntime::LerpBoneTransforms(TArray<FTransform> & A, const TArray<
 	}
 }
 
-void FAnimationRuntime::BlendAdditivePose(const FCompactPose& SourcePose, const FCompactPose& AdditiveBlendPose, const float BlendWeight,/*out*/ FCompactPose& ResultPose)
-{
-	const ScalarRegister VBlendWeight(BlendWeight);
-	// Subsequent poses need to be blended in
-
-	for (FCompactPoseBoneIndex BoneIndex : SourcePose.ForEachBoneIndex())
-	{
-		FTransform Additive = AdditiveBlendPose[BoneIndex];
-		ResultPose[BoneIndex] = SourcePose[BoneIndex];
-		FTransform::BlendFromIdentityAndAccumulate(ResultPose[BoneIndex], Additive, VBlendWeight);
-	}
-
-	// Ensure that all of the resulting rotations are normalized
-	ResultPose.NormalizeRotations();
-}
-
-
 void FAnimationRuntime::CombineWithAdditiveAnimations(int32 NumAdditivePoses, const FTransformArrayA2** SourceAdditivePoses, const float* SourceAdditiveWeights, const FBoneContainer& RequiredBones, /*inout*/ FTransformArrayA2& Atoms)
 {
 	const TArray<FBoneIndexType> & RequiredBoneIndices = RequiredBones.GetBoneIndicesArray();
@@ -466,20 +449,69 @@ void FAnimationRuntime::ConvertPoseToAdditive(FCompactPose& TargetPose, const FC
 void FAnimationRuntime::ConvertPoseToMeshRotation(FCompactPose& LocalPose)
 {
 	// Convert all rotations to mesh space
-	for (FCompactPoseBoneIndex BoneIndex : LocalPose.ForEachBoneIndex())
+	// only the root bone doesn't have a parent. So skip it to save a branch in the iteration.
+	for (FCompactPoseBoneIndex BoneIndex(1); BoneIndex < LocalPose.GetNumBones(); ++BoneIndex)
 	{
-		if (!BoneIndex.IsRootBone()) // Skip root
-		{
-			const FCompactPoseBoneIndex ParentIndex = LocalPose.GetParentBoneIndex(BoneIndex);
+		const FCompactPoseBoneIndex ParentIndex = LocalPose.GetParentBoneIndex(BoneIndex);
+
+		const FQuat MeshSpaceRotation = LocalPose[ParentIndex].GetRotation() * LocalPose[BoneIndex].GetRotation();
+		LocalPose[BoneIndex].SetRotation(MeshSpaceRotation);
+	}
+}
+
+void FAnimationRuntime::ConvertMeshRotationPoseToLocalSpace(FCompactPose& Pose)
+{
+	// Convert all rotations to mesh space
+	// only the root bone doesn't have a parent. So skip it to save a branch in the iteration.
+	for (FCompactPoseBoneIndex BoneIndex(Pose.GetNumBones()-1); BoneIndex > 0; --BoneIndex)
+	{
+		const FCompactPoseBoneIndex ParentIndex = Pose.GetParentBoneIndex(BoneIndex);
 
-			FTransform& LocalTransform = LocalPose[BoneIndex];
-			const FTransform& LocalParentTransform = LocalPose[ParentIndex];
+		FQuat LocalSpaceRotation = Pose[ParentIndex].GetRotation().Inverse() * Pose[BoneIndex].GetRotation();
+		Pose[BoneIndex].SetRotation(LocalSpaceRotation);
+	}
+}
 
-			LocalTransform.SetRotation(LocalParentTransform.GetRotation() * LocalTransform.GetRotation());
+void FAnimationRuntime::AccumulateAdditivePose(FCompactPose& BasePose, const FCompactPose& AdditivePose, float Weight)
+{
+	if (Weight > ZERO_ANIMWEIGHT_THRESH)
+	{
+		if (Weight >= (1.f - ZERO_ANIMWEIGHT_THRESH))
+		{
+			// fast path, no need to weight additive.
+			for (FCompactPoseBoneIndex BoneIndex : BasePose.ForEachBoneIndex())
+			{
+				BasePose[BoneIndex].Accumulate(AdditivePose[BoneIndex]);
+			}
+		}
+		else
+		{
+			// Slower path w/ weighting
+			const ScalarRegister VBlendWeight(Weight);
+			for (FCompactPoseBoneIndex BoneIndex : BasePose.ForEachBoneIndex())
+			{
+				// copy additive, because BlendFromIdentityAndAccumulate modifies it.
+				FTransform Additive = AdditivePose[BoneIndex];
+				FTransform::BlendFromIdentityAndAccumulate(BasePose[BoneIndex], Additive, VBlendWeight);
+			}
 		}
 	}
 }
 
+void FAnimationRuntime::AccumulateMeshSpaceRotationAdditiveToLocalPose(FCompactPose& BasePose, const FCompactPose& MeshSpaceRotationAdditive, float Weight)
+{
+	if (Weight > ZERO_ANIMWEIGHT_THRESH)
+	{
+		// Convert base pose from local space to mesh space rotation.
+		FAnimationRuntime::ConvertPoseToMeshRotation(BasePose);
+
+		// Add MeshSpaceRotAdditive to it
+		FAnimationRuntime::AccumulateAdditivePose(BasePose, MeshSpaceRotationAdditive, Weight);
+
+		// Convert back to local space
+		FAnimationRuntime::ConvertMeshRotationPoseToLocalSpace(BasePose);
+	}
+}
 
 /** 
  * return ETypeAdvanceAnim type
diff --git a/Engine/Source/Runtime/Engine/Public/AnimationRuntime.h b/Engine/Source/Runtime/Engine/Public/AnimationRuntime.h
--- a/Engine/Source/Runtime/Engine/Public/AnimationRuntime.h
+++ b/Engine/Source/Runtime/Engine/Public/AnimationRuntime.h
@@ -153,9 +153,18 @@ public:
 	/** Convert TargetPose into an AdditivePose, by doing TargetPose = TargetPose - BasePose */
 	static void ConvertPoseToAdditive(FCompactPose& TargetPose, const FCompactPose& BasePose);
 
-	/** Convert LocalPose into MeshSpaceRotations. */
+	/** Convert LocalPose into MeshSpaceRotations. Rotations are NOT normalized. */
 	static void ConvertPoseToMeshRotation(FCompactPose& LocalPose);
 
+	/** Convert a MeshSpaceRotation pose to Local Space. Rotations are NOT normalized. */
+	static void ConvertMeshRotationPoseToLocalSpace(FCompactPose& Pose);
+
+	/** Accumulates weighted AdditivePose to BasePose. Rotations are NOT normalized. */
+	static void AccumulateAdditivePose(FCompactPose& BasePose, const FCompactPose& AdditivePose, float Weight);
+
+	/** Accumulate a MeshSpaceRotation Additive pose to a local pose. Rotations are NOT normalized */
+	static void AccumulateMeshSpaceRotationAdditiveToLocalPose(FCompactPose& BasePose, const FCompactPose& MeshSpaceRotationAdditive, float Weight);
+
 	/**
 	 * Accumulates BlendPoses to ResultAtoms with BlendWeight. 
 	 * ResultAtoms += BlendPose*BlendWeight.
@@ -170,19 +179,6 @@ public:
 		const float BlendWeight,
 		/*inout*/ FCompactPose& ResultPose);
 
-	/**
-	* Blends SourcePose + BlendPose*BlendWeight = ResultAtoms. As it states, this accumulates to the SourcePose
-	*
-	* Right now this came from BlendPosesTogether, light weight blend, no slerp or anything.
-	*
-	* @param	ResultPose		Output array of relative bone transforms.
-	*/
-	static void BlendAdditivePose(
-		const FCompactPose& SourcePose,
-		const FCompactPose& AdditiveBlendPose,
-		const float BlendWeight,
-		/*out*/ FCompactPose& ResultPose);
-
 	/** Lerp for BoneTransforms. Stores results in A. Performs A = Lerp(A, B, Alpha);
 	 * @param A : In/Out transform array.
 	 * @param B : B In transform array.
