#ue4 - Overlaps optimizations during swept movement. Approximate 10% or greater reduction to movement cost when in the presence of overlaps (depending on scene complexity).

- Overlaps at the end location are now computed based on tests against the individual candidates found during the sweep, rather than using a full scene query.
- Changes in rotation still allow cached overlaps as long as the rotation is symmetric (ie couldn't have changed overlaps).
- UpdateOverlaps() now avoids any processing for child components when done for the root component. This makes optimization simpler and children will test the parent anyway.

diff --git a/Engine/Source/Runtime/Engine/Classes/Components/PrimitiveComponent.h b/Engine/Source/Runtime/Engine/Classes/Components/PrimitiveComponent.h
--- a/Engine/Source/Runtime/Engine/Classes/Components/PrimitiveComponent.h
+++ b/Engine/Source/Runtime/Engine/Classes/Components/PrimitiveComponent.h
@@ -422,7 +422,7 @@ protected:
 
 	virtual void UpdateNavigationData() override;
 
-	/** Returns true if all descendant components that we can possibly collide with use relative location and rotation. */
+	/** Returns true if all descendant components that we can possibly overlap with use relative location and rotation. */
 	virtual bool AreAllCollideableDescendantsRelative(bool bAllowCachedValue = true) const;
 
 	/** Result of last call to AreAllCollideableDescendantsRelative(). */
@@ -509,6 +509,16 @@ protected:
 	/** Set of components that this component is currently overlapping. */
 	TArray<FOverlapInfo> OverlappingComponents;
 
+private:
+	/** Convert a set of overlaps from a sweep to a subset that includes only those at the end location (filling in OverlapsAtEndLocation). */
+	const TArray<FOverlapInfo>* ConvertSweptOverlapsToCurrentOverlaps(TArray<FOverlapInfo>& OverlapsAtEndLocation, const TArray<FOverlapInfo>& SweptOverlaps, int32 SweptOverlapsIndex, const FVector& EndLocation, const FQuat& EndRotationQuat);
+
+	/** Convert a set of overlaps from a symmetric change in rotation to a subset that includes only those at the end location (filling in OverlapsAtEndLocation). */
+	const TArray<FOverlapInfo>* ConvertRotationOverlapsToCurrentOverlaps(TArray<FOverlapInfo>& OverlapsAtEndLocation, const TArray<FOverlapInfo>& CurrentOverlaps);
+
+	// FScopedMovementUpdate needs access to the above two functions.
+	friend FScopedMovementUpdate;
+
 public:
 	/** 
 	 * Begin tracking an overlap interaction with the component specified.
@@ -567,13 +577,13 @@ public:
 
 	/** 
 	 * Queries world and updates overlap tracking state for this component.
-	 * @param PendingOverlaps			An ordered list of components that the MovedComponent overlapped during its movement (i.e. generated during a sweep).
-	 *									May not be overlapping them now.
+	 * @param NewPendingOverlaps		An ordered list of components that the MovedComponent overlapped during its movement (eg. generated during a sweep). Only used to add potentially new overlaps.
+	 *									Might not be overlapping them now.
 	 * @param bDoNotifies				True to dispatch being/end overlap notifications when these events occur.
-	 * @param OverlapsAtEndLocation		If non-null, the given list of overlaps will be used as the overlaps for this component at the current location, rather than checking for them separately.
+	 * @param OverlapsAtEndLocation		If non-null, the given list of overlaps will be used as the overlaps for this component at the current location, rather than checking for them with a scene query.
 	 *									Generally this should only be used if this component is the RootComponent of the owning actor and overlaps with other descendant components have been verified.
 	 */
-	virtual void UpdateOverlaps(TArray<FOverlapInfo> const* PendingOverlaps=NULL, bool bDoNotifies=true, const TArray<FOverlapInfo>* OverlapsAtEndLocation=NULL) override;
+	virtual void UpdateOverlaps(TArray<FOverlapInfo> const* NewPendingOverlaps=nullptr, bool bDoNotifies=true, const TArray<FOverlapInfo>* OverlapsAtEndLocation=nullptr) override;
 
 	/** Update current physics volume for this component, if bShouldUpdatePhysicsVolume is true. Overridden to use the overlaps to find the physics volume. */
 	virtual void UpdatePhysicsVolume( bool bTriggerNotifiers ) override;
@@ -1774,3 +1784,8 @@ FORCEINLINE_DEBUGGABLE bool UPrimitiveComponent::ComponentOverlapComponent(class
 {
 	return ComponentOverlapComponentImpl(PrimComp, Pos, Rot.Quaternion(), Params);
 }
+
+FORCEINLINE_DEBUGGABLE const TArray<FOverlapInfo>& UPrimitiveComponent::GetOverlapInfos() const
+{
+	return OverlappingComponents;
+}
diff --git a/Engine/Source/Runtime/Engine/Classes/Components/SceneComponent.h b/Engine/Source/Runtime/Engine/Classes/Components/SceneComponent.h
--- a/Engine/Source/Runtime/Engine/Classes/Components/SceneComponent.h
+++ b/Engine/Source/Runtime/Engine/Classes/Components/SceneComponent.h
@@ -561,7 +561,7 @@ private:
 public:
 
 	/** Queries world and updates overlap tracking state for this component */
-	virtual void UpdateOverlaps(TArray<FOverlapInfo> const* PendingOverlaps=NULL, bool bDoNotifies=true, const TArray<FOverlapInfo>* OverlapsAtEndLocation=NULL);
+	virtual void UpdateOverlaps(TArray<FOverlapInfo> const* PendingOverlaps=nullptr, bool bDoNotifies=true, const TArray<FOverlapInfo>* OverlapsAtEndLocation=nullptr);
 
 	/**
 	 * Tries to move the component by a movement vector (Delta) and sets rotation to NewRotation.
@@ -962,6 +962,20 @@ public:
 	FScopedMovementUpdate( class USceneComponent* Component, EScopedUpdate::Type ScopeBehavior = EScopedUpdate::DeferredUpdates );
 	~FScopedMovementUpdate();
 
+	enum class EHasMovedTransformOption
+	{
+		eTestTransform,
+		eIgnoreTransform
+	};
+
+	enum class EOverlapState
+	{
+		eUseParent,
+		eUnknown,
+		eIncludesOverlaps,
+		eForceUpdate,
+	};
+
 	/** Get the scope containing this scope. A scope only has an outer scope if they both defer updates. */
 	const FScopedMovementUpdate* GetOuterDeferredScope() const;
 
@@ -971,12 +985,6 @@ public:
 	/** Revert movement to the initial location of the Component at the start of the scoped update. Also clears pending overlaps and sets bHasMoved to false. */
 	void RevertMove();
 
-	enum class EHasMovedTransformOption
-	{
-		eTestTransform,
-		eIgnoreTransform
-	};
-
 	/** Returns whether movement has occurred at all during this scope, optionally checking if the transform is different (since changing scale does not go through a move). RevertMove() sets this back to false. */
 	bool HasMoved(EHasMovedTransformOption CheckTransform) const;
 
@@ -986,44 +994,36 @@ public:
 	/** Returns true if there are pending overlaps queued in this scope. */
 	bool HasPendingOverlaps() const;
 
-	/** Returns true if there are pending overlaps queued in this scope or any outer scope. */
-	bool HasPendingOverlapsInAncestorOrSelf() const;
-
 	/** Returns the pending overlaps within this scope. */
 	const TArray<FOverlapInfo>& GetPendingOverlaps() const;
 
-	/** Returns whether the overlaps at the current location are completely known. May walk the outer scope chain to compute this if no movement has occurred. */
-	bool HasKnownOverlapStateAtEnd() const;
+	/** Returns the list of pending blocking hits, which will be used for notifications once the move is committed. */
+	const TBlockingHitArray& GetPendingBlockingHits() const;
 
-	/** Returns the overlaps at the current location if they are known, or null otherwise. See HasKnownOverlapStateAtEnd(). */
-	const TArray<FOverlapInfo>* GetKnownOverlapsAtEnd(const TArray<struct FOverlapInfo>* DefaultOverlaps) const;
+	//--------------------------------------------------------------------------------------------------------//
+	// These methods are intended only to be used by SceneComponent and derived classes.
 
-	/** Add overlaps to the queued overlaps array. This is intended for use only by SceneComponent and its derived classes whenever movement is performed. Null NewOverlapsAtEndLocation indicates an unknown overlap state. */
-	void AppendOverlapsAfterMove(const TArray<FOverlapInfo>& NewPendingOverlaps, const TArray<FOverlapInfo>* NewOverlapsAtEndLocation);
+	/** Add overlaps to the queued overlaps array. This is intended for use only by SceneComponent and its derived classes whenever movement is performed. */
+	void AppendOverlapsAfterMove(const TArray<FOverlapInfo>& NewPendingOverlaps, bool bSweep, bool bIncludesOverlapsAtEnd);
+
+	/** Keep current pending overlaps after a move but make note that there was movement (just a symmetric rotation). */
+	void KeepCurrentOverlapsAfterRotation(bool bSweep);
 
 	/** Add blocking hit that will get processed once the move is committed. This is intended for use only by SceneComponent and its derived classes. */
 	void AppendBlockingHitAfterMove(const FHitResult& Hit);
-	
-	/** Returns the list of pending blocking hits, which will be used for notifications once the move is committed. */
-	const TBlockingHitArray& GetPendingBlockingHits() const;
-
 
-	DEPRECATED(4.9, "This function has been renamed to AppendOverlapsAfterMove.")
-	void AppendOverlaps(const TArray<FOverlapInfo>& NewPendingOverlaps, const TArray<FOverlapInfo>* NewOverlapsAtEndLocation);
+	/** Clear overlap state at current location, we don't know what it is. */
+	void InvalidateCurrentOverlaps();
 
-	DEPRECATED(4.9, "This function has been renamed to AppendOverlapsAfterMove.")
-	void AppendBlockingHit(const FHitResult& Hit);
+	/** Force full overlap update once this scope finishes. */
+	void ForceOverlapUpdate();
 
-	DEPRECATED(4.9, "This function is ambiguous. Use either HasStoredOverlapsAtEnd() or HasKnownOverlapStateAtEnd().")
-	bool HasValidOverlapsAtEnd() const;
+	//--------------------------------------------------------------------------------------------------------//
 
 private:
 
-	/** Returns whether there are stored pending overlaps at the end location. If there has been no movement, this returns false, but overlaps may be extracted from the outer scope. */
-	bool HasStoredOverlapsAtEnd() const;
-
-	/** Returns the list of overlaps at the end location, or null if HasStoredOverlapsAtEnd() is false. */
-	const TArray<FOverlapInfo>* GetOverlapsAtEnd() const;
+	/** Fills in the list of overlaps at the end location (in EndOverlaps). Returns pointer to the list, or null if it can't be computed. */
+	const TArray<FOverlapInfo>* GetOverlapsAtEnd(class UPrimitiveComponent& PrimComponent, TArray<FOverlapInfo>& EndOverlaps, bool bTransformChanged) const;
 
 	/** Notify this scope that the given inner scope completed its update (ie is going out of scope). Only occurs for deferred updates. */
 	void OnInnerScopeComplete(const FScopedMovementUpdate& InnerScope);
@@ -1040,15 +1040,16 @@ private:
 	class USceneComponent* Owner;
 	FScopedMovementUpdate* OuterDeferredScope;
 	uint32 bDeferUpdates:1;
-	uint32 bHasStoredOverlapsAtEnd:1;
 	uint32 bHasMoved:1;
+	EOverlapState CurrentOverlapState;
+
 	FTransform InitialTransform;
 	FVector InitialRelativeLocation;
 	FRotator InitialRelativeRotation;
 	FVector InitialRelativeScale;
 
+	int32 FinalOverlapCandidatesIndex;
 	TArray<FOverlapInfo> PendingOverlaps;
-	TArray<FOverlapInfo> OverlapsAtEnd;
 	TBlockingHitArray BlockingHits;
 
 	friend class USceneComponent;
@@ -1077,59 +1078,69 @@ FORCEINLINE bool FScopedMovementUpdate::HasPendingOverlaps() const
 	return PendingOverlaps.Num() > 0;
 }
 
-FORCEINLINE_DEBUGGABLE bool FScopedMovementUpdate::HasPendingOverlapsInAncestorOrSelf() const
+FORCEINLINE const TArray<struct FOverlapInfo>& FScopedMovementUpdate::GetPendingOverlaps() const
 {
-	const FScopedMovementUpdate* Scope = this;
-	do
-	{
-		if (Scope->HasPendingOverlaps())
-		{
-			return true;
-		}
-		Scope = Scope->OuterDeferredScope;
-	}
-	while (Scope != nullptr);
-
-	return false;
+	return PendingOverlaps;
 }
 
-FORCEINLINE const TArray<struct FOverlapInfo>& FScopedMovementUpdate::GetPendingOverlaps() const
+FORCEINLINE const FScopedMovementUpdate::TBlockingHitArray& FScopedMovementUpdate::GetPendingBlockingHits() const
 {
-	return PendingOverlaps;
+	return BlockingHits;
 }
 
-FORCEINLINE bool FScopedMovementUpdate::HasStoredOverlapsAtEnd() const
+FORCEINLINE_DEBUGGABLE void FScopedMovementUpdate::AppendBlockingHitAfterMove(const FHitResult& Hit)
 {
-	return bHasStoredOverlapsAtEnd;
+	BlockingHits.Add(Hit);
 }
 
-FORCEINLINE const TArray<struct FOverlapInfo>* FScopedMovementUpdate::GetOverlapsAtEnd() const
+FORCEINLINE_DEBUGGABLE void FScopedMovementUpdate::KeepCurrentOverlapsAfterRotation(bool bSweep)
 {
-	return bHasStoredOverlapsAtEnd ? &OverlapsAtEnd : nullptr;
+	bHasMoved = true;
+	// CurrentOverlapState is unchanged
 }
 
-FORCEINLINE const FScopedMovementUpdate::TBlockingHitArray& FScopedMovementUpdate::GetPendingBlockingHits() const
+FORCEINLINE_DEBUGGABLE void FScopedMovementUpdate::InvalidateCurrentOverlaps()
 {
-	return BlockingHits;
+	bHasMoved = true;
+	CurrentOverlapState = EOverlapState::eUnknown;
+	FinalOverlapCandidatesIndex = INDEX_NONE;
 }
 
-FORCEINLINE_DEBUGGABLE void FScopedMovementUpdate::AppendBlockingHitAfterMove(const FHitResult& Hit)
+FORCEINLINE_DEBUGGABLE void FScopedMovementUpdate::ForceOverlapUpdate()
 {
-	BlockingHits.Add(Hit);
+	bHasMoved = true;
+	CurrentOverlapState = EOverlapState::eForceUpdate;
+	FinalOverlapCandidatesIndex = INDEX_NONE;
 }
 
-FORCEINLINE void FScopedMovementUpdate::AppendOverlaps(const TArray<FOverlapInfo>& NewPendingOverlaps, const TArray<FOverlapInfo>* NewOverlapsAtEndLocation)
+
+
+FORCEINLINE_DEBUGGABLE class FScopedMovementUpdate* USceneComponent::GetCurrentScopedMovement() const
 {
-	AppendOverlapsAfterMove(NewPendingOverlaps, NewOverlapsAtEndLocation);
+	if (ScopedMovementStack.Num() > 0)
+	{
+		return ScopedMovementStack.Last();
+	}
+	return nullptr;
 }
 
-FORCEINLINE void FScopedMovementUpdate::AppendBlockingHit(const FHitResult& Hit)
+//////////////////////////////////////////////////////////////////////////
+// SceneComponent inlines that depend on FScopedMovementUpdate
+
+FORCEINLINE_DEBUGGABLE bool USceneComponent::IsDeferringMovementUpdates() const
 {
-	AppendBlockingHitAfterMove(Hit);
+	if (ScopedMovementStack.Num() > 0)
+	{
+		checkSlow(ScopedMovementStack.Last()->IsDeferringUpdates());
+		return true;
+	}
+	return false;
 }
 
-FORCEINLINE bool FScopedMovementUpdate::HasValidOverlapsAtEnd() const
+FORCEINLINE_DEBUGGABLE void USceneComponent::BeginScopedMovementUpdate(class FScopedMovementUpdate& ScopedUpdate)
 {
-	// This was the old behavior, but should be re-evaluated by calling code.
-	return HasStoredOverlapsAtEnd();
+	checkSlow(IsInGameThread());
+	checkSlow(ScopedUpdate.IsDeferringUpdates());
+	ScopedMovementStack.Push(&ScopedUpdate);
 }
+
diff --git a/Engine/Source/Runtime/Engine/Private/Components/PrimitiveComponent.cpp b/Engine/Source/Runtime/Engine/Private/Components/PrimitiveComponent.cpp
--- a/Engine/Source/Runtime/Engine/Private/Components/PrimitiveComponent.cpp
+++ b/Engine/Source/Runtime/Engine/Private/Components/PrimitiveComponent.cpp
@@ -59,6 +59,46 @@ TAutoConsoleVariable<int32> CVarShowInitialOverlaps(
 	ECVF_Cheat);
 #endif // !(UE_BUILD_SHIPPING || UE_BUILD_TEST)
 
+static int32 bEnableFastOverlapCheck = 1;
+static FAutoConsoleVariableRef CVarEnableFastOverlapCheck(TEXT("p.EnableFastOverlapCheck"), bEnableFastOverlapCheck, TEXT("Enable fast overlap check against sweep hits, avoiding UpdateOverlaps (for the swept component)."));
+DECLARE_CYCLE_STAT(TEXT("MoveComponent FastOverlap"), STAT_MoveComponent_FastOverlap, STATGROUP_Game);
+
+
+// Predicate to determine if an overlap is with a certain AActor.
+struct FPredicateOverlapHasSameActor
+{
+	FPredicateOverlapHasSameActor(const AActor& Owner)
+	: MyOwner(Owner)
+	{
+	}
+
+	bool operator() (const FOverlapInfo& Info)
+	{
+		return Info.OverlapInfo.Actor == &MyOwner;
+	}
+
+private:
+	const AActor& MyOwner;
+};
+
+// Predicate to determine if an overlap is *NOT* with a certain AActor.
+struct FPredicateOverlapHasDifferentActor
+{
+	FPredicateOverlapHasDifferentActor(const AActor& Owner)
+	: MyOwner(Owner)
+	{
+	}
+
+	bool operator() (const FOverlapInfo& Info)
+	{
+		return Info.OverlapInfo.Actor != &MyOwner;
+	}
+
+private:
+	const AActor& MyOwner;
+};
+
+
 ///////////////////////////////////////////////////////////////////////////////
 // PRIMITIVE COMPONENT
 ///////////////////////////////////////////////////////////////////////////////
@@ -1355,13 +1395,14 @@ bool UPrimitiveComponent::MoveComponentImpl( const FVector& Delta, const FQuat&
 
 	// Set up.
 	float DeltaSizeSq = Delta.SizeSquared();
+	const FQuat InitialRotationQuat = ComponentToWorld.GetRotation();
 
 	// ComponentSweepMulti does nothing if moving < KINDA_SMALL_NUMBER in distance, so it's important to not try to sweep distances smaller than that. 
 	const float MinMovementDistSq = (bSweep ? FMath::Square(4.f*KINDA_SMALL_NUMBER) : 0.f);
 	if (DeltaSizeSq <= MinMovementDistSq)
 	{
 		// Skip if no vector or rotation.
-		if (NewRotationQuat.Equals(ComponentToWorld.GetRotation(), 1.e-6f))
+		if (NewRotationQuat.Equals(InitialRotationQuat, 1.e-6f))
 		{
 			// copy to optional output param
 			if (OutHit)
@@ -1376,15 +1417,17 @@ bool UPrimitiveComponent::MoveComponentImpl( const FVector& Delta, const FQuat&
 	const bool bSkipPhysicsMove = ((MoveFlags & MOVECOMP_SkipPhysicsMove) != MOVECOMP_NoFlags);
 
 	bool bMoved = false;
+	bool bIncludesOverlapsAtEnd = false;
+	bool bRotationOnly = false;
 	TArray<FOverlapInfo> PendingOverlaps;
-	TArray<FOverlapInfo> OverlapsAtEndLocation;
-	const TArray<FOverlapInfo>* OverlapsAtEndLocationPtr = nullptr; // When non-null, used as optimization to avoid work in UpdateOverlaps.
 	AActor* const Actor = GetOwner();
 
 	if ( !bSweep )
 	{
 		// not sweeping, just go directly to the new transform
 		bMoved = InternalSetWorldLocationAndRotation(TraceEnd, NewRotationQuat, bSkipPhysicsMove, !bSweep);
+		bRotationOnly = (DeltaSizeSq == 0);
+		bIncludesOverlapsAtEnd = bRotationOnly && (AreSymmetricRotations(InitialRotationQuat, NewRotationQuat, GetComponentScale())) && IsCollisionEnabled();
 	}
 	else
 	{
@@ -1418,7 +1461,7 @@ bool UPrimitiveComponent::MoveComponentImpl( const FVector& Delta, const FQuat&
 			FComponentQueryParams Params(Name_MoveComponent, Actor);
 			FCollisionResponseParams ResponseParam;
 			InitSweepCollisionParams(Params, ResponseParam);
-			bool const bHadBlockingHit = MyWorld->ComponentSweepMulti(Hits, this, TraceStart, TraceEnd, GetComponentQuat(), Params);
+			bool const bHadBlockingHit = MyWorld->ComponentSweepMulti(Hits, this, TraceStart, TraceEnd, InitialRotationQuat, Params);
 
 			if (Hits.Num() > 0)
 			{
@@ -1521,21 +1564,7 @@ bool UPrimitiveComponent::MoveComponentImpl( const FVector& Delta, const FQuat&
 				}
 			}
 
-			// We have performed a sweep that tested for all overlaps (not including components on the owning actor).
-			// However any rotation was set at the end and not swept, so we can't assume the overlaps at the end location are correct if rotation changed.
-			if (PendingOverlaps.Num() == 0 && CVarAllowCachedOverlaps->GetInt())
-			{
-				if (Actor && Actor->GetRootComponent() == this && AreSymmetricRotations(NewRotationQuat, ComponentToWorld.GetRotation(), ComponentToWorld.GetScale3D()))
-				{
-					// We know we are not overlapping any new components at the end location.
-					// Keep known overlapping child components, as long as we know their overlap status could not have changed (ie they are positioned relative to us).
-					if (AreAllCollideableDescendantsRelative())
-					{
-						GetOverlapsWithActor(Actor, OverlapsAtEndLocation);
-						OverlapsAtEndLocationPtr = &OverlapsAtEndLocation;
-					}
-				}
-			}
+			bIncludesOverlapsAtEnd = AreSymmetricRotations(InitialRotationQuat, NewRotationQuat, GetComponentScale());
 
 #if !(UE_BUILD_SHIPPING || UE_BUILD_TEST)
 			if ( (BlockingHit.Time < 1.f) && !IsZeroExtent() )
@@ -1560,24 +1589,12 @@ bool UPrimitiveComponent::MoveComponentImpl( const FVector& Delta, const FQuat&
 		{
 			// apply move delta even if components has collisions disabled
 			NewLocation += Delta;
+			bIncludesOverlapsAtEnd = false;
 		}
 		else if (DeltaSizeSq == 0.f && bCollisionEnabled)
 		{
-			// We didn't move, and any rotation that doesn't change our overlap bounds means we already know what we are overlapping at this point.
-			if (Actor && Actor->GetRootComponent() == this && CVarAllowCachedOverlaps->GetInt())
-			{
-				// Can only do this if current known overlaps are valid (not deferring updates, or we know there are no new pending overlaps).
-				const FScopedMovementUpdate* ScopedUpdate = GetCurrentScopedMovement();
-				if (ScopedUpdate == nullptr || ScopedUpdate->HasKnownOverlapStateAtEnd())
-				{
-					// Only if we know that we won't change our overlap status with children by moving.
-					if (AreSymmetricRotations(NewRotationQuat, ComponentToWorld.GetRotation(), ComponentToWorld.GetScale3D()) &&
-						AreAllCollideableDescendantsRelative())
-					{
-						OverlapsAtEndLocationPtr = (ScopedUpdate ? ScopedUpdate->GetKnownOverlapsAtEnd(&OverlappingComponents) : &OverlappingComponents);
-					}
-				}
-			}
+			bIncludesOverlapsAtEnd = AreSymmetricRotations(InitialRotationQuat, NewRotationQuat, GetComponentScale());
+			bRotationOnly = true;
 		}
 
 		// Update the location.  This will teleport any child components as well (not sweep).
@@ -1587,18 +1604,39 @@ bool UPrimitiveComponent::MoveComponentImpl( const FVector& Delta, const FQuat&
 	// Handle overlap notifications.
 	if (bMoved)
 	{
-		// Check if we are deferring the movement updates.
 		if (IsDeferringMovementUpdates())
 		{
 			// Defer UpdateOverlaps until the scoped move ends.
 			FScopedMovementUpdate* ScopedUpdate = GetCurrentScopedMovement();
-			checkSlow(ScopedUpdate != NULL);
-			ScopedUpdate->AppendOverlapsAfterMove(PendingOverlaps, OverlapsAtEndLocationPtr);
+			if (bRotationOnly && bIncludesOverlapsAtEnd)
+			{
+				ScopedUpdate->KeepCurrentOverlapsAfterRotation(bSweep);
+			}
+			else
+			{
+				ScopedUpdate->AppendOverlapsAfterMove(PendingOverlaps, bSweep, bIncludesOverlapsAtEnd);
+			}
 		}
 		else
 		{
-			// still need to do this even if bGenerateOverlapEvents is false for this component, since we could have child components where it is true
-			UpdateOverlaps(&PendingOverlaps, true, OverlapsAtEndLocationPtr);
+			if (bIncludesOverlapsAtEnd)
+			{
+				TArray<FOverlapInfo> OverlapsAtEndLocation;
+				const TArray<FOverlapInfo>* OverlapsAtEndLocationPtr = nullptr; // When non-null, used as optimization to avoid work in UpdateOverlaps.
+				if (bRotationOnly)
+				{
+					OverlapsAtEndLocationPtr = ConvertRotationOverlapsToCurrentOverlaps(OverlapsAtEndLocation, GetOverlapInfos());
+				}
+				else
+				{
+					OverlapsAtEndLocationPtr = ConvertSweptOverlapsToCurrentOverlaps(OverlapsAtEndLocation, PendingOverlaps, 0, GetComponentLocation(), GetComponentQuat());
+				}
+				UpdateOverlaps(&PendingOverlaps, true, OverlapsAtEndLocationPtr);
+			}
+			else
+			{
+				UpdateOverlaps(&PendingOverlaps, true, nullptr);
+			}
 		}
 	}
 
@@ -1608,7 +1646,6 @@ bool UPrimitiveComponent::MoveComponentImpl( const FVector& Delta, const FQuat&
 		if (IsDeferringMovementUpdates())
 		{
 			FScopedMovementUpdate* ScopedUpdate = GetCurrentScopedMovement();
-			checkSlow(ScopedUpdate != NULL);
 			ScopedUpdate->AppendBlockingHitAfterMove(BlockingHit);
 		}
 		else
@@ -2042,6 +2079,87 @@ void UPrimitiveComponent::UpdateNavigationData()
 }
 
 
+const TArray<FOverlapInfo>* UPrimitiveComponent::ConvertSweptOverlapsToCurrentOverlaps(
+	TArray<FOverlapInfo>& OverlapsAtEndLocation, const TArray<FOverlapInfo>& SweptOverlaps, int32 SweptOverlapsIndex,
+	const FVector& EndLocation, const FQuat& EndRotationQuat)
+{
+	checkSlow(SweptOverlapsIndex >= 0);
+
+	const TArray<FOverlapInfo>* Result = nullptr;
+	if (bGenerateOverlapEvents && CVarAllowCachedOverlaps->GetInt())
+	{
+		const AActor* Actor = GetOwner();
+		if (Actor && Actor->GetRootComponent() == this)
+		{
+			// We know we are not overlapping any new components at the end location. Children are ignored here (see note below).
+			if (bEnableFastOverlapCheck)
+			{
+				SCOPE_CYCLE_COUNTER(STAT_MoveComponent_FastOverlap);
+
+				// Check components we hit during the sweep, keep only those still overlapping
+				const FCollisionQueryParams UnusedQueryParams;
+				for (int32 Index = SweptOverlapsIndex; Index < SweptOverlaps.Num(); ++Index)
+				{
+					const FOverlapInfo& OtherOverlap = SweptOverlaps[Index];
+					UPrimitiveComponent* OtherPrimitive = OtherOverlap.OverlapInfo.GetComponent();
+					if (OtherPrimitive && OtherPrimitive->bGenerateOverlapEvents)
+					{
+						if (OtherPrimitive->bMultiBodyOverlap)
+						{
+							// Not handled yet. We could do it by checking every body explicitly and track each body index in the overlap test, but this seems like a rare need.
+							return nullptr;
+						}
+						else if (ComponentOverlapComponent(OtherPrimitive, EndLocation, EndRotationQuat, UnusedQueryParams))
+						{
+							OverlapsAtEndLocation.Add(OtherOverlap);
+						}
+					}
+				}
+
+				// Note: we don't worry about adding any child components here, because they are not included in the sweep results.
+				// Children test for their own overlaps after we update our own, and we ignore children in our own update.
+				checkfSlow(OverlapsAtEndLocation.FindByPredicate(FPredicateOverlapHasSameActor(*Actor)) == nullptr,
+					TEXT("Child overlaps should not be included in the SweptOverlaps() array in UPrimitiveComponent::ConvertSweptOverlapsToCurrentOverlaps()."));
+
+				Result = &OverlapsAtEndLocation;
+			}
+			else
+			{
+				if (SweptOverlaps.Num() == 0 && AreAllCollideableDescendantsRelative())
+				{
+					// Add overlaps with components in this actor.
+					GetOverlapsWithActor(Actor, OverlapsAtEndLocation);
+					Result = &OverlapsAtEndLocation;
+				}
+			}
+		}
+	}
+
+	return Result;
+}
+
+
+const TArray<FOverlapInfo>* UPrimitiveComponent::ConvertRotationOverlapsToCurrentOverlaps(TArray<FOverlapInfo>& OverlapsAtEndLocation, const TArray<FOverlapInfo>& CurrentOverlaps)
+{
+	const TArray<FOverlapInfo>* Result = nullptr;
+	if (bGenerateOverlapEvents && CVarAllowCachedOverlaps->GetInt())
+	{
+		const AActor* Actor = GetOwner();
+		if (Actor && Actor->GetRootComponent() == this)
+		{
+			if (bEnableFastOverlapCheck)
+			{
+				// Add all current overlaps that are not children. Children test for their own overlaps after we update our own, and we ignore children in our own update.
+				OverlapsAtEndLocation = CurrentOverlaps.FilterByPredicate(FPredicateOverlapHasDifferentActor(*Actor));
+				Result = &OverlapsAtEndLocation;
+			}
+		}
+	}
+
+	return Result;
+}
+
+
 bool UPrimitiveComponent::AreAllCollideableDescendantsRelative(bool bAllowCachedValue) const
 {
 	UPrimitiveComponent* MutableThis = const_cast<UPrimitiveComponent*>(this);
@@ -2070,7 +2188,8 @@ bool UPrimitiveComponent::AreAllCollideableDescendantsRelative(bool bAllowCached
 				if (CurrentComp->bAbsoluteLocation || CurrentComp->bAbsoluteRotation)
 				{
 					// Can we possibly collide with the component?
-					if (CurrentComp->IsCollisionEnabled() && CurrentComp->GetCollisionResponseToChannel(GetCollisionObjectType()) != ECR_Ignore)
+					UPrimitiveComponent* const CurrentPrimitive = Cast<UPrimitiveComponent>(CurrentComp);
+					if (CurrentPrimitive && CurrentPrimitive->bGenerateOverlapEvents && CurrentPrimitive->IsCollisionEnabled() && CurrentPrimitive->GetCollisionResponseToChannel(GetCollisionObjectType()) != ECR_Ignore)
 					{
 						MutableThis->bCachedAllCollideableDescendantsRelative = false;
 						MutableThis->LastCheckedAllCollideableDescendantsTime = MyWorld->GetTimeSeconds();
@@ -2087,11 +2206,6 @@ bool UPrimitiveComponent::AreAllCollideableDescendantsRelative(bool bAllowCached
 	return true;
 }
 
-const TArray<FOverlapInfo>& UPrimitiveComponent::GetOverlapInfos() const
-{
-	return OverlappingComponents;
-}
-
 void UPrimitiveComponent::IgnoreActorWhenMoving(AActor* Actor, bool bShouldIgnore)
 {
 	// Clean up stale references
@@ -2124,12 +2238,14 @@ void UPrimitiveComponent::ClearMoveIgnoreActors()
 }
 
 
-void UPrimitiveComponent::UpdateOverlaps(const TArray<FOverlapInfo>* PendingOverlaps, bool bDoNotifies, const TArray<FOverlapInfo>* OverlapsAtEndLocation)
+void UPrimitiveComponent::UpdateOverlaps(const TArray<FOverlapInfo>* NewPendingOverlaps, bool bDoNotifies, const TArray<FOverlapInfo>* OverlapsAtEndLocation)
 {
 	SCOPE_CYCLE_COUNTER(STAT_UpdateOverlaps); 
 
 	if (IsDeferringMovementUpdates())
 	{
+		// Someone tried to call UpdateOverlaps() explicitly during a deferred update, this means they really have a good reason to force it.
+		GetCurrentScopedMovement()->ForceOverlapUpdate();
 		return;
 	}
 
@@ -2141,11 +2257,15 @@ void UPrimitiveComponent::UpdateOverlaps(const TArray<FOverlapInfo>* PendingOver
 		AActor* const MyActor = GetOwner();
 		if ( MyActor && MyActor->IsActorInitialized() )
 		{
-			if (PendingOverlaps)
+			// If we are the root component we ignore child components. Those children will update their overlaps when we descend into the child tree.
+			// This aids an optimization in MoveComponent.
+			const bool bIgnoreChildren = (MyActor->GetRootComponent() == this);
+
+			if (NewPendingOverlaps)
 			{
-				for (int32 Idx=0; Idx<PendingOverlaps->Num(); ++Idx)
+				for (int32 Idx=0; Idx<NewPendingOverlaps->Num(); ++Idx)
 				{
-					BeginComponentOverlap( (*PendingOverlaps)[Idx], bDoNotifies );
+					BeginComponentOverlap( (*NewPendingOverlaps)[Idx], bDoNotifies );
 				}
 			}
 
@@ -2160,17 +2280,17 @@ void UPrimitiveComponent::UpdateOverlaps(const TArray<FOverlapInfo>* PendingOver
 				// Might be able to avoid testing for new overlaps at the end location.
 				if (OverlapsAtEndLocation != NULL && CVarAllowCachedOverlaps->GetInt())
 				{
-					UE_LOG(LogPrimitiveComponent, VeryVerbose, TEXT("%s Skipping overlap test!"), *GetName());
+					UE_LOG(LogPrimitiveComponent, VeryVerbose, TEXT("%s->%s Skipping overlap test!"), *GetNameSafe(GetOwner()), *GetName());
 					NewOverlappingComponents = *OverlapsAtEndLocation;
 				}
 				else
 				{
-					UE_LOG(LogPrimitiveComponent, VeryVerbose, TEXT("%s Performing overlaps!"), *GetName());
+					UE_LOG(LogPrimitiveComponent, VeryVerbose, TEXT("%s->%s Performing overlaps!"), *GetNameSafe(GetOwner()), *GetName());
 					UWorld* const MyWorld = MyActor->GetWorld();
 					TArray<FOverlapResult> Overlaps;
 					static FName NAME_UpdateOverlaps = FName(TEXT("UpdateOverlaps"));
-					// note this will include overlaps with components in the same actor.  
-					FComponentQueryParams Params(NAME_UpdateOverlaps);
+					// note this will optionally include overlaps with components in the same actor (depending on bIgnoreChildren). 
+					FComponentQueryParams Params(NAME_UpdateOverlaps, bIgnoreChildren ? MyActor : nullptr);
 					Params.bTraceAsyncScene = bCheckAsyncSceneOnMove;
 					Params.AddIgnoredActors(MoveIgnoreActors);
 					MyWorld->ComponentOverlapMulti(Overlaps, this, GetComponentLocation(), GetComponentQuat(), Params);
@@ -2194,7 +2314,16 @@ void UPrimitiveComponent::UpdateOverlaps(const TArray<FOverlapInfo>* PendingOver
 			if (OverlappingComponents.Num() > 0)
 			{
 				// make a copy of the old that we can manipulate to avoid n^2 searching later
-				TInlineOverlapInfoArray OldOverlappingComponents(OverlappingComponents);
+				TInlineOverlapInfoArray OldOverlappingComponents;
+				if (bIgnoreChildren)
+				{
+					OldOverlappingComponents = OverlappingComponents.FilterByPredicate(FPredicateOverlapHasDifferentActor(*MyActor));
+					checkSlow(!NewOverlappingComponents.FindByPredicate(FPredicateOverlapHasSameActor(*MyActor)));
+				}
+				else
+				{
+					OldOverlappingComponents = OverlappingComponents;
+				}
 
 				// Now we want to compare the old and new overlap lists to determine 
 				// what overlaps are in old and not in new (need end overlap notifies), and 
@@ -2250,12 +2379,6 @@ void UPrimitiveComponent::UpdateOverlaps(const TArray<FOverlapInfo>* PendingOver
 		}
 	}
 
-	// Update physics volume using most current overlaps
-	if (bShouldUpdatePhysicsVolume)
-	{
-		UpdatePhysicsVolume(bDoNotifies);
-	}
-
 	// now update any children down the chain.
 	for (int32 ChildIdx=0; ChildIdx<AttachChildren.Num(); ++ChildIdx)
 	{
@@ -2266,6 +2389,12 @@ void UPrimitiveComponent::UpdateOverlaps(const TArray<FOverlapInfo>* PendingOver
 			ChildComp->UpdateOverlaps(NULL, bDoNotifies, NULL);
 		}
 	}
+
+	// Update physics volume using most current overlaps
+	if (bShouldUpdatePhysicsVolume)
+	{
+		UpdatePhysicsVolume(bDoNotifies);
+	}
 }
 
 bool UPrimitiveComponent::ComponentOverlapMultiImpl(TArray<struct FOverlapResult>& OutOverlaps, const UWorld* World, const FVector& Pos, const FQuat& Quat, ECollisionChannel TestChannel, const struct FComponentQueryParams& Params, const struct FCollisionObjectQueryParams& ObjectQueryParams) const
diff --git a/Engine/Source/Runtime/Engine/Private/Components/SceneComponent.cpp b/Engine/Source/Runtime/Engine/Private/Components/SceneComponent.cpp
--- a/Engine/Source/Runtime/Engine/Private/Components/SceneComponent.cpp
+++ b/Engine/Source/Runtime/Engine/Private/Components/SceneComponent.cpp
@@ -437,38 +437,6 @@ void USceneComponent::PropagateTransformUpdate(bool bTransformChanged, bool bSki
 }
 
 
-class FScopedMovementUpdate* USceneComponent::GetCurrentScopedMovement() const
-{
-	if (ScopedMovementStack.Num() > 0)
-	{
-		return ScopedMovementStack.Last();
-	}
-
-	return NULL;
-}
-
-
-bool USceneComponent::IsDeferringMovementUpdates() const
-{
-	if (ScopedMovementStack.Num() > 0)
-	{
-		checkSlow(ScopedMovementStack.Last()->IsDeferringUpdates());
-		return true;
-	}
-	
-	return false;
-}
-
-
-void USceneComponent::BeginScopedMovementUpdate(class FScopedMovementUpdate& ScopedUpdate)
-{
-	checkSlow(IsInGameThread());
-	checkSlow(ScopedUpdate.IsDeferringUpdates());
-	
-	ScopedMovementStack.Push(&ScopedUpdate);
-}
-
-
 void USceneComponent::EndScopedMovementUpdate(class FScopedMovementUpdate& CompletedScope)
 {
 	SCOPE_CYCLE_COUNTER(STAT_EndScopedMovementUpdate);
@@ -491,6 +459,7 @@ void USceneComponent::EndScopedMovementUpdate(class FScopedMovementUpdate& Compl
 			const bool bTransformChanged = CurrentScopedUpdate->IsTransformDirty();
 			if (bTransformChanged)
 			{
+				// TODO: handle teleporting flag when it differs between accumulated moves.
 				PropagateTransformUpdate(true);
 			}
 
@@ -498,7 +467,17 @@ void USceneComponent::EndScopedMovementUpdate(class FScopedMovementUpdate& Compl
 			// If no movement and no change in transform, nothing changed.
 			if (bTransformChanged || CurrentScopedUpdate->bHasMoved)
 			{
-				UpdateOverlaps(&CurrentScopedUpdate->GetPendingOverlaps(), true, CurrentScopedUpdate->GetOverlapsAtEnd());
+				UPrimitiveComponent* PrimitiveThis = Cast<UPrimitiveComponent>(this);
+				if (PrimitiveThis)
+				{
+					TArray<FOverlapInfo> EndOverlaps;
+					const TArray<FOverlapInfo>* EndOverlapsPtr = CurrentScopedUpdate->GetOverlapsAtEnd(*PrimitiveThis, EndOverlaps, bTransformChanged);
+					UpdateOverlaps(&CurrentScopedUpdate->GetPendingOverlaps(), true, EndOverlapsPtr);
+				}
+				else
+				{
+					UpdateOverlaps(nullptr, true, nullptr);
+				}
 			}
 
 			// Dispatch all deferred blocking hits
@@ -798,7 +777,16 @@ void USceneComponent::SetRelativeScale3D(FVector NewScale3D)
 		
 		if (IsRegistered())
 		{
-			UpdateOverlaps();
+			if (!IsDeferringMovementUpdates())
+			{
+				UpdateOverlaps();
+			}
+			else
+			{
+				// Invalidate cached overlap state at this location.
+				TArray<FOverlapInfo> EmptyOverlaps;
+				GetCurrentScopedMovement()->AppendOverlapsAfterMove(EmptyOverlaps, false, false);
+			}
 		}
 	}
 }
@@ -1865,14 +1853,10 @@ void USceneComponent::UpdateOverlaps(TArray<FOverlapInfo> const* PendingOverlaps
 
 	if (IsDeferringMovementUpdates())
 	{
+		GetCurrentScopedMovement()->ForceOverlapUpdate();
 		return;
 	}
 	
-	if (bShouldUpdatePhysicsVolume)
-	{
-		UpdatePhysicsVolume(bDoNotifies);
-	}
-
 	// SceneComponent has no physical representation, so no overlaps to test for/
 	// But, we need to test down the attachment chain since there might be PrimitiveComponents below.
 	for (int32 ChildIdx=0; ChildIdx<AttachChildren.Num(); ++ChildIdx)
@@ -1883,6 +1867,11 @@ void USceneComponent::UpdateOverlaps(TArray<FOverlapInfo> const* PendingOverlaps
 			AttachChildren[ChildIdx]->UpdateOverlaps(NULL, bDoNotifies);
 		}
 	}
+
+	if (bShouldUpdatePhysicsVolume)
+	{
+		UpdatePhysicsVolume(bDoNotifies);
+	}
 }
 
 bool USceneComponent::CheckStaticMobilityAndWarn(const FText& ActionText) const
@@ -2249,8 +2238,9 @@ FScopedMovementUpdate::FScopedMovementUpdate( class USceneComponent* Component,
 : Owner(Component)
 , OuterDeferredScope(nullptr)
 , bDeferUpdates(ScopeBehavior == EScopedUpdate::DeferredUpdates)
-, bHasStoredOverlapsAtEnd(false)
 , bHasMoved(false)
+, CurrentOverlapState(EOverlapState::eUseParent)
+, FinalOverlapCandidatesIndex(INDEX_NONE)
 {
 	if (IsValid(Component))
 	{
@@ -2313,9 +2303,8 @@ void FScopedMovementUpdate::RevertMove()
 	USceneComponent* Component = Owner;
 	if (IsValid(Component))
 	{
-		bHasStoredOverlapsAtEnd = false;
+		FinalOverlapCandidatesIndex = INDEX_NONE;
 		PendingOverlaps.Reset();
-		OverlapsAtEnd.Reset();
 		BlockingHits.Reset();
 		
 		if (IsTransformDirty())
@@ -2334,26 +2323,44 @@ void FScopedMovementUpdate::RevertMove()
 		}
 	}
 	bHasMoved = false;
+	CurrentOverlapState = EOverlapState::eUseParent;
 }
 
-
-void FScopedMovementUpdate::AppendOverlapsAfterMove(const TArray<FOverlapInfo>& NewPendingOverlaps, const TArray<FOverlapInfo>* NewOverlapsAtEndLocation)
+void FScopedMovementUpdate::AppendOverlapsAfterMove(const TArray<FOverlapInfo>& NewPendingOverlaps, bool bSweep, bool bIncludesOverlapsAtEnd)
 {
 	bHasMoved = true;
-	PendingOverlaps.Append(NewPendingOverlaps);
-	if (NewOverlapsAtEndLocation)
+	const bool bWasForcing = (CurrentOverlapState == EOverlapState::eForceUpdate);
+
+	if (bIncludesOverlapsAtEnd)
 	{
-		bHasStoredOverlapsAtEnd = true;
-		OverlapsAtEnd = *NewOverlapsAtEndLocation;		
+		CurrentOverlapState = EOverlapState::eIncludesOverlaps;
+		if (NewPendingOverlaps.Num())
+		{
+			FinalOverlapCandidatesIndex = PendingOverlaps.Num();
+			PendingOverlaps.Append(NewPendingOverlaps);
+		}
+		else
+		{
+			// No new pending overlaps means we're not overlapping anything at the end location.
+			FinalOverlapCandidatesIndex = INDEX_NONE;
+		}
 	}
 	else
 	{
-		// null indicates unknown overlap state at the end.
-		bHasStoredOverlapsAtEnd = false;
-		OverlapsAtEnd.Reset();
+		// We don't expect any pending overlaps in the case of a teleport.
+		checkSlow(!bSweep);
+		checkSlow(NewPendingOverlaps.Num() == 0);
+		CurrentOverlapState = EOverlapState::eUnknown;
+		FinalOverlapCandidatesIndex = INDEX_NONE;
+	}
+
+	if (bWasForcing)
+	{
+		CurrentOverlapState = EOverlapState::eForceUpdate;
 	}
 }
 
+
 void FScopedMovementUpdate::OnInnerScopeComplete(const FScopedMovementUpdate& InnerScope)
 {
 	if (IsValid(Owner))
@@ -2366,51 +2373,91 @@ void FScopedMovementUpdate::OnInnerScopeComplete(const FScopedMovementUpdate& In
 		if (InnerScope.HasMoved(EHasMovedTransformOption::eTestTransform))
 		{
 			bHasMoved = true;
-			AppendOverlapsAfterMove(InnerScope.GetPendingOverlaps(), InnerScope.GetOverlapsAtEnd());
+			
+			if (InnerScope.CurrentOverlapState == EOverlapState::eUseParent)
+			{
+				// Unchanged, use our own
+			}
+			else
+			{
+				// Bubble up from inner scope.
+				CurrentOverlapState = InnerScope.CurrentOverlapState;
+				if (InnerScope.FinalOverlapCandidatesIndex == INDEX_NONE)
+				{
+					FinalOverlapCandidatesIndex = INDEX_NONE;
+				}
+				else
+				{
+					checkSlow(InnerScope.GetPendingOverlaps().Num() > 0);
+					FinalOverlapCandidatesIndex = PendingOverlaps.Num() + InnerScope.FinalOverlapCandidatesIndex;
+				}
+				PendingOverlaps.Append(InnerScope.GetPendingOverlaps());
+				checkSlow(FinalOverlapCandidatesIndex < PendingOverlaps.Num());
+			}
 		}
 		else
 		{
 			// Don't want to invalidate a parent scope when nothing changed in the child.
+			checkSlow(InnerScope.CurrentOverlapState == EOverlapState::eUseParent);
 		}
 
 		BlockingHits.Append(InnerScope.GetPendingBlockingHits());
 	}	
 }
 
-bool FScopedMovementUpdate::HasKnownOverlapStateAtEnd() const
+const TArray<FOverlapInfo>* FScopedMovementUpdate::GetOverlapsAtEnd(class UPrimitiveComponent& PrimComponent, TArray<FOverlapInfo>& EndOverlaps, bool bTransformChanged) const
 {
-	const FScopedMovementUpdate* CurrentScope = this;
-	do
+	const TArray<FOverlapInfo>* EndOverlapsPtr = nullptr;
+	switch (CurrentOverlapState)
 	{
-		if (CurrentScope->HasMoved(EHasMovedTransformOption::eTestTransform))
+		case FScopedMovementUpdate::EOverlapState::eUseParent:
+		{
+			// Only rotation could have possibly changed
+			if (bTransformChanged && PrimComponent.AreSymmetricRotations(InitialTransform.GetRotation(), PrimComponent.GetComponentQuat(), PrimComponent.GetComponentScale()))
+			{
+				EndOverlapsPtr = PrimComponent.ConvertRotationOverlapsToCurrentOverlaps(EndOverlaps, PrimComponent.GetOverlapInfos());
+			}
+			else
+			{
+				// Use current overlaps (unchanged)
+				EndOverlapsPtr = &PrimComponent.GetOverlapInfos();
+			}
+			break;
+		}
+		case FScopedMovementUpdate::EOverlapState::eUnknown:
+		case FScopedMovementUpdate::EOverlapState::eForceUpdate:
 		{
-			return bHasStoredOverlapsAtEnd;
+			EndOverlapsPtr = nullptr;
+			break;
 		}
-
-		CurrentScope = CurrentScope->GetOuterDeferredScope();
-	}
-	while (CurrentScope != nullptr);
-
-	// Top level scope that has not moved. Current state of overlaps is same as component.
-	return true;
-}
-
-const TArray<FOverlapInfo>* FScopedMovementUpdate::GetKnownOverlapsAtEnd(const TArray<FOverlapInfo>* DefaultOverlaps) const
-{
-	const FScopedMovementUpdate* CurrentScope = this;
-	do
-	{
-		if (CurrentScope->HasMoved(EHasMovedTransformOption::eTestTransform))
+		case FScopedMovementUpdate::EOverlapState::eIncludesOverlaps:
+		{
+			if (FinalOverlapCandidatesIndex == INDEX_NONE)
+			{
+				// Overlapping nothing
+				EndOverlapsPtr = &EndOverlaps;
+			}
+			else
+			{
+				// Fill in EndOverlaps with overlaps valid at the end location.
+				const bool bMatchingScale = InitialTransform.GetScale3D().Equals(PrimComponent.GetComponentScale());
+				if (bMatchingScale)
+				{
+					EndOverlapsPtr = PrimComponent.ConvertSweptOverlapsToCurrentOverlaps(
+						EndOverlaps, GetPendingOverlaps(), FinalOverlapCandidatesIndex,
+						PrimComponent.GetComponentLocation(), PrimComponent.GetComponentQuat());
+				}
+			}
+			break;
+		}
+		default:
 		{
-			return CurrentScope->GetOverlapsAtEnd();
+			checkf(false, TEXT("Unknown FScopedMovementUpdate::EOverlapState value"));
+			break;
 		}
-
-		CurrentScope = CurrentScope->GetOuterDeferredScope();
 	}
-	while (CurrentScope != nullptr);
 
-	// Top level scope that has not moved.
-	return DefaultOverlaps;
+	return EndOverlapsPtr;
 }
 
 
