Fix Mac build.

diff --git a/Engine/Source/Runtime/Engine/Private/PhysicsEngine/PhysXSupport.h b/Engine/Source/Runtime/Engine/Private/PhysicsEngine/PhysXSupport.h
--- a/Engine/Source/Runtime/Engine/Private/PhysicsEngine/PhysXSupport.h
+++ b/Engine/Source/Runtime/Engine/Private/PhysicsEngine/PhysXSupport.h
@@ -413,8 +413,8 @@ class FPhysxSharedData
 {
 public:
 	static FPhysxSharedData& Get(){ return *Singleton; }
-	static void FPhysxSharedData::Initialize();
-	static void FPhysxSharedData::Terminate();
+	static void Initialize();
+	static void Terminate();
 
 	void Add(PxBase* Obj);
 	void Remove(PxBase* Obj)	{ if(Obj) { SharedObjects->remove(*Obj); } }
