Fix simulated destructible mesh not respecting set world location in begin play.

UE-15684

diff --git a/Engine/Source/Runtime/Engine/Private/Components/DestructibleComponent.cpp b/Engine/Source/Runtime/Engine/Private/Components/DestructibleComponent.cpp
--- a/Engine/Source/Runtime/Engine/Private/Components/DestructibleComponent.cpp
+++ b/Engine/Source/Runtime/Engine/Private/Components/DestructibleComponent.cpp
@@ -159,7 +159,7 @@ void UDestructibleComponent::OnUpdateTransform(bool bSkipPhysicsMove, bool bTele
 #endif
 
 #if WITH_APEX
-	if (ApexDestructibleActor != NULL && !BodyInstance.bSimulatePhysics)
+	if (ApexDestructibleActor)
 	{
 		PxRigidDynamic* PRootActor = ApexDestructibleActor->getChunkPhysXActor(0);
 		PxMat44 GlobalPose(PxMat33(U2PQuat(CurrentLocalToWorld.GetRotation())), U2PVector(CurrentLocalToWorld.GetTranslation()));
