Fix complex collision that is movable not being set as part of physics simulation. Was not correctly setting actors as kinematic.

UE-15899

diff --git a/Engine/Source/Runtime/Engine/Private/PhysicsEngine/BodyInstance.cpp b/Engine/Source/Runtime/Engine/Private/PhysicsEngine/BodyInstance.cpp
--- a/Engine/Source/Runtime/Engine/Private/PhysicsEngine/BodyInstance.cpp
+++ b/Engine/Source/Runtime/Engine/Private/PhysicsEngine/BodyInstance.cpp
@@ -1055,7 +1055,7 @@ struct FInitBodiesHelper
 
 	physx::PxRigidActor* CreateActor_PhysX_AssumesLocked(FBodyInstance* Instance, const PxTransform& PTransform) const
 	{
-		physx::PxRigidActor* PNewDynamic = nullptr;
+		physx::PxRigidDynamic* PNewDynamic = nullptr;
 
 		if (bCompileStatic || bStatic)
 		{
@@ -1076,6 +1076,11 @@ struct FInitBodiesHelper
 			{
 				Instance->RigidActorSync = PNewDynamic;
 			}
+
+			if(!Instance->ShouldInstanceSimulatingPhysics())
+			{
+				PNewDynamic->setRigidDynamicFlag(PxRigidDynamicFlag::eKINEMATIC, true);
+			}
 		}
 
 		return PNewDynamic;
