cleanup

diff --git a/Engine/Source/Runtime/Engine/Private/Components/BillboardComponent.cpp b/Engine/Source/Runtime/Engine/Private/Components/BillboardComponent.cpp
--- a/Engine/Source/Runtime/Engine/Private/Components/BillboardComponent.cpp
+++ b/Engine/Source/Runtime/Engine/Private/Components/BillboardComponent.cpp
@@ -212,7 +212,6 @@ public:
 		FPrimitiveViewRelevance Result;
 		Result.bDrawRelevance = IsShown(View) && bVisible;
 		Result.bOpaqueRelevance = true;
-		Result.bNormalTranslucencyRelevance = false;
 		Result.bDynamicRelevance = true;
 		Result.bShadowRelevance = IsShadowCast(View);
 		Result.bEditorPrimitiveRelevance = UseEditorCompositing(View);
