[UE-15470] Add stacktrace to Blueprint array error messages (with -scriptstackonwarnings option)

PR #1159
=> https://github.com/EpicGames/UnrealEngine/pull/1159

diff --git a/Engine/Source/Runtime/Engine/Private/KismetArrayLibrary.cpp b/Engine/Source/Runtime/Engine/Private/KismetArrayLibrary.cpp
--- a/Engine/Source/Runtime/Engine/Private/KismetArrayLibrary.cpp
+++ b/Engine/Source/Runtime/Engine/Private/KismetArrayLibrary.cpp
@@ -92,7 +92,7 @@ void UKismetArrayLibrary::GenericArray_Insert(void* TargetArray, const UArrayPro
 		}
 		else
 		{
-			UE_LOG(LogArray, Warning, TEXT("Attempted to insert an item into array %s out of bounds [%d/%d]!"), *ArrayProp->GetName(), Index, GetLastIndex(ArrayHelper));
+			FFrame::KismetExecutionMessage(*FString::Printf(TEXT("Attempted to insert an item into array %s out of bounds [%d/%d]!"), *ArrayProp->GetName(), Index, GetLastIndex(ArrayHelper)), ELogVerbosity::Warning);
 		}
 	}
 }
@@ -108,7 +108,7 @@ void UKismetArrayLibrary::GenericArray_Remove(void* TargetArray, const UArrayPro
 		}
 		else
 		{
-			UE_LOG(LogArray, Warning, TEXT("Attempted to remove an item from an invalid index from array %s [%d/%d]!"), *ArrayProp->GetName(), IndexToRemove, GetLastIndex(ArrayHelper));
+			FFrame::KismetExecutionMessage(*FString::Printf(TEXT("Attempted to remove an item from an invalid index from array %s [%d/%d]!"), *ArrayProp->GetName(), IndexToRemove, GetLastIndex(ArrayHelper)), ELogVerbosity::Warning);
 		}
 	}
 }
@@ -172,7 +172,7 @@ void UKismetArrayLibrary::GenericArray_Resize(void* TargetArray, const UArrayPro
 		}
 		else
 		{
-			UE_LOG(LogArray, Warning, TEXT("Attempted to resize an array using negative size: Array = %s, Size = %d!"), *ArrayProp->GetName(), Size);
+			FFrame::KismetExecutionMessage(*FString::Printf(TEXT("Attempted to resize an array using negative size: Array = %s, Size = %d!"), *ArrayProp->GetName(), Size), ELogVerbosity::Warning);
 		}
 	}
 }
@@ -214,7 +214,7 @@ void UKismetArrayLibrary::GenericArray_Get(void* TargetArray, const UArrayProper
 		}
 		else
 		{
-			UE_LOG(LogArray, Warning, TEXT("Attempted to get an item from array %s out of bounds [%d/%d]!"), *ArrayProp->GetName(), Index, GetLastIndex(ArrayHelper));
+			FFrame::KismetExecutionMessage(*FString::Printf(TEXT("Attempted to get an item from array %s out of bounds [%d/%d]!"), *ArrayProp->GetName(), Index, GetLastIndex(ArrayHelper)), ELogVerbosity::Warning);
 			InnerProp->InitializeValue(Item);
 		}
 	}
@@ -240,7 +240,7 @@ void UKismetArrayLibrary::GenericArray_Set(void* TargetArray, const UArrayProper
 		}
 		else
 		{
-			UE_LOG(LogArray, Warning, TEXT("Attempted to set an invalid index on array %s [%d/%d]!"), *ArrayProp->GetName(), Index, GetLastIndex(ArrayHelper));
+			FFrame::KismetExecutionMessage(*FString::Printf(TEXT("Attempted to set an invalid index on array %s [%d/%d]!"), *ArrayProp->GetName(), Index, GetLastIndex(ArrayHelper)), ELogVerbosity::Warning);
 		}
 	}
 }
