#ue4 - Add FBodyInstance::OverlapTestForBody(). Avoids heap allocation when only method was one that took a TArray.

- Templatized OverlapTestForBodies() to allow TArray with different allocator types (namely TInlineAllocator for OverlapTestForBody()).

diff --git a/Engine/Source/Runtime/Engine/Classes/PhysicsEngine/BodyInstance.h b/Engine/Source/Runtime/Engine/Classes/PhysicsEngine/BodyInstance.h
--- a/Engine/Source/Runtime/Engine/Classes/PhysicsEngine/BodyInstance.h
+++ b/Engine/Source/Runtime/Engine/Classes/PhysicsEngine/BodyInstance.h
@@ -911,6 +911,7 @@ public:
 	 *  @return true if any of the bodies passed in overlap with this
 	 */
 	bool OverlapTestForBodies(const FVector& Position, const FQuat& Rotation, const TArray<FBodyInstance*>& Bodies) const;
+	bool OverlapTestForBody(const FVector& Position, const FQuat& Rotation, FBodyInstance* Body) const;
 
 	/**
 	 *  Determines the set of components that this body instance would overlap with at the supplied location/rotation
@@ -1011,6 +1012,9 @@ private:
 	 */
 	static bool IsValidCollisionProfileName(FName InCollisionProfileName);
 
+	template<typename AllocatorType>
+	bool OverlapTestForBodiesImpl(const FVector& Position, const FQuat& Rotation, const TArray<FBodyInstance*, AllocatorType>& Bodies) const;
+
 	friend class UCollisionProfile;
 	friend class FBodyInstanceCustomization;
 	
@@ -1054,9 +1058,20 @@ private:
 //////////////////////////////////////////////////////////////////////////
 // BodyInstance inlines
 
-inline bool FBodyInstance::OverlapMulti(TArray<struct FOverlapResult>& InOutOverlaps, const class UWorld* World, const FTransform* pWorldToComponent, const FVector& Pos, const FRotator& Rot, ECollisionChannel TestChannel, const struct FComponentQueryParams& Params, const struct FCollisionResponseParams& ResponseParams, const struct FCollisionObjectQueryParams& ObjectQueryParams) const
+FORCEINLINE_DEBUGGABLE bool FBodyInstance::OverlapMulti(TArray<struct FOverlapResult>& InOutOverlaps, const class UWorld* World, const FTransform* pWorldToComponent, const FVector& Pos, const FRotator& Rot, ECollisionChannel TestChannel, const struct FComponentQueryParams& Params, const struct FCollisionResponseParams& ResponseParams, const struct FCollisionObjectQueryParams& ObjectQueryParams) const
 {
 	// Pass on to FQuat version
 	return OverlapMulti(InOutOverlaps, World, pWorldToComponent, Pos, Rot.Quaternion(), TestChannel, Params, ResponseParams, ObjectQueryParams);
 }
 
+FORCEINLINE_DEBUGGABLE bool FBodyInstance::OverlapTestForBodies(const FVector& Position, const FQuat& Rotation, const TArray<FBodyInstance*>& Bodies) const
+{
+	return OverlapTestForBodiesImpl(Position, Rotation, Bodies);
+}
+
+FORCEINLINE_DEBUGGABLE bool FBodyInstance::OverlapTestForBody(const FVector& Position, const FQuat& Rotation, FBodyInstance* Body) const
+{
+	TArray<FBodyInstance*, TInlineAllocator<1>> InlineArray;
+	InlineArray.Add(Body);
+	return OverlapTestForBodiesImpl(Position, Rotation, InlineArray);
+}
diff --git a/Engine/Source/Runtime/Engine/Private/Components/PrimitiveComponent.cpp b/Engine/Source/Runtime/Engine/Private/Components/PrimitiveComponent.cpp
--- a/Engine/Source/Runtime/Engine/Private/Components/PrimitiveComponent.cpp
+++ b/Engine/Source/Runtime/Engine/Private/Components/PrimitiveComponent.cpp
@@ -1785,9 +1785,7 @@ bool UPrimitiveComponent::ComponentOverlapComponentImpl(class UPrimitiveComponen
 
 	if(FBodyInstance* BI = PrimComp->GetBodyInstance())
 	{
-		TArray<FBodyInstance*> Bodies;
-		Bodies.Add(GetBodyInstance());
-		return BI->OverlapTestForBodies(Pos, Quat, Bodies);
+		return BI->OverlapTestForBody(Pos, Quat, GetBodyInstance());
 	}
 
 	return false;
diff --git a/Engine/Source/Runtime/Engine/Private/PhysicsEngine/BodyInstance.cpp b/Engine/Source/Runtime/Engine/Private/PhysicsEngine/BodyInstance.cpp
--- a/Engine/Source/Runtime/Engine/Private/PhysicsEngine/BodyInstance.cpp
+++ b/Engine/Source/Runtime/Engine/Private/PhysicsEngine/BodyInstance.cpp
@@ -4017,7 +4017,8 @@ float FBodyInstance::GetDistanceToBody(const FVector& Point, FVector& OutPointOn
 	//@TODO: BOX2D: Implement DistanceToBody
 }
 
-bool FBodyInstance::OverlapTestForBodies(const FVector& Pos, const FQuat& Rot, const TArray<FBodyInstance*>& Bodies) const
+template <typename AllocatorType>
+bool FBodyInstance::OverlapTestForBodiesImpl(const FVector& Pos, const FQuat& Rot, const TArray<FBodyInstance*, AllocatorType>& Bodies) const
 {
 	bool bHaveOverlap = false;
 #if WITH_PHYSX
@@ -4058,6 +4059,11 @@ bool FBodyInstance::OverlapTestForBodies(const FVector& Pos, const FQuat& Rot, c
 	return bHaveOverlap;
 }
 
+// Explicit template instantiation for the above.
+template bool FBodyInstance::OverlapTestForBodiesImpl(const FVector& Pos, const FQuat& Rot, const TArray<FBodyInstance*>& Bodies) const;
+template bool FBodyInstance::OverlapTestForBodiesImpl(const FVector& Pos, const FQuat& Rot, const TArray<FBodyInstance*, TInlineAllocator<1>>& Bodies) const;
+
+
 bool FBodyInstance::OverlapTest(const FVector& Position, const FQuat& Rotation, const struct FCollisionShape& CollisionShape, FMTDResult* OutMTD) const
 {
 	bool bHasOverlap = false;
