Allow changing the static mesh on UCameraComponent

Merging 2561299

//depot/UE4-UT/...

to //depot/UE4/...

diff --git a/Engine/Source/Runtime/Engine/Classes/Camera/CameraComponent.h b/Engine/Source/Runtime/Engine/Classes/Camera/CameraComponent.h
--- a/Engine/Source/Runtime/Engine/Classes/Camera/CameraComponent.h
+++ b/Engine/Source/Runtime/Engine/Classes/Camera/CameraComponent.h
@@ -90,10 +90,15 @@ protected:
 
 	UPROPERTY(transient)
 	class UStaticMesh* CameraMesh;
-
+	
 	// The camera mesh to show visually where the camera is placed
 	UPROPERTY(transient)
 	class UStaticMeshComponent* ProxyMeshComponent;
+
+public:
+
+	ENGINE_API virtual void SetCameraMesh(UStaticMesh* Mesh) { CameraMesh = Mesh; }
+
 #endif
 public:
 #if WITH_EDITORONLY_DATA
