DistanceToSurface node has an optional input position

diff --git a/Engine/Source/Runtime/Engine/Classes/Materials/MaterialExpressionDistanceFieldGradient.h b/Engine/Source/Runtime/Engine/Classes/Materials/MaterialExpressionDistanceFieldGradient.h
--- a/Engine/Source/Runtime/Engine/Classes/Materials/MaterialExpressionDistanceFieldGradient.h
+++ b/Engine/Source/Runtime/Engine/Classes/Materials/MaterialExpressionDistanceFieldGradient.h
@@ -10,6 +10,9 @@ class UMaterialExpressionDistanceFieldGradient : public UMaterialExpression
 {
 	GENERATED_UCLASS_BODY()
 
+	UPROPERTY(meta = (RequiredInput = "false", ToolTip = "Defaults to current world position if not specified"))
+	FExpressionInput Position;
+
 	// Begin UMaterialExpression Interface
 	virtual int32 Compile(class FMaterialCompiler* Compiler, int32 OutputIndex, int32 MultiplexIndex) override;
 	virtual void GetCaption(TArray<FString>& OutCaptions) const override;
diff --git a/Engine/Source/Runtime/Engine/Classes/Materials/MaterialExpressionDistanceToNearestSurface.h b/Engine/Source/Runtime/Engine/Classes/Materials/MaterialExpressionDistanceToNearestSurface.h
--- a/Engine/Source/Runtime/Engine/Classes/Materials/MaterialExpressionDistanceToNearestSurface.h
+++ b/Engine/Source/Runtime/Engine/Classes/Materials/MaterialExpressionDistanceToNearestSurface.h
@@ -10,6 +10,9 @@ class UMaterialExpressionDistanceToNearestSurface : public UMaterialExpression
 {
 	GENERATED_UCLASS_BODY()
 
+	UPROPERTY(meta = (RequiredInput = "false", ToolTip = "Defaults to current world position if not specified"))
+	FExpressionInput Position;
+
 	// Begin UMaterialExpression Interface
 	virtual int32 Compile(class FMaterialCompiler* Compiler, int32 OutputIndex, int32 MultiplexIndex) override;
 	virtual void GetCaption(TArray<FString>& OutCaptions) const override;
diff --git a/Engine/Source/Runtime/Engine/Private/Materials/HLSLMaterialTranslator.h b/Engine/Source/Runtime/Engine/Private/Materials/HLSLMaterialTranslator.h
--- a/Engine/Source/Runtime/Engine/Private/Materials/HLSLMaterialTranslator.h
+++ b/Engine/Source/Runtime/Engine/Private/Materials/HLSLMaterialTranslator.h
@@ -3638,28 +3638,38 @@ protected:
 		return AddCodeChunk( MCT_Float3, TEXT("MaterialExpressionBlackBody(%s)"), *GetParameterCode(Temp) );
 	}
 
-	virtual int32 DistanceToNearestSurface() override
+	virtual int32 DistanceToNearestSurface(int32 PositionArg) override
 	{
 		if (ErrorUnlessFeatureLevelSupported(ERHIFeatureLevel::SM5) == INDEX_NONE)
 		{
 			return INDEX_NONE;
 		}
 
+		if (PositionArg == INDEX_NONE)
+		{
+			return INDEX_NONE;
+		}
+
 		MaterialCompilationOutput.bUsesGlobalDistanceField = true;
 
-		return AddCodeChunk(MCT_Float, TEXT("GetDistanceToNearestSurfaceGlobal(Parameters.WorldPosition)"));
+		return AddCodeChunk(MCT_Float, TEXT("GetDistanceToNearestSurfaceGlobal(%s)"), *GetParameterCode(PositionArg));
 	}
 
-	virtual int32 DistanceFieldGradient() override
+	virtual int32 DistanceFieldGradient(int32 PositionArg) override
 	{
 		if (ErrorUnlessFeatureLevelSupported(ERHIFeatureLevel::SM5) == INDEX_NONE)
 		{
 			return INDEX_NONE;
 		}
 
+		if (PositionArg == INDEX_NONE)
+		{
+			return INDEX_NONE;
+		}
+
 		MaterialCompilationOutput.bUsesGlobalDistanceField = true;
 
-		return AddCodeChunk(MCT_Float3, TEXT("GetDistanceFieldGradientGlobal(Parameters.WorldPosition)"));
+		return AddCodeChunk(MCT_Float3, TEXT("GetDistanceFieldGradientGlobal(%s)"), *GetParameterCode(PositionArg));
 	}
 
 	virtual int32 AtmosphericFogColor( int32 WorldPosition ) override
diff --git a/Engine/Source/Runtime/Engine/Private/Materials/MaterialExpressions.cpp b/Engine/Source/Runtime/Engine/Private/Materials/MaterialExpressions.cpp
--- a/Engine/Source/Runtime/Engine/Private/Materials/MaterialExpressions.cpp
+++ b/Engine/Source/Runtime/Engine/Private/Materials/MaterialExpressions.cpp
@@ -8274,7 +8274,18 @@ UMaterialExpressionDistanceToNearestSurface::UMaterialExpressionDistanceToNeares
 
 int32 UMaterialExpressionDistanceToNearestSurface::Compile(class FMaterialCompiler* Compiler, int32 OutputIndex, int32 MultiplexIndex)
 {
-	return Compiler->DistanceToNearestSurface();
+	int32 PositionArg = INDEX_NONE;
+
+	if (Position.Expression)
+	{
+		PositionArg = Position.Compile(Compiler);
+	}
+	else 
+	{
+		PositionArg = Compiler->WorldPosition(WPT_Default);
+	}
+
+	return Compiler->DistanceToNearestSurface(PositionArg);
 }
 
 void UMaterialExpressionDistanceToNearestSurface::GetCaption(TArray<FString>& OutCaptions) const
@@ -8305,7 +8316,18 @@ UMaterialExpressionDistanceFieldGradient::UMaterialExpressionDistanceFieldGradie
 
 int32 UMaterialExpressionDistanceFieldGradient::Compile(class FMaterialCompiler* Compiler, int32 OutputIndex, int32 MultiplexIndex)
 {
-	return Compiler->DistanceFieldGradient();
+	int32 PositionArg = INDEX_NONE;
+
+	if (Position.Expression)
+	{
+		PositionArg = Position.Compile(Compiler);
+	}
+	else 
+	{
+		PositionArg = Compiler->WorldPosition(WPT_Default);
+	}
+
+	return Compiler->DistanceFieldGradient(PositionArg);
 }
 
 void UMaterialExpressionDistanceFieldGradient::GetCaption(TArray<FString>& OutCaptions) const
diff --git a/Engine/Source/Runtime/Engine/Public/MaterialCompiler.h b/Engine/Source/Runtime/Engine/Public/MaterialCompiler.h
--- a/Engine/Source/Runtime/Engine/Public/MaterialCompiler.h
+++ b/Engine/Source/Runtime/Engine/Public/MaterialCompiler.h
@@ -173,8 +173,8 @@ public:
 	virtual int32 AntialiasedTextureMask(int32 Tex, int32 UV, float Threshold, uint8 Channel) = 0;
 	virtual int32 Noise(int32 Position, float Scale, int32 Quality, uint8 NoiseFunction, bool bTurbulence, int32 Levels, float OutputMin, float OutputMax, float LevelScale, int32 FilterWidth) = 0;
 	virtual int32 BlackBody( int32 Temp ) = 0;
-	virtual int32 DistanceToNearestSurface() = 0;
-	virtual int32 DistanceFieldGradient() = 0;
+	virtual int32 DistanceToNearestSurface(int32 PositionArg) = 0;
+	virtual int32 DistanceFieldGradient(int32 PositionArg) = 0;
 	virtual int32 DepthOfFieldFunction(int32 Depth, int32 FunctionValueIndex) = 0;
 	virtual int32 AtmosphericFogColor(int32 WorldPosition) = 0;
 	virtual int32 SpeedTree(ESpeedTreeGeometryType GeometryType, ESpeedTreeWindType WindType, ESpeedTreeLODType LODType, float BillboardThreshold, bool bAccurateWindVelocities) = 0;
@@ -336,8 +336,8 @@ public:
 		return Compiler->Noise(Position, Scale, Quality, NoiseFunction, bTurbulence, Levels, OutputMin, OutputMax, LevelScale, FilterWidth);
 	}
 	virtual int32 BlackBody( int32 Temp ) override { return Compiler->BlackBody(Temp); }
-	virtual int32 DistanceToNearestSurface() override { return Compiler->DistanceToNearestSurface(); }
-	virtual int32 DistanceFieldGradient() override { return Compiler->DistanceFieldGradient(); }
+	virtual int32 DistanceToNearestSurface(int32 PositionArg) override { return Compiler->DistanceToNearestSurface(PositionArg); }
+	virtual int32 DistanceFieldGradient(int32 PositionArg) override { return Compiler->DistanceFieldGradient(PositionArg); }
 	virtual int32 PerInstanceRandom() override { return Compiler->PerInstanceRandom(); }
 	virtual int32 PerInstanceFadeAmount() override { return Compiler->PerInstanceFadeAmount(); }
 	virtual int32 DepthOfFieldFunction(int32 Depth, int32 FunctionValueIndex) override
