PR #812: Support for custom physics projection with FBodyInstance (Contributed by [REDACTED_PERSON])

UE-9687

diff --git a/Engine/Source/Runtime/Engine/Classes/PhysicsEngine/BodyInstance.h b/Engine/Source/Runtime/Engine/Classes/PhysicsEngine/BodyInstance.h
--- a/Engine/Source/Runtime/Engine/Classes/PhysicsEngine/BodyInstance.h
+++ b/Engine/Source/Runtime/Engine/Classes/PhysicsEngine/BodyInstance.h
@@ -19,6 +19,13 @@ class UPhysicsConstraintComponent;
   * Do not expect this callback to be called from the main game thread! It may get called from a physics simulation thread. */
 DECLARE_DELEGATE_TwoParams(FCalculateCustomPhysics, float, FBodyInstance*);
 
+/** Delegate for applying custom physics projection upon the body. When this is set for the body instance,
+  * it will be called whenever component transformation is requested from the physics engine. If
+  * projection is required (for example, visual position of an object must be different to the one in physics engine,
+  * e.g. the box should not penetrate the wall visually) the transformation of body must be updated to account for it.
+  * Since this could be called many times by GetWorldTransform any expensive computations should be cached if possible.*/
+DECLARE_DELEGATE_TwoParams(FCalculateCustomProjection, const FBodyInstance*, FTransform&);
+
 #if WITH_PHYSX
 struct FShapeData;
 
@@ -711,14 +718,17 @@ public:
 	/** Enables/disables whether this body is affected by gravity. */
 	void SetEnableGravity(bool bGravityEnabled);
 
+	/** Custom projection for physics (callback to update component transform based on physics data) */
+	FCalculateCustomProjection OnCalculateCustomProjection;
+
 	/** See if this body is valid. */
 	bool IsValidBodyInstance() const;
 
 	/** Get current transform in world space from physics body. */
-	FTransform GetUnrealWorldTransform() const;
+	FTransform GetUnrealWorldTransform(bool bWithProjection = true) const;
 
 	/** Get current transform in world space from physics body. */
-	FTransform GetUnrealWorldTransform_AssumesLocked() const;
+	FTransform GetUnrealWorldTransform_AssumesLocked(bool bWithProjection = true) const;
 
 	/**
 	 *	Move the physics body to a new pose.
diff --git a/Engine/Source/Runtime/Engine/Private/PhysicsEngine/BodyInstance.cpp b/Engine/Source/Runtime/Engine/Private/PhysicsEngine/BodyInstance.cpp
--- a/Engine/Source/Runtime/Engine/Private/PhysicsEngine/BodyInstance.cpp
+++ b/Engine/Source/Runtime/Engine/Private/PhysicsEngine/BodyInstance.cpp
@@ -587,7 +587,7 @@ void FBodyInstance::CreateDOFLock()
 		DOFConstraint->bTwistLimitSoft = false;
 		DOFConstraint->bLinearLimitSoft = false;
 
-		const FTransform TM = GetUnrealWorldTransform();
+		const FTransform TM = GetUnrealWorldTransform(false);
 		FVector Normal = FVector(1, 0, 0);
 		FVector Sec = FVector(0, 1, 0);
 
@@ -1839,7 +1839,7 @@ bool FBodyInstance::Weld(FBodyInstance* TheirBody, const FTransform& TheirTM)
 #if WITH_PHYSX
 	TArray<PxShape *> PNewShapes;
 
-	FTransform MyTM = GetUnrealWorldTransform();
+	FTransform MyTM = GetUnrealWorldTransform(false);
 	MyTM.SetScale3D(Scale3D);	//physx doesn't store 3d so set it here
 
 	FTransform RelativeTM = TheirTM.GetRelativeTransform(MyTM);
@@ -2523,7 +2523,7 @@ bool FBodyInstance::IsValidBodyInstance() const
 }
 
 template <bool NeedsLock>
-FTransform GetUnrealWorldTransformImp(const FBodyInstance* BodyInstance)
+FTransform GetUnrealWorldTransformImp(const FBodyInstance* BodyInstance, bool bWithProjection)
 {
 	FTransform WorldTM = FTransform::Identity;
 #if WITH_PHYSX
@@ -2531,6 +2531,11 @@ FTransform GetUnrealWorldTransformImp(const FBodyInstance* BodyInstance)
 	{
 		PxTransform PTM = PActor->getGlobalPose();
 		WorldTM = P2UTransform(PTM);
+
+		if(bWithProjection)
+		{
+			BodyInstance->OnCalculateCustomProjection.ExecuteIfBound(BodyInstance, WorldTM);
+		}
 	});
 #endif // WITH_PHYSX
 
@@ -2544,21 +2549,26 @@ FTransform GetUnrealWorldTransformImp(const FBodyInstance* BodyInstance)
 		const FRotator Rotation3D(FMath::RadiansToDegrees(RotationInRadians), 0.0f, 0.0f); //@TODO: BOX2D: Should be moved to FPhysicsIntegration2D
 
 		WorldTM = FTransform(Rotation3D, Translation3D, BodyInstance->Scale3D);
+
+		if (bWithProjection)
+		{
+			BodyInstance->OnCalculateCustomProjection.ExecuteIfBound(BodyInstance, WorldTM);
+		}
 	}
 #endif
 
 	return WorldTM;
 }
 
-FTransform FBodyInstance::GetUnrealWorldTransform() const
+FTransform FBodyInstance::GetUnrealWorldTransform(bool bWithProjection /* = true*/) const
 {
-	return GetUnrealWorldTransformImp<true>(this);
+	return GetUnrealWorldTransformImp<true>(this, bWithProjection);
 }
 
 
-FTransform FBodyInstance::GetUnrealWorldTransform_AssumesLocked() const
+FTransform FBodyInstance::GetUnrealWorldTransform_AssumesLocked(bool bWithProjection /* = true*/) const
 {
-	return GetUnrealWorldTransformImp<false>(this);
+	return GetUnrealWorldTransformImp<false>(this, bWithProjection);
 }
 
 void FBodyInstance::SetBodyTransform(const FTransform& NewTransform, bool bTeleport)
