PR #938: Give games the ability to write custom logic for delta seconds to use for ticking. (Contributed by [REDACTED_PERSON])

UE-12139

diff --git a/Engine/Source/Runtime/Engine/Classes/GameFramework/WorldSettings.h b/Engine/Source/Runtime/Engine/Classes/GameFramework/WorldSettings.h
--- a/Engine/Source/Runtime/Engine/Classes/GameFramework/WorldSettings.h
+++ b/Engine/Source/Runtime/Engine/Classes/GameFramework/WorldSettings.h
@@ -494,7 +494,12 @@ public:
 		return TimeDilation * MatineeTimeDilation * DemoPlayTimeDilation;
 	}
 
-	/** 
+	/**
+	 * Returns the delta time to be used by the tick. Can be overridden if game specific logic is needed.
+	 */
+	virtual float FixupDeltaSeconds(float DeltaSeconds, float RealDeltaSeconds);
+
+	/**
 	 * Called by GameMode.HandleMatchIsWaitingToStart, calls BeginPlay on all actors
 	 */
 	virtual void NotifyBeginPlay();
diff --git a/Engine/Source/Runtime/Engine/Private/LevelTick.cpp b/Engine/Source/Runtime/Engine/Private/LevelTick.cpp
--- a/Engine/Source/Runtime/Engine/Private/LevelTick.cpp
+++ b/Engine/Source/Runtime/Engine/Private/LevelTick.cpp
@@ -1119,8 +1119,12 @@ void UWorld::Tick( ELevelTick TickType, float DeltaSeconds )
 
 	// apply time multipliers
 	DeltaSeconds *= Info->GetEffectiveTimeDilation();
-	// Clamp time between 2000 fps and 2.5 fps.
-	DeltaSeconds = FMath::Clamp(DeltaSeconds,0.0005f,0.40f);
+
+	// Handle clamping of time to an acceptable value
+	const float GameDeltaSeconds = Info->FixupDeltaSeconds(DeltaSeconds, RealDeltaSeconds);
+	check(GameDeltaSeconds >= 0.0f);
+
+	DeltaSeconds = GameDeltaSeconds;
 	DeltaTimeSeconds = DeltaSeconds;
 
 	if ( !bIsPaused )
diff --git a/Engine/Source/Runtime/Engine/Private/WorldSettings.cpp b/Engine/Source/Runtime/Engine/Private/WorldSettings.cpp
--- a/Engine/Source/Runtime/Engine/Private/WorldSettings.cpp
+++ b/Engine/Source/Runtime/Engine/Private/WorldSettings.cpp
@@ -116,6 +116,12 @@ float AWorldSettings::GetGravityZ() const
 	return WorldGravityZ;
 }
 
+float AWorldSettings::FixupDeltaSeconds(float DeltaSeconds, float RealDeltaSeconds)
+{
+	// Clamp time between 2000 fps and 2.5 fps.
+	return FMath::Clamp(DeltaSeconds, 0.0005f, 0.4f);	
+}
+
 void AWorldSettings::NotifyBeginPlay()
 {
 	UWorld* World = GetWorld();
