Allow skipping the queued bunch throttling mechanism by setting net.ProcessQueuedBunchesMillisecondLimit to 0. This is a short-term fix for UE-15696.

diff --git a/Engine/Source/Runtime/Engine/Private/DataChannel.cpp b/Engine/Source/Runtime/Engine/Private/DataChannel.cpp
--- a/Engine/Source/Runtime/Engine/Private/DataChannel.cpp
+++ b/Engine/Source/Runtime/Engine/Private/DataChannel.cpp
@@ -32,10 +32,10 @@ static TAutoConsoleVariable<int32> CVarNetReliableDebug(
 	TEXT(" 2: Print reliable bunch buffer each net update"),
 	ECVF_Default);
 
-static TAutoConsoleVariable<float> CVarNetProcessQueuedBunchesMillisecondLimit(
+static TAutoConsoleVariable<int> CVarNetProcessQueuedBunchesMillisecondLimit(
 	TEXT("net.ProcessQueuedBunchesMillisecondLimit"),
-	30.0f,
-	TEXT("Time threshold for processing queued bunches. If it takes longer than this in a single frame, wait until the next frame to continue processing queued bunches."));
+	30,
+	TEXT("Time threshold for processing queued bunches. If it takes longer than this in a single frame, wait until the next frame to continue processing queued bunches. For unlimited time, set to 0."));
 
 /*-----------------------------------------------------------------------------
 	UChannel implementation.
@@ -1743,7 +1743,8 @@ bool UActorChannel::ProcessQueuedBunches()
 		}
 	}
 
-	const bool bHasTimeToProcess = Connection->Driver->ProcessQueuedBunchesCurrentFrameMilliseconds < CVarNetProcessQueuedBunchesMillisecondLimit.GetValueOnGameThread();
+	const int BunchTimeLimit = CVarNetProcessQueuedBunchesMillisecondLimit.GetValueOnGameThread();
+	const bool bHasTimeToProcess = BunchTimeLimit == 0 || Connection->Driver->ProcessQueuedBunchesCurrentFrameMilliseconds < BunchTimeLimit;
 
 	// We can process all of the queued up bunches if ALL of these are true:
 	//	1. We have queued bunches to process
