Assert to gather more information about 'Assertion failed: Value->TextureReference.TextureReferenceRHI'

diff --git a/Engine/Source/Runtime/Engine/Private/Materials/MaterialUniformExpressions.cpp b/Engine/Source/Runtime/Engine/Private/Materials/MaterialUniformExpressions.cpp
--- a/Engine/Source/Runtime/Engine/Private/Materials/MaterialUniformExpressions.cpp
+++ b/Engine/Source/Runtime/Engine/Private/Materials/MaterialUniformExpressions.cpp
@@ -352,10 +352,11 @@ FUniformBufferRHIRef FUniformExpressionSet::CreateUniformBuffer(const FMaterialR
 			if (Value && Value->Resource)
 			{
 				//@todo-rco: Help track down a invalid values
-				checkf(Value->IsValidLowLevel() && Value->IsA(UTexture::StaticClass()), TEXT("Expecting a UTexture! Value='%s' class='%s'"), *Value->GetName(), *Value->GetClass()->GetName());
+				checkf(Value->IsA(UTexture::StaticClass()), TEXT("Expecting a UTexture! Value='%s' class='%s'"), *Value->GetName(), *Value->GetClass()->GetName());
 
 				// UMaterial / UMaterialInstance should have caused all dependent textures to be PostLoaded, which initializes their rendering resource
-				check(Value->TextureReference.TextureReferenceRHI);
+				checkf(Value->TextureReference.TextureReferenceRHI, TEXT("Texture %s of class %s had invalid texture reference"), *Value->GetName(), *Value->GetClass()->GetName());
+
 				*ResourceTable++ = Value->TextureReference.TextureReferenceRHI;
 				FSamplerStateRHIRef* SamplerSource = &Value->Resource->SamplerStateRHI;
 
