Fix shape sharing incorrectly cloning shapes when updating filter data.

UE-16025

diff --git a/Engine/Source/Runtime/Engine/Private/PhysicsEngine/BodyInstance.cpp b/Engine/Source/Runtime/Engine/Private/PhysicsEngine/BodyInstance.cpp
--- a/Engine/Source/Runtime/Engine/Private/PhysicsEngine/BodyInstance.cpp
+++ b/Engine/Source/Runtime/Engine/Private/PhysicsEngine/BodyInstance.cpp
@@ -741,87 +741,94 @@ void FBodyInstance::UpdatePhysicsShapeFilterData(uint32 SkelMeshCompID, bool bUs
 				}
 			}
 
-			PShape->setSimulationFilterData(PSimFilterData);
 
-			// If query collision is enabled..
-			if (UseCollisionEnabled != ECollisionEnabled::NoCollision)
+			ExecuteOnPxShapeWrite(this, PShape, [&](PxShape* PGivenShape)
 			{
-				const bool bQueryEnabled = ((UseCollisionEnabled == ECollisionEnabled::QueryAndPhysics) || (UseCollisionEnabled == ECollisionEnabled::QueryOnly));
-
-				// Only perform scene queries in the synchronous scene for static shapes
-				if (bPhysicsStatic)
-				{
-					bool bIsSyncShape = (ShapeIdx < NumSyncShapes);
-					PShape->setFlag(PxShapeFlag::eSCENE_QUERY_SHAPE, bIsSyncShape && bQueryEnabled);
-				}
-				// If non-static, enable scene queries if requested
-				else
-				{
-					PShape->setFlag(PxShapeFlag::eSCENE_QUERY_SHAPE, bQueryEnabled);
-				}
-
-				// See if we want physics collision
-				const bool bSimCollision = ((UseCollisionEnabled == ECollisionEnabled::QueryAndPhysics) || (UseCollisionEnabled == ECollisionEnabled::PhysicsOnly));
+				PGivenShape->setSimulationFilterData(PSimFilterData);
 
-				// Triangle mesh is 'complex' geom
-				if (PShape->getGeometryType() == PxGeometryType::eTRIANGLEMESH)
+				// If query collision is enabled..
+				if (UseCollisionEnabled != ECollisionEnabled::NoCollision)
 				{
-					PShape->setQueryFilterData(PComplexQueryData);
+					const bool bQueryEnabled = ((UseCollisionEnabled == ECollisionEnabled::QueryAndPhysics) || (UseCollisionEnabled == ECollisionEnabled::QueryOnly));
 
-					// on dynamic objects and objects which don't use complex as simple, tri mesh not used for sim
-					if (!bSimCollision || !bUseComplexAsSimple)
+					// Only perform scene queries in the synchronous scene for static shapes
+					if (bPhysicsStatic)
 					{
-						PShape->setFlag(PxShapeFlag::eSIMULATION_SHAPE, false);
+						bool bIsSyncShape = (ShapeIdx < NumSyncShapes);
+						PGivenShape->setFlag(PxShapeFlag::eSCENE_QUERY_SHAPE, bIsSyncShape && bQueryEnabled);
 					}
+					// If non-static, enable scene queries if requested
 					else
 					{
-						PShape->setFlag(PxShapeFlag::eSIMULATION_SHAPE, true);
+						PGivenShape->setFlag(PxShapeFlag::eSCENE_QUERY_SHAPE, bQueryEnabled);
 					}
 
-					if (OwnerPrimitiveComponent == NULL || !OwnerPrimitiveComponent->IsA(UModelComponent::StaticClass()))
-					{
-						PShape->setFlag(PxShapeFlag::eVISUALIZATION, false); // dont draw the tri mesh, we can see it anyway, and its slow
-					}
-				}
-				// Everything else is 'simple'
-				else
-				{
-					PShape->setQueryFilterData(PSimpleQueryData);
+					// See if we want physics collision
+					const bool bSimCollision = ((UseCollisionEnabled == ECollisionEnabled::QueryAndPhysics) || (UseCollisionEnabled == ECollisionEnabled::PhysicsOnly));
 
-					// See if we currently have sim collision
-					const bool bCurrentSimCollision = (PShape->getFlags() & PxShapeFlag::eSIMULATION_SHAPE);
-					// Enable sim collision
-					if (bSimCollision && !bCurrentSimCollision)
-					{
-						bUpdateMassProperties = true;
-						PShape->setFlag(PxShapeFlag::eSIMULATION_SHAPE, true);
-					}
-					// Disable sim collision
-					else if (!bSimCollision && bCurrentSimCollision)
+					// Triangle mesh is 'complex' geom
+					if (PGivenShape->getGeometryType() == PxGeometryType::eTRIANGLEMESH)
 					{
-						bUpdateMassProperties = true;
-						PShape->setFlag(PxShapeFlag::eSIMULATION_SHAPE, false);
-					}
+						PGivenShape->setQueryFilterData(PComplexQueryData);
 
-					// enable swept bounds for CCD for this shape
-					PxRigidBody* PBody = PActor->is<PxRigidBody>();
-					if (bSimCollision && !bPhysicsStatic && bUseCCD && PBody)
-					{
-						PBody->setRigidBodyFlag(PxRigidBodyFlag::eENABLE_CCD, true);
+						// on dynamic objects and objects which don't use complex as simple, tri mesh not used for sim
+						if (!bSimCollision || !bUseComplexAsSimple)
+						{
+							PGivenShape->setFlag(PxShapeFlag::eSIMULATION_SHAPE, false);
+						}
+						else
+						{
+							PGivenShape->setFlag(PxShapeFlag::eSIMULATION_SHAPE, true);
+						}
+
+						if (OwnerPrimitiveComponent == NULL || !OwnerPrimitiveComponent->IsA(UModelComponent::StaticClass()))
+						{
+							PGivenShape->setFlag(PxShapeFlag::eVISUALIZATION, false); // dont draw the tri mesh, we can see it anyway, and its slow
+						}
 					}
-					else if (PBody)
+					// Everything else is 'simple'
+					else
 					{
+						PGivenShape->setQueryFilterData(PSimpleQueryData);
+
+						// See if we currently have sim collision
+						const bool bCurrentSimCollision = (PGivenShape->getFlags() & PxShapeFlag::eSIMULATION_SHAPE);
+						// Enable sim collision
+						if (bSimCollision && !bCurrentSimCollision)
+						{
+							bUpdateMassProperties = true;
+							PGivenShape->setFlag(PxShapeFlag::eSIMULATION_SHAPE, true);
+						}
+						// Disable sim collision
+						else if (!bSimCollision && bCurrentSimCollision)
+						{
+							bUpdateMassProperties = true;
+							PGivenShape->setFlag(PxShapeFlag::eSIMULATION_SHAPE, false);
+						}
+
+						// enable swept bounds for CCD for this shape
+						PxRigidBody* PBody = PActor->is<PxRigidBody>();
+						if (bSimCollision && !bPhysicsStatic && bUseCCD && PBody)
+						{
+							PBody->setRigidBodyFlag(PxRigidBodyFlag::eENABLE_CCD, true);
+						}
+						else if (PBody)
+						{
 
-						PBody->setRigidBodyFlag(PxRigidBodyFlag::eENABLE_CCD, false);
+							PBody->setRigidBodyFlag(PxRigidBodyFlag::eENABLE_CCD, false);
+						}
 					}
 				}
-			}
-			// No collision enabled
-			else
-			{
-				PShape->setFlag(PxShapeFlag::eSIMULATION_SHAPE, false);
-				PShape->setFlag(PxShapeFlag::eSCENE_QUERY_SHAPE, false);
-			}
+				// No collision enabled
+				else
+				{
+					PGivenShape->setFlag(PxShapeFlag::eSIMULATION_SHAPE, false);
+					PGivenShape->setFlag(PxShapeFlag::eSCENE_QUERY_SHAPE, false);
+				}
+			});
+
+
+			
 		}
 
 		if (bUpdateMassProperties)
