Rotation support and component filtering added to spring node

API Change: Deprecation of bNoZSpring flag (Use bTranslateZ instead)

diff --git a/Engine/Source/Editor/AnimGraph/Classes/AnimGraphNode_SpringBone.h b/Engine/Source/Editor/AnimGraph/Classes/AnimGraphNode_SpringBone.h
--- a/Engine/Source/Editor/AnimGraph/Classes/AnimGraphNode_SpringBone.h
+++ b/Engine/Source/Editor/AnimGraph/Classes/AnimGraphNode_SpringBone.h
@@ -14,6 +14,10 @@ class UAnimGraphNode_SpringBone : public UAnimGraphNode_SkeletalControlBase
 	UPROPERTY(EditAnywhere, Category=Settings)
 	FAnimNode_SpringBone Node;
 
+	// Begin UObject interface
+	virtual void PostLoad();
+	// End UObject interface
+
 public:
 	// UEdGraphNode interface
 	virtual FText GetNodeTitle(ENodeTitleType::Type TitleType) const override;
diff --git a/Engine/Source/Editor/AnimGraph/Private/AnimGraphNode_SpringBone.cpp b/Engine/Source/Editor/AnimGraph/Private/AnimGraphNode_SpringBone.cpp
--- a/Engine/Source/Editor/AnimGraph/Private/AnimGraphNode_SpringBone.cpp
+++ b/Engine/Source/Editor/AnimGraph/Private/AnimGraphNode_SpringBone.cpp
@@ -13,6 +13,17 @@ UAnimGraphNode_SpringBone::UAnimGraphNode_SpringBone(const FObjectInitializer& O
 {
 }
 
+void UAnimGraphNode_SpringBone::PostLoad()
+{
+	Super::PostLoad();
+
+	if (GetLinkerUE4Version() < VER_UE4_REPLACE_SPRING_NOZ_PROPERTY)
+	{
+		//Fix up z translation from old NoZSpring property
+		Node.bTranslateZ = !Node.bNoZSpring_DEPRECATED;
+	}
+}
+
 FText UAnimGraphNode_SpringBone::GetControllerDescription() const
 {
 	return LOCTEXT("SpringController", "Spring controller");
diff --git a/Engine/Source/Runtime/Core/Public/UObject/ObjectVersion.h b/Engine/Source/Runtime/Core/Public/UObject/ObjectVersion.h
--- a/Engine/Source/Runtime/Core/Public/UObject/ObjectVersion.h
+++ b/Engine/Source/Runtime/Core/Public/UObject/ObjectVersion.h
@@ -518,6 +518,8 @@ enum EUnrealEngineObjectUE4Version
 	VER_UE4_ADD_BLEND_MODE_TO_WIDGET_COMPONENT,
 	// Added lightmass primitive setting
 	VER_UE4_NEW_LIGHTMASS_PRIMITIVE_SETTING,
+	// Deprecate NoZSpring property on spring nodes to be replaced with TranslateZ property
+	VER_UE4_REPLACE_SPRING_NOZ_PROPERTY,
 
 	// -----<new versions can be added before this line>-------------------------------------------------
 	// - this needs to be the last line (see note below)
diff --git a/Engine/Source/Runtime/Engine/Classes/Animation/BoneControllers/AnimNode_SpringBone.h b/Engine/Source/Runtime/Engine/Classes/Animation/BoneControllers/AnimNode_SpringBone.h
--- a/Engine/Source/Runtime/Engine/Classes/Animation/BoneControllers/AnimNode_SpringBone.h
+++ b/Engine/Source/Runtime/Engine/Classes/Animation/BoneControllers/AnimNode_SpringBone.h
@@ -39,8 +39,32 @@ struct ENGINE_API FAnimNode_SpringBone : public FAnimNode_SkeletalControlBase
 	float ErrorResetThresh;
 
 	/** If true, Z position is always correct, no spring applied */
-	UPROPERTY(EditAnywhere, BlueprintReadWrite, Category=Spring)
-	bool bNoZSpring;
+	UPROPERTY()
+	bool bNoZSpring_DEPRECATED;
+
+	/** If true take the spring calculation for translation in X */
+	UPROPERTY(EditAnywhere, BlueprintReadWrite, Category=FilterChannels)
+	bool bTranslateX;
+
+	/** If true take the spring calculation for translation in Y */
+	UPROPERTY(EditAnywhere, BlueprintReadWrite, Category=FilterChannels)
+	bool bTranslateY;
+
+	/** If true take the spring calculation for translation in Z */
+	UPROPERTY(EditAnywhere, BlueprintReadWrite, Category=FilterChannels)
+	bool bTranslateZ;
+
+	/** If true take the spring calculation for rotation in X */
+	UPROPERTY(EditAnywhere, BlueprintReadWrite, Category=FilterChannels)
+	bool bRotateX;
+
+	/** If true take the spring calculation for rotation in Y */
+	UPROPERTY(EditAnywhere, BlueprintReadWrite, Category=FilterChannels)
+	bool bRotateY;
+
+	/** If true take the spring calculation for rotation in Z */
+	UPROPERTY(EditAnywhere, BlueprintReadWrite, Category=FilterChannels)
+	bool bRotateZ;
 
 	/** Internal use - Amount of time we need to simulate. */
 	float RemainingTime;
diff --git a/Engine/Source/Runtime/Engine/Private/Animation/AnimNode_SpringBone.cpp b/Engine/Source/Runtime/Engine/Private/Animation/AnimNode_SpringBone.cpp
--- a/Engine/Source/Runtime/Engine/Private/Animation/AnimNode_SpringBone.cpp
+++ b/Engine/Source/Runtime/Engine/Private/Animation/AnimNode_SpringBone.cpp
@@ -12,11 +12,17 @@ FAnimNode_SpringBone::FAnimNode_SpringBone()
 	, SpringStiffness(50.0f)
 	, SpringDamping(4.0f)
 	, ErrorResetThresh(256.0f)
-	, bNoZSpring(false)
+	, bNoZSpring_DEPRECATED(false)
 	, RemainingTime(0.f)
 	, bHadValidStrength(false)
 	, BoneLocation(FVector::ZeroVector)
 	, BoneVelocity(FVector::ZeroVector)
+	, bTranslateX(true)
+	, bTranslateY(true)
+	, bTranslateZ(true)
+	, bRotateX(false)
+	, bRotateY(false)
+	, bRotateZ(false)
 {
 }
 
@@ -57,10 +63,32 @@ void FAnimNode_SpringBone::GatherDebugData(FNodeDebugData& DebugData)
 	ComponentPose.GatherDebugData(DebugData);
 }
 
+FORCEINLINE void CopyToVectorByFlags(FVector& DestVec, const FVector& SrcVec, bool bX, bool bY, bool bZ)
+{
+	if (bX)
+	{
+		DestVec.X = SrcVec.X;
+	}
+	if (bY)
+	{
+		DestVec.Y = SrcVec.Y;
+	}
+	if (bZ)
+	{
+		DestVec.Z = SrcVec.Z;
+	}
+}
+
 void FAnimNode_SpringBone::EvaluateBoneTransforms(USkeletalMeshComponent* SkelComp, FCSPose<FCompactPose>& MeshBases, TArray<FBoneTransform>& OutBoneTransforms)
 {
 	check(OutBoneTransforms.Num() == 0);
 
+	const bool bNoOffset = !bTranslateX && !bTranslateY && !bTranslateZ;
+	if (bNoOffset)
+	{
+		return;
+	}
+
 	// Location of our bone in world space
 	const FBoneContainer& BoneContainer = MeshBases.GetPose().GetBoneContainer();
 
@@ -129,11 +157,9 @@ void FAnimNode_SpringBone::EvaluateBoneTransforms(USkeletalMeshComponent* SkelCo
 		FVector const DeltaMove = (BoneVelocity * FixedTimeStep);
 		BoneLocation += DeltaMove;
 
-		// Force z to be correct if desired
-		if (bNoZSpring)
-		{
-			BoneLocation.Z = TargetPos.Z;
-		}
+		// Filter out spring translation based on our filter properties
+		CopyToVectorByFlags(BoneLocation, TargetPos, !bTranslateX, !bTranslateY, !bTranslateZ);
+
 
 		// If desired, limit error
 		if (bLimitDisplacement)
@@ -160,6 +186,24 @@ void FAnimNode_SpringBone::EvaluateBoneTransforms(USkeletalMeshComponent* SkelCo
 	FTransform OutBoneTM = SpaceBase;
 	OutBoneTM.SetLocation( SkelComp->GetComponentToWorld().InverseTransformPosition(BoneLocation) );
 
+	const bool bUseRotation = bRotateX || bRotateY || bRotateZ;
+	if (bUseRotation)
+	{
+		FCompactPoseBoneIndex ParentBoneIndex = MeshBases.GetPose().GetParentBoneIndex(SpringBoneIndex);
+		const FTransform& ParentSpaceBase = MeshBases.GetComponentSpaceTransform(ParentBoneIndex);
+
+		FVector ParentToTarget = (TargetPos - ParentSpaceBase.GetLocation()).GetSafeNormal();
+		FVector ParentToCurrent = (BoneLocation - ParentSpaceBase.GetLocation()).GetSafeNormal();
+
+		FQuat AdditionalRotation = FQuat::FindBetween(ParentToTarget, ParentToCurrent);
+
+		// Filter rotation based on our filter properties
+		FVector EularRot = AdditionalRotation.Euler();
+		CopyToVectorByFlags(EularRot, FVector::ZeroVector, !bRotateX, !bRotateY, !bRotateZ);
+
+		OutBoneTM.SetRotation(FQuat::MakeFromEuler(EularRot) * OutBoneTM.GetRotation());
+	}
+
 	// Output new transform for current bone.
 	OutBoneTransforms.Add(FBoneTransform(SpringBoneIndex, OutBoneTM));
 }
