Mitigated rare Persona crash when the preview instance is freed, added extra logging to catch.

diff --git a/Engine/Source/Runtime/Engine/Private/Animation/AnimNode_StateMachine.cpp b/Engine/Source/Runtime/Engine/Private/Animation/AnimNode_StateMachine.cpp
--- a/Engine/Source/Runtime/Engine/Private/Animation/AnimNode_StateMachine.cpp
+++ b/Engine/Source/Runtime/Engine/Private/Animation/AnimNode_StateMachine.cpp
@@ -263,6 +263,14 @@ void FAnimNode_StateMachine::Update(const FAnimationUpdateContext& Context)
 		{
 			return;
 		}
+		else if(!Machine->States.IsValidIndex(CurrentState))
+		{
+			// Attempting to catch a crash where the state machine has been freed.
+			UE_LOG(LogAnimation, Warning, TEXT("FAnimNode_StateMachine::Update - Invalid current state, please report. Attempting to use state %d in state machine %d"), CurrentState, StateMachineIndexInClass);
+			UE_LOG(LogAnimation, Warning, TEXT("\t\tWhen updating AnimInstance: %s"), *Context.AnimInstance->GetName())
+
+			return;
+		}
 	}
 	else
 	{
@@ -557,7 +565,7 @@ void FAnimNode_StateMachine::Evaluate(FPoseContext& Output)
 {
 	if (FBakedAnimationStateMachine* Machine = GetMachineDescription())
 	{
-		if (Machine->States.Num() == 0)
+		if (Machine->States.Num() == 0 || !Machine->States.IsValidIndex(CurrentState))
 		{
 			Output.Pose.ResetToRefPose();
 			return;
