#ue4 - Perf: Avoid Quat->Rotator->Quat conversion for every physics component that moves during simulation.

diff --git a/Engine/Source/Runtime/Engine/Private/PhysicsEngine/PhysScene.cpp b/Engine/Source/Runtime/Engine/Private/PhysicsEngine/PhysScene.cpp
--- a/Engine/Source/Runtime/Engine/Private/PhysicsEngine/PhysScene.cpp
+++ b/Engine/Source/Runtime/Engine/Private/PhysicsEngine/PhysScene.cpp
@@ -795,7 +795,7 @@ void FPhysScene::SyncComponentsToBodies_AssumesLocked(uint32 SceneType)
 		if (!NewTransform.EqualsNoScale(BodyInstance->OwnerComponent->ComponentToWorld))
 		{
 			const FVector MoveBy = NewTransform.GetLocation() - BodyInstance->OwnerComponent->ComponentToWorld.GetLocation();
-			const FRotator NewRotation = NewTransform.Rotator();
+			const FQuat NewRotation = NewTransform.GetRotation();
 
 			//@warning: do not reference BodyInstance again after calling MoveComponent() - events from the move could have made it unusable (destroying the actor, SetPhysics(), etc)
 			BodyInstance->OwnerComponent->MoveComponent(MoveBy, NewRotation, false, NULL, MOVECOMP_SkipPhysicsMove);
