Use negative scale for trimeshes instead of a negative cooked mesh. Saves on memory and streaming time.

UE-16069

diff --git a/Engine/Source/Runtime/Engine/Classes/PhysicsEngine/BodySetup.h b/Engine/Source/Runtime/Engine/Classes/PhysicsEngine/BodySetup.h
--- a/Engine/Source/Runtime/Engine/Classes/PhysicsEngine/BodySetup.h
+++ b/Engine/Source/Runtime/Engine/Classes/PhysicsEngine/BodySetup.h
@@ -122,9 +122,6 @@ class UBodySetup : public UObject
 #if WITH_PHYSX
 	/** Physics triangle mesh, created from cooked data in CreatePhysicsMeshes */
 	physx::PxTriangleMesh* TriMesh;
-
-	/** Physics triangle mesh, flipped across X, created from cooked data in CreatePhysicsMeshes */
-	physx::PxTriangleMesh* TriMeshNegX;
 #endif
 
 	/** Flag used to know if we have created the physics convex and tri meshes from the cooked data yet */
diff --git a/Engine/Source/Runtime/Engine/Private/PhysicsEngine/BodyInstance.cpp b/Engine/Source/Runtime/Engine/Private/PhysicsEngine/BodyInstance.cpp
--- a/Engine/Source/Runtime/Engine/Private/PhysicsEngine/BodyInstance.cpp
+++ b/Engine/Source/Runtime/Engine/Private/PhysicsEngine/BodyInstance.cpp
@@ -2272,32 +2272,19 @@ bool FBodyInstance::UpdateBodyScale(const FVector& InScale3D)
 					PShape->getTriangleMeshGeometry(PTriMeshGeom);
 
 					// Create tri-mesh shape
-					if (BodySetup.IsValid() && (BodySetup->TriMesh != NULL || BodySetup->TriMeshNegX != NULL))
+					if (BodySetup.IsValid() && BodySetup->TriMesh)
 					{
 						// Please note that this one we don't inverse old scale, but just set new one (but still adjust for scale mode)
 						FVector NewScale3D = RelativeScale3D * OldScale3D;
-						FVector Scale3DAbs(FMath::Abs(NewScale3D.X), FMath::Abs(NewScale3D.Y), FMath::Abs(NewScale3D.Z)); // magnitude of scale (sign removed)
-
-						PxTransform PNewLocalPose;
-						bool bUseNegX = CalcMeshNegScaleCompensation(NewScale3D, PNewLocalPose);
-
-						// Only case where TriMeshNegX should be null is BSP, which should not require negX version
-						if (bUseNegX && BodySetup->TriMeshNegX == NULL)
-						{
-							UE_LOG(LogPhysics, Warning, TEXT("FBodyInstance::UpdateBodyScale: Want to use NegX but it doesn't exist! %s"), *BodySetup->GetPathName());
-						}
-
-						PxTriangleMesh* UseTriMesh = bUseNegX ? BodySetup->TriMeshNegX : BodySetup->TriMesh;
-						if (UseTriMesh != NULL)
+						if (BodySetup->TriMesh)
 						{
-							PTriMeshGeom.triangleMesh = bUseNegX ? BodySetup->TriMeshNegX : BodySetup->TriMesh;
-							PTriMeshGeom.scale.scale = U2PVector(Scale3DAbs);
+							PTriMeshGeom.triangleMesh = BodySetup->TriMesh;
+							PTriMeshGeom.scale.scale = U2PVector(Scale3D);
 
 							if (PTriMeshGeom.isValid())
 							{
 								UpdatedGeometry = &PTriMeshGeom;
 								bSuccess = true;
-
 							}
 							else
 							{
diff --git a/Engine/Source/Runtime/Engine/Private/PhysicsEngine/BodySetup.cpp b/Engine/Source/Runtime/Engine/Private/PhysicsEngine/BodySetup.cpp
--- a/Engine/Source/Runtime/Engine/Private/PhysicsEngine/BodySetup.cpp
+++ b/Engine/Source/Runtime/Engine/Private/PhysicsEngine/BodySetup.cpp
@@ -168,17 +168,8 @@ void UBodySetup::CreatePhysicsMeshes()
 			}
 		}
 
-		if(bGenerateNonMirroredCollision)
-		{
-			TriMesh = CookedDataReader.TriMesh;
-			FPhysxSharedData::Get().Add(TriMesh);
-		}
-
-		if(bGenerateMirroredCollision)
-		{
-			TriMeshNegX = CookedDataReader.TriMeshNegX;
-			FPhysxSharedData::Get().Add(TriMeshNegX);
-		}
+		TriMesh = CookedDataReader.TriMesh;
+		FPhysxSharedData::Get().Add(TriMesh);
 	}
 	else
 	{
@@ -227,13 +218,6 @@ void UBodySetup::ClearPhysicsMeshes()
 		TriMesh = NULL;
 	}
 
-	if(TriMeshNegX != NULL)
-	{
-		GPhysXPendingKillTriMesh.Add(TriMeshNegX);
-		FPhysxSharedData::Get().Remove(TriMeshNegX);
-		TriMeshNegX = NULL;
-	}
-
 	bCreatedPhysicsMeshes = false;
 #endif
 
@@ -537,45 +521,31 @@ public:
 		float ContactOffsetFactor, MaxContactOffset;
 		GetContactOffsetParams(ContactOffsetFactor, MaxContactOffset);
 
-		if (BodySetup->TriMesh || BodySetup->TriMeshNegX)
+		if (BodySetup->TriMesh)
 		{
-			PxTransform PLocalPose;
-			bool bUseNegX = CalcMeshNegScaleCompensation(Scale3D, PLocalPose);
-
-			// Only case where TriMeshNegX should be null is BSP, which should not require negX version
-			if (bUseNegX && BodySetup->TriMeshNegX == NULL)
+			PxTriangleMeshGeometry PTriMeshGeom;
+			PTriMeshGeom.triangleMesh = BodySetup->TriMesh;
+			PTriMeshGeom.scale.scale = U2PVector(Scale3D);
+			if (BodySetup->bDoubleSidedGeometry)
 			{
-				UE_LOG(LogPhysics, Log, TEXT("AddTriMeshToRigidActor: Want to use NegX but it doesn't exist! %s"), *BodySetup->GetPathName());
+				PTriMeshGeom.meshFlags |= PxMeshGeometryFlag::eDOUBLE_SIDED;
 			}
 
-			PxTriangleMesh* UseTriMesh = bUseNegX ? BodySetup->TriMeshNegX : BodySetup->TriMesh;
-			if (UseTriMesh != NULL)
-			{
-				PxTriangleMeshGeometry PTriMeshGeom;
-				PTriMeshGeom.triangleMesh = bUseNegX ? BodySetup->TriMeshNegX : BodySetup->TriMesh;
-				PTriMeshGeom.scale.scale = U2PVector(Scale3DAbs);
-				if (BodySetup->bDoubleSidedGeometry)
-				{
-					PTriMeshGeom.meshFlags |= PxMeshGeometryFlag::eDOUBLE_SIDED;
-				}
 
+			if (PTriMeshGeom.isValid())
+			{
 
-				if (PTriMeshGeom.isValid())
+				// Create without 'sim shape' flag, problematic if it's kinematic, and it gets set later anyway.
 				{
-					ensure(PLocalPose.isValid());
-
-					// Create without 'sim shape' flag, problematic if it's kinematic, and it gets set later anyway.
+					if (!AttachShape_AssumesLocked(PTriMeshGeom, PxTransform::createIdentity(), MaxContactOffset, PxShapeFlag::eSCENE_QUERY_SHAPE | PxShapeFlag::eVISUALIZATION))
 					{
-						if (!AttachShape_AssumesLocked(PTriMeshGeom, PLocalPose, MaxContactOffset, PxShapeFlag::eSCENE_QUERY_SHAPE | PxShapeFlag::eVISUALIZATION))
-						{
-							UE_LOG(LogPhysics, Log, TEXT("Can't create new mesh shape in AddShapesToRigidActor"));
-						}
+						UE_LOG(LogPhysics, Log, TEXT("Can't create new mesh shape in AddShapesToRigidActor"));
 					}
 				}
-				else
-				{
-					UE_LOG(LogPhysics, Log, TEXT("AddTriMeshToRigidActor: TriMesh invalid"));
-				}
+			}
+			else
+			{
+				UE_LOG(LogPhysics, Log, TEXT("AddTriMeshToRigidActor: TriMesh invalid"));
 			}
 		}
 	}
@@ -997,11 +967,6 @@ SIZE_T UBodySetup::GetResourceSize( EResourceSizeMode::Type Mode )
 		ResourceSize += GetPhysxObjectSize(TriMesh, NULL);
 	}
 
-	if (TriMeshNegX != NULL)
-	{
-		ResourceSize += GetPhysxObjectSize(TriMeshNegX, NULL);
-	}
-
 	// Count PhysX convex mem usage
 	for(int ConvIdx=0; ConvIdx<AggGeom.ConvexElems.Num(); ConvIdx++)
 	{
diff --git a/Engine/Source/Runtime/Engine/Private/PhysicsEngine/PhysDerivedData.cpp b/Engine/Source/Runtime/Engine/Private/PhysicsEngine/PhysDerivedData.cpp
--- a/Engine/Source/Runtime/Engine/Private/PhysicsEngine/PhysDerivedData.cpp
+++ b/Engine/Source/Runtime/Engine/Private/PhysicsEngine/PhysDerivedData.cpp
@@ -57,13 +57,11 @@ bool FDerivedDataPhysXCooker::Build( TArray<uint8>& OutData )
 	int32 NumConvexElementsCooked = 0;
 	int32 NumMirroredElementsCooked = 0;
 	uint8 bTriMeshCooked = false;
-	uint8 bMirroredTriMeshCooked = false;
 	Ar << bLittleEndian;
 	int64 CookedMeshInfoOffset = Ar.Tell();
 	Ar << NumConvexElementsCooked;	
 	Ar << NumMirroredElementsCooked;
 	Ar << bTriMeshCooked;
-	Ar << bMirroredTriMeshCooked;
 
 	//TODO: we must save an id with convex and tri meshes for serialization. We must save this here and patch it up at runtime somehow
 
@@ -84,14 +82,7 @@ bool FDerivedDataPhysXCooker::Build( TArray<uint8>& OutData )
 	bool bUsingAllTriData = BodySetup != NULL ? BodySetup->bMeshCollideAll : false;
 	if( (BodySetup == NULL || BodySetup->CollisionTraceFlag != CTF_UseSimpleAsComplex) && ShouldGenerateTriMeshData(bUsingAllTriData) )
 	{
-		if( bGenerateNormalMesh )
-		{
-			bTriMeshCooked = (uint8)BuildTriMesh( OutData, false, bUsingAllTriData );
-		}
-		if( bGenerateMirroredMesh && ShouldGenerateNegXTriMeshData() )
-		{
-			bMirroredTriMeshCooked = (uint8)BuildTriMesh( OutData, true, bUsingAllTriData );
-		}
+		bTriMeshCooked = (uint8)BuildTriMesh( OutData, false, bUsingAllTriData );
 	}
 
 	// Update info on what actually got cooked
@@ -99,7 +90,6 @@ bool FDerivedDataPhysXCooker::Build( TArray<uint8>& OutData )
 	Ar << NumConvexElementsCooked;	
 	Ar << NumMirroredElementsCooked;
 	Ar << bTriMeshCooked;
-	Ar << bMirroredTriMeshCooked;
 
 	// Whatever got cached return true. We want to cache 'failure' too.
 	return true;
@@ -153,14 +143,6 @@ bool FDerivedDataPhysXCooker::ShouldGenerateTriMeshData(bool InUseAllTriData)
 	return bPerformCook;
 }
 
-bool FDerivedDataPhysXCooker::ShouldGenerateNegXTriMeshData()
-{
-	IInterface_CollisionDataProvider* CDP = Cast<IInterface_CollisionDataProvider>(CollisionDataProvider);
-	const bool bWantsNegX = ( CDP != NULL ) ? CDP->WantsNegXTriMesh() : false;
-	return bWantsNegX;
-}
-
-
 bool FDerivedDataPhysXCooker::BuildTriMesh( TArray<uint8>& OutData, bool bInMirrored, bool InUseAllTriData )
 {
 	check(Cooker != NULL);
diff --git a/Engine/Source/Runtime/Engine/Private/PhysicsEngine/PhysDerivedData.h b/Engine/Source/Runtime/Engine/Private/PhysicsEngine/PhysDerivedData.h
--- a/Engine/Source/Runtime/Engine/Private/PhysicsEngine/PhysDerivedData.h
+++ b/Engine/Source/Runtime/Engine/Private/PhysicsEngine/PhysDerivedData.h
@@ -38,7 +38,7 @@ public:
 		// This is a version string that mimics the old versioning scheme. If you
 		// want to bump this version, generate a new guid using VS->Tools->Create GUID and
 		// return it here. Ex.
-		return TEXT("{1F0627AE-ABEB-4206-8D78-E16BEB5DDC7E}");
+		return TEXT("7F9691E5320E4745AE426E72F60B484D");
 	}
 
 	virtual FString GetPluginSpecificCacheKeySuffix() const override
@@ -83,7 +83,6 @@ private:
 	int32 BuildConvex( TArray<uint8>& OutData, bool InMirrored );
 	bool BuildTriMesh( TArray<uint8>& OutData, bool InMirrored, bool InUseAllTriData );
 	bool ShouldGenerateTriMeshData(bool InUseAllTriData);
-	bool ShouldGenerateNegXTriMeshData();
 };
 
 //////////////////////////////////////////////////////////////////////////
diff --git a/Engine/Source/Runtime/Engine/Private/PhysicsEngine/PhysXSupport.cpp b/Engine/Source/Runtime/Engine/Private/PhysicsEngine/PhysXSupport.cpp
--- a/Engine/Source/Runtime/Engine/Private/PhysicsEngine/PhysXSupport.cpp
+++ b/Engine/Source/Runtime/Engine/Private/PhysicsEngine/PhysXSupport.cpp
@@ -645,7 +645,6 @@ PxU32 FPhysXCPUDispatcherSingleThread::getWorkerCount() const
 
 FPhysXFormatDataReader::FPhysXFormatDataReader( FByteBulkData& InBulkData )
 	: TriMesh( NULL )
-	, TriMeshNegX( NULL )
 {
 	// Read cooked physics data
 	uint8* DataPtr = (uint8*)InBulkData.Lock( LOCK_READ_ONLY );
@@ -655,14 +654,12 @@ FPhysXFormatDataReader::FPhysXFormatDataReader( FByteBulkData& InBulkData )
 	int32 NumConvexElementsCooked = 0;
 	int32 NumMirroredElementsCooked = 0;
 	uint8 bTriMeshCooked = false;
-	uint8 bMirroredTriMeshCooked = false;
 
 	Ar << bLittleEndian;
 	Ar.SetByteSwapping( PLATFORM_LITTLE_ENDIAN ? !bLittleEndian : !!bLittleEndian );
 	Ar << NumConvexElementsCooked;	
 	Ar << NumMirroredElementsCooked;
 	Ar << bTriMeshCooked;
-	Ar << bMirroredTriMeshCooked;
 	
 	ConvexMeshes.Empty( NumConvexElementsCooked );
 	ConvexMeshesNegX.Empty( NumMirroredElementsCooked );
@@ -685,12 +682,6 @@ FPhysXFormatDataReader::FPhysXFormatDataReader( FByteBulkData& InBulkData )
 		check(TriMesh);	
 	}
 
-	if( bMirroredTriMeshCooked )
-	{
-		TriMeshNegX = ReadTriMesh( Ar, DataPtr, InBulkData.GetBulkDataSize() );
-		check(TriMeshNegX);	
-	}
-
 	InBulkData.Unlock();
 }
 
@@ -924,7 +915,6 @@ PxCollection* MakePhysXCollection(const TArray<UPhysicalMaterial*>& PhysicalMate
 	for (UBodySetup* BodySetup : BodySetups)
 	{
 		AddToCollection(PCollection, BodySetup->TriMesh);
-		AddToCollection(PCollection, BodySetup->TriMeshNegX);
 
 		for (const FKConvexElem& ConvexElem : BodySetup->AggGeom.ConvexElems)
 		{
diff --git a/Engine/Source/Runtime/Engine/Private/PhysicsEngine/PhysXSupport.h b/Engine/Source/Runtime/Engine/Private/PhysicsEngine/PhysXSupport.h
--- a/Engine/Source/Runtime/Engine/Private/PhysicsEngine/PhysXSupport.h
+++ b/Engine/Source/Runtime/Engine/Private/PhysicsEngine/PhysXSupport.h
@@ -491,7 +491,6 @@ public:
 	TArray< PxConvexMesh* > ConvexMeshes;
 	TArray< PxConvexMesh* > ConvexMeshesNegX;
 	PxTriangleMesh* TriMesh;
-	PxTriangleMesh* TriMeshNegX;
 
 	FPhysXFormatDataReader( FByteBulkData& InBulkData );
 
