#ue4 - Perf: Avoid some rotation comparisons in UCharacterMovementComponent::PhysicsRotation() if nothing can possibly change.

diff --git a/Engine/Source/Runtime/Engine/Private/Components/CharacterMovementComponent.cpp b/Engine/Source/Runtime/Engine/Private/Components/CharacterMovementComponent.cpp
--- a/Engine/Source/Runtime/Engine/Private/Components/CharacterMovementComponent.cpp
+++ b/Engine/Source/Runtime/Engine/Private/Components/CharacterMovementComponent.cpp
@@ -4383,27 +4383,27 @@ void UCharacterMovementComponent::PhysicsRotation(float DeltaTime)
 	// Accumulate a desired new rotation.
 	const float AngleTolerance = 1e-3f;
 
-	// PITCH
-	if (!FMath::IsNearlyEqual(CurrentRotation.Pitch, DesiredRotation.Pitch, AngleTolerance))
+	if (!CurrentRotation.Equals(DesiredRotation, AngleTolerance))
 	{
-		DesiredRotation.Pitch = FMath::FixedTurn(CurrentRotation.Pitch, DesiredRotation.Pitch, DeltaRot.Pitch);
-	}
+		// PITCH
+		if (!FMath::IsNearlyEqual(CurrentRotation.Pitch, DesiredRotation.Pitch, AngleTolerance))
+		{
+			DesiredRotation.Pitch = FMath::FixedTurn(CurrentRotation.Pitch, DesiredRotation.Pitch, DeltaRot.Pitch);
+		}
 
-	// YAW
-	if (!FMath::IsNearlyEqual(CurrentRotation.Yaw, DesiredRotation.Yaw, AngleTolerance))
-	{
-		DesiredRotation.Yaw = FMath::FixedTurn(CurrentRotation.Yaw, DesiredRotation.Yaw, DeltaRot.Yaw);
-	}
+		// YAW
+		if (!FMath::IsNearlyEqual(CurrentRotation.Yaw, DesiredRotation.Yaw, AngleTolerance))
+		{
+			DesiredRotation.Yaw = FMath::FixedTurn(CurrentRotation.Yaw, DesiredRotation.Yaw, DeltaRot.Yaw);
+		}
 
-	// ROLL
-	if (!FMath::IsNearlyEqual(CurrentRotation.Roll, DesiredRotation.Roll, AngleTolerance))
-	{
-		DesiredRotation.Roll = FMath::FixedTurn(CurrentRotation.Roll, DesiredRotation.Roll, DeltaRot.Roll);
-	}
+		// ROLL
+		if (!FMath::IsNearlyEqual(CurrentRotation.Roll, DesiredRotation.Roll, AngleTolerance))
+		{
+			DesiredRotation.Roll = FMath::FixedTurn(CurrentRotation.Roll, DesiredRotation.Roll, DeltaRot.Roll);
+		}
 
-	// Set the new rotation.
-	if (!CurrentRotation.Equals(DesiredRotation, AngleTolerance))
-	{
+		// Set the new rotation.
 		MoveUpdatedComponent( FVector::ZeroVector, DesiredRotation, true );
 	}
 }
