Add new warnings for physics operations on an object that has its physics collision type set to No Collision or QueryOnly.
Removed warning about simulating physics with no collision. Although this is a common mistake, it's a valid operation for advanced users.

UE-16076

diff --git a/Engine/Source/Runtime/Engine/Private/Components/PrimitiveComponent.cpp b/Engine/Source/Runtime/Engine/Private/Components/PrimitiveComponent.cpp
--- a/Engine/Source/Runtime/Engine/Private/Components/PrimitiveComponent.cpp
+++ b/Engine/Source/Runtime/Engine/Private/Components/PrimitiveComponent.cpp
@@ -894,12 +894,6 @@ bool UPrimitiveComponent::ShouldCreatePhysicsState() const
 	bool bShouldCreatePhysicsState = IsRegistered() && (bAlwaysCreatePhysicsState || IsCollisionEnabled());
 
 #if WITH_EDITOR
-	if (BodyInstance.bSimulatePhysics && GetCollisionEnabled() == ECollisionEnabled::NoCollision)
-	{
-		FMessageLog("PIE").Warning(FText::Format(LOCTEXT("InvalidSimulateOptions", "Invalid Simulate Options: Body ({0}) is set to simulate physics but Collision Enabled is set to No Collision"),
-			FText::FromString(GetReadableName())));
-	}
-
 	// if it shouldn't create physics state, but if world wants to enable trace collision for components, allow it
 	if (!bShouldCreatePhysicsState && World && World->bEnableTraceCollision)
 	{
diff --git a/Engine/Source/Runtime/Engine/Private/PrimitiveComponentPhysics.cpp b/Engine/Source/Runtime/Engine/Private/PrimitiveComponentPhysics.cpp
--- a/Engine/Source/Runtime/Engine/Private/PrimitiveComponentPhysics.cpp
+++ b/Engine/Source/Runtime/Engine/Private/PrimitiveComponentPhysics.cpp
@@ -164,9 +164,18 @@ void UPrimitiveComponent::WarnInvalidPhysicsOperations_Internal(const FText& Act
 #if !(UE_BUILD_SHIPPING || UE_BUILD_TEST)
 	if (!CheckStaticMobilityAndWarn(ActionText))	//all physics operations require non-static mobility
 	{
-		if (BI && BI->bSimulatePhysics == false)	//some require to be simulating too
+		if (BI)
 		{
-			FMessageLog("PIE").Warning(FText::Format(LOCTEXT("InvalidPhysicsOperation", "{0} has to have 'Simulate Physics' enabled if you'd like to {1}. "), FText::FromString(GetReadableName()), ActionText));
+			ECollisionEnabled::Type CollisionEnabled = BI->GetCollisionEnabled();
+
+			if(BI->bSimulatePhysics == false)	//some require to be simulating too
+			{
+				FMessageLog("PIE").Warning(FText::Format(LOCTEXT("InvalidPhysicsOperation", "{0} has to have 'Simulate Physics' enabled if you'd like to {1}. "), FText::FromString(GetReadableName()), ActionText));
+			}
+			else if (CollisionEnabled == ECollisionEnabled::NoCollision || CollisionEnabled == ECollisionEnabled::QueryOnly)	//shapes need to be simulating
+			{
+				FMessageLog("PIE").Warning(FText::Format(LOCTEXT("InvalidPhysicsOperation", "{0} has to have 'CollisionEnabled' set to 'Query and Physics' or 'Physics only' if you'd like to {1}. "), FText::FromString(GetReadableName()), ActionText));
+			}
 		}
 	}
 #endif
@@ -437,9 +446,9 @@ float UPrimitiveComponent::GetAngularDamping() const
 
 void UPrimitiveComponent::SetMassScale(FName BoneName, float InMassScale)
 {
-	FBodyInstance* BI = GetBodyInstance(BoneName);
-	if (BI)
+	if (FBodyInstance* BI = GetBodyInstance(BoneName))
 	{
+		WarnInvalidPhysicsOperations(LOCTEXT("SetMassScale", "SetMassScale"), BI);
 		BI->SetMassScale(InMassScale);
 	}
 }
@@ -451,9 +460,9 @@ void UPrimitiveComponent::SetAllMassScale(float InMassScale)
 
 float UPrimitiveComponent::GetMass() const
 {
-	FBodyInstance* BI = GetBodyInstance();
-	if (BI)
+	if (FBodyInstance* BI = GetBodyInstance())
 	{
+		WarnInvalidPhysicsOperations(LOCTEXT("GetMass", "GetMass"), BI);
 		return BI->GetBodyMass();
 	}
 
