Fixed UE-15909 "Opening a level with a navmesh in standalone game crashes standalone .exe"

diff --git a/Engine/Source/Runtime/Engine/Private/AI/Navigation/RecastNavMesh.cpp b/Engine/Source/Runtime/Engine/Private/AI/Navigation/RecastNavMesh.cpp
--- a/Engine/Source/Runtime/Engine/Private/AI/Navigation/RecastNavMesh.cpp
+++ b/Engine/Source/Runtime/Engine/Private/AI/Navigation/RecastNavMesh.cpp
@@ -441,7 +441,7 @@ UPrimitiveComponent* ARecastNavMesh::ConstructRenderingComponentImpl()
 
 	return NULL;
 #else
-	return NewObject<UNavMeshRenderingComponent>(this, TEXT("NavMeshRenderer"));
+	return NewObject<UNavMeshRenderingComponent>(this, TEXT("NavMeshRenderer"), RF_Transient);
 #endif // DO_NAVMESH_DEBUG_DRAWING_PER_TILE
 }
 
