#ue4 - Fix stat tracker in "stat game" for CharacterMovement not tracking the total cost. Add stats for Character network movement and smoothing.

diff --git a/Engine/Source/Runtime/Engine/Private/Components/CharacterMovementComponent.cpp b/Engine/Source/Runtime/Engine/Private/Components/CharacterMovementComponent.cpp
--- a/Engine/Source/Runtime/Engine/Private/Components/CharacterMovementComponent.cpp
+++ b/Engine/Source/Runtime/Engine/Private/Components/CharacterMovementComponent.cpp
@@ -35,6 +35,8 @@ DECLARE_STATS_GROUP(TEXT("Character"), STATGROUP_Character, STATCAT_Advanced);
 DECLARE_CYCLE_STAT(TEXT("Char Movement Tick"), STAT_CharacterMovementTick, STATGROUP_Character);
 DECLARE_CYCLE_STAT(TEXT("Char Movement Authority Time"), STAT_CharacterMovementAuthority, STATGROUP_Character);
 DECLARE_CYCLE_STAT(TEXT("Char Movement Simulated Time"), STAT_CharacterMovementSimulated, STATGROUP_Character);
+DECLARE_CYCLE_STAT(TEXT("Char CombineNetMove"), STAT_CharacterMovementCombineNetMove, STATGROUP_Character);
+DECLARE_CYCLE_STAT(TEXT("Char SmoothClientPosition"), STAT_CharacterMovementSmoothClientPosition, STATGROUP_Character);
 DECLARE_CYCLE_STAT(TEXT("Char Physics Interation"), STAT_CharPhysicsInteraction, STATGROUP_Character);
 DECLARE_CYCLE_STAT(TEXT("Char StepUp"), STAT_CharStepUp, STATGROUP_Character);
 DECLARE_CYCLE_STAT(TEXT("Char AdjustFloorHeight"), STAT_CharAdjustFloorHeight, STATGROUP_Character);
@@ -848,6 +850,7 @@ void UCharacterMovementComponent::PerformAirControlForPathFollowing(FVector Dire
 
 void UCharacterMovementComponent::TickComponent(float DeltaTime, enum ELevelTick TickType, FActorComponentTickFunction *ThisTickFunction)
 {
+	SCOPE_CYCLE_COUNTER(STAT_CharacterMovement);
 	SCOPE_CYCLE_COUNTER(STAT_CharacterMovementTick);
 
 	const FVector InputVector = ConsumeInputVector();
@@ -1008,7 +1011,6 @@ void UCharacterMovementComponent::AdjustProxyCapsuleSize()
 
 void UCharacterMovementComponent::SimulatedTick(float DeltaSeconds)
 {
-	SCOPE_CYCLE_COUNTER(STAT_CharacterMovement);
 	SCOPE_CYCLE_COUNTER(STAT_CharacterMovementSimulated);
 
 	// If we are playing a RootMotion AnimMontage.
@@ -1488,7 +1490,6 @@ void UCharacterMovementComponent::DisableMovement()
 
 void UCharacterMovementComponent::PerformMovement(float DeltaSeconds)
 {
-	SCOPE_CYCLE_COUNTER(STAT_CharacterMovement);
 	SCOPE_CYCLE_COUNTER(STAT_CharacterMovementAuthority);
 
 	if (!HasValidData())
@@ -5645,6 +5646,7 @@ void UCharacterMovementComponent::SmoothCorrection(const FVector& OldLocation, c
 
 void UCharacterMovementComponent::SmoothClientPosition(float DeltaSeconds)
 {
+	SCOPE_CYCLE_COUNTER(STAT_CharacterMovementSmoothClientPosition);
 	if (!HasValidData() || GetNetMode() != NM_Client)
 	{
 		return;
@@ -5954,6 +5956,8 @@ void UCharacterMovementComponent::ReplicateMoveToServer(float DeltaTime, const F
 	// do not combine moves which have different TimeStamps (before and after reset).
 	if( ClientData->PendingMove.IsValid() && !ClientData->PendingMove->bOldTimeStampBeforeReset && ClientData->PendingMove->CanCombineWith(NewMove, CharacterOwner, ClientData->MaxResponseTime * CharacterOwner->GetWorldSettings()->GetEffectiveTimeDilation()))
 	{
+		SCOPE_CYCLE_COUNTER(STAT_CharacterMovementCombineNetMove);
+
 		// Only combine and move back to the start location if we don't move back in to a spot that would make us collide with something new.
 		const FVector OldStartLocation = ClientData->PendingMove->GetRevertedLocation();
 		if (!OverlapTest(OldStartLocation, ClientData->PendingMove->StartRotation.Quaternion(), UpdatedComponent->GetCollisionObjectType(), GetPawnCapsuleCollisionShape(SHRINK_None), CharacterOwner))
