Engine - Changing the location we ignore incoming mouse events to be in player update, doing it in the hit testing code prevents manual usage of the code from being useful based on the cursor location, which shouldn't be considered part of it.

diff --git a/Engine/Source/Runtime/Engine/Private/PlayerController.cpp b/Engine/Source/Runtime/Engine/Private/PlayerController.cpp
--- a/Engine/Source/Runtime/Engine/Private/PlayerController.cpp
+++ b/Engine/Source/Runtime/Engine/Private/PlayerController.cpp
@@ -1957,12 +1957,6 @@ bool APlayerController::GetHitResultAtScreenPosition(const FVector2D ScreenPosit
 	{
 		UGameViewportClient* ViewportClient = LocalPlayer->ViewportClient;
 
-		// Early out if were over hit testable slate widgets other then this viewport
-		if (!(ViewportClient->GetGameViewportWidget().IsValid() && ViewportClient->GetGameViewportWidget()->IsDirectlyHovered()))
-		{
-			return false;
-		}
-
 		// Create a view family for the game viewport
 		FSceneViewFamilyContext ViewFamily( FSceneViewFamily::ConstructionValues(
 			ViewportClient->Viewport,
@@ -3756,9 +3750,15 @@ void APlayerController::TickPlayerInput(const float DeltaSeconds, const bool bGa
 			FHitResult HitResult;
 			bool bHit = false;
 			
-			if (LocalPlayer->ViewportClient->GetMousePosition(MousePosition))
+			UGameViewportClient* ViewportClient = LocalPlayer->ViewportClient;
+
+			// Only send mouse hit events if we're directly over the viewport.
+			if ( ViewportClient->GetGameViewportWidget().IsValid() && ViewportClient->GetGameViewportWidget()->IsDirectlyHovered() )
 			{
-				bHit = GetHitResultAtScreenPosition(MousePosition, CurrentClickTraceChannel, true, /*out*/ HitResult);
+				if ( LocalPlayer->ViewportClient->GetMousePosition(MousePosition) )
+				{
+					bHit = GetHitResultAtScreenPosition(MousePosition, CurrentClickTraceChannel, true, /*out*/ HitResult);
+				}
 			}
 
 			UPrimitiveComponent* PreviousComponent = CurrentClickablePrimitive.Get();
