Added code gaurds to remove UWorld::AttemptDeriveFromPlayInSettings unless the editor is present

diff --git a/Engine/Source/Runtime/Engine/Classes/Engine/World.h b/Engine/Source/Runtime/Engine/Classes/Engine/World.h
--- a/Engine/Source/Runtime/Engine/Classes/Engine/World.h
+++ b/Engine/Source/Runtime/Engine/Classes/Engine/World.h
@@ -2937,8 +2937,10 @@ public:
 	/** Returns the net mode this world is running under */
 	ENetMode GetNetMode() const;
 
+#if WITH_EDITOR
 	/** Attempts to derive the net mode from PlayInSettings for PIE*/
 	ENetMode AttemptDeriveFromPlayInSettings() const;
+#endif
 
 	/** Attempts to derive the net mode from URL */
 	ENetMode AttemptDeriveFromURL() const;
