Fix welded bodies not obtaining a physx lock

UE-16131

diff --git a/Engine/Source/Runtime/Engine/Private/PhysicsEngine/BodyInstance.cpp b/Engine/Source/Runtime/Engine/Private/PhysicsEngine/BodyInstance.cpp
--- a/Engine/Source/Runtime/Engine/Private/PhysicsEngine/BodyInstance.cpp
+++ b/Engine/Source/Runtime/Engine/Private/PhysicsEngine/BodyInstance.cpp
@@ -2878,7 +2878,8 @@ void FBodyInstance::ExecuteOnPhysicsReadOnly(TFunctionRef<void()> Func) const
 	//Because of this our read operations are typically on a dynamic or the sync actor
 
 #if WITH_PHYSX
-	const int32 SceneIndex = RigidActorSync ? SceneIndexSync : SceneIndexAsync;
+	const FBodyInstance* BI = WeldParent ? WeldParent : this;
+	const int32 SceneIndex = BI->RigidActorSync ? BI->SceneIndexSync : BI->SceneIndexAsync;
 	SCOPED_SCENE_READ_LOCK(GetPhysXSceneFromIndex(SceneIndex));
 	Func();
 #endif
