PR #1176: Destructible component physics bugfix (Contributed by [REDACTED_PERSON])

diff --git a/Engine/Source/Runtime/Engine/Private/Components/DestructibleComponent.cpp b/Engine/Source/Runtime/Engine/Private/Components/DestructibleComponent.cpp
--- a/Engine/Source/Runtime/Engine/Private/Components/DestructibleComponent.cpp
+++ b/Engine/Source/Runtime/Engine/Private/Components/DestructibleComponent.cpp
@@ -450,17 +450,11 @@ bool UDestructibleComponent::CanEditSimulatePhysics()
 void UDestructibleComponent::AddImpulse( FVector Impulse, FName BoneName /*= NAME_None*/, bool bVelChange /*= false*/ )
 {
 #if WITH_APEX
-	if ( ApexDestructibleActor )
+	if (ApexDestructibleActor)
 	{
-		int32 ChunkIdx = BoneIdxToChunkIdx(GetBoneIndex(BoneName));
+		const int32 ChunkIdx = BoneIdxToChunkIdx(GetBoneIndex(BoneName));
 		PxRigidDynamic* PActor = ApexDestructibleActor->getChunkPhysXActor(ChunkIdx);
-
-		if ( PActor != NULL )
-		{
-			SCOPED_SCENE_WRITE_LOCK(PActor->getScene());
-
-			PActor->addForce(U2PVector(Impulse), bVelChange ? PxForceMode::eVELOCITY_CHANGE : PxForceMode::eIMPULSE);
-		}
+		PActor->addForce(U2PVector(Impulse), bVelChange ? PxForceMode::eVELOCITY_CHANGE : PxForceMode::eIMPULSE);
 	}
 #endif
 }
@@ -471,14 +465,9 @@ void UDestructibleComponent::AddImpulseAtLocation( FVector Impulse, FVector Posi
 #if WITH_APEX
 	if (ApexDestructibleActor)
 	{
-		int32 ChunkIdx = NxModuleDestructibleConst::INVALID_CHUNK_INDEX;
-		if (BoneName != NAME_None)
-		{
-			ChunkIdx = BoneIdxToChunkIdx(GetBoneIndex(BoneName));
-		}
-
+		const int32 ChunkIdx = BoneIdxToChunkIdx(GetBoneIndex(BoneName));
 		PxRigidDynamic* Actor = ApexDestructibleActor->getChunkPhysXActor(ChunkIdx);
-		Actor->addForce(U2PVector(Impulse), PxForceMode::eIMPULSE);
+		PxRigidBodyExt::addForceAtPos(*Actor, U2PVector(Impulse), U2PVector(Position), PxForceMode::eIMPULSE);
 	}
 #endif
 }
@@ -486,14 +475,12 @@ void UDestructibleComponent::AddImpulseAtLocation( FVector Impulse, FVector Posi
 void UDestructibleComponent::AddForce( FVector Force, FName BoneName /*= NAME_None*/, bool bAccelChange /* = false */ )
 {
 #if WITH_APEX
-	int32 ChunkIdx = NxModuleDestructibleConst::INVALID_CHUNK_INDEX;
-	if (BoneName != NAME_None)
+	if (ApexDestructibleActor)
 	{
-		ChunkIdx = BoneIdxToChunkIdx(GetBoneIndex(BoneName));
+		const int32 ChunkIdx = BoneIdxToChunkIdx(GetBoneIndex(BoneName));
+		PxRigidDynamic* Actor = ApexDestructibleActor->getChunkPhysXActor(ChunkIdx);
+		Actor->addForce(U2PVector(Force), bAccelChange ? PxForceMode::eACCELERATION : PxForceMode::eFORCE);
 	}
-
-	PxRigidDynamic* Actor = ApexDestructibleActor->getChunkPhysXActor(ChunkIdx);
-	Actor->addForce(U2PVector(Force), bAccelChange ? PxForceMode::eACCELERATION : PxForceMode::eFORCE);
 #endif
 }
 
@@ -502,14 +489,9 @@ void UDestructibleComponent::AddForceAtLocation( FVector Force, FVector Location
 #if WITH_APEX
 	if (ApexDestructibleActor)
 	{
-		int32 ChunkIdx = NxModuleDestructibleConst::INVALID_CHUNK_INDEX;
-		if (BoneName != NAME_None)
-		{
-			ChunkIdx = BoneIdxToChunkIdx(GetBoneIndex(BoneName));
-		}
-
+		int32 ChunkIdx = BoneIdxToChunkIdx(GetBoneIndex(BoneName));
 		PxRigidDynamic* Actor = ApexDestructibleActor->getChunkPhysXActor(ChunkIdx);
-		Actor->addForce(U2PVector(Force), PxForceMode::eFORCE);
+		PxRigidBodyExt::addForceAtPos(*Actor, U2PVector(Force), U2PVector(Location), PxForceMode::eFORCE);
 	}
 #endif
 }
@@ -568,6 +550,11 @@ void UDestructibleComponent::AddRadialForce(FVector Origin, float Radius, float
 		return;
 	}
 
+	if (ApexDestructibleActor == NULL)
+	{
+		return;
+	}
+
 	PxRigidDynamic** PActorBuffer = NULL;
 	PxU32 PActorCount = 0;
 	if (ApexDestructibleActor->acquirePhysXActorBuffer(PActorBuffer, PActorCount, NxDestructiblePhysXActorQueryFlags::Dynamic))
@@ -1428,4 +1415,4 @@ void UDestructibleComponent::SetMaterial(int32 ElementIndex, UMaterialInterface*
 #endif
 }
 
-#endif // WITH_PHYSX
\ No newline at end of file
+#endif // WITH_PHYSX
