Fix destructible actor operations not obtaining the physx lock

UE-16136

diff --git a/Engine/Source/Runtime/Engine/Classes/Components/DestructibleComponent.h b/Engine/Source/Runtime/Engine/Classes/Components/DestructibleComponent.h
--- a/Engine/Source/Runtime/Engine/Classes/Components/DestructibleComponent.h
+++ b/Engine/Source/Runtime/Engine/Classes/Components/DestructibleComponent.h
@@ -222,6 +222,12 @@ private:
 
 	bool IsFracturedOrInitiallyStatic() const;
 
+	/** Obtains the appropriate PhysX scene lock for READING and executes the passed in lambda. */
+	bool ExecuteOnPhysicsReadOnly(TFunctionRef<void()> Func) const;
+
+	/** Obtains the appropriate PhysX scene lock for WRITING and executes the passed in lambda. */
+	bool ExecuteOnPhysicsReadWrite(TFunctionRef<void()> Func) const;
+
 	/** Collision response used for chunks */
 	FCollisionResponse LargeChunkCollisionResponse;
 	FCollisionResponse SmallChunkCollisionResponse;
diff --git a/Engine/Source/Runtime/Engine/Private/Components/DestructibleComponent.cpp b/Engine/Source/Runtime/Engine/Private/Components/DestructibleComponent.cpp
--- a/Engine/Source/Runtime/Engine/Private/Components/DestructibleComponent.cpp
+++ b/Engine/Source/Runtime/Engine/Private/Components/DestructibleComponent.cpp
@@ -450,12 +450,12 @@ bool UDestructibleComponent::CanEditSimulatePhysics()
 void UDestructibleComponent::AddImpulse( FVector Impulse, FName BoneName /*= NAME_None*/, bool bVelChange /*= false*/ )
 {
 #if WITH_APEX
-	if (ApexDestructibleActor)
+	ExecuteOnPhysicsReadWrite([&]
 	{
 		const int32 ChunkIdx = BoneIdxToChunkIdx(GetBoneIndex(BoneName));
 		PxRigidDynamic* PActor = ApexDestructibleActor->getChunkPhysXActor(ChunkIdx);
 		PActor->addForce(U2PVector(Impulse), bVelChange ? PxForceMode::eVELOCITY_CHANGE : PxForceMode::eIMPULSE);
-	}
+	});
 #endif
 }
 
@@ -463,36 +463,36 @@ void UDestructibleComponent::AddImpulse( FVector Impulse, FName BoneName /*= NAM
 void UDestructibleComponent::AddImpulseAtLocation( FVector Impulse, FVector Position, FName BoneName /*= NAME_None*/ )
 {
 #if WITH_APEX
-	if (ApexDestructibleActor)
+	ExecuteOnPhysicsReadWrite([&]
 	{
 		const int32 ChunkIdx = BoneIdxToChunkIdx(GetBoneIndex(BoneName));
 		PxRigidDynamic* Actor = ApexDestructibleActor->getChunkPhysXActor(ChunkIdx);
 		PxRigidBodyExt::addForceAtPos(*Actor, U2PVector(Impulse), U2PVector(Position), PxForceMode::eIMPULSE);
-	}
+	});
 #endif
 }
 
 void UDestructibleComponent::AddForce( FVector Force, FName BoneName /*= NAME_None*/, bool bAccelChange /* = false */ )
 {
 #if WITH_APEX
-	if (ApexDestructibleActor)
+	ExecuteOnPhysicsReadWrite([&]
 	{
 		const int32 ChunkIdx = BoneIdxToChunkIdx(GetBoneIndex(BoneName));
 		PxRigidDynamic* Actor = ApexDestructibleActor->getChunkPhysXActor(ChunkIdx);
 		Actor->addForce(U2PVector(Force), bAccelChange ? PxForceMode::eACCELERATION : PxForceMode::eFORCE);
-	}
+	});
 #endif
 }
 
 void UDestructibleComponent::AddForceAtLocation( FVector Force, FVector Location, FName BoneName /*= NAME_None*/ )
 {
 #if WITH_APEX
-	if (ApexDestructibleActor)
+	ExecuteOnPhysicsReadWrite([&]
 	{
 		int32 ChunkIdx = BoneIdxToChunkIdx(GetBoneIndex(BoneName));
 		PxRigidDynamic* Actor = ApexDestructibleActor->getChunkPhysXActor(ChunkIdx);
 		PxRigidBodyExt::addForceAtPos(*Actor, U2PVector(Force), U2PVector(Location), PxForceMode::eFORCE);
-	}
+	});
 #endif
 }
 
@@ -721,6 +721,46 @@ bool UDestructibleComponent::IsFracturedOrInitiallyStatic() const
 	return bFractured || bInitiallyStatic;
 }
 
+bool UDestructibleComponent::ExecuteOnPhysicsReadOnly(TFunctionRef<void()> Func) const
+{
+#if WITH_APEX
+	if (ApexDestructibleActor)
+	{
+		FPhysScene* PhysScene = GetWorld()->GetPhysicsScene();
+		// Destructibles are always dynamic or kinematic, and therefore only go into one of the scenes
+		const uint32 SceneType = BodyInstance.UseAsyncScene(PhysScene) ? PST_Async : PST_Sync;
+		PxScene* PScene = PhysScene->GetPhysXScene(SceneType);
+
+		SCOPED_SCENE_READ_LOCK(PScene);
+		Func();
+
+		return true;
+	}
+#endif
+
+	return false;
+}
+
+bool UDestructibleComponent::ExecuteOnPhysicsReadWrite(TFunctionRef<void()> Func) const
+{
+#if WITH_APEX
+	if (ApexDestructibleActor)
+	{
+		FPhysScene* PhysScene = GetWorld()->GetPhysicsScene();
+		// Destructibles are always dynamic or kinematic, and therefore only go into one of the scenes
+		const uint32 SceneType = BodyInstance.UseAsyncScene(PhysScene) ? PST_Async : PST_Sync;
+		PxScene* PScene = PhysScene->GetPhysXScene(SceneType);
+
+		SCOPED_SCENE_WRITE_LOCK(PScene);
+		Func();
+
+		return true;
+	}
+#endif
+
+	return false;
+}
+
 void UDestructibleComponent::RefreshBoneTransforms(FActorComponentTickFunction* TickFunction)
 {
 }
