Fix Write lock on welded bodies not grabbing the right physx lock

UE-16131

diff --git a/Engine/Source/Runtime/Engine/Private/PhysicsEngine/BodyInstance.cpp b/Engine/Source/Runtime/Engine/Private/PhysicsEngine/BodyInstance.cpp
--- a/Engine/Source/Runtime/Engine/Private/PhysicsEngine/BodyInstance.cpp
+++ b/Engine/Source/Runtime/Engine/Private/PhysicsEngine/BodyInstance.cpp
@@ -2891,26 +2891,27 @@ void FBodyInstance::ExecuteOnPhysicsReadWrite(TFunctionRef<void()> Func) const
 	//Because of that write operations on static actors are more expensive and require both locks.
 
 #if WITH_PHYSX
-	if(RigidActorSync)
+	const FBodyInstance* BI = WeldParent ? WeldParent : this;
+	if(BI->RigidActorSync)
 	{
-		SCENE_LOCK_WRITE(GetPhysXSceneFromIndex(SceneIndexSync));
+		SCENE_LOCK_WRITE(GetPhysXSceneFromIndex(BI->SceneIndexSync));
 	}
 	
-	if (RigidActorAsync)
+	if (BI->RigidActorAsync)
 	{
-		SCENE_LOCK_WRITE(GetPhysXSceneFromIndex(SceneIndexAsync));
+		SCENE_LOCK_WRITE(GetPhysXSceneFromIndex(BI->SceneIndexAsync));
 	}
 
 	Func();
 
-	if (RigidActorSync)
+	if (BI->RigidActorSync)
 	{
-		SCENE_UNLOCK_WRITE(GetPhysXSceneFromIndex(SceneIndexSync));
+		SCENE_UNLOCK_WRITE(GetPhysXSceneFromIndex(BI->SceneIndexSync));
 	}
 
-	if (RigidActorAsync)
+	if (BI->RigidActorAsync)
 	{
-		SCENE_UNLOCK_WRITE(GetPhysXSceneFromIndex(SceneIndexAsync));
+		SCENE_UNLOCK_WRITE(GetPhysXSceneFromIndex(BI->SceneIndexAsync));
 	}
 #endif
 }
