Merging from /usr/John_Barrett: Adjusted plugin commandlet to allow enabling/disabling of multiple plugins at once, delimited with a comma ','.

Also added a check to see that the plugin actually exists before enabling it, and a '-Force' commandline switch, to override this check.

diff --git a/Engine/Source/Runtime/Engine/Private/Commandlets/PluginCommandlet.cpp b/Engine/Source/Runtime/Engine/Private/Commandlets/PluginCommandlet.cpp
--- a/Engine/Source/Runtime/Engine/Private/Commandlets/PluginCommandlet.cpp
+++ b/Engine/Source/Runtime/Engine/Private/Commandlets/PluginCommandlet.cpp
@@ -14,10 +14,14 @@ DEFINE_LOG_CATEGORY_STATIC(LogPluginCommandlet, Log, All);
  * Commandlet used for enabling/disabling plugins
  *
  * Usage:
- *	Plugin Enable/Disable PluginName
+ *	Plugin Enable/Disable PluginName,PluginName2
+ *
+ * Optional parameters:
+ *	-Force: Forces a plugin to be enabled/disabled, even if it can not be found.
  *
  * Example:
  *	Plugin Enable NetcodeUnitTest
+ *	Plugin Enable NetcodeUnitTest,NUTUnrealEngine4
  */
 
 UPluginCommandlet::UPluginCommandlet(const FObjectInitializer& ObjectInitializer)
@@ -59,17 +63,31 @@ int32 UPluginCommandlet::Main(const FString& Params)
 			// Trim all tokens up to the command, to keep things simple
 			CmdLineTokens.RemoveAt(0, CommandIdx+1);
 
-			FString PluginName = (CmdLineTokens.Num() > 0 ? CmdLineTokens.Pop() : TEXT(""));
+			FString PluginList = (CmdLineTokens.Num() > 0 ? CmdLineTokens.Pop() : TEXT(""));
+			TArray<FString> PluginArray;
+
+			PluginList.ParseIntoArray(PluginArray, TEXT(","));
 
-			if (PluginName.Len() > 0)
+			if (PluginArray.Num() > 0 && IProjectManager::Get().LoadProjectFile(FPaths::GetProjectFilePath()))
 			{
-				if (IProjectManager::Get().LoadProjectFile(FPaths::GetProjectFilePath()))
+				FText FailReason;
+
+				bSuccess = true;
+
+				for (FString PluginName : PluginArray)
 				{
-					FText FailReason;
+					bool bSetPluginState = false;
 
-					bSuccess = IProjectManager::Get().SetPluginEnabled(PluginName, bEnable, FailReason);
+					if (IPluginManager::Get().FindPlugin(PluginName).IsValid() || CmdLineSwitches.Contains(TEXT("FORCE")))
+					{
+						bSetPluginState = IProjectManager::Get().SetPluginEnabled(PluginName, bEnable, FailReason);
+					}
+					else
+					{
+						FailReason = FText::FromString(TEXT("Plugin not found. Add -Force to the commandline, to override."));
+					}
 
-					if (bSuccess)
+					if (bSetPluginState)
 					{
 						UE_LOG(LogPluginCommandlet, Log, TEXT("Successfully %s plugin '%s'"),
 								(bEnable ? TEXT("enabled") : TEXT("disabled")), *PluginName);
@@ -78,17 +96,19 @@ int32 UPluginCommandlet::Main(const FString& Params)
 					{
 						UE_LOG(LogPluginCommandlet, Error, TEXT("Failed to %s plugin '%s' - error: %s"),
 								(bEnable ? TEXT("enable") : TEXT("disable")), *PluginName, *FailReason.ToString());
+
+						bSuccess = false;
 					}
 				}
-				else
-				{
-					UE_LOG(LogPluginCommandlet, Error, TEXT("Failed to load project file."));
-					bSuccess = false;
-				}
+			}
+			else if (PluginArray.Num() == 0)
+			{
+				UE_LOG(LogPluginCommandlet, Error, TEXT("Failed to parse a plugin name."));
+				bSuccess = false;
 			}
 			else
 			{
-				UE_LOG(LogPluginCommandlet, Error, TEXT("Failed to parse plugin name."));
+				UE_LOG(LogPluginCommandlet, Error, TEXT("Failed to load project file."));
 				bSuccess = false;
 			}
 		}
