#ue4 - Expose UPrimitiveComponent::CharacterCanStepUpOn variable and UPrimitiveComponent::CharacterCanStepUp() function to blueprints.

- Also improved comments/docs on these to indicate that it only prevents step ups when bumping in to objects, not walkability on them.

diff --git a/Engine/Source/Runtime/Engine/Classes/Components/PrimitiveComponent.h b/Engine/Source/Runtime/Engine/Classes/Components/PrimitiveComponent.h
--- a/Engine/Source/Runtime/Engine/Classes/Components/PrimitiveComponent.h
+++ b/Engine/Source/Runtime/Engine/Classes/Components/PrimitiveComponent.h
@@ -467,8 +467,10 @@ private:
 public:
 	/**
 	 * Determine whether a Character can step up onto this component.
+	 * This controls whether they can try to step up on it when they bump in to it, not whether they can walk on it after landing on it.
+	 * @see FWalkableSlopeOverride
 	 */
-	UPROPERTY(EditDefaultsOnly, BlueprintReadOnly, Category=Collision)
+	UPROPERTY(EditAnywhere, BlueprintReadWrite, Category=Collision)
 	TEnumAsByte<enum ECanBeCharacterBase> CanCharacterStepUpOn;
 
 	/**
@@ -1700,8 +1702,11 @@ public:
 
 	/**
 	 * Return true if the given Pawn can step up onto this component.
-	 * @param Pawn is the Pawn that wants to step onto this component.
+	 * This controls whether they can try to step up on it when they bump in to it, not whether they can walk on it after landing on it.
+	 * @param Pawn the Pawn that wants to step onto this component.
+	 * @see CanCharacterStepUpOn
 	 */
+	UFUNCTION(BlueprintCallable, Category=Collision)
 	virtual bool CanCharacterStepUp(class APawn* Pawn) const;
 
 	/** Can this component potentially influence navigation */
