Merge to main: Fix to DemoNetDriver to not delete spectator pawn on when scrubbing

diff --git a/Engine/Source/Runtime/Engine/Private/DemoNetDriver.cpp b/Engine/Source/Runtime/Engine/Private/DemoNetDriver.cpp
--- a/Engine/Source/Runtime/Engine/Private/DemoNetDriver.cpp
+++ b/Engine/Source/Runtime/Engine/Private/DemoNetDriver.cpp
@@ -1510,7 +1510,7 @@ void UDemoNetDriver::LoadCheckpoint()
 			AddNonQueuedActorForScrubbing(*It);
 		}
 		
-		if ( *It == SpectatorController )
+		if (*It == SpectatorController || *It == SpectatorController->GetSpectatorPawn())
 		{
 			continue;
 		}
