Merge to main: Fix crash when viewing live replays

diff --git a/Engine/Source/Runtime/Engine/Private/DemoNetDriver.cpp b/Engine/Source/Runtime/Engine/Private/DemoNetDriver.cpp
--- a/Engine/Source/Runtime/Engine/Private/DemoNetDriver.cpp
+++ b/Engine/Source/Runtime/Engine/Private/DemoNetDriver.cpp
@@ -1510,14 +1510,17 @@ void UDemoNetDriver::LoadCheckpoint()
 			AddNonQueuedActorForScrubbing(*It);
 		}
 		
-		if (*It == SpectatorController || *It == SpectatorController->GetSpectatorPawn())
+		if ( SpectatorController != nullptr )
 		{
-			continue;
-		}
+			if ( *It == SpectatorController || *It == SpectatorController->GetSpectatorPawn() )
+			{
+				continue;
+			}
 
-		if ( It->GetOwner() == SpectatorController )
-		{
-			continue;
+			if ( It->GetOwner() == SpectatorController )
+			{
+				continue;
+			}
 		}
 
 		if ( It->IsNetStartupActor() )
