Fix for CRASH: EDITOR: Crash occurs when connecting MF_DepthEnv to the M_Forest Material Attributes

Several recent additions to the material properties have made the mapping from IO index to EMaterialPropertiy very confusing and fragile.
I've pulled the mapping to a single place in the code and it should now be robust to new additions.
Changing the order of pins will still require fix up of old FExpressionInputs but this is true of all expressions. No way around that.
Though I've now placed some static asserts and comments around the code so that the need for this and how to easily avoid it at all is clear and unmissable.

diff --git a/Engine/Shaders/ShaderVersion.usf b/Engine/Shaders/ShaderVersion.usf
--- a/Engine/Shaders/ShaderVersion.usf
+++ b/Engine/Shaders/ShaderVersion.usf
@@ -2,4 +2,4 @@
 // Each time the console command is executed it generates a new GUID. As this file is included
 // in common.usf (which should be included in any shader) it allows to invalidate the shader DDC.
 // 
-// GUID = 639690ED421CFC45604E3687590EE0F7
\ No newline at end of file
+// GUID = 48155F744FBD837A7BD439A9F40EC4D7
\ No newline at end of file
diff --git a/Engine/Source/Runtime/Engine/Private/Materials/Material.cpp b/Engine/Source/Runtime/Engine/Private/Materials/Material.cpp
--- a/Engine/Source/Runtime/Engine/Private/Materials/Material.cpp
+++ b/Engine/Source/Runtime/Engine/Private/Materials/Material.cpp
@@ -1998,8 +1998,20 @@ void UMaterial::Serialize(FArchive& Ar)
 	DoMaterialAttributeReorder(&WorldDisplacement,		Ar.UE4Ver());
 	DoMaterialAttributeReorder(&TessellationMultiplier,	Ar.UE4Ver());
 	DoMaterialAttributeReorder(&SubsurfaceColor,		Ar.UE4Ver());
+	DoMaterialAttributeReorder(&ClearCoat,				Ar.UE4Ver());
+	DoMaterialAttributeReorder(&ClearCoatRoughness,		Ar.UE4Ver());
 	DoMaterialAttributeReorder(&AmbientOcclusion,		Ar.UE4Ver());
 	DoMaterialAttributeReorder(&Refraction,				Ar.UE4Ver());
+	DoMaterialAttributeReorder(&CustomizedUVs[0],		Ar.UE4Ver());
+	DoMaterialAttributeReorder(&CustomizedUVs[1],		Ar.UE4Ver());
+	DoMaterialAttributeReorder(&CustomizedUVs[2],		Ar.UE4Ver());
+	DoMaterialAttributeReorder(&CustomizedUVs[3],		Ar.UE4Ver());
+	DoMaterialAttributeReorder(&CustomizedUVs[4],		Ar.UE4Ver());
+	DoMaterialAttributeReorder(&CustomizedUVs[5],		Ar.UE4Ver());
+	DoMaterialAttributeReorder(&CustomizedUVs[6],		Ar.UE4Ver());
+	DoMaterialAttributeReorder(&CustomizedUVs[7],		Ar.UE4Ver());
+	DoMaterialAttributeReorder(&PixelDepthOffset,		Ar.UE4Ver());
+	static_assert(MP_MAX == 28, "New material properties must have DoMaterialAttributesReorder called on them to ensure that any future reordering of property pins is correctly applied.");
 
 	if (Ar.UE4Ver() < VER_UE4_MATERIAL_MASKED_BLENDMODE_TIDY)
 	{
diff --git a/Engine/Source/Runtime/Engine/Private/Materials/MaterialExpressions.cpp b/Engine/Source/Runtime/Engine/Private/Materials/MaterialExpressions.cpp
--- a/Engine/Source/Runtime/Engine/Private/Materials/MaterialExpressions.cpp
+++ b/Engine/Source/Runtime/Engine/Private/Materials/MaterialExpressions.cpp
@@ -3198,43 +3198,44 @@ int32 UMaterialExpressionMakeMaterialAttributes::Compile(class FMaterialCompiler
 {
 	int32 Ret = INDEX_NONE;
 	UMaterialExpression* Expression = NULL;
-	switch( MultiplexIndex )
-	{
-		case 0: Ret = BaseColor.Compile(Compiler); Expression = BaseColor.Expression; break; 
-		case 1: Ret = Metallic.Compile(Compiler); Expression = Metallic.Expression; break; 
-		case 2: Ret = Specular.Compile(Compiler); Expression = Specular.Expression; break; 
-		case 3: Ret = Roughness.Compile(Compiler); Expression = Roughness.Expression; break; 
-		case 4: Ret = EmissiveColor.Compile(Compiler); Expression = EmissiveColor.Expression; break; 
-		case 5: Ret = Opacity.Compile(Compiler); Expression = Opacity.Expression; break; 
-		case 6: Ret = OpacityMask.Compile(Compiler); Expression = OpacityMask.Expression; break; 
-		case 7: Ret = Normal.Compile(Compiler); Expression = Normal.Expression; break; 
-		case 8: Ret = WorldPositionOffset.Compile(Compiler); Expression = WorldPositionOffset.Expression; break; 
-		case 9: Ret = WorldDisplacement.Compile(Compiler); Expression = WorldDisplacement.Expression; break; 
-		case 10: Ret = TessellationMultiplier.Compile(Compiler); Expression = TessellationMultiplier.Expression; break; 
-		case 11: Ret = SubsurfaceColor.Compile(Compiler); Expression = SubsurfaceColor.Expression; break;
-		case 12: Ret = ClearCoat.Compile(Compiler); Expression = ClearCoat.Expression; break;
-		case 13: Ret = ClearCoatRoughness.Compile(Compiler); Expression = ClearCoatRoughness.Expression; break;
-		case 14: Ret = AmbientOcclusion.Compile(Compiler); Expression = AmbientOcclusion.Expression; break;
-		case 15: Ret = Refraction.Compile(Compiler); Expression = Refraction.Expression; break; 
-	};
-
-	const int32 CustomUVStart = 16;
 
-	if (MultiplexIndex >= CustomUVStart && MultiplexIndex <= CustomUVStart + MP_CustomizedUVs7 - MP_CustomizedUVs0)
-	{
-		Ret = CustomizedUVs[MultiplexIndex - CustomUVStart].Compile(Compiler); Expression = CustomizedUVs[MultiplexIndex - CustomUVStart].Expression; 
-	}
+	static_assert(MP_MAX == 28, 
+		"New material properties should be added to the end of the inputs for this expression. \
+		The order of properties here should match the material results pins, the make material attriubtes node inputs and the mapping of IO indices to properties in GetMaterialPropertyFromInputOutputIndex().\
+		Insertions into the middle of the properties or a change in the order of properties will also require that existing data is fixed up in DoMaterialAttriubtesReorder().\
+		");
+
+	EMaterialProperty Property = GetMaterialPropertyFromInputOutputIndex(MultiplexIndex);
+	switch (Property)
+	{
+	case MP_BaseColor: Ret = BaseColor.Compile(Compiler); Expression = BaseColor.Expression; break;
+	case MP_Metallic: Ret = Metallic.Compile(Compiler); Expression = Metallic.Expression; break;
+	case MP_Specular: Ret = Specular.Compile(Compiler); Expression = Specular.Expression; break;
+	case MP_Roughness: Ret = Roughness.Compile(Compiler); Expression = Roughness.Expression; break;
+	case MP_EmissiveColor: Ret = EmissiveColor.Compile(Compiler); Expression = EmissiveColor.Expression; break;
+	case MP_Opacity: Ret = Opacity.Compile(Compiler); Expression = Opacity.Expression; break;
+	case MP_OpacityMask: Ret = OpacityMask.Compile(Compiler); Expression = OpacityMask.Expression; break;
+	case MP_Normal: Ret = Normal.Compile(Compiler); Expression = Normal.Expression; break;
+	case MP_WorldPositionOffset: Ret = WorldPositionOffset.Compile(Compiler); Expression = WorldPositionOffset.Expression; break;
+	case MP_WorldDisplacement: Ret = WorldDisplacement.Compile(Compiler); Expression = WorldDisplacement.Expression; break;
+	case MP_TessellationMultiplier: Ret = TessellationMultiplier.Compile(Compiler); Expression = TessellationMultiplier.Expression; break;
+	case MP_SubsurfaceColor: Ret = SubsurfaceColor.Compile(Compiler); Expression = SubsurfaceColor.Expression; break;
+	case MP_ClearCoat: Ret = ClearCoat.Compile(Compiler); Expression = ClearCoat.Expression; break;
+	case MP_ClearCoatRoughness: Ret = ClearCoatRoughness.Compile(Compiler); Expression = ClearCoatRoughness.Expression; break;
+	case MP_AmbientOcclusion: Ret = AmbientOcclusion.Compile(Compiler); Expression = AmbientOcclusion.Expression; break;
+	case MP_Refraction: Ret = Refraction.Compile(Compiler); Expression = Refraction.Expression; break;
+	case MP_PixelDepthOffset: Ret = PixelDepthOffset.Compile(Compiler); Expression = PixelDepthOffset.Expression; break;
+	};
 
-	if (MultiplexIndex == CustomUVStart + 2)
+	if (Property >= MP_CustomizedUVs0 && Property <= MP_CustomizedUVs7)
 	{
-		Ret = PixelDepthOffset.Compile(Compiler);
-		Expression = PixelDepthOffset.Expression;
+		Ret = CustomizedUVs[Property - MP_CustomizedUVs0].Compile(Compiler); Expression = CustomizedUVs[Property - MP_CustomizedUVs0].Expression;
 	}
 
 	//If we've connected an expression but its still returned INDEX_NONE, flag the error.
-	if( Expression && INDEX_NONE == Ret )
+	if (Expression && INDEX_NONE == Ret)
 	{
-		Compiler->Errorf(TEXT("Error on property %s"), *GetNameOfMaterialProperty(GetMaterialPropertyFromInputOutputIndex(MultiplexIndex)));
+		Compiler->Errorf(TEXT("Error on property %s"), *GetNameOfMaterialProperty(Property));
 	}
 
 	return Ret;
@@ -3264,6 +3265,12 @@ UMaterialExpressionBreakMaterialAttributes::UMaterialExpressionBreakMaterialAttr
 	bShowOutputNameOnPin = true;
 
 	MenuCategories.Add(ConstructorStatics.NAME_MaterialAttributes);
+	
+	static_assert(MP_MAX == 28, 
+		"New material properties should be added to the end of the outputs for this expression. \
+		The order of properties here should match the material results pins, the make material attriubtes node inputs and the mapping of IO indices to properties in GetMaterialPropertyFromInputOutputIndex().\
+		Insertions into the middle of the properties or a change in the order of properties will also require that existing data is fixed up in DoMaterialAttriubtesReorder().\
+		");
 
 	Outputs.Reset();
 	Outputs.Add(FExpressionOutput(TEXT("BaseColor"), 1, 1, 1, 1, 0));
diff --git a/Engine/Source/Runtime/Engine/Private/Materials/MaterialShared.cpp b/Engine/Source/Runtime/Engine/Private/Materials/MaterialShared.cpp
--- a/Engine/Source/Runtime/Engine/Private/Materials/MaterialShared.cpp
+++ b/Engine/Source/Runtime/Engine/Private/Materials/MaterialShared.cpp
@@ -2114,90 +2114,63 @@ FMaterialUpdateContext::~FMaterialUpdateContext()
 		);
 }
 
+const TMap<EMaterialProperty, int32> CreatePropertyToIOIndexMap()
+{	
+	static_assert(MP_MAX == 28, 
+		"New material properties should be added to the end of \"real\" properties in this map. Immediately before MP_MaterialAttributes . \
+		The order of properties here should match the material results pins, the inputs to MakeMaterialAttriubtes and the outputs of BreakMaterialAttriubtes.\
+		Insertions into the middle of the properties or a change in the order of properties will also require that existing data is fixed up in DoMaterialAttriubtesReorder().\
+		");
+
+	TMap<EMaterialProperty, int32> Ret;
+	Ret.Add(MP_BaseColor, 0);
+	Ret.Add(MP_Metallic, 1);
+	Ret.Add(MP_Specular, 2);
+	Ret.Add(MP_Roughness, 3);
+	Ret.Add(MP_EmissiveColor, 4);
+	Ret.Add(MP_Opacity, 5);
+	Ret.Add(MP_OpacityMask, 6);
+	Ret.Add(MP_Normal, 7);
+	Ret.Add(MP_WorldPositionOffset, 8);
+	Ret.Add(MP_WorldDisplacement, 9);
+	Ret.Add(MP_TessellationMultiplier, 10);
+	Ret.Add(MP_SubsurfaceColor, 11);
+	Ret.Add(MP_ClearCoat, 12);
+	Ret.Add(MP_ClearCoatRoughness, 13);
+	Ret.Add(MP_AmbientOcclusion, 14);
+	Ret.Add(MP_Refraction, 15);
+	Ret.Add(MP_CustomizedUVs0, 16);
+	Ret.Add(MP_CustomizedUVs1, 17);
+	Ret.Add(MP_CustomizedUVs2, 18);
+	Ret.Add(MP_CustomizedUVs3, 19);
+	Ret.Add(MP_CustomizedUVs4, 20);
+	Ret.Add(MP_CustomizedUVs5, 21);
+	Ret.Add(MP_CustomizedUVs6, 22);
+	Ret.Add(MP_CustomizedUVs7, 23);
+	Ret.Add(MP_PixelDepthOffset, 24);
+	//^^^^ New properties go above here ^^^^
+	//Below here are legacy or utility properties that don't match up to material inputs.
+	Ret.Add(MP_MaterialAttributes, INDEX_NONE);
+	Ret.Add(MP_DiffuseColor, INDEX_NONE);
+	Ret.Add(MP_SpecularColor, INDEX_NONE);
+	Ret.Add(MP_MAX, INDEX_NONE);
+	return Ret;
+}
+static const TMap<EMaterialProperty, int32> PropertyToIOIndexMap = CreatePropertyToIOIndexMap();
+
 EMaterialProperty GetMaterialPropertyFromInputOutputIndex(int32 Index)
 {
-	// Warning: If you change this mapping, you will need to handle backwards compatibility for nodes that use it, otherwise pins will get swapped
-
-	switch(Index)
-	{
-	case 0: return MP_BaseColor;
-	case 1: return MP_Metallic;
-	case 2: return MP_Specular;
-	case 3: return MP_Roughness;
-	case 4: return MP_EmissiveColor;
-	case 5: return MP_Opacity;
-	case 6: return MP_OpacityMask;
-	case 7: return MP_Normal;
-	case 8: return MP_WorldPositionOffset;
-	case 9: return MP_WorldDisplacement;
-	case 10: return MP_TessellationMultiplier;
-	case 11: return MP_SubsurfaceColor;
-	case 12: return MP_ClearCoat;
-	case 13: return MP_ClearCoatRoughness;
-	case 14: return MP_AmbientOcclusion;
-	case 15: return MP_Refraction;
-	};
-
-	const int32 UVStart = 16;
-	const int32 UVEnd = UVStart + MP_CustomizedUVs7 - MP_CustomizedUVs0;
-
-	if (Index >= UVStart && Index <= UVEnd)
-	{
-		return (EMaterialProperty)(MP_CustomizedUVs0 + Index - UVStart);
-	}
-
-	if (Index == UVEnd + 1)
-	{
-		return MP_MaterialAttributes;
-	}
-
-	if (Index == UVEnd + 2)
-	{
-		return MP_PixelDepthOffset;
-	}
-
-	return MP_MAX;
+	const EMaterialProperty* Property = PropertyToIOIndexMap.FindKey(Index);
+	check(Property);
+	return *Property;
 }
 
 int32 GetInputOutputIndexFromMaterialProperty(EMaterialProperty Property)
 {
-	// Warning: If you change this mapping, you will need to handle backwards compatibility for nodes that use it, otherwise pins will get swapped
-
-	switch(Property)
-	{
-	case MP_BaseColor: return 0;
-	case MP_Metallic: return 1;
-	case MP_Specular: return 2;
-	case MP_Roughness: return 3;
-	case MP_EmissiveColor: return 4;
-	case MP_Opacity: return 5;
-	case MP_OpacityMask: return 6;
-	case MP_Normal: return 7;
-	case MP_WorldPositionOffset: return 8;
-	case MP_WorldDisplacement: return 9;
-	case MP_TessellationMultiplier: return 10;
-	case MP_SubsurfaceColor: return 11;
-	case MP_ClearCoat: return 12;
-	case MP_ClearCoatRoughness: return 13;
-	case MP_AmbientOcclusion: return 14;
-	case MP_Refraction: return 15;
-	case MP_MaterialAttributes: UE_LOG(LogMaterial, Fatal, TEXT("We should never need the IO index of the MaterialAttriubtes property."));
-	};
-
-	const int32 UVStart = 16;
-	const int32 UVEnd = UVStart + MP_CustomizedUVs7 - MP_CustomizedUVs0;
-
-	if (Property >= MP_CustomizedUVs0 && Property <= MP_CustomizedUVs7)
-	{
-		return UVStart + Property - MP_CustomizedUVs0;
-	}
-
-	if (Property == MP_PixelDepthOffset)
-	{
-		return UVEnd + 2;
-	}
-
-	return -1;
+	const int32* Index = PropertyToIOIndexMap.Find(Property);
+	check(Index);
+	check(*Index != INDEX_NONE);
+	return *Index;
 }
 
 int32 GetDefaultExpressionForMaterialProperty(FMaterialCompiler* Compiler, EMaterialProperty Property)
@@ -2304,6 +2277,9 @@ int32 UMaterialInterface::CompileProperty(FMaterialCompiler* Compiler, EMaterial
 	}
 }
 
+//Reorder the output index for any FExpressionInput connected to a UMaterialExpressionBreakMaterialAttributes.
+//If the order of pins in the material results or the make/break attributes nodes changes 
+//then the OutputIndex stored in any FExpressionInput coming from UMaterialExpressionBreakMaterialAttributes will be wrong and needs reordering.
 void DoMaterialAttributeReorder(FExpressionInput* Input, int32 UE4Ver)
 {
 	if( Input && Input->Expression && Input->Expression->IsA(UMaterialExpressionBreakMaterialAttributes::StaticClass()) )
