Add support for welded bodies to use complex collision. This means we need support for multiple trimeshes per UBodySetup.

UE-14520

diff --git a/Engine/Source/Runtime/Engine/Classes/PhysicsEngine/BodyInstance.h b/Engine/Source/Runtime/Engine/Classes/PhysicsEngine/BodyInstance.h
--- a/Engine/Source/Runtime/Engine/Classes/PhysicsEngine/BodyInstance.h
+++ b/Engine/Source/Runtime/Engine/Classes/PhysicsEngine/BodyInstance.h
@@ -844,7 +844,7 @@ public:
 #endif
 
 	/** Update the instances collision filtering data */
-	void UpdatePhysicsFilterData(bool bForceSimpleAsComplex = false);
+	void UpdatePhysicsFilterData();
 
 	friend FArchive& operator<<(FArchive& Ar,FBodyInstance& BodyInst);
 
diff --git a/Engine/Source/Runtime/Engine/Classes/PhysicsEngine/BodySetup.h b/Engine/Source/Runtime/Engine/Classes/PhysicsEngine/BodySetup.h
--- a/Engine/Source/Runtime/Engine/Classes/PhysicsEngine/BodySetup.h
+++ b/Engine/Source/Runtime/Engine/Classes/PhysicsEngine/BodySetup.h
@@ -121,7 +121,7 @@ class UBodySetup : public UObject
 
 #if WITH_PHYSX
 	/** Physics triangle mesh, created from cooked data in CreatePhysicsMeshes */
-	physx::PxTriangleMesh* TriMesh;
+	TArray<physx::PxTriangleMesh*> TriMeshes;
 #endif
 
 	/** Flag used to know if we have created the physics convex and tri meshes from the cooked data yet */
diff --git a/Engine/Source/Runtime/Engine/Private/AI/Navigation/RecastNavMeshGenerator.cpp b/Engine/Source/Runtime/Engine/Private/AI/Navigation/RecastNavMeshGenerator.cpp
--- a/Engine/Source/Runtime/Engine/Private/AI/Navigation/RecastNavMeshGenerator.cpp
+++ b/Engine/Source/Runtime/Engine/Private/AI/Navigation/RecastNavMeshGenerator.cpp
@@ -712,15 +712,18 @@ FORCEINLINE_DEBUGGABLE void ExportRigidBodyTriMesh(UBodySetup& BodySetup, TNavSt
 												   FBox& UnrealBounds, const FTransform& LocalToWorld)
 {
 #if WITH_PHYSX
-	if (BodySetup.TriMesh != NULL && BodySetup.CollisionTraceFlag == CTF_UseComplexAsSimple)
+	if (BodySetup.CollisionTraceFlag == CTF_UseComplexAsSimple)
 	{
-		if (BodySetup.TriMesh->getTriangleMeshFlags() & PxTriangleMeshFlag::eHAS_16BIT_TRIANGLE_INDICES)
+		for(PxTriangleMesh* TriMesh : BodySetup.TriMeshes)
 		{
-			ExportPxTriMesh<PxU16>(BodySetup.TriMesh, LocalToWorld, VertexBuffer, IndexBuffer, UnrealBounds);
-		}
-		else
-		{
-			ExportPxTriMesh<PxU32>(BodySetup.TriMesh, LocalToWorld, VertexBuffer, IndexBuffer, UnrealBounds);
+			if (TriMesh->getTriangleMeshFlags() & PxTriangleMeshFlag::eHAS_16BIT_TRIANGLE_INDICES)
+			{
+				ExportPxTriMesh<PxU16>(TriMesh, LocalToWorld, VertexBuffer, IndexBuffer, UnrealBounds);
+			}
+			else
+			{
+				ExportPxTriMesh<PxU32>(TriMesh, LocalToWorld, VertexBuffer, IndexBuffer, UnrealBounds);
+			}
 		}
 	}
 #endif // WITH_PHYSX
diff --git a/Engine/Source/Runtime/Engine/Private/Components/SplineMeshComponent.cpp b/Engine/Source/Runtime/Engine/Private/Components/SplineMeshComponent.cpp
--- a/Engine/Source/Runtime/Engine/Private/Components/SplineMeshComponent.cpp
+++ b/Engine/Source/Runtime/Engine/Private/Components/SplineMeshComponent.cpp
@@ -774,7 +774,7 @@ UBodySetup* USplineMeshComponent::GetBodySetup()
 {
 	// Don't return a body setup that has no collision, it means we are interactively moving the spline and don't want to build collision.
 	// Instead we explicitly build collision with USplineMeshComponent::RecreateCollision()
-	if (BodySetup != NULL && (BodySetup->TriMesh != NULL || BodySetup->AggGeom.GetElementCount() > 0))
+	if (BodySetup != NULL && (BodySetup->TriMeshes.Num() || BodySetup->AggGeom.GetElementCount() > 0))
 	{
 		return BodySetup;
 	}
diff --git a/Engine/Source/Runtime/Engine/Private/PhysicsEngine/BodyInstance.cpp b/Engine/Source/Runtime/Engine/Private/PhysicsEngine/BodyInstance.cpp
--- a/Engine/Source/Runtime/Engine/Private/PhysicsEngine/BodyInstance.cpp
+++ b/Engine/Source/Runtime/Engine/Private/PhysicsEngine/BodyInstance.cpp
@@ -841,7 +841,7 @@ void FBodyInstance::UpdatePhysicsShapeFilterData(uint32 SkelMeshCompID, bool bUs
 #endif
 
 /** Update the filter data on the physics shapes, based on the owning component flags. */
-void FBodyInstance::UpdatePhysicsFilterData(bool bForceSimpleAsComplex)
+void FBodyInstance::UpdatePhysicsFilterData()
 {
 	SCOPE_CYCLE_COUNTER(STAT_UpdatePhysFilter);
 
@@ -929,8 +929,8 @@ void FBodyInstance::UpdatePhysicsFilterData(bool bForceSimpleAsComplex)
 	}
 #endif
 
-	const bool bUseComplexAsSimple = !bForceSimpleAsComplex && (BodySetup.Get()->CollisionTraceFlag == CTF_UseComplexAsSimple);
-	const bool bUseSimpleAsComplex = bForceSimpleAsComplex || (BodySetup.Get()->CollisionTraceFlag == CTF_UseSimpleAsComplex);
+	const bool bUseComplexAsSimple = (BodySetup.Get()->CollisionTraceFlag == CTF_UseComplexAsSimple);
+	const bool bUseSimpleAsComplex = (BodySetup.Get()->CollisionTraceFlag == CTF_UseSimpleAsComplex);
 
 #if WITH_PHYSX
 	const TEnumAsByte<ECollisionEnabled::Type>* CollisionEnabledOverride = bUseCollisionEnabledOverride ? &UseCollisionEnabled : NULL;
@@ -1961,7 +1961,7 @@ void FBodyInstance::PostShapeChange()
 	UpdatePhysicalMaterials();
 
 	// Set the filter data on the shapes (call this after setting BodyData because it uses that pointer)
-	UpdatePhysicsFilterData(true);
+	UpdatePhysicsFilterData();
 
 	UpdateMassProperties();
 	// Update damping
@@ -2271,27 +2271,33 @@ bool FBodyInstance::UpdateBodyScale(const FVector& InScale3D)
 				{
 					PShape->getTriangleMeshGeometry(PTriMeshGeom);
 
-					// Create tri-mesh shape
-					if (BodySetup.IsValid() && BodySetup->TriMesh)
+					// find which trimesh elems it is
+					// it would be nice to know if the order of PShapes array index is in the order of createShape
+					if (BodySetup.IsValid())
 					{
 						// Please note that this one we don't inverse old scale, but just set new one (but still adjust for scale mode)
 						FVector NewScale3D = RelativeScale3D * OldScale3D;
-						if (BodySetup->TriMesh)
-						{
-							PTriMeshGeom.triangleMesh = BodySetup->TriMesh;
-							PTriMeshGeom.scale.scale = U2PVector(Scale3D);
 
-							if (PTriMeshGeom.isValid())
-							{
-								UpdatedGeometry = &PTriMeshGeom;
-								bSuccess = true;
-							}
-							else
+						for (PxTriangleMesh* TriMesh : BodySetup->TriMeshes)
+						{
+							// found it
+							if (TriMesh == PTriMeshGeom.triangleMesh)
 							{
-								bInvalid = true;
+								PTriMeshGeom.scale.scale = U2PVector(Scale3D);
+
+								if (PTriMeshGeom.isValid())
+								{
+									UpdatedGeometry = &PTriMeshGeom;
+									bSuccess = true;
+								}
+								else
+								{
+									bInvalid = true;
+								}
 							}
 						}
 					}
+
 					break;
 				}
 				case PxGeometryType::eHEIGHTFIELD:
diff --git a/Engine/Source/Runtime/Engine/Private/PhysicsEngine/BodySetup.cpp b/Engine/Source/Runtime/Engine/Private/PhysicsEngine/BodySetup.cpp
--- a/Engine/Source/Runtime/Engine/Private/PhysicsEngine/BodySetup.cpp
+++ b/Engine/Source/Runtime/Engine/Private/PhysicsEngine/BodySetup.cpp
@@ -168,8 +168,12 @@ void UBodySetup::CreatePhysicsMeshes()
 			}
 		}
 
-		TriMesh = CookedDataReader.TriMesh;
-		FPhysxSharedData::Get().Add(TriMesh);
+		for(PxTriangleMesh* TriMesh : CookedDataReader.TriMeshes)
+		{
+			check(TriMesh);
+			TriMeshes.Add(TriMesh);
+			FPhysxSharedData::Get().Add(TriMesh);
+		}
 	}
 	else
 	{
@@ -211,12 +215,13 @@ void UBodySetup::ClearPhysicsMeshes()
 		}
 	}
 
-	if(TriMesh != NULL)
+	for(int32 ElementIndex = 0; ElementIndex < TriMeshes.Num(); ++ElementIndex)
 	{
-		GPhysXPendingKillTriMesh.Add(TriMesh);
-		FPhysxSharedData::Get().Remove(TriMesh);
-		TriMesh = NULL;
+		GPhysXPendingKillTriMesh.Add(TriMeshes[ElementIndex]);
+		FPhysxSharedData::Get().Remove(TriMeshes[ElementIndex]);
+		TriMeshes[ElementIndex] = NULL;
 	}
+	TriMeshes.Empty();
 
 	bCreatedPhysicsMeshes = false;
 #endif
@@ -521,10 +526,11 @@ public:
 		float ContactOffsetFactor, MaxContactOffset;
 		GetContactOffsetParams(ContactOffsetFactor, MaxContactOffset);
 
-		if (BodySetup->TriMesh)
+		for(PxTriangleMesh* TriMesh : BodySetup->TriMeshes)
 		{
+		
 			PxTriangleMeshGeometry PTriMeshGeom;
-			PTriMeshGeom.triangleMesh = BodySetup->TriMesh;
+			PTriMeshGeom.triangleMesh = TriMesh;
 			PTriMeshGeom.scale.scale = U2PVector(Scale3D);
 			if (BodySetup->bDoubleSidedGeometry)
 			{
@@ -534,13 +540,11 @@ public:
 
 			if (PTriMeshGeom.isValid())
 			{
-
+				PxTransform PElementTransform = U2PTransform(RelativeTM);
 				// Create without 'sim shape' flag, problematic if it's kinematic, and it gets set later anyway.
+				if (!AttachShape_AssumesLocked(PTriMeshGeom, PElementTransform, MaxContactOffset, PxShapeFlag::eSCENE_QUERY_SHAPE | PxShapeFlag::eVISUALIZATION))
 				{
-					if (!AttachShape_AssumesLocked(PTriMeshGeom, PxTransform::createIdentity(), MaxContactOffset, PxShapeFlag::eSCENE_QUERY_SHAPE | PxShapeFlag::eVISUALIZATION))
-					{
-						UE_LOG(LogPhysics, Log, TEXT("Can't create new mesh shape in AddShapesToRigidActor"));
-					}
+					UE_LOG(LogPhysics, Log, TEXT("Can't create new mesh shape in AddShapesToRigidActor"));
 				}
 			}
 			else
@@ -857,16 +861,16 @@ void UBodySetup::UpdateTriMeshVertices(const TArray<FVector> & NewPositions)
 {
 	SCOPE_CYCLE_COUNTER(STAT_UpdateTriMeshVertices);
 #if WITH_PHYSX
-	if (TriMesh)
+	if (TriMeshes.Num())
 	{
 		
-		PxVec3 * PNewPositions = TriMesh->getVerticesForModification();
+		PxVec3 * PNewPositions = TriMeshes[0]->getVerticesForModification();	//we only update the first trimesh. We assume this per poly case is not updating welded trimeshes
 		for (int32 i = 0; i < NewPositions.Num(); ++i)
 		{
 			PNewPositions[i] = U2PVector(NewPositions[i]);
 		}
 
-		TriMesh->refitBVH();
+		TriMeshes[0]->refitBVH();
 	}
 #endif
 }
@@ -962,7 +966,7 @@ SIZE_T UBodySetup::GetResourceSize( EResourceSizeMode::Type Mode )
 
 #if WITH_PHYSX
 	// Count PhysX trimesh mem usage
-	if (TriMesh != NULL)
+	for(PxTriangleMesh* TriMesh : TriMeshes)
 	{
 		ResourceSize += GetPhysxObjectSize(TriMesh, NULL);
 	}
diff --git a/Engine/Source/Runtime/Engine/Private/PhysicsEngine/PhysDerivedData.cpp b/Engine/Source/Runtime/Engine/Private/PhysicsEngine/PhysDerivedData.cpp
--- a/Engine/Source/Runtime/Engine/Private/PhysicsEngine/PhysDerivedData.cpp
+++ b/Engine/Source/Runtime/Engine/Private/PhysicsEngine/PhysDerivedData.cpp
@@ -56,12 +56,12 @@ bool FDerivedDataPhysXCooker::Build( TArray<uint8>& OutData )
 	uint8 bLittleEndian = PLATFORM_LITTLE_ENDIAN;	//TODO: We should pass the target platform into this function and write it. Then swap the endian on the writer so the reader doesn't have to do it at runtime
 	int32 NumConvexElementsCooked = 0;
 	int32 NumMirroredElementsCooked = 0;
-	uint8 bTriMeshCooked = false;
+	int32 NumTriMeshesCooked = 0;
 	Ar << bLittleEndian;
 	int64 CookedMeshInfoOffset = Ar.Tell();
 	Ar << NumConvexElementsCooked;	
 	Ar << NumMirroredElementsCooked;
-	Ar << bTriMeshCooked;
+	Ar << NumTriMeshesCooked;
 
 	//TODO: we must save an id with convex and tri meshes for serialization. We must save this here and patch it up at runtime somehow
 
@@ -82,14 +82,14 @@ bool FDerivedDataPhysXCooker::Build( TArray<uint8>& OutData )
 	bool bUsingAllTriData = BodySetup != NULL ? BodySetup->bMeshCollideAll : false;
 	if( (BodySetup == NULL || BodySetup->CollisionTraceFlag != CTF_UseSimpleAsComplex) && ShouldGenerateTriMeshData(bUsingAllTriData) )
 	{
-		bTriMeshCooked = (uint8)BuildTriMesh( OutData, false, bUsingAllTriData );
+		NumTriMeshesCooked = BuildTriMesh( OutData, bUsingAllTriData );
 	}
 
 	// Update info on what actually got cooked
 	Ar.Seek( CookedMeshInfoOffset );
 	Ar << NumConvexElementsCooked;	
 	Ar << NumMirroredElementsCooked;
-	Ar << bTriMeshCooked;
+	Ar << NumTriMeshesCooked;
 
 	// Whatever got cached return true. We want to cache 'failure' too.
 	return true;
@@ -143,7 +143,7 @@ bool FDerivedDataPhysXCooker::ShouldGenerateTriMeshData(bool InUseAllTriData)
 	return bPerformCook;
 }
 
-bool FDerivedDataPhysXCooker::BuildTriMesh( TArray<uint8>& OutData, bool bInMirrored, bool InUseAllTriData )
+int32 FDerivedDataPhysXCooker::BuildTriMesh( TArray<uint8>& OutData, bool InUseAllTriData )
 {
 	check(Cooker != NULL);
 
@@ -166,37 +166,23 @@ bool FDerivedDataPhysXCooker::BuildTriMesh( TArray<uint8>& OutData, bool bInMirr
 		}
 
 		TArray<FVector>* MeshVertices = NULL;
-		TArray<FVector> MirroredVerts;
-
-		if( bInMirrored )
-		{
-			MirroredVerts.AddUninitialized(NumVerts);
-			for(int32 VertIdx=0; VertIdx<NumVerts; VertIdx++)
-			{
-				MirroredVerts[VertIdx] = TriangleMeshDesc.Vertices[VertIdx] * FVector(-1,1,1);
-			}
-			MeshVertices = &MirroredVerts;
-		}
-		else
-		{
-			MeshVertices = &TriangleMeshDesc.Vertices;
-		}
-
-		UE_LOG(LogPhysics, Log, TEXT("Cook TriMesh: %s (FlipX: %d)"), *CollisionDataProvider->GetPathName(), bInMirrored);
+		MeshVertices = &TriangleMeshDesc.Vertices;
+		
+		UE_LOG(LogPhysics, Log, TEXT("Cook TriMesh: %s"), *CollisionDataProvider->GetPathName());
 		bool bPerPolySkeletalMesh = false;
 		if (USkeletalMesh* SkeletalMesh = Cast<USkeletalMesh>(CollisionDataProvider))
 		{
 			ensure(SkeletalMesh->bEnablePerPolyCollision);
 			bPerPolySkeletalMesh = true;
 		}
-		bResult = Cooker->CookTriMesh( Format, *MeshVertices, TriangleMeshDesc.Indices, TriangleMeshDesc.MaterialIndices, bInMirrored ? !TriangleMeshDesc.bFlipNormals : TriangleMeshDesc.bFlipNormals, OutData, bPerPolySkeletalMesh );
+		bResult = Cooker->CookTriMesh( Format, *MeshVertices, TriangleMeshDesc.Indices, TriangleMeshDesc.MaterialIndices, TriangleMeshDesc.bFlipNormals, OutData, bPerPolySkeletalMesh );
 		if( !bResult )
 		{
-			UE_LOG(LogPhysics, Warning, TEXT("Failed to cook TriMesh: %s (FlipX:%d)."), *CollisionDataProvider->GetPathName(), bInMirrored );
+			UE_LOG(LogPhysics, Warning, TEXT("Failed to cook TriMesh: %s."), *CollisionDataProvider->GetPathName());
 		}
 	}
 
-	return bResult;
+	return bResult == true ? 1 : 0;	//the cooker only generates 1 or 0 trimeshes. We return an int because we support multiple trimeshes for welding and we might want to do this per static mesh in the future.
 }
 
 FDerivedDataPhysXBinarySerializer::FDerivedDataPhysXBinarySerializer(FName InFormat, const TArray<FBodyInstance*>& InBodies, const TArray<class UBodySetup*>& InBodySetups, const TArray<class UPhysicalMaterial*>& InPhysicalMaterials, const FGuid& InGuid)
diff --git a/Engine/Source/Runtime/Engine/Private/PhysicsEngine/PhysDerivedData.h b/Engine/Source/Runtime/Engine/Private/PhysicsEngine/PhysDerivedData.h
--- a/Engine/Source/Runtime/Engine/Private/PhysicsEngine/PhysDerivedData.h
+++ b/Engine/Source/Runtime/Engine/Private/PhysicsEngine/PhysDerivedData.h
@@ -38,7 +38,7 @@ public:
 		// This is a version string that mimics the old versioning scheme. If you
 		// want to bump this version, generate a new guid using VS->Tools->Create GUID and
 		// return it here. Ex.
-		return TEXT("7F9691E5320E4745AE426E72F60B484D");
+		return TEXT("1D9C17C451EA4FB88CB7A27536A65DB0");
 	}
 
 	virtual FString GetPluginSpecificCacheKeySuffix() const override
@@ -81,7 +81,7 @@ private:
 
 	void InitCooker();
 	int32 BuildConvex( TArray<uint8>& OutData, bool InMirrored );
-	bool BuildTriMesh( TArray<uint8>& OutData, bool InMirrored, bool InUseAllTriData );
+	int32 BuildTriMesh( TArray<uint8>& OutData, bool InUseAllTriData );
 	bool ShouldGenerateTriMeshData(bool InUseAllTriData);
 };
 
diff --git a/Engine/Source/Runtime/Engine/Private/PhysicsEngine/PhysXSupport.cpp b/Engine/Source/Runtime/Engine/Private/PhysicsEngine/PhysXSupport.cpp
--- a/Engine/Source/Runtime/Engine/Private/PhysicsEngine/PhysXSupport.cpp
+++ b/Engine/Source/Runtime/Engine/Private/PhysicsEngine/PhysXSupport.cpp
@@ -644,7 +644,6 @@ PxU32 FPhysXCPUDispatcherSingleThread::getWorkerCount() const
 // FPhysXFormatDataReader
 
 FPhysXFormatDataReader::FPhysXFormatDataReader( FByteBulkData& InBulkData )
-	: TriMesh( NULL )
 {
 	// Read cooked physics data
 	uint8* DataPtr = (uint8*)InBulkData.Lock( LOCK_READ_ONLY );
@@ -653,33 +652,33 @@ FPhysXFormatDataReader::FPhysXFormatDataReader( FByteBulkData& InBulkData )
 	uint8 bLittleEndian = true;
 	int32 NumConvexElementsCooked = 0;
 	int32 NumMirroredElementsCooked = 0;
-	uint8 bTriMeshCooked = false;
+	int32 NumTriMeshesCooked = 0;
 
 	Ar << bLittleEndian;
 	Ar.SetByteSwapping( PLATFORM_LITTLE_ENDIAN ? !bLittleEndian : !!bLittleEndian );
 	Ar << NumConvexElementsCooked;	
 	Ar << NumMirroredElementsCooked;
-	Ar << bTriMeshCooked;
+	Ar << NumTriMeshesCooked;
 	
-	ConvexMeshes.Empty( NumConvexElementsCooked );
-	ConvexMeshesNegX.Empty( NumMirroredElementsCooked );
-
+	ConvexMeshes.Empty(NumConvexElementsCooked);
 	for( int32 ElementIndex = 0; ElementIndex < NumConvexElementsCooked; ElementIndex++ )
 	{
 		PxConvexMesh* ConvexMesh = ReadConvexMesh( Ar, DataPtr, InBulkData.GetBulkDataSize() );
 		ConvexMeshes.Add( ConvexMesh );
 	}
 
+	ConvexMeshesNegX.Empty(NumMirroredElementsCooked);
 	for( int32 ElementIndex = 0; ElementIndex < NumMirroredElementsCooked; ElementIndex++ )
 	{
 		PxConvexMesh* ConvexMeshNegX = ReadConvexMesh( Ar, DataPtr, InBulkData.GetBulkDataSize() );
 		ConvexMeshesNegX.Add( ConvexMeshNegX );
 	}
 
-	if( bTriMeshCooked )
+	TriMeshes.Empty(NumTriMeshesCooked);
+	for(int32 ElementIndex = 0; ElementIndex < NumTriMeshesCooked; ++ElementIndex)
 	{
-		TriMesh = ReadTriMesh( Ar, DataPtr, InBulkData.GetBulkDataSize() );
-		check(TriMesh);	
+		PxTriangleMesh* TriMesh = ReadTriMesh( Ar, DataPtr, InBulkData.GetBulkDataSize() );
+		TriMeshes.Add(TriMesh);
 	}
 
 	InBulkData.Unlock();
@@ -914,7 +913,10 @@ PxCollection* MakePhysXCollection(const TArray<UPhysicalMaterial*>& PhysicalMate
 
 	for (UBodySetup* BodySetup : BodySetups)
 	{
-		AddToCollection(PCollection, BodySetup->TriMesh);
+		for(PxTriangleMesh* TriMesh : BodySetup->TriMeshes)
+		{
+			AddToCollection(PCollection, TriMesh);
+		}
 
 		for (const FKConvexElem& ConvexElem : BodySetup->AggGeom.ConvexElems)
 		{
diff --git a/Engine/Source/Runtime/Engine/Private/PhysicsEngine/PhysXSupport.h b/Engine/Source/Runtime/Engine/Private/PhysicsEngine/PhysXSupport.h
--- a/Engine/Source/Runtime/Engine/Private/PhysicsEngine/PhysXSupport.h
+++ b/Engine/Source/Runtime/Engine/Private/PhysicsEngine/PhysXSupport.h
@@ -488,9 +488,9 @@ public:
 class FPhysXFormatDataReader
 {
 public:
-	TArray< PxConvexMesh* > ConvexMeshes;
-	TArray< PxConvexMesh* > ConvexMeshesNegX;
-	PxTriangleMesh* TriMesh;
+	TArray<PxConvexMesh*> ConvexMeshes;
+	TArray<PxConvexMesh*> ConvexMeshesNegX;
+	TArray<PxTriangleMesh*> TriMeshes;
 
 	FPhysXFormatDataReader( FByteBulkData& InBulkData );
 
diff --git a/Engine/Source/Runtime/Engine/Private/Vehicles/WheeledVehicleMovementComponent.cpp b/Engine/Source/Runtime/Engine/Private/Vehicles/WheeledVehicleMovementComponent.cpp
--- a/Engine/Source/Runtime/Engine/Private/Vehicles/WheeledVehicleMovementComponent.cpp
+++ b/Engine/Source/Runtime/Engine/Private/Vehicles/WheeledVehicleMovementComponent.cpp
@@ -236,9 +236,9 @@ void UWheeledVehicleMovementComponent::SetupVehicleShapes()
 					MeshScaleV.Z = Wheel->ShapeRadius / MeshBounds.BoxExtent.Z;
 				}
 				PxMeshScale MeshScale(U2PVector(UpdatedComponent->RelativeScale3D * MeshScaleV), PxQuat::createIdentity());
-				if (Wheel->CollisionMesh->BodySetup->TriMesh)
+				if (Wheel->CollisionMesh->BodySetup->TriMeshes.Num())
 				{
-					PxTriangleMesh* TriMesh = Wheel->CollisionMesh->BodySetup->TriMesh;
+					PxTriangleMesh* TriMesh = Wheel->CollisionMesh->BodySetup->TriMeshes[0];
 
 					// No eSIMULATION_SHAPE flag for wheels
 					PWheelShape = PVehicleActor->createShape(PxTriangleMeshGeometry(TriMesh, MeshScale), *WheelMaterial, PxShapeFlag::eSCENE_QUERY_SHAPE | PxShapeFlag::eVISUALIZATION);
