Fix for root motion not being cleared on additive animations

diff --git a/Engine/Source/Runtime/Engine/Private/Animation/AnimSingleNodeInstance.cpp b/Engine/Source/Runtime/Engine/Private/Animation/AnimSingleNodeInstance.cpp
--- a/Engine/Source/Runtime/Engine/Private/Animation/AnimSingleNodeInstance.cpp
+++ b/Engine/Source/Runtime/Engine/Private/Animation/AnimSingleNodeInstance.cpp
@@ -322,7 +322,7 @@ bool UAnimSingleNodeInstance::NativeEvaluateAnimation(FPoseContext& Output)
 				FCompactPose AdditivePose;
 				AdditivePose.SetBoneContainer(&Output.Pose.GetBoneContainer());
 
-				FAnimExtractContext ExtractionContext(CurrentTime);
+				FAnimExtractContext ExtractionContext(CurrentTime, Sequence->bEnableRootMotion);
 				Sequence->GetAdditiveBasePose(AdditivePose, ExtractionContext);
 				Sequence->GetAnimationPose(Output.Pose, ExtractionContext);
 				if (Sequence->AdditiveAnimType == AAT_LocalSpaceBase)
