Added a project wide setting for turning off the default per poly complex collision

#UE-16180

diff --git a/Engine/Source/Runtime/Engine/Classes/PhysicsEngine/BodySetup.h b/Engine/Source/Runtime/Engine/Classes/PhysicsEngine/BodySetup.h
--- a/Engine/Source/Runtime/Engine/Classes/PhysicsEngine/BodySetup.h
+++ b/Engine/Source/Runtime/Engine/Classes/PhysicsEngine/BodySetup.h
@@ -95,6 +95,8 @@ class UBodySetup : public UObject
 	UPROPERTY(EditAnywhere, Category=Collision, meta=(DisplayName = "Collision Complexity"))
 	TEnumAsByte<enum ECollisionTraceFlag> CollisionTraceFlag;
 
+	TEnumAsByte<enum ECollisionTraceFlag> GetCollisionTraceFlag() const;
+
 	/** Default properties of the body instance, copied into objects on instantiation, was URB_BodyInstance */
 	UPROPERTY(EditAnywhere, Category=Collision, meta=(FullyExpand = "true"))
 	struct FBodyInstance DefaultInstance;
diff --git a/Engine/Source/Runtime/Engine/Classes/PhysicsEngine/PhysicsSettings.h b/Engine/Source/Runtime/Engine/Classes/PhysicsEngine/PhysicsSettings.h
--- a/Engine/Source/Runtime/Engine/Classes/PhysicsEngine/PhysicsSettings.h
+++ b/Engine/Source/Runtime/Engine/Classes/PhysicsEngine/PhysicsSettings.h
@@ -162,6 +162,12 @@ class ENGINE_API UPhysicsSettings : public UDeveloperSettings
 	UPROPERTY(config, EditAnywhere, Category = Simulation)
 	bool bSimulateSkeletalMeshOnDedicatedServer;
 
+	
+	/**
+	*  If true, static meshes will use per poly collision as complex collision by default. If false the default behavior is the same as UseSimpleAsComplex. */
+	UPROPERTY(config, EditAnywhere, Category = Simulation)
+	bool bDefaultHasComplexCollision;
+
 
 	/** Max Physics Delta Time to be clamped. */
 	UPROPERTY(config, EditAnywhere, meta=(ClampMin="0.0013", UIMin = "0.0013", ClampMax="1.0", UIMax="1.0"), Category=Framerate)
diff --git a/Engine/Source/Runtime/Engine/Private/AI/Navigation/RecastNavMeshGenerator.cpp b/Engine/Source/Runtime/Engine/Private/AI/Navigation/RecastNavMeshGenerator.cpp
--- a/Engine/Source/Runtime/Engine/Private/AI/Navigation/RecastNavMeshGenerator.cpp
+++ b/Engine/Source/Runtime/Engine/Private/AI/Navigation/RecastNavMeshGenerator.cpp
@@ -712,7 +712,7 @@ FORCEINLINE_DEBUGGABLE void ExportRigidBodyTriMesh(UBodySetup& BodySetup, TNavSt
 												   FBox& UnrealBounds, const FTransform& LocalToWorld)
 {
 #if WITH_PHYSX
-	if (BodySetup.CollisionTraceFlag == CTF_UseComplexAsSimple)
+	if (BodySetup.GetCollisionTraceFlag() == CTF_UseComplexAsSimple)
 	{
 		for(PxTriangleMesh* TriMesh : BodySetup.TriMeshes)
 		{
diff --git a/Engine/Source/Runtime/Engine/Private/Components/SplineMeshComponent.cpp b/Engine/Source/Runtime/Engine/Private/Components/SplineMeshComponent.cpp
--- a/Engine/Source/Runtime/Engine/Private/Components/SplineMeshComponent.cpp
+++ b/Engine/Source/Runtime/Engine/Private/Components/SplineMeshComponent.cpp
@@ -859,7 +859,7 @@ void USplineMeshComponent::RecreateCollision()
 		BodySetup->BodySetupGuid = StaticMesh->BodySetup->BodySetupGuid;
 		CachedMeshBodySetupGuid = StaticMesh->BodySetup->BodySetupGuid;
 
-		if (BodySetup->CollisionTraceFlag == CTF_UseComplexAsSimple)
+		if (BodySetup->GetCollisionTraceFlag() == CTF_UseComplexAsSimple)
 		{
 			BodySetup->AggGeom.EmptyElements();
 		}
diff --git a/Engine/Source/Runtime/Engine/Private/Components/StaticMeshComponent.cpp b/Engine/Source/Runtime/Engine/Private/Components/StaticMeshComponent.cpp
--- a/Engine/Source/Runtime/Engine/Private/Components/StaticMeshComponent.cpp
+++ b/Engine/Source/Runtime/Engine/Private/Components/StaticMeshComponent.cpp
@@ -616,7 +616,7 @@ bool UStaticMeshComponent::CanEditSimulatePhysics()
 {
 	if (UBodySetup* BodySetup = GetBodySetup())
 	{
-		return (BodySetup->AggGeom.GetElementCount() > 0) || (BodySetup->CollisionTraceFlag == CTF_UseComplexAsSimple);
+		return (BodySetup->AggGeom.GetElementCount() > 0) || (BodySetup->GetCollisionTraceFlag() == CTF_UseComplexAsSimple);
 	}
 	else
 	{
diff --git a/Engine/Source/Runtime/Engine/Private/PhysicsEngine/BodyInstance.cpp b/Engine/Source/Runtime/Engine/Private/PhysicsEngine/BodyInstance.cpp
--- a/Engine/Source/Runtime/Engine/Private/PhysicsEngine/BodyInstance.cpp
+++ b/Engine/Source/Runtime/Engine/Private/PhysicsEngine/BodyInstance.cpp
@@ -929,8 +929,8 @@ void FBodyInstance::UpdatePhysicsFilterData()
 	}
 #endif
 
-	const bool bUseComplexAsSimple = (BodySetup.Get()->CollisionTraceFlag == CTF_UseComplexAsSimple);
-	const bool bUseSimpleAsComplex = (BodySetup.Get()->CollisionTraceFlag == CTF_UseSimpleAsComplex);
+	const bool bUseComplexAsSimple = (BodySetup.Get()->GetCollisionTraceFlag() == CTF_UseComplexAsSimple);
+	const bool bUseSimpleAsComplex = (BodySetup.Get()->GetCollisionTraceFlag() == CTF_UseSimpleAsComplex);
 
 #if WITH_PHYSX
 	const TEnumAsByte<ECollisionEnabled::Type>* CollisionEnabledOverride = bUseCollisionEnabledOverride ? &UseCollisionEnabled : NULL;
@@ -1101,7 +1101,7 @@ struct FInitBodiesHelper
 
 		FShapeData ShapeData;
 		Instance->GetFilterData_AssumesLocked(ShapeData);
-		Instance->GetShapeFlags_AssumesLocked(ShapeData, ShapeData.CollisionEnabled, BodySetup->CollisionTraceFlag == CTF_UseComplexAsSimple);
+		Instance->GetShapeFlags_AssumesLocked(ShapeData, ShapeData.CollisionEnabled, BodySetup->GetCollisionTraceFlag() == CTF_UseComplexAsSimple);
 
 		if (!bCompileStatic && PNewDynamic)
 		{
@@ -1864,7 +1864,7 @@ bool FBodyInstance::Weld(FBodyInstance* TheirBody, const FTransform& TheirTM)
 
 	FShapeData ShapeData;
 	GetFilterData_AssumesLocked(ShapeData);
-	GetShapeFlags_AssumesLocked(ShapeData, ShapeData.CollisionEnabled, BodySetup->CollisionTraceFlag == CTF_UseComplexAsSimple);
+	GetShapeFlags_AssumesLocked(ShapeData, ShapeData.CollisionEnabled, BodySetup->GetCollisionTraceFlag() == CTF_UseComplexAsSimple);
 
 	//child body gets placed into the same scenes as parent body
 	if (PxRigidActor* MyBody = RigidActorSync)
@@ -4599,8 +4599,8 @@ void FBodyInstance::GetFilterData_AssumesLocked(FShapeData& ShapeData, bool bFor
 	}
 #endif
 
-	const bool bUseComplexAsSimple = !bForceSimpleAsComplex && (BodySetup.Get()->CollisionTraceFlag == CTF_UseComplexAsSimple);
-	const bool bUseSimpleAsComplex = bForceSimpleAsComplex || (BodySetup.Get()->CollisionTraceFlag == CTF_UseSimpleAsComplex);
+	const bool bUseComplexAsSimple = !bForceSimpleAsComplex && (BodySetup.Get()->GetCollisionTraceFlag() == CTF_UseComplexAsSimple);
+	const bool bUseSimpleAsComplex = bForceSimpleAsComplex || (BodySetup.Get()->GetCollisionTraceFlag() == CTF_UseSimpleAsComplex);
 
 	if(GetPxRigidActor_AssumesLocked())
 	{
diff --git a/Engine/Source/Runtime/Engine/Private/PhysicsEngine/BodySetup.cpp b/Engine/Source/Runtime/Engine/Private/PhysicsEngine/BodySetup.cpp
--- a/Engine/Source/Runtime/Engine/Private/PhysicsEngine/BodySetup.cpp
+++ b/Engine/Source/Runtime/Engine/Private/PhysicsEngine/BodySetup.cpp
@@ -122,7 +122,7 @@ void UBodySetup::CreatePhysicsMeshes()
 
 		FPhysXFormatDataReader CookedDataReader(*FormatData);
 
-		if (CollisionTraceFlag != CTF_UseComplexAsSimple)
+		if (GetCollisionTraceFlag() != CTF_UseComplexAsSimple)
 		{
 			bool bNeedsCooking = bGenerateNonMirroredCollision && CookedDataReader.ConvexMeshes.Num() != AggGeom.ConvexElems.Num();
 			bNeedsCooking = bNeedsCooking || (bGenerateMirroredCollision && CookedDataReader.ConvexMeshesNegX.Num() != AggGeom.ConvexElems.Num());
@@ -138,7 +138,7 @@ void UBodySetup::CreatePhysicsMeshes()
 		
 		ClearPhysicsMeshes();
 
-		if (CollisionTraceFlag != CTF_UseComplexAsSimple)
+		if (GetCollisionTraceFlag() != CTF_UseComplexAsSimple)
 		{
 
 			ensure(!bGenerateNonMirroredCollision || CookedDataReader.ConvexMeshes.Num() == 0 || CookedDataReader.ConvexMeshes.Num() == AggGeom.ConvexElems.Num());
@@ -607,7 +607,7 @@ void UBodySetup::AddShapesToRigidActor_AssumesLocked(FBodyInstance* OwningInstan
 
 	// Create shapes for simple collision if we do not want to use the complex collision mesh 
 	// for simple queries as well
-	if (CollisionTraceFlag != ECollisionTraceFlag::CTF_UseComplexAsSimple)
+	if (GetCollisionTraceFlag() != ECollisionTraceFlag::CTF_UseComplexAsSimple)
 	{
 		AddShapesHelper.AddSpheresToRigidActor_AssumesLocked();
 		AddShapesHelper.AddBoxesToRigidActor_AssumesLocked();
@@ -617,7 +617,7 @@ void UBodySetup::AddShapesToRigidActor_AssumesLocked(FBodyInstance* OwningInstan
 
 	// Create tri-mesh shape, when we are not using simple collision shapes for 
 	// complex queries as well
-	if (CollisionTraceFlag != ECollisionTraceFlag::CTF_UseSimpleAsComplex)
+	if (GetCollisionTraceFlag() != ECollisionTraceFlag::CTF_UseSimpleAsComplex)
 	{
 		AddShapesHelper.AddTriMeshToRigidActor_AssumesLocked();
 	}
@@ -1251,3 +1251,9 @@ float UBodySetup::GetVolume(const FVector& Scale) const
 {
 	return AggGeom.GetVolume(Scale);
 }
+
+TEnumAsByte<enum ECollisionTraceFlag> UBodySetup::GetCollisionTraceFlag() const
+{
+	ECollisionTraceFlag DefaultFlag = UPhysicsSettings::Get()->bDefaultHasComplexCollision ? ECollisionTraceFlag::CTF_UseDefault : ECollisionTraceFlag::CTF_UseSimpleAsComplex;
+	return CollisionTraceFlag == ECollisionTraceFlag::CTF_UseDefault ? DefaultFlag : CollisionTraceFlag;
+}
\ No newline at end of file
diff --git a/Engine/Source/Runtime/Engine/Private/PhysicsEngine/PhysDerivedData.cpp b/Engine/Source/Runtime/Engine/Private/PhysicsEngine/PhysDerivedData.cpp
--- a/Engine/Source/Runtime/Engine/Private/PhysicsEngine/PhysDerivedData.cpp
+++ b/Engine/Source/Runtime/Engine/Private/PhysicsEngine/PhysDerivedData.cpp
@@ -66,7 +66,7 @@ bool FDerivedDataPhysXCooker::Build( TArray<uint8>& OutData )
 	//TODO: we must save an id with convex and tri meshes for serialization. We must save this here and patch it up at runtime somehow
 
 	// Cook convex meshes, but only if we are not forcing complex collision to be used as simple collision as well
-	if( BodySetup && BodySetup->CollisionTraceFlag != CTF_UseComplexAsSimple && BodySetup->AggGeom.ConvexElems.Num() > 0 )
+	if( BodySetup && BodySetup->GetCollisionTraceFlag() != CTF_UseComplexAsSimple && BodySetup->AggGeom.ConvexElems.Num() > 0 )
 	{
 		if( bGenerateNormalMesh )
 		{
@@ -80,7 +80,7 @@ bool FDerivedDataPhysXCooker::Build( TArray<uint8>& OutData )
 
 	// Cook trimeshes, but only if we do not frce simple collision to be used as complex collision as well
 	bool bUsingAllTriData = BodySetup != NULL ? BodySetup->bMeshCollideAll : false;
-	if( (BodySetup == NULL || BodySetup->CollisionTraceFlag != CTF_UseSimpleAsComplex) && ShouldGenerateTriMeshData(bUsingAllTriData) )
+	if( (BodySetup == NULL || BodySetup->GetCollisionTraceFlag() != CTF_UseSimpleAsComplex) && ShouldGenerateTriMeshData(bUsingAllTriData) )
 	{
 		NumTriMeshesCooked = BuildTriMesh( OutData, bUsingAllTriData );
 	}
diff --git a/Engine/Source/Runtime/Engine/Private/PhysicsEngine/PhysicsSettings.cpp b/Engine/Source/Runtime/Engine/Private/PhysicsEngine/PhysicsSettings.cpp
--- a/Engine/Source/Runtime/Engine/Private/PhysicsEngine/PhysicsSettings.cpp
+++ b/Engine/Source/Runtime/Engine/Private/PhysicsEngine/PhysicsSettings.cpp
@@ -20,6 +20,7 @@ UPhysicsSettings::UPhysicsSettings(const FObjectInitializer& ObjectInitializer)
 	, LockedAxis_DEPRECATED(ESettingsLockedAxis::Invalid)
 	, BounceThresholdVelocity(200.f)
 	, bSimulateSkeletalMeshOnDedicatedServer(true)
+	, bDefaultHasComplexCollision(true)
 	, MaxPhysicsDeltaTime(1.f / 30.f)
 	, bSubstepping(false)
 	, bSubsteppingAsync(false)
diff --git a/Engine/Source/Runtime/Engine/Private/StaticMeshRender.cpp b/Engine/Source/Runtime/Engine/Private/StaticMeshRender.cpp
--- a/Engine/Source/Runtime/Engine/Private/StaticMeshRender.cpp
+++ b/Engine/Source/Runtime/Engine/Private/StaticMeshRender.cpp
@@ -137,7 +137,7 @@ FStaticMeshSceneProxy::FStaticMeshSceneProxy(UStaticMeshComponent* InComponent):
 
 	if (BodySetup)
 	{
-		CollisionTraceFlag = BodySetup->CollisionTraceFlag;
+		CollisionTraceFlag = BodySetup->GetCollisionTraceFlag();
 	}
 }
 
@@ -578,8 +578,11 @@ bool FStaticMeshSceneProxy::IsCollisionView(const FEngineShowFlags& EngineShowFl
 
 		if(bHasResponse)
 		{
-			bDrawComplexCollision = EngineShowFlags.CollisionVisibility;
-			bDrawSimpleCollision = EngineShowFlags.CollisionPawn;
+			const ECollisionTraceFlag TraceFlag = BodySetup ? BodySetup->GetCollisionTraceFlag() : ECollisionTraceFlag::CTF_UseDefault;
+			
+			//Visiblity uses complex and pawn uses simple. However, if UseSimpleAsComplex or UseComplexAsSimple is used we need to adjust accordingly
+			bDrawComplexCollision = (EngineShowFlags.CollisionVisibility && TraceFlag != ECollisionTraceFlag::CTF_UseSimpleAsComplex) || (EngineShowFlags.CollisionPawn && TraceFlag == ECollisionTraceFlag::CTF_UseComplexAsSimple);
+			bDrawSimpleCollision  = (EngineShowFlags.CollisionPawn && TraceFlag != ECollisionTraceFlag::CTF_UseComplexAsSimple) || (EngineShowFlags.CollisionVisibility && TraceFlag == ECollisionTraceFlag::CTF_UseSimpleAsComplex);
 		}
 	}
 
