Rename bPendingKillPending to bActorIsBeingDestroyed and made it private
Added blueprint callable IsActorBeingDestroyed utility function
Added similar bIsBeingDestroyed recursion protection in to ActorComponent

diff --git a/Engine/Source/Runtime/Engine/Classes/Components/ActorComponent.h b/Engine/Source/Runtime/Engine/Classes/Components/ActorComponent.h
--- a/Engine/Source/Runtime/Engine/Classes/Components/ActorComponent.h
+++ b/Engine/Source/Runtime/Engine/Classes/Components/ActorComponent.h
@@ -136,6 +136,9 @@ private:
 	/** Indicates that BeginPlay has been called, but EndPlay has not yet */
 	uint32 bHasBegunPlay:1;
 
+	/** Indicates the the destruction process has begun for this component to avoid recursion */
+	uint32 bIsBeingDestroyed:1;
+
 #if WITH_EDITOR
 	/** During undo/redo it isn't safe to cache owner */
 	uint32 bCanUseCachedOwner:1;
@@ -170,6 +173,12 @@ public:
 	bool HasBeenInitialized() const { return bHasBeenInitialized; }
 	bool HasBegunPlay() const { return bHasBegunPlay; }
 
+	UFUNCTION(BlueprintCallable, Category="Components")
+	bool IsBeingDestroyed() const
+	{
+		return bIsBeingDestroyed;
+	}
+
 	bool IsCreatedByConstructionScript() const;
 
 	UFUNCTION()
diff --git a/Engine/Source/Runtime/Engine/Classes/GameFramework/Actor.h b/Engine/Source/Runtime/Engine/Classes/GameFramework/Actor.h
--- a/Engine/Source/Runtime/Engine/Classes/GameFramework/Actor.h
+++ b/Engine/Source/Runtime/Engine/Classes/GameFramework/Actor.h
@@ -317,12 +317,14 @@ public:
 	UPROPERTY(EditAnywhere, BlueprintReadWrite, SaveGame, Replicated, Category=Actor)
 	uint32 bCanBeDamaged:1;
 
+private:
 	/**
 	 * Set when actor is about to be deleted.
-	 * @see IsPendingKillPending()
 	 */
 	UPROPERTY(Transient, DuplicateTransient)
-	uint32 bPendingKillPending:1;    
+	uint32 bActorIsBeingDestroyed:1;    
+
+public:
 
 	/** This actor collides with the world when placing in the editor or when spawned, even if RootComponent collision is disabled */
 	UPROPERTY()
@@ -1023,6 +1025,12 @@ public:
 	/** Returns whether an actor has had BeginPlay called on it (and not subsequently had EndPlay called) */
 	bool HasActorBegunPlay() const { return bActorHasBegunPlay; }
 
+	UFUNCTION(BlueprintCallable, Category="Game")
+	bool IsActorBeingDestroyed() const 
+	{
+		return bActorIsBeingDestroyed;
+	}
+
 	/** Event when this actor takes ANY damage */
 	UFUNCTION(BlueprintImplementableEvent, BlueprintAuthorityOnly, meta=(DisplayName = "AnyDamage"), Category="Game|Damage")
 	void ReceiveAnyDamage(float Damage, const class UDamageType* DamageType, class AController* InstigatedBy, AActor* DamageCauser);
@@ -1835,7 +1843,7 @@ public:
 	 **/
 	inline bool IsPendingKillPending() const
 	{
-		return bPendingKillPending || IsPendingKill();
+		return bActorIsBeingDestroyed || IsPendingKill();
 	}
 
 	/** Invalidate lighting cache with default options. */
@@ -2432,6 +2440,19 @@ private:
 
 	// Helper that already assumes the Hit info is reversed, and avoids creating a temp FHitResult if possible.
 	void InternalDispatchBlockingHit(UPrimitiveComponent* MyComp, UPrimitiveComponent* OtherComp, bool bSelfMoved, FHitResult const& Hit);
+
+	friend struct FMarkActorIsBeingDestroyed;
+};
+
+struct FMarkActorIsBeingDestroyed
+{
+private:
+	FMarkActorIsBeingDestroyed(AActor* InActor)
+	{
+		InActor->bActorIsBeingDestroyed = true;
+	}
+
+	friend UWorld;
 };
 
 /**
diff --git a/Engine/Source/Runtime/Engine/Private/Actor.cpp b/Engine/Source/Runtime/Engine/Private/Actor.cpp
--- a/Engine/Source/Runtime/Engine/Private/Actor.cpp
+++ b/Engine/Source/Runtime/Engine/Private/Actor.cpp
@@ -92,7 +92,6 @@ void AActor::InitializeDefaults()
 	InputConsumeOption_DEPRECATED = ICO_ConsumeBoundKeys;
 	bBlockInput = false;
 	bCanBeDamaged = true;
-	bPendingKillPending = false;
 	bFindCameraComponentWhenViewTarget = true;
 	bAllowReceiveTickEventOnDedicatedServer = true;
 	bRelevantForNetworkReplays = true;
diff --git a/Engine/Source/Runtime/Engine/Private/Components/ActorComponent.cpp b/Engine/Source/Runtime/Engine/Private/Components/ActorComponent.cpp
--- a/Engine/Source/Runtime/Engine/Private/Components/ActorComponent.cpp
+++ b/Engine/Source/Runtime/Engine/Private/Components/ActorComponent.cpp
@@ -873,6 +873,14 @@ void UActorComponent::UnregisterComponent()
 
 void UActorComponent::DestroyComponent(bool bPromoteChildren/*= false*/)
 {
+	// Avoid re-entrancy
+	if (bIsBeingDestroyed)
+	{
+		return;
+	}
+
+	bIsBeingDestroyed = true;
+
 	if (bHasBegunPlay)
 	{
 		EndPlay(EEndPlayReason::Destroyed);
diff --git a/Engine/Source/Runtime/Engine/Private/LevelActor.cpp b/Engine/Source/Runtime/Engine/Private/LevelActor.cpp
--- a/Engine/Source/Runtime/Engine/Private/LevelActor.cpp
+++ b/Engine/Source/Runtime/Engine/Private/LevelActor.cpp
@@ -456,7 +456,7 @@ bool UWorld::DestroyActor( AActor* ThisActor, bool bNetForce, bool bShouldModify
 	}
 
 	// Prevent recursion
-	ThisActor->bPendingKillPending = true;
+	FMarkActorIsBeingDestroyed MarkActorIsBeingDestroyed(ThisActor);
 
 	// Notify the texture streaming manager about the destruction of this actor.
 	IStreamingManager::Get().NotifyActorDestroyed( ThisActor );
