Fix mac build

diff --git a/Engine/Source/Runtime/Engine/Private/PhysicsEngine/BodySetup.cpp b/Engine/Source/Runtime/Engine/Private/PhysicsEngine/BodySetup.cpp
--- a/Engine/Source/Runtime/Engine/Private/PhysicsEngine/BodySetup.cpp
+++ b/Engine/Source/Runtime/Engine/Private/PhysicsEngine/BodySetup.cpp
@@ -1254,6 +1254,6 @@ float UBodySetup::GetVolume(const FVector& Scale) const
 
 TEnumAsByte<enum ECollisionTraceFlag> UBodySetup::GetCollisionTraceFlag() const
 {
-	ECollisionTraceFlag DefaultFlag = UPhysicsSettings::Get()->bDefaultHasComplexCollision ? ECollisionTraceFlag::CTF_UseDefault : ECollisionTraceFlag::CTF_UseSimpleAsComplex;
+	TEnumAsByte<enum ECollisionTraceFlag> DefaultFlag = UPhysicsSettings::Get()->bDefaultHasComplexCollision ? ECollisionTraceFlag::CTF_UseDefault : ECollisionTraceFlag::CTF_UseSimpleAsComplex;
 	return CollisionTraceFlag == ECollisionTraceFlag::CTF_UseDefault ? DefaultFlag : CollisionTraceFlag;
 }
\ No newline at end of file
