Fix Clang compile error

diff --git a/Engine/Source/Runtime/Engine/Private/StaticMeshRender.cpp b/Engine/Source/Runtime/Engine/Private/StaticMeshRender.cpp
--- a/Engine/Source/Runtime/Engine/Private/StaticMeshRender.cpp
+++ b/Engine/Source/Runtime/Engine/Private/StaticMeshRender.cpp
@@ -578,7 +578,7 @@ bool FStaticMeshSceneProxy::IsCollisionView(const FEngineShowFlags& EngineShowFl
 
 		if(bHasResponse)
 		{
-			const ECollisionTraceFlag TraceFlag = BodySetup ? BodySetup->GetCollisionTraceFlag() : ECollisionTraceFlag::CTF_UseDefault;
+			const ECollisionTraceFlag TraceFlag = BodySetup ? BodySetup->GetCollisionTraceFlag().GetValue() : ECollisionTraceFlag::CTF_UseDefault;
 			
 			//Visiblity uses complex and pawn uses simple. However, if UseSimpleAsComplex or UseComplexAsSimple is used we need to adjust accordingly
 			bDrawComplexCollision = (EngineShowFlags.CollisionVisibility && TraceFlag != ECollisionTraceFlag::CTF_UseSimpleAsComplex) || (EngineShowFlags.CollisionPawn && TraceFlag == ECollisionTraceFlag::CTF_UseComplexAsSimple);
