If a template is specified but a name is not supplied via the actor spawn parameters, then use the template's name as a base for the new actor's name.

diff --git a/Engine/Source/Runtime/Engine/Private/LevelActor.cpp b/Engine/Source/Runtime/Engine/Private/LevelActor.cpp
--- a/Engine/Source/Runtime/Engine/Private/LevelActor.cpp
+++ b/Engine/Source/Runtime/Engine/Private/LevelActor.cpp
@@ -284,12 +284,24 @@ AActor* UWorld::SpawnActor( UClass* Class, FVector const* Location, FRotator con
 		
 	}
 
+	ULevel* LevelToSpawnIn = SpawnParameters.OverrideLevel;
+	if (LevelToSpawnIn == NULL)
+	{
+		// Spawn in the same level as the owner if we have one. @warning: this relies on the outer of an actor being the level.
+		LevelToSpawnIn = (SpawnParameters.Owner != NULL) ? CastChecked<ULevel>(SpawnParameters.Owner->GetOuter()) : CurrentLevel;
+	}
+
+	FName NewActorName = SpawnParameters.Name;
 	AActor* Template = SpawnParameters.Template;
 	// Use class's default actor as a template.
 	if( !Template )
 	{
 		Template = Class->GetDefaultObject<AActor>();
 	}
+	else if (NewActorName.IsNone())
+	{
+		NewActorName = MakeUniqueObjectName(LevelToSpawnIn, Template->GetClass(), *Template->GetFName().GetPlainNameString());
+	}
 	check(Template!=NULL);
 
 	// See if we can spawn on ded.server/client only etc (check NeedsLoadForClient & NeedsLoadForServer)
@@ -313,13 +325,7 @@ AActor* UWorld::SpawnActor( UClass* Class, FVector const* Location, FRotator con
 		}
 	}
 
-	ULevel* LevelToSpawnIn = SpawnParameters.OverrideLevel;
-	if (LevelToSpawnIn == NULL)
-	{
-		// Spawn in the same level as the owner if we have one. @warning: this relies on the outer of an actor being the level.
-		LevelToSpawnIn = (SpawnParameters.Owner != NULL) ? CastChecked<ULevel>(SpawnParameters.Owner->GetOuter()) : CurrentLevel;
-	}
-	AActor* Actor = NewObject<AActor>(LevelToSpawnIn, Class, SpawnParameters.Name, SpawnParameters.ObjectFlags, Template);
+	AActor* Actor = NewObject<AActor>(LevelToSpawnIn, Class, NewActorName, SpawnParameters.ObjectFlags, Template);
 	check(Actor);
 
 #if WITH_EDITOR
