Comparing InMaxPacket against MAX DEMO READ WRITE BUFFER instead of MAX PACKET SIZE for DemoNetDriver.cpp

diff --git a/Engine/Source/Runtime/Engine/Private/DemoNetDriver.cpp b/Engine/Source/Runtime/Engine/Private/DemoNetDriver.cpp
--- a/Engine/Source/Runtime/Engine/Private/DemoNetDriver.cpp
+++ b/Engine/Source/Runtime/Engine/Private/DemoNetDriver.cpp
@@ -1716,7 +1716,7 @@ void UDemoNetConnection::InitConnection( UNetDriver* InDriver, EConnectionState
 	// default implementation
 	Super::InitConnection( InDriver, InState, InURL, InConnectionSpeed );
 
-	MaxPacket = (InMaxPacket == 0 || InMaxPacket > MAX_PACKET_SIZE) ? MAX_DEMO_READ_WRITE_BUFFER : InMaxPacket;
+	MaxPacket = (InMaxPacket == 0 || InMaxPacket > MAX_DEMO_READ_WRITE_BUFFER) ? MAX_DEMO_READ_WRITE_BUFFER : InMaxPacket;
 	InternalAck = true;
 
 	InitSendBuffer();
