#ue4 - Fix build after CL 2567834. RemoveAt() is incompatible with TSet; queue up elements for removal to be processed later.

diff --git a/Engine/Source/Runtime/Engine/Private/DemoNetDriver.cpp b/Engine/Source/Runtime/Engine/Private/DemoNetDriver.cpp
--- a/Engine/Source/Runtime/Engine/Private/DemoNetDriver.cpp
+++ b/Engine/Source/Runtime/Engine/Private/DemoNetDriver.cpp
@@ -938,17 +938,18 @@ void UDemoNetDriver::TickDemoRecord( float DeltaSeconds )
 
 	DemoReplicateActor( World->GetWorldSettings(), ClientConnections[0], IsNetClient );
 
+	TArray<AActor*> ActorsToRemove;
 	for ( AActor* Actor : World->NetworkActors )
 	{
 		if ( Actor->IsPendingKill() )
 		{
-			World->NetworkActors.RemoveAtSwap( i );
+			ActorsToRemove.Add(Actor);
 			continue;
 		}
 
 		if ( Actor->GetRemoteRole() == ROLE_None )
 		{
-			World->NetworkActors.RemoveAtSwap( i );
+			ActorsToRemove.Add(Actor);
 			continue;
 		}
 
@@ -956,6 +957,12 @@ void UDemoNetDriver::TickDemoRecord( float DeltaSeconds )
 		DemoReplicateActor( Actor, ClientConnections[0], IsNetClient );
 	}
 
+	// Remove actors marked for removal.
+	for ( AActor* Actor : ActorsToRemove )
+	{
+		World->NetworkActors.Remove(Actor);
+	}
+
 	// Make sure nothing is left over
 	ClientConnections[0]->FlushNet();
 
