Fixing crashes when launching standalone game

diff --git a/Engine/Source/Runtime/Engine/Private/CurveBase.cpp b/Engine/Source/Runtime/Engine/Private/CurveBase.cpp
--- a/Engine/Source/Runtime/Engine/Private/CurveBase.cpp
+++ b/Engine/Source/Runtime/Engine/Private/CurveBase.cpp
@@ -837,7 +837,10 @@ UCurveBase::UCurveBase(const FObjectInitializer& ObjectInitializer)
 #if WITH_EDITORONLY_DATA
 void UCurveBase::GetAssetRegistryTags(TArray<FAssetRegistryTag>& OutTags) const
 {
-	OutTags.Add( FAssetRegistryTag(SourceFileTagName(), AssetImportData->ToJson(), FAssetRegistryTag::TT_Hidden) );
+	if (AssetImportData)
+	{
+		OutTags.Add(FAssetRegistryTag(SourceFileTagName(), AssetImportData->ToJson(), FAssetRegistryTag::TT_Hidden));
+	}
 
 	Super::GetAssetRegistryTags(OutTags);
 }
@@ -845,7 +848,7 @@ void UCurveBase::GetAssetRegistryTags(TArray<FAssetRegistryTag>& OutTags) const
 void UCurveBase::PostLoad()
 {
 	Super::PostLoad();
-	if (!ImportPath_DEPRECATED.IsEmpty())
+	if (!ImportPath_DEPRECATED.IsEmpty() && AssetImportData)
 	{
 		FAssetImportInfo Info;
 		Info.Insert(FAssetImportInfo::FSourceFile(ImportPath_DEPRECATED));
diff --git a/Engine/Source/Runtime/Engine/Private/CurveTable.cpp b/Engine/Source/Runtime/Engine/Private/CurveTable.cpp
--- a/Engine/Source/Runtime/Engine/Private/CurveTable.cpp
+++ b/Engine/Source/Runtime/Engine/Private/CurveTable.cpp
@@ -107,7 +107,7 @@ void UCurveTable::GetAssetRegistryTags(TArray<FAssetRegistryTag>& OutTags) const
 void UCurveTable::PostLoad()
 {
 	Super::PostLoad();
-	if (!ImportPath_DEPRECATED.IsEmpty())
+	if (!ImportPath_DEPRECATED.IsEmpty() && AssetImportData)
 	{
 		FAssetImportInfo Info;
 		Info.Insert(FAssetImportInfo::FSourceFile(ImportPath_DEPRECATED));
diff --git a/Engine/Source/Runtime/Engine/Private/DataTable.cpp b/Engine/Source/Runtime/Engine/Private/DataTable.cpp
--- a/Engine/Source/Runtime/Engine/Private/DataTable.cpp
+++ b/Engine/Source/Runtime/Engine/Private/DataTable.cpp
@@ -171,7 +171,7 @@ void UDataTable::GetAssetRegistryTags(TArray<FAssetRegistryTag>& OutTags) const
 void UDataTable::PostLoad()
 {
 	Super::PostLoad();
-	if (!ImportPath_DEPRECATED.IsEmpty())
+	if (!ImportPath_DEPRECATED.IsEmpty() && AssetImportData)
 	{
 		FAssetImportInfo Info;
 		Info.Insert(FAssetImportInfo::FSourceFile(ImportPath_DEPRECATED));
diff --git a/Engine/Source/Runtime/Engine/Private/SoundWave.cpp b/Engine/Source/Runtime/Engine/Private/SoundWave.cpp
--- a/Engine/Source/Runtime/Engine/Private/SoundWave.cpp
+++ b/Engine/Source/Runtime/Engine/Private/SoundWave.cpp
@@ -356,7 +356,7 @@ void USoundWave::PostLoad()
 	}
 
 #if WITH_EDITORONLY_DATA
-	if (!SourceFilePath_DEPRECATED.IsEmpty())
+	if (!SourceFilePath_DEPRECATED.IsEmpty() && AssetImportData)
 	{
 		FAssetImportInfo Info;
 		Info.Insert(FAssetImportInfo::FSourceFile(SourceFilePath_DEPRECATED));
diff --git a/Engine/Source/Runtime/Engine/Private/Texture.cpp b/Engine/Source/Runtime/Engine/Private/Texture.cpp
--- a/Engine/Source/Runtime/Engine/Private/Texture.cpp
+++ b/Engine/Source/Runtime/Engine/Private/Texture.cpp
@@ -270,7 +270,7 @@ void UTexture::PostLoad()
 	Super::PostLoad();
 
 #if WITH_EDITORONLY_DATA
-	if (!SourceFilePath_DEPRECATED.IsEmpty())
+	if (!SourceFilePath_DEPRECATED.IsEmpty() && AssetImportData)
 	{
 		FAssetImportInfo Info;
 		Info.Insert(FAssetImportInfo::FSourceFile(SourceFilePath_DEPRECATED));
diff --git a/Engine/Source/Runtime/Engine/Private/VectorField.cpp b/Engine/Source/Runtime/Engine/Private/VectorField.cpp
--- a/Engine/Source/Runtime/Engine/Private/VectorField.cpp
+++ b/Engine/Source/Runtime/Engine/Private/VectorField.cpp
@@ -326,7 +326,7 @@ void UVectorFieldStatic::PostLoad()
 	}
 
 #if WITH_EDITORONLY_DATA
-	if (!SourceFilePath_DEPRECATED.IsEmpty())
+	if (!SourceFilePath_DEPRECATED.IsEmpty() && AssetImportData)
 	{
 		FAssetImportInfo Info;
 		Info.Insert(FAssetImportInfo::FSourceFile(SourceFilePath_DEPRECATED));
