Simplified Raw Texture Update function by using FSlateTextureData

diff --git a/Engine/Source/Runtime/Engine/Private/Slate/SlateTextures.cpp b/Engine/Source/Runtime/Engine/Private/Slate/SlateTextures.cpp
--- a/Engine/Source/Runtime/Engine/Private/Slate/SlateTextures.cpp
+++ b/Engine/Source/Runtime/Engine/Private/Slate/SlateTextures.cpp
@@ -186,25 +186,20 @@ void FSlateTexture2DRHIRef::UpdateTextureThreadSafeRaw(uint32 InWidth, uint32 In
 	if (IsInGameThread())
 	{
 		// No cheap way to avoid having to copy the Buffer, as we cannot guarantee it will not be touched before the rendering thread is done with it.
-		const uint32 BufferSize = InWidth * InHeight * 4;
-		TArray<uint8>* BufferCopy = new TArray<uint8>();
-		BufferCopy->AddUninitialized(BufferSize);
-		FMemory::Memcpy(BufferCopy->GetData(), Buffer, BufferSize);
-		FIntPoint Dimensions(InWidth, InHeight);
+		FSlateTextureData* BulkData = new FSlateTextureData( (uint8*)Buffer, InWidth, InHeight, 4 );
 
 		// Update the texture RHI
-		ENQUEUE_UNIQUE_RENDER_COMMAND_THREEPARAMETER(
+		ENQUEUE_UNIQUE_RENDER_COMMAND_TWOPARAMETER(
 			FSlateTexture2DRHIRef_UpdateTextureRaw,
 			FSlateTexture2DRHIRef*, ThisTexture, this,
-			FIntPoint, InDimensions, Dimensions,
-			TArray<uint8>*, BufferCopy, BufferCopy,
+			FSlateTextureData*, BulkData, BulkData,
 			{
-				if (ThisTexture->GetWidth() != InDimensions.X || ThisTexture->GetHeight() != InDimensions.Y)
+				if (ThisTexture->GetWidth() != BulkData->GetWidth() || ThisTexture->GetHeight() != BulkData->GetHeight())
 				{
-					ThisTexture->Resize(InDimensions.X, InDimensions.Y);
+					ThisTexture->Resize(BulkData->GetWidth(), BulkData->GetHeight());
 				}
-				ThisTexture->UpdateTexture(*BufferCopy);
-				delete BufferCopy;
+				ThisTexture->UpdateTexture(BulkData->GetRawBytes());
+				delete BulkData;
 			});
 	}
 }
