PR #1191: DestructibleComponent: null pointer checks (Contributed by [REDACTED_PERSON])

#UE-16252

diff --git a/Engine/Source/Runtime/Engine/Private/Components/DestructibleComponent.cpp b/Engine/Source/Runtime/Engine/Private/Components/DestructibleComponent.cpp
--- a/Engine/Source/Runtime/Engine/Private/Components/DestructibleComponent.cpp
+++ b/Engine/Source/Runtime/Engine/Private/Components/DestructibleComponent.cpp
@@ -453,21 +453,24 @@ void UDestructibleComponent::AddImpulse( FVector Impulse, FName BoneName /*= NAM
 	ExecuteOnPhysicsReadWrite([&]
 	{
 		const int32 ChunkIdx = BoneIdxToChunkIdx(GetBoneIndex(BoneName));
-		PxRigidDynamic* PActor = ApexDestructibleActor->getChunkPhysXActor(ChunkIdx);
-		PActor->addForce(U2PVector(Impulse), bVelChange ? PxForceMode::eVELOCITY_CHANGE : PxForceMode::eIMPULSE);
+		if(PxRigidDynamic* Actor = ApexDestructibleActor->getChunkPhysXActor(ChunkIdx))
+		{
+			Actor->addForce(U2PVector(Impulse), bVelChange ? PxForceMode::eVELOCITY_CHANGE : PxForceMode::eIMPULSE);
+		}
 	});
 #endif
 }
 
-
 void UDestructibleComponent::AddImpulseAtLocation( FVector Impulse, FVector Position, FName BoneName /*= NAME_None*/ )
 {
 #if WITH_APEX
 	ExecuteOnPhysicsReadWrite([&]
 	{
 		const int32 ChunkIdx = BoneIdxToChunkIdx(GetBoneIndex(BoneName));
-		PxRigidDynamic* Actor = ApexDestructibleActor->getChunkPhysXActor(ChunkIdx);
-		PxRigidBodyExt::addForceAtPos(*Actor, U2PVector(Impulse), U2PVector(Position), PxForceMode::eIMPULSE);
+		if (PxRigidDynamic* Actor = ApexDestructibleActor->getChunkPhysXActor(ChunkIdx))
+		{
+			PxRigidBodyExt::addForceAtPos(*Actor, U2PVector(Impulse), U2PVector(Position), PxForceMode::eIMPULSE);
+		}
 	});
 #endif
 }
@@ -478,8 +481,10 @@ void UDestructibleComponent::AddForce( FVector Force, FName BoneName /*= NAME_No
 	ExecuteOnPhysicsReadWrite([&]
 	{
 		const int32 ChunkIdx = BoneIdxToChunkIdx(GetBoneIndex(BoneName));
-		PxRigidDynamic* Actor = ApexDestructibleActor->getChunkPhysXActor(ChunkIdx);
-		Actor->addForce(U2PVector(Force), bAccelChange ? PxForceMode::eACCELERATION : PxForceMode::eFORCE);
+		if (PxRigidDynamic* Actor = ApexDestructibleActor->getChunkPhysXActor(ChunkIdx))
+		{
+			Actor->addForce(U2PVector(Force), bAccelChange ? PxForceMode::eACCELERATION : PxForceMode::eFORCE);
+		}
 	});
 #endif
 }
@@ -490,8 +495,10 @@ void UDestructibleComponent::AddForceAtLocation( FVector Force, FVector Location
 	ExecuteOnPhysicsReadWrite([&]
 	{
 		int32 ChunkIdx = BoneIdxToChunkIdx(GetBoneIndex(BoneName));
-		PxRigidDynamic* Actor = ApexDestructibleActor->getChunkPhysXActor(ChunkIdx);
-		PxRigidBodyExt::addForceAtPos(*Actor, U2PVector(Force), U2PVector(Location), PxForceMode::eFORCE);
+		if (PxRigidDynamic* Actor = ApexDestructibleActor->getChunkPhysXActor(ChunkIdx))
+		{
+			PxRigidBodyExt::addForceAtPos(*Actor, U2PVector(Force), U2PVector(Location), PxForceMode::eFORCE);
+		}
 	});
 #endif
 }
