#UE4 Re-fixing demonetdriver errors

diff --git a/Engine/Source/Runtime/Engine/Private/DemoNetDriver.cpp b/Engine/Source/Runtime/Engine/Private/DemoNetDriver.cpp
--- a/Engine/Source/Runtime/Engine/Private/DemoNetDriver.cpp
+++ b/Engine/Source/Runtime/Engine/Private/DemoNetDriver.cpp
@@ -938,18 +938,19 @@ void UDemoNetDriver::TickDemoRecord( float DeltaSeconds )
 
 	DemoReplicateActor( World->GetWorldSettings(), ClientConnections[0], IsNetClient );
 
-	TArray<AActor*> ActorsToRemove;
-	for ( AActor* Actor : World->NetworkActors )
+	for ( TSet<AActor*>::TIterator ActorIt = World->NetworkActors.CreateIterator(); ActorIt; ++ActorIt)
 	{
+		AActor* Actor = *ActorIt;
+		
 		if ( Actor->IsPendingKill() )
 		{
-			ActorsToRemove.Add(Actor);
+			ActorIt.RemoveCurrent();
 			continue;
 		}
 
 		if ( Actor->GetRemoteRole() == ROLE_None )
 		{
-			ActorsToRemove.Add(Actor);
+			ActorIt.RemoveCurrent();
 			continue;
 		}
 
@@ -957,12 +958,6 @@ void UDemoNetDriver::TickDemoRecord( float DeltaSeconds )
 		DemoReplicateActor( Actor, ClientConnections[0], IsNetClient );
 	}
 
-	// Remove actors marked for removal.
-	for ( AActor* Actor : ActorsToRemove )
-	{
-		World->NetworkActors.Remove(Actor);
-	}
-
 	// Make sure nothing is left over
 	ClientConnections[0]->FlushNet();
 
