Changing the class name for the automation log tests.

diff --git a/Engine/Source/Runtime/Engine/Private/Tests/EngineAutomationTests.cpp b/Engine/Source/Runtime/Engine/Private/Tests/EngineAutomationTests.cpp
--- a/Engine/Source/Runtime/Engine/Private/Tests/EngineAutomationTests.cpp
+++ b/Engine/Source/Runtime/Engine/Private/Tests/EngineAutomationTests.cpp
@@ -428,9 +428,9 @@ bool FCinematicFPSPerfTest::RunTest(const FString& Parameters)
 }
 
 
-IMPLEMENT_SIMPLE_AUTOMATION_TEST(FLogAddMessage, "System.Automation.Log.Add Log Message", EAutomationTestFlags::ATF_None)
+IMPLEMENT_SIMPLE_AUTOMATION_TEST(FAutomationLogAddMessage, "System.Automation.Log.Add Log Message", EAutomationTestFlags::ATF_None)
 
-bool FLogAddMessage::RunTest(const FString& Parameters)
+bool FAutomationLogAddMessage::RunTest(const FString& Parameters)
 {
 	//** TEST **//
 	AddLogItem(TEXT("Test log message."));
@@ -445,9 +445,9 @@ bool FLogAddMessage::RunTest(const FString& Parameters)
 	return true;
 }
 
-IMPLEMENT_SIMPLE_AUTOMATION_TEST(FLogAddWarning, "System.Automation.Log.Add Warning Message", EAutomationTestFlags::ATF_None)
+IMPLEMENT_SIMPLE_AUTOMATION_TEST(FAutomationLogAddWarning, "System.Automation.Log.Add Warning Message", EAutomationTestFlags::ATF_None)
 
-bool FLogAddWarning::RunTest(const FString& Parameters)
+bool FAutomationLogAddWarning::RunTest(const FString& Parameters)
 {
 	//** TEST **//
 	AddWarning(TEXT("Test warning message."));
@@ -462,9 +462,9 @@ bool FLogAddWarning::RunTest(const FString& Parameters)
 	return true;
 }
 
-IMPLEMENT_SIMPLE_AUTOMATION_TEST(FLogAddError, "System.Automation.Log.Add Error Message", EAutomationTestFlags::ATF_None)
+IMPLEMENT_SIMPLE_AUTOMATION_TEST(FAutomationLogAddError, "System.Automation.Log.Add Error Message", EAutomationTestFlags::ATF_None)
 
-bool FLogAddError::RunTest(const FString& Parameters)
+bool FAutomationLogAddError::RunTest(const FString& Parameters)
 {
 	//** TEST **//
 	AddError(TEXT("Test error message"));
