Making GTestRegisterTickFunctions thread safe

diff --git a/Engine/Source/Runtime/Engine/Private/Actor.cpp b/Engine/Source/Runtime/Engine/Private/Actor.cpp
--- a/Engine/Source/Runtime/Engine/Private/Actor.cpp
+++ b/Engine/Source/Runtime/Engine/Private/Actor.cpp
@@ -539,7 +539,19 @@ void AActor::ApplyWorldOffset(const FVector& InOffset, bool bWorldShift)
 	}
 }
 
-static AActor* GTestRegisterTickFunctions = NULL;
+/** Thread safe container for actor related global variables */
+class FActorThreadContext : public TThreadSingleton<FActorThreadContext>
+{
+	friend TThreadSingleton<FActorThreadContext>;
+
+	FActorThreadContext()
+		: TestRegisterTickFunctions(nullptr)
+	{
+	}
+public:
+	/** Tests tick function registration */
+	AActor* TestRegisterTickFunctions;
+};
 
 void AActor::RegisterActorTickFunctions(bool bRegister)
 {
@@ -562,17 +574,18 @@ void AActor::RegisterActorTickFunctions(bool bRegister)
 		}
 	}
 
-	GTestRegisterTickFunctions = this; // we will verify the super call chain is intact. Don't copy and paste this to another actor class!
+	FActorThreadContext::Get().TestRegisterTickFunctions = this; // we will verify the super call chain is intact. Don't copy and paste this to another actor class!
 }
 
 void AActor::RegisterAllActorTickFunctions(bool bRegister, bool bDoComponents)
 {
 	if(!IsTemplate())
 	{
-		check(GTestRegisterTickFunctions == NULL);
+		FActorThreadContext& ThreadContext = FActorThreadContext::Get();
+		check(ThreadContext.TestRegisterTickFunctions == nullptr);
 		RegisterActorTickFunctions(bRegister);
-		checkf(GTestRegisterTickFunctions == this, TEXT("Failed to route Actor RegisterTickFunctions (%s)"), *GetFullName());
-		GTestRegisterTickFunctions = NULL;
+		checkf(ThreadContext.TestRegisterTickFunctions == this, TEXT("Failed to route Actor RegisterTickFunctions (%s)"), *GetFullName());
+		ThreadContext.TestRegisterTickFunctions = nullptr;
 
 		if (bDoComponents)
 		{
