Can now Make and Break the DateTime struct in Blueprints.

diff --git a/Engine/Source/Runtime/CoreUObject/Classes/Object.h b/Engine/Source/Runtime/CoreUObject/Classes/Object.h
--- a/Engine/Source/Runtime/CoreUObject/Classes/Object.h
+++ b/Engine/Source/Runtime/CoreUObject/Classes/Object.h
@@ -831,7 +831,7 @@ public:
 
 // A date/time value.
 
-USTRUCT(immutable, noexport, BlueprintType)
+USTRUCT(immutable, noexport, BlueprintType, meta=(HasNativeMake="Engine.KismetMathLibrary.MakeDateTime", HasNativeBreak="Engine.KismetMathLibrary.BreakDateTime"))
 struct FDateTime
 {
 	int64 Ticks;
diff --git a/Engine/Source/Runtime/Engine/Classes/Kismet/KismetMathLibrary.h b/Engine/Source/Runtime/Engine/Classes/Kismet/KismetMathLibrary.h
--- a/Engine/Source/Runtime/Engine/Classes/Kismet/KismetMathLibrary.h
+++ b/Engine/Source/Runtime/Engine/Classes/Kismet/KismetMathLibrary.h
@@ -790,6 +790,14 @@ class ENGINE_API UKismetMathLibrary : public UBlueprintFunctionLibrary
 	// DateTime functions.
 	//
 
+	/* Makes a DateTime struct */
+	UFUNCTION(BlueprintPure, Category="Math|DateTime", meta=(NativeMakeFunc, AdvancedDisplay = "3"))
+	static FDateTime MakeDateTime(int32 Year, int32 Month, int32 Day, int32 Hour = 0, int32 Minute = 0, int32 Second = 0, int32 Millisecond = 0);
+
+	/* Breaks a DateTime into its components */
+	UFUNCTION(BlueprintPure, Category="Math|DateTime", meta=(NativeBreakFunc))
+	static void BreakDateTime(FDateTime InDateTime, int32& Year, int32& Month, int32& Day, int32& Hour, int32& Minute, int32& Second, int32& Millisecond);
+
 	/* Addition (A + B) */
 	UFUNCTION(BlueprintPure, meta=(DisplayName="DateTime + Timespan", CompactNodeTitle="+", Keywords="+ add plus"), Category="Math|DateTime")
 	static FDateTime Add_DateTimeTimespan( FDateTime A, FTimespan B );
diff --git a/Engine/Source/Runtime/Engine/Private/KismetMathLibrary.cpp b/Engine/Source/Runtime/Engine/Private/KismetMathLibrary.cpp
--- a/Engine/Source/Runtime/Engine/Private/KismetMathLibrary.cpp
+++ b/Engine/Source/Runtime/Engine/Private/KismetMathLibrary.cpp
@@ -1175,6 +1175,21 @@ bool UKismetMathLibrary::ClassIsChildOf(TSubclassOf<class UObject> TestClass, TS
 
 /* DateTime functions
  *****************************************************************************/
+FDateTime UKismetMathLibrary::MakeDateTime(int32 Year, int32 Month, int32 Day, int32 Hour, int32 Minute, int32 Second, int32 Millisecond)
+{
+	return FDateTime(Year, Month, Day, Hour, Minute, Second, Millisecond);
+}
+
+void UKismetMathLibrary::BreakDateTime(FDateTime InDateTime, int32& Year, int32& Month, int32& Day, int32& Hour, int32& Minute, int32& Second, int32& Millisecond)
+{
+	Year = GetYear(InDateTime);
+	Month = GetMonth(InDateTime);
+	Day = GetDay(InDateTime);
+	Hour = GetHour(InDateTime);
+	Minute = GetMinute(InDateTime);
+	Second = GetSecond(InDateTime);
+	Millisecond = GetMillisecond(InDateTime);
+}
 
 FDateTime UKismetMathLibrary::Add_DateTimeTimespan( FDateTime A, FTimespan B )
 {
