Can now Make and Break the Timespan struct in Blueprints.

Extension of CL# 2569222

diff --git a/Engine/Source/Runtime/CoreUObject/Classes/Object.h b/Engine/Source/Runtime/CoreUObject/Classes/Object.h
--- a/Engine/Source/Runtime/CoreUObject/Classes/Object.h
+++ b/Engine/Source/Runtime/CoreUObject/Classes/Object.h
@@ -840,7 +840,7 @@ struct FDateTime
 
 // A time span value.
 
-USTRUCT(immutable, noexport, BlueprintType)
+USTRUCT(immutable, noexport, BlueprintType, meta=(HasNativeMake="Engine.KismetMathLibrary.MakeTimespan", HasNativeBreak="Engine.KismetMathLibrary.BreakTimespan"))
 struct FTimespan
 {
 	int64 Ticks;
diff --git a/Engine/Source/Runtime/Engine/Classes/Kismet/KismetMathLibrary.h b/Engine/Source/Runtime/Engine/Classes/Kismet/KismetMathLibrary.h
--- a/Engine/Source/Runtime/Engine/Classes/Kismet/KismetMathLibrary.h
+++ b/Engine/Source/Runtime/Engine/Classes/Kismet/KismetMathLibrary.h
@@ -922,6 +922,14 @@ class ENGINE_API UKismetMathLibrary : public UBlueprintFunctionLibrary
 	// Timespan functions.
 	//
 
+	/* Makes a Timespan struct */
+	UFUNCTION(BlueprintPure, Category="Math|Timespan", meta=(NativeMakeFunc))
+	static FTimespan MakeTimespan(int32 Days, int32 Hours, int32 Minutes, int32 Seconds, int32 Milliseconds);
+
+	/* Breaks a Timespan into its components */
+	UFUNCTION(BlueprintPure, Category="Math|Timespan", meta=(NativeBreakFunc))
+	static void BreakTimespan(FTimespan InTimespan, int32& Days, int32& Hours, int32& Minutes, int32& Seconds, int32& Milliseconds);
+
 	/* Addition (A + B) */
 	UFUNCTION(BlueprintPure, meta=(DisplayName="Timespan + Timespan", CompactNodeTitle="+", Keywords="+ add plus"), Category="Math|Timespan")
 	static FTimespan Add_TimespanTimespan( FTimespan A, FTimespan B );
diff --git a/Engine/Source/Runtime/Engine/Private/KismetMathLibrary.cpp b/Engine/Source/Runtime/Engine/Private/KismetMathLibrary.cpp
--- a/Engine/Source/Runtime/Engine/Private/KismetMathLibrary.cpp
+++ b/Engine/Source/Runtime/Engine/Private/KismetMathLibrary.cpp
@@ -1379,6 +1379,20 @@ FDateTime UKismetMathLibrary::UtcNow( )
 /* Timespan functions
  *****************************************************************************/
 
+FTimespan UKismetMathLibrary::MakeTimespan(int32 Days, int32 Hours, int32 Minutes, int32 Seconds, int32 Milliseconds)
+{
+	return FTimespan(Days, Hours, Minutes, Seconds, Milliseconds);
+}
+
+void UKismetMathLibrary::BreakTimespan(FTimespan InTimespan, int32& Days, int32& Hours, int32& Minutes, int32& Seconds, int32& Milliseconds)
+{
+	Days = InTimespan.GetDays();
+	Hours = InTimespan.GetHours();
+	Minutes = InTimespan.GetMinutes();
+	Seconds = InTimespan.GetSeconds();
+	Milliseconds = InTimespan.GetMilliseconds();
+}
+
 FTimespan UKismetMathLibrary::Add_TimespanTimespan( FTimespan A, FTimespan B )
 {
 	return A + B;
