UE4 - fix editing blueprint CDO of a hierarchical static mesh component never building the tree.

diff --git a/Engine/Source/Runtime/Engine/Private/HierarchicalInstancedStaticMesh.cpp b/Engine/Source/Runtime/Engine/Private/HierarchicalInstancedStaticMesh.cpp
--- a/Engine/Source/Runtime/Engine/Private/HierarchicalInstancedStaticMesh.cpp
+++ b/Engine/Source/Runtime/Engine/Private/HierarchicalInstancedStaticMesh.cpp
@@ -2117,6 +2117,13 @@ FPrimitiveSceneProxy* UHierarchicalInstancedStaticMeshComponent::CreateSceneProx
 	}
 	ProxySize = 0;
 
+	if (PerInstanceSMData.Num() != InstanceReorderTable.Num())
+	{
+		// this is a bizzaro case where folks are editing the PerInstanceSMData in a template of a blueprint component
+		// not a great solution, but we can fix it now.
+		BuildTree();
+	}
+
 	FlushAsyncBuildInstanceBufferTask();
 
 	// Verify that the mesh is valid before using it.
