Fix crash in console after removing all players from the game.

diff --git a/Engine/Source/Runtime/Engine/Private/UserInterface/Console.cpp b/Engine/Source/Runtime/Engine/Private/UserInterface/Console.cpp
--- a/Engine/Source/Runtime/Engine/Private/UserInterface/Console.cpp
+++ b/Engine/Source/Runtime/Engine/Private/UserInterface/Console.cpp
@@ -491,18 +491,21 @@ void UConsole::StartTyping(const FString& Text)
 
 void UConsole::FlushPlayerInput()
 {
-	APlayerController* PC = NULL;
-	if(ConsoleTargetPlayer != NULL)
+	APlayerController* PC = nullptr;
+	if (ConsoleTargetPlayer)
 	{
 		PC = ConsoleTargetPlayer->PlayerController;
 	}
 	else
 	{
-		UWorld *World = GetOuterUGameViewportClient()->GetWorld();
-		PC = GEngine->GetFirstGamePlayer(World)->PlayerController;
+		UWorld* World = GetOuterUGameViewportClient()->GetWorld();
+		if (ULocalPlayer* LocalPlayer = GEngine->GetFirstGamePlayer(World))
+		{
+			PC = LocalPlayer->PlayerController;
+		}
 	}
 
-	if ( PC != NULL && PC->PlayerInput != NULL )
+	if (PC && PC->PlayerInput)
 	{
 		PC->PlayerInput->FlushPressedKeys();
 	}
@@ -510,8 +513,6 @@ void UConsole::FlushPlayerInput()
 
 bool UConsole::ProcessControlKey(FKey Key, EInputEvent Event)
 {	
-	UWorld *World = GetOuterUGameViewportClient()->GetWorld();
-	APlayerController* PlayerController = GEngine->GetFirstGamePlayer(World)->PlayerController;
 	if (Key == EKeys::LeftControl || Key == EKeys::RightControl)
 	{
 		if (Event == IE_Released)
@@ -525,7 +526,7 @@ bool UConsole::ProcessControlKey(FKey Key, EInputEvent Event)
 
 		return true;
 	}
-	else if (bCtrl && Event == IE_Pressed && PlayerController)
+	else if (bCtrl && Event == IE_Pressed)
 	{
 		if (Key == EKeys::V)
 		{
