[UE-15826] Guarding against the resetting of package loaders (when renaming some Blueprint element) while their linker is still actively loading (modified fix from RobM).

diff --git a/Engine/Source/Runtime/Engine/Private/Blueprint.cpp b/Engine/Source/Runtime/Engine/Private/Blueprint.cpp
--- a/Engine/Source/Runtime/Engine/Private/Blueprint.cpp
+++ b/Engine/Source/Runtime/Engine/Private/Blueprint.cpp
@@ -436,6 +436,35 @@ void UBlueprint::PostLoad()
 {
 	Super::PostLoad();
 
+	// Can't use TGuardValue here as bIsRegeneratingOnLoad is a bitfield
+	struct FScopedRegeneratingOnLoad
+	{
+		UBlueprint& Blueprint;
+		bool bPreviousValue;
+		FScopedRegeneratingOnLoad(UBlueprint& InBlueprint)
+			: Blueprint(InBlueprint)
+			, bPreviousValue(InBlueprint.bIsRegeneratingOnLoad)
+		{
+			// if the blueprint's package is still in the midst of loading, then
+			// bIsRegeneratingOnLoad needs to be set to prevent UObject renames
+			// from resetting loaders
+			Blueprint.bIsRegeneratingOnLoad = true;
+			if (UPackage* Package = Blueprint.GetOutermost())
+			{
+				// checking (Package->LinkerLoad != nullptr) ensures this 
+				// doesn't get set when duplicating blueprints (which also calls 
+				// PostLoad), and checking RF_WasLoaded makes sure we only 
+				// forcefully set bIsRegeneratingOnLoad for blueprints that need 
+				// it (ones still actively loading)
+				Blueprint.bIsRegeneratingOnLoad = bPreviousValue || ((Package->LinkerLoad != nullptr) && !Package->HasAnyFlags(RF_WasLoaded));
+			}
+		}
+		~FScopedRegeneratingOnLoad()
+		{
+			Blueprint.bIsRegeneratingOnLoad = bPreviousValue;
+		}
+	} GuardIsRegeneratingOnLoad(*this);
+
 	// Mark the blueprint as in error if there has been a major version bump
 	if (BlueprintSystemVersion < UBlueprint::GetCurrentBlueprintSystemVersion())
 	{
