bReplicateMovement actor flag will no longer result in client replicating movement on server

UE-7453

Server and Clients characters no longer replicate movement if bReplicates is false

Client no longer extrapolates Server's position if the server is moving and bReplicate/bReplicateMovement is toggled off

Server no longer extrapolates Client's position for the period when bReplicateMovement is turned off and on again

bReplicateMovement was marked for Replication in UPROPERTY but was not replicating, it is now replicating and uses a repnotify so developers can detect if the server has stopped replicated movement of an actor

diff --git a/Engine/Source/Runtime/Engine/Classes/GameFramework/Actor.h b/Engine/Source/Runtime/Engine/Classes/GameFramework/Actor.h
--- a/Engine/Source/Runtime/Engine/Classes/GameFramework/Actor.h
+++ b/Engine/Source/Runtime/Engine/Classes/GameFramework/Actor.h
@@ -130,9 +130,13 @@ public:
 	 * @see SetReplicates()
 	 * @see https://docs.unrealengine.com/latest/INT/Gameplay/Networking/Replication/
 	 */
-	UPROPERTY(Replicated, Category=Replication, EditDefaultsOnly)
+	UPROPERTY(ReplicatedUsing=OnRep_ReplicateMovement, Category=Replication, EditDefaultsOnly)
 	uint32 bReplicateMovement:1;    
 
+	/** Called on client when updated bReplicateMovement value is received for this actor. */
+	UFUNCTION()
+	virtual void OnRep_ReplicateMovement();
+
 	/**
 	 * If true, this actor is no longer replicated to new clients, and is "torn off" (becomes a ROLE_Authority) on clients to which it was being replicated.
 	 * @see TornOff()
@@ -182,7 +186,7 @@ protected:
 	 * If true, this actor will replicate to remote machines
 	 * @see SetReplicates()
 	 */
-	UPROPERTY(Category = Replication, EditDefaultsOnly, BlueprintReadOnly)
+	UPROPERTY(EditDefaultsOnly, BlueprintReadOnly, Category = "Replication")
 	uint32 bReplicates:1;
 
 	/** This function should only be used in the constructor of classes that need to set the RemoteRole for backwards compatibility purposes */
@@ -221,6 +225,13 @@ public:
 	UFUNCTION(BlueprintCallable, Category = "Replication")
 	void SetReplicates(bool bInReplicates);
 
+	/**
+	* Set whether this actor's movement replicates to network clients.
+	* @param bInReplicateMovement Whether this Actor's movement replicates to clients.
+	*/
+	UFUNCTION(BlueprintCallable, Category = "Replication")
+	virtual void SetReplicateMovement(bool bInReplicateMovement);
+
 	/** Sets whether or not this Actor is an autonomous proxy, which is an actor on a network client that is controlled by a user on that client. */
 	void SetAutonomousProxy(bool bInAutonomousProxy);
 	
diff --git a/Engine/Source/Runtime/Engine/Classes/GameFramework/Character.h b/Engine/Source/Runtime/Engine/Classes/GameFramework/Character.h
--- a/Engine/Source/Runtime/Engine/Classes/GameFramework/Character.h
+++ b/Engine/Source/Runtime/Engine/Classes/GameFramework/Character.h
@@ -259,6 +259,10 @@ public:
 	UFUNCTION()
 	virtual void OnRep_ReplicatedBasedMovement();
 
+	/** Set whether this actor's movement replicates to network clients. */
+	UFUNCTION(BlueprintCallable, Category = "Replication")
+	virtual void SetReplicateMovement(bool bInReplicateMovement) override;
+
 protected:
 	/** Saved translation offset of mesh. */
 	UPROPERTY()
diff --git a/Engine/Source/Runtime/Engine/Private/Actor.cpp b/Engine/Source/Runtime/Engine/Private/Actor.cpp
--- a/Engine/Source/Runtime/Engine/Private/Actor.cpp
+++ b/Engine/Source/Runtime/Engine/Private/Actor.cpp
@@ -678,6 +678,7 @@ void AActor::Tick( float DeltaSeconds )
 		ReceiveTick(DeltaSeconds);
 	}
 
+
 	// Update any latent actions we have for this actor
 	GetWorld()->GetLatentActionManager().ProcessLatentActions(this, DeltaSeconds);
 
@@ -1600,6 +1601,8 @@ void AActor::PrestreamTextures( float Seconds, bool bEnableStreaming, int32 Cine
 
 void AActor::OnRep_Instigator() {}
 
+void AActor::OnRep_ReplicateMovement() {}
+
 void AActor::RouteEndPlay(const EEndPlayReason::Type EndPlayReason)
 {
 	if (bActorInitialized)
@@ -2533,6 +2536,11 @@ void AActor::SetReplicates(bool bInReplicates)
 	}
 }
 
+void AActor::SetReplicateMovement(bool bInReplicateMovement)
+{
+	bReplicateMovement = bInReplicateMovement;
+}
+
 void AActor::SetAutonomousProxy(bool bInAutonomousProxy)
 {
 	if (bReplicates)
diff --git a/Engine/Source/Runtime/Engine/Private/ActorReplication.cpp b/Engine/Source/Runtime/Engine/Private/ActorReplication.cpp
--- a/Engine/Source/Runtime/Engine/Private/ActorReplication.cpp
+++ b/Engine/Source/Runtime/Engine/Private/ActorReplication.cpp
@@ -234,6 +234,7 @@ void AActor::GetLifetimeReplicatedProps( TArray< FLifetimeProperty > & OutLifeti
 		BPClass->GetLifetimeBlueprintReplicationList(OutLifetimeProps);
 	}
 
+	DOREPLIFETIME( AActor, bReplicateMovement );
 	DOREPLIFETIME( AActor, Role );
 	DOREPLIFETIME( AActor, RemoteRole );
 	DOREPLIFETIME( AActor, Owner );
diff --git a/Engine/Source/Runtime/Engine/Private/Character.cpp b/Engine/Source/Runtime/Engine/Private/Character.cpp
--- a/Engine/Source/Runtime/Engine/Private/Character.cpp
+++ b/Engine/Source/Runtime/Engine/Private/Character.cpp
@@ -285,6 +285,23 @@ void ACharacter::OnRep_IsCrouched()
 	}
 }
 
+void ACharacter::SetReplicateMovement(bool bInReplicateMovement)
+{
+	Super::SetReplicateMovement(bInReplicateMovement);
+
+	if (CharacterMovement != nullptr)
+	{
+		// Set prediction data time stamp to current time to stop extrapolating
+		// from time bReplicateMovement was turned off to when it was turned on again
+		FNetworkPredictionData_Server* NetworkPrediction = CharacterMovement->GetPredictionData_Server();
+
+		if (NetworkPrediction != nullptr)
+		{
+			NetworkPrediction->ServerTimeStamp = GetWorld()->GetTimeSeconds();
+		}
+	}
+}
+
 bool ACharacter::CanCrouch()
 {
 	return !bIsCrouched && CharacterMovement && CharacterMovement->CanEverCrouch() && GetRootComponent() && !GetRootComponent()->IsSimulatingPhysics();
diff --git a/Engine/Source/Runtime/Engine/Private/Components/CharacterMovementComponent.cpp b/Engine/Source/Runtime/Engine/Private/Components/CharacterMovementComponent.cpp
--- a/Engine/Source/Runtime/Engine/Private/Components/CharacterMovementComponent.cpp
+++ b/Engine/Source/Runtime/Engine/Private/Components/CharacterMovementComponent.cpp
@@ -1060,6 +1060,7 @@ void UCharacterMovementComponent::SimulatedTick(float DeltaSeconds)
 			CharacterOwner->SimulatedRootMotionPositionFixup(DeltaSeconds);
 		}
 	}
+	// Not playing RootMotion AnimMontage
 	else
 	{
 		// if we were simulating root motion, we've been ignoring regular ReplicatedMovement updates.
@@ -1076,15 +1077,18 @@ void UCharacterMovementComponent::SimulatedTick(float DeltaSeconds)
 			}
 		}
 
-		if( UpdatedComponent->IsSimulatingPhysics() 
-			|| (CharacterOwner && CharacterOwner->IsMatineeControlled()) 
-			|| (CharacterOwner && CharacterOwner->IsPlayingRootMotion()))
+		if (CharacterOwner->bReplicateMovement)
 		{
-			PerformMovement(DeltaSeconds);
-		}
-		else
-		{
-			SimulateMovement(DeltaSeconds);
+			if ((UpdatedComponent->IsSimulatingPhysics()
+				|| (CharacterOwner && CharacterOwner->IsMatineeControlled())
+				|| (CharacterOwner && CharacterOwner->IsPlayingRootMotion())))
+			{
+				PerformMovement(DeltaSeconds);
+			}
+			else
+			{
+				SimulateMovement(DeltaSeconds);
+			}
 		}
 	}
 
@@ -6011,45 +6015,49 @@ void UCharacterMovementComponent::ReplicateMoveToServer(float DeltaTime, const F
 	NewMove->PostUpdate(CharacterOwner, FSavedMove_Character::PostUpdate_Record);
 
 	// Add NewMove to the list
-	ClientData->SavedMoves.Push(NewMove);
-
-	const bool bCanDelayMove = (CVarNetEnableMoveCombining.GetValueOnGameThread() != 0) && CanDelaySendingMove(NewMove);
-
-	if (bCanDelayMove && ClientData->PendingMove.IsValid() == false)
+	if (CharacterOwner->bReplicateMovement)
 	{
-		// Decide whether to hold off on move
-		// send moves more frequently in small games where server isn't likely to be saturated
-		float NetMoveDelta;
-		UPlayer* Player = (PC ? PC->Player : NULL);
+		ClientData->SavedMoves.Push(NewMove);
 
-		if (Player && (Player->CurrentNetSpeed > 10000) && (GetWorld()->GameState != NULL) && (GetWorld()->GameState->PlayerArray.Num() <= 10))
-		{
-			NetMoveDelta = 0.011f;
-		}
-		else if (Player && CharacterOwner->GetWorldSettings()->GameNetworkManagerClass) 
-		{
-			NetMoveDelta = FMath::Max(0.0222f,2 * GetDefault<AGameNetworkManager>(CharacterOwner->GetWorldSettings()->GameNetworkManagerClass)->MoveRepSize/Player->CurrentNetSpeed);
-		}
-		else
-		{
-			NetMoveDelta = 0.011f;
-		}
+		const bool bCanDelayMove = (CVarNetEnableMoveCombining.GetValueOnGameThread() != 0) && CanDelaySendingMove(NewMove);
 
-		if ((GetWorld()->TimeSeconds - ClientData->ClientUpdateTime) * CharacterOwner->GetWorldSettings()->GetEffectiveTimeDilation() < NetMoveDelta)
+		if (bCanDelayMove && ClientData->PendingMove.IsValid() == false)
 		{
-			// Delay sending this move.
-			ClientData->PendingMove = NewMove;
-			return;
+			// Decide whether to hold off on move
+			// send moves more frequently in small games where server isn't likely to be saturated
+			float NetMoveDelta;
+			UPlayer* Player = (PC ? PC->Player : NULL);
+
+			if (Player && (Player->CurrentNetSpeed > 10000) && (GetWorld()->GameState != NULL) && (GetWorld()->GameState->PlayerArray.Num() <= 10))
+			{
+				NetMoveDelta = 0.011f;
+			}
+			else if (Player && CharacterOwner->GetWorldSettings()->GameNetworkManagerClass) 
+			{
+				NetMoveDelta = FMath::Max(0.0222f,2 * GetDefault<AGameNetworkManager>(CharacterOwner->GetWorldSettings()->GameNetworkManagerClass)->MoveRepSize/Player->CurrentNetSpeed);
+			}
+			else
+			{
+				NetMoveDelta = 0.011f;
+			}
+
+			if ((GetWorld()->TimeSeconds - ClientData->ClientUpdateTime) * CharacterOwner->GetWorldSettings()->GetEffectiveTimeDilation() < NetMoveDelta)
+			{
+				// Delay sending this move.
+				ClientData->PendingMove = NewMove;
+				return;
+			}
 		}
-	}
 
-	ClientData->ClientUpdateTime = GetWorld()->TimeSeconds;
+		ClientData->ClientUpdateTime = GetWorld()->TimeSeconds;
 
-	UE_LOG(LogNetPlayerMovement, Verbose, TEXT("Client ReplicateMove Time %f Acceleration %s Position %s DeltaTime %f"),
-		NewMove->TimeStamp, *NewMove->Acceleration.ToString(), *UpdatedComponent->GetComponentLocation().ToString(), DeltaTime);
+		UE_LOG(LogNetPlayerMovement, Verbose, TEXT("Client ReplicateMove Time %f Acceleration %s Position %s DeltaTime %f"),
+			NewMove->TimeStamp, *NewMove->Acceleration.ToString(), *UpdatedComponent->GetComponentLocation().ToString(), DeltaTime);
+
+		// Send move to server if this character is replicating movement
+		CallServerMove(NewMove.Get(), OldMove.Get());
+	}
 
-	// Send to the server
-	CallServerMove( NewMove.Get(), OldMove.Get() );
 	ClientData->PendingMove = NULL;
 }
 
diff --git a/Engine/Source/Runtime/Engine/Private/LevelActor.cpp b/Engine/Source/Runtime/Engine/Private/LevelActor.cpp
--- a/Engine/Source/Runtime/Engine/Private/LevelActor.cpp
+++ b/Engine/Source/Runtime/Engine/Private/LevelActor.cpp
@@ -462,7 +462,7 @@ bool UWorld::DestroyActor( AActor* ThisActor, bool bNetForce, bool bShouldModify
 	}
 
 	// Prevent recursion
-	FMarkActorIsBeingDestroyed MarkActorIsBeingDestroyed(ThisActor);
+	//FMarkActorIsBeingDestroyed MarkActorIsBeingDestroyed(ThisActor);
 
 	// Notify the texture streaming manager about the destruction of this actor.
 	IStreamingManager::Get().NotifyActorDestroyed( ThisActor );
diff --git a/Engine/Source/Runtime/Engine/Private/PlayerController.cpp b/Engine/Source/Runtime/Engine/Private/PlayerController.cpp
--- a/Engine/Source/Runtime/Engine/Private/PlayerController.cpp
+++ b/Engine/Source/Runtime/Engine/Private/PlayerController.cpp
@@ -3822,7 +3822,7 @@ void APlayerController::TickActor( float DeltaSeconds, ELevelTick TickType, FAct
 		// force physics update for clients that aren't sending movement updates in a timely manner 
 		// this prevents cheats associated with artificially induced ping spikes
 		// skip updates if pawn lost autonomous proxy role (e.g. TurnOff() call)
-		if (GetPawn() && !GetPawn()->IsPendingKill() && GetPawn()->GetRemoteRole() == ROLE_AutonomousProxy)
+		if (GetPawn() && !GetPawn()->IsPendingKill() && GetPawn()->GetRemoteRole() == ROLE_AutonomousProxy && GetPawn()->bReplicateMovement)
 		{
 			INetworkPredictionInterface* NetworkPredictionInterface = Cast<INetworkPredictionInterface>(GetPawn()->GetMovementComponent());
 			if (NetworkPredictionInterface)
