Added UE_SECURITY_LOG and logged cases that demonstrate the log

diff --git a/Engine/Source/Runtime/Core/Public/Misc/AssertionMacros.h b/Engine/Source/Runtime/Core/Public/Misc/AssertionMacros.h
--- a/Engine/Source/Runtime/Core/Public/Misc/AssertionMacros.h
+++ b/Engine/Source/Runtime/Core/Public/Misc/AssertionMacros.h
@@ -263,6 +263,7 @@ struct FTCharArrayTester
 	#define DECLARE_LOG_CATEGORY_CLASS(...)
 	#define DEFINE_LOG_CATEGORY_CLASS(...)
 	#define LOG_SCOPE_VERBOSITY_OVERRIDE(...)
+	#define UE_SECURITY_LOG(...)
 
 #else
 
@@ -317,6 +318,27 @@ struct FTCharArrayTester
 			} \
 		}
 
+		/**
+		* A  macro that outputs a formatted message to the log specifically used for security events
+		* @param NetConnection, a valid UNetConnection
+		* @param SecurityEventType, a security event type (ESecurityEvent::Type)
+		* @param Format, format text
+		***/
+		#define UE_SECURITY_LOG(NetConnection, SecurityEventType, Format, ...) \
+		{ \
+			static_assert(IS_TCHAR_ARRAY(Format), "Formatting string must be a TCHAR array."); \
+			check(NetConnection != nullptr); \
+			if (UE_LOG_CHECK_COMPILEDIN_VERBOSITY(LogSecurity, Warning)) \
+			{ \
+				if (!LogSecurity.IsSuppressed(ELogVerbosity::Warning)) \
+				{ \
+				FString Test = FString::Printf(TEXT("%s: %s: %s"), *(NetConnection->RemoteAddressToString()), ToString(SecurityEventType), Format); \
+					FMsg::Logf_Internal(__FILE__, __LINE__, LogSecurity.GetCategoryName(), ELogVerbosity::Warning, *Test, ##__VA_ARGS__); \
+				} \
+			} \
+		 }
+
+
 		// Conditional logging. Will only log if Condition is met.
 		#define UE_CLOG(Condition, CategoryName, Verbosity, Format, ...) \
 		{ \
@@ -339,7 +361,6 @@ struct FTCharArrayTester
 			} \
 		}
 	#endif
-
 	/** 
 	 * A macro that executes some code within a scope if a given logging category is active at a given verbosity level
 	 * Also, withing the scope of the execution, the default category and verbosity is set up for the low level logging 
diff --git a/Engine/Source/Runtime/Engine/Classes/Engine/NetConnection.h b/Engine/Source/Runtime/Engine/Classes/Engine/NetConnection.h
--- a/Engine/Source/Runtime/Engine/Classes/Engine/NetConnection.h
+++ b/Engine/Source/Runtime/Engine/Classes/Engine/NetConnection.h
@@ -44,6 +44,40 @@ enum EConnectionState
 	USOCK_Open      = 3, // Connection is open.
 };
 
+// 
+// Security event types used for UE_SECURITY_LOG
+//
+namespace ESecurityEvent
+{ 
+	enum Type
+	{
+		Malformed_Packet = 0, // The packet didn't follow protocol
+		Invalid_Data = 1,     // The packet contained invalid data
+		Closed = 2            // The connection had issues (potentially malicious) and was closed
+	};
+	
+	/** @return the stringified version of the enum passed in */
+	inline const TCHAR* ToString(const ESecurityEvent::Type EnumVal)
+	{
+		switch (EnumVal)
+		{
+			case Malformed_Packet:
+			{
+				return TEXT("Malformed_Packet");
+			}
+			case Invalid_Data:
+			{
+				return TEXT("Invalid_Data");
+			}
+			case Closed:
+			{
+				return TEXT("Closed");
+			}
+		}
+		return TEXT("");
+	}
+}
+
 /** If this connection is from a client, this is the current login state of this connection/login attempt */
 namespace EClientLoginState
 {
@@ -94,14 +128,6 @@ struct DelayedPacket
 //	(optionally, 'PING_ACK_PACKET_INTERVAL' can be tweaked, so that the interval checks take a similar amount of time as this delay)
 #define PING_ACK_DELAY 0.5
 
-
-class SecurityLog
-{
-	void SecurityEvent(const char* string);
-};
-
-
-
 UCLASS(customConstructor, Abstract, MinimalAPI, transient, config=Engine)
 class UNetConnection : public UPlayer
 {
@@ -115,8 +141,6 @@ class UNetConnection : public UPlayer
 	UPROPERTY()
 	class UNetDriver* Driver;	
 
-	class SecurityLog Log;
-
 	UPROPERTY()
 	/** Package map between local and remote. (negotiates net serialization) */
 	class UPackageMap* PackageMap;
@@ -482,6 +506,14 @@ public:
 	 */
 	ENGINE_API virtual void InitConnection(UNetDriver* InDriver, EConnectionState InState, const FURL& InURL, int32 InConnectionSpeed=0, int32 InMaxPacket=0);
 
+
+	/** 
+	* Gets a unique ID for the connection, this ID depends on the underlying connection
+	* For IP connections this is an IP Address and port, for steam this is a SteamID
+	*/
+	ENGINE_API virtual FString RemoteAddressToString() PURE_VIRTUAL(UNetConnection::RemoteAddressToString, return TEXT("Error"););
+	
+	
 	// Functions.
 
 	/** Resend any pending acks. */
diff --git a/Engine/Source/Runtime/Engine/Private/DataChannel.cpp b/Engine/Source/Runtime/Engine/Private/DataChannel.cpp
--- a/Engine/Source/Runtime/Engine/Private/DataChannel.cpp
+++ b/Engine/Source/Runtime/Engine/Private/DataChannel.cpp
@@ -19,6 +19,7 @@ DEFINE_LOG_CATEGORY(LogNetPlayerMovement);
 DEFINE_LOG_CATEGORY(LogNetTraffic);
 DEFINE_LOG_CATEGORY(LogRepTraffic);
 DEFINE_LOG_CATEGORY(LogNetDormancy);
+DEFINE_LOG_CATEGORY(LogSecurity);
 DEFINE_LOG_CATEGORY_STATIC(LogNetPartialBunch, Warning, All);
 
 extern FAutoConsoleVariable CVarDoReplicationContextString;
@@ -1101,6 +1102,7 @@ bool UControlChannel::CheckEndianess(FInBunch& Bunch)
 void UControlChannel::ReceivedBunch( FInBunch& Bunch )
 {
 	check(!Closing);
+
 	// If this is a new client connection inspect the raw packet for endianess
 	if (Connection && bNeedsEndianInspection && !CheckEndianess(Bunch))
 	{
diff --git a/Engine/Source/Runtime/Engine/Private/NetConnection.cpp b/Engine/Source/Runtime/Engine/Private/NetConnection.cpp
--- a/Engine/Source/Runtime/Engine/Private/NetConnection.cpp
+++ b/Engine/Source/Runtime/Engine/Private/NetConnection.cpp
@@ -459,16 +459,20 @@ void UNetConnection::ReceivedRawPacket( void* InData, int32 Count )
 			FBitReader Reader( Data, BitSize );
 			ReceivedPacket( Reader );
 		}
+		// MalformedPacket - Received a packet with 0's in the last byte
 		else 
 		{
-			UE_LOG( LogNetTraffic, Error, TEXT( "Packet missing trailing 1" ) );
+			UE_SECURITY_LOG(this, ESecurityEvent::Malformed_Packet, TEXT("Received packet with 0's in last byte of packet"));
 			Close();	// If they have the secret key (or haven't opened the control channel yet) and get here, assume they are being malicious
+			UE_SECURITY_LOG(this, ESecurityEvent::Closed, TEXT("Connection closed"));
 		}
 	}
+	// MalformedPacket - Received a packet of 0 bytes
 	else 
 	{
-		UE_LOG( LogNetTraffic, Error, TEXT( "Received zero-size packet" ) );
+		UE_SECURITY_LOG(this, ESecurityEvent::Malformed_Packet, TEXT("Received zero-size packet"));
 		Close();	// If they have the secret key (or haven't opened the control channel yet) and get here, assume they are being malicious
+		UE_SECURITY_LOG(this, ESecurityEvent::Closed, TEXT("Connection closed"));
 	}
 }
 
@@ -661,7 +665,8 @@ void UNetConnection::ReceivedPacket( FBitReader& Reader )
 		Driver->InPacketsLost += PacketsLost;
 		InPacketId = PacketId;
 	}
-	else
+	// Invalid Data Packet: Received a duplicate packet
+	else if ( PacketId > 0)
 	{
 		Driver->InOutOfOrderPackets++;
 		// Protect against replay attacks
@@ -669,8 +674,18 @@ void UNetConnection::ReceivedPacket( FBitReader& Reader )
 		// The only bunch we would process would be unreliable RPC's, which could allow for replay attacks
 		// So rather than add individual protection for unreliable RPC's as well, just kill it at the source, 
 		// which protects everything in one fell swoop
+
+		UE_SECURITY_LOG(this, ESecurityEvent::Invalid_Data, TEXT("Received a packet with an Acked PacketID"));
+
 		return;		
 	}
+	// Invalid Data Packet: User has sent an invalid PacketId
+	else if ( PacketId < 0)
+	{
+		UE_SECURITY_LOG(this, ESecurityEvent::Invalid_Data, TEXT("Received a packet with an invalid PacketID"));
+
+		return;
+	}
 
 	// Detect packets on the client, which should trigger PingAck verification
 	bool bPendingSendPingAck = Driver->ServerConnection != NULL && (PacketId % PING_ACK_PACKET_INTERVAL) == 0;
@@ -687,8 +702,9 @@ void UNetConnection::ReceivedPacket( FBitReader& Reader )
 		bool IsAck = !!Reader.ReadBit();
 		if ( Reader.IsError() )
 		{
-			UE_LOG( LogNetTraffic, Error, TEXT( "Bunch missing ack flag" ) );
+			UE_SECURITY_LOG(this, ESecurityEvent::Invalid_Data, TEXT("Bunch missing ack flag"));
 			Close();
+			UE_SECURITY_LOG(this, ESecurityEvent::Closed, TEXT("Connection Closed"));
 			return;
 		}
 
@@ -699,10 +715,12 @@ void UNetConnection::ReceivedPacket( FBitReader& Reader )
 
 			// This is an acknowledgment.
 			int32 AckPacketId = MakeRelative(Reader.ReadInt(MAX_PACKETID),OutAckPacketId,MAX_PACKETID);
+
 			if( Reader.IsError() )
 			{
-				UE_LOG( LogNetTraffic, Error, TEXT( "Bunch missing ack" ) );
+				UE_SECURITY_LOG(this, ESecurityEvent::Invalid_Data, TEXT("Bunch missing ack"));
 				Close();
+				UE_SECURITY_LOG(this, ESecurityEvent::Closed, TEXT("Connection Closed"));
 				return;
 			}
 
@@ -1898,10 +1916,3 @@ void UNetConnection::ResetPacketBitCounts()
 	NumAckBits = 0;
 	NumPaddingBits = 0;
 }
-
-
-
-void SecurityLog::SecurityEvent(const char* string)
-{
-	UE_LOG(LogNet, Warning, TEXT("GHI"));
-}
\ No newline at end of file
diff --git a/Engine/Source/Runtime/Engine/Public/EngineLogs.h b/Engine/Source/Runtime/Engine/Public/EngineLogs.h
--- a/Engine/Source/Runtime/Engine/Public/EngineLogs.h
+++ b/Engine/Source/Runtime/Engine/Public/EngineLogs.h
@@ -25,4 +25,5 @@ ENGINE_API DECLARE_LOG_CATEGORY_EXTERN(LogNetDormancy, Warning, All);
 ENGINE_API DECLARE_LOG_CATEGORY_EXTERN(LogSkeletalControl, Warning, All);
 ENGINE_API DECLARE_LOG_CATEGORY_EXTERN(LogSubtitle, Log, All);
 ENGINE_API DECLARE_LOG_CATEGORY_EXTERN(LogTexture, Log, All);
-ENGINE_API DECLARE_LOG_CATEGORY_EXTERN(LogPlayerManagement, Error, All);
\ No newline at end of file
+ENGINE_API DECLARE_LOG_CATEGORY_EXTERN(LogPlayerManagement, Error, All);
+ENGINE_API DECLARE_LOG_CATEGORY_EXTERN(LogSecurity, Warning, All);
\ No newline at end of file
diff --git a/Engine/Source/Runtime/Online/OnlineSubsystemUtils/Classes/IpConnection.h b/Engine/Source/Runtime/Online/OnlineSubsystemUtils/Classes/IpConnection.h
--- a/Engine/Source/Runtime/Online/OnlineSubsystemUtils/Classes/IpConnection.h
+++ b/Engine/Source/Runtime/Online/OnlineSubsystemUtils/Classes/IpConnection.h
@@ -25,5 +25,6 @@ class ONLINESUBSYSTEMUTILS_API UIpConnection : public UNetConnection
 	FString LowLevelDescribe() override;
 	virtual int32 GetAddrAsInt(void) override;
 	virtual int32 GetAddrPort(void) override;
+	virtual FString RemoteAddressToString() override;
 	// End NetConnection Interface
 };
diff --git a/Engine/Source/Runtime/Online/OnlineSubsystemUtils/Private/IpConnection.cpp b/Engine/Source/Runtime/Online/OnlineSubsystemUtils/Private/IpConnection.cpp
--- a/Engine/Source/Runtime/Online/OnlineSubsystemUtils/Private/IpConnection.cpp
+++ b/Engine/Source/Runtime/Online/OnlineSubsystemUtils/Private/IpConnection.cpp
@@ -173,3 +173,8 @@ int32 UIpConnection::GetAddrPort(void)
 	RemoteAddr->GetPort(OutPort);
 	return OutPort;
 }
+
+FString UIpConnection::RemoteAddressToString()
+{
+	return RemoteAddr->ToString(true);
+}
\ No newline at end of file
diff --git a/Engine/Source/Runtime/Online/OnlineSubsystemUtils/Private/IpNetDriver.cpp b/Engine/Source/Runtime/Online/OnlineSubsystemUtils/Private/IpNetDriver.cpp
--- a/Engine/Source/Runtime/Online/OnlineSubsystemUtils/Private/IpNetDriver.cpp
+++ b/Engine/Source/Runtime/Online/OnlineSubsystemUtils/Private/IpNetDriver.cpp
@@ -190,6 +190,22 @@ void UIpNetDriver::TickDispatch( float DeltaTime )
 			}
 			else
 			{
+				// MalformedPacket: Client tried sending a packet that exceeded the maximum packet limit
+				// enforced by the server
+				if (Error == SE_EMSGSIZE)
+				{
+					UIpConnection* Connection = nullptr;
+					if (GetServerConnection() && (*GetServerConnection()->RemoteAddr == *FromAddr))
+					{
+						Connection = GetServerConnection();
+					}
+
+					if (Connection != nullptr)
+					{
+						UE_SECURITY_LOG(Connection, ESecurityEvent::Malformed_Packet, TEXT("Received Packet with bytes > max MTU"));
+					}
+				}
+
 				if( Error != SE_ECONNRESET && Error != SE_UDP_ERR_PORT_UNREACH )
 				{
 					UE_LOG(LogNet, Warning, TEXT("UDP recvfrom error: %i (%s) from %s"),
