Fixed Material Editor so that comment bubbles' toggle states are preserved when exiting the editor.

diff --git a/Engine/Source/Editor/GraphEditor/Private/SCommentBubble.cpp b/Engine/Source/Editor/GraphEditor/Private/SCommentBubble.cpp
--- a/Engine/Source/Editor/GraphEditor/Private/SCommentBubble.cpp
+++ b/Engine/Source/Editor/GraphEditor/Private/SCommentBubble.cpp
@@ -45,6 +45,7 @@ void SCommentBubble::Construct( const FArguments& InArgs )
 	GraphNode				= InArgs._GraphNode;
 	CommentAttribute		= InArgs._Text;
 	OnTextCommittedDelegate	= InArgs._OnTextCommitted;
+	OnToggledDelegate		= InArgs._OnToggled;
 	ColorAndOpacity			= InArgs._ColorAndOpacity;
 	bAllowPinning			= InArgs._AllowPinning;
 	bEnableTitleBarBubble	= InArgs._EnableTitleBarBubble;
@@ -391,9 +392,13 @@ FSlateColor SCommentBubble::GetTextForegroundColor() const
 
 void SCommentBubble::OnCommentTextCommitted( const FText& NewText, ETextCommit::Type CommitInfo )
 {
-	CachedComment = NewText.ToString();
-	CachedCommentText = NewText;
-	OnTextCommittedDelegate.ExecuteIfBound( CachedCommentText, CommitInfo );
+	if (CommitInfo == ETextCommit::OnCleared)
+	{
+		// Don't respond to OnEnter, as it will be immediately followed by OnCleared anyway (due to loss of keyboard focus) and generate a second transaction
+		CachedComment = NewText.ToString();
+		CachedCommentText = NewText;
+		OnTextCommittedDelegate.ExecuteIfBound(CachedCommentText, CommitInfo);
+	}
 }
 
 EVisibility SCommentBubble::GetToggleButtonVisibility() const
@@ -426,11 +431,12 @@ void SCommentBubble::OnCommentBubbleToggle( ECheckBoxState State )
 {
 	if( !IsReadOnly() )
 	{
-		const FScopedTransaction Transaction( NSLOCTEXT( "CommentBubble", "BubbleVisibility", "Comment Bubble Visibilty" ) );
+		const FScopedTransaction Transaction( NSLOCTEXT( "CommentBubble", "BubbleVisibility", "Comment Bubble Visibility" ) );
 		GraphNode->Modify();
 		GraphNode->bCommentBubbleVisible = State == ECheckBoxState::Checked;
 		OpacityValue = 0.f;
 		UpdateBubble();
+		OnToggledDelegate.ExecuteIfBound(GraphNode->bCommentBubbleVisible);
 	}
 }
 
diff --git a/Engine/Source/Editor/GraphEditor/Private/SGraphNode.cpp b/Engine/Source/Editor/GraphEditor/Private/SGraphNode.cpp
--- a/Engine/Source/Editor/GraphEditor/Private/SGraphNode.cpp
+++ b/Engine/Source/Editor/GraphEditor/Private/SGraphNode.cpp
@@ -180,6 +180,11 @@ void SGraphNode::OnCommentTextCommitted(const FText& NewComment, ETextCommit::Ty
 	GetNodeObj()->OnUpdateCommentText(NewComment.ToString());
 }
 
+void SGraphNode::OnCommentBubbleToggled(bool bCommentBubbleVisible)
+{
+	GetNodeObj()->OnCommentBubbleToggled(bCommentBubbleVisible);
+}
+
 void SGraphNode::SetDisallowedPinConnectionEvent(SGraphEditor::FOnDisallowedPinConnection InOnDisallowedPinConnection)
 {
 	OnDisallowedPinConnection = InOnDisallowedPinConnection;
@@ -860,6 +865,7 @@ void SGraphNode::UpdateGraphNode()
 	.GraphNode( GraphNode )
 	.Text( this, &SGraphNode::GetNodeComment )
 	.OnTextCommitted( this, &SGraphNode::OnCommentTextCommitted )
+	.OnToggled( this, &SGraphNode::OnCommentBubbleToggled )
 	.ColorAndOpacity( CommentColor )
 	.AllowPinning( true )
 	.EnableTitleBarBubble( true )
diff --git a/Engine/Source/Editor/GraphEditor/Public/SCommentBubble.h b/Engine/Source/Editor/GraphEditor/Public/SCommentBubble.h
--- a/Engine/Source/Editor/GraphEditor/Public/SCommentBubble.h
+++ b/Engine/Source/Editor/GraphEditor/Public/SCommentBubble.h
@@ -3,6 +3,7 @@
 #pragma once
 
 DECLARE_DELEGATE_RetVal( bool, FIsGraphNodeHovered );
+DECLARE_DELEGATE_OneParam( FOnCommentBubbleToggled, bool );
 
 class GRAPHEDITOR_API SCommentBubble : public SCompoundWidget
 {
@@ -28,6 +29,9 @@ class GRAPHEDITOR_API SCommentBubble : public SCompoundWidget
 		/** Called when the comment text is committed */
 		SLATE_EVENT( FOnTextCommitted, OnTextCommitted )
 
+		/** Called when the comment bubble is toggled */
+		SLATE_EVENT( FOnCommentBubbleToggled, OnToggled )
+
 		/** The comment hint text for the bubble */
 		SLATE_ATTRIBUTE( FText, HintText )
 
@@ -149,6 +153,8 @@ protected:
 
 	/** Optional delegate to call when the comment text is committed */
 	FOnTextCommitted OnTextCommittedDelegate;
+	/** Optional delegate to call when the comment bubble is toggled */
+	FOnCommentBubbleToggled OnToggledDelegate;
 	/** Delegate to determine if the graph node is currently hovered */
 	FIsGraphNodeHovered IsGraphNodeHovered;
 
diff --git a/Engine/Source/Editor/GraphEditor/Public/SGraphNode.h b/Engine/Source/Editor/GraphEditor/Public/SGraphNode.h
--- a/Engine/Source/Editor/GraphEditor/Public/SGraphNode.h
+++ b/Engine/Source/Editor/GraphEditor/Public/SGraphNode.h
@@ -185,6 +185,8 @@ public:
 	void SetDisallowedPinConnectionEvent(SGraphEditor::FOnDisallowedPinConnection InOnDisallowedPinConnection);
 	/** called to replace this nodes comment text */
 	void OnCommentTextCommitted(const FText& NewComment, ETextCommit::Type CommitInfo);
+	/** called when the node's comment bubble is toggled */
+	void OnCommentBubbleToggled(bool bCommentBubbleVisible);
 	/** returns true if a rename is pending on this node */
 	bool IsRenamePending() const { return bRenameIsPending; }
 
diff --git a/Engine/Source/Editor/UnrealEd/Classes/MaterialGraph/MaterialGraphNode.h b/Engine/Source/Editor/UnrealEd/Classes/MaterialGraph/MaterialGraphNode.h
--- a/Engine/Source/Editor/UnrealEd/Classes/MaterialGraph/MaterialGraphNode.h
+++ b/Engine/Source/Editor/UnrealEd/Classes/MaterialGraph/MaterialGraphNode.h
@@ -60,6 +60,7 @@ public:
 	virtual void NodeConnectionListChanged() override;
 	virtual void OnRenameNode(const FString& NewName) override;
 	virtual void OnUpdateCommentText( const FString& NewComment ) override;
+	virtual void OnCommentBubbleToggled( bool bCommentBubbleVisible ) override;
 	virtual void GetPinHoverText(const UEdGraphPin& Pin, FString& HoverTextOut) const override;
 	virtual FString GetDocumentationExcerptName() const override;
 	// End UEdGraphNode interface.
diff --git a/Engine/Source/Editor/UnrealEd/Private/MaterialGraphNode.cpp b/Engine/Source/Editor/UnrealEd/Private/MaterialGraphNode.cpp
--- a/Engine/Source/Editor/UnrealEd/Private/MaterialGraphNode.cpp
+++ b/Engine/Source/Editor/UnrealEd/Private/MaterialGraphNode.cpp
@@ -592,6 +592,7 @@ void UMaterialGraphNode::PostPlacedNewNode()
 	if (MaterialExpression)
 	{
 		NodeComment = MaterialExpression->Desc;
+		bCommentBubbleVisible = MaterialExpression->bCommentBubbleVisible;
 		NodePosX = MaterialExpression->MaterialExpressionEditorX;
 		NodePosY = MaterialExpression->MaterialExpressionEditorY;
 		bCanRenameNode = MaterialExpression->CanRenameNode();
@@ -628,6 +629,17 @@ void UMaterialGraphNode::OnUpdateCommentText( const FString& NewComment )
 	{
 		MaterialExpression->Modify();
 		MaterialExpression->Desc = NewComment;
+		MaterialDirtyDelegate.ExecuteIfBound();
+	}
+}
+
+void UMaterialGraphNode::OnCommentBubbleToggled( bool bCommentBubbleVisible )
+{
+	if ( MaterialExpression )
+	{
+		MaterialExpression->Modify();
+		MaterialExpression->bCommentBubbleVisible = bCommentBubbleVisible;
+		MaterialDirtyDelegate.ExecuteIfBound();
 	}
 }
 
diff --git a/Engine/Source/Runtime/Engine/Classes/EdGraph/EdGraphNode.h b/Engine/Source/Runtime/Engine/Classes/EdGraph/EdGraphNode.h
--- a/Engine/Source/Runtime/Engine/Classes/EdGraph/EdGraphNode.h
+++ b/Engine/Source/Runtime/Engine/Classes/EdGraph/EdGraphNode.h
@@ -415,6 +415,9 @@ public:
 	// called to replace this nodes comment text
 	virtual void OnUpdateCommentText( const FString& NewComment );
 
+	// called when the node's comment bubble is toggled
+	virtual void OnCommentBubbleToggled( bool bCommentBubbleVisible ) {}
+
 	/** Return whether to draw this node as a comment node */
 	virtual bool ShouldDrawNodeAsComment() const { return false; }
 
diff --git a/Engine/Source/Runtime/Engine/Classes/Materials/MaterialExpression.h b/Engine/Source/Runtime/Engine/Classes/Materials/MaterialExpression.h
--- a/Engine/Source/Runtime/Engine/Classes/Materials/MaterialExpression.h
+++ b/Engine/Source/Runtime/Engine/Classes/Materials/MaterialExpression.h
@@ -105,18 +105,6 @@ class ENGINE_API UMaterialExpression : public UObject
 	FGuid MaterialExpressionGuid;
 
 #endif // WITH_EDITORONLY_DATA
-	/** Set to true by RecursiveUpdateRealtimePreview() if the expression's preview needs to be updated in realtime in the material editor. */
-	UPROPERTY()
-	uint32 bRealtimePreview:1;
-
-	/** If true, we should update the preview next render. This is set when changing bRealtimePreview. */
-	UPROPERTY(transient)
-	uint32 bNeedToUpdatePreview:1;
-
-	/** Indicates that this is a 'parameter' type of expression and should always be loaded (ie not cooked away) because we might want the default parameter. */
-	UPROPERTY()
-	uint32 bIsParameterExpression:1;
-
 	/** 
 	 * The material that this expression is currently being compiled in.  
 	 * This is not necessarily the object which owns this expression, for example a preview material compiling a material function's expressions.
@@ -139,6 +127,22 @@ class ENGINE_API UMaterialExpression : public UObject
 	UPROPERTY()
 	FColor BorderColor;
 
+	/** Set to true by RecursiveUpdateRealtimePreview() if the expression's preview needs to be updated in realtime in the material editor. */
+	UPROPERTY()
+	uint32 bRealtimePreview:1;
+
+	/** If true, we should update the preview next render. This is set when changing bRealtimePreview. */
+	UPROPERTY(transient)
+	uint32 bNeedToUpdatePreview:1;
+
+	/** Indicates that this is a 'parameter' type of expression and should always be loaded (ie not cooked away) because we might want the default parameter. */
+	UPROPERTY()
+	uint32 bIsParameterExpression:1;
+
+	/** If true, the comment bubble will be visible in the graph editor */
+	UPROPERTY()
+	uint32 bCommentBubbleVisible:1;
+
 	/** If true, use the output name as the label for the pin */
 	UPROPERTY()
 	uint32 bShowOutputNameOnPin:1;
