Added LogNavOctree command line parameter, to log navoctree as debug geometry to visual logger on stop logging. It fixes UE-16335.

diff --git a/Engine/Source/Developer/LogVisualizer/Private/SVisualLogger.cpp b/Engine/Source/Developer/LogVisualizer/Private/SVisualLogger.cpp
--- a/Engine/Source/Developer/LogVisualizer/Private/SVisualLogger.cpp
+++ b/Engine/Source/Developer/LogVisualizer/Private/SVisualLogger.cpp
@@ -447,9 +447,15 @@ bool SVisualLogger::HandleStopRecordingCommandCanExecute() const
 
 void SVisualLogger::HandleStopRecordingCommandExecute()
 {
+	UWorld* World = FLogVisualizer::Get().GetWorld();
+
+	if (FParse::Param(FCommandLine::Get(), TEXT("DumpNavOctree")) == false && ULogVisualizerSettings::StaticClass()->GetDefaultObject<ULogVisualizerSettings>()->bLogNavOctreeOnStop)
+	{
+		FVisualLogger::NavigationDataDump(World, LogNavigation, ELogVerbosity::Log, INDEX_NONE, FBox());
+	}
+
 	FVisualLogger::Get().SetIsRecording(false);
 
-	UWorld* World = FLogVisualizer::Get().GetWorld();
 	if (AVisualLoggerCameraController::IsEnabled(World))
 	{
 		AVisualLoggerCameraController::DisableCamera(World);
diff --git a/Engine/Source/Developer/LogVisualizer/Public/LogVisualizerSettings.h b/Engine/Source/Developer/LogVisualizer/Public/LogVisualizerSettings.h
--- a/Engine/Source/Developer/LogVisualizer/Public/LogVisualizerSettings.h
+++ b/Engine/Source/Developer/LogVisualizer/Public/LogVisualizerSettings.h
@@ -97,6 +97,10 @@ public:
 	UPROPERTY(EditAnywhere, config, Category = "VisualLogger")
 	bool bUsePlayersOnlyForPause;
 
+	/** Whether to dump Navigation Octree on Stop recording or not */
+	UPROPERTY(EditAnywhere, config, Category = "VisualLogger")
+	bool bLogNavOctreeOnStop;
+
 	// UObject overrides
 #if WITH_EDITOR
 	virtual void PostEditChangeProperty(struct FPropertyChangedEvent& PropertyChangedEvent) override;
diff --git a/Engine/Source/Runtime/Engine/Private/VisualLogger/VisualLogger.cpp b/Engine/Source/Runtime/Engine/Private/VisualLogger/VisualLogger.cpp
--- a/Engine/Source/Runtime/Engine/Private/VisualLogger/VisualLogger.cpp
+++ b/Engine/Source/Runtime/Engine/Private/VisualLogger/VisualLogger.cpp
@@ -263,7 +263,7 @@ void FVisualLogger::NavigationDataDump(const class UObject* Object, const struct
 	const FNavDataGenerator* Generator = MainNavData ? MainNavData->GetGenerator() : nullptr;
 	if (Generator)
 	{
-		Generator->GrabDebugSnapshot(CurrentEntry, Box, Category, Verbosity);
+		Generator->GrabDebugSnapshot(CurrentEntry, FMath::IsNearlyZero(Box.GetVolume()) ? MainNavData->GetBounds() : Box, Category, Verbosity);
 	}
 }
 
@@ -282,24 +282,32 @@ FVisualLogger::FVisualLogger()
 	}
 }
 
-UWorld* FVisualLogger::GetWorld(const class UObject* Object)
+namespace
 {
-	UWorld* World = Object ? GEngine->GetWorldFromContextObject(Object, false) : nullptr;
-#if WITH_EDITOR
-	UEditorEngine *EEngine = Cast<UEditorEngine>(GEngine);
-	if (GIsEditor && EEngine != nullptr && World == nullptr)
+	static UWorld* GetWorldForVisualLogger(const class UObject* Object)
 	{
-		// lets use PlayWorld during PIE/Simulate and regular world from editor otherwise, to draw debug information
-		World = EEngine->PlayWorld != nullptr ? EEngine->PlayWorld : EEngine->GetEditorWorldContext().World();
-	}
+		UWorld* World = Object ? GEngine->GetWorldFromContextObject(Object, false) : nullptr;
+#if WITH_EDITOR
+		UEditorEngine *EEngine = Cast<UEditorEngine>(GEngine);
+		if (GIsEditor && EEngine != nullptr && World == nullptr)
+		{
+			// lets use PlayWorld during PIE/Simulate and regular world from editor otherwise, to draw debug information
+			World = EEngine->PlayWorld != nullptr ? EEngine->PlayWorld : EEngine->GetEditorWorldContext().World();
+		}
 
 #endif
-	if (!GIsEditor && World == nullptr)
-	{
-		World = GEngine->GetWorld();
+		if (!GIsEditor && World == nullptr)
+		{
+			World = GEngine->GetWorld();
+		}
+
+		return World;
 	}
+}
 
-	return World;
+UWorld* FVisualLogger::GetWorld(const class UObject* Object)
+{
+	return GetWorldForVisualLogger(Object);
 }
 
 void FVisualLogger::Shutdown()
@@ -383,6 +391,10 @@ class UObject* FVisualLogger::FindRedirection(const UObject* Object)
 
 void FVisualLogger::SetIsRecording(bool InIsRecording) 
 { 
+	if (InIsRecording == false && InIsRecording != !!bIsRecording && FParse::Param(FCommandLine::Get(), TEXT("LogNavOctree")))
+	{
+		FVisualLogger::NavigationDataDump(GetWorld(nullptr), LogNavigation, ELogVerbosity::Log, INDEX_NONE, FBox());
+	}
 	if (IsRecordingToFile())
 	{
 		SetIsRecordingToFile(false);
@@ -481,6 +493,10 @@ public:
 #endif
 			}
 		}
+		else if (FParse::Command(&Cmd, TEXT("LogNavOctree")))
+		{
+			FVisualLogger::NavigationDataDump(GetWorldForVisualLogger(nullptr), LogNavigation, ELogVerbosity::Log, INDEX_NONE, FBox());
+		}
 		return false;
 	}
 } LogVisualizerExec;
