Remove WITH_EDITOR define around GetActiveSocketList in QuerySupportedSockets

diff --git a/Engine/Source/Runtime/Engine/Private/Components/SkinnedMeshComponent.cpp b/Engine/Source/Runtime/Engine/Private/Components/SkinnedMeshComponent.cpp
--- a/Engine/Source/Runtime/Engine/Private/Components/SkinnedMeshComponent.cpp
+++ b/Engine/Source/Runtime/Engine/Private/Components/SkinnedMeshComponent.cpp
@@ -1318,8 +1318,6 @@ void USkinnedMeshComponent::QuerySupportedSockets(TArray<FComponentSocketDescrip
 {
 	if (SkeletalMesh != NULL)
 	{
-		//@TODO: need to make this work in the game too
-#if WITH_EDITOR
 		// Grab all the mesh and skeleton sockets
 		const TArray<USkeletalMeshSocket*> AllSockets = SkeletalMesh->GetActiveSocketList();
 
@@ -1330,7 +1328,6 @@ void USkinnedMeshComponent::QuerySupportedSockets(TArray<FComponentSocketDescrip
 				new (OutSockets) FComponentSocketDescription(Socket->SocketName, EComponentSocketType::Socket);
 			}
 		}
-#endif
 
 		// Now grab the bones, which can behave exactly like sockets
 		for (int32 BoneIdx = 0; BoneIdx < SkeletalMesh->RefSkeleton.GetNum(); ++BoneIdx)
