Update AttachmentReplication structure on all transform changes
Don't update AttachmentReplication structure without Authority
Setup the relative transforms before calling AttachTo when applying AttachmentReplication structure

diff --git a/Engine/Source/Runtime/Engine/Private/Actor.cpp b/Engine/Source/Runtime/Engine/Private/Actor.cpp
--- a/Engine/Source/Runtime/Engine/Private/Actor.cpp
+++ b/Engine/Source/Runtime/Engine/Private/Actor.cpp
@@ -1260,22 +1260,14 @@ void AActor::OnRep_AttachmentReplication()
 	{
 		if (RootComponent)
 		{
-			USceneComponent* ParentComponent = AttachmentReplication.AttachParent->GetRootComponent();
-			
-			if (AttachmentReplication.AttachComponent != NULL)
-			{
-				ParentComponent = AttachmentReplication.AttachComponent;
-			}
+			USceneComponent* ParentComponent = (AttachmentReplication.AttachComponent ? AttachmentReplication.AttachComponent : AttachmentReplication.AttachParent->GetRootComponent());
 
 			if (ParentComponent)
 			{
-				RootComponent->AttachTo(ParentComponent, AttachmentReplication.AttachSocket);
 				RootComponent->RelativeLocation = AttachmentReplication.LocationOffset;
 				RootComponent->RelativeRotation = AttachmentReplication.RotationOffset;
 				RootComponent->RelativeScale3D = AttachmentReplication.RelativeScale3D;
-
-				RootComponent->UpdateComponentToWorld();
-
+				RootComponent->AttachTo(ParentComponent, AttachmentReplication.AttachSocket);
 			}
 		}
 	}
diff --git a/Engine/Source/Runtime/Engine/Private/Components/SceneComponent.cpp b/Engine/Source/Runtime/Engine/Private/Components/SceneComponent.cpp
--- a/Engine/Source/Runtime/Engine/Private/Components/SceneComponent.cpp
+++ b/Engine/Source/Runtime/Engine/Private/Components/SceneComponent.cpp
@@ -422,6 +422,25 @@ void USceneComponent::PropagateTransformUpdate(bool bTransformChanged, bool bSki
 
 		// Refresh navigation
 		UpdateNavigationData();
+
+
+		AActor* Owner = GetOwner();
+		if (Owner && Owner->HasAuthority() && Owner->GetRootComponent() == this)
+		{
+			if (AttachParent)
+			{
+				Owner->AttachmentReplication.AttachParent = AttachParent->GetAttachmentRootActor();
+				Owner->AttachmentReplication.LocationOffset = RelativeLocation;
+				Owner->AttachmentReplication.RotationOffset = RelativeRotation;
+				Owner->AttachmentReplication.RelativeScale3D = RelativeScale3D;
+				Owner->AttachmentReplication.AttachSocket = AttachSocketName;
+				Owner->AttachmentReplication.AttachComponent = AttachParent;
+			}
+			else
+			{
+				Owner->AttachmentReplication.AttachParent = nullptr;
+			}
+		}
 	}
 	else
 	{
@@ -1286,18 +1305,6 @@ void USceneComponent::AttachTo(class USceneComponent* Parent, FName InSocketName
 		{
 			UpdateOverlaps();
 		}
-
-		AActor* Owner = GetOwner();
-
-		if (Owner && Owner->GetRootComponent() == this)
-		{
-			Owner->AttachmentReplication.AttachParent = AttachParent->GetAttachmentRootActor();
-			Owner->AttachmentReplication.LocationOffset = RelativeLocation;
-			Owner->AttachmentReplication.RotationOffset = RelativeRotation;
-			Owner->AttachmentReplication.RelativeScale3D = RelativeScale3D;
-			Owner->AttachmentReplication.AttachSocket = InSocketName;
-			Owner->AttachmentReplication.AttachComponent = AttachParent;
-		}
 	}
 }
 
@@ -1361,11 +1368,6 @@ void USceneComponent::DetachFromParent(bool bMaintainWorldPosition, bool bCallMo
 		{
 			UpdateOverlaps();
 		}
-
-		if (Owner && Owner->GetRootComponent() == this)
-		{
-			Owner->AttachmentReplication.AttachParent = nullptr;
-		}
 	}
 }
 
