Added a parameter to FObjectReplicator::CallRepNotifies to allow calling notifies even if there are queued bunches on the owning channel.
Call notifies in FObjectReplicator::UpdateUnmappedObjects even with queued bunches, so that we don't fail the check for zero RepNotifies on the next call to UpdateUnmappedObjects.
Fixes UE-15696.

diff --git a/Engine/Source/Runtime/Engine/Private/DataChannel.cpp b/Engine/Source/Runtime/Engine/Private/DataChannel.cpp
--- a/Engine/Source/Runtime/Engine/Private/DataChannel.cpp
+++ b/Engine/Source/Runtime/Engine/Private/DataChannel.cpp
@@ -1770,7 +1770,7 @@ bool UActorChannel::ProcessQueuedBunches()
 		// Call any onreps that were delayed because we were queuing bunches
 		for (auto& ReplicatorPair : ReplicationMap)
 		{
-			ReplicatorPair.Value->CallRepNotifies();
+			ReplicatorPair.Value->CallRepNotifies(true);
 		}
 	}
 
diff --git a/Engine/Source/Runtime/Engine/Private/DataReplication.cpp b/Engine/Source/Runtime/Engine/Private/DataReplication.cpp
--- a/Engine/Source/Runtime/Engine/Private/DataReplication.cpp
+++ b/Engine/Source/Runtime/Engine/Private/DataReplication.cpp
@@ -761,7 +761,7 @@ void FObjectReplicator::PostReceivedBunch()
 	}
 
 	// Call RepNotifies
-	CallRepNotifies();
+	CallRepNotifies(true);
 }
 
 static FORCEINLINE FPropertyRetirement ** UpdateAckedRetirements( FPropertyRetirement &	Retire, int32 OutAckPacketId )
@@ -1126,7 +1126,7 @@ void FObjectReplicator::StartBecomingDormant()
 	bLastUpdateEmpty = false; // Ensure we get one more attempt to update properties
 }
 
-void FObjectReplicator::CallRepNotifies()
+void FObjectReplicator::CallRepNotifies(bool bSkipIfChannelHasQueuedBunches)
 {
 	UObject* Object = GetObject();
 
@@ -1140,7 +1140,7 @@ void FObjectReplicator::CallRepNotifies()
 		return;
 	}
 
-	if ( OwningChannel != NULL && OwningChannel->QueuedBunches.Num() > 0 )
+	if ( bSkipIfChannelHasQueuedBunches && ( OwningChannel != NULL && OwningChannel->QueuedBunches.Num() > 0 ) )
 	{
 		return;
 	}
@@ -1278,7 +1278,9 @@ void FObjectReplicator::UpdateUnmappedObjects( bool & bOutHasMoreUnmapped )
 	}
 
 	// Call any rep notifies that need to happen when object pointers change
-	CallRepNotifies();
+	// Pass in false to override the check for queued bunches. Otherwise, if the owning channel has queued bunches,
+	// the RepNotifies will remain in the list and the check for 0 RepNotifies above will fail next time.
+	CallRepNotifies(false);
 
 	if ( bSomeObjectsWereMapped )
 	{
diff --git a/Engine/Source/Runtime/Engine/Private/DemoNetDriver.cpp b/Engine/Source/Runtime/Engine/Private/DemoNetDriver.cpp
--- a/Engine/Source/Runtime/Engine/Private/DemoNetDriver.cpp
+++ b/Engine/Source/Runtime/Engine/Private/DemoNetDriver.cpp
@@ -1353,7 +1353,7 @@ void UDemoNetDriver::TickDemoPlayback( float DeltaSeconds )
 				{
 					for (auto& ReplicatorPair : ChannelPair.Value->ReplicationMap)
 					{
-						ReplicatorPair.Value->CallRepNotifies();
+						ReplicatorPair.Value->CallRepNotifies(true);
 					}
 				}
 			}
diff --git a/Engine/Source/Runtime/Engine/Public/Net/DataReplication.h b/Engine/Source/Runtime/Engine/Public/Net/DataReplication.h
--- a/Engine/Source/Runtime/Engine/Public/Net/DataReplication.h
+++ b/Engine/Source/Runtime/Engine/Public/Net/DataReplication.h
@@ -130,7 +130,7 @@ public:
 
 	void StartBecomingDormant();
 
-	void CallRepNotifies();
+	void CallRepNotifies(bool bSkipIfChannelHasQueuedBunches);
 
 	void UpdateUnmappedObjects( bool & bOutHasMoreUnmapped );
 
