TextureSample material nodes now fetch the correct tooltips.

diff --git a/Engine/Source/Runtime/Engine/Classes/Materials/MaterialExpressionTextureSample.h b/Engine/Source/Runtime/Engine/Classes/Materials/MaterialExpressionTextureSample.h
--- a/Engine/Source/Runtime/Engine/Classes/Materials/MaterialExpressionTextureSample.h
+++ b/Engine/Source/Runtime/Engine/Classes/Materials/MaterialExpressionTextureSample.h
@@ -74,6 +74,7 @@ class ENGINE_API UMaterialExpressionTextureSample : public UMaterialExpressionTe
 	virtual bool MatchesSearchQuery( const TCHAR* SearchQuery ) override;
 #if WITH_EDITOR
 	virtual uint32 GetInputType(int32 InputIndex) override;
+	virtual void GetConnectorToolTip(int32 InputIndex, int32 OutputIndex, TArray<FString>& OutToolTip) override;
 #endif
 	// End UMaterialExpression Interface
 
diff --git a/Engine/Source/Runtime/Engine/Private/Materials/MaterialExpressions.cpp b/Engine/Source/Runtime/Engine/Private/Materials/MaterialExpressions.cpp
--- a/Engine/Source/Runtime/Engine/Private/Materials/MaterialExpressions.cpp
+++ b/Engine/Source/Runtime/Engine/Private/Materials/MaterialExpressions.cpp
@@ -1405,6 +1405,17 @@ uint32 UMaterialExpressionTextureSample::GetInputType(int32 InputIndex)
 	return MCT_Unknown;
 }
 #undef IF_INPUT_RETURN
+
+void UMaterialExpressionTextureSample::GetConnectorToolTip(int32 InputIndex, int32 OutputIndex, TArray<FString>& OutToolTip)
+{
+	if (InputIndex == 1 && !GetOuter()->IsA(UMaterialFunction::StaticClass()))
+	{
+		// If input pin 1 is omitted, increment the index so the correct tooltip is looked up
+		InputIndex++;
+	}
+
+	Super::GetConnectorToolTip(InputIndex, OutputIndex, OutToolTip);
+}
 #endif
 
 
