Clamp unreasonable sizes for the checkerboard pattern in the texture editor

diff --git a/Engine/Source/Editor/TextureEditor/Classes/TextureEditorSettings.h b/Engine/Source/Editor/TextureEditor/Classes/TextureEditorSettings.h
--- a/Engine/Source/Editor/TextureEditor/Classes/TextureEditorSettings.h
+++ b/Engine/Source/Editor/TextureEditor/Classes/TextureEditorSettings.h
@@ -46,7 +46,7 @@ public:
 	FColor CheckerColorTwo;
 
 	/** The size of the checkered background tiles. */
-	UPROPERTY(config, EditAnywhere, Category=Background)
+	UPROPERTY(config, EditAnywhere, Category=Background, meta=(ClampMin="2", ClampMax="4096"))
 	int32 CheckerSize;
 
 public:
diff --git a/Engine/Source/Runtime/Engine/Private/ImageUtils.cpp b/Engine/Source/Runtime/Engine/Private/ImageUtils.cpp
--- a/Engine/Source/Runtime/Engine/Private/ImageUtils.cpp
+++ b/Engine/Source/Runtime/Engine/Private/ImageUtils.cpp
@@ -259,7 +259,7 @@ void FImageUtils::CompressImageArray( int32 ImageWidth, int32 ImageHeight, TArra
 
 UTexture2D* FImageUtils::CreateCheckerboardTexture(FColor ColorOne, FColor ColorTwo, int32 CheckerSize)
 {
-	CheckerSize = FMath::RoundUpToPowerOfTwo(CheckerSize);
+	CheckerSize = FMath::Min<uint32>( FMath::RoundUpToPowerOfTwo(CheckerSize), 4096 );
 	const int32 HalfPixelNum = CheckerSize >> 1;
 
 	// Create the texture
