fixed Clang errors/warnings

diff --git a/Engine/Source/Runtime/Engine/Private/SceneView.cpp b/Engine/Source/Runtime/Engine/Private/SceneView.cpp
--- a/Engine/Source/Runtime/Engine/Private/SceneView.cpp
+++ b/Engine/Source/Runtime/Engine/Private/SceneView.cpp
@@ -748,179 +748,180 @@ void FSceneView::OverridePostProcessSettings(const FPostProcessSettings& Src, fl
 		Weight = 1.0f;
 	}
 
-	FFinalPostProcessSettings& Dest = FinalPostProcessSettings;
-
-	// The following code needs to be adjusted when settings in FPostProcessSettings change.
-	LERP_PP(WhiteTemp);
-	LERP_PP(WhiteTint);
-	
-	LERP_PP(ColorSaturation);
-	LERP_PP(ColorContrast);
-	LERP_PP(ColorGamma);
-	LERP_PP(ColorGain);
-	LERP_PP(ColorOffset);
-
-	LERP_PP(FilmWhitePoint);
-	LERP_PP(FilmSaturation);
-	LERP_PP(FilmChannelMixerRed);
-	LERP_PP(FilmChannelMixerGreen);
-	LERP_PP(FilmChannelMixerBlue);
-	LERP_PP(FilmContrast);
-	LERP_PP(FilmDynamicRange);
-	LERP_PP(FilmHealAmount);
-	LERP_PP(FilmToeAmount);
-	LERP_PP(FilmShadowTint);
-	LERP_PP(FilmShadowTintBlend);
-	LERP_PP(FilmShadowTintAmount);
-
-	LERP_PP(FilmSlope);
-	LERP_PP(FilmToe);
-	LERP_PP(FilmShoulder);
-	LERP_PP(FilmBlackClip);
-	LERP_PP(FilmWhiteClip);
-
-	LERP_PP(SceneColorTint);
-	LERP_PP(SceneFringeIntensity);
-	LERP_PP(SceneFringeSaturation);
-	LERP_PP(BloomIntensity);
-	LERP_PP(BloomThreshold);
-	LERP_PP(Bloom1Tint);
-	LERP_PP(BloomSizeScale);
-	LERP_PP(Bloom1Size);
-	LERP_PP(Bloom2Tint);
-	LERP_PP(Bloom2Size);
-	LERP_PP(Bloom3Tint);
-	LERP_PP(Bloom3Size);
-	LERP_PP(Bloom4Tint);
-	LERP_PP(Bloom4Size);
-	LERP_PP(Bloom5Tint);
-	LERP_PP(Bloom5Size);
-	LERP_PP(Bloom6Tint);
-	LERP_PP(Bloom6Size);
-	LERP_PP(BloomDirtMaskIntensity);
-	LERP_PP(BloomDirtMaskTint);
-	LERP_PP(AmbientCubemapIntensity);
-	LERP_PP(AmbientCubemapTint);
-	LERP_PP(AutoExposureLowPercent);
-	LERP_PP(AutoExposureHighPercent);
-	LERP_PP(AutoExposureMinBrightness);
-	LERP_PP(AutoExposureMaxBrightness);
-	LERP_PP(AutoExposureSpeedUp);
-	LERP_PP(AutoExposureSpeedDown);
-	LERP_PP(AutoExposureBias);
-	LERP_PP(HistogramLogMin);
-	LERP_PP(HistogramLogMax);
-	LERP_PP(LensFlareIntensity);
-	LERP_PP(LensFlareTint);
-	LERP_PP(LensFlareBokehSize);
-	LERP_PP(LensFlareThreshold);
-	LERP_PP(VignetteIntensity);
-	LERP_PP(GrainIntensity);
-	LERP_PP(GrainJitter);
-	LERP_PP(AmbientOcclusionIntensity);
-	LERP_PP(AmbientOcclusionStaticFraction);
-	LERP_PP(AmbientOcclusionRadius);
-	LERP_PP(AmbientOcclusionFadeDistance);
-	LERP_PP(AmbientOcclusionFadeRadius);
-	LERP_PP(AmbientOcclusionDistance_DEPRECATED);
-	LERP_PP(AmbientOcclusionPower);
-	LERP_PP(AmbientOcclusionBias);
-	LERP_PP(AmbientOcclusionQuality);
-	LERP_PP(AmbientOcclusionMipBlend);
-	LERP_PP(AmbientOcclusionMipScale);
-	LERP_PP(AmbientOcclusionMipThreshold);
-	LERP_PP(IndirectLightingColor);
-	LERP_PP(IndirectLightingIntensity);
-	LERP_PP(DepthOfFieldFocalDistance);
-	LERP_PP(DepthOfFieldFstop);
-	LERP_PP(DepthOfFieldDepthBlurRadius);
-	LERP_PP(DepthOfFieldDepthBlurAmount);
-	LERP_PP(DepthOfFieldFocalRegion);
-	LERP_PP(DepthOfFieldNearTransitionRegion);
-	LERP_PP(DepthOfFieldFarTransitionRegion);
-	LERP_PP(DepthOfFieldScale);
-	LERP_PP(DepthOfFieldMaxBokehSize);
-	LERP_PP(DepthOfFieldNearBlurSize);
-	LERP_PP(DepthOfFieldFarBlurSize);
-	LERP_PP(DepthOfFieldOcclusion);
-	LERP_PP(DepthOfFieldColorThreshold);
-	LERP_PP(DepthOfFieldSizeThreshold);
-	LERP_PP(DepthOfFieldSkyFocusDistance);
-	LERP_PP(MotionBlurAmount);
-	LERP_PP(MotionBlurMax);
-	LERP_PP(MotionBlurPerObjectSize);
-	LERP_PP(ScreenPercentage);
-	LERP_PP(ScreenSpaceReflectionQuality);
-	LERP_PP(ScreenSpaceReflectionIntensity);
-	LERP_PP(ScreenSpaceReflectionMaxRoughness);
-
-	// cubemaps are getting blended additively - in contrast to other properties, maybe we should make that consistent
-	if(Src.AmbientCubemap && Src.bOverride_AmbientCubemapIntensity)
-	{
-		FFinalPostProcessSettings::FCubemapEntry Entry;
-
-		Entry.AmbientCubemapTintMulScaleValue = FLinearColor(1, 1, 1, 1) * Src.AmbientCubemapIntensity;
-
-		if(Src.bOverride_AmbientCubemapTint)
+	{
+		FFinalPostProcessSettings& Dest = FinalPostProcessSettings;
+
+		// The following code needs to be adjusted when settings in FPostProcessSettings change.
+		LERP_PP(WhiteTemp);
+		LERP_PP(WhiteTint);
+
+		LERP_PP(ColorSaturation);
+		LERP_PP(ColorContrast);
+		LERP_PP(ColorGamma);
+		LERP_PP(ColorGain);
+		LERP_PP(ColorOffset);
+
+		LERP_PP(FilmWhitePoint);
+		LERP_PP(FilmSaturation);
+		LERP_PP(FilmChannelMixerRed);
+		LERP_PP(FilmChannelMixerGreen);
+		LERP_PP(FilmChannelMixerBlue);
+		LERP_PP(FilmContrast);
+		LERP_PP(FilmDynamicRange);
+		LERP_PP(FilmHealAmount);
+		LERP_PP(FilmToeAmount);
+		LERP_PP(FilmShadowTint);
+		LERP_PP(FilmShadowTintBlend);
+		LERP_PP(FilmShadowTintAmount);
+
+		LERP_PP(FilmSlope);
+		LERP_PP(FilmToe);
+		LERP_PP(FilmShoulder);
+		LERP_PP(FilmBlackClip);
+		LERP_PP(FilmWhiteClip);
+
+		LERP_PP(SceneColorTint);
+		LERP_PP(SceneFringeIntensity);
+		LERP_PP(SceneFringeSaturation);
+		LERP_PP(BloomIntensity);
+		LERP_PP(BloomThreshold);
+		LERP_PP(Bloom1Tint);
+		LERP_PP(BloomSizeScale);
+		LERP_PP(Bloom1Size);
+		LERP_PP(Bloom2Tint);
+		LERP_PP(Bloom2Size);
+		LERP_PP(Bloom3Tint);
+		LERP_PP(Bloom3Size);
+		LERP_PP(Bloom4Tint);
+		LERP_PP(Bloom4Size);
+		LERP_PP(Bloom5Tint);
+		LERP_PP(Bloom5Size);
+		LERP_PP(Bloom6Tint);
+		LERP_PP(Bloom6Size);
+		LERP_PP(BloomDirtMaskIntensity);
+		LERP_PP(BloomDirtMaskTint);
+		LERP_PP(AmbientCubemapIntensity);
+		LERP_PP(AmbientCubemapTint);
+		LERP_PP(AutoExposureLowPercent);
+		LERP_PP(AutoExposureHighPercent);
+		LERP_PP(AutoExposureMinBrightness);
+		LERP_PP(AutoExposureMaxBrightness);
+		LERP_PP(AutoExposureSpeedUp);
+		LERP_PP(AutoExposureSpeedDown);
+		LERP_PP(AutoExposureBias);
+		LERP_PP(HistogramLogMin);
+		LERP_PP(HistogramLogMax);
+		LERP_PP(LensFlareIntensity);
+		LERP_PP(LensFlareTint);
+		LERP_PP(LensFlareBokehSize);
+		LERP_PP(LensFlareThreshold);
+		LERP_PP(VignetteIntensity);
+		LERP_PP(GrainIntensity);
+		LERP_PP(GrainJitter);
+		LERP_PP(AmbientOcclusionIntensity);
+		LERP_PP(AmbientOcclusionStaticFraction);
+		LERP_PP(AmbientOcclusionRadius);
+		LERP_PP(AmbientOcclusionFadeDistance);
+		LERP_PP(AmbientOcclusionFadeRadius);
+		LERP_PP(AmbientOcclusionDistance_DEPRECATED);
+		LERP_PP(AmbientOcclusionPower);
+		LERP_PP(AmbientOcclusionBias);
+		LERP_PP(AmbientOcclusionQuality);
+		LERP_PP(AmbientOcclusionMipBlend);
+		LERP_PP(AmbientOcclusionMipScale);
+		LERP_PP(AmbientOcclusionMipThreshold);
+		LERP_PP(IndirectLightingColor);
+		LERP_PP(IndirectLightingIntensity);
+		LERP_PP(DepthOfFieldFocalDistance);
+		LERP_PP(DepthOfFieldFstop);
+		LERP_PP(DepthOfFieldDepthBlurRadius);
+		LERP_PP(DepthOfFieldDepthBlurAmount);
+		LERP_PP(DepthOfFieldFocalRegion);
+		LERP_PP(DepthOfFieldNearTransitionRegion);
+		LERP_PP(DepthOfFieldFarTransitionRegion);
+		LERP_PP(DepthOfFieldScale);
+		LERP_PP(DepthOfFieldMaxBokehSize);
+		LERP_PP(DepthOfFieldNearBlurSize);
+		LERP_PP(DepthOfFieldFarBlurSize);
+		LERP_PP(DepthOfFieldOcclusion);
+		LERP_PP(DepthOfFieldColorThreshold);
+		LERP_PP(DepthOfFieldSizeThreshold);
+		LERP_PP(DepthOfFieldSkyFocusDistance);
+		LERP_PP(MotionBlurAmount);
+		LERP_PP(MotionBlurMax);
+		LERP_PP(MotionBlurPerObjectSize);
+		LERP_PP(ScreenPercentage);
+		LERP_PP(ScreenSpaceReflectionQuality);
+		LERP_PP(ScreenSpaceReflectionIntensity);
+		LERP_PP(ScreenSpaceReflectionMaxRoughness);
+
+		// cubemaps are getting blended additively - in contrast to other properties, maybe we should make that consistent
+		if (Src.AmbientCubemap && Src.bOverride_AmbientCubemapIntensity)
 		{
-			Entry.AmbientCubemapTintMulScaleValue *= Src.AmbientCubemapTint;
-		}
+			FFinalPostProcessSettings::FCubemapEntry Entry;
 
-		Entry.AmbientCubemap = Src.AmbientCubemap;
-		Dest.UpdateEntry(Entry, Weight);
-	}
+			Entry.AmbientCubemapTintMulScaleValue = FLinearColor(1, 1, 1, 1) * Src.AmbientCubemapIntensity;
 
-	IF_PP(ColorGradingLUT)
-	{
-		float ColorGradingIntensity = FMath::Clamp(Src.ColorGradingIntensity, 0.0f, 1.0f);
-		Dest.LerpTo(Src.ColorGradingLUT, ColorGradingIntensity * Weight);
-	}
+			if (Src.bOverride_AmbientCubemapTint)
+			{
+				Entry.AmbientCubemapTintMulScaleValue *= Src.AmbientCubemapTint;
+			}
 
-	// actual texture cannot be blended but the intensity can be blended
-	IF_PP(BloomDirtMask)
-	{
-		Dest.BloomDirtMask = Src.BloomDirtMask;
-	}
+			Entry.AmbientCubemap = Src.AmbientCubemap;
+			Dest.UpdateEntry(Entry, Weight);
+		}
 
-	// actual texture cannot be blended but the intensity can be blended
-	IF_PP(DepthOfFieldBokehShape)
-	{
-		Dest.DepthOfFieldBokehShape = Src.DepthOfFieldBokehShape;
-	}
+		IF_PP(ColorGradingLUT)
+		{
+			float ColorGradingIntensity = FMath::Clamp(Src.ColorGradingIntensity, 0.0f, 1.0f);
+			Dest.LerpTo(Src.ColorGradingLUT, ColorGradingIntensity * Weight);
+		}
 
-	// actual texture cannot be blended but the intensity can be blended
-	IF_PP(LensFlareBokehShape)
-	{
-		Dest.LensFlareBokehShape = Src.LensFlareBokehShape;
-	}
+		// actual texture cannot be blended but the intensity can be blended
+		IF_PP(BloomDirtMask)
+		{
+			Dest.BloomDirtMask = Src.BloomDirtMask;
+		}
 
-	if(Src.bOverride_LensFlareTints)
-	{
-		for(uint32 i = 0; i < 8; ++i)
+		// actual texture cannot be blended but the intensity can be blended
+		IF_PP(DepthOfFieldBokehShape)
 		{
-			Dest.LensFlareTints[i] = FMath::Lerp(Dest.LensFlareTints[i], Src.LensFlareTints[i], Weight);
+			Dest.DepthOfFieldBokehShape = Src.DepthOfFieldBokehShape;
 		}
-	}
 
-	if(Src.bOverride_DepthOfFieldMethod)
-	{
-		Dest.DepthOfFieldMethod = Src.DepthOfFieldMethod;
-	}
+		// actual texture cannot be blended but the intensity can be blended
+		IF_PP(LensFlareBokehShape)
+		{
+			Dest.LensFlareBokehShape = Src.LensFlareBokehShape;
+		}
 
-	if(Src.bOverride_AmbientOcclusionRadiusInWS)
-	{
-		Dest.AmbientOcclusionRadiusInWS = Src.AmbientOcclusionRadiusInWS;
-	}
+		if (Src.bOverride_LensFlareTints)
+		{
+			for (uint32 i = 0; i < 8; ++i)
+			{
+				Dest.LensFlareTints[i] = FMath::Lerp(Dest.LensFlareTints[i], Src.LensFlareTints[i], Weight);
+			}
+		}
 
-	if(Src.bOverride_AntiAliasingMethod)
-	{
-		Dest.AntiAliasingMethod = Src.AntiAliasingMethod;
+		if (Src.bOverride_DepthOfFieldMethod)
+		{
+			Dest.DepthOfFieldMethod = Src.DepthOfFieldMethod;
+		}
+
+		if (Src.bOverride_AmbientOcclusionRadiusInWS)
+		{
+			Dest.AmbientOcclusionRadiusInWS = Src.AmbientOcclusionRadiusInWS;
+		}
+
+		if (Src.bOverride_AntiAliasingMethod)
+		{
+			Dest.AntiAliasingMethod = Src.AntiAliasingMethod;
+		}
 	}
 	
 	// will be deprecated soon, use the new asset LightPropagationVolumeBlendable instead
 	{
 		FLightPropagationVolumeSettings& Dest = FinalPostProcessSettings.BlendableManager.GetSingleFinalData<FLightPropagationVolumeSettings>();
-		const float Weight = 1.0f;
 
 		LERP_PP(LPVIntensity);
 		LERP_PP(LPVSecondaryOcclusionIntensity);
