fixed clang warning

diff --git a/Engine/Source/Runtime/Engine/Classes/EdGraph/EdGraphNode.h b/Engine/Source/Runtime/Engine/Classes/EdGraph/EdGraphNode.h
--- a/Engine/Source/Runtime/Engine/Classes/EdGraph/EdGraphNode.h
+++ b/Engine/Source/Runtime/Engine/Classes/EdGraph/EdGraphNode.h
@@ -416,7 +416,7 @@ public:
 	virtual void OnUpdateCommentText( const FString& NewComment );
 
 	// called when the node's comment bubble is toggled
-	virtual void OnCommentBubbleToggled( bool bCommentBubbleVisible ) {}
+	virtual void OnCommentBubbleToggled( bool bInCommentBubbleVisible ) {}
 
 	/** Return whether to draw this node as a comment node */
 	virtual bool ShouldDrawNodeAsComment() const { return false; }
