#ue4 - Perf: Avoid recomputing transforms and Quat<->Rotator conversions in FAnimNode_RotateRootBone::Evaluate() when neither Pitch nor Yaw are set.

diff --git a/Engine/Source/Runtime/Engine/Private/Animation/AnimNode_RotateRootBone.cpp b/Engine/Source/Runtime/Engine/Private/Animation/AnimNode_RotateRootBone.cpp
--- a/Engine/Source/Runtime/Engine/Private/Animation/AnimNode_RotateRootBone.cpp
+++ b/Engine/Source/Runtime/Engine/Private/Animation/AnimNode_RotateRootBone.cpp
@@ -30,17 +30,20 @@ void FAnimNode_RotateRootBone::Evaluate(FPoseContext& Output)
 	checkSlow(!FMath::IsNaN(Yaw) && FMath::IsFinite(Yaw));
 	checkSlow(!FMath::IsNaN(Pitch) && FMath::IsFinite(Pitch));
 
-	// Build our desired rotation
-	const FRotator DeltaRotation(Pitch, Yaw, 0.f);
-	const FQuat DeltaQuat(DeltaRotation);
-	const FQuat MeshToComponentQuat(MeshToComponent);
+	if (Pitch != 0.f || Yaw != 0.f)
+	{
+		// Build our desired rotation
+		const FRotator DeltaRotation(Pitch, Yaw, 0.f);
+		const FQuat DeltaQuat(DeltaRotation);
+		const FQuat MeshToComponentQuat(MeshToComponent);
 
-	// Convert our rotation from Component Space to Mesh Space.
-	const FQuat MeshSpaceDeltaQuat = MeshToComponentQuat.Inverse() * DeltaQuat * MeshToComponentQuat;
+		// Convert our rotation from Component Space to Mesh Space.
+		const FQuat MeshSpaceDeltaQuat = MeshToComponentQuat.Inverse() * DeltaQuat * MeshToComponentQuat;
 
-	// Apply rotation to root bone.
-	FCompactPoseBoneIndex RootBone(0);
-	Output.Pose[RootBone].SetRotation(Output.Pose[RootBone].GetRotation() * MeshSpaceDeltaQuat);
+		// Apply rotation to root bone.
+		FCompactPoseBoneIndex RootBone(0);
+		Output.Pose[RootBone].SetRotation(Output.Pose[RootBone].GetRotation() * MeshSpaceDeltaQuat);
+	}
 }
 
 
