#ue4 - Exported FOverlapInfo with ENGINE_API.

diff --git a/Engine/Source/Runtime/Engine/Classes/Components/SceneComponent.h b/Engine/Source/Runtime/Engine/Classes/Components/SceneComponent.h
--- a/Engine/Source/Runtime/Engine/Classes/Components/SceneComponent.h
+++ b/Engine/Source/Runtime/Engine/Classes/Components/SceneComponent.h
@@ -9,7 +9,7 @@
 
 /** Overlap info consisting of the primitive and the body that is overlapping */
 USTRUCT()
-struct FOverlapInfo
+struct ENGINE_API FOverlapInfo
 {
 	GENERATED_USTRUCT_BODY()
 
