#ue4 - Move static initialization of some text out of functions and in to global scope. Avoids branches and extra instructions in these functions.

diff --git a/Engine/Source/Runtime/Engine/Private/Components/CharacterMovementComponent.cpp b/Engine/Source/Runtime/Engine/Private/Components/CharacterMovementComponent.cpp
--- a/Engine/Source/Runtime/Engine/Private/Components/CharacterMovementComponent.cpp
+++ b/Engine/Source/Runtime/Engine/Private/Components/CharacterMovementComponent.cpp
@@ -59,6 +59,20 @@ const float UCharacterMovementComponent::MAX_FLOOR_DIST = 2.4f;
 const float UCharacterMovementComponent::BRAKE_TO_STOP_VELOCITY = 10.f;
 const float UCharacterMovementComponent::SWEEP_EDGE_REJECT_DISTANCE = 0.15f;
 
+// Statics
+namespace CharacterMovementComponentStatics
+{
+	static const FName CrouchTraceName = FName(TEXT("CrouchTrace"));
+	static const FName FindWaterLineName = FName(TEXT("FindWaterLine"));
+	static const FName FallingTraceParamsTag = FName(TEXT("PhysFalling"));
+	static const FName CheckLedgeDirectionName = FName(TEXT("CheckLedgeDirection"));
+	static const FName ProjectLocationName = FName(TEXT("NavProjectLocation"));
+	static const FName CheckWaterJumpName = FName(TEXT("CheckWaterJump"));
+	static const FName ComputeFloorDistName = FName(TEXT("ComputeFloorDistSweep"));
+	static const FName FloorLineTraceName = FName(TEXT("ComputeFloorDistLineTrace"));
+	static const FName ImmersionDepthName = FName(TEXT("MovementComp_Character_ImmersionDepth"));
+}
+
 // CVars
 static TAutoConsoleVariable<int32> CVarNetEnableMoveCombining(
 	TEXT("p.NetEnableMoveCombining"),
@@ -1773,8 +1787,7 @@ void UCharacterMovementComponent::Crouch(bool bClientSimulation)
 		// Crouching to a larger height? (this is rare)
 		if (CrouchedHalfHeight > OldUnscaledHalfHeight)
 		{
-			static const FName NAME_CrouchTrace = FName(TEXT("CrouchTrace"));
-			FCollisionQueryParams CapsuleParams(NAME_CrouchTrace, false, CharacterOwner);
+			FCollisionQueryParams CapsuleParams(CharacterMovementComponentStatics::CrouchTraceName, false, CharacterOwner);
 			FCollisionResponseParams ResponseParam;
 			InitCollisionParams(CapsuleParams, ResponseParam);
 			const bool bEncroached = GetWorld()->OverlapBlockingTestByChannel(UpdatedComponent->GetComponentLocation() - FVector(0.f,0.f,ScaledHalfHeightAdjust), FQuat::Identity,
@@ -1846,9 +1859,8 @@ void UCharacterMovementComponent::UnCrouch(bool bClientSimulation)
 	if( !bClientSimulation )
 	{
 		// Try to stay in place and see if the larger capsule fits. We use a slightly taller capsule to avoid penetration.
-		static const FName NAME_CrouchTrace = FName(TEXT("CrouchTrace"));
 		const float SweepInflation = KINDA_SMALL_NUMBER * 10.f;
-		FCollisionQueryParams CapsuleParams(NAME_CrouchTrace, false, CharacterOwner);
+		FCollisionQueryParams CapsuleParams(CharacterMovementComponentStatics::CrouchTraceName, false, CharacterOwner);
 		FCollisionResponseParams ResponseParam;
 		InitCollisionParams(CapsuleParams, ResponseParam);
 		const FCollisionShape StandingCapsuleShape = GetPawnCapsuleCollisionShape(SHRINK_HeightCustom, -SweepInflation); // Shrink by negative amount, so actually grow it.
@@ -2202,9 +2214,7 @@ float UCharacterMovementComponent::ImmersionDepth()
 				const FVector TraceStart = UpdatedComponent->GetComponentLocation() + FVector(0.f,0.f,CollisionHalfHeight);
 				const FVector TraceEnd = UpdatedComponent->GetComponentLocation() - FVector(0.f,0.f,CollisionHalfHeight);
 
-				const static FName MovementComp_Character_ImmersionDepthName(TEXT("MovementComp_Character_ImmersionDepth"));
-				FCollisionQueryParams NewTraceParams( MovementComp_Character_ImmersionDepthName, true );
-
+				FCollisionQueryParams NewTraceParams(CharacterMovementComponentStatics::ImmersionDepthName, true);
 				VolumeBrushComp->LineTraceComponent( Hit, TraceStart, TraceEnd, NewTraceParams );
 			}
 
@@ -2958,8 +2968,7 @@ FVector UCharacterMovementComponent::FindWaterLine(FVector InWater, FVector Outo
 	FVector Result = OutofWater;
 
 	TArray<FHitResult> Hits;
-	static const FName NAME_FindWaterLine = FName(TEXT("FindWaterLine"));
-	GetWorld()->LineTraceMultiByChannel(Hits, OutofWater, InWater, UpdatedComponent->GetCollisionObjectType(), FCollisionQueryParams(NAME_FindWaterLine, true, CharacterOwner));
+	GetWorld()->LineTraceMultiByChannel(Hits, OutofWater, InWater, UpdatedComponent->GetCollisionObjectType(), FCollisionQueryParams(CharacterMovementComponentStatics::FindWaterLineName, true, CharacterOwner));
 
 	for( int32 HitIdx = 0; HitIdx < Hits.Num(); HitIdx++ )
 	{
@@ -3295,8 +3304,7 @@ bool UCharacterMovementComponent::FindAirControlImpact(float DeltaTime, float Ad
 	
 	if (!TestWalk.IsZero())
 	{
-		static const FName FallingTraceParamsTag = FName(TEXT("PhysFalling"));
-		FCollisionQueryParams CapsuleQuery(FallingTraceParamsTag, false, CharacterOwner);
+		FCollisionQueryParams CapsuleQuery(CharacterMovementComponentStatics::FallingTraceParamsTag, false, CharacterOwner);
 		FCollisionResponseParams ResponseParam;
 		InitCollisionParams(CapsuleQuery, ResponseParam);
 		const FVector CapsuleLocation = UpdatedComponent->GetComponentLocation();
@@ -3341,8 +3349,7 @@ FVector UCharacterMovementComponent::LimitAirControl(float DeltaTime, const FVec
 bool UCharacterMovementComponent::CheckLedgeDirection(const FVector& OldLocation, const FVector& SideStep, const FVector& GravDir)
 {
 	FVector SideDest = OldLocation + SideStep;
-	static const FName CheckLedgeDirectionName(TEXT("CheckLedgeDirection"));
-	FCollisionQueryParams CapsuleParams(CheckLedgeDirectionName, false, CharacterOwner);
+	FCollisionQueryParams CapsuleParams(CharacterMovementComponentStatics::CheckLedgeDirectionName, false, CharacterOwner);
 	FCollisionResponseParams ResponseParam;
 	InitCollisionParams(CapsuleParams, ResponseParam);
 	FCollisionShape CapsuleShape = GetPawnCapsuleCollisionShape(SHRINK_None);
@@ -4007,8 +4014,7 @@ FVector UCharacterMovementComponent::ProjectLocationFromNavMesh(float DeltaSecon
 			// raycast to underlying mesh to allow us to more closely follow geometry
 			// we use static objects here as a best approximation to accept only objects that
 			// influence navmesh generation
-			static const FName ProjectLocationName(TEXT("NavProjectLocation"));
-			FCollisionQueryParams Params(ProjectLocationName, false);
+			FCollisionQueryParams Params(CharacterMovementComponentStatics::ProjectLocationName, false);
 			FCollisionResponseParams ResponseParams(ECR_Ignore); // ignore everything
 			ResponseParams.CollisionResponse.SetResponse(ECC_WorldStatic, ECR_Block); // get blocked only by WorldStatic
 			GetWorld()->LineTraceSingleByChannel(CachedProjectedNavMeshHitResult, TraceStart, TraceEnd, ECC_WorldStatic, Params, ResponseParams);
@@ -4487,8 +4493,7 @@ bool UCharacterMovementComponent::CheckWaterJump(FVector CheckPoint, FVector& Wa
 	CheckPoint = UpdatedComponent->GetComponentLocation() + 1.2f * PawnCapsuleRadius * CheckNorm;
 	FVector Extent(PawnCapsuleRadius, PawnCapsuleRadius, PawnCapsuleHalfHeight);
 	FHitResult HitInfo(1.f);
-	static const FName CheckWaterJumpName(TEXT("CheckWaterJump"));
-	FCollisionQueryParams CapsuleParams(CheckWaterJumpName, false, CharacterOwner);
+	FCollisionQueryParams CapsuleParams(CharacterMovementComponentStatics::CheckWaterJumpName, false, CharacterOwner);
 	FCollisionResponseParams ResponseParam;
 	InitCollisionParams(CapsuleParams, ResponseParam);
 	FCollisionShape CapsuleShape = GetPawnCapsuleCollisionShape(SHRINK_None);
@@ -4502,7 +4507,7 @@ bool UCharacterMovementComponent::CheckWaterJump(FVector CheckPoint, FVector& Wa
 		FVector Start = UpdatedComponent->GetComponentLocation();
 		Start.Z += MaxOutOfWaterStepHeight;
 		CheckPoint = Start + 3.2f * PawnCapsuleRadius * WallNormal;
-		FCollisionQueryParams LineParams(CheckWaterJumpName, true, CharacterOwner);
+		FCollisionQueryParams LineParams(CharacterMovementComponentStatics::CheckWaterJumpName, true, CharacterOwner);
 		FCollisionResponseParams LineResponseParam;
 		InitCollisionParams(LineParams, LineResponseParam);
 		bHit = GetWorld()->LineTraceSingleByChannel( HitInfo, Start, CheckPoint, CollisionChannel, LineParams, LineResponseParam );
@@ -4740,9 +4745,7 @@ void UCharacterMovementComponent::ComputeFloorDist(const FVector& CapsuleLocatio
 		const float ShrinkScaleOverlap = 0.1f;
 		float ShrinkHeight = (PawnHalfHeight - PawnRadius) * (1.f - ShrinkScale);
 		float TraceDist = SweepDistance + ShrinkHeight;
-
-		static const FName ComputeFloorDistName(TEXT("ComputeFloorDistSweep"));
-		QueryParams.TraceTag = ComputeFloorDistName;
+		QueryParams.TraceTag = CharacterMovementComponentStatics::ComputeFloorDistName;
 		FCollisionShape CapsuleShape = FCollisionShape::MakeCapsule(SweepRadius, PawnHalfHeight - ShrinkHeight);
 
 		FHitResult Hit(1.f);
@@ -4797,9 +4800,7 @@ void UCharacterMovementComponent::ComputeFloorDist(const FVector& CapsuleLocatio
 		const FVector LineTraceStart = CapsuleLocation;	
 		const float TraceDist = LineDistance + ShrinkHeight;
 		const FVector Down = FVector(0.f, 0.f, -TraceDist);
-
-		static const FName FloorLineTraceName = FName(TEXT("ComputeFloorDistLineTrace"));
-		QueryParams.TraceTag = FloorLineTraceName;
+		QueryParams.TraceTag = CharacterMovementComponentStatics::FloorLineTraceName;
 
 		FHitResult Hit(1.f);
 		bBlockingHit = GetWorld()->LineTraceSingleByChannel(Hit, LineTraceStart, LineTraceStart + Down, CollisionChannel, QueryParams, ResponseParam);
diff --git a/Engine/Source/Runtime/Engine/Private/Components/PrimitiveComponent.cpp b/Engine/Source/Runtime/Engine/Private/Components/PrimitiveComponent.cpp
--- a/Engine/Source/Runtime/Engine/Private/Components/PrimitiveComponent.cpp
+++ b/Engine/Source/Runtime/Engine/Private/Components/PrimitiveComponent.cpp
@@ -34,6 +34,14 @@
 //////////////////////////////////////////////////////////////////////////
 // Globals
 
+namespace PrimitiveComponentStatics
+{
+	static const FText MobilityWarnText = LOCTEXT("InvalidMove", "move");
+	static const FName MoveComponentName(TEXT("MoveComponent"));
+	static const FName UpdateOverlapsName(TEXT("UpdateOverlaps"));
+	static const FName PrimitiveComponentInstanceDataTypeName(TEXT("PrimitiveComponentInstanceData"));
+}
+
 DEFINE_LOG_CATEGORY_STATIC(LogPrimitiveComponent, Log, All);
 
 static FAutoConsoleVariable CVarAllowCachedOverlaps(
@@ -419,8 +427,7 @@ FActorComponentInstanceData* UPrimitiveComponent::GetComponentInstanceData() con
 
 FName UPrimitiveComponent::GetComponentInstanceDataType() const
 {
-	static const FName PrimitiveComponentInstanceDataTypeName(TEXT("PrimitiveComponentInstanceData"));
-	return PrimitiveComponentInstanceDataTypeName;
+	return PrimitiveComponentStatics::PrimitiveComponentInstanceDataTypeName;
 }
 
 void UPrimitiveComponent::OnAttachmentChanged()
@@ -1345,6 +1352,7 @@ FCollisionShape UPrimitiveComponent::GetCollisionShape(float Inflation) const
 	return FCollisionShape::MakeBox(Bounds.BoxExtent + Inflation);
 }
 
+
 bool UPrimitiveComponent::MoveComponentImpl( const FVector& Delta, const FQuat& NewRotationQuat, bool bSweep, FHitResult* OutHit, EMoveComponentFlags MoveFlags)
 {
 	SCOPE_CYCLE_COUNTER(STAT_MoveComponentTime);
@@ -1369,8 +1377,7 @@ bool UPrimitiveComponent::MoveComponentImpl( const FVector& Delta, const FQuat&
 	}
 
 	// static things can move before they are registered (e.g. immediately after streaming), but not after.
-	static const FText WarnText = LOCTEXT("InvalidMove", "move");
-	if (CheckStaticMobilityAndWarn(WarnText))
+	if (CheckStaticMobilityAndWarn(PrimitiveComponentStatics::MobilityWarnText))
 	{
 		if (OutHit)
 		{
@@ -1450,10 +1457,9 @@ bool UPrimitiveComponent::MoveComponentImpl( const FVector& Delta, const FQuat&
 #if !(UE_BUILD_SHIPPING || UE_BUILD_TEST) && PERF_MOVECOMPONENT_STATS
 			MoveTimer.bDidLineCheck = true;
 #endif 
-			static const FName Name_MoveComponent(TEXT("MoveComponent"));
 			UWorld* const MyWorld = GetWorld();
 
-			FComponentQueryParams Params(Name_MoveComponent, Actor);
+			FComponentQueryParams Params(PrimitiveComponentStatics::MoveComponentName, Actor);
 			FCollisionResponseParams ResponseParam;
 			InitSweepCollisionParams(Params, ResponseParam);
 			bool const bHadBlockingHit = MyWorld->ComponentSweepMulti(Hits, this, TraceStart, TraceEnd, InitialRotationQuat, Params);
@@ -2301,9 +2307,8 @@ void UPrimitiveComponent::UpdateOverlaps(const TArray<FOverlapInfo>* NewPendingO
 					UE_LOG(LogPrimitiveComponent, VeryVerbose, TEXT("%s->%s Performing overlaps!"), *GetNameSafe(GetOwner()), *GetName());
 					UWorld* const MyWorld = MyActor->GetWorld();
 					TArray<FOverlapResult> Overlaps;
-					static FName NAME_UpdateOverlaps = FName(TEXT("UpdateOverlaps"));
 					// note this will optionally include overlaps with components in the same actor (depending on bIgnoreChildren). 
-					FComponentQueryParams Params(NAME_UpdateOverlaps, bIgnoreChildren ? MyActor : nullptr);
+					FComponentQueryParams Params(PrimitiveComponentStatics::UpdateOverlapsName, bIgnoreChildren ? MyActor : nullptr);
 					Params.bTraceAsyncScene = bCheckAsyncSceneOnMove;
 					Params.AddIgnoredActors(MoveIgnoreActors);
 					MyWorld->ComponentOverlapMulti(Overlaps, this, GetComponentLocation(), GetComponentQuat(), Params);
diff --git a/Engine/Source/Runtime/Engine/Private/Components/SceneComponent.cpp b/Engine/Source/Runtime/Engine/Private/Components/SceneComponent.cpp
--- a/Engine/Source/Runtime/Engine/Private/Components/SceneComponent.cpp
+++ b/Engine/Source/Runtime/Engine/Private/Components/SceneComponent.cpp
@@ -18,6 +18,15 @@
 #define LOCTEXT_NAMESPACE "SceneComponent"
 
 
+namespace SceneComponentStatics
+{
+	static const FName DefaultSceneRootVariableName(TEXT("DefaultSceneRoot"));
+	static const FName SceneComponentInstanceDataTypeName(TEXT("SceneComponentInstanceData"));
+	static const FName MobilityName(TEXT("Mobility"));
+	static const FText MobilityWarnText = LOCTEXT("InvalidMove", "move");
+	static const FName PhysicsVolumeTraceName(TEXT("PhysicsVolumeTrace"));
+}
+
 DEFINE_LOG_CATEGORY_STATIC(LogSceneComponent, Log, All);
 
 FOverlapInfo::FOverlapInfo(UPrimitiveComponent* InComponent, int32 InBodyIndex)
@@ -30,9 +39,7 @@ FOverlapInfo::FOverlapInfo(UPrimitiveComponent* InComponent, int32 InBodyIndex)
 
 const FName& USceneComponent::GetDefaultSceneRootVariableName()
 {
-	static FName DefaultSceneRootVariableName = FName(TEXT("DefaultSceneRoot"));
-
-	return DefaultSceneRootVariableName;
+	return SceneComponentStatics::DefaultSceneRootVariableName;
 }
 
 USceneComponent::USceneComponent(const FObjectInitializer& ObjectInitializer /*= FObjectInitializer::Get()*/)
@@ -261,8 +268,7 @@ static void UpdateAttachedMobility(USceneComponent* ComponentThatChanged)
 void USceneComponent::PostEditChangeProperty(FPropertyChangedEvent& PropertyChangedEvent)
 {
 	// Note: This must be called before UActorComponent::PostEditChangeChainProperty is called because this component will be reset when UActorComponent reruns construction scripts 
-	static FName MobilityName("Mobility");
-	if(PropertyChangedEvent.Property && PropertyChangedEvent.Property->GetFName() == MobilityName)
+	if(PropertyChangedEvent.Property && PropertyChangedEvent.Property->GetFName() == SceneComponentStatics::MobilityName)
 	{
 		UpdateAttachedMobility(this);
 	}
@@ -273,8 +279,7 @@ void USceneComponent::PostEditChangeProperty(FPropertyChangedEvent& PropertyChan
 void USceneComponent::PostEditChangeChainProperty(FPropertyChangedChainEvent& PropertyChangedEvent)
 {
 	// Note: This must be called before UActorComponent::PostEditChangeChainProperty is called because this component will be reset when UActorComponent reruns construction scripts 
-	static FName MobilityName("Mobility");
-	if(PropertyChangedEvent.Property && PropertyChangedEvent.Property->GetFName() == MobilityName)
+	if (PropertyChangedEvent.Property && PropertyChangedEvent.Property->GetFName() == SceneComponentStatics::MobilityName)
 	{
 		UpdateAttachedMobility(this);
 	}
@@ -1473,8 +1478,7 @@ FActorComponentInstanceData* USceneComponent::GetComponentInstanceData() const
 
 FName USceneComponent::GetComponentInstanceDataType() const
 {
-	static const FName SceneComponentInstanceDataTypeName(TEXT("SceneComponentInstanceData"));
-	return SceneComponentInstanceDataTypeName;
+	return SceneComponentStatics::SceneComponentInstanceDataTypeName;
 }
 
 void USceneComponent::UpdateChildTransforms(bool bSkipPhysicsMove, bool bTeleport)
@@ -1744,8 +1748,7 @@ void USceneComponent::UpdatePhysicsVolume( bool bTriggerNotifiers )
 			{
 				// check for all volumes that overlap the component
 				TArray<FOverlapResult> Hits;
-				static FName NAME_PhysicsVolumeTrace = FName(TEXT("PhysicsVolumeTrace"));
-				FComponentQueryParams Params(NAME_PhysicsVolumeTrace, GetOwner());
+				FComponentQueryParams Params(SceneComponentStatics::PhysicsVolumeTraceName, GetOwner());
 
 				bool bOverlappedOrigin = false;
 				const UPrimitiveComponent* SelfAsPrimitive = Cast<UPrimitiveComponent>(this);
@@ -1905,8 +1908,7 @@ bool USceneComponent::MoveComponentImpl( const FVector& Delta, const FQuat& NewR
 	SCOPE_CYCLE_COUNTER(STAT_MoveComponentSceneComponentTime);
 
 	// static things can move before they are registered (e.g. immediately after streaming), but not after.
-	static const FText WarnText = LOCTEXT("InvalidMove", "move");
-	if (CheckStaticMobilityAndWarn(WarnText))
+	if (CheckStaticMobilityAndWarn(SceneComponentStatics::MobilityWarnText))
 	{
 		if (OutHit)
 		{
