#ue4 - PlayerController uses a reliable RPC to request that the pawn gets resent from the server if it has not yet replicated. Previously this was an unreliable call so it didn't always work.

#1192 - Git author: ZKShao - https://github.com/EpicGames/UnrealEngine/pull/1192

diff --git a/Engine/Source/Runtime/Engine/Classes/GameFramework/PlayerController.h b/Engine/Source/Runtime/Engine/Classes/GameFramework/PlayerController.h
--- a/Engine/Source/Runtime/Engine/Classes/GameFramework/PlayerController.h
+++ b/Engine/Source/Runtime/Engine/Classes/GameFramework/PlayerController.h
@@ -950,6 +950,10 @@ public:
 	UFUNCTION(unreliable, server, WithValidation)
 	void ServerCheckClientPossession();
 
+	/** Reliable version of ServerCheckClientPossession to be used when there is no likely danger of spamming the network. */
+	UFUNCTION(reliable, server, WithValidation)
+	void ServerCheckClientPossessionReliable();
+
 	/** Call ServerCheckClientPossession on the server, but only if the current pawn is not the acknowledged pawn (and throttled to avoid saturating the network). */
 	virtual void SafeServerCheckClientPossession();
 
diff --git a/Engine/Source/Runtime/Engine/Private/PlayerController.cpp b/Engine/Source/Runtime/Engine/Private/PlayerController.cpp
--- a/Engine/Source/Runtime/Engine/Private/PlayerController.cpp
+++ b/Engine/Source/Runtime/Engine/Private/PlayerController.cpp
@@ -731,7 +731,7 @@ void APlayerController::ClientRestart_Implementation(APawn* NewPawn)
 	if ( GetPawn() == NULL )
 	{
 		// We failed to possess, ask server to verify and potentially resend the pawn
-		ServerCheckClientPossession();
+		ServerCheckClientPossessionReliable();
 		return;
 	}
 
@@ -2569,6 +2569,17 @@ void APlayerController::ServerCheckClientPossession_Implementation()
 	}
 }
 
+bool APlayerController::ServerCheckClientPossessionReliable_Validate()
+{
+	return true;
+}
+
+void APlayerController::ServerCheckClientPossessionReliable_Implementation()
+{
+	ServerCheckClientPossession_Implementation();
+}
+
+
 void APlayerController::SafeServerCheckClientPossession()
 {
 	if (GetPawn() && AcknowledgedPawn != GetPawn())
