#ue4 - Perf: Optimize weak pointer usage in UPrimitiveComponent::DispatchBlockingHit().

diff --git a/Engine/Source/Runtime/Engine/Private/Components/PrimitiveComponent.cpp b/Engine/Source/Runtime/Engine/Private/Components/PrimitiveComponent.cpp
--- a/Engine/Source/Runtime/Engine/Private/Components/PrimitiveComponent.cpp
+++ b/Engine/Source/Runtime/Engine/Private/Components/PrimitiveComponent.cpp
@@ -1684,12 +1684,16 @@ bool UPrimitiveComponent::MoveComponentImpl( const FVector& Delta, const FQuat&
 
 void UPrimitiveComponent::DispatchBlockingHit(AActor& Owner, FHitResult const& BlockingHit)
 {
-	Owner.DispatchBlockingHit(this, BlockingHit.Component.Get(), true, BlockingHit);
+	UPrimitiveComponent* const BlockingHitComponent = BlockingHit.Component.Get();
+	Owner.DispatchBlockingHit(this, BlockingHitComponent, true, BlockingHit);
 
-	// BlockingHit.GetActor() could be marked for deletion in DispatchBlockingHit(), which would make the weak pointer return NULL.
-	if ( BlockingHit.GetActor() != nullptr && BlockingHit.Component.Get() != nullptr)
+	if (!BlockingHitComponent->IsPendingKill())
 	{
-		BlockingHit.GetActor()->DispatchBlockingHit(BlockingHit.Component.Get(), this, false, BlockingHit);
+		// BlockingHit.GetActor() could be marked for deletion in DispatchBlockingHit(), which would make the weak pointer return NULL.
+		if (AActor* const BlockingHitActor = BlockingHit.GetActor())
+		{
+			BlockingHitActor->DispatchBlockingHit(BlockingHitComponent, this, false, BlockingHit);
+		}
 	}
 }
 
