Fixing build error from CL 2570924.

diff --git a/Engine/Source/Runtime/Engine/Classes/EdGraph/EdGraphPin.h b/Engine/Source/Runtime/Engine/Classes/EdGraph/EdGraphPin.h
--- a/Engine/Source/Runtime/Engine/Classes/EdGraph/EdGraphPin.h
+++ b/Engine/Source/Runtime/Engine/Classes/EdGraph/EdGraphPin.h
@@ -297,10 +297,12 @@ public:
 	/** Helper to safely set a pin's bHidden property only if it has no sub-pins that are influencing it to be hidden */
 	void SafeSetHidden(bool bIsHidden)
 	{
+#if WITH_EDITORONLY_DATA
 		 if (SubPins.Num() == 0)
 		 {
 			 bHidden = bIsHidden;
 		 }
+#endif // WITH_EDITORONLY_DATA
 	}
 
 	/** Get the current DefaultObject path name, or DefaultValue if its null */
