Fixed crash in FDebugRenderSceneProxy class.

diff --git a/Engine/Source/Runtime/Engine/Private/DebugRenderSceneProxy.cpp b/Engine/Source/Runtime/Engine/Private/DebugRenderSceneProxy.cpp
--- a/Engine/Source/Runtime/Engine/Private/DebugRenderSceneProxy.cpp
+++ b/Engine/Source/Runtime/Engine/Private/DebugRenderSceneProxy.cpp
@@ -98,14 +98,14 @@ void FDebugRenderSceneProxy::GetDynamicMeshElements(const TArray<const FSceneVie
 				{
 					
 					MeshColor = new(FMemStack::Get())  FColoredMaterialRenderProxy(
-						SolidMeshMaterial->GetRenderProxy(false, false/*IsSelected(), IsHovered()*/),
+						SolidMeshMaterial->GetRenderProxy(false, false),
 						Color,
 						"GizmoColor"
 						);
 				}
 				else
 				{
-					MeshColor = new(FMemStack::Get()) FColoredMaterialRenderProxy(GEngine->DebugMeshMaterial->GetRenderProxy(true), Color);
+					MeshColor = new(FMemStack::Get()) FColoredMaterialRenderProxy(GEngine->DebugMeshMaterial->GetRenderProxy(false, false), Color);
 				}
 
 				MeshColorInstances.Add(HashKey, MeshColor);
