fixed navigable collision export from covex elements of BodySetup not using element's local transform
#ue4 UE-13943

diff --git a/Engine/Source/Runtime/Engine/Classes/AI/Navigation/NavCollision.h b/Engine/Source/Runtime/Engine/Classes/AI/Navigation/NavCollision.h
--- a/Engine/Source/Runtime/Engine/Classes/AI/Navigation/NavCollision.h
+++ b/Engine/Source/Runtime/Engine/Classes/AI/Navigation/NavCollision.h
@@ -73,6 +73,9 @@ class ENGINE_API UNavCollision : public UObject
 	/** convex collisions are ready to use */
 	uint32 bHasConvexGeometry : 1;
 
+	/** if set, convex geometry will be rebuild instead of using cooked data */
+	uint32 bForceGeometryRebuild : 1;
+
 	/** Guid of associated BodySetup */
 	FGuid BodySetupGuid;
 
diff --git a/Engine/Source/Runtime/Engine/Private/AI/Navigation/NavCollision.cpp b/Engine/Source/Runtime/Engine/Private/AI/Navigation/NavCollision.cpp
--- a/Engine/Source/Runtime/Engine/Private/AI/Navigation/NavCollision.cpp
+++ b/Engine/Source/Runtime/Engine/Private/AI/Navigation/NavCollision.cpp
@@ -136,6 +136,8 @@ bool FDerivedDataNavCollisionCooker::Build( TArray<uint8>& OutData )
 //----------------------------------------------------------------------//
 UNavCollision::UNavCollision(const FObjectInitializer& ObjectInitializer) : Super(ObjectInitializer)
 {	
+	bHasConvexGeometry = false;
+	bForceGeometryRebuild = false;
 }
 
 FGuid UNavCollision::GetGuid() const
@@ -161,7 +163,7 @@ void UNavCollision::Setup(UBodySetup* BodySetup)
 	{
 		// Find or create cooked navcollision data
 		FByteBulkData* FormatData = GetCookedData(NAVCOLLISION_FORMAT);
-		if( FormatData )
+		if (!bForceGeometryRebuild && FormatData)
 		{
 			// if it's not being already processed
 			if (FormatData->IsLocked() == false)
@@ -343,7 +345,8 @@ void UNavCollision::Serialize(FArchive& Ar)
 
 	const int32 VerInitial = 1;
 	const int32 VerAreaClass = 2;
-	const int32 VerLatest = VerAreaClass;
+	const int32 VerConvexTransforms = 3;
+	const int32 VerLatest = VerConvexTransforms;
 
 	// use magic number to determine if serialized stream has version :/
 	const int32 MagicNum = 0xA237F237;
@@ -400,6 +403,11 @@ void UNavCollision::Serialize(FArchive& Ar)
 	{
 		Ar << AreaClass;
 	}
+
+	if (Version < VerConvexTransforms && Ar.IsLoading() && GIsEditor)
+	{
+		bForceGeometryRebuild = true;
+	}
 }
 
 void UNavCollision::PostLoad()
diff --git a/Engine/Source/Runtime/Engine/Private/AI/Navigation/RecastNavMeshGenerator.cpp b/Engine/Source/Runtime/Engine/Private/AI/Navigation/RecastNavMeshGenerator.cpp
--- a/Engine/Source/Runtime/Engine/Private/AI/Navigation/RecastNavMeshGenerator.cpp
+++ b/Engine/Source/Runtime/Engine/Private/AI/Navigation/RecastNavMeshGenerator.cpp
@@ -703,7 +703,7 @@ FORCEINLINE_DEBUGGABLE void ExportRigidBodyConvexElements(UBodySetup& BodySetup,
 		ShapeBuffer.Add(VertexBuffer.Num() / 3);
 
 		// Get verts/triangles from this hull.
-		ExportPxConvexMesh(ConvexElem->ConvexMesh, LocalToWorld, VertexBuffer, IndexBuffer, UnrealBounds);
+		ExportPxConvexMesh(ConvexElem->ConvexMesh, ConvexElem->Transform * LocalToWorld, VertexBuffer, IndexBuffer, UnrealBounds);
 	}
 #endif // WITH_PHYSX
 }
