fixed navigation relevancy and update order of custom navigation links in NavLinkProxy actor
#ue4 UE-16339

diff --git a/Engine/Source/Runtime/Engine/Private/AI/Navigation/NavLinkProxy.cpp b/Engine/Source/Runtime/Engine/Private/AI/Navigation/NavLinkProxy.cpp
--- a/Engine/Source/Runtime/Engine/Private/AI/Navigation/NavLinkProxy.cpp
+++ b/Engine/Source/Runtime/Engine/Private/AI/Navigation/NavLinkProxy.cpp
@@ -55,16 +55,26 @@ ANavLinkProxy::ANavLinkProxy(const FObjectInitializer& ObjectInitializer) : Supe
 #if WITH_EDITOR
 void ANavLinkProxy::PostEditChangeProperty(FPropertyChangedEvent& PropertyChangedEvent) 
 {
-	Super::PostEditChangeProperty(PropertyChangedEvent);
-
-	if (PropertyChangedEvent.Property && PropertyChangedEvent.Property->GetFName() == GET_MEMBER_NAME_CHECKED(ANavLinkProxy,bSmartLinkIsRelevant))
+	bool bUpdateInNavOctree = false;
+	if (PropertyChangedEvent.Property && PropertyChangedEvent.Property->GetFName() == GET_MEMBER_NAME_CHECKED(ANavLinkProxy, bSmartLinkIsRelevant))
 	{
 		SmartLinkComp->SetNavigationRelevancy(bSmartLinkIsRelevant);
+		if (bSmartLinkIsRelevant)
+		{
+			SmartLinkComp->ForceNavigationRelevancy(true);
+		}
+
+		bUpdateInNavOctree = true;
 	}
 
 	const FName CategoryName = FObjectEditorUtils::GetCategoryFName(PropertyChangedEvent.Property);
 	const FName MemberCategoryName = FObjectEditorUtils::GetCategoryFName(PropertyChangedEvent.MemberProperty);
 	if (CategoryName == TEXT("SimpleLink") || MemberCategoryName == TEXT("SimpleLink"))
+	{
+		bUpdateInNavOctree = true;
+	}
+
+	if (bUpdateInNavOctree)
 	{
 		UNavigationSystem* NavSys = UNavigationSystem::GetCurrent(GetWorld());
 		if (NavSys)
@@ -72,6 +82,8 @@ void ANavLinkProxy::PostEditChangeProperty(FPropertyChangedEvent& PropertyChange
 			NavSys->UpdateNavOctree(this);
 		}
 	}
+
+	Super::PostEditChangeProperty(PropertyChangedEvent);
 }
 #endif // WITH_EDITOR
 
