Updating asset import paths now attempts to reuse file timestamp/md5 information

diff --git a/Engine/Source/Runtime/Engine/Private/EditorFramework/AssetImportData.cpp b/Engine/Source/Runtime/Engine/Private/EditorFramework/AssetImportData.cpp
--- a/Engine/Source/Runtime/Engine/Private/EditorFramework/AssetImportData.cpp
+++ b/Engine/Source/Runtime/Engine/Private/EditorFramework/AssetImportData.cpp
@@ -86,8 +86,16 @@ TOptional<FAssetImportInfo> FAssetImportInfo::FromJson(FString InJsonString)
 
 void UAssetImportData::UpdateFilenameOnly(const FString& InPath)
 {
-	SourceFiles.Reset();
-	SourceFiles.Emplace(SanitizeImportFilename(InPath));
+	// Try and retain the MD5 and timestamp if possible
+	if (SourceFiles.Num() == 1)
+	{
+		SourceFiles[0].RelativeFilename = SanitizeImportFilename(InPath);
+	}
+	else
+	{
+		SourceFiles.Reset();
+		SourceFiles.Emplace(SanitizeImportFilename(InPath));
+	}
 }
 
 void UAssetImportData::Update(const FString& InPath)
