custom navigation link registration can now queue for processing before navigation system is created
#ue4 UE-16339

diff --git a/Engine/Source/Runtime/Engine/Classes/AI/Navigation/NavigationSystem.h b/Engine/Source/Runtime/Engine/Classes/AI/Navigation/NavigationSystem.h
--- a/Engine/Source/Runtime/Engine/Classes/AI/Navigation/NavigationSystem.h
+++ b/Engine/Source/Runtime/Engine/Classes/AI/Navigation/NavigationSystem.h
@@ -566,6 +566,9 @@ protected:
 	/** registers NavArea classes awaiting registration in PendingNavAreaRegistration */
 	void ProcessNavAreaPendingRegistration();
 
+	/** registers CustomLinks awaiting registration in PendingCustomLinkRegistration */
+	void ProcessCustomLinkPendingRegistration();
+
 	/** used to apply updates of nav volumes in navigation system's tick */
 	void PerformNavigationBoundsUpdate(const TArray<FNavigationBoundsUpdateRequest>& UpdateRequests);
 	
@@ -628,6 +631,9 @@ public:
 	void RegisterCustomLink(INavLinkCustomInterface& CustomLink);
 	void UnregisterCustomLink(INavLinkCustomInterface& CustomLink);
 	
+	static void RequestCustomLinkRegistering(INavLinkCustomInterface& CustomLink, UObject* OwnerOb);
+	static void RequestCustomLinkUnregistering(INavLinkCustomInterface& CustomLink, UObject* ObjectOb);
+
 	/** find custom link by unique ID */
 	INavLinkCustomInterface* GetCustomLink(uint32 UniqueLinkId) const;
 
@@ -799,6 +805,7 @@ protected:
 	// async queries
 	FCriticalSection NavDataRegistrationSection;
 	static FCriticalSection NavAreaRegistrationSection;
+	static FCriticalSection CustomLinkRegistrationSection;
 	
 #if WITH_EDITOR
 	uint8 NavUpdateLockFlags;
@@ -831,6 +838,7 @@ protected:
 
 	static bool bUpdateNavOctreeOnComponentChange;
 
+	static TMap<INavLinkCustomInterface*, FWeakObjectPtr> PendingCustomLinkRegistration;
 	static TArray<UClass*> PendingNavAreaRegistration;
 	static TArray<const UClass*> NavAreaClasses;
 	static TSubclassOf<UNavArea> DefaultWalkableArea;
diff --git a/Engine/Source/Runtime/Engine/Private/AI/Navigation/NavLinkCustomComponent.cpp b/Engine/Source/Runtime/Engine/Private/AI/Navigation/NavLinkCustomComponent.cpp
--- a/Engine/Source/Runtime/Engine/Private/AI/Navigation/NavLinkCustomComponent.cpp
+++ b/Engine/Source/Runtime/Engine/Private/AI/Navigation/NavLinkCustomComponent.cpp
@@ -120,27 +120,14 @@ void UNavLinkCustomComponent::OnRegister()
 		NavLinkUserId = INavLinkCustomInterface::GetUniqueId();
 	}
 
-	UNavigationSystem* NavSys = UNavigationSystem::GetCurrent(GetWorld());
-	if (NavSys == NULL)
-	{
-		return;
-	}
-
-	NavSys->RegisterCustomLink(*this);
+	UNavigationSystem::RequestCustomLinkRegistering(*this, this);
 }
 
 void UNavLinkCustomComponent::OnUnregister()
 {
 	Super::OnUnregister();
 
-	UNavigationSystem* NavSys = UNavigationSystem::GetCurrent(GetWorld());
-	if (NavSys == NULL)
-	{
-		return;
-	}
-
-	// always try to unregister, even if not relevant right now
-	NavSys->UnregisterCustomLink(*this);
+	UNavigationSystem::RequestCustomLinkUnregistering(*this, this);
 }
 
 void UNavLinkCustomComponent::SetLinkData(const FVector& RelativeStart, const FVector& RelativeEnd, ENavLinkDirection::Type Direction)
diff --git a/Engine/Source/Runtime/Engine/Private/AI/Navigation/NavigationSystem.cpp b/Engine/Source/Runtime/Engine/Private/AI/Navigation/NavigationSystem.cpp
--- a/Engine/Source/Runtime/Engine/Private/AI/Navigation/NavigationSystem.cpp
+++ b/Engine/Source/Runtime/Engine/Private/AI/Navigation/NavigationSystem.cpp
@@ -232,9 +232,11 @@ void FNavigationLockContext::UnlockUpdates()
 // UNavigationSystem                                                                
 //----------------------------------------------------------------------//
 bool UNavigationSystem::bNavigationAutoUpdateEnabled = true;
+TMap<INavLinkCustomInterface*, FWeakObjectPtr> UNavigationSystem::PendingCustomLinkRegistration;
 TArray<const UClass*> UNavigationSystem::NavAreaClasses;
 TArray<UClass*> UNavigationSystem::PendingNavAreaRegistration;
 FCriticalSection UNavigationSystem::NavAreaRegistrationSection;
+FCriticalSection UNavigationSystem::CustomLinkRegistrationSection;
 TSubclassOf<UNavArea> UNavigationSystem::DefaultWalkableArea = NULL;
 TSubclassOf<UNavArea> UNavigationSystem::DefaultObstacleArea = NULL;
 #if !(UE_BUILD_SHIPPING || UE_BUILD_TEST)
@@ -315,7 +317,6 @@ void UNavigationSystem::DoInitialSetup()
 	}
 	
 	UpdateAbstractNavData();
-
 	CreateCrowdManager();
 
 	bInitialSetupHasBeenPerformed = true;
@@ -722,6 +723,11 @@ void UNavigationSystem::Tick(float DeltaSeconds)
 		SCOPE_CYCLE_COUNTER(STAT_Navigation_TickNavAreaRegister);
 		ProcessNavAreaPendingRegistration();
 	}
+	
+	if (PendingCustomLinkRegistration.Num())
+	{
+		ProcessCustomLinkPendingRegistration();
+	}
 
 	if (PendingNavBoundsUpdates.Num() > 0)
 	{
@@ -1703,6 +1709,25 @@ void UNavigationSystem::ProcessNavAreaPendingRegistration()
 	}
 }
 
+void UNavigationSystem::ProcessCustomLinkPendingRegistration()
+{
+	FScopeLock AccessLock(&CustomLinkRegistrationSection);
+
+	TMap<INavLinkCustomInterface*, FWeakObjectPtr> TempPending = PendingCustomLinkRegistration;
+	PendingCustomLinkRegistration.Empty();
+
+	for (TMap<INavLinkCustomInterface*, FWeakObjectPtr>::TIterator It(TempPending); It; ++It)
+	{
+		INavLinkCustomInterface* ILink = It.Key();
+		FWeakObjectPtr LinkOb = It.Value();
+		
+		if (LinkOb.IsValid() && ILink)
+		{
+			RegisterCustomLink(*ILink);
+		}
+	}
+}
+
 UNavigationSystem::ERegistrationResult UNavigationSystem::RegisterNavData(ANavigationData* NavData)
 {
 	if (NavData == NULL)
@@ -1836,6 +1861,34 @@ void UNavigationSystem::UpdateCustomLink(const INavLinkCustomInterface* CustomLi
 	}
 }
 
+void UNavigationSystem::RequestCustomLinkRegistering(INavLinkCustomInterface& CustomLink, UObject* OwnerOb)
+{
+	UNavigationSystem* NavSys = UNavigationSystem::GetCurrent(OwnerOb);
+	if (NavSys)
+	{
+		NavSys->RegisterCustomLink(CustomLink);
+	}
+	else
+	{
+		FScopeLock AccessLock(&CustomLinkRegistrationSection);
+		PendingCustomLinkRegistration.Add(&CustomLink, OwnerOb);
+	}
+}
+
+void UNavigationSystem::RequestCustomLinkUnregistering(INavLinkCustomInterface& CustomLink, UObject* OwnerOb)
+{
+	UNavigationSystem* NavSys = UNavigationSystem::GetCurrent(OwnerOb);
+	if (NavSys)
+	{
+		NavSys->UnregisterCustomLink(CustomLink);
+	}
+	else
+	{
+		FScopeLock AccessLock(&CustomLinkRegistrationSection);
+		PendingCustomLinkRegistration.Remove(&CustomLink);
+	}
+}
+
 void UNavigationSystem::RequestAreaUnregistering(UClass* NavAreaClass)
 {
 	if (NavAreaClasses.Contains(NavAreaClass))
