added custom navigation export for skeletal meshes
#ue4 UE-11969

diff --git a/Engine/Source/Runtime/Engine/Classes/Components/SkeletalMeshComponent.h b/Engine/Source/Runtime/Engine/Classes/Components/SkeletalMeshComponent.h
--- a/Engine/Source/Runtime/Engine/Classes/Components/SkeletalMeshComponent.h
+++ b/Engine/Source/Runtime/Engine/Classes/Components/SkeletalMeshComponent.h
@@ -882,6 +882,7 @@ public:
 	virtual void SetAllMassScale(float InMassScale = 1.f) override;
 	virtual float GetMass() const override;
 	virtual float CalculateMass(FName BoneName = NAME_None) override;
+	virtual bool DoCustomNavigableGeometryExport(FNavigableGeometryExport& GeomExport) const override;
 #if WITH_EDITOR
 	virtual bool ComponentIsTouchingSelectionBox(const FBox& InSelBBox, const FEngineShowFlags& ShowFlags, const bool bConsiderOnlyBSP, const bool bMustEncompassEntireComponent) const override;
 	virtual bool ComponentIsTouchingSelectionFrustum(const FConvexVolume& InFrustum, const FEngineShowFlags& ShowFlags, const bool bConsiderOnlyBSP, const bool bMustEncompassEntireComponent) const override;
diff --git a/Engine/Source/Runtime/Engine/Private/Components/SkeletalMeshComponent.cpp b/Engine/Source/Runtime/Engine/Private/Components/SkeletalMeshComponent.cpp
--- a/Engine/Source/Runtime/Engine/Private/Components/SkeletalMeshComponent.cpp
+++ b/Engine/Source/Runtime/Engine/Private/Components/SkeletalMeshComponent.cpp
@@ -19,6 +19,7 @@
 #include "Particles/ParticleSystemComponent.h"
 #include "Animation/AnimSingleNodeInstance.h"
 #include "Engine/SkeletalMeshSocket.h"
+#include "AI/NavigationSystemHelpers.h"
 #if WITH_EDITOR
 #include "ShowFlags.h"
 #include "Collision.h"
@@ -150,6 +151,8 @@ USkeletalMeshComponent::USkeletalMeshComponent(const FObjectInitializer& ObjectI
 	bEnableUpdateRateOptimizations = false;
 	RagdollAggregateThreshold = UPhysicsSettings::Get()->RagdollAggregateThreshold;
 
+	bHasCustomNavigableGeometry = EHasCustomNavigableGeometry::Yes;
+
 	bTickInEditor = true;
 }
 
@@ -2265,3 +2268,29 @@ void USkeletalMeshComponent::UnbindClothFromMasterPoseComponent(bool bRestoreSim
 	}
 #endif
 }
+
+bool USkeletalMeshComponent::DoCustomNavigableGeometryExport(FNavigableGeometryExport& GeomExport) const
+{
+	UPhysicsAsset* PhysicsAsset = GetPhysicsAsset();
+	if (PhysicsAsset && ComponentToWorld.GetScale3D().IsUniform())
+	{
+		const int32 MaxBodies = PhysicsAsset->BodySetup.Num();
+		for (int32 Idx = 0; Idx < MaxBodies; Idx++)
+		{
+			UBodySetup* BodySetup = PhysicsAsset->BodySetup[Idx];
+			int32 BoneIndex = BodySetup ? GetBoneIndex(BodySetup->BoneName) : INDEX_NONE;
+
+			if (BoneIndex != INDEX_NONE)
+			{
+				FTransform WorldBoneTransform = GetBoneTransform(BoneIndex, ComponentToWorld);
+				if (FMath::Abs(WorldBoneTransform.GetDeterminant()) > (float)KINDA_SMALL_NUMBER)
+				{
+					GeomExport.ExportRigidBodySetup(*BodySetup, WorldBoneTransform);
+				}
+			}
+		}
+	}
+
+	// skip fallback export of body setup data
+	return false;
+}
