Making sure ResetLoaders does not get called when renaming during loading.

diff --git a/Engine/Source/Runtime/CoreUObject/Private/UObject/UObjectGlobals.cpp b/Engine/Source/Runtime/CoreUObject/Private/UObject/UObjectGlobals.cpp
--- a/Engine/Source/Runtime/CoreUObject/Private/UObject/UObjectGlobals.cpp
+++ b/Engine/Source/Runtime/CoreUObject/Private/UObject/UObjectGlobals.cpp
@@ -1479,7 +1479,7 @@ FObjectDuplicationParameters::FObjectDuplicationParameters( UObject* InSourceObj
 
 UObject* StaticDuplicateObject(UObject const* SourceObject, UObject* DestOuter, const TCHAR* DestName, EObjectFlags FlagMask, UClass* DestClass, EDuplicateForPie DuplicateForPIE)
 {
-	if (!IsAsyncLoading()  && SourceObject->HasAnyFlags(RF_ClassDefaultObject))
+	if (!IsAsyncLoading() && !IsLoading() && SourceObject->HasAnyFlags(RF_ClassDefaultObject))
 	{
 		// Detach linker for the outer if it already exists, to avoid problems with PostLoad checking the Linker version
 		ResetLoaders(DestOuter);
diff --git a/Engine/Source/Runtime/Engine/Private/Components/ChildActorComponent.cpp b/Engine/Source/Runtime/Engine/Private/Components/ChildActorComponent.cpp
--- a/Engine/Source/Runtime/Engine/Private/Components/ChildActorComponent.cpp
+++ b/Engine/Source/Runtime/Engine/Private/Components/ChildActorComponent.cpp
@@ -178,7 +178,7 @@ void UChildActorComponent::ApplyComponentInstanceData(FChildActorComponentInstan
 			const FString ChildActorNameString = ChildActorName.ToString();
 			if (ChildActor->Rename(*ChildActorNameString, nullptr, REN_Test))
 			{
-				ChildActor->Rename(*ChildActorNameString, nullptr, REN_DoNotDirty);
+				ChildActor->Rename(*ChildActorNameString, nullptr, REN_DoNotDirty | (IsLoading() ? REN_ForceNoResetLoaders : REN_None));
 			}
 		}
 
@@ -313,7 +313,7 @@ void UChildActorComponent::DestroyChildActor(const bool bRequiresRename)
 				if (bRequiresRename)
 				{
 					const FString ObjectBaseName = FString::Printf(TEXT("DESTROYED_%s_CHILDACTOR"), *ChildClass->GetName());
-					const ERenameFlags RenameFlags = (GetWorld()->IsGameWorld() ? REN_DoNotDirty | REN_ForceNoResetLoaders : REN_DoNotDirty);
+					const ERenameFlags RenameFlags = ((GetWorld()->IsGameWorld() || IsLoading()) ? REN_DoNotDirty | REN_ForceNoResetLoaders : REN_DoNotDirty);
 					ChildActor->Rename(*MakeUniqueObjectName(ChildActor->GetOuter(), ChildClass, *ObjectBaseName).ToString(), nullptr, RenameFlags);
 				}
 				World->DestroyActor(ChildActor);
