Fix WakeRigidBody call on destructible causing a crash for deferred apex actors.

#UE-16447

diff --git a/Engine/Source/Runtime/Engine/Classes/Components/DestructibleComponent.h b/Engine/Source/Runtime/Engine/Classes/Components/DestructibleComponent.h
--- a/Engine/Source/Runtime/Engine/Classes/Components/DestructibleComponent.h
+++ b/Engine/Source/Runtime/Engine/Classes/Components/DestructibleComponent.h
@@ -135,6 +135,8 @@ public:
 	virtual bool SweepComponent( FHitResult& OutHit, const FVector Start, const FVector End, const FCollisionShape& CollisionShape, bool bTraceComplex=false) override;
 	virtual void SetEnableGravity(bool bGravityEnabled) override;
 
+	virtual void WakeRigidBody(FName BoneName /* = NAME_None */) override;
+
 	virtual void SetMaterial(int32 ElementIndex, UMaterialInterface* Material) override;
 	// End UPrimitiveComponent interface.
 
diff --git a/Engine/Source/Runtime/Engine/Classes/Components/PrimitiveComponent.h b/Engine/Source/Runtime/Engine/Classes/Components/PrimitiveComponent.h
--- a/Engine/Source/Runtime/Engine/Classes/Components/PrimitiveComponent.h
+++ b/Engine/Source/Runtime/Engine/Classes/Components/PrimitiveComponent.h
@@ -908,7 +908,7 @@ public:
 	 *	@param	BoneName	If a SkeletalMeshComponent, name of body to wake. 'None' indicates root body.
 	 */
 	UFUNCTION(BlueprintCallable, Category="Physics")
-	void WakeRigidBody(FName BoneName = NAME_None);
+	virtual void WakeRigidBody(FName BoneName = NAME_None);
 
 	/** 
 	 *	Force a single body back to sleep. 
diff --git a/Engine/Source/Runtime/Engine/Private/Components/DestructibleComponent.cpp b/Engine/Source/Runtime/Engine/Private/Components/DestructibleComponent.cpp
--- a/Engine/Source/Runtime/Engine/Private/Components/DestructibleComponent.cpp
+++ b/Engine/Source/Runtime/Engine/Private/Components/DestructibleComponent.cpp
@@ -1260,6 +1260,17 @@ void UDestructibleComponent::ResetFakeBodyInstance( FFakeBodyInstanceState& Prev
 
 #endif
 
+void UDestructibleComponent::WakeRigidBody(FName BoneName /* = NAME_None */)
+{
+#if WITH_APEX
+	ExecuteOnPhysicsReadWrite([&]
+	{
+		const int32 ChunkIdx = BoneIdxToChunkIdx(GetBoneIndex(BoneName));
+		ApexDestructibleActor->setChunkPhysXActorAwakeState(ChunkIdx, true);
+	});
+#endif
+}
+
 void UDestructibleComponent::SetEnableGravity(bool bGravityEnabled)
 {
 	Super::SetEnableGravity(bGravityEnabled);
