remove the function that isn't defined.

diff --git a/Engine/Source/Runtime/Engine/Classes/Components/SkinnedMeshComponent.h b/Engine/Source/Runtime/Engine/Classes/Components/SkinnedMeshComponent.h
--- a/Engine/Source/Runtime/Engine/Classes/Components/SkinnedMeshComponent.h
+++ b/Engine/Source/Runtime/Engine/Classes/Components/SkinnedMeshComponent.h
@@ -799,18 +799,6 @@ public:
 	 */
 	virtual class UMorphTarget* FindMorphTarget( FName MorphTargetName ) const;
 
-	/**	
-	 * Find all bones by name within given radius 
-	 * 
-	 * @param Origin origin of the radius
-	 * @param Radius radius of the sphere
-	 * @param TraceFlags Trace flags
-	 * @param (out) OutBones Bones within the sphere
-	 *
-	 * @return true if any bone is within the radius. false otherwise
-	 */
-	bool GetBonesWithinRadius( FVector Origin, float Radius, int32 TraceFlags, TArray<FName>& OutBones );
-
 	/**
 	 *	Hides the specified bone. You can also set option for physics body.
 	 *
