[INTEGRATE] Change 2573216 by [REDACTED_PERSON][REDACTED_HANDLE] on 2015/06/01 17:49:33

	Allow dynamically spawned particle and decal components to be destroyed

diff --git a/Engine/Source/Runtime/Engine/Classes/Components/ActorComponent.h b/Engine/Source/Runtime/Engine/Classes/Components/ActorComponent.h
--- a/Engine/Source/Runtime/Engine/Classes/Components/ActorComponent.h
+++ b/Engine/Source/Runtime/Engine/Classes/Components/ActorComponent.h
@@ -100,6 +100,9 @@ public:
 	/** Can we tick this concurrently on other threads? */
 	uint32 bAllowConcurrentTick:1;
 
+	/** Can this component be destroyed (via K2_DestroyComponent) by any parent */
+	uint32 bAllowAnyoneToDestroyMe:1;
+
 	/** True if this component was created by a construction script, and will be destroyed by DestroyConstructedComponents */
 	UPROPERTY()
 	uint32 bCreatedByConstructionScript_DEPRECATED:1;
diff --git a/Engine/Source/Runtime/Engine/Private/Components/ActorComponent.cpp b/Engine/Source/Runtime/Engine/Private/Components/ActorComponent.cpp
--- a/Engine/Source/Runtime/Engine/Private/Components/ActorComponent.cpp
+++ b/Engine/Source/Runtime/Engine/Private/Components/ActorComponent.cpp
@@ -938,7 +938,7 @@ void UActorComponent::OnComponentDestroyed()
 void UActorComponent::K2_DestroyComponent(UObject* Object)
 {
 	AActor* MyOwner = GetOwner();
-	if (Object == this || MyOwner == NULL || MyOwner == Object)
+	if (bAllowAnyoneToDestroyMe || Object == this || MyOwner == NULL || MyOwner == Object)
 	{
 		DestroyComponent();
 	}
diff --git a/Engine/Source/Runtime/Engine/Private/GameplayStatics.cpp b/Engine/Source/Runtime/Engine/Private/GameplayStatics.cpp
--- a/Engine/Source/Runtime/Engine/Private/GameplayStatics.cpp
+++ b/Engine/Source/Runtime/Engine/Private/GameplayStatics.cpp
@@ -542,6 +542,7 @@ UParticleSystemComponent* CreateParticleSystem(UParticleSystem* EmitterTemplate,
 {
 	UParticleSystemComponent* PSC = NewObject<UParticleSystemComponent>((Actor ? Actor : (UObject*)World));
 	PSC->bAutoDestroy = bAutoDestroy;
+	PSC->bAllowAnyoneToDestroyMe = true;
 	PSC->SecondsBeforeInactive = 0.0f;
 	PSC->bAutoActivate = false;
 	PSC->SetTemplate(EmitterTemplate);
@@ -1006,6 +1007,7 @@ UDecalComponent* CreateDecalComponent(class UMaterialInterface* DecalMaterial, F
 	const FMatrix DecalInternalTransform = FRotationMatrix(FRotator(0.f, 90.0f, -90.0f));
 
 	UDecalComponent* DecalComp = NewObject<UDecalComponent>((Actor ? Actor : (UObject*)World));
+	DecalComp->bAllowAnyoneToDestroyMe = true;
 	DecalComp->DecalMaterial = DecalMaterial;
 	DecalComp->RelativeScale3D = DecalInternalTransform.TransformVector(DecalSize);
 	DecalComp->bAbsoluteScale = true;
