UKismetStringLibrary::ReplaceInline works properly. It's not Pure anymore. The input param is UPARAM(ref)

diff --git a/Engine/Source/Runtime/Engine/Classes/Kismet/KismetStringLibrary.h b/Engine/Source/Runtime/Engine/Classes/Kismet/KismetStringLibrary.h
--- a/Engine/Source/Runtime/Engine/Classes/Kismet/KismetStringLibrary.h
+++ b/Engine/Source/Runtime/Engine/Classes/Kismet/KismetStringLibrary.h
@@ -389,7 +389,7 @@ class ENGINE_API UKismetStringLibrary : public UBlueprintFunctionLibrary
 	 * @return a copy of this string with the replacement made
 	 */
 	UFUNCTION(BlueprintPure, Category = "Utilities|String")
-		static FString Replace(const FString& SourceString, const FString& From, const FString& To, ESearchCase::Type SearchCase = ESearchCase::IgnoreCase);
+	static FString Replace(const FString& SourceString, const FString& From, const FString& To, ESearchCase::Type SearchCase = ESearchCase::IgnoreCase);
 
 	/**
 	 * Replace all occurrences of SearchText with ReplacementText in this string.
@@ -400,8 +400,8 @@ class ENGINE_API UKismetStringLibrary : public UBlueprintFunctionLibrary
 	 *
 	 * @return	the number of occurrences of SearchText that were replaced.
 	 */
-	UFUNCTION(BlueprintPure, Category = "Utilities|String")
-		static int32 ReplaceInline(FString& SourceString, const FString& SearchText, const FString& ReplacementText, ESearchCase::Type SearchCase = ESearchCase::IgnoreCase);
+	UFUNCTION(BlueprintCallable, Category = "Utilities|String")
+	static int32 ReplaceInline(UPARAM(ref) FString& SourceString, const FString& SearchText, const FString& ReplacementText, ESearchCase::Type SearchCase = ESearchCase::IgnoreCase);
 
 	/**
 	* Splits this string at given string position case sensitive.
