Additional fix for CL 2568322.

diff --git a/Engine/Source/Runtime/Engine/Private/Model.cpp b/Engine/Source/Runtime/Engine/Private/Model.cpp
--- a/Engine/Source/Runtime/Engine/Private/Model.cpp
+++ b/Engine/Source/Runtime/Engine/Private/Model.cpp
@@ -332,6 +332,15 @@ void UModel::PostLoad()
 			CurSurf.bHiddenEdTemporary = ( ( CurSurf.PolyFlags & PF_HiddenEd ) != 0 );
 			CurSurf.bHiddenEdLevel = 0;
 		}
+
+		if (ABrush* Owner = Cast<ABrush>(GetOuter()))
+		{
+			OwnerLocationWhenLastBuilt = Owner->GetActorLocation();
+			OwnerPrepivotWhenLastBuilt = Owner->GetPrePivot();
+			OwnerScaleWhenLastBuilt = Owner->GetActorScale();
+			OwnerRotationWhenLastBuilt = -Owner->GetActorRotation();
+			bCachedOwnerTransformValid = true;
+		}
 	}
 }
 
