Fixed non-editor builds.

diff --git a/Engine/Source/Runtime/Engine/Private/Model.cpp b/Engine/Source/Runtime/Engine/Private/Model.cpp
--- a/Engine/Source/Runtime/Engine/Private/Model.cpp
+++ b/Engine/Source/Runtime/Engine/Private/Model.cpp
@@ -333,6 +333,7 @@ void UModel::PostLoad()
 			CurSurf.bHiddenEdLevel = 0;
 		}
 
+#if WITH_EDITOR
 		if (ABrush* Owner = Cast<ABrush>(GetOuter()))
 		{
 			OwnerLocationWhenLastBuilt = Owner->GetActorLocation();
@@ -341,6 +342,7 @@ void UModel::PostLoad()
 			OwnerRotationWhenLastBuilt = -Owner->GetActorRotation();
 			bCachedOwnerTransformValid = true;
 		}
+#endif
 	}
 }
 
