Don't process for seamless travel/route end play an Actor that was destroyed by another Actor's EndPlay.

diff --git a/Engine/Source/Runtime/Engine/Private/UnrealEngine.cpp b/Engine/Source/Runtime/Engine/Private/UnrealEngine.cpp
--- a/Engine/Source/Runtime/Engine/Private/UnrealEngine.cpp
+++ b/Engine/Source/Runtime/Engine/Private/UnrealEngine.cpp
@@ -8902,7 +8902,11 @@ bool UEngine::LoadMap( FWorldContext& WorldContext, FURL URL, class UPendingNetG
 
 		for (FActorIterator ActorIt(WorldContext.World()); ActorIt; ++ActorIt)
 		{
-			ActorIt->RouteEndPlay(EEndPlayReason::LevelTransition);
+			// Routing the end play of an earlier Actor could have resulted in this Actor already being cleaned up/destroyed
+			if (!ActorIt->IsPendingKill())
+			{
+				ActorIt->RouteEndPlay(EEndPlayReason::LevelTransition);
+			}
 		}
 
 		// Do this after destroying pawns/playercontrollers, in case that spawns new things (e.g. dropped weapons)
diff --git a/Engine/Source/Runtime/Engine/Private/World.cpp b/Engine/Source/Runtime/Engine/Private/World.cpp
--- a/Engine/Source/Runtime/Engine/Private/World.cpp
+++ b/Engine/Source/Runtime/Engine/Private/World.cpp
@@ -4670,6 +4670,13 @@ UWorld* FSeamlessTravelHandler::Tick()
 			for (FActorIterator It(CurrentWorld); It; ++It)
 			{
 				AActor* TheActor = *It;
+
+				// Routing the end play of an earlier Actor could have resulted in this Actor already being cleaned up/destroyed
+				if (TheActor->IsPendingKill())
+				{
+					continue;
+				}
+
 				bool bSameLevel = TheActor->GetLevel() == CurrentWorld->PersistentLevel;
 				bool bShouldKeep = KeepAnnotation.Get(TheActor);
 				bool bDormant = false;
