Defer apex damage events until fetchResults is done. This is needed in case DestroyComponent is called from inside the callback

#UE-16462

diff --git a/Engine/Source/Runtime/Engine/Private/PhysicsEngine/PhysScene.cpp b/Engine/Source/Runtime/Engine/Private/PhysicsEngine/PhysScene.cpp
--- a/Engine/Source/Runtime/Engine/Private/PhysicsEngine/PhysScene.cpp
+++ b/Engine/Source/Runtime/Engine/Private/PhysicsEngine/PhysScene.cpp
@@ -54,9 +54,41 @@ FORCEINLINE static bool FrameLagAsync()
 	return true;
 }
 
+#if WITH_APEX
+
+//This level of indirection is needed because we don't want to expose NxApexDamageEventReportData in a public engine header
+struct FPendingApexDamageEvent
+{
+	TWeakObjectPtr<class UDestructibleComponent> DestructibleComponent;
+	NxApexDamageEventReportData DamageEvent;
+	TArray<NxApexChunkData> ApexChunkData;
+
+	FPendingApexDamageEvent(UDestructibleComponent* InDestructibleComponent, const NxApexDamageEventReportData& InDamageEvent)
+		: DestructibleComponent(InDestructibleComponent)
+		, DamageEvent(InDamageEvent)
+	{
+		ApexChunkData.AddUninitialized(InDamageEvent.fractureEventListSize);
+		for(uint32 ChunkIdx = 0; ChunkIdx < InDamageEvent.fractureEventListSize; ++ChunkIdx)
+		{
+			ApexChunkData[ChunkIdx] = InDamageEvent.fractureEventList[ChunkIdx];
+		}
+
+		DamageEvent.fractureEventList = ApexChunkData.GetData();
+	}
+
+};
+
+struct FPendingApexDamageManager
+{
+	TArray<FPendingApexDamageEvent> PendingDamageEvents;
+};
+#endif
 
 /** Exposes creation of physics-engine scene outside Engine (for use with PhAT for example). */
 FPhysScene::FPhysScene()
+#if WITH_APEX
+	: PendingApexDamageManager(new FPendingApexDamageManager)
+#endif
 {
 	LineBatcher = NULL;
 	OwningWorld = NULL;
@@ -366,10 +398,10 @@ void FPhysScene::TermBody_AssumesLocked(FBodyInstance* BodyInstance)
 #if WITH_SUBSTEPPING
 
 #if WITH_APEX
-void FPhysScene::DeferredDestructibleDamageNotify(const NxApexDamageEventReportData& damageEvent)
+void FPhysScene::AddPendingDamageEvent(UDestructibleComponent* DestructibleComponent, const NxApexDamageEventReportData& DamageEvent)
 {
-	//TODO: waiting for support on this from NVIDIA
-	//DestructibleDamageEventQueue.Add(damageEvent);
+	check(IsInGameThread());
+	FPendingApexDamageEvent* Pending = new (PendingApexDamageManager->PendingDamageEvents) FPendingApexDamageEvent(DestructibleComponent, DamageEvent);
 }
 #endif
 
@@ -852,23 +884,19 @@ void FPhysScene::DispatchPhysNotifications_AssumesLocked()
 		PendingCollisionNotifies.Empty();
 	}
 
-#if WITH_SUBSTEPPING
 #if WITH_APEX
-	//TODO: Queue is always empty for now - waiting for support from NVIDIA
-	//Destructible notification
+	for (const FPendingApexDamageEvent& PendingDamageEvent : PendingApexDamageManager->PendingDamageEvents)
 	{
-		for (int32 i = 0; i < DestructibleDamageEventQueue.Num(); ++i)
+		if(UDestructibleComponent* DestructibleComponent = PendingDamageEvent.DestructibleComponent.Get())
 		{
-			const NxApexDamageEventReportData & damageEvent = DestructibleDamageEventQueue[i];
-			UDestructibleComponent* DestructibleComponent = Cast<UDestructibleComponent>(FPhysxUserData::Get<UPrimitiveComponent>(damageEvent.destructible->userData));
+			const NxApexDamageEventReportData & damageEvent = PendingDamageEvent.DamageEvent;
+			check(DestructibleComponent == Cast<UDestructibleComponent>(FPhysxUserData::Get<UPrimitiveComponent>(damageEvent.destructible->userData)))	//we store this as a weak pointer above in case one of the callbacks decided to call DestroyComponent
 			DestructibleComponent->OnDamageEvent(damageEvent);
 		}
-
-		DestructibleDamageEventQueue.Empty();
 	}
-#endif
-#endif
 
+	PendingApexDamageManager->PendingDamageEvents.Empty();
+#endif
 }
 
 void FPhysScene::SetUpForFrame(const FVector* NewGrav, float InDeltaSeconds, float InMaxPhysicsDeltaTime)
diff --git a/Engine/Source/Runtime/Engine/Private/PhysicsEngine/PhysXSupport.cpp b/Engine/Source/Runtime/Engine/Private/PhysicsEngine/PhysXSupport.cpp
--- a/Engine/Source/Runtime/Engine/Private/PhysicsEngine/PhysXSupport.cpp
+++ b/Engine/Source/Runtime/Engine/Private/PhysicsEngine/PhysXSupport.cpp
@@ -738,7 +738,7 @@ void FApexChunkReport::onDamageNotify(const NxApexDamageEventReportData& damageE
 		return;
 	}
 
-	DestructibleComponent->OnDamageEvent(damageEvent);
+	DestructibleComponent->GetWorld()->GetPhysicsScene()->AddPendingDamageEvent(DestructibleComponent, damageEvent);
 }
 
 void FApexChunkReport::onStateChangeNotify(const NxApexChunkStateEventData& visibilityEvent)
diff --git a/Engine/Source/Runtime/Engine/Public/PhysicsPublic.h b/Engine/Source/Runtime/Engine/Public/PhysicsPublic.h
--- a/Engine/Source/Runtime/Engine/Public/PhysicsPublic.h
+++ b/Engine/Source/Runtime/Engine/Public/PhysicsPublic.h
@@ -562,11 +562,9 @@ public:
 	/** Adds to queue of skelmesh we want to remove from collision disable table */
 	void DeferredRemoveCollisionDisableTable(uint32 SkelMeshCompID);
 
-#if WITH_SUBSTEPPING
 #if WITH_APEX
-	/** Adds a damage event to be fired when substepping is done */
-	void DeferredDestructibleDamageNotify(const NxApexDamageEventReportData& damageEvent);
-#endif
+	/** Adds a damage event to be fired when fetchResults is done */
+	void AddPendingDamageEvent(class UDestructibleComponent* DestructibleComponent, const NxApexDamageEventReportData& DamageEvent);
 #endif
 
 	/** DeferredCommandHandler - this is mainly used for scene specific APEX calls that need to be deferred */
@@ -608,7 +606,7 @@ private:
 	class FPhysSubstepTask * PhysSubSteppers[PST_MAX];
 
 #if WITH_APEX
-	TArray<NxApexDamageEventReportData> DestructibleDamageEventQueue;
+	TUniquePtr<struct FPendingApexDamageManager> PendingApexDamageManager;
 #endif
 #endif
 
