Fix so that matinee sequences play correctly if Stop() is called before Play().

diff --git a/Engine/Source/Runtime/Engine/Private/Interpolation.cpp b/Engine/Source/Runtime/Engine/Private/Interpolation.cpp
--- a/Engine/Source/Runtime/Engine/Private/Interpolation.cpp
+++ b/Engine/Source/Runtime/Engine/Private/Interpolation.cpp
@@ -482,7 +482,10 @@ void AMatineeActor::Reverse()
 
 void AMatineeActor::Stop()
 {
-	bPendingStop = true;
+	if (bIsPlaying)
+	{
+		bPendingStop = true;
+	}
 }
 
 void AMatineeActor::Pause()
