Rename USoundWaveStreaming to USoundWaveProcedural

diff --git a/Engine/Plugins/Experimental/Phya/Source/Phya/Classes/PhyaCollisionHandler.h b/Engine/Plugins/Experimental/Phya/Source/Phya/Classes/PhyaCollisionHandler.h
--- a/Engine/Plugins/Experimental/Phya/Source/Phya/Classes/PhyaCollisionHandler.h
+++ b/Engine/Plugins/Experimental/Phya/Source/Phya/Classes/PhyaCollisionHandler.h
@@ -37,13 +37,13 @@ class UPhyaCollisionHandler : public UPhysicsCollisionHandler
 
 	void TestImpact();
 
-	void StreamingWaveUnderflow(USoundWaveStreaming* InStreamingWave, int32 SamplesRequired);
+	void ProceduralWaveUnderflow(USoundWaveProcedural* InProceduralWave, int32 SamplesRequired);
 
 	UPROPERTY()
 	UAudioComponent* AudioComp;
 
 	UPROPERTY()
-	USoundWaveStreaming* StreamingWave;
+	USoundWaveProcedural* ProceduralWave;
 
 	TMap< FPhyaBodyInstancePair, TSharedPtr<FPhyaPairInfo> >	PairHash;
 
diff --git a/Engine/Plugins/Experimental/Phya/Source/Phya/Private/PhyaCollisionHandler.cpp b/Engine/Plugins/Experimental/Phya/Source/Phya/Private/PhyaCollisionHandler.cpp
--- a/Engine/Plugins/Experimental/Phya/Source/Phya/Private/PhyaCollisionHandler.cpp
+++ b/Engine/Plugins/Experimental/Phya/Source/Phya/Private/PhyaCollisionHandler.cpp
@@ -69,16 +69,16 @@ void UPhyaCollisionHandler::InitCollisionHandler()
 	InitPhya();
 
 	// Create streaming wave object
-	StreamingWave = NewObject<USoundWaveStreaming>(this);
-	StreamingWave->SampleRate = SampleRate;
-	StreamingWave->NumChannels = 1;
-	StreamingWave->Duration = INDEFINITELY_LOOPING_DURATION;
-	StreamingWave->bLooping = false;
-	StreamingWave->OnSoundWaveStreamingUnderflow = FOnSoundWaveStreamingUnderflow::CreateUObject(this, &UPhyaCollisionHandler::StreamingWaveUnderflow);
+	ProceduralWave = NewObject<USoundWaveProcedural>(this);
+	ProceduralWave->SampleRate = SampleRate;
+	ProceduralWave->NumChannels = 1;
+	ProceduralWave->Duration = INDEFINITELY_LOOPING_DURATION;
+	ProceduralWave->bLooping = false;
+	ProceduralWave->OnSoundWaveProceduralUnderflow = FOnSoundWaveProceduralUnderflow::CreateUObject(this, &UPhyaCollisionHandler::ProceduralWaveUnderflow);
 
 	// Create audio component
 	AudioComp = NewObject<UAudioComponent>(this);
-	AudioComp->SetSound(StreamingWave);
+	AudioComp->SetSound(ProceduralWave);
 	AudioComp->Play();
 
 	Whitesurf = new paFunSurface;			// Function surface.. uses surface functions, combined with
@@ -128,11 +128,11 @@ void UPhyaCollisionHandler::InitCollisionHandler()
 	//GenerateAudio();
 }
 
-void UPhyaCollisionHandler::StreamingWaveUnderflow(USoundWaveStreaming* InStreamingWave, int32 SamplesRequired)
+void UPhyaCollisionHandler::ProceduralWaveUnderflow(USoundWaveProcedural* InProceduralWave, int32 SamplesRequired)
 {
-	check(StreamingWave == InStreamingWave);
+	check(ProceduralWave == InProceduralWave);
 
-	const int32 QueuedSamples = StreamingWave->GetAvailableAudioByteCount()/sizeof(uint16);
+	const int32 QueuedSamples = ProceduralWave->GetAvailableAudioByteCount()/sizeof(uint16);
 	const int32 SamplesNeeded = SamplesRequired - QueuedSamples;
 	const int32 BlocksNeeded = FMath::CeilToInt(SamplesNeeded/BlockSize);
 
@@ -155,7 +155,7 @@ void UPhyaCollisionHandler::StreamingWaveUnderflow(USoundWaveStreaming* InStream
 		}
 	}
 
-	StreamingWave->QueueAudio((uint8*)SampleData.GetData(), SampleData.Num() * sizeof(int16));
+	ProceduralWave->QueueAudio((uint8*)SampleData.GetData(), SampleData.Num() * sizeof(int16));
 }
 
 
diff --git a/Engine/Plugins/Runtime/SoundMod/Source/SoundMod/Classes/SoundModWave.h b/Engine/Plugins/Runtime/SoundMod/Source/SoundMod/Classes/SoundModWave.h
--- a/Engine/Plugins/Runtime/SoundMod/Source/SoundMod/Classes/SoundModWave.h
+++ b/Engine/Plugins/Runtime/SoundMod/Source/SoundMod/Classes/SoundModWave.h
@@ -9,7 +9,7 @@
 #include "SoundModWave.generated.h"
 
 UCLASS()
-class USoundModWave : public USoundWaveStreaming
+class USoundModWave : public USoundWaveProcedural
 {
 	GENERATED_UCLASS_BODY()
 
diff --git a/Engine/Source/Runtime/Engine/Classes/Sound/SoundWaveProcedural.h b/Engine/Source/Runtime/Engine/Classes/Sound/SoundWaveProcedural.h
new file mode 100644
--- /dev/null
+++ b/Engine/Source/Runtime/Engine/Classes/Sound/SoundWaveProcedural.h
@@ -0,0 +1,56 @@
+// Copyright 1998-2015 Epic Games, Inc. All Rights Reserved.
+
+
+/** 
+ * Playable sound object for wave files that are streamed, particularly VOIP
+ */
+
+#pragma once
+#include "Sound/SoundWave.h"
+#include "SoundWaveProcedural.generated.h"
+
+DECLARE_DELEGATE_TwoParams( FOnSoundWaveProceduralUnderflow, class USoundWaveProcedural*, int32 );
+
+DEPRECATED(4.9, "FOnSoundWaveStreamingUnderflow has been renamed FOnSoundWaveProceduralUnderflow")
+typedef FOnSoundWaveProceduralUnderflow FOnSoundWaveStreamingUnderflow;
+
+DEPRECATED(4.9, "USoundWaveStreaming has been renamed USoundWaveProcedural.")
+typedef class USoundWaveProcedural USoundWaveStreaming;
+
+UCLASS()
+class ENGINE_API USoundWaveProcedural : public USoundWave
+{
+	GENERATED_UCLASS_BODY()
+
+private:
+	TArray<uint8> QueuedAudio;
+
+public:
+
+	// Begin UObject interface. 
+	virtual void Serialize( FArchive& Ar ) override;
+	virtual void GetAssetRegistryTags(TArray<FAssetRegistryTag>& OutTags) const override;
+	virtual SIZE_T GetResourceSize(EResourceSizeMode::Type Mode) override;
+	// End UObject interface. 
+
+	// Begin USoundWave interface.
+	virtual int32 GeneratePCMData(uint8* PCMData, const int32 SamplesNeeded) override;
+	virtual FByteBulkData* GetCompressedData(FName Format) override;
+	virtual void InitAudioResource( FByteBulkData& CompressedData ) override;
+	virtual bool InitAudioResource(FName Format) override;
+	virtual int32 GetResourceSizeForFormat(FName Format) override;
+	// End USoundWave interface.
+
+
+	/** Add data to the FIFO that feeds the audio device. */
+	void QueueAudio(const uint8* AudioData, const int32 BufferSize);
+
+	/** Remove all queued data from the FIFO. This is only necessary if you want to start over, or GeneratePCMData() isn't going to be called, since that will eventually drain it. */
+	void ResetAudio();
+
+	/** Query bytes queued for playback */
+	int32 GetAvailableAudioByteCount();
+
+	/** Called when GeneratePCMData is called but not enough data is available. Allows more data to be added, and will try again */
+	FOnSoundWaveProceduralUnderflow OnSoundWaveProceduralUnderflow;
+};
\ No newline at end of file
diff --git a/Engine/Source/Runtime/Engine/Classes/Sound/SoundWaveStreaming.h b/Engine/Source/Runtime/Engine/Classes/Sound/SoundWaveStreaming.h
deleted file mode 100644
diff --git a/Engine/Source/Runtime/Engine/Private/SoundWaveProcedural.cpp b/Engine/Source/Runtime/Engine/Private/SoundWaveProcedural.cpp
new file mode 100644
--- /dev/null
+++ b/Engine/Source/Runtime/Engine/Private/SoundWaveProcedural.cpp
@@ -0,0 +1,95 @@
+// Copyright 1998-2015 Epic Games, Inc. All Rights Reserved.
+
+#include "EnginePrivate.h"
+#include "Sound/SoundWaveProcedural.h"
+
+USoundWaveProcedural::USoundWaveProcedural(const FObjectInitializer& ObjectInitializer)
+	: Super(ObjectInitializer)
+{
+	bProcedural = true;
+}
+
+void USoundWaveProcedural::QueueAudio( const uint8* AudioData, const int32 BufferSize )
+{
+	if (BufferSize == 0 || !ensure( ( BufferSize % sizeof( int16 ) ) == 0 ))
+	{
+		return;
+	}
+	const int32 Position = QueuedAudio.AddUninitialized( BufferSize );
+	FMemory::Memcpy( &QueuedAudio[ Position ], AudioData, BufferSize );
+}
+
+int32 USoundWaveProcedural::GeneratePCMData( uint8* PCMData, const int32 SamplesNeeded )
+{
+	int32 SamplesAvailable = QueuedAudio.Num() / sizeof( int16 );
+
+	// if delegate is bound and we don't have enough samples, call it so system can supply more
+	if (SamplesNeeded > SamplesAvailable && OnSoundWaveProceduralUnderflow.IsBound())
+	{
+		OnSoundWaveProceduralUnderflow.Execute(this, SamplesNeeded);
+		// Update available samples
+		SamplesAvailable = QueuedAudio.Num() / sizeof( int16 );
+	}
+
+	if (SamplesAvailable > 0 && SamplesNeeded > 0)
+	{
+		const int32 SamplesToCopy = FMath::Min<int32>( SamplesNeeded, SamplesAvailable );
+		const int32 BytesToCopy = SamplesToCopy * sizeof( int16 );
+
+		FMemory::Memcpy( ( void* )PCMData,  &QueuedAudio[ 0 ], BytesToCopy );
+		QueuedAudio.RemoveAt( 0, BytesToCopy );
+		return BytesToCopy;
+	}
+
+	return 0;
+}
+
+void USoundWaveProcedural::ResetAudio()
+{
+	QueuedAudio.Empty();
+}
+
+int32 USoundWaveProcedural::GetAvailableAudioByteCount()
+{
+	return QueuedAudio.Num();
+}
+
+SIZE_T USoundWaveProcedural::GetResourceSize(EResourceSizeMode::Type Mode)
+{
+	return 0;
+}
+
+int32 USoundWaveProcedural::GetResourceSizeForFormat(FName Format)
+{
+	return 0;
+}
+
+void USoundWaveProcedural::GetAssetRegistryTags(TArray<FAssetRegistryTag>& OutTags) const
+{
+	// SoundWaveProcedural should never be in the asset registry
+	check(false);
+}
+
+FByteBulkData* USoundWaveProcedural::GetCompressedData(FName Format)
+{
+	// SoundWaveProcedural does not have compressed data and should generally not be asked about it
+	return nullptr;
+}
+
+void USoundWaveProcedural::Serialize( FArchive& Ar )
+{
+	// Do not call the USoundWave version of serialize
+	USoundBase::Serialize( Ar );
+}
+
+void USoundWaveProcedural::InitAudioResource( FByteBulkData& CompressedData )
+{
+	// Should never be pushing compressed data to a SoundWaveProcedural
+	check(false);
+}
+
+bool USoundWaveProcedural::InitAudioResource(FName Format)
+{
+	// Nothing to be done to initialize a USoundWaveProcedural
+	return true;
+}
\ No newline at end of file
diff --git a/Engine/Source/Runtime/Engine/Private/SoundWaveStreaming.cpp b/Engine/Source/Runtime/Engine/Private/SoundWaveStreaming.cpp
deleted file mode 100644
diff --git a/Engine/Source/Runtime/Engine/Public/Engine.h b/Engine/Source/Runtime/Engine/Public/Engine.h
--- a/Engine/Source/Runtime/Engine/Public/Engine.h
+++ b/Engine/Source/Runtime/Engine/Public/Engine.h
@@ -282,7 +282,6 @@
 #include "Sound/SoundClass.h"
 #include "Sound/SoundCue.h"
 #include "Sound/SoundGroups.h"
-#include "Sound/SoundWaveStreaming.h"
 #include "Sound/SoundMix.h"
 #include "Engine/StaticMeshSocket.h"
 #include "Camera/CameraStackTypes.h"
diff --git a/Engine/Source/Runtime/Online/OnlineSubsystemSteam/Private/VoiceEngineSteam.cpp b/Engine/Source/Runtime/Online/OnlineSubsystemSteam/Private/VoiceEngineSteam.cpp
--- a/Engine/Source/Runtime/Online/OnlineSubsystemSteam/Private/VoiceEngineSteam.cpp
+++ b/Engine/Source/Runtime/Online/OnlineSubsystemSteam/Private/VoiceEngineSteam.cpp
@@ -9,6 +9,7 @@
 
 #include "Engine.h"
 #include "SoundDefinitions.h"
+#include "Runtime/Engine/Classes/Sound/SoundWaveProcedural.h"
 #include "OnlineSubsystemUtils.h"
 
 /** Largest size Steam says it will need to compress data */
@@ -324,7 +325,7 @@ uint32 FVoiceEngineSteam::SubmitRemoteVoiceData(const FUniqueNetId& RemoteTalker
 
 	if (QueuedData->AudioComponent != NULL)
 	{
-		USoundWaveStreaming* SoundStreaming = CastChecked<USoundWaveStreaming>(QueuedData->AudioComponent->Sound);
+		USoundWaveProcedural* SoundStreaming = CastChecked<USoundWaveProcedural>(QueuedData->AudioComponent->Sound);
 		if (SoundStreaming->GetAvailableAudioByteCount() == 0)
 		{
 			UE_LOG(LogVoiceDecode, Log, TEXT("VOIP audio component was starved!"));
diff --git a/Engine/Source/Runtime/Online/OnlineSubsystemUtils/Private/OnlineSubsystemUtils.cpp b/Engine/Source/Runtime/Online/OnlineSubsystemUtils/Private/OnlineSubsystemUtils.cpp
--- a/Engine/Source/Runtime/Online/OnlineSubsystemUtils/Private/OnlineSubsystemUtils.cpp
+++ b/Engine/Source/Runtime/Online/OnlineSubsystemUtils/Private/OnlineSubsystemUtils.cpp
@@ -11,6 +11,7 @@
 #include "Voice.h"
 #include "SoundDefinitions.h"
 #include "Runtime/Engine/Classes/Sound/AudioSettings.h"
+#include "Runtime/Engine/Classes/Sound/SoundWaveProcedural.h"
 
 // Testing classes
 #include "Tests/TestFriendsInterface.h"
@@ -47,7 +48,7 @@ UAudioComponent* CreateVoiceAudioComponent(uint32 SampleRate)
 	{
 		if (FAudioDevice* AudioDevice = GEngine->GetMainAudioDevice())
 		{
-			USoundWaveStreaming* SoundStreaming = NewObject<USoundWaveStreaming>();
+			USoundWaveProcedural* SoundStreaming = NewObject<USoundWaveProcedural>();
 			SoundStreaming->SampleRate = SampleRate;
 			SoundStreaming->NumChannels = 1;
 			SoundStreaming->Duration = INDEFINITELY_LOOPING_DURATION;
diff --git a/Engine/Source/Runtime/Online/OnlineSubsystemUtils/Private/Tests/TestVoice.cpp b/Engine/Source/Runtime/Online/OnlineSubsystemUtils/Private/Tests/TestVoice.cpp
--- a/Engine/Source/Runtime/Online/OnlineSubsystemUtils/Private/Tests/TestVoice.cpp
+++ b/Engine/Source/Runtime/Online/OnlineSubsystemUtils/Private/Tests/TestVoice.cpp
@@ -7,6 +7,7 @@
 #include "Engine.h"
 #include "OnlineSubsystemUtils.h"
 #include "Voice.h"
+#include "Runtime/Engine/Classes/Sound/SoundWaveProcedural.h"
 
 FTestVoice::FTestVoice() :
 	VoiceComp(NULL),
@@ -114,7 +115,7 @@ bool FTestVoice::Tick(float DeltaTime)
 		if (VoiceComp)
 		{
 			static bool bLastHasData = true;
-			USoundWaveStreaming* SoundStreaming = CastChecked<USoundWaveStreaming>(VoiceComp->Sound);
+			USoundWaveProcedural* SoundStreaming = CastChecked<USoundWaveProcedural>(VoiceComp->Sound);
 			bool bHasData = SoundStreaming->GetAvailableAudioByteCount() != 0;
 
 			if (bHasData != bLastHasData)
diff --git a/Engine/Source/Runtime/Online/OnlineSubsystemUtils/Private/VoiceEngineImpl.cpp b/Engine/Source/Runtime/Online/OnlineSubsystemUtils/Private/VoiceEngineImpl.cpp
--- a/Engine/Source/Runtime/Online/OnlineSubsystemUtils/Private/VoiceEngineImpl.cpp
+++ b/Engine/Source/Runtime/Online/OnlineSubsystemUtils/Private/VoiceEngineImpl.cpp
@@ -7,6 +7,7 @@
 
 #include "Engine.h"
 #include "SoundDefinitions.h"
+#include "Runtime/Engine/Classes/Sound/SoundWaveProcedural.h"
 #include "OnlineSubsystemUtils.h"
 
 /** Largest size preallocated for compressed data */
@@ -387,7 +388,7 @@ uint32 FVoiceEngineImpl::SubmitRemoteVoiceData(const FUniqueNetId& RemoteTalkerI
 			QueuedData.AudioComponent->Play();
 		}
 
-		USoundWaveStreaming* SoundStreaming = CastChecked<USoundWaveStreaming>(QueuedData.AudioComponent->Sound);
+		USoundWaveProcedural* SoundStreaming = CastChecked<USoundWaveProcedural>(QueuedData.AudioComponent->Sound);
 
 		if (0)
 		{
