Re-add SoundWaveStreaming.h with deprecation and redirection to SoundWaveProcedural.h

diff --git a/Engine/Source/Runtime/Engine/Classes/Sound/SoundWaveStreaming.h b/Engine/Source/Runtime/Engine/Classes/Sound/SoundWaveStreaming.h
new file mode 100644
--- /dev/null
+++ b/Engine/Source/Runtime/Engine/Classes/Sound/SoundWaveStreaming.h
@@ -0,0 +1,11 @@
+// Copyright 1998-2015 Epic Games, Inc. All Rights Reserved.
+
+#ifdef _MSC_VER
+    #pragma message(__FILE__"(9): warning : #include SoundWaveProcedural.h instead of SoundWaveStreaming.h")
+#else
+    #pragma message("#include SoundWaveProcedural.h instead of SoundWaveStreaming.h")
+#endif
+
+#include "Sound/SoundWaveProcedural.h"
+
+//DEPRECATED(4.9)
\ No newline at end of file
