Make all actor iterators skip Actors that have become pending kill since first gathered
(remove IsPendingKills added in CL# 2573823 that are no longer necessary)

diff --git a/Engine/Source/Runtime/Engine/Private/UnrealEngine.cpp b/Engine/Source/Runtime/Engine/Private/UnrealEngine.cpp
--- a/Engine/Source/Runtime/Engine/Private/UnrealEngine.cpp
+++ b/Engine/Source/Runtime/Engine/Private/UnrealEngine.cpp
@@ -8902,11 +8902,7 @@ bool UEngine::LoadMap( FWorldContext& WorldContext, FURL URL, class UPendingNetG
 
 		for (FActorIterator ActorIt(WorldContext.World()); ActorIt; ++ActorIt)
 		{
-			// Routing the end play of an earlier Actor could have resulted in this Actor already being cleaned up/destroyed
-			if (!ActorIt->IsPendingKill())
-			{
-				ActorIt->RouteEndPlay(EEndPlayReason::LevelTransition);
-			}
+			ActorIt->RouteEndPlay(EEndPlayReason::LevelTransition);
 		}
 
 		// Do this after destroying pawns/playercontrollers, in case that spawns new things (e.g. dropped weapons)
diff --git a/Engine/Source/Runtime/Engine/Private/World.cpp b/Engine/Source/Runtime/Engine/Private/World.cpp
--- a/Engine/Source/Runtime/Engine/Private/World.cpp
+++ b/Engine/Source/Runtime/Engine/Private/World.cpp
@@ -4670,13 +4670,6 @@ UWorld* FSeamlessTravelHandler::Tick()
 			for (FActorIterator It(CurrentWorld); It; ++It)
 			{
 				AActor* TheActor = *It;
-
-				// Routing the end play of an earlier Actor could have resulted in this Actor already being cleaned up/destroyed
-				if (TheActor->IsPendingKill())
-				{
-					continue;
-				}
-
 				bool bSameLevel = TheActor->GetLevel() == CurrentWorld->PersistentLevel;
 				bool bShouldKeep = KeepAnnotation.Get(TheActor);
 				bool bDormant = false;
diff --git a/Engine/Source/Runtime/Engine/Public/EngineUtils.h b/Engine/Source/Runtime/Engine/Public/EngineUtils.h
--- a/Engine/Source/Runtime/Engine/Public/EngineUtils.h
+++ b/Engine/Source/Runtime/Engine/Public/EngineUtils.h
@@ -337,7 +337,7 @@ protected:
 	 */
 	FORCEINLINE static bool IsActorSuitable(AActor* Actor)
 	{
-		return true;
+		return !Actor->IsPendingKill();
 	}
 
 	/**
@@ -540,7 +540,7 @@ protected:
 	 */
 	FORCEINLINE static bool IsActorSuitable(AActor* Actor)
 	{
-		return Actor->IsSelected();
+		return !Actor->IsPendingKill() && Actor->IsSelected();
 	}
 };
 
