Integrate - Release StaticShadowDepthMap if it has been initialized, regardless of mobility.

diff --git a/Engine/Source/Runtime/Engine/Private/Components/LightComponent.cpp b/Engine/Source/Runtime/Engine/Private/Components/LightComponent.cpp
--- a/Engine/Source/Runtime/Engine/Private/Components/LightComponent.cpp
+++ b/Engine/Source/Runtime/Engine/Private/Components/LightComponent.cpp
@@ -611,11 +611,8 @@ void ULightComponent::BeginDestroy()
 {
 	Super::BeginDestroy();
 
-	if (HasStaticShadowing() && !HasStaticLighting())
-	{
-		BeginReleaseResource(&StaticShadowDepthMap);
-		StaticShadowDepthMap.Empty();
-	}
+	BeginReleaseResource(&StaticShadowDepthMap);
+	StaticShadowDepthMap.Empty();
 
 	// Use a fence to keep track of when the rendering thread executes the release command
 	DestroyFence.BeginFence();
