Fixed non-Unity build error

diff --git a/Engine/Source/Runtime/Engine/Private/DemoNetDriver.cpp b/Engine/Source/Runtime/Engine/Private/DemoNetDriver.cpp
--- a/Engine/Source/Runtime/Engine/Private/DemoNetDriver.cpp
+++ b/Engine/Source/Runtime/Engine/Private/DemoNetDriver.cpp
@@ -21,6 +21,8 @@
 #include "Net/NetworkProfiler.h"
 #include "Net/DataReplication.h"
 #include "GameFramework/GameMode.h"
+#include "GameFramework/GameState.h"
+#include "GameFramework/PlayerState.h"
 
 DEFINE_LOG_CATEGORY( LogDemo );
 
