#ue4 - Fix physics volume change notifications not working for the first change in physics volume.

diff --git a/Engine/Source/Runtime/Engine/Private/Components/PrimitiveComponent.cpp b/Engine/Source/Runtime/Engine/Private/Components/PrimitiveComponent.cpp
--- a/Engine/Source/Runtime/Engine/Private/Components/PrimitiveComponent.cpp
+++ b/Engine/Source/Runtime/Engine/Private/Components/PrimitiveComponent.cpp
@@ -2456,10 +2456,7 @@ void UPrimitiveComponent::UpdatePhysicsVolume( bool bTriggerNotifiers )
 				}
 			}
 
-			if (BestVolume != GetPhysicsVolume())
-			{
-				SetPhysicsVolume(BestVolume, bTriggerNotifiers);
-			}
+			SetPhysicsVolume(BestVolume, bTriggerNotifiers);
 		}
 		else
 		{
