#ue4 - Avoid blocking hit notifications from FScopedMovementUpdates when the first hit destroys the object, and protect against null access.

diff --git a/Engine/Source/Runtime/Engine/Private/Components/PrimitiveComponent.cpp b/Engine/Source/Runtime/Engine/Private/Components/PrimitiveComponent.cpp
--- a/Engine/Source/Runtime/Engine/Private/Components/PrimitiveComponent.cpp
+++ b/Engine/Source/Runtime/Engine/Private/Components/PrimitiveComponent.cpp
@@ -1685,14 +1685,18 @@ bool UPrimitiveComponent::MoveComponentImpl( const FVector& Delta, const FQuat&
 void UPrimitiveComponent::DispatchBlockingHit(AActor& Owner, FHitResult const& BlockingHit)
 {
 	UPrimitiveComponent* const BlockingHitComponent = BlockingHit.Component.Get();
-	Owner.DispatchBlockingHit(this, BlockingHitComponent, true, BlockingHit);
-
-	if (!BlockingHitComponent->IsPendingKill())
+	if (BlockingHitComponent)
 	{
-		// BlockingHit.GetActor() could be marked for deletion in DispatchBlockingHit(), which would make the weak pointer return NULL.
-		if (AActor* const BlockingHitActor = BlockingHit.GetActor())
+		Owner.DispatchBlockingHit(this, BlockingHitComponent, true, BlockingHit);
+
+		// Dispatch above could kill the component, so we need to check that.
+		if (!BlockingHitComponent->IsPendingKill())
 		{
-			BlockingHitActor->DispatchBlockingHit(BlockingHitComponent, this, false, BlockingHit);
+			// BlockingHit.GetActor() could be marked for deletion in DispatchBlockingHit(), which would make the weak pointer return NULL.
+			if (AActor* const BlockingHitActor = BlockingHit.GetActor())
+			{
+				BlockingHitActor->DispatchBlockingHit(BlockingHitComponent, this, false, BlockingHit);
+			}
 		}
 	}
 }
