Setting r.Shadow.DistanceScale to 0 will fully disable CSM

diff --git a/Engine/Source/Runtime/Engine/Private/Components/DirectionalLightComponent.cpp b/Engine/Source/Runtime/Engine/Private/Components/DirectionalLightComponent.cpp
--- a/Engine/Source/Runtime/Engine/Private/Components/DirectionalLightComponent.cpp
+++ b/Engine/Source/Runtime/Engine/Private/Components/DirectionalLightComponent.cpp
@@ -329,7 +329,7 @@ private:
 	uint32 GetNumShadowMappedCascades(uint32 MaxShadowCascades, bool bPrecomputedLightingIsValid) const
 	{
 		const int32 EffectiveNumDynamicShadowCascades = bPrecomputedLightingIsValid ? DynamicShadowCascades : FMath::Max(0, CVarUnbuiltNumWholeSceneDynamicShadowCascades.GetValueOnAnyThread());
-		const int32 NumCascades = GetEffectiveWholeSceneDynamicShadowRadius(bPrecomputedLightingIsValid) > 0.0f ? EffectiveNumDynamicShadowCascades : 0;
+		const int32 NumCascades = GetCSMMaxDistance(bPrecomputedLightingIsValid) > 0.0f ? EffectiveNumDynamicShadowCascades : 0;
 		return FMath::Min<int32>(NumCascades, MaxShadowCascades);
 	}
 
@@ -337,7 +337,7 @@ private:
 	{
 		static const auto CVar = IConsoleManager::Get().FindTConsoleVariableDataFloat(TEXT("r.Shadow.DistanceScale"));
 		 
-		float Scale = FMath::Clamp(CVar->GetValueOnRenderThread(), 0.1f, 2.0f);
+		float Scale = FMath::Clamp(CVar->GetValueOnRenderThread(), 0.0f, 2.0f);
 		float Distance = GetEffectiveWholeSceneDynamicShadowRadius(bPrecomputedLightingIsValid) * Scale;
 		return Distance;
 	}
