#ue4 - Fix blocking hits dispatched from FScopedMovementUpdate() being dispatched when collision response has changed (from overlaps or other blocking hits).

diff --git a/Engine/Source/Runtime/Engine/Private/Components/SceneComponent.cpp b/Engine/Source/Runtime/Engine/Private/Components/SceneComponent.cpp
--- a/Engine/Source/Runtime/Engine/Private/Components/SceneComponent.cpp
+++ b/Engine/Source/Runtime/Engine/Private/Components/SceneComponent.cpp
@@ -514,7 +514,11 @@ void USceneComponent::EndScopedMovementUpdate(class FScopedMovementUpdate& Compl
 					UPrimitiveComponent* PrimitiveThis = CastChecked<UPrimitiveComponent>(this);
 					for (const FHitResult& Hit : CurrentScopedUpdate->BlockingHits)
 					{
-						PrimitiveThis->DispatchBlockingHit(*Owner, Hit);
+						// Collision response may change (due to overlaps or multiple blocking hits), make sure it's still considered blocking.
+						if (PrimitiveThis->GetCollisionResponseToComponent(Hit.GetComponent()) == ECR_Block)
+						{
+							PrimitiveThis->DispatchBlockingHit(*Owner, Hit);
+						}						
 					}
 				}
 			}
