Fix for skeletal mesh APEX blueprint functions not appearing in-editor

diff --git a/Engine/Source/Runtime/Engine/Classes/Engine/SkeletalMesh.h b/Engine/Source/Runtime/Engine/Classes/Engine/SkeletalMesh.h
--- a/Engine/Source/Runtime/Engine/Classes/Engine/SkeletalMesh.h
+++ b/Engine/Source/Runtime/Engine/Classes/Engine/SkeletalMesh.h
@@ -847,6 +847,7 @@ public:
 	ENGINE_API void  LoadClothCollisionVolumes(int32 AssetIndex, physx::apex::NxClothingAsset* ClothingAsset);
 	ENGINE_API bool IsMappedClothingLOD(int32 LODIndex, int32 AssetIndex);
 	ENGINE_API int32 GetClothAssetIndex(int32 LODIndex, int32 SectionIndex);
+#endif
 
 	/** 
 	 * Checks whether the provided section is using APEX cloth. if bCheckCorrespondingSections is true
@@ -858,7 +859,6 @@ public:
 	 */
 	UFUNCTION(BlueprintCallable, Category="Cloth")
 	ENGINE_API bool IsSectionUsingCloth(int32 InSectionIndex, bool bCheckCorrespondingSections = true) const;
-#endif// #if WITH_APEX_CLOTHING
 
 	ENGINE_API void CreateBodySetup();
 	ENGINE_API UBodySetup* GetBodySetup();
