Documentation cleanup - references to "sphere" in world collision API changed to "shape".

diff --git a/Engine/Source/Runtime/Engine/Classes/Engine/World.h b/Engine/Source/Runtime/Engine/Classes/Engine/World.h
--- a/Engine/Source/Runtime/Engine/Classes/Engine/World.h
+++ b/Engine/Source/Runtime/Engine/Classes/Engine/World.h
@@ -1277,9 +1277,9 @@ public:
 	bool LineTraceMultiByProfile(TArray<struct FHitResult>& OutHits, const FVector& Start, const FVector& End, FName ProfileName, const FCollisionQueryParams& Params = FCollisionQueryParams::DefaultQueryParam) const;
 
 	/**
-	 *  Sweep a sphere against the world using a specific channel and return if a blocking hit is found.
-	 *  @param  Start           Start location of the sphere
-	 *  @param  End             End location of the sphere
+	 *  Sweep a shape against the world using a specific channel and return if a blocking hit is found.
+	 *  @param  Start           Start location of the shape
+	 *  @param  End             End location of the shape
 	 *  @param  TraceChannel    The 'channel' that this trace uses, used to determine which components to hit
 	 *  @param	CollisionShape	CollisionShape - supports Box, Sphere, Capsule
 	 *  @param  Params          Additional parameters used for the trace
@@ -1293,9 +1293,9 @@ public:
 	}
 
 	/**
-	 *  Sweep a sphere against the world using a specific channel and return if a blocking hit is found.
-	 *  @param  Start           Start location of the sphere
-	 *  @param  End             End location of the sphere
+	 *  Sweep a shape against the world using a specific channel and return if a blocking hit is found.
+	 *  @param  Start           Start location of the shape
+	 *  @param  End             End location of the shape
 	 *  @param  TraceChannel    The 'channel' that this trace uses, used to determine which components to hit
 	 *  @param	CollisionShape	CollisionShape - supports Box, Sphere, Capsule
 	 *  @param  Params          Additional parameters used for the trace
@@ -1305,9 +1305,9 @@ public:
 	bool SweepTestByChannel(const FVector& Start, const FVector& End, const FQuat& Rot, ECollisionChannel TraceChannel, const FCollisionShape& CollisionShape, const FCollisionQueryParams& Params = FCollisionQueryParams::DefaultQueryParam, const FCollisionResponseParams& ResponseParam = FCollisionResponseParams::DefaultResponseParam) const;
 
 	/**
-	 *  Sweep a sphere against the world using object types and return if a blocking hit is found.
-	 *  @param  Start           Start location of the sphere
-	 *  @param  End             End location of the sphere
+	 *  Sweep a shape against the world using object types and return if a blocking hit is found.
+	 *  @param  Start           Start location of the shape
+	 *  @param  End             End location of the shape
 	 *  @param	CollisionShape	CollisionShape - supports Box, Sphere, Capsule
 	 *  @param  Params          Additional parameters used for the trace
 	 *	@param	ObjectQueryParams	List of object types it's looking for
@@ -1320,9 +1320,9 @@ public:
 	}
 
 	/**
-	 *  Sweep a sphere against the world using object types and return if a blocking hit is found.
-	 *  @param  Start           Start location of the sphere
-	 *  @param  End             End location of the sphere
+	 *  Sweep a shape against the world using object types and return if a blocking hit is found.
+	 *  @param  Start           Start location of the shape
+	 *  @param  End             End location of the shape
 	 *	@param	ObjectQueryParams	List of object types it's looking for
 	 *  @param	CollisionShape	CollisionShape - supports Box, Sphere, Capsule
 	 *  @param  Params          Additional parameters used for the trace
@@ -1332,9 +1332,9 @@ public:
 
 
 	/**
-	 *  Sweep a sphere against the world using a specific profile and return if a blocking hit is found.
-	 *  @param  Start           Start location of the sphere
-	 *  @param  End             End location of the sphere
+	 *  Sweep a shape against the world using a specific profile and return if a blocking hit is found.
+	 *  @param  Start           Start location of the shape
+	 *  @param  End             End location of the shape
 	 *  @param  ProfileName     The 'profile' used to determine which components to hit
 	 *  @param	CollisionShape	CollisionShape - supports Box, Sphere, Capsule
 	 *  @param  Params          Additional parameters used for the trace
@@ -1343,10 +1343,10 @@ public:
 	bool SweepTestByProfile(const FVector& Start, const FVector& End, const FQuat& Rot, FName ProfileName, const FCollisionShape& CollisionShape, const FCollisionQueryParams& Params) const;
 
 	/**
-	 *  Sweep a sphere against the world and return the first blocking hit using a specific channel
+	 *  Sweep a shape against the world and return the first blocking hit using a specific channel
 	 *  @param  OutHit          First blocking hit found
-	 *  @param  Start           Start location of the sphere
-	 *  @param  End             End location of the sphere
+	 *  @param  Start           Start location of the shape
+	 *  @param  End             End location of the shape
 	 *  @param  TraceChannel    The 'channel' that this trace is in, used to determine which components to hit
 	 *  @param	CollisionShape	CollisionShape - supports Box, Sphere, Capsule
 	 *  @param  Params          Additional parameters used for the trace
@@ -1360,10 +1360,10 @@ public:
 	}
 
 	/**
-	 *  Sweep a sphere against the world and return the first blocking hit using a specific channel
+	 *  Sweep a shape against the world and return the first blocking hit using a specific channel
 	 *  @param  OutHit          First blocking hit found
-	 *  @param  Start           Start location of the sphere
-	 *  @param  End             End location of the sphere
+	 *  @param  Start           Start location of the shape
+	 *  @param  End             End location of the shape
 	 *  @param  TraceChannel    The 'channel' that this trace is in, used to determine which components to hit
 	 *  @param	CollisionShape	CollisionShape - supports Box, Sphere, Capsule
 	 *  @param  Params          Additional parameters used for the trace
@@ -1373,10 +1373,10 @@ public:
 	bool SweepSingleByChannel(struct FHitResult& OutHit, const FVector& Start, const FVector& End, const FQuat& Rot, ECollisionChannel TraceChannel, const FCollisionShape& CollisionShape, const FCollisionQueryParams& Params = FCollisionQueryParams::DefaultQueryParam, const FCollisionResponseParams& ResponseParam = FCollisionResponseParams::DefaultResponseParam) const;
 
 	/**
-	 *  Sweep a sphere against the world and return the first blocking hit using object types
+	 *  Sweep a shape against the world and return the first blocking hit using object types
 	 *  @param  OutHit          First blocking hit found
-	 *  @param  Start           Start location of the sphere
-	 *  @param  End             End location of the sphere
+	 *  @param  Start           Start location of the shape
+	 *  @param  End             End location of the shape
 	 *  @param	CollisionShape	CollisionShape - supports Box, Sphere, Capsule
 	 *  @param  Params          Additional parameters used for the trace
 	 *	@param	ObjectQueryParams	List of object types it's looking for
@@ -1389,10 +1389,10 @@ public:
 	}
 
 	/**
-	 *  Sweep a sphere against the world and return the first blocking hit using object types
+	 *  Sweep a shape against the world and return the first blocking hit using object types
 	 *  @param  OutHit          First blocking hit found
-	 *  @param  Start           Start location of the sphere
-	 *  @param  End             End location of the sphere
+	 *  @param  Start           Start location of the shape
+	 *  @param  End             End location of the shape
 	 *	@param	ObjectQueryParams	List of object types it's looking for
 	 *  @param	CollisionShape	CollisionShape - supports Box, Sphere, Capsule
 	 *  @param  Params          Additional parameters used for the trace
@@ -1401,10 +1401,10 @@ public:
 	bool SweepSingleByObjectType(struct FHitResult& OutHit, const FVector& Start, const FVector& End, const FQuat& Rot, const FCollisionObjectQueryParams& ObjectQueryParams, const FCollisionShape& CollisionShape, const FCollisionQueryParams& Params = FCollisionQueryParams::DefaultQueryParam) const;
 
 	/**
-	 *  Sweep a sphere against the world and return the first blocking hit using a specific profile
+	 *  Sweep a shape against the world and return the first blocking hit using a specific profile
 	 *  @param  OutHit          First blocking hit found
-	 *  @param  Start           Start location of the sphere
-	 *  @param  End             End location of the sphere
+	 *  @param  Start           Start location of the shape
+	 *  @param  End             End location of the shape
 	 *  @param  ProfileName     The 'profile' used to determine which components to hit
 	 *  @param	CollisionShape	CollisionShape - supports Box, Sphere, Capsule
 	 *  @param  Params          Additional parameters used for the trace
@@ -1413,12 +1413,12 @@ public:
 	bool SweepSingleByProfile(struct FHitResult& OutHit, const FVector& Start, const FVector& End, const FQuat& Rot, FName ProfileName, const FCollisionShape& CollisionShape, const FCollisionQueryParams& Params = FCollisionQueryParams::DefaultQueryParam) const;
 
 	/**
-	 *  Sweep a sphere against the world and return all initial overlaps using a specific channel (including blocking) if requested, then overlapping hits and then first blocking hit
+	 *  Sweep a shape against the world and return all initial overlaps using a specific channel (including blocking) if requested, then overlapping hits and then first blocking hit
 	 *  Results are sorted, so a blocking hit (if found) will be the last element of the array
 	 *  Only the single closest blocking result will be generated, no tests will be done after that
 	 *  @param  OutHits         Array of hits found between ray and the world
-	 *  @param  Start           Start location of the sphere
-	 *  @param  End             End location of the sphere
+	 *  @param  Start           Start location of the shape
+	 *  @param  End             End location of the shape
 	 *  @param  TraceChannel    The 'channel' that this ray is in, used to determine which components to hit
 	 *  @param	CollisionShape	CollisionShape - supports Box, Sphere, Capsule
 	 *  @param  Params          Additional parameters used for the trace
@@ -1432,12 +1432,12 @@ public:
 	}
 
 	/**
-	 *  Sweep a sphere against the world and return all initial overlaps using a specific channel (including blocking) if requested, then overlapping hits and then first blocking hit
+	 *  Sweep a shape against the world and return all initial overlaps using a specific channel (including blocking) if requested, then overlapping hits and then first blocking hit
 	 *  Results are sorted, so a blocking hit (if found) will be the last element of the array
 	 *  Only the single closest blocking result will be generated, no tests will be done after that
 	 *  @param  OutHits         Array of hits found between ray and the world
-	 *  @param  Start           Start location of the sphere
-	 *  @param  End             End location of the sphere
+	 *  @param  Start           Start location of the shape
+	 *  @param  End             End location of the shape
 	 *  @param  TraceChannel    The 'channel' that this ray is in, used to determine which components to hit
 	 *  @param	CollisionShape	CollisionShape - supports Box, Sphere, Capsule
 	 *  @param  Params          Additional parameters used for the trace
@@ -1447,12 +1447,12 @@ public:
 	bool SweepMultiByChannel(TArray<struct FHitResult>& OutHits, const FVector& Start, const FVector& End, const FQuat& Rot, ECollisionChannel TraceChannel, const FCollisionShape& CollisionShape, const FCollisionQueryParams& Params = FCollisionQueryParams::DefaultQueryParam, const FCollisionResponseParams& ResponseParam = FCollisionResponseParams::DefaultResponseParam) const;
 
 	/**
-	 *  Sweep a sphere against the world and return all initial overlaps using object types (including blocking) if requested, then overlapping hits and then first blocking hit
+	 *  Sweep a shape against the world and return all initial overlaps using object types (including blocking) if requested, then overlapping hits and then first blocking hit
 	 *  Results are sorted, so a blocking hit (if found) will be the last element of the array
 	 *  Only the single closest blocking result will be generated, no tests will be done after that
 	 *  @param  OutHits         Array of hits found between ray and the world
-	 *  @param  Start           Start location of the sphere
-	 *  @param  End             End location of the sphere
+	 *  @param  Start           Start location of the shape
+	 *  @param  End             End location of the shape
 	 *  @param	CollisionShape	CollisionShape - supports Box, Sphere, Capsule
 	 *  @param  Params          Additional parameters used for the trace
 	 *	@param	ObjectQueryParams	List of object types it's looking for
@@ -1465,12 +1465,12 @@ public:
 	}
 
 	/**
-	 *  Sweep a sphere against the world and return all initial overlaps using object types (including blocking) if requested, then overlapping hits and then first blocking hit
+	 *  Sweep a shape against the world and return all initial overlaps using object types (including blocking) if requested, then overlapping hits and then first blocking hit
 	 *  Results are sorted, so a blocking hit (if found) will be the last element of the array
 	 *  Only the single closest blocking result will be generated, no tests will be done after that
 	 *  @param  OutHits         Array of hits found between ray and the world
-	 *  @param  Start           Start location of the sphere
-	 *  @param  End             End location of the sphere
+	 *  @param  Start           Start location of the shape
+	 *  @param  End             End location of the shape
 	 *	@param	ObjectQueryParams	List of object types it's looking for
 	 *  @param	CollisionShape	CollisionShape - supports Box, Sphere, Capsule
 	 *  @param  Params          Additional parameters used for the trace
@@ -1479,12 +1479,12 @@ public:
 	bool SweepMultiByObjectType(TArray<struct FHitResult>& OutHits, const FVector& Start, const FVector& End, const FQuat& Rot, const FCollisionObjectQueryParams& ObjectQueryParams, const FCollisionShape& CollisionShape, const FCollisionQueryParams& Params = FCollisionQueryParams::DefaultQueryParam) const;
 
 	/**
-	 *  Sweep a sphere against the world and return all initial overlaps using a specific profile, then overlapping hits and then first blocking hit
+	 *  Sweep a shape against the world and return all initial overlaps using a specific profile, then overlapping hits and then first blocking hit
 	 *  Results are sorted, so a blocking hit (if found) will be the last element of the array
 	 *  Only the single closest blocking result will be generated, no tests will be done after that
 	 *  @param  OutHits         Array of hits found between ray and the world
-	 *  @param  Start           Start location of the sphere
-	 *  @param  End             End location of the sphere
+	 *  @param  Start           Start location of the shape
+	 *  @param  End             End location of the shape
 	 *  @param  ProfileName     The 'profile' used to determine which components to hit
 	 *  @param	CollisionShape	CollisionShape - supports Box, Sphere, Capsule
 	 *  @param  Params          Additional parameters used for the trace
@@ -1590,7 +1590,7 @@ public:
 	/**
 	 *  Test the collision of a shape at the supplied location using a specific channel, and determine the set of components that it overlaps
 	 *  @param  OutOverlap      Component found to overlap supplied shape
-	 *  @param  Pos             Location of center of sphere to test against the world
+	 *  @param  Pos             Location of center of shape to test against the world
 	 *  @param  TraceChannel    The 'channel' that this query is in, used to determine which components to hit
 	 *  @param	CollisionShape	CollisionShape - supports Box, Sphere, Capsule
 	 *  @param  Params          Additional parameters used for the trace
@@ -1603,7 +1603,7 @@ public:
 	/**
 	 *  Test the collision of a shape at the supplied location using object types, and determine the set of components that it overlaps
 	 *  @param  OutOverlap      Component found to overlap supplied shape
-	 *  @param  Pos             Location of center of sphere to test against the world
+	 *  @param  Pos             Location of center of shape to test against the world
 	 *  @param	CollisionShape	CollisionShape - supports Box, Sphere, Capsule
 	 *  @param  Params          Additional parameters used for the trace
 	 *	@param	ObjectQueryParams	List of object types it's looking for
@@ -1615,7 +1615,7 @@ public:
 	/**
 	 *  Test the collision of a shape at the supplied location using a specific profile, and determine the set of components that it overlaps
 	 *  @param  OutOverlap      Component found to overlap supplied shape
-	 *  @param  Pos             Location of center of sphere to test against the world
+	 *  @param  Pos             Location of center of shape to test against the world
 	 *  @param  ProfileName     The 'profile' used to determine which components to hit
 	 *  @param	CollisionShape	CollisionShape - supports Box, Sphere, Capsule
 	 *  @param  Params          Additional parameters used for the trace
@@ -1627,7 +1627,7 @@ public:
 	/**
 	 *  Test the collision of a shape at the supplied location using a specific channel, and determine the set of components that it overlaps
 	 *  @param  OutOverlaps     Array of components found to overlap supplied box
-	 *  @param  Pos             Location of center of sphere to test against the world
+	 *  @param  Pos             Location of center of shape to test against the world
 	 *  @param  TraceChannel    The 'channel' that this query is in, used to determine which components to hit
 	 *  @param	CollisionShape	CollisionShape - supports Box, Sphere, Capsule
 	 *  @param  Params          Additional parameters used for the trace
@@ -1643,7 +1643,7 @@ public:
 	/**
 	 *  Test the collision of a shape at the supplied location using a specific channel, and determine the set of components that it overlaps
 	 *  @param  OutOverlaps     Array of components found to overlap supplied box
-	 *  @param  Pos             Location of center of sphere to test against the world
+	 *  @param  Pos             Location of center of shape to test against the world
 	 *  @param  TraceChannel    The 'channel' that this query is in, used to determine which components to hit
 	 *  @param	CollisionShape	CollisionShape - supports Box, Sphere, Capsule
 	 *  @param  Params          Additional parameters used for the trace
@@ -1655,7 +1655,7 @@ public:
 	/**
 	 *  Test the collision of a shape at the supplied location using object types, and determine the set of components that it overlaps
 	 *  @param  OutOverlaps     Array of components found to overlap supplied box
-	 *  @param  Pos             Location of center of sphere to test against the world
+	 *  @param  Pos             Location of center of shape to test against the world
 	 *  @param	CollisionShape	CollisionShape - supports Box, Sphere, Capsule
 	 *  @param  Params          Additional parameters used for the trace
 	 *	@param	ObjectQueryParams	List of object types it's looking for
@@ -1670,7 +1670,7 @@ public:
 	/**
 	 *  Test the collision of a shape at the supplied location using object types, and determine the set of components that it overlaps
 	 *  @param  OutOverlaps     Array of components found to overlap supplied box
-	 *  @param  Pos             Location of center of sphere to test against the world
+	 *  @param  Pos             Location of center of shape to test against the world
 	 *	@param	ObjectQueryParams	List of object types it's looking for
 	 *  @param	CollisionShape	CollisionShape - supports Box, Sphere, Capsule
 	 *  @param  Params          Additional parameters used for the trace
@@ -1681,7 +1681,7 @@ public:
 	/**
 	 *  Test the collision of a shape at the supplied location using a specific profile, and determine the set of components that it overlaps
 	 *  @param  OutOverlaps     Array of components found to overlap supplied box
-	 *  @param  Pos             Location of center of sphere to test against the world
+	 *  @param  Pos             Location of center of shape to test against the world
 	 *  @param  ProfileName     The 'profile' used to determine which components to hit
 	 *  @param	CollisionShape	CollisionShape - supports Box, Sphere, Capsule
 	 *  @param  Params          Additional parameters used for the trace
@@ -1848,8 +1848,8 @@ public:
 	 * if no delegate, you can query trace data using QueryTraceData or QueryOverlapData
 	 * the data is available only in the next frame after request is made - in other words, if request is made in frame X, you can get the result in frame (X+1)
 	 *
-	 *  @param  Start           Start location of the sphere
-	 *  @param  End             End location of the sphere
+	 *  @param  Start           Start location of the shape
+	 *  @param  End             End location of the shape
 	 *  @param  TraceChannel    The 'channel' that this trace is in, used to determine which components to hit
 	 *  @param	CollisionShape		CollisionShape - supports Box, Sphere, Capsule
 	 *  @param  Params          Additional parameters used for the trace
@@ -1876,8 +1876,8 @@ public:
 	 * if no delegate, you can query trace data using QueryTraceData or QueryOverlapData
 	 * the data is available only in the next frame after request is made - in other words, if request is made in frame X, you can get the result in frame (X+1)
 	 *
-	 *  @param  Start           Start location of the sphere
-	 *  @param  End             End location of the sphere
+	 *  @param  Start           Start location of the shape
+	 *  @param  End             End location of the shape
 	 *  @param  TraceChannel    The 'channel' that this trace is in, used to determine which components to hit
 	 *  @param	CollisionShape		CollisionShape - supports Box, Sphere, Capsule
 	 *  @param  Params          Additional parameters used for the trace
@@ -1900,8 +1900,8 @@ public:
 	 * if no delegate, you can query trace data using QueryTraceData or QueryOverlapData
 	 * the data is available only in the next frame after request is made - in other words, if request is made in frame X, you can get the result in frame (X+1)
 	 *
-	 *  @param  Start           Start location of the sphere
-	 *  @param  End             End location of the sphere
+	 *  @param  Start           Start location of the shape
+	 *  @param  End             End location of the shape
 	 *  @param	CollisionShape		CollisionShape - supports Box, Sphere, Capsule
 	 *  @param  Params          Additional parameters used for the trace
 	 *	@param	ObjectQueryParams	List of object types it's looking for
@@ -1928,8 +1928,8 @@ public:
 	 * if no delegate, you can query trace data using QueryTraceData or QueryOverlapData
 	 * the data is available only in the next frame after request is made - in other words, if request is made in frame X, you can get the result in frame (X+1)
 	 *
-	 *  @param  Start           Start location of the sphere
-	 *  @param  End             End location of the sphere
+	 *  @param  Start           Start location of the shape
+	 *  @param  End             End location of the shape
 	 *	@param	ObjectQueryParams	List of object types it's looking for
 	 *  @param	CollisionShape		CollisionShape - supports Box, Sphere, Capsule
 	 *  @param  Params          Additional parameters used for the trace
@@ -1953,7 +1953,7 @@ public:
 	 * if no delegate, you can query trace data using QueryTraceData or QueryOverlapData
 	 * the data is available only in the next frame after request is made - in other words, if request is made in frame X, you can get the result in frame (X+1)
 	 *
-	 *  @param  Pos             Location of center of sphere to test against the world
+	 *  @param  Pos             Location of center of shape to test against the world
 	 *	@param	bMultiTrace		true if you'd like to do multi trace, or false otherwise
 	 *  @param  TraceChannel    The 'channel' that this query is in, used to determine which components to hit
 	 *  @param	CollisionShape		CollisionShape - supports Box, Sphere, Capsule
@@ -1980,7 +1980,7 @@ public:
 	 * if no delegate, you can query trace data using QueryTraceData or QueryOverlapData
 	 * the data is available only in the next frame after request is made - in other words, if request is made in frame X, you can get the result in frame (X+1)
 	 *
-	 *  @param  Pos             Location of center of sphere to test against the world
+	 *  @param  Pos             Location of center of shape to test against the world
 	 *	@param	bMultiTrace		true if you'd like to do multi trace, or false otherwise
 	 *  @param  TraceChannel    The 'channel' that this query is in, used to determine which components to hit
 	 *  @param	CollisionShape		CollisionShape - supports Box, Sphere, Capsule
@@ -2003,7 +2003,7 @@ public:
 	 * if no delegate, you can query trace data using QueryTraceData or QueryOverlapData
 	 * the data is available only in the next frame after request is made - in other words, if request is made in frame X, you can get the result in frame (X+1)
 	 *
-	 *  @param  Pos             Location of center of sphere to test against the world
+	 *  @param  Pos             Location of center of shape to test against the world
 	 *  @param  TraceChannel    The 'channel' that this query is in, used to determine which components to hit
 	 *  @param	CollisionShape		CollisionShape - supports Box, Sphere, Capsule
 	 *  @param  Params          Additional parameters used for the trace
@@ -2029,7 +2029,7 @@ public:
 	 * if no delegate, you can query trace data using QueryTraceData or QueryOverlapData
 	 * the data is available only in the next frame after request is made - in other words, if request is made in frame X, you can get the result in frame (X+1)
 	 *
-	 *  @param  Pos             Location of center of sphere to test against the world
+	 *  @param  Pos             Location of center of shape to test against the world
 	 *	@param	ObjectQueryParams	List of object types it's looking for
 	 *  @param	CollisionShape		CollisionShape - supports Box, Sphere, Capsule
 	 *  @param  Params          Additional parameters used for the trace
