Fixed a crash when a level that has a MatineeActor is unloaded at runtime

https://jira.ol.epicgames.net/browse/UE-16549

diff --git a/Engine/Source/Runtime/Engine/Classes/Matinee/MatineeActor.h b/Engine/Source/Runtime/Engine/Classes/Matinee/MatineeActor.h
--- a/Engine/Source/Runtime/Engine/Classes/Matinee/MatineeActor.h
+++ b/Engine/Source/Runtime/Engine/Classes/Matinee/MatineeActor.h
@@ -267,6 +267,7 @@ public:
 	ENGINE_API void OnObjectsReplaced(const TMap<UObject*,UObject*>& ReplacementMap);
 #endif //WITH_EDITOR
 
+	virtual void Destroyed() override;
 protected:
 
 	/** Handle for efficient management of CheckPriorityRefresh timer */
diff --git a/Engine/Source/Runtime/Engine/Private/Interpolation.cpp b/Engine/Source/Runtime/Engine/Private/Interpolation.cpp
--- a/Engine/Source/Runtime/Engine/Private/Interpolation.cpp
+++ b/Engine/Source/Runtime/Engine/Private/Interpolation.cpp
@@ -536,6 +536,16 @@ void AMatineeActor::OnObjectsReplaced(const TMap<UObject*,UObject*>& Replacement
 }
 #endif //WITH_EDITOR
 
+void AMatineeActor::Destroyed()
+{
+	if( GetWorld() )
+	{
+		GetWorldTimerManager().ClearTimer(TimerHandle_CheckPriorityRefresh);
+	}
+
+	Super::Destroyed();
+}
+
 void AMatineeActor::SetPosition(float NewPosition,bool bJump)
 {
 	// if we aren't currently active, temporarily activate to change the position
