#ue4 - Fix characters not rotating with a base object that is moving at a very low rotation rate in games with a high framerate.

- Quat comparison tolerance was too strict.

diff --git a/Engine/Source/Runtime/Engine/Private/Components/CharacterMovementComponent.cpp b/Engine/Source/Runtime/Engine/Private/Components/CharacterMovementComponent.cpp
--- a/Engine/Source/Runtime/Engine/Private/Components/CharacterMovementComponent.cpp
+++ b/Engine/Source/Runtime/Engine/Private/Components/CharacterMovementComponent.cpp
@@ -1374,7 +1374,7 @@ void UCharacterMovementComponent::UpdateBasedMovement(float DeltaSeconds)
 	}
 
 	// Find change in rotation
-	const bool bRotationChanged = !OldBaseQuat.Equals(NewBaseQuat);
+	const bool bRotationChanged = !OldBaseQuat.Equals(NewBaseQuat, 1e-6f);
 	if (bRotationChanged)
 	{
 		DeltaQuat = NewBaseQuat * OldBaseQuat.Inverse();
