#ue4 - Perf: Inline some frequently used Actor, FBodyInstance, and FCollisionResponse functions.

diff --git a/Engine/Source/Runtime/Engine/Classes/GameFramework/Actor.h b/Engine/Source/Runtime/Engine/Classes/GameFramework/Actor.h
--- a/Engine/Source/Runtime/Engine/Classes/GameFramework/Actor.h
+++ b/Engine/Source/Runtime/Engine/Classes/GameFramework/Actor.h
@@ -2511,6 +2511,34 @@ FORCEINLINE FVector AActor::GetSimpleCollisionCylinderExtent() const
 	return FVector(Radius, Radius, HalfHeight);
 }
 
+FORCEINLINE_DEBUGGABLE bool AActor::GetActorEnableCollision() const
+{
+	return bActorEnableCollision;
+}
+
+FORCEINLINE_DEBUGGABLE bool AActor::HasAuthority() const
+{
+	return (Role == ROLE_Authority);
+}
+
+FORCEINLINE_DEBUGGABLE AActor* AActor::GetOwner() const
+{ 
+	return Owner; 
+}
+
+FORCEINLINE_DEBUGGABLE const AActor* AActor::GetNetOwner() const
+{
+	// NetOwner is the Actor Owner unless otherwise overridden (see PlayerController/Pawn/Beacon)
+	// Used in ServerReplicateActors
+	return Owner;
+}
+
+FORCEINLINE_DEBUGGABLE ENetRole AActor::GetRemoteRole() const
+{
+	return RemoteRole;
+}
+
+
 //////////////////////////////////////////////////////////////////////////
 // Macro to hide common Transform functions in native code for classes where they don't make sense.
 // Note that this doesn't prevent access through function calls from parent classes (ie an AActor*), but
diff --git a/Engine/Source/Runtime/Engine/Classes/PhysicsEngine/BodyInstance.h b/Engine/Source/Runtime/Engine/Classes/PhysicsEngine/BodyInstance.h
--- a/Engine/Source/Runtime/Engine/Classes/PhysicsEngine/BodyInstance.h
+++ b/Engine/Source/Runtime/Engine/Classes/PhysicsEngine/BodyInstance.h
@@ -88,7 +88,7 @@ struct ENGINE_API FCollisionResponse
 	void ReplaceChannels(ECollisionResponse OldResponse, ECollisionResponse NewResponse);
 
 	/** Returns the response set on the specified channel */
-	ECollisionResponse GetResponse(ECollisionChannel Channel) const;
+	FORCEINLINE_DEBUGGABLE ECollisionResponse GetResponse(ECollisionChannel Channel) const { return ResponseToChannels.GetResponse(Channel); }
 	const FCollisionResponseContainer& GetResponseContainer() const { return ResponseToChannels; }
 
 	/** Set all channels from ChannelResponse Array **/
@@ -764,7 +764,7 @@ public:
 	void SetResponseToChannel(ECollisionChannel Channel, ECollisionResponse NewResponse);
 
 	/** Get the collision response of this body to a particular channel */
-	ECollisionResponse GetResponseToChannel(ECollisionChannel Channel) const;
+	FORCEINLINE_DEBUGGABLE ECollisionResponse GetResponseToChannel(ECollisionChannel Channel) const { return CollisionResponses.GetResponse(Channel); }
 
 	/** Set the response of this body to all channels */
 	void SetResponseToAllChannels(ECollisionResponse NewResponse);
@@ -776,13 +776,13 @@ public:
 	void SetResponseToChannels(const FCollisionResponseContainer& NewReponses);
 
 	/** Get Collision ResponseToChannels container for this component **/
-	const FCollisionResponseContainer& GetResponseToChannels() const;
+	FORCEINLINE_DEBUGGABLE const FCollisionResponseContainer& GetResponseToChannels() const { return CollisionResponses.GetResponseContainer(); }
 
 	/** Set the movement channel of this body to the one supplied */
 	void SetObjectType(ECollisionChannel Channel);
 
 	/** Get the movement channel of this body **/
-	ECollisionChannel GetObjectType() const;
+	FORCEINLINE_DEBUGGABLE ECollisionChannel GetObjectType() const { return ObjectType; }
 
 	/** Controls what kind of collision is enabled for this body and allows optional disable physics rebuild */
 	void SetCollisionEnabled(ECollisionEnabled::Type NewType, bool bUpdatePhysicsFilterData = true);
@@ -800,7 +800,7 @@ public:
 	void SetCollisionProfileName(FName InCollisionProfileName);
 
 	/** Get the current collision profile assigned to this body */
-	FName GetCollisionProfileName() const;
+	FORCEINLINE_DEBUGGABLE FName GetCollisionProfileName() const { return CollisionProfileName; }
 
 	/** return true if it uses Collision Profile System. False otherwise*/
 	bool DoesUseCollisionProfile() const;
diff --git a/Engine/Source/Runtime/Engine/Private/Actor.cpp b/Engine/Source/Runtime/Engine/Private/Actor.cpp
--- a/Engine/Source/Runtime/Engine/Private/Actor.cpp
+++ b/Engine/Source/Runtime/Engine/Private/Actor.cpp
@@ -1212,18 +1212,6 @@ void AActor::SetOwner( AActor *NewOwner )
 	}
 }
 
-AActor* AActor::GetOwner() const
-{ 
-	return Owner; 
-}
-
-const AActor* AActor::GetNetOwner() const
-{
-	// NetOwner is the Actor Owner unless otherwise overridden (see PlayerController/Pawn/Beacon)
-	// Used in ServerReplicateActors
-	return Owner;
-}
-
 bool AActor::HasNetOwner() const
 {
 	if (Owner == NULL)
@@ -2554,11 +2542,6 @@ void AActor::CopyRemoteRoleFrom(const AActor* CopyFromActor)
 	}
 }
 
-ENetRole AActor::GetRemoteRole() const
-{
-	return RemoteRole;
-}
-
 void AActor::PostNetInit()
 {
 	if(RemoteRole != ROLE_Authority)
@@ -3004,11 +2987,6 @@ void AActor::SetActorEnableCollision(bool bNewActorEnableCollision)
 }
 
 
-bool AActor::GetActorEnableCollision() const
-{
-	return bActorEnableCollision;
-}
-
 bool AActor::Destroy( bool bNetForce, bool bShouldModifyLevel )
 {
 	// It's already pending kill or in DestroyActor(), no need to beat the corpse
@@ -3033,11 +3011,6 @@ void AActor::K2_DestroyActor()
 	Destroy();
 }
 
-bool AActor::HasAuthority() const
-{
-	return (Role == ROLE_Authority);
-}
-
 void AActor::K2_DestroyComponent(UActorComponent* Component)
 {
 	// If its a valid component, and we own it, destroy it
diff --git a/Engine/Source/Runtime/Engine/Private/PhysicsEngine/BodyInstance.cpp b/Engine/Source/Runtime/Engine/Private/PhysicsEngine/BodyInstance.cpp
--- a/Engine/Source/Runtime/Engine/Private/PhysicsEngine/BodyInstance.cpp
+++ b/Engine/Source/Runtime/Engine/Private/PhysicsEngine/BodyInstance.cpp
@@ -87,12 +87,6 @@ void FCollisionResponse::ReplaceChannels(ECollisionResponse OldResponse, ECollis
 #endif
 }
 
-/** Returns the response set on the specified channel */
-ECollisionResponse FCollisionResponse::GetResponse(ECollisionChannel Channel) const
-{
-	return ResponseToChannels.GetResponse(Channel);
-}
-
 /** Set all channels from ChannelResponse Array **/
 void FCollisionResponse::SetCollisionResponseContainer(const FCollisionResponseContainer& InResponseToChannels)
 {
@@ -422,11 +416,6 @@ void FBodyInstance::InvalidateCollisionProfileName()
 	CollisionProfileName = UCollisionProfile::CustomCollisionProfileName;
 }
 
-ECollisionResponse FBodyInstance::GetResponseToChannel(ECollisionChannel Channel) const
-{
-	return CollisionResponses.GetResponse(Channel);
-}
-
 void FBodyInstance::SetResponseToChannel(ECollisionChannel Channel, ECollisionResponse NewResponse)
 {
 	InvalidateCollisionProfileName();
@@ -435,11 +424,6 @@ void FBodyInstance::SetResponseToChannel(ECollisionChannel Channel, ECollisionRe
 	UpdatePhysicsFilterData();
 }
 
-const FCollisionResponseContainer& FBodyInstance::GetResponseToChannels() const
-{
-	return CollisionResponses.GetResponseContainer();
-}
-
 void FBodyInstance::SetResponseToAllChannels(ECollisionResponse NewResponse)
 {
 	InvalidateCollisionProfileName();
@@ -470,11 +454,6 @@ void FBodyInstance::SetObjectType(ECollisionChannel Channel)
 	UpdatePhysicsFilterData();
 }
 
-ECollisionChannel FBodyInstance::GetObjectType() const
-{
-	return ObjectType;
-}
-
 void FBodyInstance::SetCollisionProfileName(FName InCollisionProfileName)
 {
 	if (CollisionProfileName != InCollisionProfileName)
@@ -485,11 +464,6 @@ void FBodyInstance::SetCollisionProfileName(FName InCollisionProfileName)
 	}
 }
 
-FName FBodyInstance::GetCollisionProfileName() const
-{
-	return CollisionProfileName;
-}
-
 bool FBodyInstance::DoesUseCollisionProfile() const
 {
 	return IsValidCollisionProfileName(CollisionProfileName);
diff --git a/Engine/Source/Runtime/Engine/Private/PhysicsEngine/PhysUtils.cpp b/Engine/Source/Runtime/Engine/Private/PhysicsEngine/PhysUtils.cpp
--- a/Engine/Source/Runtime/Engine/Private/PhysicsEngine/PhysUtils.cpp
+++ b/Engine/Source/Runtime/Engine/Private/PhysicsEngine/PhysUtils.cpp
@@ -227,12 +227,6 @@ void FCollisionResponseContainer::ReplaceChannels(ECollisionResponse OldResponse
 	}
 }
 
-/** Returns if the specified channel is set */
-ECollisionResponse FCollisionResponseContainer::GetResponse(ECollisionChannel Channel) const
-{
-	return (ECollisionResponse)EnumArray[Channel];
-}
-
 FCollisionResponseContainer FCollisionResponseContainer::CreateMinContainer(const FCollisionResponseContainer& A, const FCollisionResponseContainer& B)
 {
 	FCollisionResponseContainer Result;
