#ue4 - Perf: Remove GetSafeNormalPrecise() and GetClampedToMaxSizePrecise() from CharacterMovementComponent source. The stock FVector methods are now sufficiently precise after my InvSqrt() changes.

diff --git a/Engine/Source/Runtime/Engine/Private/Components/CharacterMovementComponent.cpp b/Engine/Source/Runtime/Engine/Private/Components/CharacterMovementComponent.cpp
--- a/Engine/Source/Runtime/Engine/Private/Components/CharacterMovementComponent.cpp
+++ b/Engine/Source/Runtime/Engine/Private/Components/CharacterMovementComponent.cpp
@@ -113,38 +113,6 @@ static TAutoConsoleVariable<float> CVarNetCorrectionLifetime(
 	ECVF_Cheat);
 #endif // !(UE_BUILD_SHIPPING || UE_BUILD_TEST)
 
-// Version that does not use inverse sqrt estimate, for higher precision.
-FORCEINLINE FVector GetSafeNormalPrecise(const FVector& V)
-{
-	const float VSq = V.SizeSquared();
-	if (VSq < SMALL_NUMBER)
-	{
-		return FVector::ZeroVector;
-	}
-	else
-	{
-		return V * (1.f / FMath::Sqrt(VSq));
-	}
-}
-
-// Version that does not use inverse sqrt estimate, for higher precision.
-FORCEINLINE FVector GetClampedToMaxSizePrecise(const FVector& V, float MaxSize)
-{
-	if (MaxSize < KINDA_SMALL_NUMBER)
-	{
-		return FVector::ZeroVector;
-	}
-
-	const float VSq = V.SizeSquared();
-	if (VSq > FMath::Square(MaxSize))
-	{
-		return V * (MaxSize / FMath::Sqrt(VSq));
-	}
-	else
-	{
-		return V;
-	}
-}
 
 
 void FFindFloorResult::SetFromSweep(const FHitResult& InHit, const float InSweepFloorDist, const bool bIsWalkableFloor)
@@ -2306,11 +2274,11 @@ void UCharacterMovementComponent::CalcVelocity(float DeltaTime, float Friction,
 		// In consideration order for direction: Acceleration, then Velocity, then Pawn's rotation.
 		if (Acceleration.SizeSquared() > SMALL_NUMBER)
 		{
-			Acceleration = GetSafeNormalPrecise(Acceleration) * MaxAccel;
+			Acceleration = Acceleration.GetSafeNormal() * MaxAccel;
 		}
 		else 
 		{
-			Acceleration = MaxAccel * (Velocity.SizeSquared() < SMALL_NUMBER ? UpdatedComponent->GetForwardVector() : GetSafeNormalPrecise(Velocity));
+			Acceleration = MaxAccel * (Velocity.SizeSquared() < SMALL_NUMBER ? UpdatedComponent->GetForwardVector() : Velocity.GetSafeNormal());
 		}
 
 		AnalogInputModifier = 1.f;
@@ -2333,13 +2301,13 @@ void UCharacterMovementComponent::CalcVelocity(float DeltaTime, float Friction,
 		// Don't allow braking to lower us below max speed if we started above it.
 		if (bVelocityOverMax && Velocity.SizeSquared() < FMath::Square(MaxSpeed) && FVector::DotProduct(Acceleration, OldVelocity) > 0.0f)
 		{
-			Velocity = GetSafeNormalPrecise(OldVelocity) * MaxSpeed;
+			Velocity = OldVelocity.GetSafeNormal() * MaxSpeed;
 		}
 	}
 	else if (!bZeroAcceleration)
 	{
 		// Friction affects our ability to change direction. This is only done for input acceleration, not path following.
-		const FVector AccelDir = GetSafeNormalPrecise(Acceleration);
+		const FVector AccelDir = Acceleration.GetSafeNormal();
 		const float VelSize = Velocity.Size();
 		Velocity = Velocity - (Velocity - AccelDir * VelSize) * FMath::Min(DeltaTime * Friction, 1.f);
 	}
@@ -2354,7 +2322,7 @@ void UCharacterMovementComponent::CalcVelocity(float DeltaTime, float Friction,
 	const float NewMaxSpeed = (IsExceedingMaxSpeed(MaxSpeed)) ? Velocity.Size() : MaxSpeed;
 	Velocity += Acceleration * DeltaTime;
 	Velocity += RequestedAcceleration * DeltaTime;
-	Velocity = GetClampedToMaxSizePrecise(Velocity, NewMaxSpeed);
+	Velocity = Velocity.GetClampedToMaxSize(NewMaxSpeed);
 
 	if (bUseRVOAvoidance)
 	{
@@ -2728,7 +2696,7 @@ void UCharacterMovementComponent::ApplyVelocityBraking(float DeltaTime, float Fr
 	const float MaxTimeStep = (1.0f / 33.0f);
 
 	// Decelerate to brake to a stop
-	const FVector RevAccel = (bZeroBraking ? FVector::ZeroVector : (-BrakingDeceleration * GetSafeNormalPrecise(Velocity)));
+	const FVector RevAccel = (bZeroBraking ? FVector::ZeroVector : (-BrakingDeceleration * Velocity.GetSafeNormal()));
 	while( RemainingTime >= MIN_TICK_TIME )
 	{
 		// Zero friction uses constant deceleration, so no need for iteration.
@@ -5544,7 +5512,7 @@ FVector UCharacterMovementComponent::ConstrainInputAcceleration(const FVector& I
 
 FVector UCharacterMovementComponent::ScaleInputAcceleration(const FVector& InputAcceleration) const
 {
-	return GetMaxAcceleration() * GetClampedToMaxSizePrecise(InputAcceleration, 1.0f);
+	return GetMaxAcceleration() * InputAcceleration.GetClampedToMaxSize(1.0f);
 }
 
 
