#ue4 - Perf: Cache rotation conversions to avoid Quat<->Rotator conversion when possible in SceneComponent.

- Added FRotationConversionCache which caches Quat<->Rotator conversions and returns the cached conversion if the source rotation is the same.
- SceneComponent uses this cache for internal uses of RelativeRotation (FRotator) and calls to GetComponentRotation() (returning a FRotator from ComponentToWorld's quaternion).

diff --git a/Engine/Source/Runtime/Engine/Classes/Components/SceneComponent.h b/Engine/Source/Runtime/Engine/Classes/Components/SceneComponent.h
--- a/Engine/Source/Runtime/Engine/Classes/Components/SceneComponent.h
+++ b/Engine/Source/Runtime/Engine/Classes/Components/SceneComponent.h
@@ -173,14 +173,17 @@ private:
 	TWeakObjectPtr<class APhysicsVolume> PhysicsVolume;
 
 public:
+	/** Current bounds of this component */
+	FBoxSphereBounds Bounds;
+
 	/** Current transform of this component, relative to the world */
 	FTransform ComponentToWorld;
 
-	/** Current bounds of this component */
-	FBoxSphereBounds Bounds;
+private:
+	/** Cache that avoids Quat<->Rotator conversions if possible. Only to be used with ComponentToWorld.GetRotation(). */
+	FRotationConversionCache WorldRotationCache;
 
-	UPROPERTY()
-	FVector RelativeTranslation_DEPRECATED;
+public:
 
 	/** Location of this component relative to its parent */
 	UPROPERTY(EditAnywhere, BlueprintReadOnly, ReplicatedUsing=OnRep_Transform, Category = Transform)
@@ -190,6 +193,21 @@ public:
 	UPROPERTY(EditAnywhere, BlueprintReadOnly, ReplicatedUsing=OnRep_Transform, Category=Transform)
 	FRotator RelativeRotation;
 
+private:
+	/** Cache that avoids Quat<->Rotator conversions if possible. Only to be used with RelativeRotation. */
+	FRotationConversionCache RelativeRotationCache;
+
+public:
+	/**
+	*	Non-uniform scaling of this component relative to its parent.
+	*	Note that scaling is always applied in local space (no shearing etc)
+	*/
+	UPROPERTY(BlueprintReadOnly, ReplicatedUsing=OnRep_Transform, interp, Category=Transform)
+	FVector RelativeScale3D;
+
+	UPROPERTY()
+	FVector RelativeTranslation_DEPRECATED;
+
 	/** How often this component is allowed to move, used to make various optimizations. Only safe to set in constructor, use SetMobility() during runtime. */
 	UPROPERTY(Category = Mobility, EditAnywhere, BlueprintReadOnly)
 	TEnumAsByte<EComponentMobility::Type> Mobility;
@@ -216,13 +234,6 @@ private:
 
 public:
 
-	/** 
-	 *	Non-uniform scaling of this component relative to its parent. 
-	 *	Note that scaling is always applied in local space (no shearing etc)
-	 */
-	UPROPERTY(BlueprintReadOnly, ReplicatedUsing=OnRep_Transform, interp, Category=Transform)
-	FVector RelativeScale3D;
-
 	/**
 	 * Velocity of this component.
 	 * @see GetComponentVelocity()
@@ -642,7 +653,7 @@ public:
 	/** Return rotation of the component, in world space */
 	FORCEINLINE FRotator GetComponentRotation() const
 	{
-		return ComponentToWorld.Rotator();
+		return WorldRotationCache.NormalizedQuatToRotator(ComponentToWorld.GetRotation());
 	}
 
 	/** Return rotation quaternion of the component, in world space */
@@ -876,7 +887,7 @@ FORCEINLINE_DEBUGGABLE bool USceneComponent::MoveComponent(const FVector& Delta,
 
 FORCEINLINE_DEBUGGABLE void USceneComponent::SetRelativeLocation(FVector NewLocation, bool bSweep, FHitResult* OutSweepHitResult)
 {
-	SetRelativeLocationAndRotation(NewLocation, RelativeRotation, bSweep, OutSweepHitResult);
+	SetRelativeLocationAndRotation(NewLocation, RelativeRotationCache.RotatorToQuat(RelativeRotation), bSweep, OutSweepHitResult);
 }
 
 FORCEINLINE_DEBUGGABLE void USceneComponent::SetRelativeRotation(FRotator NewRotation, bool bSweep, FHitResult* OutSweepHitResult)
@@ -896,7 +907,7 @@ FORCEINLINE_DEBUGGABLE void USceneComponent::SetRelativeLocationAndRotation(FVec
 
 FORCEINLINE_DEBUGGABLE void USceneComponent::AddRelativeLocation(FVector DeltaLocation, bool bSweep, FHitResult* OutSweepHitResult)
 {
-	SetRelativeLocationAndRotation(RelativeLocation + DeltaLocation, RelativeRotation, bSweep, OutSweepHitResult);
+	SetRelativeLocationAndRotation(RelativeLocation + DeltaLocation, RelativeRotationCache.RotatorToQuat(RelativeRotation), bSweep, OutSweepHitResult);
 }
 
 FORCEINLINE_DEBUGGABLE void USceneComponent::AddRelativeRotation(FRotator DeltaRotation, bool bSweep, FHitResult* OutSweepHitResult)
diff --git a/Engine/Source/Runtime/Engine/Classes/Engine/EngineTypes.h b/Engine/Source/Runtime/Engine/Classes/Engine/EngineTypes.h
--- a/Engine/Source/Runtime/Engine/Classes/Engine/EngineTypes.h
+++ b/Engine/Source/Runtime/Engine/Classes/Engine/EngineTypes.h
@@ -775,7 +775,7 @@ struct ENGINE_API FCollisionResponseContainer
 	void ReplaceChannels(ECollisionResponse OldResponse, ECollisionResponse NewResponse);
 
 	/** Returns the response set on the specified channel */
-	ECollisionResponse GetResponse(ECollisionChannel Channel) const;
+	FORCEINLINE_DEBUGGABLE ECollisionResponse GetResponse(ECollisionChannel Channel) const { return (ECollisionResponse)EnumArray[Channel]; }
 
 	/** Set all channels from ChannelResponse Array **/
 	void UpdateResponsesFromArray(TArray<FResponseChannel> & ChannelResponses);
@@ -1070,6 +1070,85 @@ struct ENGINE_API FBasedPosition
 };
 
 
+/** Struct for caching Quat<->Rotator conversions. */
+struct ENGINE_API FRotationConversionCache
+{
+	FRotationConversionCache()
+		: CachedQuat(FQuat::Identity)
+		, CachedRotator(FRotator::ZeroRotator)
+	{
+	}
+
+	/** Convert a FRotator to FQuat. Uses the cached conversion if possible, and updates it if there was no match. */
+	FORCEINLINE_DEBUGGABLE FQuat FRotationConversionCache::RotatorToQuat(const FRotator& InRotator) const
+	{
+		if (CachedRotator != InRotator)
+		{
+			CachedRotator = InRotator.GetNormalized();;
+			CachedQuat = CachedRotator.Quaternion();
+		}
+		return CachedQuat;
+	}
+
+	/** Convert a FRotator to FQuat. Uses the cached conversion if possible, but does *NOT* update the cache if there was no match. */
+	FORCEINLINE_DEBUGGABLE FQuat FRotationConversionCache::RotatorToQuat_ReadOnly(const FRotator& InRotator) const
+	{
+		if (CachedRotator == InRotator)
+		{
+			return CachedQuat;
+		}
+		return InRotator.Quaternion();
+	}
+
+	/** Convert a FQuat to FRotator. Uses the cached conversion if possible, and updates it if there was no match. */
+	FORCEINLINE_DEBUGGABLE FRotator FRotationConversionCache::QuatToRotator(const FQuat& InQuat) const
+	{
+		if (CachedQuat != InQuat)
+		{
+			CachedQuat = InQuat.GetNormalized();
+			CachedRotator = CachedQuat.Rotator();
+		}
+		return CachedRotator;
+	}
+
+	/** Convert a FQuat to FRotator. Uses the cached conversion if possible, but does *NOT* update the cache if there was no match. */
+	FORCEINLINE_DEBUGGABLE FRotator FRotationConversionCache::QuatToRotator_ReadOnly(const FQuat& InQuat) const
+	{
+		if (CachedQuat == InQuat)
+		{
+			return CachedRotator;
+		}
+		return InQuat.Rotator();
+	}
+
+	/** Version of QuatToRotator when the Quat is known to already be normalized. */
+	FORCEINLINE_DEBUGGABLE FRotator FRotationConversionCache::NormalizedQuatToRotator(const FQuat& InNormalizedQuat) const
+	{
+		if (CachedQuat != InNormalizedQuat)
+		{
+			CachedQuat = InNormalizedQuat;
+			CachedRotator = InNormalizedQuat.Rotator();
+		}
+		return CachedRotator;
+	}
+
+	/** Version of QuatToRotator when the Quat is known to already be normalized. Does *NOT* update the cache if there was no match. */
+	FORCEINLINE_DEBUGGABLE FRotator FRotationConversionCache::NormalizedQuatToRotator_ReadOnly(const FQuat& InNormalizedQuat) const
+	{
+		if (CachedQuat == InNormalizedQuat)
+		{
+			return CachedRotator;
+		}
+		return InNormalizedQuat.Rotator();
+	}
+
+private:
+	mutable FQuat		CachedQuat;		// FQuat matching CachedRotator such that CachedQuat.Rotator() == CachedRotator.
+	mutable FRotator	CachedRotator;	// FRotator matching CachedQuat such that CachedRotator.Quaternion() == CachedQuat.
+};
+
+
+
 /** A line of subtitle text and the time at which it should be displayed. */
 USTRUCT()
 struct FSubtitleCue
diff --git a/Engine/Source/Runtime/Engine/Private/Components/SceneComponent.cpp b/Engine/Source/Runtime/Engine/Private/Components/SceneComponent.cpp
--- a/Engine/Source/Runtime/Engine/Private/Components/SceneComponent.cpp
+++ b/Engine/Source/Runtime/Engine/Private/Components/SceneComponent.cpp
@@ -347,6 +347,10 @@ void USceneComponent::UpdateComponentToWorldWithParent(USceneComponent* Parent,
 	// Minimize accumulation of errors after many composed transforms.
 	NewTransform.NormalizeRotation();
 
+	// Update RelativeRotation and cache (not done earlier because normalization can change the rotation)
+	const FQuat RelativeQuatNormalized = (Parent ? RelativeRotationQuat.GetNormalized() : NewTransform.GetRotation());
+	RelativeRotation = RelativeRotationCache.NormalizedQuatToRotator(RelativeQuatNormalized);
+
 #if !(UE_BUILD_SHIPPING || UE_BUILD_TEST)
 	ensureOnce(NewTransform.IsValid());
 #endif
@@ -356,7 +360,6 @@ void USceneComponent::UpdateComponentToWorldWithParent(USceneComponent* Parent,
 	{
 		// Update transform
 		ComponentToWorld = NewTransform;
-
 		PropagateTransformUpdate(true, bSkipPhysicsMove, bTeleport);
 	}
 	else
@@ -399,7 +402,7 @@ void USceneComponent::OnRegister()
 
 void USceneComponent::UpdateComponentToWorld(bool bSkipPhysicsMove, bool bTeleport)
 {
-	UpdateComponentToWorldWithParent(AttachParent, bSkipPhysicsMove, RelativeRotation.Quaternion(), bTeleport);
+	UpdateComponentToWorldWithParent(AttachParent, bSkipPhysicsMove, RelativeRotationCache.RotatorToQuat(RelativeRotation), bTeleport);
 }
 
 
@@ -734,34 +737,35 @@ void USceneComponent::SetRelativeLocationAndRotation(FVector NewLocation, const
 
 void USceneComponent::AddRelativeRotation(const FQuat& DeltaRotation, bool bSweep, FHitResult* OutSweepHitResult)
 {
-	const FQuat CurRelRotQuat = RelativeRotation.Quaternion();
+	const FQuat CurRelRotQuat = RelativeRotationCache.RotatorToQuat(RelativeRotation);
 	const FQuat NewRelRotQuat = DeltaRotation * CurRelRotQuat;
 	SetRelativeLocationAndRotation(RelativeLocation, NewRelRotQuat, bSweep, OutSweepHitResult);
 }
 
 void USceneComponent::AddLocalOffset(FVector DeltaLocation, bool bSweep, FHitResult* OutSweepHitResult)
 {
-	const FVector LocalOffset = RelativeRotation.RotateVector(DeltaLocation);
-	SetRelativeLocationAndRotation(RelativeLocation + LocalOffset, RelativeRotation, bSweep, OutSweepHitResult);
+	const FQuat CurRelRotQuat = RelativeRotationCache.RotatorToQuat(RelativeRotation);
+	const FVector LocalOffset = CurRelRotQuat.RotateVector(DeltaLocation);
+	SetRelativeLocationAndRotation(RelativeLocation + LocalOffset, CurRelRotQuat, bSweep, OutSweepHitResult);
 }
 
 void USceneComponent::AddLocalRotation(FRotator DeltaRotation, bool bSweep, FHitResult* OutSweepHitResult)
 {
-	const FQuat CurRelRotQuat = RelativeRotation.Quaternion();
+	const FQuat CurRelRotQuat = RelativeRotationCache.RotatorToQuat(RelativeRotation);
 	const FQuat NewRelRotQuat = CurRelRotQuat * DeltaRotation.Quaternion();
 	SetRelativeLocationAndRotation(RelativeLocation, NewRelRotQuat, bSweep, OutSweepHitResult);
 }
 
 void USceneComponent::AddLocalRotation(const FQuat& DeltaRotation, bool bSweep, FHitResult* OutSweepHitResult)
 {
-	const FQuat CurRelRotQuat = RelativeRotation.Quaternion();
+	const FQuat CurRelRotQuat = RelativeRotationCache.RotatorToQuat(RelativeRotation);
 	const FQuat NewRelRotQuat = CurRelRotQuat * DeltaRotation;
 	SetRelativeLocationAndRotation(RelativeLocation, NewRelRotQuat, bSweep, OutSweepHitResult);
 }
 
 void USceneComponent::AddLocalTransform(const FTransform& DeltaTransform, bool bSweep, FHitResult* OutSweepHitResult)
 {
-	const FTransform RelativeTransform( RelativeRotation, RelativeLocation, FVector(1,1,1) ); // don't use scaling, so it matches how AddLocalRotation/Offset work
+	const FTransform RelativeTransform( RelativeRotationCache.RotatorToQuat(RelativeRotation), RelativeLocation, FVector(1,1,1) ); // don't use scaling, so it matches how AddLocalRotation/Offset work
 	const FTransform NewRelTransform = DeltaTransform * RelativeTransform;
 	SetRelativeTransform(NewRelTransform, bSweep, OutSweepHitResult);
 }
@@ -833,7 +837,7 @@ void USceneComponent::SetRelativeTransform(const FTransform& NewTransform, bool
 
 FTransform USceneComponent::GetRelativeTransform() const
 {
-	const FTransform RelativeTransform( RelativeRotation.Quaternion(), RelativeLocation, RelativeScale3D );
+	const FTransform RelativeTransform( RelativeRotationCache.RotatorToQuat(RelativeRotation), RelativeLocation, RelativeScale3D );
 	return RelativeTransform;
 }
 
@@ -966,7 +970,7 @@ void USceneComponent::SetWorldLocationAndRotationNoPhysics(const FVector& NewLoc
 		{
 			// Quat multiplication works reverse way, make sure you do Parent(-1) * World = Local, not World*Parent(-) = Local (the way matrix does)
 			const FQuat NewRelQuat = ParentToWorld.GetRotation().Inverse() * NewRotation.Quaternion();
-			RelativeRotation = NewRelQuat.Rotator();
+			RelativeRotation = RelativeRotationCache.QuatToRotator(NewRelQuat);
 		}
 	}
 	else
@@ -1187,10 +1191,10 @@ void USceneComponent::AttachTo(class USceneComponent* Parent, FName InSocketName
 				//Since the object is physically simulated it can't be the case that it's a child of object A and being attached to object B (at runtime)
 				if (bMaintainWorldPosition == false)	//User tried to attach but physically based so detach. However, if they provided relative coordinates we should still get the correct position
 				{
-					UpdateComponentToWorldWithParent(Parent, false, RelativeRotation.Quaternion());
-					RelativeLocation = ComponentToWorld.GetLocation();
-					RelativeRotation = ComponentToWorld.GetRotation().Rotator();
-					RelativeScale3D = ComponentToWorld.GetScale3D();
+					UpdateComponentToWorldWithParent(Parent, false, RelativeRotationCache.RotatorToQuat(RelativeRotation));
+					RelativeLocation = ComponentToWorld.GetLocation(); // or GetComponentLocation(), but worried about custom location...
+					RelativeRotation = GetComponentRotation();
+					RelativeScale3D = GetComponentScale();
 				}
 
 				return;
@@ -1235,11 +1239,11 @@ void USceneComponent::AttachTo(class USceneComponent* Parent, FName InSocketName
 
 				if(bAbsoluteRotation)
 				{
-					RelativeRotation = ComponentToWorld.Rotator();
+					RelativeRotation = GetComponentRotation();
 				}
 				else
 				{
-					RelativeRotation = RelativeTM.Rotator();
+					RelativeRotation = RelativeRotationCache.QuatToRotator(RelativeTM.GetRotation());
 				}
 				if(bAbsoluteScale)
 				{
@@ -1364,9 +1368,9 @@ void USceneComponent::DetachFromParent(bool bMaintainWorldPosition, bool bCallMo
 		// If desired, update RelativeLocation and RelativeRotation to maintain current world position after detachment
 		if(bMaintainWorldPosition)
 		{
-			RelativeLocation = ComponentToWorld.GetTranslation();
-			RelativeRotation = ComponentToWorld.Rotator();
-			RelativeScale3D = ComponentToWorld.GetScale3D();
+			RelativeLocation = ComponentToWorld.GetTranslation(); // or GetComponentLocation, but worried about custom location...
+			RelativeRotation = GetComponentRotation();
+			RelativeScale3D = GetComponentScale();
 		}
 
 		// calculate transform with new attachment condition
@@ -1825,6 +1829,7 @@ void USceneComponent::BeginDestroy()
 
 bool USceneComponent::InternalSetWorldLocationAndRotation(FVector NewLocation, const FQuat& RotationQuat, bool bNoPhysics, bool bTeleport)
 {
+	checkSlow(bWorldToComponentUpdated);
 	FQuat NewRotationQuat(RotationQuat);
 
 	// If attached to something, transform into local space
@@ -1844,11 +1849,9 @@ bool USceneComponent::InternalSetWorldLocationAndRotation(FVector NewLocation, c
 		}
 	}
 
-	const FRotator NewRotator = NewRotationQuat.Rotator();
-	if (!NewLocation.Equals(RelativeLocation) || !NewRotator.Equals(RelativeRotation))
+	if (!NewLocation.Equals(RelativeLocation) || !NewRotationQuat.Equals(RelativeRotationCache.RotatorToQuat_ReadOnly(RelativeRotation), 1e-6f))
 	{
 		RelativeLocation = NewLocation;
-		RelativeRotation = NewRotator;
 		UpdateComponentToWorldWithParent(AttachParent, bNoPhysics, NewRotationQuat, bTeleport);
 		return true;
 	}
@@ -2084,7 +2087,7 @@ void USceneComponent::ApplyWorldOffset(const FVector& InOffset, bool bWorldShift
 	// Calculate current ComponentToWorld transform
 	// We do this because at level load/duplication ComponentToWorld is uninitialized
 	{
-		const FTransform RelativeTransform(RelativeRotation, RelativeLocation, RelativeScale3D);
+		const FTransform RelativeTransform(RelativeRotationCache.RotatorToQuat(RelativeRotation), RelativeLocation, RelativeScale3D);
 		ComponentToWorld = CalcNewComponentToWorld(RelativeTransform);
 	}
 
@@ -2106,7 +2109,7 @@ void USceneComponent::ApplyWorldOffset(const FVector& InOffset, bool bWorldShift
 	RelativeLocation = NewLocation;
 		
 	// Calculate the new ComponentToWorld transform
-	const FTransform RelativeTransform(RelativeRotation, RelativeLocation, RelativeScale3D);
+	const FTransform RelativeTransform(RelativeRotationCache.RotatorToQuat(RelativeRotation), RelativeLocation, RelativeScale3D);
 	ComponentToWorld = CalcNewComponentToWorld(RelativeTransform);
 
 	// Physics move is skipped if physics state is not created or physics scene supports origin shifting
