#ue4 - FindBestOverlappingNormal() can just use GetSafeNormal(), it's sufficiently accurate now.

diff --git a/Engine/Source/Runtime/Engine/Private/Collision/CollisionConversions.cpp b/Engine/Source/Runtime/Engine/Private/Collision/CollisionConversions.cpp
--- a/Engine/Source/Runtime/Engine/Private/Collision/CollisionConversions.cpp
+++ b/Engine/Source/Runtime/Engine/Private/Collision/CollisionConversions.cpp
@@ -559,10 +559,7 @@ FVector FindBestOverlappingNormal(const UWorld* World, const PxGeometry& Geom, c
 		const FVector C = P2UVector(Tri.verts[2]);
 
 		FVector TriNormal = ((B - A) ^ (C - A));
-
-		// Use a more accurate normalization that avoids InvSqrtEst
-		const float TriNormalSize = TriNormal.Size();
-		TriNormal = (TriNormalSize >= KINDA_SMALL_NUMBER ? TriNormal / TriNormalSize : FVector::ZeroVector);
+		TriNormal = TriNormal.GetSafeNormal();
 
 		const FPlane TriPlane(A, TriNormal);
 
