Fixed Linux build (removed extra qualifiers on struct methods)

diff --git a/Engine/Source/Runtime/Engine/Classes/Engine/EngineTypes.h b/Engine/Source/Runtime/Engine/Classes/Engine/EngineTypes.h
--- a/Engine/Source/Runtime/Engine/Classes/Engine/EngineTypes.h
+++ b/Engine/Source/Runtime/Engine/Classes/Engine/EngineTypes.h
@@ -1080,7 +1080,7 @@ struct ENGINE_API FRotationConversionCache
 	}
 
 	/** Convert a FRotator to FQuat. Uses the cached conversion if possible, and updates it if there was no match. */
-	FORCEINLINE_DEBUGGABLE FQuat FRotationConversionCache::RotatorToQuat(const FRotator& InRotator) const
+	FORCEINLINE_DEBUGGABLE FQuat RotatorToQuat(const FRotator& InRotator) const
 	{
 		if (CachedRotator != InRotator)
 		{
@@ -1091,7 +1091,7 @@ struct ENGINE_API FRotationConversionCache
 	}
 
 	/** Convert a FRotator to FQuat. Uses the cached conversion if possible, but does *NOT* update the cache if there was no match. */
-	FORCEINLINE_DEBUGGABLE FQuat FRotationConversionCache::RotatorToQuat_ReadOnly(const FRotator& InRotator) const
+	FORCEINLINE_DEBUGGABLE FQuat RotatorToQuat_ReadOnly(const FRotator& InRotator) const
 	{
 		if (CachedRotator == InRotator)
 		{
@@ -1101,7 +1101,7 @@ struct ENGINE_API FRotationConversionCache
 	}
 
 	/** Convert a FQuat to FRotator. Uses the cached conversion if possible, and updates it if there was no match. */
-	FORCEINLINE_DEBUGGABLE FRotator FRotationConversionCache::QuatToRotator(const FQuat& InQuat) const
+	FORCEINLINE_DEBUGGABLE FRotator QuatToRotator(const FQuat& InQuat) const
 	{
 		if (CachedQuat != InQuat)
 		{
@@ -1112,7 +1112,7 @@ struct ENGINE_API FRotationConversionCache
 	}
 
 	/** Convert a FQuat to FRotator. Uses the cached conversion if possible, but does *NOT* update the cache if there was no match. */
-	FORCEINLINE_DEBUGGABLE FRotator FRotationConversionCache::QuatToRotator_ReadOnly(const FQuat& InQuat) const
+	FORCEINLINE_DEBUGGABLE FRotator QuatToRotator_ReadOnly(const FQuat& InQuat) const
 	{
 		if (CachedQuat == InQuat)
 		{
@@ -1122,7 +1122,7 @@ struct ENGINE_API FRotationConversionCache
 	}
 
 	/** Version of QuatToRotator when the Quat is known to already be normalized. */
-	FORCEINLINE_DEBUGGABLE FRotator FRotationConversionCache::NormalizedQuatToRotator(const FQuat& InNormalizedQuat) const
+	FORCEINLINE_DEBUGGABLE FRotator NormalizedQuatToRotator(const FQuat& InNormalizedQuat) const
 	{
 		if (CachedQuat != InNormalizedQuat)
 		{
@@ -1133,7 +1133,7 @@ struct ENGINE_API FRotationConversionCache
 	}
 
 	/** Version of QuatToRotator when the Quat is known to already be normalized. Does *NOT* update the cache if there was no match. */
-	FORCEINLINE_DEBUGGABLE FRotator FRotationConversionCache::NormalizedQuatToRotator_ReadOnly(const FQuat& InNormalizedQuat) const
+	FORCEINLINE_DEBUGGABLE FRotator NormalizedQuatToRotator_ReadOnly(const FQuat& InNormalizedQuat) const
 	{
 		if (CachedQuat == InNormalizedQuat)
 		{
