FMaterialUniformExpressionTextureParameter::GetGameThreadTextureValue no longer searches the base material for a matching parameter default texture value, as it already knows the parameter default texture value.  This reduces ULevel::BuildStreamingData editor hitches from 600ms to 30ms.  The previous implementation could return the wrong texture in the case of duplicately named parameters with different default texture values.

diff --git a/Engine/Source/Runtime/Engine/Classes/Materials/MaterialInstance.h b/Engine/Source/Runtime/Engine/Classes/Materials/MaterialInstance.h
--- a/Engine/Source/Runtime/Engine/Classes/Materials/MaterialInstance.h
+++ b/Engine/Source/Runtime/Engine/Classes/Materials/MaterialInstance.h
@@ -252,6 +252,8 @@ public:
 	virtual ENGINE_API bool GetFontParameterValue(FName ParameterName, class UFont*& OutFontValue, int32& OutFontPage) const override;
 	virtual ENGINE_API bool GetScalarParameterValue(FName ParameterName, float& OutValue) const override;
 	virtual ENGINE_API bool GetTextureParameterValue(FName ParameterName, class UTexture*& OutValue) const override;
+	/** Searches for a parameter override for the named parameter and returns true if one was found.  Does not search the base UMaterial. */
+	virtual ENGINE_API bool GetTextureParameterOverrideValue(FName ParameterName, class UTexture*& OutValue) const override;
 	virtual ENGINE_API bool GetVectorParameterValue(FName ParameterName, FLinearColor& OutValue) const override;
 	virtual ENGINE_API void GetUsedTextures(TArray<UTexture*>& OutTextures, EMaterialQualityLevel::Type QualityLevel, bool bAllQualityLevels, ERHIFeatureLevel::Type FeatureLevel, bool bAllFeatureLevels) const override;
 	virtual ENGINE_API void OverrideTexture(const UTexture* InTextureToOverride, UTexture* OverrideTexture, ERHIFeatureLevel::Type InFeatureLevel) override;
diff --git a/Engine/Source/Runtime/Engine/Classes/Materials/MaterialInterface.h b/Engine/Source/Runtime/Engine/Classes/Materials/MaterialInterface.h
--- a/Engine/Source/Runtime/Engine/Classes/Materials/MaterialInterface.h
+++ b/Engine/Source/Runtime/Engine/Classes/Materials/MaterialInterface.h
@@ -531,6 +531,7 @@ public:
 	ENGINE_API virtual bool GetScalarParameterValue(FName ParameterName, float& OutValue) const;
 	ENGINE_API virtual bool GetScalarCurveParameterValue(FName ParameterName, FInterpCurveFloat& OutValue) const;
 	ENGINE_API virtual bool GetTextureParameterValue(FName ParameterName, class UTexture*& OutValue) const;
+	ENGINE_API virtual bool GetTextureParameterOverrideValue(FName ParameterName, class UTexture*& OutValue) const;
 	ENGINE_API virtual bool GetVectorParameterValue(FName ParameterName, FLinearColor& OutValue) const;
 	ENGINE_API virtual bool GetVectorCurveParameterValue(FName ParameterName, FInterpCurveVector& OutValue) const;
 	ENGINE_API virtual bool GetLinearColorParameterValue(FName ParameterName, FLinearColor& OutValue) const;
diff --git a/Engine/Source/Runtime/Engine/Private/Materials/Material.cpp b/Engine/Source/Runtime/Engine/Private/Materials/Material.cpp
--- a/Engine/Source/Runtime/Engine/Private/Materials/Material.cpp
+++ b/Engine/Source/Runtime/Engine/Private/Materials/Material.cpp
@@ -1356,6 +1356,7 @@ bool UMaterial::GetVectorParameterValue(FName ParameterName, FLinearColor& OutVa
 			const UMaterialExpressionVectorParameter* Parameter = CastChecked<const UMaterialExpressionVectorParameter>(Expression);
 			if (Parameter->IsNamedParameter(ParameterName, OutValue))
 			{
+				// Warning: in the case of duplicate parameters with different default values, this will find the first in the expression array, not necessarily the one that's used for rendering
 				return true;
 			}
 		}
@@ -1398,6 +1399,7 @@ bool UMaterial::GetScalarParameterValue(FName ParameterName, float& OutValue) co
 			const UMaterialExpressionScalarParameter* Parameter = CastChecked<const UMaterialExpressionScalarParameter>(Expression);
 			if (Parameter->IsNamedParameter(ParameterName, OutValue))
 			{
+				// Warning: in the case of duplicate parameters with different default values, this will find the first in the expression array, not necessarily the one that's used for rendering
 				return true;
 			}
 		}
@@ -1440,6 +1442,7 @@ bool UMaterial::GetTextureParameterValue(FName ParameterName, UTexture*& OutValu
 			const UMaterialExpressionTextureSampleParameter* Parameter = CastChecked<const UMaterialExpressionTextureSampleParameter>(Expression);
 			if (Parameter->IsNamedParameter(ParameterName, OutValue))
 			{
+				// Warning: in the case of duplicate parameters with different default values, this will find the first in the expression array, not necessarily the one that's used for rendering
 				return true;
 			}
 		}
@@ -1482,6 +1485,7 @@ bool UMaterial::GetFontParameterValue(FName ParameterName, UFont*& OutFontValue,
 			const UMaterialExpressionFontSampleParameter* Parameter = CastChecked<const UMaterialExpressionFontSampleParameter>(Expression);
 			if (Parameter->IsNamedParameter(ParameterName, OutFontValue, OutFontPage))
 			{
+				// Warning: in the case of duplicate parameters with different default values, this will find the first in the expression array, not necessarily the one that's used for rendering
 				return true;
 			}
 		}
@@ -1525,6 +1529,7 @@ bool UMaterial::GetStaticSwitchParameterValue(FName ParameterName, bool& OutValu
 			const UMaterialExpressionStaticBoolParameter* Parameter = CastChecked<const UMaterialExpressionStaticBoolParameter>(Expression);
 			if (Parameter->IsNamedParameter(ParameterName, OutValue, OutExpressionGuid))
 			{
+				// Warning: in the case of duplicate parameters with different default values, this will find the first in the expression array, not necessarily the one that's used for rendering
 				return true;
 			}
 		}
@@ -1568,6 +1573,7 @@ bool UMaterial::GetStaticComponentMaskParameterValue(FName ParameterName, bool&
 			const UMaterialExpressionStaticComponentMaskParameter* Parameter = CastChecked<const UMaterialExpressionStaticComponentMaskParameter>(Expression);
 			if (Parameter->IsNamedParameter(ParameterName, OutR, OutG, OutB, OutA, OutExpressionGuid))
 			{
+				// Warning: in the case of duplicate parameters with different default values, this will find the first in the expression array, not necessarily the one that's used for rendering
 				return true;
 			}
 		}
diff --git a/Engine/Source/Runtime/Engine/Private/Materials/MaterialInstance.cpp b/Engine/Source/Runtime/Engine/Private/Materials/MaterialInstance.cpp
--- a/Engine/Source/Runtime/Engine/Private/Materials/MaterialInstance.cpp
+++ b/Engine/Source/Runtime/Engine/Private/Materials/MaterialInstance.cpp
@@ -610,6 +610,30 @@ bool UMaterialInstance::GetTextureParameterValue(FName ParameterName, UTexture*&
 	}
 }
 
+bool UMaterialInstance::GetTextureParameterOverrideValue(FName ParameterName, UTexture*& OutValue) const
+{
+	if(ReentrantFlag)
+	{
+		return false;
+	}
+
+	const FTextureParameterValue* ParameterValue = GameThread_FindParameterByName(TextureParameterValues,ParameterName);
+	if(ParameterValue && ParameterValue->ParameterValue)
+	{
+		OutValue = ParameterValue->ParameterValue;
+		return true;
+	}
+	else if(Parent)
+	{
+		FMICReentranceGuard	Guard(this);
+		return Parent->GetTextureParameterOverrideValue(ParameterName,OutValue);
+	}
+	else
+	{
+		return false;
+	}
+}
+
 bool UMaterialInstance::GetFontParameterValue(FName ParameterName,class UFont*& OutFontValue, int32& OutFontPage) const
 {
 	if( ReentrantFlag )
diff --git a/Engine/Source/Runtime/Engine/Private/Materials/MaterialInterface.cpp b/Engine/Source/Runtime/Engine/Private/Materials/MaterialInterface.cpp
--- a/Engine/Source/Runtime/Engine/Private/Materials/MaterialInterface.cpp
+++ b/Engine/Source/Runtime/Engine/Private/Materials/MaterialInterface.cpp
@@ -200,6 +200,11 @@ bool UMaterialInterface::GetTextureParameterValue(FName ParameterName, UTexture*
 	return false;
 }
 
+bool UMaterialInterface::GetTextureParameterOverrideValue(FName ParameterName, UTexture*& OutValue) const
+{
+	return false;
+}
+
 bool UMaterialInterface::GetFontParameterValue(FName ParameterName,class UFont*& OutFontValue,int32& OutFontPage) const
 {
 	return false;
diff --git a/Engine/Source/Runtime/Engine/Private/Materials/MaterialUniformExpressions.h b/Engine/Source/Runtime/Engine/Private/Materials/MaterialUniformExpressions.h
--- a/Engine/Source/Runtime/Engine/Private/Materials/MaterialUniformExpressions.h
+++ b/Engine/Source/Runtime/Engine/Private/Materials/MaterialUniformExpressions.h
@@ -311,7 +311,7 @@ public:
 		else
 		{
 			OutValue = NULL;
-			if(!MaterialInterface->GetTextureParameterValue(ParameterName,OutValue))
+			if(!MaterialInterface->GetTextureParameterOverrideValue(ParameterName,OutValue))
 			{
 				OutValue = GetIndexedTexture(Material, TextureIndex);
 			}
