Fix writing strings in network profiler data file headers. Now handles strings of any length, and doesn't waste space if a string is less than 255 characters long. Previously if a string, like the URL, was longer than 255 characters, it would overwrite the beginning of the token stream, corrupting it. Incremented profiler version. Fixes UE-16607.

diff --git a/Engine/Source/Programs/NetworkProfiler/NetworkProfiler/StreamHeader.cs b/Engine/Source/Programs/NetworkProfiler/NetworkProfiler/StreamHeader.cs
--- a/Engine/Source/Programs/NetworkProfiler/NetworkProfiler/StreamHeader.cs
+++ b/Engine/Source/Programs/NetworkProfiler/NetworkProfiler/StreamHeader.cs
@@ -17,7 +17,7 @@ namespace NetworkProfiler
 		public const UInt32 ExpectedMagic = 0x1DBF348C;
 
 		/** We expect this version, or we can't proceed.			*/
-		public const UInt32 ExpectedVersion = 7;
+		public const UInt32 ExpectedVersion = 8;
 
 		/** Magic to ensure we're opening the right file.			*/
 		public UInt32 Magic;
@@ -132,7 +132,7 @@ namespace NetworkProfiler
 				NameTableOffset = BinaryStream.ReadUInt32();
 				NameTableEntries = BinaryStream.ReadUInt32();
 
-				// Serialize various fixed size strings.
+				// Serialize various dynamically-sized strings.
 				Tag = SerializeAnsiString( BinaryStream );
 				GameName = SerializeAnsiString( BinaryStream );
 				URL = SerializeAnsiString( BinaryStream );
@@ -140,8 +140,7 @@ namespace NetworkProfiler
 		}
 
 		/**
-		 * We serialize a fixed size string into the header so its size doesn't change later on. This means
-		 * we need to handle loading slightly differently to trim the extra padding.
+		 * Serialize a string from the header.
 		 * 
 		 * @param	BinaryStream	stream to read from
 		 * @return	string being read
@@ -149,13 +148,7 @@ namespace NetworkProfiler
 		private string SerializeAnsiString( BinaryReader BinaryStream )
 		{
 			UInt32 SerializedLength = BinaryStream.ReadUInt32();
-			string ResultString = new string(BinaryStream.ReadChars((int)SerializedLength)); 
-			// We serialize a fixed size string. Trim the null characters that make it in by converting char[] to string.
-			int RealLength = 0;
-			while( ( RealLength < ResultString.Length ) && ( ResultString[RealLength++] != '\0' ) ) 
-			{
-			}
-			return ResultString.Remove(RealLength-1);
+			return new string(BinaryStream.ReadChars((int)SerializedLength));
 		}
 	}
 }
diff --git a/Engine/Source/Runtime/Engine/Private/NetworkProfiler.cpp b/Engine/Source/Runtime/Engine/Private/NetworkProfiler.cpp
--- a/Engine/Source/Runtime/Engine/Private/NetworkProfiler.cpp
+++ b/Engine/Source/Runtime/Engine/Private/NetworkProfiler.cpp
@@ -26,7 +26,7 @@ FNetworkProfiler GNetworkProfiler;
 /** Magic value, determining that file is a network profiler file.				*/
 #define NETWORK_PROFILER_MAGIC						0x1DBF348C
 /** Version of memory profiler. Incremented on serialization changes.			*/
-#define NETWORK_PROFILER_VERSION					7
+#define NETWORK_PROFILER_VERSION					8
 
 static const FString UnknownName("UnknownName");
 
@@ -50,54 +50,43 @@ enum ENetworkProfilingPayloadType
 	NPTYPE_WritePropertyHandle			// Property handles
 };
 
-
 /*=============================================================================
-	Network profiler header.
+	FNetworkProfilerHeader implementation.
 =============================================================================*/
 
-struct FNetworkProfilerHeader
+FNetworkProfilerHeader::FNetworkProfilerHeader()
+	: Magic(NETWORK_PROFILER_MAGIC)
+	, Version(NETWORK_PROFILER_VERSION)
+	, NameTableOffset(0)
+	, NameTableEntries(0)
 {
-	/** Magic to ensure we're opening the right file.	*/
-	uint32	Magic;
-	/** Version number to detect version mismatches.	*/
-	uint32	Version;
-
-	/** Offset in file for name table.					*/
-	uint32	NameTableOffset;
-	/** Number of name table entries.					*/
-	uint32	NameTableEntries;
-
-	/** Tag, set via -networkprofiler=TAG				*/
-	FString Tag;
-	/** Game name, e.g. Example							*/
-	FString GameName;
-	/** URL used to open/ browse to the map.			*/
-	FString URL;
-
-	/**
-	 * Serialization operator.
-	 *
-	 * @param	Ar			Archive to serialize to
-	 * @param	Header		Header to serialize
-	 * @return	Passed in archive
-	 */
-	friend FArchive& operator << ( FArchive& Ar, FNetworkProfilerHeader Header )
-	{
-		check( Ar.IsSaving() );
-		Ar	<< Header.Magic
-			<< Header.Version
-			<< Header.NameTableOffset
-			<< Header.NameTableEntries;
-		Header.Tag.SerializeAsANSICharArray( Ar, 255 );
-		Header.GameName.SerializeAsANSICharArray( Ar, 255 );
-		Header.URL.SerializeAsANSICharArray( Ar, 255 );
-		return Ar;
-	}
-};
+}
+
+void FNetworkProfilerHeader::Reset(const FURL& InURL)
+{
+	NameTableOffset = 0;
+	NameTableEntries = 0;
 
+	FParse::Value(FCommandLine::Get(), TEXT("NETWORKPROFILER="), Tag);
+	GameName = FApp::GetGameName();
+	URL = InURL.ToString();
+}
+
+FArchive& operator << ( FArchive& Ar, FNetworkProfilerHeader& Header )
+{
+	check( Ar.IsSaving() );
+	Ar	<< Header.Magic
+		<< Header.Version
+		<< Header.NameTableOffset
+		<< Header.NameTableEntries;
+	Header.Tag.SerializeAsANSICharArray( Ar );
+	Header.GameName.SerializeAsANSICharArray( Ar );
+	Header.URL.SerializeAsANSICharArray( Ar );
+	return Ar;
+}
 
 /*=============================================================================
-	FMallocProfiler implementation.
+	FNetworkProfiler implementation.
 =============================================================================*/
 
 /**
@@ -107,7 +96,6 @@ FNetworkProfiler::FNetworkProfiler()
 :	FileWriter(NULL)
 ,	bHasNoticeableNetworkTrafficOccured(false)
 ,	bIsTrackingEnabled(false)
-,	CurrentURL(NoInit)
 {
 }
 
@@ -494,24 +482,15 @@ void FNetworkProfiler::TrackSessionChange( bool bShouldContinueTracking, const F
 		{	
 			if( bHasNoticeableNetworkTrafficOccured )
 			{
-				UE_LOG(LogNet, Log, TEXT("Network Profiler: Writing out session file for '%s'"), *CurrentURL.ToString());
+				UE_LOG(LogNet, Log, TEXT("Network Profiler: Writing out session file for '%s'"), *CurrentHeader.GetURL());
 
 				// Write end of stream marker.
 				uint8 Type = NPTYPE_EndOfStreamMarker;
 				(*FileWriter) << Type;
 
-				// Real header, written at start of the file but overwritten out right before we close the file.
-				FNetworkProfilerHeader Header;
-				Header.Magic	= NETWORK_PROFILER_MAGIC;
-				Header.Version	= NETWORK_PROFILER_VERSION;
-				Header.Tag		= TEXT("");
-				FParse::Value(FCommandLine::Get(), TEXT("NETWORKPROFILER="), Header.Tag);
-				Header.GameName = FApp::GetGameName();
-				Header.URL		= CurrentURL.ToString();
-
 				// Write out name table and update header with offset and count.
-				Header.NameTableOffset	= FileWriter->Tell();
-				Header.NameTableEntries	= NameArray.Num();
+				CurrentHeader.SetNameTableValues(FileWriter->Tell(), NameArray.Num());
+
 				for( int32 NameIndex=0; NameIndex<NameArray.Num(); NameIndex++ )
 				{
 					NameArray[NameIndex].SerializeAsANSICharArray( *FileWriter );
@@ -519,7 +498,7 @@ void FNetworkProfiler::TrackSessionChange( bool bShouldContinueTracking, const F
 
 				// Seek to the beginning of the file and write out proper header.
 				FileWriter->Seek( 0 );
-				(*FileWriter) << Header;
+				(*FileWriter) << CurrentHeader;
 
 				// Close file writer so we can rename the file to its final destination.
 				FileWriter->Close();
@@ -566,14 +545,12 @@ void FNetworkProfiler::TrackSessionChange( bool bShouldContinueTracking, const F
 			IFileManager::Get().MakeDirectory( *FPaths::GetPath(TempFileName) );
 			FileWriter = IFileManager::Get().CreateFileWriter( *TempFileName, FILEWRITE_EvenIfReadOnly );
 			check( FileWriter );
+			
+			CurrentHeader.Reset(InURL);
 
-			// Serialize dummy header, overwritten when session ends.
-			FNetworkProfilerHeader DummyHeader;
-			FMemory::Memzero( &DummyHeader, sizeof(DummyHeader) );
-			(*FileWriter) << DummyHeader;
+			// Serialize a header of the proper size, overwritten when session ends.
+			(*FileWriter) << CurrentHeader;
 		}
-
-		CurrentURL = InURL;
 	}
 #endif	//#if ALLOW_DEBUG_FILES
 }
diff --git a/Engine/Source/Runtime/Engine/Public/Net/NetworkProfiler.h b/Engine/Source/Runtime/Engine/Public/Net/NetworkProfiler.h
--- a/Engine/Source/Runtime/Engine/Public/Net/NetworkProfiler.h
+++ b/Engine/Source/Runtime/Engine/Public/Net/NetworkProfiler.h
@@ -4,8 +4,9 @@
 	NetworkProfiler.h: network profiling support.
 =============================================================================*/
 
-#ifndef UNREAL_NETWORK_PROFILER_H
-#define UNREAL_NETWORK_PROFILER_H
+#pragma once
+
+#include "UnrealString.h"
 
 class UNetConnection;
 class FInternetAddr;
@@ -14,6 +15,57 @@ class FInternetAddr;
 
 #define NETWORK_PROFILER( x ) if ( GNetworkProfiler.IsTrackingEnabled() ) { x; }
 
+/*=============================================================================
+	Network profiler header.
+=============================================================================*/
+
+class FNetworkProfilerHeader
+{
+private:
+	/** Magic to ensure we're opening the right file.	*/
+	uint32	Magic;
+	/** Version number to detect version mismatches.	*/
+	uint32	Version;
+
+	/** Offset in file for name table.					*/
+	uint32	NameTableOffset;
+	/** Number of name table entries.					*/
+	uint32	NameTableEntries;
+
+	/** Tag, set via -networkprofiler=TAG				*/
+	FString Tag;
+	/** Game name, e.g. Example							*/
+	FString GameName;
+	/** URL used to open/ browse to the map.			*/
+	FString URL;
+
+public:
+	/** Constructor.									*/
+	FNetworkProfilerHeader();
+
+	/** Resets the header info for a new session.		*/
+	void Reset(const FURL& InURL);
+
+	/** Sets the final name table values.				*/
+	void SetNameTableValues(uint32 Offset, uint32 Entries)
+	{
+		NameTableOffset = Offset;
+		NameTableEntries = Entries;
+	}
+
+	/** Returns the URL stored in the header.			*/
+	FString GetURL() const { return URL; }
+
+	/**
+	 * Serialization operator.
+	 *
+	 * @param	Ar			Archive to serialize to
+	 * @param	Header		Header to serialize
+	 * @return	Passed in archive
+	 */
+	friend FArchive& operator << ( FArchive& Ar, FNetworkProfilerHeader& Header );
+};
+
 /*=============================================================================
 	FNetworkProfiler
 =============================================================================*/
@@ -43,9 +95,9 @@ private:
 	/** Whether tracking is enabled.																*/
 	bool									bIsTrackingEnabled;	
 
-	/** URL used for current tracking session.														*/
-	FURL									CurrentURL;
-	
+	/** Header for the current session.																*/
+	FNetworkProfilerHeader					CurrentHeader;
+
 	/** All the data required for writing sent bunches to the profiler stream						*/
 	struct FSendBunchInfo
 	{
@@ -341,5 +393,3 @@ extern ENGINE_API FNetworkProfiler GNetworkProfiler;
 #define NETWORK_PROFILER(x)
 
 #endif
-
-#endif	//#ifndef UNREAL_NETWORK_PROFILER_H
