#ue4 - Fix Character::bSimGravityDisabled not being restored to false for simulated pawns under some conditions.

diff --git a/Engine/Source/Runtime/Engine/Private/Character.cpp b/Engine/Source/Runtime/Engine/Private/Character.cpp
--- a/Engine/Source/Runtime/Engine/Private/Character.cpp
+++ b/Engine/Source/Runtime/Engine/Private/Character.cpp
@@ -1187,16 +1187,13 @@ void ACharacter::UpdateSimulatedPosition(const FVector& NewLocation, const FRota
 
 		// Only need to check for encroachment when teleported without any velocity.
 		// Normal movement pops the character out of geometry anyway, no use doing it before and after (with different rules).
-		if (CharacterMovement && CharacterMovement->Velocity.IsZero())
+		bSimGravityDisabled = false;
+		if (CharacterMovement->Velocity.IsZero())
 		{
 			if (GetWorld()->EncroachingBlockingGeometry(this, NewLocation, NewRotation))
 			{
 				bSimGravityDisabled = true;
 			}
-			else
-			{
-				bSimGravityDisabled = false;
-			}
 		}
 		
 		// Don't use TeleportTo(), that clears our base.
