#ue4 - Perf: inline some common PhysX<->Unreal type conversions.

diff --git a/Engine/Source/Runtime/Engine/Private/Collision/PhysXCollision.cpp b/Engine/Source/Runtime/Engine/Private/Collision/PhysXCollision.cpp
--- a/Engine/Source/Runtime/Engine/Private/Collision/PhysXCollision.cpp
+++ b/Engine/Source/Runtime/Engine/Private/Collision/PhysXCollision.cpp
@@ -496,14 +496,14 @@ struct FScopedMultiSceneReadLock
 
 	inline void LockRead(PxScene* Scene, int32 SceneIndex)
 	{
-		check(SceneLocks[SceneIndex] == nullptr); // no nested locks allowed.
+		checkSlow(SceneLocks[SceneIndex] == nullptr); // no nested locks allowed.
 		SCENE_LOCK_READ(Scene);
 		SceneLocks[SceneIndex] = Scene;
 	}
 
 	inline void UnlockRead(PxScene* Scene, int32 SceneIndex)
 	{
-		check(SceneLocks[SceneIndex] == Scene || SceneLocks[SceneIndex] == nullptr);
+		checkSlow(SceneLocks[SceneIndex] == Scene || SceneLocks[SceneIndex] == nullptr);
 		SCENE_UNLOCK_READ(Scene);
 		SceneLocks[SceneIndex] = nullptr;
 	}
diff --git a/Engine/Source/Runtime/Engine/Private/PhysicsEngine/PhysXSupport.cpp b/Engine/Source/Runtime/Engine/Private/PhysicsEngine/PhysXSupport.cpp
--- a/Engine/Source/Runtime/Engine/Private/PhysicsEngine/PhysXSupport.cpp
+++ b/Engine/Source/Runtime/Engine/Private/PhysicsEngine/PhysXSupport.cpp
@@ -73,16 +73,6 @@ PxTransform U2PTransform(const FTransform& UTransform)
 	return Result;
 }
 
-PxVec3 U2PVector(const FVector& UVec)
-{
-	return PxVec3(UVec.X, UVec.Y, UVec.Z);
-}
-
-PxQuat U2PQuat(const FQuat& UQuat)
-{
-	return PxQuat( UQuat.X, UQuat.Y, UQuat.Z, UQuat.W );
-}
-
 PxMat44 U2PMatrix(const FMatrix& UTM)
 {
 	PxMat44 Result;
@@ -93,11 +83,6 @@ PxMat44 U2PMatrix(const FMatrix& UTM)
 	return Result;
 }
 
-PxPlane U2PPlane(FPlane& Plane)
-{
-	return PxPlane(Plane.X, Plane.Y, Plane.Z, -Plane.W);
-}
-
 UCollision2PGeom::UCollision2PGeom(const FCollisionShape& CollisionShape)
 {
 	switch (CollisionShape.ShapeType)
@@ -153,25 +138,6 @@ FTransform P2UTransform(const PxTransform& PTM)
 	return Result;
 }
 
-FVector P2UVector(const PxVec3& PVec)
-{
-	return FVector(PVec.x, PVec.y, PVec.z);
-}
-
-FQuat P2UQuat(const PxQuat& PQuat)
-{
-	return FQuat(PQuat.x, PQuat.y, PQuat.z, PQuat.w);
-}
-
-FPlane P2UPlane(PxReal P[4])
-{
-	return FPlane(P[0], P[1], P[2], -P[3]);
-}
-
-FPlane P2UPlane(PxPlane& Plane)
-{
-	return FPlane(Plane.n.x, Plane.n.y, Plane.n.z, -Plane.d);
-}
 ///////////////////// Utils /////////////////////
 
 
diff --git a/Engine/Source/Runtime/Engine/Private/PhysicsEngine/PhysXSupport.h b/Engine/Source/Runtime/Engine/Private/PhysicsEngine/PhysXSupport.h
--- a/Engine/Source/Runtime/Engine/Private/PhysicsEngine/PhysXSupport.h
+++ b/Engine/Source/Runtime/Engine/Private/PhysicsEngine/PhysXSupport.h
@@ -296,19 +296,57 @@ ENGINE_API PxQuat U2PQuat(const FQuat& UQuat);
 /** Convert Unreal FMatrix to PhysX PxMat44 */
 ENGINE_API PxMat44 U2PMatrix(const FMatrix& UTM);
 /** Convert Unreal FPlane to PhysX plane def */
-ENGINE_API PxPlane U2PPlane(FPlane& Plane);
+ENGINE_API PxPlane U2PPlane(const FPlane& Plane);
 /** Convert PhysX PxTransform to Unreal PxTransform */
 ENGINE_API FTransform P2UTransform(const PxTransform& PTM);
 /** Convert PhysX PxQuat to Unreal FQuat */
 ENGINE_API FQuat P2UQuat(const PxQuat& PQuat);
 /** Convert PhysX plane def to Unreal FPlane */
-ENGINE_API FPlane P2UPlane(PxReal P[4]);
-ENGINE_API FPlane P2UPlane(PxPlane& Plane);
+ENGINE_API FPlane P2UPlane(const PxReal P[4]);
+ENGINE_API FPlane P2UPlane(const PxPlane& Plane);
 /** Convert PhysX PxMat44 to Unreal FMatrix */
 ENGINE_API FMatrix P2UMatrix(const PxMat44& PMat);
 /** Convert PhysX PxTransform to Unreal FMatrix */
 ENGINE_API FMatrix PTransform2UMatrix(const PxTransform& PTM);
 
+// inlines
+
+ENGINE_API FORCEINLINE_DEBUGGABLE PxVec3 U2PVector(const FVector& UVec)
+{
+	return PxVec3(UVec.X, UVec.Y, UVec.Z);
+}
+
+ENGINE_API FORCEINLINE_DEBUGGABLE PxQuat U2PQuat(const FQuat& UQuat)
+{
+	return PxQuat( UQuat.X, UQuat.Y, UQuat.Z, UQuat.W );
+}
+
+ENGINE_API FORCEINLINE_DEBUGGABLE PxPlane U2PPlane(const FPlane& Plane)
+{
+	return PxPlane(Plane.X, Plane.Y, Plane.Z, -Plane.W);
+}
+
+ENGINE_API FORCEINLINE_DEBUGGABLE FVector P2UVector(const PxVec3& PVec)
+{
+	return FVector(PVec.x, PVec.y, PVec.z);
+}
+
+ENGINE_API FORCEINLINE_DEBUGGABLE FQuat P2UQuat(const PxQuat& PQuat)
+{
+	return FQuat(PQuat.x, PQuat.y, PQuat.z, PQuat.w);
+}
+
+ENGINE_API FORCEINLINE_DEBUGGABLE FPlane P2UPlane(const PxReal P[4])
+{
+	return FPlane(P[0], P[1], P[2], -P[3]);
+}
+
+ENGINE_API FORCEINLINE_DEBUGGABLE FPlane P2UPlane(const PxPlane& Plane)
+{
+	return FPlane(Plane.n.x, Plane.n.y, Plane.n.z, -Plane.d);
+}
+
+
 //////// GEOM CONVERSION
 // we need this helper struct since PhysX needs geoms to be on the stack
 struct UCollision2PGeom
