#ue4 - Change a check() to checkSlow() in FScopedMovementUpdate.

diff --git a/Engine/Source/Runtime/Engine/Private/Components/SceneComponent.cpp b/Engine/Source/Runtime/Engine/Private/Components/SceneComponent.cpp
--- a/Engine/Source/Runtime/Engine/Private/Components/SceneComponent.cpp
+++ b/Engine/Source/Runtime/Engine/Private/Components/SceneComponent.cpp
@@ -477,7 +477,7 @@ void USceneComponent::EndScopedMovementUpdate(class FScopedMovementUpdate& Compl
 
 	// Process top of the stack
 	FScopedMovementUpdate* CurrentScopedUpdate = ScopedMovementStack.Pop(false);
-	check(CurrentScopedUpdate == &CompletedScope);
+	checkSlow(CurrentScopedUpdate == &CompletedScope);
 	{
 		checkSlow(CurrentScopedUpdate->IsDeferringUpdates());
 		if (ScopedMovementStack.Num() == 0)
