#ue4 - Fix mono build.

diff --git a/Engine/Source/Editor/DestructibleMeshEditor/Private/SDestructibleMeshEditorViewport.cpp b/Engine/Source/Editor/DestructibleMeshEditor/Private/SDestructibleMeshEditorViewport.cpp
--- a/Engine/Source/Editor/DestructibleMeshEditor/Private/SDestructibleMeshEditorViewport.cpp
+++ b/Engine/Source/Editor/DestructibleMeshEditor/Private/SDestructibleMeshEditorViewport.cpp
@@ -3,6 +3,7 @@
 #include "DestructibleMeshEditorPrivatePCH.h"
 
 #include "PhysicsPublic.h"
+#include "Runtime/Engine/Private/PhysicsEngine/PhysXSupport.h"
 #include "MouseDeltaTracker.h"
 #include "Runtime/Engine/Public/Slate/SceneViewport.h"
 #include "PreviewScene.h"
diff --git a/Engine/Source/Runtime/Engine/Private/PhysicsEngine/PhysXSupport.h b/Engine/Source/Runtime/Engine/Private/PhysicsEngine/PhysXSupport.h
--- a/Engine/Source/Runtime/Engine/Private/PhysicsEngine/PhysXSupport.h
+++ b/Engine/Source/Runtime/Engine/Private/PhysicsEngine/PhysXSupport.h
@@ -299,6 +299,8 @@ ENGINE_API PxMat44 U2PMatrix(const FMatrix& UTM);
 ENGINE_API PxPlane U2PPlane(const FPlane& Plane);
 /** Convert PhysX PxTransform to Unreal PxTransform */
 ENGINE_API FTransform P2UTransform(const PxTransform& PTM);
+/** Convert PhysX PxVec3 to Unreal FVector */
+ENGINE_API FVector P2UVector(const PxVec3& PVec);
 /** Convert PhysX PxQuat to Unreal FQuat */
 ENGINE_API FQuat P2UQuat(const PxQuat& PQuat);
 /** Convert PhysX plane def to Unreal FPlane */
diff --git a/Engine/Source/Runtime/Engine/Public/PhysicsPublic.h b/Engine/Source/Runtime/Engine/Public/PhysicsPublic.h
--- a/Engine/Source/Runtime/Engine/Public/PhysicsPublic.h
+++ b/Engine/Source/Runtime/Engine/Public/PhysicsPublic.h
@@ -106,9 +106,6 @@ extern ENGINE_API class FPhysXAllocator* GPhysXAllocator;
 /** Pointer to PhysX Command Handler */
 extern ENGINE_API class FPhysCommandHandler* GPhysCommandHandler;
 
-/** Convert PhysX PxVec3 to Unreal FVector */
-ENGINE_API FVector P2UVector(const PxVec3& PVec);
-
 #if WITH_APEX
 
 namespace NxParameterized
