Stats - Misc tweak

diff --git a/Engine/Source/Runtime/Engine/Private/UnrealEngine.cpp b/Engine/Source/Runtime/Engine/Private/UnrealEngine.cpp
--- a/Engine/Source/Runtime/Engine/Private/UnrealEngine.cpp
+++ b/Engine/Source/Runtime/Engine/Private/UnrealEngine.cpp
@@ -4155,10 +4155,10 @@ bool UEngine::HandleMemCommand( const TCHAR* Cmd, FOutputDevice& Ar )
 		TArray<FStatMessage> Stats;
 		GetPermanentStats(Stats);
 
-		FName NAME_STATGROUP_SceneMemory("STATGROUP_SceneMemory");
-		FName NAME_STATGROUP_Memory("STATGROUP_Memory");
+		FName NAME_STATGROUP_SceneMemory( FStatGroup_STATGROUP_SceneMemory::GetGroupName() );
+		FName NAME_STATGROUP_Memory(FStatGroup_STATGROUP_Memory::GetGroupName());
 		FName NAME_STATGROUP_TextureGroup("STATGROUP_TextureGroup");
-		FName NAME_STATGROUP_RHI("STATGROUP_RHI");
+		FName NAME_STATGROUP_RHI(FStatGroup_STATGROUP_RHI::GetGroupName());
 
 		for (int32 Index = 0; Index < Stats.Num(); Index++)
 		{
@@ -4698,7 +4698,7 @@ struct FHierarchy
 	}
 };
 
-// @TODO yrx 2014-09-15 Move to ObjectCommads.cpp or ObjectExec.cpp
+// #TODO_YRX 2014-09-15 Move to ObjectCommads.cpp or ObjectExec.cpp
 bool UEngine::HandleObjCommand( const TCHAR* Cmd, FOutputDevice& Ar )
 {
 	if( FParse::Command(&Cmd,TEXT("GARBAGE")) || FParse::Command(&Cmd,TEXT("GC")) )
