UE-13798 Make the shared state of the event graph use something akin to weak ptrs - ENABLED

diff --git a/Engine/Source/Runtime/Engine/Private/BlueprintGeneratedClass.cpp b/Engine/Source/Runtime/Engine/Private/BlueprintGeneratedClass.cpp
--- a/Engine/Source/Runtime/Engine/Private/BlueprintGeneratedClass.cpp
+++ b/Engine/Source/Runtime/Engine/Private/BlueprintGeneratedClass.cpp
@@ -676,7 +676,7 @@ public:
 protected:
 	virtual FArchive& operator<<(UObject*& Object) override
 	{
-		const bool bWeakRef = false; //@TODO - enable: Object ? !Object->HasAnyFlags(RF_StrongRefOnFrame) : false;
+		const bool bWeakRef = Object ? !Object->HasAnyFlags(RF_StrongRefOnFrame) : false;
 		Collector.SetShouldHandleAsWeakRef(bWeakRef); 
 		return FSimpleObjectReferenceCollectorArchive::operator<<(Object);
 	}
@@ -698,6 +698,8 @@ void UBlueprintGeneratedClass::AddReferencedObjectsInUbergraphFrame(UObject* InT
 				{
 					FPersistentFrameCollectorArchive ObjectReferenceCollector(InThis, Collector);
 					BPGC->UberGraphFunction->SerializeBin(ObjectReferenceCollector, PointerToUberGraphFrame->RawPointer);
+
+					// Reset the ShouldHandleAsWeakRef state, before the collector is used by a different archive.
 					Collector.SetShouldHandleAsWeakRef(false);
 				}
 			}
