[UE-16566] Fix a potential crash on load during ActorComponent serialization for older SCS node templates.

change summary:
- modified UActorComponent::PostLoad() to ensure that templates older than VER_UE4_ACTOR_COMPONENT_CREATION_METHOD are set to the default CreationMethod on load.

diff --git a/Engine/Source/Runtime/Engine/Private/Components/ActorComponent.cpp b/Engine/Source/Runtime/Engine/Private/Components/ActorComponent.cpp
--- a/Engine/Source/Runtime/Engine/Private/Components/ActorComponent.cpp
+++ b/Engine/Source/Runtime/Engine/Private/Components/ActorComponent.cpp
@@ -129,9 +129,13 @@ void UActorComponent::PostLoad()
 
 	if (GetLinkerUE4Version() < VER_UE4_ACTOR_COMPONENT_CREATION_METHOD)
 	{
-		if (bCreatedByConstructionScript_DEPRECATED)
+		if (IsTemplate())
 		{
-				CreationMethod = EComponentCreationMethod::SimpleConstructionScript;
+			CreationMethod = EComponentCreationMethod::Native;
+		}
+		else if (bCreatedByConstructionScript_DEPRECATED)
+		{
+			CreationMethod = EComponentCreationMethod::SimpleConstructionScript;
 		}
 		else if (bInstanceComponent_DEPRECATED)
 		{
