#ue4 - Initialize UActorComponent::bCanUseCachedOwner to true.

diff --git a/Engine/Source/Runtime/Engine/Private/Components/ActorComponent.cpp b/Engine/Source/Runtime/Engine/Private/Components/ActorComponent.cpp
--- a/Engine/Source/Runtime/Engine/Private/Components/ActorComponent.cpp
+++ b/Engine/Source/Runtime/Engine/Private/Components/ActorComponent.cpp
@@ -111,6 +111,7 @@ UActorComponent::UActorComponent(const FObjectInitializer& ObjectInitializer /*=
 	bAutoRegister = true;
 	bNetAddressable = false;
 	bEditableWhenInherited = true;
+	bCanUseCachedOwner = true;
 }
 
 void UActorComponent::PostInitProperties()
