#Anim
Meta data support

diff --git a/Engine/Source/Editor/Persona/Private/Persona.cpp b/Engine/Source/Editor/Persona/Private/Persona.cpp
--- a/Engine/Source/Editor/Persona/Private/Persona.cpp
+++ b/Engine/Source/Editor/Persona/Private/Persona.cpp
@@ -434,6 +434,9 @@ TSharedPtr<SDockTab> FPersona::OpenNewAnimationDocumentTab(UObject* InAnimAsset)
 
 TSharedPtr<SDockTab> FPersona::OpenNewDocumentTab(class UAnimationAsset* InAnimAsset)
 {
+	/// before opening new asset, clear the currently selected object
+	SetDetailObject(NULL);
+
 	TSharedPtr<SDockTab> NewTab;
 	if (InAnimAsset)
 	{
diff --git a/Engine/Source/Editor/Persona/Private/SAnimEditorBase.cpp b/Engine/Source/Editor/Persona/Private/SAnimEditorBase.cpp
--- a/Engine/Source/Editor/Persona/Private/SAnimEditorBase.cpp
+++ b/Engine/Source/Editor/Persona/Private/SAnimEditorBase.cpp
@@ -271,8 +271,9 @@ void SAnimEditorBase::OnSelectionChanged(const FGraphPanelSelectionSet& Selected
 {
 	if (SelectedItems.Num() == 0)
 	{
-		// Edit the sequence
-		PersonaPtr.Pin()->UpdateSelectionDetails(GetEditorObject(), LOCTEXT("Edit Sequence", "Edit Sequence"));
+		// If nothing is selected, unselect it instead of selecting current editor object. 
+		// That means, it shows same information as Anim Detail Panel
+		PersonaPtr.Pin()->UpdateSelectionDetails(NULL, LOCTEXT("Edit Sequence", "Edit Sequence"));
 	}
 	else
 	{
diff --git a/Engine/Source/Editor/Persona/Private/SAnimMontagePanel.cpp b/Engine/Source/Editor/Persona/Private/SAnimMontagePanel.cpp
--- a/Engine/Source/Editor/Persona/Private/SAnimMontagePanel.cpp
+++ b/Engine/Source/Editor/Persona/Private/SAnimMontagePanel.cpp
@@ -78,7 +78,6 @@ SAnimMontagePanel::~SAnimMontagePanel()
 /** This is the main function that creates the UI widgets for the montage tool.*/
 void SAnimMontagePanel::Update()
 {
-	ClearSelected();
 	if ( Montage != NULL )
 	{
 		SMontageEditor * Editor = MontageEditor.Pin().Get();
diff --git a/Engine/Source/Runtime/Engine/Classes/Animation/AnimCompositeBase.h b/Engine/Source/Runtime/Engine/Classes/Animation/AnimCompositeBase.h
--- a/Engine/Source/Runtime/Engine/Classes/Animation/AnimCompositeBase.h
+++ b/Engine/Source/Runtime/Engine/Classes/Animation/AnimCompositeBase.h
@@ -106,6 +106,30 @@ struct FAnimSegment
 		return ((CurPos >= StartPos) && (CurPos <= (StartPos + GetLength())));
 	}
 
+	/*
+	 * Return true if it's included within the input range
+	 */
+	bool IsIncluded(float InStartPos, float InEndPos) const
+	{
+		float EndPos = StartPos + GetLength(); 
+		// InStartPos is between Start and End, it is included
+		if (StartPos <= InStartPos && EndPos > InStartPos)
+		{
+			return true;
+		}
+		// InEndPos is between Start and End, it is also included
+		if (StartPos < InEndPos && EndPos >= InEndPos)
+		{
+			return true;
+		}
+		// if it is within Start and End, it is also included
+		if (StartPos >= InStartPos && EndPos <= InEndPos)
+		{
+			return true;
+		}
+
+		return false;
+	}
 	/**
 	 * Get Animation Data, for now have weight to be controlled here, in the future, it will be controlled in Track
 	 */
diff --git a/Engine/Source/Runtime/Engine/Classes/Animation/AnimMetaData.h b/Engine/Source/Runtime/Engine/Classes/Animation/AnimMetaData.h
new file mode 100644
--- /dev/null
+++ b/Engine/Source/Runtime/Engine/Classes/Animation/AnimMetaData.h
@@ -0,0 +1,20 @@
+// Copyright 1998-2015 Epic Games, Inc. All Rights Reserved.
+
+/* 
+ * Base Class for UAnimMetaData that can be implemented for each game needs
+ * this data will be saved to animation asset as well as montage sections, and you can query that data and decide what to do
+ *
+ * Refer : GetMetaData/GetSectionMetaData
+ */
+
+#pragma once
+#include "AnimMetadata.generated.h"
+
+UCLASS(Blueprintable, abstract, const, MinimalAPI, editinlinenew, hidecategories=Object, collapsecategories)
+class UAnimMetaData : public UObject
+{
+	GENERATED_UCLASS_BODY()
+};
+
+
+
diff --git a/Engine/Source/Runtime/Engine/Classes/Animation/AnimMontage.h b/Engine/Source/Runtime/Engine/Classes/Animation/AnimMontage.h
--- a/Engine/Source/Runtime/Engine/Classes/Animation/AnimMontage.h
+++ b/Engine/Source/Runtime/Engine/Classes/Animation/AnimMontage.h
@@ -35,12 +35,25 @@ struct FCompositeSection : public FAnimLinkableElement
 	UPROPERTY()
 	FName NextSectionName;
 
+private:
+	/** Meta data that can be saved with the asset 
+	 * 
+	 * You can query by GetMetaData function
+	 */
+	UPROPERTY(Category=Section, Instanced, EditAnywhere)
+	TArray<class UAnimMetaData*> MetaData;
+
+public:
 	FCompositeSection()
 		: FAnimLinkableElement()
 		, SectionName(NAME_None)
 		, NextSectionName(NAME_None)
 	{
 	}
+
+	/** Get available Metadata for this section
+	 */
+	const TArray<class UAnimMetaData*>& GetMetaData() const { return MetaData; }
 };
 
 /**
@@ -462,6 +475,21 @@ public:
 	/** Return Section Index from Position */
 	ENGINE_API int32 GetSectionIndexFromPosition(float Position) const;
 	
+	/**
+	 * Get Section Metadata for the montage including metadata belong to the anim reference
+	 * This will remove redundant entry if found - i.e. multiple same anim reference is used
+	 * 
+	 * @param : SectionName - Name of section you'd like to get meta data for. 
+	 *						- If SectionName == NONE, it will return all the section data
+	 * @param : bIncludeSequence - if true, it returns all metadata of the animation within that section
+	 *						- whether partial or full
+	 * @param : SlotName - this only matters if bIncludeSequence is true.
+	 *						- If true, and if SlotName is given, it will only look for SlotName.
+	 *						- If true and if SlotName is none, then it will look for all slot nodes
+	 ***/
+
+	ENGINE_API const TArray<class UAnimMetaData*> GetSectionMetaData(FName SectionName, bool bIncludeSequence=true, FName SlotName = NAME_None);
+
 	/** Get Section Index from CurrentTime with PosWithinCompositeSection */
 	int32 GetAnimCompositeSectionIndexFromPos(float CurrentTime, float& PosWithinCompositeSection) const;
 
diff --git a/Engine/Source/Runtime/Engine/Classes/Animation/AnimationAsset.h b/Engine/Source/Runtime/Engine/Classes/Animation/AnimationAsset.h
--- a/Engine/Source/Runtime/Engine/Classes/Animation/AnimationAsset.h
+++ b/Engine/Source/Runtime/Engine/Classes/Animation/AnimationAsset.h
@@ -440,6 +440,13 @@ private:
 	/** Skeleton guid. If changes, you need to remap info*/
 	FGuid SkeletonGuid;
 
+	/** Meta data that can be saved with the asset 
+	 * 
+	 * You can query by GetMetaData function
+	 */
+	UPROPERTY(Category=MetaData, instanced, EditAnywhere)
+	TArray<class UAnimMetaData*> MetaData;
+
 public:
 	/** Advances the asset player instance 
 	 * 
@@ -462,6 +469,10 @@ public:
 
 	virtual void Serialize(FArchive& Ar) override;
 
+	/** Get available Metadata within the animation asset
+	 */
+	ENGINE_API const TArray<class UAnimMetaData*>& GetMetaData() const { return MetaData; }
+
 #if WITH_EDITOR
 	/** Replace Skeleton 
 	 * 
diff --git a/Engine/Source/Runtime/Engine/Private/Animation/AnimMetaData.cpp b/Engine/Source/Runtime/Engine/Private/Animation/AnimMetaData.cpp
new file mode 100644
--- /dev/null
+++ b/Engine/Source/Runtime/Engine/Private/Animation/AnimMetaData.cpp
@@ -0,0 +1,12 @@
+// Copyright 1998-2015 Epic Games, Inc. All Rights Reserved.
+
+#include "EnginePrivate.h"
+#include "Animation/AnimMetaData.h"
+
+/////////////////////////////////////////////////////
+// UAnimMetaDatay
+
+UAnimMetaData::UAnimMetaData(const FObjectInitializer& ObjectInitializer)
+	: Super(ObjectInitializer)
+{
+}
diff --git a/Engine/Source/Runtime/Engine/Private/Animation/AnimMontage.cpp b/Engine/Source/Runtime/Engine/Private/Animation/AnimMontage.cpp
--- a/Engine/Source/Runtime/Engine/Private/Animation/AnimMontage.cpp
+++ b/Engine/Source/Runtime/Engine/Private/Animation/AnimMontage.cpp
@@ -1625,6 +1625,83 @@ float UAnimMontage::CalculateSequenceLength()
 	return CalculatedSequenceLength;
 }
 
+const TArray<class UAnimMetaData*> UAnimMontage::GetSectionMetaData(FName SectionName, bool bIncludeSequence/*=true*/, FName SlotName /*= NAME_None*/)
+{
+	TArray<class UAnimMetaData*> MetadataList;
+	bool bShouldIIncludeSequence = bIncludeSequence;
+	
+	for (int32 SectionIndex=0; SectionIndex<CompositeSections.Num(); ++SectionIndex)
+	{
+		const auto& CurSection = CompositeSections[SectionIndex];
+		if (SectionName == NAME_None || CurSection.SectionName == SectionName)
+		{
+			// add to the list
+			MetadataList.Append(CurSection.GetMetaData());
+
+			if (bShouldIIncludeSequence)
+			{
+				if (SectionName == NAME_None)
+				{
+					for(auto& SlotIter : SlotAnimTracks)
+					{
+						if(SlotName == NAME_None || SlotIter.SlotName ==  SlotName)
+						{
+							// now add the animations within this section
+							for(auto& SegmentIter : SlotIter.AnimTrack.AnimSegments)
+							{
+								if(SegmentIter.AnimReference)
+								{
+									// only add unique here
+									TArray<UAnimMetaData*> RefMetadata = SegmentIter.AnimReference->GetMetaData();
+
+									for (auto& RefData : RefMetadata)
+									{
+										MetadataList.AddUnique(RefData);
+									}
+								}
+							}
+						}
+					}
+
+					// if section name == None, we only grab slots once
+					// otherwise, it will grab multiple times
+					bShouldIIncludeSequence = false;
+				}
+				else
+				{
+					float SectionStartTime = 0.f, SectionEndTime = 0.f;
+					GetSectionStartAndEndTime(SectionIndex, SectionStartTime, SectionEndTime);
+					for(auto& SlotIter : SlotAnimTracks)
+					{
+						if(SlotName == NAME_None || SlotIter.SlotName ==  SlotName)
+						{
+							// now add the animations within this section
+							for(auto& SegmentIter : SlotIter.AnimTrack.AnimSegments)
+							{
+								if (SegmentIter.IsIncluded(SectionStartTime, SectionEndTime))
+								{
+									if(SegmentIter.AnimReference)
+									{
+										// only add unique here
+										TArray<UAnimMetaData*> RefMetadata = SegmentIter.AnimReference->GetMetaData();
+
+										for(auto& RefData : RefMetadata)
+										{
+											MetadataList.AddUnique(RefData);
+										}
+									}
+								}
+							}
+						}
+					}
+				}
+			}
+		}
+	}
+
+	return MetadataList;
+}
+
 #if WITH_EDITOR
 bool UAnimMontage::GetAllAnimationSequencesReferred(TArray<UAnimSequence*>& AnimationSequences)
 {
@@ -1645,7 +1722,7 @@ void UAnimMontage::ReplaceReferredAnimations(const TMap<UAnimSequence*, UAnimSeq
 	}
 }
 
-ENGINE_API void UAnimMontage::UpdateLinkableElements()
+void UAnimMontage::UpdateLinkableElements()
 {
 	// Update all linkable elements
 	for(FCompositeSection& Section : CompositeSections)
