Multiple instances of the same project can now record network profiles simultaneously.
Fixed by using a unqiue temporary filename to record profiling data.
Delete old temporary files so that they don't keep accumulating.

diff --git a/Engine/Source/Runtime/Engine/Private/NetworkProfiler.cpp b/Engine/Source/Runtime/Engine/Private/NetworkProfiler.cpp
--- a/Engine/Source/Runtime/Engine/Private/NetworkProfiler.cpp
+++ b/Engine/Source/Runtime/Engine/Private/NetworkProfiler.cpp
@@ -30,6 +30,8 @@ FNetworkProfiler GNetworkProfiler;
 
 static const FString UnknownName("UnknownName");
 
+static const double MaxTempFileAgeSeconds = 60.0 * 60.0 * 24.0 * 5.0;
+
 enum ENetworkProfilingPayloadType
 {
 	NPTYPE_FrameMarker			= 0,	// Frame marker, signaling beginning of frame.	
@@ -525,6 +527,9 @@ void FNetworkProfiler::TrackSessionChange( bool bShouldContinueTracking, const F
 			{
 				UE_LOG(LogNet, Warning, TEXT("Network Profiler: Nothing important happened"));
 				FileWriter->Close();
+
+				// Delete the temporary file.
+				IFileManager::Get().Delete(*TempFileName);
 			}
 
 			// Clean up.
@@ -535,11 +540,24 @@ void FNetworkProfiler::TrackSessionChange( bool bShouldContinueTracking, const F
 
 		if( bShouldContinueTracking )
 		{
+			// Delete any stale .tmp files.
+			TArray<FString> FoundTempFiles;
+			IFileManager::Get().FindFiles(FoundTempFiles, *FPaths::ProfilingDir(), TEXT(".tmp"));
+			
+			for ( const FString& FoundFile : FoundTempFiles)
+			{
+				const FString FullFilename = FPaths::ProfilingDir() + FoundFile;
+				if ( IFileManager::Get().GetFileAgeSeconds(*FullFilename) > MaxTempFileAgeSeconds)
+				{
+					IFileManager::Get().Delete(*FullFilename);
+				}
+			}
+
 			// Start a new tracking session.
 			check( FileWriter == NULL );
 
 			// Use a dummy name for sessions in progress that is renamed at end.
-			TempFileName = FPaths::ProfilingDir() + TEXT("NetworkProfiling.tmp");
+			TempFileName = FPaths::CreateTempFilename(*FPaths::ProfilingDir(), TEXT("NetworkProfiling-"));
 
 			// Create folder and file writer.
 			IFileManager::Get().MakeDirectory( *FPaths::GetPath(TempFileName) );
