Added virtual destructor for FDamageEvent.
Github PR#1141
(thanks bozaro!)

diff --git a/Engine/Source/Runtime/Engine/Classes/Engine/EngineTypes.h b/Engine/Source/Runtime/Engine/Classes/Engine/EngineTypes.h
--- a/Engine/Source/Runtime/Engine/Classes/Engine/EngineTypes.h
+++ b/Engine/Source/Runtime/Engine/Classes/Engine/EngineTypes.h
@@ -2401,6 +2401,8 @@ public:
 		: DamageTypeClass(InDamageEvent.DamageTypeClass)
 	{ }
 	
+	virtual ~FDamageEvent() { }
+
 	explicit FDamageEvent(TSubclassOf<class UDamageType> InDamageTypeClass)
 		: DamageTypeClass(InDamageTypeClass)
 	{ }
