Fix compile error

diff --git a/Engine/Source/Runtime/Engine/Private/Components/ActorComponent.cpp b/Engine/Source/Runtime/Engine/Private/Components/ActorComponent.cpp
--- a/Engine/Source/Runtime/Engine/Private/Components/ActorComponent.cpp
+++ b/Engine/Source/Runtime/Engine/Private/Components/ActorComponent.cpp
@@ -111,7 +111,9 @@ UActorComponent::UActorComponent(const FObjectInitializer& ObjectInitializer /*=
 	bAutoRegister = true;
 	bNetAddressable = false;
 	bEditableWhenInherited = true;
+#if WITH_EDITOR
 	bCanUseCachedOwner = true;
+#endif
 }
 
 void UActorComponent::PostInitProperties()
