#ue4 - Add stat counter for CharacterMovementComponent::FindFloor().

diff --git a/Engine/Source/Runtime/Engine/Private/Components/CharacterMovementComponent.cpp b/Engine/Source/Runtime/Engine/Private/Components/CharacterMovementComponent.cpp
--- a/Engine/Source/Runtime/Engine/Private/Components/CharacterMovementComponent.cpp
+++ b/Engine/Source/Runtime/Engine/Private/Components/CharacterMovementComponent.cpp
@@ -39,6 +39,7 @@ DECLARE_CYCLE_STAT(TEXT("Char CombineNetMove"), STAT_CharacterMovementCombineNet
 DECLARE_CYCLE_STAT(TEXT("Char SmoothClientPosition"), STAT_CharacterMovementSmoothClientPosition, STATGROUP_Character);
 DECLARE_CYCLE_STAT(TEXT("Char Physics Interation"), STAT_CharPhysicsInteraction, STATGROUP_Character);
 DECLARE_CYCLE_STAT(TEXT("Char StepUp"), STAT_CharStepUp, STATGROUP_Character);
+DECLARE_CYCLE_STAT(TEXT("Char FindFloor"), STAT_CharFindFloor, STATGROUP_Character);
 DECLARE_CYCLE_STAT(TEXT("Char AdjustFloorHeight"), STAT_CharAdjustFloorHeight, STATGROUP_Character);
 DECLARE_CYCLE_STAT(TEXT("Char Update Acceleration"), STAT_CharUpdateAcceleration, STATGROUP_Character);
 DECLARE_CYCLE_STAT(TEXT("Char MoveUpdateDelegate"), STAT_CharMoveUpdateDelegate, STATGROUP_Character);
@@ -4800,6 +4801,8 @@ void UCharacterMovementComponent::ComputeFloorDist(const FVector& CapsuleLocatio
 
 void UCharacterMovementComponent::FindFloor(const FVector& CapsuleLocation, FFindFloorResult& OutFloorResult, bool bZeroDelta, const FHitResult* DownwardSweepResult) const
 {
+	SCOPE_CYCLE_COUNTER(STAT_CharFindFloor);
+
 	// No collision, no floor...
 	if (!HasValidData() || !UpdatedComponent->IsCollisionEnabled())
 	{
