#ue4 - Avoid useless function calls in CharacterMovementComponent::ConstrainInputAcceleration() when possible.

diff --git a/Engine/Source/Runtime/Engine/Private/Components/CharacterMovementComponent.cpp b/Engine/Source/Runtime/Engine/Private/Components/CharacterMovementComponent.cpp
--- a/Engine/Source/Runtime/Engine/Private/Components/CharacterMovementComponent.cpp
+++ b/Engine/Source/Runtime/Engine/Private/Components/CharacterMovementComponent.cpp
@@ -5501,15 +5501,13 @@ void UCharacterMovementComponent::ForceReplicationUpdate()
 
 FVector UCharacterMovementComponent::ConstrainInputAcceleration(const FVector& InputAcceleration) const
 {
-	FVector NewAccel = InputAcceleration;
-
 	// walking or falling pawns ignore up/down sliding
-	if (IsMovingOnGround() || IsFalling())
+	if (InputAcceleration.Z != 0.f && (IsMovingOnGround() || IsFalling()))
 	{
-		NewAccel.Z = 0.f;
+		return FVector(InputAcceleration.X, InputAcceleration.Y, 0.f);
 	}
 
-	return NewAccel;
+	return InputAcceleration;
 }
 
 
