#ue4 - Inline FBodyInstance::IsInstanceSimulatingPhysics() and the accessor for bSimulatePhysics.

diff --git a/Engine/Source/Runtime/Engine/Classes/PhysicsEngine/BodyInstance.h b/Engine/Source/Runtime/Engine/Classes/PhysicsEngine/BodyInstance.h
--- a/Engine/Source/Runtime/Engine/Classes/PhysicsEngine/BodyInstance.h
+++ b/Engine/Source/Runtime/Engine/Classes/PhysicsEngine/BodyInstance.h
@@ -1085,3 +1085,13 @@ FORCEINLINE_DEBUGGABLE bool FBodyInstance::OverlapTestForBody(const FVector& Pos
 	InlineArray.Add(Body);
 	return OverlapTestForBodiesImpl(Position, Rotation, InlineArray);
 }
+
+FORCEINLINE_DEBUGGABLE bool FBodyInstance::IsInstanceSimulatingPhysics()
+{
+	return ShouldInstanceSimulatingPhysics() && IsValidBodyInstance();
+}
+
+FORCEINLINE_DEBUGGABLE bool FBodyInstance::ShouldInstanceSimulatingPhysics()
+{
+	return bSimulatePhysics;
+}
diff --git a/Engine/Source/Runtime/Engine/Private/PhysicsEngine/BodyInstance.cpp b/Engine/Source/Runtime/Engine/Private/PhysicsEngine/BodyInstance.cpp
--- a/Engine/Source/Runtime/Engine/Private/PhysicsEngine/BodyInstance.cpp
+++ b/Engine/Source/Runtime/Engine/Private/PhysicsEngine/BodyInstance.cpp
@@ -2466,17 +2466,6 @@ void FBodyInstance::SetInstanceSimulatePhysics(bool bSimulate, bool bMaintainPhy
 	UpdateInstanceSimulatePhysics();
 }
 
-bool FBodyInstance::IsInstanceSimulatingPhysics()
-{
-	// if I'm simulating or owner is simulating
-	return ShouldInstanceSimulatingPhysics() && IsValidBodyInstance();
-}
-
-bool FBodyInstance::ShouldInstanceSimulatingPhysics()
-{
-	return bSimulatePhysics;
-}
-
 bool FBodyInstance::IsValidBodyInstance() const
 {
 #if WITH_PHYSX
