added navigation updates for destructible mesh
#ue4

diff --git a/Engine/Source/Runtime/Engine/Classes/PhysicsEngine/DestructibleActor.h b/Engine/Source/Runtime/Engine/Classes/PhysicsEngine/DestructibleActor.h
--- a/Engine/Source/Runtime/Engine/Classes/PhysicsEngine/DestructibleActor.h
+++ b/Engine/Source/Runtime/Engine/Classes/PhysicsEngine/DestructibleActor.h
@@ -34,9 +34,11 @@ public:
 	// Begin AActor interface.
 #if WITH_EDITOR
 	ENGINE_API virtual bool GetReferencedContentObjects( TArray<UObject*>& Objects ) const override;
+	virtual void PostEditChangeProperty(FPropertyChangedEvent& PropertyChangedEvent) override;
 #endif // WITH_EDITOR
 	// End AActor interface.
 
+	virtual void PostLoad() override;
 
 
 public:
diff --git a/Engine/Source/Runtime/Engine/Private/Components/DestructibleComponent.cpp b/Engine/Source/Runtime/Engine/Private/Components/DestructibleComponent.cpp
--- a/Engine/Source/Runtime/Engine/Private/Components/DestructibleComponent.cpp
+++ b/Engine/Source/Runtime/Engine/Private/Components/DestructibleComponent.cpp
@@ -16,6 +16,7 @@
 #include "PhysicsEngine/DestructibleActor.h"
 #include "Engine/DestructibleMesh.h"
 #include "Components/DestructibleComponent.h"
+#include "AI/Navigation/NavigationSystem.h"
 #include "NavigationSystemHelpers.h"
 #include "PhysicalMaterials/PhysicalMaterial.h"
 #include "Kismet/GameplayStatics.h"
@@ -962,6 +963,8 @@ void UDestructibleComponent::UpdateDestructibleChunkTM(const TArray<const PxRigi
 			//if we haven't fractured it must mean that we're simulating a destructible and so we should update our ComponentToWorld based on the single rigid body
 			DestructibleComponent->SyncComponentToRBPhysics();
 		}
+
+		UNavigationSystem::UpdateNavOctree(DestructibleComponent);
 	}
 
 }
@@ -1043,12 +1046,96 @@ void UDestructibleComponent::ApplyRadiusDamage(float BaseDamage, const FVector&
 #endif
 }
 
+#if WITH_APEX
+int32 ExportPxActorGeometry(const PxRigidDynamic* PActor, TArray<PxShape*>& Shapes, FNavigableGeometryExport& GeomExport)
+{
+	int32 ShapesExportedCount = 0;
+	if (PActor == nullptr)
+	{
+		return ShapesExportedCount;
+	}
+
+	const FTransform PActorGlobalPose = P2UTransform(PActor->getGlobalPose());
+
+	const PxU32 ShapesCount = PActor->getNbShapes();
+	if (ShapesCount > PxU32(Shapes.Num()))
+	{
+		Shapes.AddUninitialized(ShapesCount - Shapes.Num());
+	}
+	const PxU32 RetrievedShapesCount = PActor->getShapes(Shapes.GetData(), Shapes.Num());
+	PxShape* const* ShapePtr = Shapes.GetData();
+	for (PxU32 ShapeIndex = 0; ShapeIndex < RetrievedShapesCount; ++ShapeIndex, ++ShapePtr)
+	{
+		if (*ShapePtr != NULL)
+		{
+			const PxTransform LocalPose = (*ShapePtr)->getLocalPose();
+			FTransform LocalToWorld = P2UTransform(LocalPose);
+			LocalToWorld.Accumulate(PActorGlobalPose);
+
+			switch ((*ShapePtr)->getGeometryType())
+			{
+			case PxGeometryType::eCONVEXMESH:
+			{
+				PxConvexMeshGeometry Geometry;
+				if ((*ShapePtr)->getConvexMeshGeometry(Geometry))
+				{
+					++ShapesExportedCount;
+
+					// @todo address Geometry.scale not being used here
+					GeomExport.ExportPxConvexMesh(Geometry.convexMesh, LocalToWorld);
+				}
+			}
+				break;
+
+			case PxGeometryType::eTRIANGLEMESH:
+			{
+				// @todo address Geometry.scale not being used here
+				PxTriangleMeshGeometry Geometry;
+				if ((*ShapePtr)->getTriangleMeshGeometry(Geometry))
+				{
+					++ShapesExportedCount;
+					if ((Geometry.triangleMesh->getTriangleMeshFlags()) & PxTriangleMeshFlag::eHAS_16BIT_TRIANGLE_INDICES)
+					{
+						GeomExport.ExportPxTriMesh16Bit(Geometry.triangleMesh, LocalToWorld);
+					}
+					else
+					{
+						GeomExport.ExportPxTriMesh32Bit(Geometry.triangleMesh, LocalToWorld);
+					}
+				}
+			}
+				break;
+
+			default:
+			{
+				UE_LOG(LogPhysics, Log, TEXT("UDestructibleComponent::DoCustomNavigableGeometryExport(): unhandled PxGeometryType, %d.")
+					, int32((*ShapePtr)->getGeometryType()));
+			}
+				break;
+			}
+		}
+	}
+
+	return ShapesExportedCount;
+}
+#endif
+
 bool UDestructibleComponent::DoCustomNavigableGeometryExport(FNavigableGeometryExport& GeomExport) const
 {
+	bool bExportFromBodySetup = true;
+
 #if WITH_APEX
 	if (ApexDestructibleActor == NULL)
 	{
-		return false;
+#if WITH_EDITORONLY_DATA
+		if (DestructibleMesh && DestructibleMesh->SourceStaticMesh)
+		{
+			GeomExport.ExportRigidBodySetup(*DestructibleMesh->SourceStaticMesh->BodySetup, ComponentToWorld);
+			bExportFromBodySetup = false;
+		}
+#endif	// WITH_EDITORONLY_DATA
+
+		return bExportFromBodySetup;
 	}
 
 	NxDestructibleActor* DestrActor = const_cast<NxDestructibleActor*>(ApexDestructibleActor);
@@ -1068,76 +1155,38 @@ bool UDestructibleComponent::DoCustomNavigableGeometryExport(FNavigableGeometryE
 		while (PActorCount--)
 		{
 			const PxRigidDynamic* PActor = *PActorBuffer++;
-			if (PActor != NULL)
-			{
-				const FTransform PActorGlobalPose = P2UTransform(PActor->getGlobalPose());
-
-				const PxU32 ShapesCount = PActor->getNbShapes();
-				if (ShapesCount > PxU32(Shapes.Num()))
-				{
-					Shapes.AddUninitialized(ShapesCount - Shapes.Num());
-				}
-				const PxU32 RetrievedShapesCount = PActor->getShapes(Shapes.GetData(), Shapes.Num());
-				PxShape* const* ShapePtr = Shapes.GetData();
-				for (PxU32 ShapeIndex = 0; ShapeIndex < RetrievedShapesCount; ++ShapeIndex, ++ShapePtr)
-				{
-					if (*ShapePtr != NULL)
-					{
-						const PxTransform LocalPose = (*ShapePtr)->getLocalPose();
-						FTransform LocalToWorld = P2UTransform(LocalPose);
-						LocalToWorld.Accumulate(PActorGlobalPose);
-
-						switch((*ShapePtr)->getGeometryType())
-						{
-						case PxGeometryType::eCONVEXMESH:
-							{
-								PxConvexMeshGeometry Geometry;
-								if ((*ShapePtr)->getConvexMeshGeometry(Geometry))
-								{
-									++ShapesExportedCount;
-
-									// @todo address Geometry.scale not being used here
-									GeomExport.ExportPxConvexMesh(Geometry.convexMesh, LocalToWorld);
-								}
-							}
-							break;
-						case PxGeometryType::eTRIANGLEMESH:
-							{
-								// @todo address Geometry.scale not being used here
-								PxTriangleMeshGeometry Geometry;
-								if ((*ShapePtr)->getTriangleMeshGeometry(Geometry))
-								{
-									++ShapesExportedCount;
-
-									if ((Geometry.triangleMesh->getTriangleMeshFlags()) & PxTriangleMeshFlag::eHAS_16BIT_TRIANGLE_INDICES)
-									{
-										GeomExport.ExportPxTriMesh16Bit(Geometry.triangleMesh, LocalToWorld);
-									}
-									else
-									{
-										GeomExport.ExportPxTriMesh32Bit(Geometry.triangleMesh, LocalToWorld);
-									}
-								}
-							}
-						default:
-							{
-								UE_LOG(LogPhysics, Log, TEXT("UDestructibleComponent::DoCustomNavigableGeometryExport(): unhandled PxGeometryType, %d.")
-									, int32((*ShapePtr)->getGeometryType()));
-							}
-							break;
-						}
-					}
-				}
-			}
+			int32 NumExported = ExportPxActorGeometry(PActor, Shapes, GeomExport);
+			ShapesExportedCount += NumExported;
 		}
+
 		ApexDestructibleActor->releasePhysXActorBuffer();
+		bExportFromBodySetup = (ShapesExportedCount == 0);
 
 		INC_DWORD_STAT_BY(STAT_Navigation_DestructiblesShapesExported, ShapesExportedCount);
 	}
+
+	// workaround for not fractured actor
+	if (bExportFromBodySetup)
+	{
+		const PxU32 VisibleChunkCount = ApexDestructibleActor->getNumVisibleChunks();
+		if (VisibleChunkCount == 1)
+		{
+			const PxU16* VisibleChunks = ApexDestructibleActor->getVisibleChunks();
+			if (VisibleChunks[0] == 0)
+			{
+				// only root chunk visible = not fractured
+				const PxRigidDynamic* PActor = ApexDestructibleActor->getChunkPhysXActor(0);
+				TArray<PxShape*> Shapes;
+				
+				const int32 NumExported = ExportPxActorGeometry(PActor, Shapes, GeomExport);
+				bExportFromBodySetup = (NumExported == 0);
+			}
+		}
+	}
 #endif // WITH_APEX
 
 	// we don't want a regular geometry export
-	return false;
+	return bExportFromBodySetup;
 }
 
 void UDestructibleComponent::Activate( bool bReset/*=false*/ )
diff --git a/Engine/Source/Runtime/Engine/Private/DestructibleActor.cpp b/Engine/Source/Runtime/Engine/Private/DestructibleActor.cpp
--- a/Engine/Source/Runtime/Engine/Private/DestructibleActor.cpp
+++ b/Engine/Source/Runtime/Engine/Private/DestructibleActor.cpp
@@ -14,7 +14,6 @@ ADestructibleActor::ADestructibleActor(const FObjectInitializer& ObjectInitializ
 	: Super(ObjectInitializer)
 {
 	DestructibleComponent = CreateDefaultSubobject<UDestructibleComponent>(TEXT("DestructibleComponent0"));
-	DestructibleComponent->bCanEverAffectNavigation = bAffectNavigation;
 	RootComponent = DestructibleComponent;
 }
 
@@ -29,9 +28,27 @@ bool ADestructibleActor::GetReferencedContentObjects( TArray<UObject*>& Objects
 	}
 	return true;
 }
-#endif // WITH_EDITOR
 
+void ADestructibleActor::PostEditChangeProperty(FPropertyChangedEvent& PropertyChangedEvent)
+{
+	Super::PostEditChangeProperty(PropertyChangedEvent);
 
+	if (PropertyChangedEvent.Property && PropertyChangedEvent.Property->GetFName() == GET_MEMBER_NAME_CHECKED(ADestructibleActor, bAffectNavigation))
+	{
+		DestructibleComponent->SetCanEverAffectNavigation(bAffectNavigation);
+	}
+}
+#endif // WITH_EDITOR
+
+void ADestructibleActor::PostLoad()
+{
+	Super::PostLoad();
+	
+	if (DestructibleComponent)
+	{
+		DestructibleComponent->bCanEverAffectNavigation = bAffectNavigation;
+	}
+}
 
 /** Returns DestructibleComponent subobject **/
 UDestructibleComponent* ADestructibleActor::GetDestructibleComponent() const { return DestructibleComponent; }
