UE4 - fix ensure on boot

diff --git a/Engine/Source/Runtime/Engine/Private/SceneManagement.cpp b/Engine/Source/Runtime/Engine/Private/SceneManagement.cpp
--- a/Engine/Source/Runtime/Engine/Private/SceneManagement.cpp
+++ b/Engine/Source/Runtime/Engine/Private/SceneManagement.cpp
@@ -189,7 +189,7 @@ static TAutoConsoleVariable<int32> CVarUseParallelGetDynamicMeshElementsTasks(
 FMeshElementCollector::FMeshElementCollector() :
 	PrimitiveSceneProxy(NULL),
 	FeatureLevel(ERHIFeatureLevel::Num),
-	bUseAsyncTasks(FApp::ShouldUseThreadingForPerformance() && CVarUseParallelGetDynamicMeshElementsTasks.GetValueOnRenderThread() > 0)
+	bUseAsyncTasks(FApp::ShouldUseThreadingForPerformance() && CVarUseParallelGetDynamicMeshElementsTasks.GetValueOnAnyThread() > 0)
 {	
 }
 
