Prevent potential crash when starting a network profile. Fixes UE-16684.

diff --git a/Engine/Source/Runtime/Engine/Private/NetworkProfiler.cpp b/Engine/Source/Runtime/Engine/Private/NetworkProfiler.cpp
--- a/Engine/Source/Runtime/Engine/Private/NetworkProfiler.cpp
+++ b/Engine/Source/Runtime/Engine/Private/NetworkProfiler.cpp
@@ -673,7 +673,7 @@ bool FNetworkProfiler::Exec( UWorld * InWorld, const TCHAR* Cmd, FOutputDevice &
 	// If we are tracking, and we don't have a file writer, force one now 
 	if ( bIsTrackingEnabled && FileWriter == NULL ) 
 	{
-		TrackSessionChange( true, InWorld->URL );
+		TrackSessionChange( true, InWorld != nullptr ? InWorld->URL : FURL() );
 		if ( FileWriter == NULL )
 		{
 			UE_LOG(LogNet, Warning, TEXT("FNetworkProfiler::Exec: FAILED to create file writer!"));
