PR #1213: Pythagorean function for use in blueprints (Contributed by [REDACTED_PERSON])  - submitted with minor modification to variable names

diff --git a/Engine/Source/Runtime/Engine/Classes/Kismet/KismetMathLibrary.h b/Engine/Source/Runtime/Engine/Classes/Kismet/KismetMathLibrary.h
--- a/Engine/Source/Runtime/Engine/Classes/Kismet/KismetMathLibrary.h
+++ b/Engine/Source/Runtime/Engine/Classes/Kismet/KismetMathLibrary.h
@@ -304,6 +304,10 @@ class ENGINE_API UKismetMathLibrary : public UBlueprintFunctionLibrary
 	UFUNCTION(BlueprintPure, meta=(DisplayName = "InRange (float)", Min="0.0", Max="1.0"), Category="Math|Float")
 	static bool InRange_FloatFloat(float Value, float Min, float Max, bool InclusiveMin = true, bool InclusiveMax = true);
 
+	/* Returns the hypotenuse of a right-angled triangle given the width and height. */
+	UFUNCTION(BlueprintPure, meta=(Keywords = "pythagorean theorem"), Category = "Math|Float")
+	static float Hypotenuse(float Width, float Height);
+	
 	/** Snaps a value to the nearest grid multiple. E.g.,
 	 *		Location = 5.1, GridSize = 10.0 : return value = 10.0
 	 * If GridSize is 0 Location is returned
diff --git a/Engine/Source/Runtime/Engine/Private/KismetMathLibrary.cpp b/Engine/Source/Runtime/Engine/Private/KismetMathLibrary.cpp
--- a/Engine/Source/Runtime/Engine/Private/KismetMathLibrary.cpp
+++ b/Engine/Source/Runtime/Engine/Private/KismetMathLibrary.cpp
@@ -355,6 +355,18 @@ bool UKismetMathLibrary::InRange_FloatFloat(float Value, float Min, float Max, b
 	return ((InclusiveMin ? (Value >= Min) : (Value > Min)) && (InclusiveMax ? (Value <= Max) : (Value < Max)));
 }	
 
+float UKismetMathLibrary::Hypotenuse(float Width, float Height)
+{
+	// This implementation avoids overflow/underflow caused by squaring width and height:
+	Width = FMath::Abs(Width);
+	Height = FMath::Abs(Height);
+
+	float Min = FGenericPlatformMath::Min(Width, Height);
+	float Max = FGenericPlatformMath::Max(Width, Height);
+	float Ratio = Min / Max;
+	return Max * FMath::Sqrt(1.f + Ratio * Ratio);
+}
+
 float UKismetMathLibrary::GridSnap_Float(float Location, float GridSize)
 {
 	return FMath::GridSnap(Location, GridSize);
