Fix up worlds with no worldsettings object on load (presumably because the class the saved one was based on no longer exists)

diff --git a/Engine/Config/BaseEngine.ini b/Engine/Config/BaseEngine.ini
--- a/Engine/Config/BaseEngine.ini
+++ b/Engine/Config/BaseEngine.ini
@@ -29,7 +29,6 @@ bDuckingOptOut=true
 ConsoleClassName=/Script/Engine.Console
 GameViewportClientClassName=/Script/Engine.GameViewportClient
 LocalPlayerClassName=/Script/Engine.LocalPlayer
-; If you change to your WorldSettings class, the previous levels won't replace first actor - WorldSettings
 WorldSettingsClassName=/Script/Engine.WorldSettings
 NavigationSystemClassName=/Script/Engine.NavigationSystem
 AvoidanceManagerClassName=/Script/Engine.AvoidanceManager
diff --git a/Engine/Source/Runtime/Engine/Classes/Engine/World.h b/Engine/Source/Runtime/Engine/Classes/Engine/World.h
--- a/Engine/Source/Runtime/Engine/Classes/Engine/World.h
+++ b/Engine/Source/Runtime/Engine/Classes/Engine/World.h
@@ -785,6 +785,9 @@ private:
 	/** Finish Async Trace Buffer **/
 	void FinishAsyncTrace();
 
+	/** Utility function that is used to ensure that a World has the correct WorldSettings */
+	void RepairWorldSettings();
+
 	/** Gameplay timers. */
 	class FTimerManager* TimerManager;
 
diff --git a/Engine/Source/Runtime/Engine/Private/Actor.cpp b/Engine/Source/Runtime/Engine/Private/Actor.cpp
--- a/Engine/Source/Runtime/Engine/Private/Actor.cpp
+++ b/Engine/Source/Runtime/Engine/Private/Actor.cpp
@@ -2385,9 +2385,10 @@ void AActor::PostSpawnInitialize(FVector const& SpawnLocation, FRotator const& S
 	// This should be the same sequence for deferred or nondeferred spawning.
 
 	// It's not safe to call UWorld accessor functions till the world info has been spawned.
-	bool const bActorsInitialized = GetWorld() && GetWorld()->AreActorsInitialized();
+	UWorld* World = GetWorld();
+	bool const bActorsInitialized = World && World->AreActorsInitialized();
 
-	CreationTime = GetWorld()->GetTimeSeconds();
+	CreationTime = (World ? World->GetTimeSeconds() : 0.f);
 
 	// Set network role.
 	check(Role == ROLE_Authority);
@@ -2452,7 +2453,8 @@ void AActor::FinishSpawning(const FTransform& Transform, bool bIsDefaultTransfor
 
 void AActor::PostActorConstruction()
 {
-	bool const bActorsInitialized = GetWorld() && GetWorld()->AreActorsInitialized();
+	UWorld* World = GetWorld();
+	bool const bActorsInitialized = World && World->AreActorsInitialized();
 
 	if (bActorsInitialized)
 	{
@@ -2473,7 +2475,7 @@ void AActor::PostActorConstruction()
 			UE_LOG(LogActor, Fatal, TEXT("%s failed to route PostInitializeComponents.  Please call Super::PostInitializeComponents() in your <className>::PostInitializeComponents() function. "), *GetFullName() );
 		}
 
-		if (GetWorld()->HasBegunPlay() && !deferBeginPlayAndUpdateOverlaps)
+		if (World->HasBegunPlay() && !deferBeginPlayAndUpdateOverlaps)
 		{
 			BeginPlay();
 		}
diff --git a/Engine/Source/Runtime/Engine/Private/World.cpp b/Engine/Source/Runtime/Engine/Private/World.cpp
--- a/Engine/Source/Runtime/Engine/Private/World.cpp
+++ b/Engine/Source/Runtime/Engine/Private/World.cpp
@@ -551,6 +551,8 @@
 	Super::PostLoad();
 	CurrentLevel = PersistentLevel;
 
+	RepairWorldSettings();
+
 	// Remove null streaming level entries (could be if level was saved with transient level streaming objects)
 	StreamingLevels.Remove(nullptr);
 	
@@ -780,6 +782,50 @@
 	return AISystem; 
 }
 
+void UWorld::RepairWorldSettings()
+{
+	// If for some reason we don't have a valid WorldSettings object go ahead and spawn one to avoid crashing.
+	// This will generally happen if a map is being moved from a different project.
+	const bool bNeedsExchange = PersistentLevel->Actors.Num() > 0;
+	const bool bNeedsDestroy = bNeedsExchange && PersistentLevel->Actors[0] != NULL;
+
+	if (PersistentLevel->Actors.Num() < 1 || PersistentLevel->Actors[0] == NULL || !PersistentLevel->Actors[0]->IsA(GEngine->WorldSettingsClass))
+	{
+		// Rename invalid WorldSettings to avoid name collisions
+		if (bNeedsDestroy)
+		{
+			PersistentLevel->Actors[0]->Rename(NULL, PersistentLevel, REN_ForceNoResetLoaders);
+		}
+		
+		FActorSpawnParameters SpawnInfo;
+		SpawnInfo.bNoCollisionFail = true;
+		SpawnInfo.Name = GEngine->WorldSettingsClass->GetFName();
+		AWorldSettings* NewWorldSettings = SpawnActor<AWorldSettings>( GEngine->WorldSettingsClass, SpawnInfo );
+
+		const int32 NewWorldSettingsActorIndex = PersistentLevel->Actors.Find( NewWorldSettings );
+
+		if (bNeedsExchange)
+		{
+			// The world info must reside at index 0.
+			Exchange(PersistentLevel->Actors[0],PersistentLevel->Actors[NewWorldSettingsActorIndex]);
+		}
+
+		// If there was an existing actor, copy its properties to the new actor and then destroy it
+		if (bNeedsDestroy)
+		{
+			NewWorldSettings->UnregisterAllComponents();
+			UEngine::CopyPropertiesForUnrelatedObjects(PersistentLevel->Actors[NewWorldSettingsActorIndex], NewWorldSettings);
+			NewWorldSettings->RegisterAllComponents();
+			PersistentLevel->Actors[NewWorldSettingsActorIndex]->Destroy();
+		}
+
+		// Re-sort actor list as we just shuffled things around.
+		PersistentLevel->SortActorList();
+
+	}
+	check(GetWorldSettings());
+}
+
 void UWorld::InitWorld(const InitializationValues IVS)
 {
 	if (!ensure(!bIsWorldInitialized))
@@ -842,46 +888,7 @@
 	PersistentLevel->OwningWorld = this;
 	PersistentLevel->bIsVisible = true;
 
-	// If for some reason we don't have a valid WorldSettings object go ahead and spawn one to avoid crashing.
-	// This will generally happen if a map is being moved from a different project.
-	const bool bNeedsExchange = PersistentLevel->Actors.Num() > 0;
-	const bool bNeedsDestroy = bNeedsExchange && PersistentLevel->Actors[0] != NULL;
+	RepairWorldSettings();
 
-	if (PersistentLevel->Actors.Num() < 1 || PersistentLevel->Actors[0] == NULL || !PersistentLevel->Actors[0]->IsA(GEngine->WorldSettingsClass))
-	{
-		// Rename invalid WorldSettings to avoid name collisions
-		if (bNeedsDestroy)
-		{
-			PersistentLevel->Actors[0]->Rename(NULL, PersistentLevel, REN_ForceNoResetLoaders);
-		}
-		
-		FActorSpawnParameters SpawnInfo;
-		SpawnInfo.bNoCollisionFail = true;
-		SpawnInfo.Name = GEngine->WorldSettingsClass->GetFName();
-		AWorldSettings* NewWorldSettings = SpawnActor<AWorldSettings>( GEngine->WorldSettingsClass, SpawnInfo );
-
-		const int32 NewWorldSettingsActorIndex = PersistentLevel->Actors.Find( NewWorldSettings );
-
-		if (bNeedsExchange)
-		{
-			// The world info must reside at index 0.
-			Exchange(PersistentLevel->Actors[0],PersistentLevel->Actors[NewWorldSettingsActorIndex]);
-		}
-
-		// If there was an existing actor, copy its properties to the new actor and then destroy it
-		if (bNeedsDestroy)
-		{
-			NewWorldSettings->UnregisterAllComponents();
-			UEngine::CopyPropertiesForUnrelatedObjects(PersistentLevel->Actors[NewWorldSettingsActorIndex], NewWorldSettings);
-			NewWorldSettings->RegisterAllComponents();
-			PersistentLevel->Actors[NewWorldSettingsActorIndex]->Destroy();
-		}
-
-		// Re-sort actor list as we just shuffled things around.
-		PersistentLevel->SortActorList();
-
-	}
-	check(GetWorldSettings());
-
 	// initialize DefaultPhysicsVolume for the world
 	// Spawned on demand by this function.
