Integrated PR #1215: fixes camera jumping in case when we try to return to target A (Contributed by [REDACTED_PERSON])

diff --git a/Engine/Source/Runtime/Engine/Private/PlayerCameraManager.cpp b/Engine/Source/Runtime/Engine/Private/PlayerCameraManager.cpp
--- a/Engine/Source/Runtime/Engine/Private/PlayerCameraManager.cpp
+++ b/Engine/Source/Runtime/Engine/Private/PlayerCameraManager.cpp
@@ -78,8 +78,8 @@ void APlayerCameraManager::SetViewTarget(class AActor* NewTarget, struct FViewTa
 		return;
 	}
 
-	// if different then new one, then assign it
-	if( NewTarget != ViewTarget.Target )
+	// if viewtarget different then new one or we're transitioning from the same target with locked outgoing, then assign it
+	if((NewTarget != ViewTarget.Target) || (PendingViewTarget.Target && BlendParams.bLockOutgoing))
 	{
 		// if a transition time is specified, then set pending view target accordingly
 		if( TransitionParams.BlendTime > 0 )
