Clean up cases where bAutoActivate was simply being used to ensure ticking didn't get disabled

diff --git a/Engine/Plugins/2D/Paper2D/Source/Paper2D/Private/Rendering/PaperBatchComponent.cpp b/Engine/Plugins/2D/Paper2D/Source/Paper2D/Private/Rendering/PaperBatchComponent.cpp
--- a/Engine/Plugins/2D/Paper2D/Source/Paper2D/Private/Rendering/PaperBatchComponent.cpp
+++ b/Engine/Plugins/2D/Paper2D/Source/Paper2D/Private/Rendering/PaperBatchComponent.cpp
@@ -12,7 +12,6 @@ UPaperBatchComponent::UPaperBatchComponent(const FObjectInitializer& ObjectIniti
 {
 	BodyInstance.SetCollisionEnabled(ECollisionEnabled::NoCollision);
 
-	bAutoActivate = true;
 	PrimaryComponentTick.TickGroup = TG_PostPhysics;
 	PrimaryComponentTick.bCanEverTick = true;
 }
diff --git a/Engine/Plugins/Runtime/CableComponent/Source/CableComponent/Private/CableComponent.cpp b/Engine/Plugins/Runtime/CableComponent/Source/CableComponent/Private/CableComponent.cpp
--- a/Engine/Plugins/Runtime/CableComponent/Source/CableComponent/Private/CableComponent.cpp
+++ b/Engine/Plugins/Runtime/CableComponent/Source/CableComponent/Private/CableComponent.cpp
@@ -342,7 +342,6 @@ UCableComponent::UCableComponent( const FObjectInitializer& ObjectInitializer )
 {
 	PrimaryComponentTick.bCanEverTick = true;
 	bTickInEditor = true;
-	bAutoActivate = true;
 
 	CableWidth = 10.f;
 	NumSegments = 10;
diff --git a/Engine/Plugins/Runtime/LeapMotionController/Source/LeapMotionController/Private/LeapMotionControllerComponent.cpp b/Engine/Plugins/Runtime/LeapMotionController/Source/LeapMotionController/Private/LeapMotionControllerComponent.cpp
--- a/Engine/Plugins/Runtime/LeapMotionController/Source/LeapMotionController/Private/LeapMotionControllerComponent.cpp
+++ b/Engine/Plugins/Runtime/LeapMotionController/Source/LeapMotionController/Private/LeapMotionControllerComponent.cpp
@@ -25,7 +25,6 @@ ULeapMotionControllerComponent::ULeapMotionControllerComponent(const class FObje
 	PrimaryComponentTick.bCanEverTick = true;
 	PrimaryComponentTick.bStartWithTickEnabled = true;
 	PrimaryComponentTick.TickGroup = TG_PrePhysics;
-	bAutoActivate = true;
 
 	Scale = 5.0f;
 
diff --git a/Engine/Plugins/Runtime/Steam/SteamVR/Source/SteamVR/Private/SteamVRChaperoneComponent.cpp b/Engine/Plugins/Runtime/Steam/SteamVR/Source/SteamVR/Private/SteamVRChaperoneComponent.cpp
--- a/Engine/Plugins/Runtime/Steam/SteamVR/Source/SteamVR/Private/SteamVRChaperoneComponent.cpp
+++ b/Engine/Plugins/Runtime/Steam/SteamVR/Source/SteamVR/Private/SteamVRChaperoneComponent.cpp
@@ -12,7 +12,6 @@ USteamVRChaperoneComponent::USteamVRChaperoneComponent(const FObjectInitializer&
 	PrimaryComponentTick.TickGroup = TG_DuringPhysics;
 
 	bTickInEditor = true;
-	bAutoActivate = true;
 
 	bWasInsideSoftBounds = true;
 }
diff --git a/Engine/Plugins/ScriptPlugin/Source/ScriptPlugin/Private/ScriptContextComponent.cpp b/Engine/Plugins/ScriptPlugin/Source/ScriptPlugin/Private/ScriptContextComponent.cpp
--- a/Engine/Plugins/ScriptPlugin/Source/ScriptPlugin/Private/ScriptContextComponent.cpp
+++ b/Engine/Plugins/ScriptPlugin/Source/ScriptPlugin/Private/ScriptContextComponent.cpp
@@ -8,7 +8,6 @@ UScriptContextComponent::UScriptContextComponent(const FObjectInitializer& Objec
 {
 	PrimaryComponentTick.bCanEverTick = true;
 	bTickInEditor = false;
-	bAutoActivate = true;
 	bWantsInitializeComponent = true;
 
 	Context = NULL;
@@ -27,7 +26,6 @@ void UScriptContextComponent::OnRegister()
 			Context = FScriptContextBase::CreateContext(ScriptClass->SourceCode, ScriptClass, ContextOwner);
 			if (!Context || !Context->CanTick())
 			{
-				bAutoActivate = false;
 				PrimaryComponentTick.bCanEverTick = false;
 			}
 		}
diff --git a/Engine/Plugins/ScriptPlugin/Source/ScriptPlugin/Private/ScriptPluginComponent.cpp b/Engine/Plugins/ScriptPlugin/Source/ScriptPlugin/Private/ScriptPluginComponent.cpp
--- a/Engine/Plugins/ScriptPlugin/Source/ScriptPlugin/Private/ScriptPluginComponent.cpp
+++ b/Engine/Plugins/ScriptPlugin/Source/ScriptPlugin/Private/ScriptPluginComponent.cpp
@@ -8,7 +8,6 @@ UScriptPluginComponent::UScriptPluginComponent(const FObjectInitializer& ObjectI
 {
 	PrimaryComponentTick.bCanEverTick = true;
 	bTickInEditor = false;
-	bAutoActivate = true;
 	bWantsInitializeComponent = true;
 
 	Context = NULL;
@@ -24,7 +23,6 @@ void UScriptPluginComponent::OnRegister()
 		Context = FScriptContextBase::CreateContext(ScriptClass->SourceCode, ScriptClass, this);
 		if (!Context || !Context->CanTick())
 		{
-			bAutoActivate = false;
 			PrimaryComponentTick.bCanEverTick = false;
 		}
 	}
diff --git a/Engine/Source/Developer/GameplayDebugger/Private/GameplayDebuggingComponent.cpp b/Engine/Source/Developer/GameplayDebugger/Private/GameplayDebuggingComponent.cpp
--- a/Engine/Source/Developer/GameplayDebugger/Private/GameplayDebuggingComponent.cpp
+++ b/Engine/Source/Developer/GameplayDebugger/Private/GameplayDebuggingComponent.cpp
@@ -134,7 +134,6 @@ UGameplayDebuggingComponent::UGameplayDebuggingComponent(const FObjectInitialize
 		
 	PrimaryComponentTick.bCanEverTick = true;
 	bWantsInitializeComponent = true;
-	bAutoActivate = false;
 	PrimaryComponentTick.bStartWithTickEnabled = false;
 
 	LastStoredPathTimeStamp = -1.f;
diff --git a/Engine/Source/Developer/GameplayDebugger/Private/GameplayDebuggingControllerComponent.cpp b/Engine/Source/Developer/GameplayDebugger/Private/GameplayDebuggingControllerComponent.cpp
--- a/Engine/Source/Developer/GameplayDebugger/Private/GameplayDebuggingControllerComponent.cpp
+++ b/Engine/Source/Developer/GameplayDebugger/Private/GameplayDebuggingControllerComponent.cpp
@@ -35,7 +35,6 @@ UGameplayDebuggingControllerComponent::UGameplayDebuggingControllerComponent(con
 #if !(UE_BUILD_SHIPPING || UE_BUILD_TEST)
 	PrimaryComponentTick.bCanEverTick = true;
 	bWantsInitializeComponent = true;
-	bAutoActivate = false;
 	bTickInEditor=true;
 	PrimaryComponentTick.bStartWithTickEnabled = false;
 
diff --git a/Engine/Source/Runtime/AIModule/Private/Actions/PawnActionsComponent.cpp b/Engine/Source/Runtime/AIModule/Private/Actions/PawnActionsComponent.cpp
--- a/Engine/Source/Runtime/AIModule/Private/Actions/PawnActionsComponent.cpp
+++ b/Engine/Source/Runtime/AIModule/Private/Actions/PawnActionsComponent.cpp
@@ -154,7 +154,6 @@ UPawnActionsComponent::UPawnActionsComponent(const FObjectInitializer& ObjectIni
 	PrimaryComponentTick.bCanEverTick = true;
 	PrimaryComponentTick.bStartWithTickEnabled = true;
 
-	bAutoActivate = true;
 	bLockedAILogic = false;
 
 	ActionEventIndex = 0;
diff --git a/Engine/Source/Runtime/AIModule/Private/BehaviorTree/BehaviorTreeComponent.cpp b/Engine/Source/Runtime/AIModule/Private/BehaviorTree/BehaviorTreeComponent.cpp
--- a/Engine/Source/Runtime/AIModule/Private/BehaviorTree/BehaviorTreeComponent.cpp
+++ b/Engine/Source/Runtime/AIModule/Private/BehaviorTree/BehaviorTreeComponent.cpp
@@ -25,7 +25,6 @@ UBehaviorTreeComponent::UBehaviorTreeComponent(const FObjectInitializer& ObjectI
 	bLoopExecution = false;
 	bWaitingForAbortingTasks = false;
 	bRequestedFlowUpdate = false;
-	bAutoActivate = true;
 	bWantsInitializeComponent = true; 
 	bIsRunning = false;
 	bIsPaused = false;
diff --git a/Engine/Source/Runtime/AIModule/Private/Navigation/PathFollowingComponent.cpp b/Engine/Source/Runtime/AIModule/Private/Navigation/PathFollowingComponent.cpp
--- a/Engine/Source/Runtime/AIModule/Private/Navigation/PathFollowingComponent.cpp
+++ b/Engine/Source/Runtime/AIModule/Private/Navigation/PathFollowingComponent.cpp
@@ -28,7 +28,6 @@ UPathFollowingComponent::FRequestCompletedSignature UPathFollowingComponent::Unb
 UPathFollowingComponent::UPathFollowingComponent(const FObjectInitializer& ObjectInitializer) : Super(ObjectInitializer)
 {
 	PrimaryComponentTick.bCanEverTick = true;
-	bAutoActivate = true;
 
 	MinAgentRadiusPct = 1.1f;
 	MinAgentHalfHeightPct = 1.05f;
diff --git a/Engine/Source/Runtime/AIModule/Private/Perception/PawnSensingComponent.cpp b/Engine/Source/Runtime/AIModule/Private/Perception/PawnSensingComponent.cpp
--- a/Engine/Source/Runtime/AIModule/Private/Perception/PawnSensingComponent.cpp
+++ b/Engine/Source/Runtime/AIModule/Private/Perception/PawnSensingComponent.cpp
@@ -24,7 +24,6 @@ UPawnSensingComponent::UPawnSensingComponent(const FObjectInitializer& ObjectIni
 
 	PrimaryComponentTick.bCanEverTick = false;
 	bWantsInitializeComponent = true;
-	bAutoActivate = false;
 
 	bEnableSensingUpdates = true;
 }
diff --git a/Engine/Source/Runtime/Engine/Private/Camera/CameraComponent.cpp b/Engine/Source/Runtime/Engine/Private/Camera/CameraComponent.cpp
--- a/Engine/Source/Runtime/Engine/Private/Camera/CameraComponent.cpp
+++ b/Engine/Source/Runtime/Engine/Private/Camera/CameraComponent.cpp
@@ -36,7 +36,6 @@ UCameraComponent::UCameraComponent(const FObjectInitializer& ObjectInitializer)
 	PostProcessBlendWeight = 1.0f;
 	bUseControllerViewRotation_DEPRECATED = true; // the previous default value before bUsePawnControlRotation replaced this var.
 	bUsePawnControlRotation = false;
-	bAutoActivate = true;
 
 	// Init deprecated var, for old code that may refer to it.
 	SetDeprecatedControllerViewRotation(*this, bUsePawnControlRotation);
diff --git a/Engine/Source/Runtime/Engine/Private/Components/LineBatchComponent.cpp b/Engine/Source/Runtime/Engine/Private/Components/LineBatchComponent.cpp
--- a/Engine/Source/Runtime/Engine/Private/Components/LineBatchComponent.cpp
+++ b/Engine/Source/Runtime/Engine/Private/Components/LineBatchComponent.cpp
@@ -100,7 +100,6 @@ private:
 ULineBatchComponent::ULineBatchComponent(const FObjectInitializer& ObjectInitializer)
 	: Super(ObjectInitializer)
 {
-	bAutoActivate = true;
 	bTickInEditor = true;
 	PrimaryComponentTick.bCanEverTick = true;
 
diff --git a/Engine/Source/Runtime/Engine/Private/Components/SceneCaptureComponent.cpp b/Engine/Source/Runtime/Engine/Private/Components/SceneCaptureComponent.cpp
--- a/Engine/Source/Runtime/Engine/Private/Components/SceneCaptureComponent.cpp
+++ b/Engine/Source/Runtime/Engine/Private/Components/SceneCaptureComponent.cpp
@@ -269,7 +269,6 @@ USceneCaptureComponent2D::USceneCaptureComponent2D(const FObjectInitializer& Obj
 	: Super(ObjectInitializer)
 {
 	FOVAngle = 90.0f;
-	bAutoActivate = true;
 	PrimaryComponentTick.bCanEverTick = true;
 	PrimaryComponentTick.TickGroup = TG_DuringPhysics;
 	// Tick in the editor so that bCaptureEveryFrame preview works
@@ -359,7 +358,6 @@ void USceneCaptureComponent2D::PostEditChangeProperty(FPropertyChangedEvent& Pro
 USceneCaptureComponentCube::USceneCaptureComponentCube(const FObjectInitializer& ObjectInitializer)
 	: Super(ObjectInitializer)
 {
-	bAutoActivate = true;
 	PrimaryComponentTick.bCanEverTick = true;
 	PrimaryComponentTick.TickGroup = TG_DuringPhysics;
 	bTickInEditor = true;
diff --git a/Engine/Source/Runtime/Engine/Private/Components/SceneComponent.cpp b/Engine/Source/Runtime/Engine/Private/Components/SceneComponent.cpp
--- a/Engine/Source/Runtime/Engine/Private/Components/SceneComponent.cpp
+++ b/Engine/Source/Runtime/Engine/Private/Components/SceneComponent.cpp
@@ -49,7 +49,6 @@ USceneComponent::USceneComponent(const FObjectInitializer& ObjectInitializer /*=
 	RelativeScale3D = FVector(1.0f, 1.0f, 1.0f);
 	// default behavior is visible
 	bVisible = true;
-	bAutoActivate = false;
 
 	NetUpdateTransform = false;
 }
diff --git a/Engine/Source/Runtime/Engine/Private/GameFramework/SpringArmComponent.cpp b/Engine/Source/Runtime/Engine/Private/GameFramework/SpringArmComponent.cpp
--- a/Engine/Source/Runtime/Engine/Private/GameFramework/SpringArmComponent.cpp
+++ b/Engine/Source/Runtime/Engine/Private/GameFramework/SpringArmComponent.cpp
@@ -15,7 +15,6 @@ USpringArmComponent::USpringArmComponent(const FObjectInitializer& ObjectInitial
 	PrimaryComponentTick.bCanEverTick = true;
 	PrimaryComponentTick.TickGroup = TG_PostPhysics;
 	
-	bAutoActivate = true;
 	bTickInEditor = true;
 	bUsePawnControlRotation = false;
 	bDoCollisionTest = true;
diff --git a/Engine/Source/Runtime/Engine/Private/PhysicsEngine/PhysicsHandleComponent.cpp b/Engine/Source/Runtime/Engine/Private/PhysicsEngine/PhysicsHandleComponent.cpp
--- a/Engine/Source/Runtime/Engine/Private/PhysicsEngine/PhysicsHandleComponent.cpp
+++ b/Engine/Source/Runtime/Engine/Private/PhysicsEngine/PhysicsHandleComponent.cpp
@@ -10,7 +10,6 @@
 UPhysicsHandleComponent::UPhysicsHandleComponent(const FObjectInitializer& ObjectInitializer)
 	: Super(ObjectInitializer)
 {
-	bAutoActivate = true;
 	PrimaryComponentTick.bCanEverTick = true;
 	PrimaryComponentTick.TickGroup = TG_PrePhysics;
 	LinearDamping = 200.0f;
diff --git a/Engine/Source/Runtime/Engine/Private/SkeletalMesh.cpp b/Engine/Source/Runtime/Engine/Private/SkeletalMesh.cpp
--- a/Engine/Source/Runtime/Engine/Private/SkeletalMesh.cpp
+++ b/Engine/Source/Runtime/Engine/Private/SkeletalMesh.cpp
@@ -4692,7 +4692,6 @@ USkinnedMeshComponent::USkinnedMeshComponent(const FObjectInitializer& ObjectIni
 	: Super(ObjectInitializer)
 	, AnimUpdateRateParams(NULL)
 {
-	bAutoActivate = true;
 	PrimaryComponentTick.bCanEverTick = true;
 	PrimaryComponentTick.TickGroup = TG_PrePhysics;	
 	WireframeColor = FColor(221, 221, 28, 255);
diff --git a/Engine/Source/Runtime/Engine/Private/WindDirectionalSource.cpp b/Engine/Source/Runtime/Engine/Private/WindDirectionalSource.cpp
--- a/Engine/Source/Runtime/Engine/Private/WindDirectionalSource.cpp
+++ b/Engine/Source/Runtime/Engine/Private/WindDirectionalSource.cpp
@@ -102,7 +102,7 @@ bool FWindSourceSceneProxy::GetWindParameters(const FVector& EvaluatePosition, F
 		return false;
 	}
 
-	Weight				= Strength;
+	Weight = Strength;
 	WindData.Direction	= Direction;
 	WindData.Speed		= Speed;
 	WindData.MinGustAmt = MinGustAmt;
@@ -119,7 +119,7 @@ bool FWindSourceSceneProxy::GetDirectionalWindParameters(FWindData& WindData, fl
 		return false;
 	}
 
-	Weight				= Strength;
+	Weight = Strength;
 	WindData.Direction	= Direction;
 	WindData.Speed		= Speed;
 	WindData.MinGustAmt = MinGustAmt;
@@ -142,9 +142,6 @@ UWindDirectionalSourceComponent::UWindDirectionalSourceComponent(const FObjectIn
 	Speed = 0.1f;
 	MinGustAmount = 0.1f;
 	MaxGustAmount = 0.2f;
-
-	// wind will be activated automatically by default
-	bAutoActivate = true;
 }
 
 void UWindDirectionalSourceComponent::CreateRenderState_Concurrent()
