- change BodyInstance to GetBodyInstance
- added warning

diff --git a/Engine/Source/Runtime/Engine/Private/Collision/WorldCollision.cpp b/Engine/Source/Runtime/Engine/Private/Collision/WorldCollision.cpp
--- a/Engine/Source/Runtime/Engine/Private/Collision/WorldCollision.cpp
+++ b/Engine/Source/Runtime/Engine/Private/Collision/WorldCollision.cpp
@@ -492,18 +492,28 @@ bool UWorld::ComponentSweepMulti(TArray<struct FHitResult>& OutHits, class UPrim
 	OutHits.Reset();
 
 #if UE_WITH_PHYSICS
-	if (!PrimComp->BodyInstance.IsValidBodyInstance())
+	const FBodyInstance* BodyInstance = PrimComp->GetBodyInstance();
+
+	if (!BodyInstance || !BodyInstance->IsValidBodyInstance())
 	{
 		UE_LOG(LogCollision, Log, TEXT("ComponentSweepMulti : (%s) No physics data"), *PrimComp->GetReadableName());
 		return false;
 	}
+
+#if !(UE_BUILD_SHIPPING || UE_BUILD_TEST)
+	if(PrimComp->IsA(USkeletalMeshComponent::StaticClass()))
+	{
+		UE_LOG(LogCollision, Warning, TEXT("ComponentSweepMulti : SkeletalMeshComponent support only root body (%s) "), *PrimComp->GetReadableName());
+	}
+#endif
+
 #endif
 
 	SCOPE_CYCLE_COUNTER(STAT_Collision_GeomSweepMultiple);
 	bool bHaveBlockingHit = false;
 
 #if WITH_PHYSX
-	ExecuteOnPxRigidActorReadOnly(&PrimComp->BodyInstance, [&] (const PxRigidActor* PRigidActor)
+	ExecuteOnPxRigidActorReadOnly(BodyInstance, [&] (const PxRigidActor* PRigidActor)
 	{
 		// Get all the shapes from the actor
 		TArray<PxShape*, TInlineAllocator<16>> PShapes;
