#ue4 - Deprecated UProjectileMovementComponent::GetEffectiveGravityZ() in favor of just overriding GetGravityZ().

diff --git a/Engine/Source/Runtime/Engine/Classes/GameFramework/ProjectileMovementComponent.h b/Engine/Source/Runtime/Engine/Classes/GameFramework/ProjectileMovementComponent.h
--- a/Engine/Source/Runtime/Engine/Classes/GameFramework/ProjectileMovementComponent.h
+++ b/Engine/Source/Runtime/Engine/Classes/GameFramework/ProjectileMovementComponent.h
@@ -294,6 +294,9 @@ protected:
 	virtual FVector ComputeHomingAcceleration(const FVector& InVelocity, float DeltaTime) const;
 
 	/** Compute gravity effect given current physics volume, projectile gravity scale, etc. */
+	virtual float GetGravityZ() const override;
+
+	DEPRECATED(4.9, "GetEffectiveGravityZ() is deprecated, use GetGravityZ() instead, or UMovementComponent::GetGravityZ() if you want only the unmodified gravity.")
 	float GetEffectiveGravityZ() const;
 
 	/** Minimum delta time considered when ticking. Delta times below this are not considered. This is a very small non-zero positive value to avoid potential divide-by-zero in simulation code. */
diff --git a/Engine/Source/Runtime/Engine/Private/Components/ProjectileMovementComponent.cpp b/Engine/Source/Runtime/Engine/Private/Components/ProjectileMovementComponent.cpp
--- a/Engine/Source/Runtime/Engine/Private/Components/ProjectileMovementComponent.cpp
+++ b/Engine/Source/Runtime/Engine/Private/Components/ProjectileMovementComponent.cpp
@@ -378,7 +378,7 @@ FVector UProjectileMovementComponent::ComputeAcceleration(const FVector& InVeloc
 {
 	FVector Acceleration(FVector::ZeroVector);
 
-	Acceleration.Z += GetEffectiveGravityZ();
+	Acceleration.Z += GetGravityZ();
 
 	if (bIsHomingProjectile && HomingTargetComponent.IsValid())
 	{
@@ -395,10 +395,16 @@ FVector UProjectileMovementComponent::ComputeHomingAcceleration(const FVector& I
 	return HomingAcceleration;
 }
 
+// Deprecated
 float UProjectileMovementComponent::GetEffectiveGravityZ() const
+{
+	return GetGravityZ();
+}
+
+float UProjectileMovementComponent::GetGravityZ() const
 {
 	// TODO: apply buoyancy if in water
-	return ShouldApplyGravity() ? GetGravityZ() * ProjectileGravityScale : 0.f;
+	return ShouldApplyGravity() ? Super::GetGravityZ() * ProjectileGravityScale : 0.f;
 }
 
 
@@ -535,7 +541,7 @@ bool UProjectileMovementComponent::CheckStillInWorld()
 
 bool UProjectileMovementComponent::ShouldUseSubStepping() const
 {
-	return bForceSubStepping || (GetEffectiveGravityZ() != 0.f) || (bIsHomingProjectile && HomingTargetComponent.IsValid());
+	return bForceSubStepping || (GetGravityZ() != 0.f) || (bIsHomingProjectile && HomingTargetComponent.IsValid());
 }
 
 
