#ue4 Adding Get/Set WorldToMeters scale HMDFL functions

diff --git a/Engine/Source/Runtime/Engine/Classes/Kismet/HeadMountedDisplayFunctionLibrary.h b/Engine/Source/Runtime/Engine/Classes/Kismet/HeadMountedDisplayFunctionLibrary.h
--- a/Engine/Source/Runtime/Engine/Classes/Kismet/HeadMountedDisplayFunctionLibrary.h
+++ b/Engine/Source/Runtime/Engine/Classes/Kismet/HeadMountedDisplayFunctionLibrary.h
@@ -114,4 +114,21 @@ class ENGINE_API UHeadMountedDisplayFunctionLibrary : public UBlueprintFunctionL
 	 */
 	UFUNCTION(BlueprintPure, Category = "Input|HeadMountedDisplay")
 	static float GetScreenPercentage();
+
+
+	/**
+	* Sets the World to Meters scale, which changes the scale of the world as perceived by the player
+	*
+	* @param NewScale	Specifies how many Unreal units correspond to one meter in the real world
+	*/
+	UFUNCTION(BlueprintCallable, Category = "Input|HeadMountedDisplay", meta = (WorldContext = "WorldContext"))
+	static void SetWorldToMetersScale(UObject* WorldContext, float NewScale = 100.f);
+
+	/**
+	* Returns the World to Meters scale, which corresponds to the scale of the world as perceived by the player
+	*
+	* @return	How many Unreal units correspond to one meter in the real world
+	*/
+	UFUNCTION(BlueprintPure, Category = "Input|HeadMountedDisplay", meta = (WorldContext = "WorldContext"))
+	static float GetWorldToMetersScale(UObject* WorldContext);
 };
diff --git a/Engine/Source/Runtime/Engine/Private/HeadMountedDisplayFunctionLibrary.cpp b/Engine/Source/Runtime/Engine/Private/HeadMountedDisplayFunctionLibrary.cpp
--- a/Engine/Source/Runtime/Engine/Private/HeadMountedDisplayFunctionLibrary.cpp
+++ b/Engine/Source/Runtime/Engine/Private/HeadMountedDisplayFunctionLibrary.cpp
@@ -155,3 +155,16 @@ float UHeadMountedDisplayFunctionLibrary::GetScreenPercentage()
 	}
 	return 0.0f;
 }
+
+void UHeadMountedDisplayFunctionLibrary::SetWorldToMetersScale(UObject* WorldContext, float NewScale)
+{
+	if (WorldContext)
+	{
+		WorldContext->GetWorld()->GetWorldSettings()->WorldToMeters = NewScale;
+	}
+}
+
+float UHeadMountedDisplayFunctionLibrary::GetWorldToMetersScale(UObject* WorldContext)
+{
+	return WorldContext ? WorldContext->GetWorld()->GetWorldSettings()->WorldToMeters : 0.f;
+}
\ No newline at end of file
