Redirected NavLinkProxy's vlogs to NavigationSystem's log #UE4

Merged over from dev branch

diff --git a/Engine/Source/Runtime/Engine/Classes/AI/Navigation/NavLinkProxy.h b/Engine/Source/Runtime/Engine/Classes/AI/Navigation/NavLinkProxy.h
--- a/Engine/Source/Runtime/Engine/Classes/AI/Navigation/NavLinkProxy.h
+++ b/Engine/Source/Runtime/Engine/Classes/AI/Navigation/NavLinkProxy.h
@@ -69,6 +69,10 @@ public:
 	virtual void PostEditChangeProperty(FPropertyChangedEvent& PropertyChangedEvent) override;
 #endif // WITH_EDITOR
 
+#if ENABLE_VISUAL_LOG
+	virtual void BeginPlay() override;
+#endif // ENABLE_VISUAL_LOG
+
 	virtual FBox GetComponentsBoundingBox(bool bNonColliding = false) const override;
 
 	//////////////////////////////////////////////////////////////////////////
diff --git a/Engine/Source/Runtime/Engine/Private/AI/Navigation/NavLinkProxy.cpp b/Engine/Source/Runtime/Engine/Private/AI/Navigation/NavLinkProxy.cpp
--- a/Engine/Source/Runtime/Engine/Private/AI/Navigation/NavLinkProxy.cpp
+++ b/Engine/Source/Runtime/Engine/Private/AI/Navigation/NavLinkProxy.cpp
@@ -87,6 +87,17 @@ void ANavLinkProxy::PostEditChangeProperty(FPropertyChangedEvent& PropertyChange
 }
 #endif // WITH_EDITOR
 
+#if ENABLE_VISUAL_LOG
+void ANavLinkProxy::BeginPlay()
+{
+	UNavigationSystem* NavSys = UNavigationSystem::GetCurrent(GetWorld());
+	if (NavSys)
+	{
+		REDIRECT_OBJECT_TO_VLOG(this, NavSys);
+	}
+}
+#endif // ENABLE_VISUAL_LOG
+
 void ANavLinkProxy::GetNavigationData(FNavigationRelevantData& Data) const
 {
 	NavigationHelper::ProcessNavLinkAndAppend(&Data.Modifiers, this, PointLinks);
