Fix non-unity build.

diff --git a/Engine/Source/Runtime/Engine/Private/AI/Navigation/NavLinkProxy.cpp b/Engine/Source/Runtime/Engine/Private/AI/Navigation/NavLinkProxy.cpp
--- a/Engine/Source/Runtime/Engine/Private/AI/Navigation/NavLinkProxy.cpp
+++ b/Engine/Source/Runtime/Engine/Private/AI/Navigation/NavLinkProxy.cpp
@@ -10,6 +10,7 @@
 #include "AI/Navigation/NavLinkProxy.h"
 #include "AI/Navigation/NavLinkRenderingComponent.h"
 #include "NavigationSystemHelpers.h"
+#include "VisualLogger.h"
 
 ANavLinkProxy::ANavLinkProxy(const FObjectInitializer& ObjectInitializer) : Super(ObjectInitializer)
 {
