Fixed annoying issue causing navmesh to redraw every single time editor regains focus #UE4

Merged over from dev branch

diff --git a/Engine/Source/Runtime/Engine/Private/Components/NavMeshRenderingComponent.cpp b/Engine/Source/Runtime/Engine/Private/Components/NavMeshRenderingComponent.cpp
--- a/Engine/Source/Runtime/Engine/Private/Components/NavMeshRenderingComponent.cpp
+++ b/Engine/Source/Runtime/Engine/Private/Components/NavMeshRenderingComponent.cpp
@@ -18,6 +18,30 @@
 
 #include "NavMeshRenderingHelpers.h"
 
+#if WITH_EDITOR
+namespace
+{
+	bool AreAnyViewportsRelevant(const UWorld* World)
+	{
+		FWorldContext* WorldContext = GEngine->GetWorldContextFromWorld(World);
+		if (WorldContext && WorldContext->GameViewport)
+		{
+			return true;
+		}
+		
+		for (FEditorViewportClient* CurrentViewport : GEditor->AllViewportClients)
+		{
+			if (CurrentViewport && CurrentViewport->IsVisible())
+			{
+				return true;
+			}
+		}
+
+		return false;
+	}
+}
+#endif
+
 UNavMeshRenderingComponent::UNavMeshRenderingComponent(const FObjectInitializer& ObjectInitializer)
 	: Super(ObjectInitializer)
 {
@@ -68,7 +92,16 @@ bool UNavMeshRenderingComponent::IsNavigationShowFlagSet(const UWorld* World)
 
 void UNavMeshRenderingComponent::TimerFunction()
 {
-	const bool bShowNavigation = IsNavigationShowFlagSet(GetWorld());
+	const UWorld* World = GetWorld();
+#if WITH_EDITOR
+	if (GEditor && (AreAnyViewportsRelevant(World) == false))
+	{
+		// unable to tell if the flag is on or not
+		return;
+	}
+#endif // WITH_EDITOR
+
+	const bool bShowNavigation = IsNavigationShowFlagSet(World);
 
 	if (bShowNavigation != !!bCollectNavigationData)
 	{
