UE4 - framework for parallel GDME tasks

diff --git a/Engine/Source/Runtime/Engine/Private/Particles/ParticleSystemRender.cpp b/Engine/Source/Runtime/Engine/Private/Particles/ParticleSystemRender.cpp
--- a/Engine/Source/Runtime/Engine/Private/Particles/ParticleSystemRender.cpp
+++ b/Engine/Source/Runtime/Engine/Private/Particles/ParticleSystemRender.cpp
@@ -1000,9 +1000,10 @@ void FDynamicSpriteEmitterData::GetDynamicMeshElementsEmitter(const FParticleSys
 			if ( Allocation.IsValid() && (!bUsesDynamicParameter || DynamicParameterAllocation.IsValid()) )
 			{
 				// Sort the particles if needed.
-				FParticleOrder* ParticleOrder = NULL;
+				bool bSort = false;
 				if (SourceData->SortMode != PSORTMODE_None)
 				{
+					QUICK_SCOPE_CYCLE_COUNTER(STAT_FDynamicSpriteEmitterData_GetDynamicMeshElementsEmitter_GetParticleOrderData);
 					// If material is using unlit translucency and the blend mode is translucent then we need to sort (back to front)
 					const FMaterial* Material = MaterialResource[bSelected]->GetMaterial(FeatureLevel);
 
@@ -1011,108 +1012,61 @@ void FDynamicSpriteEmitterData::GetDynamicMeshElementsEmitter(const FParticleSys
 						((SourceData->SortMode == PSORTMODE_Age_OldestFirst) || (SourceData->SortMode == PSORTMODE_Age_NewestFirst)))
 						)
 					{
-						ParticleOrder = GParticleOrderPool.GetParticleOrderData(ParticleCount);
+						bSort = true;
 					}
 				}
-
-				if (Collector.ShouldUseTasks())
 				{
-					class FHandleParticleSortAndFill
+					QUICK_SCOPE_CYCLE_COUNTER(STAT_FDynamicSpriteEmitterData_PerParticleWorkOrTasks);
+					if (Collector.ShouldUseTasks())
 					{
-						const FDynamicSpriteEmitterData* Target;
-						const FSceneView* View;
-						FParticleOrder* ParticleOrder;
-						FMatrix LocalToWorld;
-						FGlobalDynamicVertexBuffer::FAllocation Allocation;
-						FGlobalDynamicVertexBuffer::FAllocation DynamicParameterAllocation;
-						bool bInstanced;
-					public:
-
-						FHandleParticleSortAndFill(						
-							const FDynamicSpriteEmitterData* InTarget,
-							const FSceneView* InView,
-							FParticleOrder* InParticleOrder,
-							const FMatrix& LocalToWorld,
-							const FGlobalDynamicVertexBuffer::FAllocation& InAllocation,
-							const FGlobalDynamicVertexBuffer::FAllocation& InDynamicParameterAllocation,
-							bool bInInstanced)
-							: Target(InTarget)
-							, View(InView)
-							, ParticleOrder(InParticleOrder)
-							, LocalToWorld(LocalToWorld)
-							, Allocation(InAllocation)
-							, DynamicParameterAllocation(InDynamicParameterAllocation)
-							, bInstanced(bInInstanced)
-						{
-						}
-
-						FORCEINLINE TStatId GetStatId() const
+						Collector.AddTask(
+							[this, SourceData, View, Proxy, Allocation, DynamicParameterAllocation, bInstanced, bSort, ParticleCount]()
+							{
+								QUICK_SCOPE_CYCLE_COUNTER(STAT_FDynamicSpriteEmitterData_GetDynamicMeshElementsEmitter_Task);
+								FMemMark Mark(FMemStack::Get());
+								FParticleOrder* ParticleOrder = NULL;
+								if (bSort)
+								{
+									ParticleOrder = (FParticleOrder*)FMemStack::Get().Alloc(sizeof(FParticleOrder)* ParticleCount, ALIGNOF(FParticleOrder));
+									SortSpriteParticles(SourceData->SortMode, SourceData->bUseLocalSpace, SourceData->ActiveParticleCount, 
+										SourceData->ParticleData, SourceData->ParticleStride, SourceData->ParticleIndices,
+										View, Proxy->GetLocalToWorld(), ParticleOrder);
+								}
+								// Fill vertex buffers.
+								if(bInstanced)
+								{
+									GetVertexAndIndexData(Allocation.Buffer, DynamicParameterAllocation.Buffer, NULL, ParticleOrder, View->ViewMatrices.ViewOrigin, Proxy->GetLocalToWorld());
+								}
+								else
+								{
+									GetVertexAndIndexDataNonInstanced(Allocation.Buffer, DynamicParameterAllocation.Buffer, NULL, ParticleOrder, View->ViewMatrices.ViewOrigin, Proxy->GetLocalToWorld());
+								}
+							}
+						);
+					}
+					else
+					{
+						FParticleOrder* ParticleOrder = NULL;
+						if (bSort)
 						{
-							RETURN_QUICK_DECLARE_CYCLE_STAT(FHandleParticleSortAndFill, STATGROUP_TaskGraphTasks);
+							ParticleOrder = GParticleOrderPool.GetParticleOrderData(ParticleCount);
+							SortSpriteParticles(SourceData->SortMode, SourceData->bUseLocalSpace, SourceData->ActiveParticleCount, 
+								SourceData->ParticleData, SourceData->ParticleStride, SourceData->ParticleIndices,
+								View, Proxy->GetLocalToWorld(), ParticleOrder);
 						}
-
-						static FORCEINLINE ENamedThreads::Type GetDesiredThread()
+						// Fill vertex buffers.
+						if(bInstanced)
 						{
-							return ENamedThreads::AnyThread;
+							GetVertexAndIndexData(Allocation.Buffer, DynamicParameterAllocation.Buffer, NULL, ParticleOrder, View->ViewMatrices.ViewOrigin, Proxy->GetLocalToWorld());
 						}
-
-						static ESubsequentsMode::Type GetSubsequentsMode() { return ESubsequentsMode::TrackSubsequents; }
-
-						void DoTask(ENamedThreads::Type CurrentThread, const FGraphEventRef& MyCompletionGraphEvent)
+						else
 						{
-							const FDynamicSpriteEmitterReplayDataBase* SourceData = Target->GetSourceData();
-							if (ParticleOrder)
-							{
-								Target->SortSpriteParticles(SourceData->SortMode, SourceData->bUseLocalSpace, SourceData->ActiveParticleCount, 
-									SourceData->ParticleData, SourceData->ParticleStride, SourceData->ParticleIndices,
-									View, LocalToWorld, ParticleOrder);
-							}
-							// Fill vertex buffers.
-							if(bInstanced)
-							{
-								Target->GetVertexAndIndexData(Allocation.Buffer, DynamicParameterAllocation.Buffer, NULL, ParticleOrder, View->ViewMatrices.ViewOrigin, LocalToWorld);
-							}
-							else
-							{
-								Target->GetVertexAndIndexDataNonInstanced(Allocation.Buffer, DynamicParameterAllocation.Buffer, NULL, ParticleOrder, View->ViewMatrices.ViewOrigin, LocalToWorld);
-							}
+							GetVertexAndIndexDataNonInstanced(Allocation.Buffer, DynamicParameterAllocation.Buffer, NULL, ParticleOrder, View->ViewMatrices.ViewOrigin, Proxy->GetLocalToWorld());
 						}
-					};
-
-					check(IsInRenderingThread());
-
-					Collector.AddTask(
-						TGraphTask<FHandleParticleSortAndFill>::CreateTask(nullptr, ENamedThreads::RenderThread).ConstructAndDispatchWhenReady(
-							this,
-							View,
-							ParticleOrder,
-							Proxy->GetLocalToWorld(),
-							Allocation,
-							DynamicParameterAllocation,
-							bInstanced
-						)
-					);
-
-				}
-				else
-				{
-					if (ParticleOrder)
-					{
-						SortSpriteParticles(SourceData->SortMode, SourceData->bUseLocalSpace, SourceData->ActiveParticleCount, 
-							SourceData->ParticleData, SourceData->ParticleStride, SourceData->ParticleIndices,
-							View, Proxy->GetLocalToWorld(), ParticleOrder);
-					}
-					// Fill vertex buffers.
-					if(bInstanced)
-					{
-						GetVertexAndIndexData(Allocation.Buffer, DynamicParameterAllocation.Buffer, NULL, ParticleOrder, View->ViewMatrices.ViewOrigin, Proxy->GetLocalToWorld());
-					}
-					else
-					{
-						GetVertexAndIndexDataNonInstanced(Allocation.Buffer, DynamicParameterAllocation.Buffer, NULL, ParticleOrder, View->ViewMatrices.ViewOrigin, Proxy->GetLocalToWorld());
 					}
 				}
 
+
 				// Create per-view uniform buffer.
 				FParticleSpriteUniformParameters PerViewUniformParameters = UniformParameters;
 				FVector2D ObjectNDCPosition;
@@ -1135,7 +1089,7 @@ void FDynamicSpriteEmitterData::GetDynamicMeshElementsEmitter(const FParticleSys
 
 				// Set the sprite uniform buffer for this view.
 				SpriteVertexFactory->SetSpriteUniformBuffer( CollectorResources.UniformBuffer );
-				SpriteVertexFactory->SetInstanceBuffer( Allocation.VertexBuffer, Allocation.VertexOffset, GetDynamicVertexStride(FeatureLevel), bInstanced );
+				SpriteVertexFactory->SetInstanceBuffer( Allocation.VertexBuffer, Allocation.VertexOffset, VertexSize, bInstanced );
 				SpriteVertexFactory->SetDynamicParameterBuffer( DynamicParameterAllocation.VertexBuffer, DynamicParameterAllocation.VertexOffset, GetDynamicParameterVertexStride(), bInstanced );
 
 				// Construct the mesh element to render.
@@ -1479,7 +1433,7 @@ void FDynamicMeshEmitterData::GetDynamicMeshElementsEmitter(const FParticleSyste
 			FDynamicMeshEmitterCollectorResources& CollectorResources = Collector.AllocateOneFrameResource<FDynamicMeshEmitterCollectorResources>();
 			CollectorResources.VertexFactory.SetParticleFactoryType(PVFT_Mesh);
 			CollectorResources.VertexFactory.SetFeatureLevel(FeatureLevel);
-			CollectorResources.VertexFactory.SetStrides(GetDynamicVertexStride(FeatureLevel), GetDynamicParameterVertexStride());
+			CollectorResources.VertexFactory.SetStrides(InstanceVertexStride, GetDynamicParameterVertexStride());
 			CollectorResources.VertexFactory.InitResource();
 			SetupVertexFactory(&CollectorResources.VertexFactory, StaticMesh->RenderData->LODResources[0]);
 
@@ -1502,10 +1456,22 @@ void FDynamicMeshEmitterData::GetDynamicMeshElementsEmitter(const FParticleSyste
 				if(Allocation.IsValid() && (!bUsesDynamicParameter || DynamicParameterAllocation.IsValid()))
 				{
 					// Fill instance buffer.
-					GetInstanceData(Allocation.Buffer, DynamicParameterAllocation.Buffer, Proxy, View);
+					if (Collector.ShouldUseTasks())
+					{
+						Collector.AddTask(
+							[this, View, Proxy, Allocation, DynamicParameterAllocation]()
+							{
+								GetInstanceData(Allocation.Buffer, DynamicParameterAllocation.Buffer, Proxy, View);
+							}
+						);
+					}
+					else
+					{
+						GetInstanceData(Allocation.Buffer, DynamicParameterAllocation.Buffer, Proxy, View);
+					}
 				}
 
-				MeshVertexFactory->SetInstanceBuffer(Allocation.VertexBuffer, Allocation.VertexOffset, GetDynamicVertexStride(FeatureLevel));
+				MeshVertexFactory->SetInstanceBuffer(Allocation.VertexBuffer, Allocation.VertexOffset, InstanceVertexStride);
 				MeshVertexFactory->SetDynamicParameterBuffer(DynamicParameterAllocation.VertexBuffer, DynamicParameterAllocation.VertexOffset , GetDynamicParameterVertexStride());
 			}
 			else
@@ -1521,7 +1487,19 @@ void FDynamicMeshEmitterData::GetDynamicMeshElementsEmitter(const FParticleSyste
 				}
 
 				// Fill instance buffer.
-				GetInstanceData((void*)InstanceVerticesCPU->InstanceDataAllocationsCPU.GetData(), (void*)InstanceVerticesCPU->DynamicParameterDataAllocationsCPU.GetData(), Proxy, View);
+				if (Collector.ShouldUseTasks())
+				{
+					Collector.AddTask(
+						[this, View, Proxy, InstanceVerticesCPU]()
+						{
+							GetInstanceData((void*)InstanceVerticesCPU->InstanceDataAllocationsCPU.GetData(), (void*)InstanceVerticesCPU->DynamicParameterDataAllocationsCPU.GetData(), Proxy, View);
+						}
+					);
+				}
+				else
+				{
+					GetInstanceData((void*)InstanceVerticesCPU->InstanceDataAllocationsCPU.GetData(), (void*)InstanceVerticesCPU->DynamicParameterDataAllocationsCPU.GetData(), Proxy, View);
+				}
 			}
 
 			Proxy->UpdateWorldSpacePrimitiveUniformBuffer();
diff --git a/Engine/Source/Runtime/Engine/Private/SceneManagement.cpp b/Engine/Source/Runtime/Engine/Private/SceneManagement.cpp
--- a/Engine/Source/Runtime/Engine/Private/SceneManagement.cpp
+++ b/Engine/Source/Runtime/Engine/Private/SceneManagement.cpp
@@ -176,6 +176,7 @@ FMeshBatchAndRelevance::FMeshBatchAndRelevance(const FMeshBatch& InMesh, const F
 	Mesh(&InMesh),
 	PrimitiveSceneProxy(InPrimitiveSceneProxy)
 {
+	QUICK_SCOPE_CYCLE_COUNTER(STAT_FMeshBatchAndRelevance);
 	EBlendMode BlendMode = InMesh.MaterialRenderProxy->GetMaterial(FeatureLevel)->GetBlendMode();
 	bHasOpaqueOrMaskedMaterial = !IsTranslucentBlendMode(BlendMode);
 	bRenderInMainPass = PrimitiveSceneProxy->ShouldRenderInMainPass();
@@ -194,16 +195,24 @@ FMeshElementCollector::FMeshElementCollector() :
 }
 
 
-void FMeshElementCollector::WaitForTasks()
+void FMeshElementCollector::ProcessTasks()
 {
 	check(IsInRenderingThread());
-	check(!TasksToWaitFor.Num() || bUseAsyncTasks);
+	check(!ParallelTasks.Num() || bUseAsyncTasks);
 
-	if (TasksToWaitFor.Num())
+	if (ParallelTasks.Num())
 	{
-		QUICK_SCOPE_CYCLE_COUNTER(STAT_FMeshElementCollector_WaitForTasks);
-		FTaskGraphInterface::Get().WaitUntilTasksComplete(TasksToWaitFor, ENamedThreads::RenderThread_Local);
-		TasksToWaitFor.Empty();
+		QUICK_SCOPE_CYCLE_COUNTER(STAT_FMeshElementCollector_ProcessTasks);
+		TArray<TFunction<void()>*, SceneRenderingAllocator>& LocalParallelTasks(ParallelTasks);
+		ParallelFor(ParallelTasks.Num(), 
+			[&LocalParallelTasks](int32 Index)
+			{
+				TFunction<void()>* Func = LocalParallelTasks[Index];
+				(*Func)();
+				Func->~TFunction<void()>();
+			}
+			);
+		ParallelTasks.Empty();
 	}
 }
 
diff --git a/Engine/Source/Runtime/Engine/Public/SceneManagement.h b/Engine/Source/Runtime/Engine/Public/SceneManagement.h
--- a/Engine/Source/Runtime/Engine/Public/SceneManagement.h
+++ b/Engine/Source/Runtime/Engine/Public/SceneManagement.h
@@ -1332,7 +1332,8 @@ public:
 	 */
 	inline FMeshBatch& AllocateMesh()
 	{
-		return *(new (MeshBatchStorage) FMeshBatch());
+		const int32 Index = MeshBatchStorage.Add(1);
+		return MeshBatchStorage[Index];
 	}
 
 	/** 
@@ -1360,12 +1361,12 @@ public:
 		return bUseAsyncTasks;
 	}
 
-	FORCEINLINE void AddTask(const FGraphEventRef& Task)
+	FORCEINLINE void AddTask(const TFunction<void()>& Task)
 	{
-		TasksToWaitFor.Add(Task);
+		ParallelTasks.Add(new (FMemStack::Get()) TFunction<void()>(Task));
 	}
 
-	ENGINE_API void WaitForTasks();
+	ENGINE_API void ProcessTasks();
 
 private:
 
@@ -1373,7 +1374,7 @@ private:
 
 	~FMeshElementCollector()
 	{
-		check(!TasksToWaitFor.Num()); // We should have blocked on this already
+		check(!ParallelTasks.Num()); // We should have blocked on this already
 		for (int32 ProxyIndex = 0; ProxyIndex < TemporaryProxies.Num(); ProxyIndex++)
 		{
 			delete TemporaryProxies[ProxyIndex];
@@ -1446,7 +1447,7 @@ private:
 	const bool bUseAsyncTasks;
 
 	/** Tasks to wait for at the end of gathering dynamic mesh elements. */
-	FGraphEventArray TasksToWaitFor;
+	TArray<TFunction<void()>*, SceneRenderingAllocator> ParallelTasks;
 
 
 	friend class FSceneRenderer;
diff --git a/Engine/Source/Runtime/Renderer/Private/ConvertToUniformMesh.cpp b/Engine/Source/Runtime/Renderer/Private/ConvertToUniformMesh.cpp
--- a/Engine/Source/Runtime/Renderer/Private/ConvertToUniformMesh.cpp
+++ b/Engine/Source/Runtime/Renderer/Private/ConvertToUniformMesh.cpp
@@ -328,7 +328,7 @@ int32 FUniformMeshConverter::Convert(
 	PrimitiveSceneInfo->Proxy->GetDynamicMeshElements(ViewsArray, *(View.Family), 0x1, Renderer.MeshCollector);
 	View.bRenderFirstInstanceOnly = false;
 
-	Renderer.MeshCollector.WaitForTasks();
+	Renderer.MeshCollector.ProcessTasks();
 
 	int32 NumTriangles = 0;
 
diff --git a/Engine/Source/Runtime/Renderer/Private/SceneVisibility.cpp b/Engine/Source/Runtime/Renderer/Private/SceneVisibility.cpp
--- a/Engine/Source/Runtime/Renderer/Private/SceneVisibility.cpp
+++ b/Engine/Source/Runtime/Renderer/Private/SceneVisibility.cpp
@@ -1273,7 +1273,7 @@ void FSceneRenderer::GatherDynamicMeshElements(
 			}
 		}
 	}
-	MeshCollector.WaitForTasks();
+	MeshCollector.ProcessTasks();
 }
 
 static void MarkAllPrimitivesForReflectionProxyUpdate(FScene* Scene)
diff --git a/Engine/Source/Runtime/Renderer/Private/ShadowSetup.cpp b/Engine/Source/Runtime/Renderer/Private/ShadowSetup.cpp
--- a/Engine/Source/Runtime/Renderer/Private/ShadowSetup.cpp
+++ b/Engine/Source/Runtime/Renderer/Private/ShadowSetup.cpp
@@ -990,7 +990,7 @@ void FProjectedShadowInfo::GatherDynamicMeshElements(FSceneRenderer& Renderer, F
     
 		    FoundView->ViewMatrices.ViewMatrix = OriginalViewMatrix;
 
-			Renderer.MeshCollector.WaitForTasks();
+			Renderer.MeshCollector.ProcessTasks();
 	    }
     }
 }
