Additional unnecessary bAutoActivate removed

diff --git a/Engine/Source/Runtime/Engine/Private/Components/SkeletalMeshComponent.cpp b/Engine/Source/Runtime/Engine/Private/Components/SkeletalMeshComponent.cpp
--- a/Engine/Source/Runtime/Engine/Private/Components/SkeletalMeshComponent.cpp
+++ b/Engine/Source/Runtime/Engine/Private/Components/SkeletalMeshComponent.cpp
@@ -108,7 +108,6 @@ public:
 USkeletalMeshComponent::USkeletalMeshComponent(const FObjectInitializer& ObjectInitializer)
 	: Super(ObjectInitializer)
 {
-	bAutoActivate = true;
 	PrimaryComponentTick.bCanEverTick = true;
 	PrimaryComponentTick.TickGroup = TG_PrePhysics;
 	bWantsInitializeComponent = true;
