Fix crash when there is no PlayAudioComp

diff --git a/Engine/Source/Runtime/Engine/Private/Interpolation.cpp b/Engine/Source/Runtime/Engine/Private/Interpolation.cpp
--- a/Engine/Source/Runtime/Engine/Private/Interpolation.cpp
+++ b/Engine/Source/Runtime/Engine/Private/Interpolation.cpp
@@ -8015,7 +8015,10 @@ void UInterpTrackSound::PreviewUpdateTrack(float NewPosition, UInterpTrackInst*
 	const bool bJustLooped = NewPosition < MatineeActor->InterpPosition && MatineeActor->bIsPlaying;
 	if (bJustLooped)
 	{
-		SoundInst->PlayAudioComp->Stop();
+		if (SoundInst->PlayAudioComp)
+		{
+			SoundInst->PlayAudioComp->Stop();
+		}
 		bPlaying = false;
 		const float Epsilon = 0.1f;
 		SoundInst->LastUpdatePosition = NewPosition - Epsilon;
