#ue4 - Remove "VectorBitwiseAND" defined only for UnrealMathSSE. There is already "VectorBitwiseAnd" defined for all platforms.

diff --git a/Engine/Source/Runtime/Core/Public/Math/UnrealMathSSE.h b/Engine/Source/Runtime/Core/Public/Math/UnrealMathSSE.h
--- a/Engine/Source/Runtime/Core/Public/Math/UnrealMathSSE.h
+++ b/Engine/Source/Runtime/Core/Public/Math/UnrealMathSSE.h
@@ -803,15 +803,6 @@ FORCEINLINE VectorRegister VectorTransformVector(const VectorRegister&  VecP,  c
  */
 #define VectorMaskBits( VecMask )			_mm_movemask_ps( VecMask )
 
-/**
- * Returns the bitwise AND.
- *
- * @param	Vec1	Vector to AND
- * @param	Vec2	Vector to AND
- * @return	bitwise per component AND operation.
- */
-#define VectorBitwiseAND( Vec1, Vec2 )	_mm_and_ps( (Vec1), (Vec2) )
-
 /**
  * Divides two vectors (component-wise) and returns the result.
  *
diff --git a/Engine/Source/Runtime/Engine/Public/kDOP.h b/Engine/Source/Runtime/Engine/Public/kDOP.h
--- a/Engine/Source/Runtime/Engine/Public/kDOP.h
+++ b/Engine/Source/Runtime/Engine/Public/kDOP.h
@@ -210,8 +210,8 @@ int32 appLineCheckTriangleSOA(const FVector3SOA& Start, const FVector3SOA& End,
 
 	// If this triangle is not closer than the previous hit, reject it
 	VectorRegister IntersectionTime = VectorLoadFloat1( &InOutIntersectionTime );
-	TriangleMask = VectorBitwiseAND( TriangleMask, VectorMask_GE( Time, VectorZero() ) );
-	TriangleMask = VectorBitwiseAND( TriangleMask, VectorMask_LT( Time, IntersectionTime ) );
+	TriangleMask = VectorBitwiseAnd( TriangleMask, VectorMask_GE( Time, VectorZero() ) );
+	TriangleMask = VectorBitwiseAnd( TriangleMask, VectorMask_LT( Time, IntersectionTime ) );
 	if ( VectorMaskBits(TriangleMask) == 0 )
 	{
 		return -1;
@@ -239,7 +239,7 @@ int32 appLineCheckTriangleSOA(const FVector3SOA& Start, const FVector3SOA& End,
 		DotW = VectorMultiply( SideDirectionX, IntersectionX );
 		DotW = VectorMultiplyAdd( SideDirectionY, IntersectionY, DotW );
 		DotW = VectorMultiplyAdd( SideDirectionZ, IntersectionZ, DotW );
-		TriangleMask = VectorBitwiseAND( TriangleMask, VectorMask_LT( VectorSubtract(DotW, SideW), GSmallNumber ) );
+		TriangleMask = VectorBitwiseAnd( TriangleMask, VectorMask_LT( VectorSubtract(DotW, SideW), GSmallNumber ) );
 		if ( VectorMaskBits(TriangleMask) == 0 )
 		{
 			return -1;
@@ -685,8 +685,8 @@ struct TkDOPNode
 
 		// 4: Calculate hit time and determine whether there was a hit.		
 		VectorStoreAligned( MinTime, &HitTime );
-		const VectorRegister OutNodeHit		= VectorBitwiseAND( VectorCompareGE( MaxTime, VectorZero() ), VectorCompareGE( MaxTime, MinTime ) );
-		const VectorRegister CloserNodeHit	= VectorBitwiseAND( OutNodeHit, VectorCompareGT( CurrentHitTime, MinTime ) );
+		const VectorRegister OutNodeHit		= VectorBitwiseAnd( VectorCompareGE( MaxTime, VectorZero() ), VectorCompareGE( MaxTime, MinTime ) );
+		const VectorRegister CloserNodeHit	= VectorBitwiseAnd( OutNodeHit, VectorCompareGT( CurrentHitTime, MinTime ) );
 		VectorStoreAligned( CloserNodeHit, (float*) NodeHit );
 #endif
 	}
