#ue4 - Simulated CharacterMovement now uses SetPostLandedPhysics() when landing, rather than hardcoded walking mode.

- Fixed SetPostLandedPhysics() not respecting default land movement mode that is not a walking mode.

diff --git a/Engine/Source/Runtime/Engine/Private/Components/CharacterMovementComponent.cpp b/Engine/Source/Runtime/Engine/Private/Components/CharacterMovementComponent.cpp
--- a/Engine/Source/Runtime/Engine/Private/Components/CharacterMovementComponent.cpp
+++ b/Engine/Source/Runtime/Engine/Private/Components/CharacterMovementComponent.cpp
@@ -604,7 +604,7 @@ float UCharacterMovementComponent::GetNetworkSafeRandomAngleDegrees() const
 void UCharacterMovementComponent::SetDefaultMovementMode()
 {
 	// check for water volume
-	if ( IsInWater() && CanEverSwim() )
+	if (CanEverSwim() && IsInWater())
 	{
 		SetMovementMode(DefaultWaterMovementMode);
 	}
@@ -1225,7 +1225,7 @@ void UCharacterMovementComponent::SimulateMovement(float DeltaSeconds)
 					if (CurrentFloor.FloorDist <= MIN_FLOOR_DIST)
 					{
 						// Landed
-						SetMovementMode(MOVE_Walking);
+						SetPostLandedPhysics(CurrentFloor.HitResult);
 					}
 					else
 					{
@@ -4160,7 +4160,7 @@ void UCharacterMovementComponent::SetPostLandedPhysics(const FHitResult& Hit)
 {
 	if( CharacterOwner )
 	{
-		if ( GetPhysicsVolume()->bWaterVolume && CanEverSwim() )
+		if (CanEverSwim() && IsInWater())
 		{
 			SetMovementMode(MOVE_Swimming);
 		}
@@ -4168,7 +4168,18 @@ void UCharacterMovementComponent::SetPostLandedPhysics(const FHitResult& Hit)
 		{
 			const FVector PreImpactAccel = Acceleration + (IsFalling() ? FVector(0.f, 0.f, GetGravityZ()) : FVector::ZeroVector);
 			const FVector PreImpactVelocity = Velocity;
-			SetMovementMode(GroundMovementMode);
+
+			if (DefaultLandMovementMode == MOVE_Walking ||
+				DefaultLandMovementMode == MOVE_NavWalking ||
+				DefaultLandMovementMode == MOVE_Falling)
+			{
+				SetMovementMode(GroundMovementMode);
+			}
+			else
+			{
+				SetDefaultMovementMode();
+			}
+			
 			ApplyImpactPhysicsForces(Hit, PreImpactAccel, PreImpactVelocity);
 		}
 	}
