GenericCreateObject node can be called in Construction Script
UGameplayStatics::SpawnObject checks ClassWithin

diff --git a/Engine/Source/Editor/BlueprintGraph/Classes/K2Node_GenericCreateObject.h b/Engine/Source/Editor/BlueprintGraph/Classes/K2Node_GenericCreateObject.h
--- a/Engine/Source/Editor/BlueprintGraph/Classes/K2Node_GenericCreateObject.h
+++ b/Engine/Source/Editor/BlueprintGraph/Classes/K2Node_GenericCreateObject.h
@@ -15,6 +15,7 @@ class BLUEPRINTGRAPH_API UK2Node_GenericCreateObject : public UK2Node_ConstructO
 	virtual void ExpandNode(class FKismetCompilerContext& CompilerContext, UEdGraph* SourceGraph) override;
 	virtual void ValidateNodeDuringCompilation(class FCompilerResultsLog& MessageLog) const override;
 	virtual void EarlyValidation(class FCompilerResultsLog& MessageLog) const override;
+	virtual bool IsCompatibleWithGraph(const UEdGraph* TargetGraph) const override;
 	// End UEdGraphNode interface.
 
 	virtual bool UseWorldContext() const override { return false; }
diff --git a/Engine/Source/Editor/BlueprintGraph/Private/K2Node_GenericCreateObject.cpp b/Engine/Source/Editor/BlueprintGraph/Private/K2Node_GenericCreateObject.cpp
--- a/Engine/Source/Editor/BlueprintGraph/Private/K2Node_GenericCreateObject.cpp
+++ b/Engine/Source/Editor/BlueprintGraph/Private/K2Node_GenericCreateObject.cpp
@@ -86,4 +86,9 @@ void UK2Node_GenericCreateObject::EarlyValidation(class FCompilerResultsLog& Mes
 	}
 }
 
+bool UK2Node_GenericCreateObject::IsCompatibleWithGraph(const UEdGraph* TargetGraph) const
+{
+	return UK2Node::IsCompatibleWithGraph(TargetGraph);
+}
+
 #undef LOCTEXT_NAMESPACE
\ No newline at end of file
diff --git a/Engine/Source/Runtime/Engine/Classes/Kismet/GameplayStatics.h b/Engine/Source/Runtime/Engine/Classes/Kismet/GameplayStatics.h
--- a/Engine/Source/Runtime/Engine/Classes/Kismet/GameplayStatics.h
+++ b/Engine/Source/Runtime/Engine/Classes/Kismet/GameplayStatics.h
@@ -31,7 +31,7 @@ class ENGINE_API UGameplayStatics : public UBlueprintFunctionLibrary
 	GENERATED_UCLASS_BODY()
 
 	// --- Create Object
-	UFUNCTION(BlueprintCallable, Category = "Spawning", meta = (UnsafeDuringActorConstruction = "true", BlueprintInternalUseOnly = "true"))
+	UFUNCTION(BlueprintCallable, Category = "Spawning", meta = (BlueprintInternalUseOnly = "true"))
 	static UObject* SpawnObject(TSubclassOf<UObject> ObjectClass, UObject* Outer);
 
 	static bool CanSpawnObjectOfClass(TSubclassOf<UObject> ObjectClass);
diff --git a/Engine/Source/Runtime/Engine/Private/GameplayStatics.cpp b/Engine/Source/Runtime/Engine/Private/GameplayStatics.cpp
--- a/Engine/Source/Runtime/Engine/Private/GameplayStatics.cpp
+++ b/Engine/Source/Runtime/Engine/Private/GameplayStatics.cpp
@@ -341,7 +341,7 @@ UObject* UGameplayStatics::SpawnObject(TSubclassOf<UObject> ObjectClass, UObject
 #if WITH_EDITOR
 		FMessageLog("PIE").Error(FText::Format(LOCTEXT("SpawnObjectWrongClass", "SpawnObject wrong class: {0}'"), FText::FromString(GetNameSafe(*ObjectClass))));
 #endif // WITH_EDITOR
-		UE_LOG(LogScript, Error, TEXT("UGameplayStatics::SpawnObject wrong class: %s"), *GetNameSafe(*ObjectClass));
+		UE_LOG(LogScript, Error, TEXT("UGameplayStatics::SpawnObject wrong class: %s"), *GetPathNameSafe(*ObjectClass));
 		return nullptr;
 	}
 
@@ -351,6 +351,12 @@ UObject* UGameplayStatics::SpawnObject(TSubclassOf<UObject> ObjectClass, UObject
 		return nullptr;
 	}
 
+	if (ObjectClass->ClassWithin && !Outer->IsA(ObjectClass->ClassWithin))
+	{
+		UE_LOG(LogScript, Warning, TEXT("UGameplayStatics::SpawnObject outer %s is not %s"), *GetPathNameSafe(Outer), *GetPathNameSafe(ObjectClass->ClassWithin));
+		return nullptr;
+	}
+
 	return NewObject<UObject>(Outer, ObjectClass, NAME_None, RF_StrongRefOnFrame);
 }
 
