Allow const Actor* to be used for TInlineComponentArray constructor

diff --git a/Engine/Source/Runtime/Engine/Classes/GameFramework/Actor.h b/Engine/Source/Runtime/Engine/Classes/GameFramework/Actor.h
--- a/Engine/Source/Runtime/Engine/Classes/GameFramework/Actor.h
+++ b/Engine/Source/Runtime/Engine/Classes/GameFramework/Actor.h
@@ -2477,7 +2477,7 @@ class TInlineComponentArray : public TArray<T, TInlineAllocator<NumElements>>
 
 public:
 	TInlineComponentArray() : Super() { }
-	TInlineComponentArray(class AActor* Actor) : Super()
+	TInlineComponentArray(const class AActor* Actor) : Super()
 	{
 		if (Actor)
 		{
