Remove GetComponentInstanceDataType and filter by Class instead
Store component instance data in an array instead of a map

diff --git a/Engine/Source/Runtime/Engine/Classes/Components/ActorComponent.h b/Engine/Source/Runtime/Engine/Classes/Components/ActorComponent.h
--- a/Engine/Source/Runtime/Engine/Classes/Components/ActorComponent.h
+++ b/Engine/Source/Runtime/Engine/Classes/Components/ActorComponent.h
@@ -566,9 +566,6 @@ public:
 	/** Called before we throw away components during RerunConstructionScripts, to cache any data we wish to persist across that operation */
 	virtual class FActorComponentInstanceData* GetComponentInstanceData() const;
 
-	/** The type of the component instance data that this component is interested in */
-	virtual FName GetComponentInstanceDataType() const;
-
 	// Begin UObject interface.
 	virtual void BeginDestroy() override;
 	virtual bool NeedsLoadForClient() const override;
diff --git a/Engine/Source/Runtime/Engine/Classes/Components/ChildActorComponent.h b/Engine/Source/Runtime/Engine/Classes/Components/ChildActorComponent.h
--- a/Engine/Source/Runtime/Engine/Classes/Components/ChildActorComponent.h
+++ b/Engine/Source/Runtime/Engine/Classes/Components/ChildActorComponent.h
@@ -48,7 +48,6 @@ public:
 	virtual void OnComponentDestroyed() override;
 	virtual void OnRegister() override;
 	virtual FActorComponentInstanceData* GetComponentInstanceData() const override;
-	virtual FName GetComponentInstanceDataType() const override;
 	// End ActorComponent interface.
 
 	/** Apply the component instance data to the child actor component */
diff --git a/Engine/Source/Runtime/Engine/Classes/Components/InstancedStaticMeshComponent.h b/Engine/Source/Runtime/Engine/Classes/Components/InstancedStaticMeshComponent.h
--- a/Engine/Source/Runtime/Engine/Classes/Components/InstancedStaticMeshComponent.h
+++ b/Engine/Source/Runtime/Engine/Classes/Components/InstancedStaticMeshComponent.h
@@ -157,7 +157,6 @@ public:
 
 	// Begin UActorComponent interface 
 	virtual FActorComponentInstanceData* GetComponentInstanceData() const override;
-	virtual FName GetComponentInstanceDataType() const override;
 	// End UActorComponent interface 
 
 	// Begin UPrimitiveComponent Interface
diff --git a/Engine/Source/Runtime/Engine/Classes/Components/LightComponent.h b/Engine/Source/Runtime/Engine/Classes/Components/LightComponent.h
--- a/Engine/Source/Runtime/Engine/Classes/Components/LightComponent.h
+++ b/Engine/Source/Runtime/Engine/Classes/Components/LightComponent.h
@@ -308,7 +308,6 @@ public:
 	// End UObject interface.
 
 	virtual FActorComponentInstanceData* GetComponentInstanceData() const override;
-	virtual FName GetComponentInstanceDataType() const override;
 	 void ApplyComponentInstanceData(class FPrecomputedLightInstanceData* ComponentInstanceData);
 
 	/** @return number of material elements in this primitive */
diff --git a/Engine/Source/Runtime/Engine/Classes/Components/PrimitiveComponent.h b/Engine/Source/Runtime/Engine/Classes/Components/PrimitiveComponent.h
--- a/Engine/Source/Runtime/Engine/Classes/Components/PrimitiveComponent.h
+++ b/Engine/Source/Runtime/Engine/Classes/Components/PrimitiveComponent.h
@@ -1033,7 +1033,6 @@ public:
 	virtual bool ShouldCreatePhysicsState() const override;
 	virtual bool HasValidPhysicsState() const override;
 	virtual class FActorComponentInstanceData* GetComponentInstanceData() const override;
-	virtual FName GetComponentInstanceDataType() const override;
 	// End UActorComponent Interface
 
 	/** @return true if the owner is selected and this component is selectable */
diff --git a/Engine/Source/Runtime/Engine/Classes/Components/SceneComponent.h b/Engine/Source/Runtime/Engine/Classes/Components/SceneComponent.h
--- a/Engine/Source/Runtime/Engine/Classes/Components/SceneComponent.h
+++ b/Engine/Source/Runtime/Engine/Classes/Components/SceneComponent.h
@@ -533,7 +533,6 @@ public:
 	virtual void OnComponentDestroyed() override;
 	virtual void ApplyWorldOffset(const FVector& InOffset, bool bWorldShift) override;
 	virtual class FActorComponentInstanceData* GetComponentInstanceData() const override;
-	virtual FName GetComponentInstanceDataType() const override;
 	// End ActorComponent interface
 
 	// Call UpdateComponentToWorld if bWorldToComponentUpdated is false.
diff --git a/Engine/Source/Runtime/Engine/Classes/Components/SkyLightComponent.h b/Engine/Source/Runtime/Engine/Classes/Components/SkyLightComponent.h
--- a/Engine/Source/Runtime/Engine/Classes/Components/SkyLightComponent.h
+++ b/Engine/Source/Runtime/Engine/Classes/Components/SkyLightComponent.h
@@ -151,7 +151,6 @@ class ENGINE_API USkyLightComponent : public ULightComponentBase
 	// End UObject Interface
 
 	virtual FActorComponentInstanceData* GetComponentInstanceData() const override;
-	virtual FName GetComponentInstanceDataType() const override;
 	void ApplyComponentInstanceData(class FPrecomputedSkyLightInstanceData* ComponentInstanceData);
 
 	/** Called each tick to recapture and queued sky captures. */
diff --git a/Engine/Source/Runtime/Engine/Classes/Components/SplineComponent.h b/Engine/Source/Runtime/Engine/Classes/Components/SplineComponent.h
--- a/Engine/Source/Runtime/Engine/Classes/Components/SplineComponent.h
+++ b/Engine/Source/Runtime/Engine/Classes/Components/SplineComponent.h
@@ -109,7 +109,6 @@ public:
 
 	// Begin UActorComponent interface.
 	virtual FActorComponentInstanceData* GetComponentInstanceData() const override;
-	virtual FName GetComponentInstanceDataType() const override;
 	// End UActorComponent interface.
 
 	void ApplyComponentInstanceData(class FSplineInstanceData* ComponentInstanceData, const bool bPostUCS);
diff --git a/Engine/Source/Runtime/Engine/Classes/Components/SplineMeshComponent.h b/Engine/Source/Runtime/Engine/Classes/Components/SplineMeshComponent.h
--- a/Engine/Source/Runtime/Engine/Classes/Components/SplineMeshComponent.h
+++ b/Engine/Source/Runtime/Engine/Classes/Components/SplineMeshComponent.h
@@ -141,7 +141,6 @@ class ENGINE_API USplineMeshComponent : public UStaticMeshComponent, public IInt
 
 	// Begin UActorComponent interface.
 	virtual FActorComponentInstanceData* GetComponentInstanceData() const override;
-	virtual FName GetComponentInstanceDataType() const override;
 	// End UActorComponent interface.
 
 	void ApplyComponentInstanceData(class FSplineMeshInstanceData* ComponentInstanceData);
diff --git a/Engine/Source/Runtime/Engine/Classes/Components/StaticMeshComponent.h b/Engine/Source/Runtime/Engine/Classes/Components/StaticMeshComponent.h
--- a/Engine/Source/Runtime/Engine/Classes/Components/StaticMeshComponent.h
+++ b/Engine/Source/Runtime/Engine/Classes/Components/StaticMeshComponent.h
@@ -270,7 +270,6 @@ public:
 	virtual void CheckForErrors() override;
 #endif
 	virtual FActorComponentInstanceData* GetComponentInstanceData() const override;
-	virtual FName GetComponentInstanceDataType() const override;
 	// End UActorComponent interface.
 
 	// Begin UPrimitiveComponent interface.
diff --git a/Engine/Source/Runtime/Engine/Private/ComponentInstanceDataCache.cpp b/Engine/Source/Runtime/Engine/Private/ComponentInstanceDataCache.cpp
--- a/Engine/Source/Runtime/Engine/Private/ComponentInstanceDataCache.cpp
+++ b/Engine/Source/Runtime/Engine/Private/ComponentInstanceDataCache.cpp
@@ -140,8 +140,9 @@ FComponentInstanceDataCache::FComponentInstanceDataCache(const AActor* Actor)
 {
 	if(Actor != NULL)
 	{
-		TInlineComponentArray<UActorComponent*> Components;
-		Actor->GetComponents(Components);
+		TInlineComponentArray<UActorComponent*> Components(Actor);
+
+		ComponentsInstanceData.Reserve(Components.Num());
 
 		// Grab per-instance data we want to persist
 		for (UActorComponent* Component : Components)
@@ -151,8 +152,7 @@ FComponentInstanceDataCache::FComponentInstanceDataCache(const AActor* Actor)
 				FActorComponentInstanceData* ComponentInstanceData = Component->GetComponentInstanceData();
 				if (ComponentInstanceData)
 				{
-					check(!Component->GetComponentInstanceDataType().IsNone());
-					TypeToDataMap.Add(Component->GetComponentInstanceDataType(), ComponentInstanceData);
+					ComponentsInstanceData.Add(ComponentInstanceData);
 				}
 			}
 			else if (Component->CreationMethod == EComponentCreationMethod::Instance)
@@ -172,9 +172,9 @@ FComponentInstanceDataCache::FComponentInstanceDataCache(const AActor* Actor)
 
 FComponentInstanceDataCache::~FComponentInstanceDataCache()
 {
-	for (auto InstanceDataPair : TypeToDataMap)
+	for (FActorComponentInstanceData* ComponentData : ComponentsInstanceData)
 	{
-		delete InstanceDataPair.Value;
+		delete ComponentData;
 	}
 }
 
@@ -190,20 +190,14 @@ void FComponentInstanceDataCache::ApplyToActor(AActor* Actor, const ECacheApplyP
 		{
 			if(Component->IsCreatedByConstructionScript()) // Only try and apply data to 'created by construction script' components
 			{
-				const FName ComponentInstanceDataType = Component->GetComponentInstanceDataType();
-
-				if (!ComponentInstanceDataType.IsNone())
+				for (FActorComponentInstanceData* ComponentInstanceData : ComponentsInstanceData)
 				{
-					TArray< FActorComponentInstanceData* > CachedData;
-					TypeToDataMap.MultiFind(ComponentInstanceDataType, CachedData);
-
-					for (FActorComponentInstanceData* ComponentInstanceData : CachedData)
+					if (	ComponentInstanceData 
+						&&	ComponentInstanceData->GetComponentClass() == Component->GetClass() // filter on class early to avoid unnecessary virtual and expensive tests
+						&&	ComponentInstanceData->MatchesComponent(Component))
 					{
-						if (ComponentInstanceData && ComponentInstanceData->MatchesComponent(Component))
-						{
-							ComponentInstanceData->ApplyToComponent(Component, CacheApplyPhase);
-							break;
-						}
+						ComponentInstanceData->ApplyToComponent(Component, CacheApplyPhase);
+						break;
 					}
 				}
 			}
@@ -226,11 +220,11 @@ void FComponentInstanceDataCache::ApplyToActor(AActor* Actor, const ECacheApplyP
 
 void FComponentInstanceDataCache::FindAndReplaceInstances(const TMap<UObject*, UObject*>& OldToNewInstanceMap)
 {
-	for (auto ComponentInstanceDataPair : TypeToDataMap)
+	for (FActorComponentInstanceData* ComponentInstanceData : ComponentsInstanceData)
 	{
-		if (ComponentInstanceDataPair.Value)
+		if (ComponentInstanceData)
 		{
-			ComponentInstanceDataPair.Value->FindAndReplaceInstances(OldToNewInstanceMap);
+			ComponentInstanceData->FindAndReplaceInstances(OldToNewInstanceMap);
 		}
 	}
 	TArray<USceneComponent*> SceneComponents;
@@ -259,11 +253,11 @@ void FComponentInstanceDataCache::AddReferencedObjects(FReferenceCollector& Coll
 
 	Collector.AddReferencedObjects(SceneComponents);
 
-	for (auto ComponentInstanceDataPair : TypeToDataMap)
+	for (FActorComponentInstanceData* ComponentInstanceData : ComponentsInstanceData)
 	{
-		if (ComponentInstanceDataPair.Value)
+		if (ComponentInstanceData)
 		{
-			ComponentInstanceDataPair.Value->AddReferencedObjects(Collector);
+			ComponentInstanceData->AddReferencedObjects(Collector);
 		}
 	}
 }
\ No newline at end of file
diff --git a/Engine/Source/Runtime/Engine/Private/Components/ActorComponent.cpp b/Engine/Source/Runtime/Engine/Private/Components/ActorComponent.cpp
--- a/Engine/Source/Runtime/Engine/Private/Components/ActorComponent.cpp
+++ b/Engine/Source/Runtime/Engine/Private/Components/ActorComponent.cpp
@@ -647,12 +647,6 @@ FActorComponentInstanceData* UActorComponent::GetComponentInstanceData() const
 	return InstanceData;
 }
 
-FName UActorComponent::GetComponentInstanceDataType() const
-{
-	static const FName ActorComponentInstanceDataTypeName(TEXT("ActorComponentInstanceData"));
-	return ActorComponentInstanceDataTypeName;
-}
-
 void FActorComponentTickFunction::ExecuteTick(float DeltaTime, enum ELevelTick TickType, ENamedThreads::Type CurrentThread, const FGraphEventRef& MyCompletionGraphEvent)
 {
 	if (Target && !Target->HasAnyFlags(RF_PendingKill | RF_Unreachable))
diff --git a/Engine/Source/Runtime/Engine/Private/Components/ChildActorComponent.cpp b/Engine/Source/Runtime/Engine/Private/Components/ChildActorComponent.cpp
--- a/Engine/Source/Runtime/Engine/Private/Components/ChildActorComponent.cpp
+++ b/Engine/Source/Runtime/Engine/Private/Components/ChildActorComponent.cpp
@@ -142,12 +142,6 @@ public:
 	TArray<FAttachedActorInfo> AttachedActors;
 };
 
-FName UChildActorComponent::GetComponentInstanceDataType() const
-{
-	static const FName ChildActorComponentInstanceDataName(TEXT("ChildActorInstanceData"));
-	return ChildActorComponentInstanceDataName;
-}
-
 FActorComponentInstanceData* UChildActorComponent::GetComponentInstanceData() const
 {
 	FChildActorComponentInstanceData* InstanceData = CachedInstanceData;
diff --git a/Engine/Source/Runtime/Engine/Private/Components/LightComponent.cpp b/Engine/Source/Runtime/Engine/Private/Components/LightComponent.cpp
--- a/Engine/Source/Runtime/Engine/Private/Components/LightComponent.cpp
+++ b/Engine/Source/Runtime/Engine/Private/Components/LightComponent.cpp
@@ -954,12 +954,6 @@ public:
 	bool bPrecomputedLightingIsValid;
 };
 
-FName ULightComponent::GetComponentInstanceDataType() const
-{
-	static const FName PrecomputedLightInstanceDataTypeName(TEXT("PrecomputedLightInstanceData"));
-	return PrecomputedLightInstanceDataTypeName;
-}
-
 FActorComponentInstanceData* ULightComponent::GetComponentInstanceData() const
 {
 	// Allocate new struct for holding light map data
diff --git a/Engine/Source/Runtime/Engine/Private/Components/PrimitiveComponent.cpp b/Engine/Source/Runtime/Engine/Private/Components/PrimitiveComponent.cpp
--- a/Engine/Source/Runtime/Engine/Private/Components/PrimitiveComponent.cpp
+++ b/Engine/Source/Runtime/Engine/Private/Components/PrimitiveComponent.cpp
@@ -39,7 +39,6 @@ namespace PrimitiveComponentStatics
 	static const FText MobilityWarnText = LOCTEXT("InvalidMove", "move");
 	static const FName MoveComponentName(TEXT("MoveComponent"));
 	static const FName UpdateOverlapsName(TEXT("UpdateOverlaps"));
-	static const FName PrimitiveComponentInstanceDataTypeName(TEXT("PrimitiveComponentInstanceData"));
 }
 
 DEFINE_LOG_CATEGORY_STATIC(LogPrimitiveComponent, Log, All);
@@ -425,11 +424,6 @@ FActorComponentInstanceData* UPrimitiveComponent::GetComponentInstanceData() con
 	return InstanceData;
 }
 
-FName UPrimitiveComponent::GetComponentInstanceDataType() const
-{
-	return PrimitiveComponentStatics::PrimitiveComponentInstanceDataTypeName;
-}
-
 void UPrimitiveComponent::OnAttachmentChanged()
 {
 	if (World && World->Scene)
diff --git a/Engine/Source/Runtime/Engine/Private/Components/SceneComponent.cpp b/Engine/Source/Runtime/Engine/Private/Components/SceneComponent.cpp
--- a/Engine/Source/Runtime/Engine/Private/Components/SceneComponent.cpp
+++ b/Engine/Source/Runtime/Engine/Private/Components/SceneComponent.cpp
@@ -21,7 +21,6 @@
 namespace SceneComponentStatics
 {
 	static const FName DefaultSceneRootVariableName(TEXT("DefaultSceneRoot"));
-	static const FName SceneComponentInstanceDataTypeName(TEXT("SceneComponentInstanceData"));
 	static const FName MobilityName(TEXT("Mobility"));
 	static const FText MobilityWarnText = LOCTEXT("InvalidMove", "move");
 	static const FName PhysicsVolumeTraceName(TEXT("PhysicsVolumeTrace"));
@@ -1483,11 +1482,6 @@ FActorComponentInstanceData* USceneComponent::GetComponentInstanceData() const
 	return InstanceData;
 }
 
-FName USceneComponent::GetComponentInstanceDataType() const
-{
-	return SceneComponentStatics::SceneComponentInstanceDataTypeName;
-}
-
 void USceneComponent::UpdateChildTransforms(bool bSkipPhysicsMove, bool bTeleport)
 {
 	for(int32 i=0; i<AttachChildren.Num(); i++)
diff --git a/Engine/Source/Runtime/Engine/Private/Components/SkyLightComponent.cpp b/Engine/Source/Runtime/Engine/Private/Components/SkyLightComponent.cpp
--- a/Engine/Source/Runtime/Engine/Private/Components/SkyLightComponent.cpp
+++ b/Engine/Source/Runtime/Engine/Private/Components/SkyLightComponent.cpp
@@ -376,13 +376,6 @@ public:
 	FSHVectorRGB3 IrradianceEnvironmentMap;
 };
 
-// Init type name static
-FName USkyLightComponent::GetComponentInstanceDataType() const
-{
-	static const FName PrecomputedSkyLightInstanceDataTypeName(TEXT("PrecomputedSkyLightInstanceData"));
-	return PrecomputedSkyLightInstanceDataTypeName;
-}
-
 FActorComponentInstanceData* USkyLightComponent::GetComponentInstanceData() const
 {
 	FPrecomputedSkyLightInstanceData* InstanceData = new FPrecomputedSkyLightInstanceData(this);
diff --git a/Engine/Source/Runtime/Engine/Private/Components/SplineComponent.cpp b/Engine/Source/Runtime/Engine/Private/Components/SplineComponent.cpp
--- a/Engine/Source/Runtime/Engine/Private/Components/SplineComponent.cpp
+++ b/Engine/Source/Runtime/Engine/Private/Components/SplineComponent.cpp
@@ -989,13 +989,6 @@ public:
 };
 
 
-FName USplineComponent::GetComponentInstanceDataType() const
-{
-	static const FName SplineInstanceDataTypeName(TEXT("SplineInstanceData"));
-	return SplineInstanceDataTypeName;
-}
-
-
 FActorComponentInstanceData* USplineComponent::GetComponentInstanceData() const
 {
 	FSplineInstanceData* SplineInstanceData = new FSplineInstanceData(this);
diff --git a/Engine/Source/Runtime/Engine/Private/Components/SplineMeshComponent.cpp b/Engine/Source/Runtime/Engine/Private/Components/SplineMeshComponent.cpp
--- a/Engine/Source/Runtime/Engine/Private/Components/SplineMeshComponent.cpp
+++ b/Engine/Source/Runtime/Engine/Private/Components/SplineMeshComponent.cpp
@@ -959,12 +959,6 @@ public:
 	FVector EndTangent;
 };
 
-FName USplineMeshComponent::GetComponentInstanceDataType() const
-{
-	static const FName SplineMeshInstanceDataTypeName(TEXT("SplineMeshInstanceData"));
-	return SplineMeshInstanceDataTypeName;
-}
-
 FActorComponentInstanceData* USplineMeshComponent::GetComponentInstanceData() const
 {
 	FActorComponentInstanceData* InstanceData = nullptr;
diff --git a/Engine/Source/Runtime/Engine/Private/Components/StaticMeshComponent.cpp b/Engine/Source/Runtime/Engine/Private/Components/StaticMeshComponent.cpp
--- a/Engine/Source/Runtime/Engine/Private/Components/StaticMeshComponent.cpp
+++ b/Engine/Source/Runtime/Engine/Private/Components/StaticMeshComponent.cpp
@@ -1540,12 +1540,6 @@ int32 UStaticMeshComponent::GetBlueprintCreatedComponentIndex() const
 	return INDEX_NONE;
 }
 
-FName UStaticMeshComponent::GetComponentInstanceDataType() const
-{
-	static const FName StaticMeshComponentInstanceDataName(TEXT("StaticMeshInstanceData"));
-	return StaticMeshComponentInstanceDataName;
-}
-
 FActorComponentInstanceData* UStaticMeshComponent::GetComponentInstanceData() const
 {
 	FStaticMeshComponentInstanceData* StaticMeshInstanceData = nullptr;
diff --git a/Engine/Source/Runtime/Engine/Private/InstancedStaticMesh.cpp b/Engine/Source/Runtime/Engine/Private/InstancedStaticMesh.cpp
--- a/Engine/Source/Runtime/Engine/Private/InstancedStaticMesh.cpp
+++ b/Engine/Source/Runtime/Engine/Private/InstancedStaticMesh.cpp
@@ -805,12 +805,6 @@ public:
 };
 #endif
 
-FName UInstancedStaticMeshComponent::GetComponentInstanceDataType() const
-{
-	static const FName InstancedStaticMeshComponentInstanceDataName(TEXT("InstancedStaticMeshComponentInstanceData"));
-	return InstancedStaticMeshComponentInstanceDataName;
-}
-
 FActorComponentInstanceData* UInstancedStaticMeshComponent::GetComponentInstanceData() const
 {
 #if WITH_EDITOR
diff --git a/Engine/Source/Runtime/Engine/Public/ComponentInstanceDataCache.h b/Engine/Source/Runtime/Engine/Public/ComponentInstanceDataCache.h
--- a/Engine/Source/Runtime/Engine/Public/ComponentInstanceDataCache.h
+++ b/Engine/Source/Runtime/Engine/Public/ComponentInstanceDataCache.h
@@ -39,6 +39,8 @@ public:
 
 	bool ContainsSavedProperties() const { return SavedProperties.Num() > 0; }
 
+	UClass* GetComponentClass() const { return SourceComponentClass; }
+
 protected:
 	/** The name of the source component */
 	FName SourceComponentName;
@@ -74,13 +76,13 @@ public:
 	/** Iterates over components and replaces any object references with the reinstanced information */
 	void FindAndReplaceInstances(const TMap<UObject*, UObject*>& OldToNewInstanceMap);
 
-	bool HasInstanceData() const { return TypeToDataMap.Num() > 0; }
+	bool HasInstanceData() const { return ComponentsInstanceData.Num() > 0; }
 
 	void AddReferencedObjects(FReferenceCollector& Collector);
 
 private:
 	/** Map of data type name to data of that type */
-	TMultiMap< FName, FActorComponentInstanceData* >	TypeToDataMap;
+	TArray< FActorComponentInstanceData* > ComponentsInstanceData;
 
 	TMap< USceneComponent*, FTransform > InstanceComponentTransformToRootMap;
 };
\ No newline at end of file
diff --git a/Engine/Source/Runtime/UMG/Private/Components/WidgetComponent.cpp b/Engine/Source/Runtime/UMG/Private/Components/WidgetComponent.cpp
--- a/Engine/Source/Runtime/UMG/Private/Components/WidgetComponent.cpp
+++ b/Engine/Source/Runtime/UMG/Private/Components/WidgetComponent.cpp
@@ -658,12 +658,6 @@ public:
 	UTextureRenderTarget2D* RenderTarget;
 };
 
-FName UWidgetComponent::GetComponentInstanceDataType() const
-{
-	static const FName InstanceDataName(TEXT("WidgetInstanceData"));
-	return InstanceDataName;
-}
-
 FActorComponentInstanceData* UWidgetComponent::GetComponentInstanceData() const
 {
 	return new FWidgetComponentInstanceData( this );
diff --git a/Engine/Source/Runtime/UMG/Public/Components/WidgetComponent.h b/Engine/Source/Runtime/UMG/Public/Components/WidgetComponent.h
--- a/Engine/Source/Runtime/UMG/Public/Components/WidgetComponent.h
+++ b/Engine/Source/Runtime/UMG/Public/Components/WidgetComponent.h
@@ -44,7 +44,6 @@ public:
 	virtual void TickComponent(float DeltaTime, enum ELevelTick TickType, FActorComponentTickFunction *ThisTickFunction) override;
 
 	virtual FActorComponentInstanceData* GetComponentInstanceData() const override;
-	virtual FName GetComponentInstanceDataType() const override;
 
 	void ApplyComponentInstanceData(class FWidgetComponentInstanceData* ComponentInstanceData);
 
