Fix shadow variable warning

diff --git a/Engine/Source/Runtime/Engine/Private/Components/DestructibleComponent.cpp b/Engine/Source/Runtime/Engine/Private/Components/DestructibleComponent.cpp
--- a/Engine/Source/Runtime/Engine/Private/Components/DestructibleComponent.cpp
+++ b/Engine/Source/Runtime/Engine/Private/Components/DestructibleComponent.cpp
@@ -1176,7 +1176,6 @@ bool UDestructibleComponent::DoCustomNavigableGeometryExport(FNavigableGeometryE
 			{
 				// only root chunk visible = not fractured
 				const PxRigidDynamic* PActor = ApexDestructibleActor->getChunkPhysXActor(0);
-				TArray<PxShape*> Shapes;
 				
 				const int32 NumExported = ExportPxActorGeometry(PActor, Shapes, GeomExport);
 				bExportFromBodySetup = (NumExported == 0);
