Fixed RPC size reporting in the network profiler. Some bits that were incorrectly reported as being in the parameter section are now correctly reported as being in the header section.

diff --git a/Engine/Source/Runtime/Engine/Private/NetworkDriver.cpp b/Engine/Source/Runtime/Engine/Private/NetworkDriver.cpp
--- a/Engine/Source/Runtime/Engine/Private/NetworkDriver.cpp
+++ b/Engine/Source/Runtime/Engine/Private/NetworkDriver.cpp
@@ -808,8 +808,6 @@ void UNetDriver::InternalProcessRemoteFunction
 		Ch->BeginContentBlock(TargetObj, Bunch);
 	}
 	
-	const int NumStartingHeaderBits = Bunch.GetNumBits();
-
 	//UE_LOG(LogScript, Log, TEXT("   Call %s"),Function->GetFullName());
 	if ( Connection->InternalAck )
 	{
@@ -822,7 +820,7 @@ void UNetDriver::InternalProcessRemoteFunction
 		Bunch.WriteIntWrapped(FieldCache->FieldNetIndex, ClassCache->GetMaxIndex()+1);
 	}
 
-	const int HeaderBits = Bunch.GetNumBits() - NumStartingHeaderBits;
+	const int HeaderBits = Bunch.GetNumBits();
 
 #if !(UE_BUILD_SHIPPING || UE_BUILD_TEST)
 	Bunch.DebugString = FString::Printf(TEXT("%.2f RPC: %s - %s"), Connection->Driver->Time, *Actor->GetName(), *Function->GetName());
