Enum comment fixes. First batch.

diff --git a/Engine/Plugins/2D/Paper2D/Source/Paper2D/Classes/PaperFlipbook.h b/Engine/Plugins/2D/Paper2D/Source/Paper2D/Classes/PaperFlipbook.h
--- a/Engine/Plugins/2D/Paper2D/Source/Paper2D/Classes/PaperFlipbook.h
+++ b/Engine/Plugins/2D/Paper2D/Source/Paper2D/Classes/PaperFlipbook.h
@@ -29,13 +29,13 @@ namespace EFlipbookCollisionMode
 {
 	enum Type
 	{
-		// The flipbook has no collision
+		/** The flipbook has no collision */
 		NoCollision,
 
-		// The flipbook has non-animated collision based on the first frame of the animation
+		/** The flipbook has non-animated collision based on the first frame of the animation */
 		FirstFrameCollision,
 
-		// The flipbook changes collision each frame based on the animation (Note: This setting is not recommended and is very expensive, recreating the physics state every frame)
+		/** The flipbook changes collision each frame based on the animation (Note: This setting is not recommended and is very expensive, recreating the physics state every frame) */
 		EachFrameCollision,
 	};
 }
diff --git a/Engine/Plugins/2D/Paper2D/Source/Paper2D/Classes/PaperSpriteAtlas.h b/Engine/Plugins/2D/Paper2D/Source/Paper2D/Classes/PaperSpriteAtlas.h
--- a/Engine/Plugins/2D/Paper2D/Source/Paper2D/Classes/PaperSpriteAtlas.h
+++ b/Engine/Plugins/2D/Paper2D/Source/Paper2D/Classes/PaperSpriteAtlas.h
@@ -33,9 +33,9 @@ struct FPaperSpriteAtlasSlot
 UENUM()
 enum class EPaperSpriteAtlasPadding : uint8
 {
-	// Dilate the texture to pad the atlas.
+	/** Dilate the texture to pad the atlas. */
 	DilateBorder,
-	// Padding border filled with zeros.
+	/** Padding border filled with zeros. */
 	PadWithZero,
 	//TODO: Tiled padding
 	//TODO: Sample outside bounds from source texture (for seamless reconstruction of a larger image)
diff --git a/Engine/Plugins/2D/Paper2D/Source/Paper2D/Classes/PaperTileMap.h b/Engine/Plugins/2D/Paper2D/Source/Paper2D/Classes/PaperTileMap.h
--- a/Engine/Plugins/2D/Paper2D/Source/Paper2D/Classes/PaperTileMap.h
+++ b/Engine/Plugins/2D/Paper2D/Source/Paper2D/Classes/PaperTileMap.h
@@ -18,16 +18,16 @@ namespace ETileMapProjectionMode
 {
 	enum Type
 	{
-		// Square tile layout
+		/** Square tile layout */
 		Orthogonal,
 
-		// Isometric tile layout (shaped like a diamond)
+		// Isometric tile layout (shaped like a diamond) */
 		IsometricDiamond,
 
-		// Isometric tile layout (roughly in a square with alternating rows staggered).  Warning: Not fully supported yet.
+		/** Isometric tile layout (roughly in a square with alternating rows staggered).  Warning: Not fully supported yet. */
 		IsometricStaggered,
 
-		// Hexagonal tile layout (roughly in a square with alternating rows staggered).  Warning: Not fully supported yet.
+		/** Hexagonal tile layout (roughly in a square with alternating rows staggered).  Warning: Not fully supported yet. */
 		HexagonalStaggered
 	};
 }
diff --git a/Engine/Plugins/2D/Paper2D/Source/Paper2D/Classes/SpriteEditorOnlyTypes.h b/Engine/Plugins/2D/Paper2D/Source/Paper2D/Classes/SpriteEditorOnlyTypes.h
--- a/Engine/Plugins/2D/Paper2D/Source/Paper2D/Classes/SpriteEditorOnlyTypes.h
+++ b/Engine/Plugins/2D/Paper2D/Source/Paper2D/Classes/SpriteEditorOnlyTypes.h
@@ -10,13 +10,13 @@ namespace ESpriteCollisionMode
 {
 	enum Type
 	{
-		// Should this have no collison and not participate in physics?
+		/** Should this have no collison and not participate in physics? */
 		None,
 
-		// EXPERIMENTAL: Should this have 2D collision geometry and participate in the 2D physics world?
+		/** EXPERIMENTAL: Should this have 2D collision geometry and participate in the 2D physics world? */
 		Use2DPhysics UMETA(DisplayName = "Use 2D Physics"),
 
-		// Should this have 3D collision geometry and participate in the 3D physics world?
+		/** Should this have 3D collision geometry and participate in the 3D physics world? */
 		Use3DPhysics UMETA(DisplayName = "Use 3D Physics")
 	};
 }
@@ -25,13 +25,13 @@ namespace ESpriteCollisionMode
 UENUM()
 enum class ESpriteShapeType : uint8
 {
-	// Box/Rectangular prism (size defined by BoxSize)
+	/** Box/Rectangular prism (size defined by BoxSize) */
 	Box,
 
-	// Circle/Sphere (major axis is defined by BoxSize.X, minor axis by BoxSize.Y)
+	/** Circle/Sphere (major axis is defined by BoxSize.X, minor axis by BoxSize.Y) */
 	Circle,
 
-	// Custom closed polygon
+	/** Custom closed polygon */
 	Polygon
 };
 
@@ -143,19 +143,19 @@ namespace ESpritePolygonMode
 {
 	enum Type
 	{
-		// Use the bounding box of the source sprite (no optimization)
+		/** Use the bounding box of the source sprite (no optimization) */
 		SourceBoundingBox,
 
-		// Tighten the bounding box around the sprite to exclude fully transparent areas (the default)
+		/** Tighten the bounding box around the sprite to exclude fully transparent areas (the default) */
 		TightBoundingBox,
 
-		// Shrink-wrapped geometry
+		/** Shrink-wrapped geometry */
 		ShrinkWrapped,
 
-		// Fully custom geometry; edited by hand
+		/** Fully custom geometry; edited by hand */
 		FullyCustom,
 
-		// Diced (split up into smaller squares, including only non-empty ones in the final geometry).  This option is only supported for Render geometry and will be ignored for Collision geometry.
+		/** Diced (split up into smaller squares, including only non-empty ones in the final geometry).  This option is only supported for Render geometry and will be ignored for Collision geometry. */
 		Diced
 	};
 }
diff --git a/Engine/Plugins/2D/Paper2D/Source/Paper2DEditor/Private/ExtractSprites/PaperExtractSpritesSettings.h b/Engine/Plugins/2D/Paper2D/Source/Paper2DEditor/Private/ExtractSprites/PaperExtractSpritesSettings.h
--- a/Engine/Plugins/2D/Paper2D/Source/Paper2DEditor/Private/ExtractSprites/PaperExtractSpritesSettings.h
+++ b/Engine/Plugins/2D/Paper2D/Source/Paper2DEditor/Private/ExtractSprites/PaperExtractSpritesSettings.h
@@ -7,10 +7,10 @@
 UENUM()
 enum class ESpriteExtractMode : uint8
 {
-	// Automatically extract sprites by detecting using alpha
+	/** Automatically extract sprites by detecting using alpha */
 	Auto,
 
-	// Extract sprites in a grid defined in the properties below
+	/** Extract sprites in a grid defined in the properties below */
 	Grid,	
 };
 
diff --git a/Engine/Plugins/2D/Paper2D/Source/SpriterImporter/Private/SpriterDataModel.h b/Engine/Plugins/2D/Paper2D/Source/SpriterImporter/Private/SpriterDataModel.h
--- a/Engine/Plugins/2D/Paper2D/Source/SpriterImporter/Private/SpriterDataModel.h
+++ b/Engine/Plugins/2D/Paper2D/Source/SpriterImporter/Private/SpriterDataModel.h
@@ -13,7 +13,8 @@ UENUM()
 enum class ESpriterObjectType : uint8
 {
 	INVALID,
-	Sprite, // Default when not specified
+	/** Default when not specified */
+	Sprite,
 	Bone,
 	Box,
 	Point,
@@ -28,7 +29,8 @@ enum class ESpriterCurveType : uint8
 {
 	INVALID,
 	Instant,
-	Linear, // Default when not specified
+	/** Default when not specified */
+	Linear,
 	Quadratic,
 	Cubic
 };
diff --git a/Engine/Plugins/NetcodeUnitTest/NetcodeUnitTest/Source/NetcodeUnitTest/Classes/UnitTest.h b/Engine/Plugins/NetcodeUnitTest/NetcodeUnitTest/Source/NetcodeUnitTest/Classes/UnitTest.h
--- a/Engine/Plugins/NetcodeUnitTest/NetcodeUnitTest/Source/NetcodeUnitTest/Classes/UnitTest.h
+++ b/Engine/Plugins/NetcodeUnitTest/NetcodeUnitTest/Source/NetcodeUnitTest/Classes/UnitTest.h
@@ -27,11 +27,16 @@ struct FTCCommonVars;
 UENUM()
 enum class EUnitTestVerification : uint8
 {
-	Unverified,				// Unit test is not yet verified
-	VerifiedNotFixed,		// Unit test is verified as not fixed
-	VerifiedFixed,			// Unit test is verified as fixed
-	VerifiedNeedsUpdate,	// Unit test is no longer functioning, needs manual check/update (issue may be fixed, or unit test broken)
-	VerifiedUnreliable,		// Unit test is verified as having executed unreliably
+	/** Unit test is not yet verified */
+	Unverified,
+	/** Unit test is verified as not fixed */
+	VerifiedNotFixed,
+	/** Unit test is verified as fixed */
+	VerifiedFixed,
+	/** Unit test is no longer functioning, needs manual check/update (issue may be fixed, or unit test broken) */
+	VerifiedNeedsUpdate,
+	/** Unit test is verified as having executed unreliably */
+	VerifiedUnreliable,
 };
 
 
diff --git a/Engine/Source/Editor/HardwareTargeting/Public/HardwareTargetingSettings.h b/Engine/Source/Editor/HardwareTargeting/Public/HardwareTargetingSettings.h
--- a/Engine/Source/Editor/HardwareTargeting/Public/HardwareTargetingSettings.h
+++ b/Engine/Source/Editor/HardwareTargeting/Public/HardwareTargetingSettings.h
@@ -10,13 +10,13 @@ namespace EHardwareClass
 {
 	enum Type
 	{
-		// Unspecified, meaning no choice has been made yet
+		/** Unspecified, meaning no choice has been made yet */
 		Unspecified,
 
-		// Desktop or console
+		/** Desktop or console */
 		Desktop,
 		
-		// Mobile or tablet
+		/** Mobile or tablet */
 		Mobile
 	};
 }
@@ -27,13 +27,13 @@ namespace EGraphicsPreset
 {
 	enum Type
 	{
-		// Unspecified, meaning no choice has been made yet
+		/** Unspecified, meaning no choice has been made yet */
 		Unspecified,
 
-		// Maximum quality - High-end features default to enabled
+		/** Maximum quality - High-end features default to enabled */
 		Maximum,
 		
-		// Scalable quality - Some features are disabled by default but can be enabled based on the actual hardware
+		/** Scalable quality - Some features are disabled by default but can be enabled based on the actual hardware */
 		Scalable
 	};
 }
diff --git a/Engine/Source/Editor/Kismet/Public/BlueprintActionMenuUtils.h b/Engine/Source/Editor/Kismet/Public/BlueprintActionMenuUtils.h
--- a/Engine/Source/Editor/Kismet/Public/BlueprintActionMenuUtils.h
+++ b/Engine/Source/Editor/Kismet/Public/BlueprintActionMenuUtils.h
@@ -18,7 +18,8 @@ namespace EContextTargetFlags
 		TARGET_SiblingPinObjects	= 0x00000010 UMETA(DisplayName="Other Object Outputs", ToolTip="Include functions and variables that belong to any of this node's output types."),
 		TARGET_BlueprintLibraries	= 0x00000020 UMETA(DisplayName="Libraries", ToolTip="Include static functions that are globally accessible (belonging to function/macro libraries, etc.)."),
 
-		ContextTargetFlagsEnd UMETA(Hidden), // +1 to the last flag (so we can easily iterate these flags)
+		// +1 to the last flag (so we can easily iterate these flags)
+		ContextTargetFlagsEnd UMETA(Hidden),
 	};
 }
 
diff --git a/Engine/Source/Editor/LandscapeEditor/Classes/LandscapeEditorObject.h b/Engine/Source/Editor/LandscapeEditor/Classes/LandscapeEditorObject.h
--- a/Engine/Source/Editor/LandscapeEditor/Classes/LandscapeEditorObject.h
+++ b/Engine/Source/Editor/LandscapeEditor/Classes/LandscapeEditorObject.h
@@ -10,15 +10,15 @@ namespace ELandscapeToolFlattenMode
 {
 	enum Type
 	{
-		//Invalid = -1,
+		// Invalid = -1,
 
-		// Flatten may both raise and lower values
+		/** Flatten may both raise and lower values */
 		Both = 0,
 
-		// Flatten may only raise values, values above the clicked point will be left unchanged
+		/** Flatten may only raise values, values above the clicked point will be left unchanged */
 		Raise = 1,
 
-		// Flatten may only lower values, values below the clicked point will be left unchanged
+		/** Flatten may only lower values, values below the clicked point will be left unchanged */
 		Lower = 2,
 	};
 }
@@ -28,15 +28,15 @@ namespace ELandscapeToolErosionMode
 {
 	enum Type
 	{
-		//Invalid = -1,
+		// Invalid = -1,
 
-		// Apply all erosion effects, both raising and lowering the heightmap
+		/** Apply all erosion effects, both raising and lowering the heightmap */
 		Both = 0,
 
-		// Only applies erosion effects that result in raising the heightmap
+		/** Only applies erosion effects that result in raising the heightmap */
 		Raise = 1,
 
-		// Only applies erosion effects that result in lowering the heightmap
+		/** Only applies erosion effects that result in lowering the heightmap */
 		Lower = 2,
 	};
 }
@@ -46,12 +46,12 @@ namespace ELandscapeToolHydroErosionMode
 {
 	enum Type
 	{
-		//Invalid = -1,
+		// Invalid = -1,
 
-		// Rains in some places and not others, randomly
+		/** Rains in some places and not others, randomly */
 		Both = 0,
 
-		// Rain is applied to the entire area
+		/** Rain is applied to the entire area */
 		Positive = 1,
 	};
 }
@@ -63,15 +63,15 @@ namespace ELandscapeToolNoiseMode
 {
 	enum Type
 	{
-		//Invalid = -1,
+		// Invalid = -1,
 
-		// Noise will both raise and lower the heightmap
+		/** Noise will both raise and lower the heightmap */
 		Both = 0,
 
-		// Noise will only raise the heightmap
+		/** Noise will only raise the heightmap */
 		Raise = 1,
 
-		// Noise will only lower the heightmap
+		/** Noise will only lower the heightmap */
 		Lower = 2,
 	};
 }
@@ -82,15 +82,15 @@ namespace ELandscapeToolPasteMode
 {
 	enum Type
 	{
-		//Invalid = -1,
+		// Invalid = -1,
 
-		// Paste may both raise and lower values
+		/** Paste may both raise and lower values */
 		Both = 0,
 
-		// Paste may only raise values, places where the pasted data would be below the heightmap are left unchanged. Good for copy/pasting mountains
+		/** Paste may only raise values, places where the pasted data would be below the heightmap are left unchanged. Good for copy/pasting mountains */
 		Raise = 1,
 
-		// Paste may only lower values, places where the pasted data would be above the heightmap are left unchanged. Good for copy/pasting valleys or pits
+		/** Paste may only lower values, places where the pasted data would be above the heightmap are left unchanged. Good for copy/pasting valleys or pits */
 		Lower = 2,
 	};
 }
@@ -100,15 +100,15 @@ namespace ELandscapeConvertMode
 {
 	enum Type
 	{
-		//Invalid = -1,
+		// Invalid = -1,
 
-		// Given the new component size, the edges of the landscape will be expanded as necessary until its overall size is a whole number of landscape components.
+		/** Given the new component size, the edges of the landscape will be expanded as necessary until its overall size is a whole number of landscape components. */
 		Expand = 0,
 
-		// Given the new component size, the edges of the landscape will be trimmed until its overall size is a whole number of landscape components.
+		/** Given the new component size, the edges of the landscape will be trimmed until its overall size is a whole number of landscape components. */
 		Clip = 1,
 
-		// The landscape will have the same overall size in the world, and have the same number of components. Existing landscape geometry and layer data will be resampled to match the new resolution.
+		/** The landscape will have the same overall size in the world, and have the same number of components. Existing landscape geometry and layer data will be resampled to match the new resolution. */
 		Resample = 2,
 	};
 }
diff --git a/Engine/Source/Editor/LocalizationDashboard/Private/LocalizationTargetTypes.h b/Engine/Source/Editor/LocalizationDashboard/Private/LocalizationTargetTypes.h
--- a/Engine/Source/Editor/LocalizationDashboard/Private/LocalizationTargetTypes.h
+++ b/Engine/Source/Editor/LocalizationDashboard/Private/LocalizationTargetTypes.h
@@ -8,9 +8,12 @@
 UENUM()
 enum class ELocalizationTargetConflictStatus : uint8
 {
-	Unknown, /* The status of conflicts in this localization target could not be determined. */
-	ConflictsPresent, /* The are outstanding conflicts present in this localization target. */
-	Clear, /* The localization target is clear of conflicts. */
+	/** The status of conflicts in this localization target could not be determined. */
+	Unknown,
+	/** The are outstanding conflicts present in this localization target. */
+	ConflictsPresent,
+	/** The localization target is clear of conflicts. */
+	Clear,
 };
 
 USTRUCT()
@@ -227,12 +230,18 @@ struct FCultureStatistics
 UENUM()
 enum class ELocalizationTargetLoadingPolicy : uint8
 {
-	Never,			/* This target's localization data will never be loaded automatically. */
-	Always,			/* This target's localization data will always be loaded automatically. */
-	Editor,			/* This target's localization data will only be loaded when running the editor. Use if this target localizes the editor. */
-	Game,			/* This target's localization data will only be loaded when running the game. Use if this target localizes your game. */
-	PropertyNames,	/* This target's localization data will only be loaded if the editor is displaying localized property names. */
-	ToolTips,		/* This target's localization data will only be loaded if the editor is displaying localized tool tips. */
+	/** This target's localization data will never be loaded automatically. */
+	Never,
+	/** This target's localization data will always be loaded automatically. */
+	Always,
+	/** This target's localization data will only be loaded when running the editor. Use if this target localizes the editor. */
+	Editor,
+	/** This target's localization data will only be loaded when running the game. Use if this target localizes your game. */
+	Game,
+	/** This target's localization data will only be loaded if the editor is displaying localized property names. */
+	PropertyNames,
+	/** This target's localization data will only be loaded if the editor is displaying localized tool tips. */
+	Tooltips,
 };
 
 USTRUCT()
diff --git a/Engine/Source/Editor/Sequencer/Private/SequencerSettings.h b/Engine/Source/Editor/Sequencer/Private/SequencerSettings.h
--- a/Engine/Source/Editor/Sequencer/Private/SequencerSettings.h
+++ b/Engine/Source/Editor/Sequencer/Private/SequencerSettings.h
@@ -10,11 +10,11 @@ namespace ESequencerCurveVisibility
 {
 	enum Type
 	{
-		/** All curves should be visible. **/
+		/** All curves should be visible. */
 		AllCurves,
-		/** Only curves from selected nodes should be visible. **/
+		/** Only curves from selected nodes should be visible. */
 		SelectedCurves,
-		/** Only curves which have keyframes should be visible. **/
+		/** Only curves which have keyframes should be visible. */
 		AnimatedCurves
 	};
 }
diff --git a/Engine/Source/Editor/UnrealEd/Classes/CookOnTheSide/CookOnTheFlyServer.h b/Engine/Source/Editor/UnrealEd/Classes/CookOnTheSide/CookOnTheFlyServer.h
--- a/Engine/Source/Editor/UnrealEd/Classes/CookOnTheSide/CookOnTheFlyServer.h
+++ b/Engine/Source/Editor/UnrealEd/Classes/CookOnTheSide/CookOnTheFlyServer.h
@@ -43,10 +43,14 @@ namespace ECookMode
 {
 	enum Type
 	{
-		CookOnTheFly,				// default mode, handles requests from network
-		CookOnTheFlyFromTheEditor,	// cook on the side
-		CookByTheBookFromTheEditor,	// precook all resources while in the editor
-		CookByTheBook,				// cooking by the book (not in the editor)
+		/** Default mode, handles requests from network. */
+		CookOnTheFly,
+		/** Cook on the side. */
+		CookOnTheFlyFromTheEditor,
+		/** Precook all resources while in the editor. */
+		CookByTheBookFromTheEditor,
+		/** Cooking by the book (not in the editor). */
+		CookByTheBook,
 	};
 }
 
diff --git a/Engine/Source/Editor/UnrealEd/Classes/Editor/EditorEngine.h b/Engine/Source/Editor/UnrealEd/Classes/Editor/EditorEngine.h
--- a/Engine/Source/Editor/UnrealEd/Classes/Editor/EditorEngine.h
+++ b/Engine/Source/Editor/UnrealEd/Classes/Editor/EditorEngine.h
@@ -22,13 +22,13 @@ class USkeleton;
 UENUM()
 enum EMapSetBrushFlags				
 {
-	// Set brush color.
+	/** Set brush color. */
 	MSB_BrushColor	= 1,
-	// Set group.
+	/** Set group. */
 	MSB_Group		= 2,
-	// Set poly flags.
+	/** Set poly flags. */
 	MSB_PolyFlags	= 4,
-	// Set CSG operation.
+	/** Set CSG operation. */
 	MSB_BrushType	= 8,
 };
 
diff --git a/Engine/Source/Editor/UnrealEd/Classes/Editor/UnrealEdEngine.h b/Engine/Source/Editor/UnrealEd/Classes/Editor/UnrealEdEngine.h
--- a/Engine/Source/Editor/UnrealEd/Classes/Editor/UnrealEdEngine.h
+++ b/Engine/Source/Editor/UnrealEd/Classes/Editor/UnrealEdEngine.h
@@ -10,13 +10,13 @@
 UENUM()
 enum EPackageNotifyState
 {
-	// The user has been prompted with the balloon taskbar message
+	/** The user has been prompted with the balloon taskbar message. */
 	NS_BalloonPrompted,
-	// The user responded to the balloon task bar message and got the modal prompt to checkout dialog and responded to it
+	/** The user responded to the balloon task bar message and got the modal prompt to checkout dialog and responded to it. */
 	NS_DialogPrompted,
-	// The package has been marked dirty and is pending a balloon prompt
+	/** The package has been marked dirty and is pending a balloon prompt. */
 	NS_PendingPrompt,
-	// The package has been marked dirty but cannot be checked out, and is pending a modal warning dialog
+	/** The package has been marked dirty but cannot be checked out, and is pending a modal warning dialog. */
 	NS_PendingWarning,
 	NS_MAX,
 };
@@ -24,11 +24,11 @@ enum EPackageNotifyState
 UENUM()
 enum EWriteDisallowedWarningState
 {
-	// The user needs to be warned about the package
+	/** The user needs to be warned about the package. */
 	WDWS_PendingWarn,
-	// The user has been warned about the package
+	/** The user has been warned about the package. */
 	WDWS_Warned,
-	// Warning for the package unnecessary
+	/** Warning for the package unnecessary. */
 	WDWS_WarningUnnecessary,
 	WDWS_MAX,
 };
diff --git a/Engine/Source/Editor/UnrealEd/Classes/Editor/UnrealEdTypes.h b/Engine/Source/Editor/UnrealEd/Classes/Editor/UnrealEdTypes.h
--- a/Engine/Source/Editor/UnrealEd/Classes/Editor/UnrealEdTypes.h
+++ b/Engine/Source/Editor/UnrealEd/Classes/Editor/UnrealEdTypes.h
@@ -115,14 +115,20 @@ struct FLightmassParameterizedMaterialSettings
 UENUM()
 enum ELevelViewportType
 {
-	LVT_OrthoXY = 0,			// Top
-	LVT_OrthoXZ = 1,			// Front
-	LVT_OrthoYZ = 2,			// Left
+	/** Top */
+	LVT_OrthoXY = 0,
+	/** Front */
+	LVT_OrthoXZ = 1,
+	/** Left */
+	LVT_OrthoYZ = 2,
 	LVT_Perspective = 3,
 	LVT_OrthoFreelook = 4,
-	LVT_OrthoNegativeXY = 5,	// Bottom
-	LVT_OrthoNegativeXZ = 6,	// Back
-	LVT_OrthoNegativeYZ = 7,	// Right
+	/** Bottom */
+	LVT_OrthoNegativeXY = 5,
+	/** Back */
+	LVT_OrthoNegativeXZ = 6,
+	/** Right */
+	LVT_OrthoNegativeYZ = 7,
 	LVT_MAX,
 
 	LVT_None = 255,
diff --git a/Engine/Source/Editor/UnrealEd/Classes/Factories/FbxAnimSequenceImportData.h b/Engine/Source/Editor/UnrealEd/Classes/Factories/FbxAnimSequenceImportData.h
--- a/Engine/Source/Editor/UnrealEd/Classes/Factories/FbxAnimSequenceImportData.h
+++ b/Engine/Source/Editor/UnrealEd/Classes/Factories/FbxAnimSequenceImportData.h
@@ -9,7 +9,7 @@
 	b) As for FbxNode::GetAnimationInternval, this one will iterate through all properties recursively, and then for all animation curves it finds, for the animation layer index specified. So in other words, if one property has been animated, it will modify this result. This is completely different from GetLocalTimeSpan since it calculates the time span depending on the keys rather than just using the start and stop time that was saved in the file.
 */
 
-/** animation length type when importing*/
+/** Animation length type when importing */
 UENUM()
 enum EFBXAnimationLengthImportType
 {
diff --git a/Engine/Source/Editor/UnrealEd/Classes/Factories/FbxStaticMeshImportData.h b/Engine/Source/Editor/UnrealEd/Classes/Factories/FbxStaticMeshImportData.h
--- a/Engine/Source/Editor/UnrealEd/Classes/Factories/FbxStaticMeshImportData.h
+++ b/Engine/Source/Editor/UnrealEd/Classes/Factories/FbxStaticMeshImportData.h
@@ -12,11 +12,11 @@ namespace EVertexColorImportOption
 {
 	enum Type
 	{
-		/** Import the static mesh using the vertex colors from the FBX file */
+		/** Import the static mesh using the vertex colors from the FBX file. */
 		Replace,
-		/** Ignore vertex colors from the FBX file, and keep the existing mesh vertex colors */
+		/** Ignore vertex colors from the FBX file, and keep the existing mesh vertex colors. */
 		Ignore,
-		/** Override all vertex colors with the specified color */
+		/** Override all vertex colors with the specified color. */
 		Override
 	};
 }
diff --git a/Engine/Source/Editor/UnrealEd/Classes/Settings/LevelEditorPlaySettings.h b/Engine/Source/Editor/UnrealEd/Classes/Settings/LevelEditorPlaySettings.h
--- a/Engine/Source/Editor/UnrealEd/Classes/Settings/LevelEditorPlaySettings.h
+++ b/Engine/Source/Editor/UnrealEd/Classes/Settings/LevelEditorPlaySettings.h
@@ -61,10 +61,10 @@ enum EPlayModeType
 	/** Runs in VR. */
 	PlayMode_InVR,
 
-	/** Simulates in viewport without possessing the player */
+	/** Simulates in viewport without possessing the player. */
 	PlayMode_Simulate,
 
-	/** The number of different Play Modes */
+	/** The number of different Play Modes. */
 	PlayMode_Count,
 };
 
@@ -81,13 +81,13 @@ enum EPlayNetMode
 UENUM()
 enum EPlayOnBuildMode
 {
-	/** Always build */
+	/** Always build. */
 	PlayOnBuild_Always UMETA(DisplayName="Always Build"),
 
 	/** Never build. */
 	PlayOnBuild_Never UMETA(DisplayName="Never Build"),
 
-	/** Build based on project type */
+	/** Build based on project type. */
 	PlayOnBuild_Default UMETA(DisplayName="Only Build Code Projects"),
 };
 
diff --git a/Engine/Source/Editor/UnrealEd/Classes/Settings/LevelEditorViewportSettings.h b/Engine/Source/Editor/UnrealEd/Classes/Settings/LevelEditorViewportSettings.h
--- a/Engine/Source/Editor/UnrealEd/Classes/Settings/LevelEditorViewportSettings.h
+++ b/Engine/Source/Editor/UnrealEd/Classes/Settings/LevelEditorViewportSettings.h
@@ -14,10 +14,10 @@
 UENUM()
 enum ERotationGridMode
 {
-	/** Using Divisions of 360 degrees (e.g 360/2. 360/3, 360/4, ... ) */
+	/** Using Divisions of 360 degrees (e.g 360/2. 360/3, 360/4, ... ). */
 	GridMode_DivisionsOf360,
 
-	/** Uses the user defined grid values */
+	/** Uses the user defined grid values. */
 	GridMode_Common,
 };
 
diff --git a/Engine/Source/Editor/UnrealEd/Classes/Settings/ProjectPackagingSettings.h b/Engine/Source/Editor/UnrealEd/Classes/Settings/ProjectPackagingSettings.h
--- a/Engine/Source/Editor/UnrealEd/Classes/Settings/ProjectPackagingSettings.h
+++ b/Engine/Source/Editor/UnrealEd/Classes/Settings/ProjectPackagingSettings.h
@@ -30,13 +30,13 @@ enum class EProjectPackagingInternationalizationPresets
 	/** English only. */
 	English,
 
-	/** English, French, Italian, German, Spanish */
+	/** English, French, Italian, German, Spanish. */
 	EFIGS,
 
-	/** English, French, Italian, German, Spanish, Chinese, Japanese, Korean */
+	/** English, French, Italian, German, Spanish, Chinese, Japanese, Korean. */
 	EFIGSCJK,
 
-	/** Chinese, Japanese, Korean */
+	/** Chinese, Japanese, Korean. */
 	CJK,
 
 	/** All known cultures. */
diff --git a/Engine/Source/Editor/UnrealEd/Classes/TexAligner/TexAligner.h b/Engine/Source/Editor/UnrealEd/Classes/TexAligner/TexAligner.h
--- a/Engine/Source/Editor/UnrealEd/Classes/TexAligner/TexAligner.h
+++ b/Engine/Source/Editor/UnrealEd/Classes/TexAligner/TexAligner.h
@@ -16,21 +16,21 @@ class FPoly;
 UENUM()
 enum ETexAlign
 {
-	// When passed to functions it means "use whatever the aligners default is"
+	/** When passed to functions it means "use whatever the aligners default is". */
 	TEXALIGN_None,
-	// No special alignment (just derive from UV vectors).
+	/** No special alignment (just derive from UV vectors). */
 	TEXALIGN_Default,
-	// Aligns to best U and V axis based on polys normal
+	/** Aligns to best U and V axis based on polys normal. */
 	TEXALIGN_Box,
-	// Maps the poly to the axis it is closest to laying parallel to
+	/** Maps the poly to the axis it is closest to laying parallel to. */
 	TEXALIGN_Planar,
-	// Fits texture to a polygon
+	/** Fits texture to a polygon. */
 	TEXALIGN_Fit,
-	// Special version of TEXALIGN_Planar.
+	/** Special version of TEXALIGN_Planar. */
 	TEXALIGN_PlanarAuto,
-	// Special version of TEXALIGN_Planar.
+	/** Special version of TEXALIGN_Planar. */
 	TEXALIGN_PlanarWall,
-	// Special version of TEXALIGN_Planar.
+	/** Special version of TEXALIGN_Planar. */
 	TEXALIGN_PlanarFloor,
 	TEXALIGN_MAX,
 };
@@ -68,7 +68,7 @@ class UTexAligner : public UObject
 	// @todo document
 	void Align( UWorld* InWorld, ETexAlign InTexAlignType, UModel* InModel );
 	
-	// Aligns a specific BSP surface
+	/** Aligns a specific BSP surface */
 	virtual void AlignSurf( ETexAlign InTexAlignType, UModel* InModel, FBspSurfIdx* InSurfIdx, FPoly* InPoly, FVector* InNormal );
 };
 
diff --git a/Engine/Source/Editor/UnrealEd/Public/EdGraphNode_Comment.h b/Engine/Source/Editor/UnrealEd/Public/EdGraphNode_Comment.h
--- a/Engine/Source/Editor/UnrealEd/Public/EdGraphNode_Comment.h
+++ b/Engine/Source/Editor/UnrealEd/Public/EdGraphNode_Comment.h
@@ -12,10 +12,10 @@ namespace ECommentBoxMode
 {
 	enum Type
 	{
-		// This comment box will move any fully contained nodes when it moves
+		/** This comment box will move any fully contained nodes when it moves. */
 		GroupMovement UMETA(DisplayName="Group Movement"),
 
-		// This comment box has no effect on nodes contained inside it
+		/** This comment box has no effect on nodes contained inside it. */
 		NoGroupMovement UMETA(DisplayName="Comment")
 	};
 }
diff --git a/Engine/Source/Runtime/AIModule/Classes/AITypes.h b/Engine/Source/Runtime/AIModule/Classes/AITypes.h
--- a/Engine/Source/Runtime/AIModule/Classes/AITypes.h
+++ b/Engine/Source/Runtime/AIModule/Classes/AITypes.h
@@ -46,6 +46,7 @@ namespace EAIOptionFlag
 	{
 		Default,
 		Enable UMETA(DisplayName = "Yes"),	// UHT was complaining when tried to use True as value instead of Enable
+
 		Disable UMETA(DisplayName = "No"),
 
 		MAX UMETA(Hidden)
@@ -86,7 +87,8 @@ namespace EPawnActionAbortState
 	{
 		NeverStarted,
 		NotBeingAborted,
-		MarkPendingAbort,	// this means waiting for child to abort before aborting self
+		/** This means waiting for child to abort before aborting self. */
+		MarkPendingAbort,
 		LatentAbortInProgress,
 		AbortDone,
 
@@ -126,11 +128,16 @@ namespace EAIRequestPriority
 {
 	enum Type
 	{
-		SoftScript, // actions requested by Level Designers by placing AI-hinting elements on the map
-		Logic,	// actions AI wants to do due to its internal logic				
-		HardScript, // actions LDs really want AI to perform
-		Reaction,	// actions being result of game-world mechanics, like hit reactions, death, falling, etc. In general things not depending on what AI's thinking
-		Ultimate,	// ultimate priority, to be used with caution, makes AI perform given action regardless of anything else (for example disabled reactions)
+		/** Actions requested by Level Designers by placing AI-hinting elements on the map. */
+		SoftScript,
+		/** Actions AI wants to do due to its internal logic. */
+		Logic,
+		/** Actions LDs really want AI to perform. */
+		HardScript,
+		/** Actions being result of game-world mechanics, like hit reactions, death, falling, etc. In general things not depending on what AI's thinking. */
+		Reaction,
+		/** Ultimate priority, to be used with caution, makes AI perform given action regardless of anything else (for example disabled reactions). */
+		Ultimate,
 
 		MAX UMETA(Hidden)
 	};
diff --git a/Engine/Source/Runtime/AIModule/Classes/BehaviorTree/BTCompositeNode.h b/Engine/Source/Runtime/AIModule/Classes/BehaviorTree/BTCompositeNode.h
--- a/Engine/Source/Runtime/AIModule/Classes/BehaviorTree/BTCompositeNode.h
+++ b/Engine/Source/Runtime/AIModule/Classes/BehaviorTree/BTCompositeNode.h
@@ -36,13 +36,18 @@ UENUM()
 namespace EBTDecoratorLogic
 {
 	// keep in sync with DescribeLogicOp() in BTCompositeNode.cpp
+
 	enum Type
 	{
 		Invalid,
-		Test,		// test decorator conditions
-		And,		// logic op: AND
-		Or,			// logic op: OR
-		Not,		// logic op: NOT
+		/** Test decorator conditions. */
+		Test,
+		/** logic op: AND */
+		And,
+		/** logic op: OR */
+		Or,
+		/** logic op: NOT */
+		Not,
 	};
 }
 
diff --git a/Engine/Source/Runtime/AIModule/Classes/EnvironmentQuery/EnvQueryTypes.h b/Engine/Source/Runtime/AIModule/Classes/EnvironmentQuery/EnvQueryTypes.h
--- a/Engine/Source/Runtime/AIModule/Classes/EnvironmentQuery/EnvQueryTypes.h
+++ b/Engine/Source/Runtime/AIModule/Classes/EnvironmentQuery/EnvQueryTypes.h
@@ -52,10 +52,14 @@ namespace EEnvTestFilterType
 {
 	enum Type
 	{
-		Minimum,	// For numeric tests
-		Maximum,	// For numeric tests
-		Range,		// For numeric tests
-		Match		// For boolean tests
+		/** For numeric tests. */
+		Minimum,
+		/** For numeric tests. */
+		Maximum,
+		/** For numeric tests. */
+		Range,
+		/** For Boolean tests. */
+		Match
 	};
 }
 
@@ -67,6 +71,7 @@ namespace EEnvTestScoreEquation
 		Linear,
 		Square,
 		InverseLinear,	// For now...
+
 		Constant
 		// What other curve shapes should be supported?  At first I was thinking we'd have parametric (F*V^P + C), but
 		// many versions of that curve would violate the [0, 1] output range which I think we should preserve.  So instead
@@ -101,9 +106,12 @@ namespace EEnvTestCost
 {
 	enum Type
 	{
-		Low,				// reading data, math operations (e.g. distance)
-		Medium,				// processing data from multiple sources (e.g. fire tickets)
-		High,				// really expensive calls (e.g. visibility traces, pathfinding)
+		/** Reading data, math operations (e.g. distance). */
+		Low,
+		/** Processing data from multiple sources (e.g. fire tickets). */
+		Medium,
+		/** Really expensive calls (e.g. visibility traces, pathfinding).  */
+		High,
 	};
 }
 
@@ -204,8 +212,10 @@ namespace EEnvQueryTestClamping
 	enum Type
 	{
 		None,			
-		SpecifiedValue,	// Clamp to value specified in test
-		FilterThreshold	// Clamp to test's filter threshold
+		/** Clamp to value specified in test. */
+		SpecifiedValue,
+		/** Clamp to test's filter threshold. */
+		FilterThreshold
 	};
 }
 
diff --git a/Engine/Source/Runtime/AIModule/Classes/Perception/AIPerceptionTypes.h b/Engine/Source/Runtime/AIModule/Classes/Perception/AIPerceptionTypes.h
--- a/Engine/Source/Runtime/AIModule/Classes/Perception/AIPerceptionTypes.h
+++ b/Engine/Source/Runtime/AIModule/Classes/Perception/AIPerceptionTypes.h
@@ -28,8 +28,10 @@ typedef FAIGenericID<FPerceptionListenerCounter> FPerceptionListenerID;
 UENUM()
 enum class EAISenseNotifyType : uint8
 {
-	OnEveryPerception,	// continuous update whenever target is perceived
-	OnPerceptionChange, // from "visible" to "not visible" or vice versa
+	/** Continuous update whenever target is perceived. */
+	OnEveryPerception,
+	/** From "visible" to "not visible" or vice versa. */
+	OnPerceptionChange,
 };
 
 struct FPerceptionChannelWhitelist
diff --git a/Engine/Source/Runtime/Android/AndroidRuntimeSettings/Classes/AndroidRuntimeSettings.h b/Engine/Source/Runtime/Android/AndroidRuntimeSettings/Classes/AndroidRuntimeSettings.h
--- a/Engine/Source/Runtime/Android/AndroidRuntimeSettings/Classes/AndroidRuntimeSettings.h
+++ b/Engine/Source/Runtime/Android/AndroidRuntimeSettings/Classes/AndroidRuntimeSettings.h
@@ -8,39 +8,41 @@ UENUM()
 namespace EAndroidScreenOrientation
 {
 	// IF THIS CHANGES, MAKE SURE TO UPDATE UEDeployAndroid.cs, ConvertOrientationIniValue()!
+
 	enum Type
 	{
-		// Portrait orientation (the display is taller than it is wide)
+		/** Portrait orientation (the display is taller than it is wide). */
 		Portrait,
 
-		// Portrait orientation rotated 180 degrees
+		/** Portrait orientation rotated 180 degrees. */
 		ReversePortrait,
 
-		// Use either portrait or reverse portrait orientation, based on the device orientation sensor
+		/** Use either portrait or reverse portrait orientation, based on the device orientation sensor. */
 		SensorPortrait,
 
-		// Landscape orientation (the display is wider than it is tall)
+		/** Landscape orientation (the display is wider than it is tall). */
 		Landscape,
 
-		// Landscape orientation rotated 180 degrees
+		/** Landscape orientation rotated 180 degrees. */
 		ReverseLandscape,
 
-		// Use either landscape or reverse landscape orientation, based on the device orientation sensor
+		/** Use either landscape or reverse landscape orientation, based on the device orientation sensor. */
 		SensorLandscape,
 
-		// Use any orientation the device normally supports, based on the device orientation sensor
+		/** Use any orientation the device normally supports, based on the device orientation sensor. */
 		Sensor,
 
-		// Use any orientation (including ones the device wouldn't choose in Sensor mode), based on the device orientation sensor
+		/** Use any orientation (including ones the device wouldn't choose in Sensor mode), based on the device orientation sensor. */
 		FullSensor,
 	};
 }
 
-// Depth buffer precision preferences
+/** Depth buffer precision preferences */
 UENUM()
 namespace EAndroidDepthBufferPreference
 {
 	// IF THIS CHANGES, MAKE SURE TO UPDATE UEDeployAndroid.cs, ConvertDepthBufferIniValue()!
+
 	enum Type
 	{
 		Default = 0 UMETA(DisplayName = "Default"),
@@ -58,11 +60,11 @@ struct FGooglePlayAchievementMapping
 {
 	GENERATED_USTRUCT_BODY()
 
-	// The game-specific achievement name (the one passed in to WriteAchievement calls)
+	/** The game-specific achievement name (the one passed in to WriteAchievement calls). */
 	UPROPERTY(EditAnywhere, Category = GooglePlayServices)
 	FString Name;
 
-	// The ID of the corresponding achievement, generated by the Google Play developer console.
+	/** The ID of the corresponding achievement, generated by the Google Play developer console. */
 	UPROPERTY(EditAnywhere, Category = GooglePlayServices)
 	FString AchievementID;
 };
@@ -75,11 +77,11 @@ struct FGooglePlayLeaderboardMapping
 {
 	GENERATED_USTRUCT_BODY()
 
-	// The game-specific leaderboard name (the one passed in to WriteLeaderboards calls)
+	/** The game-specific leaderboard name (the one passed in to WriteLeaderboards calls). */
 	UPROPERTY(EditAnywhere, Category = GooglePlayServices)
 	FString Name;
 
-	// The ID of the corresponding leaderboard, generated by the Google Play developer console.
+	/** The ID of the corresponding leaderboard, generated by the Google Play developer console. */
 	UPROPERTY(EditAnywhere, Category = GooglePlayServices)
 	FString LeaderboardID;
 };
diff --git a/Engine/Source/Runtime/CoreUObject/Classes/Object.h b/Engine/Source/Runtime/CoreUObject/Classes/Object.h
--- a/Engine/Source/Runtime/CoreUObject/Classes/Object.h
+++ b/Engine/Source/Runtime/CoreUObject/Classes/Object.h
@@ -42,13 +42,13 @@ namespace ELogTimes
 {
 	enum Type
 	{
-		// Do not display log timestamps
+		/** Do not display log timestamps. */
 		None UMETA(DisplayName = "None"),
 
-		// Display log timestamps in UTC
+		/** Display log timestamps in UTC. */
 		UTC UMETA(DisplayName = "UTC"),
 
-		// Display log timestamps in seconds elapsed since GStartTime
+		/** Display log timestamps in seconds elapsed since GStartTime. */
 		SinceGStartTime UMETA(DisplayName = "Time since application start"),
 	};
 }
@@ -82,26 +82,27 @@ enum EInterpCurveMode
 // @warning:	When you update this, you must add an entry to GPixelFormats(see RenderUtils.cpp)
 // @warning:	When you update this, you must add an entries to PixelFormat.h, usually just copy the generated section on the header into EPixelFormat
 // @warning:	The *Tools DLLs will also need to be recompiled if the ordering is changed, but should not need code changes.
+
 UENUM()
 enum EPixelFormat
 {
 	PF_Unknown,
 	PF_A32B32G32R32F,
 	PF_B8G8R8A8,
-	// UNORM red (0..1)
+	/** UNORM red (0..1) */
 	PF_G8,
 	PF_G16,
 	PF_DXT1,
 	PF_DXT3,
 	PF_DXT5,
 	PF_UYVY,
-	// A RGB FP format with platform-specific implementation, for use with render targets
+	/** A RGB FP format with platform-specific implementation, for use with render targets. */
 	PF_FloatRGB,
-	// A RGBA FP format with platform-specific implementation, for use with render targets
+	/** A RGBA FP format with platform-specific implementation, for use with render targets. */
 	PF_FloatRGBA,
-	// A depth+stencil format with platform-specific implementation, for use with render targets
+	/** A depth+stencil format with platform-specific implementation, for use with render targets. */
 	PF_DepthStencil,
-	// A depth format with platform-specific implementation, for use with render targets
+	/** A depth format with platform-specific implementation, for use with render targets. */
 	PF_ShadowDepth,
 	PF_R32_FLOAT,
 	PF_G16R16,
@@ -114,10 +115,10 @@ enum EPixelFormat
 	PF_R16F,
 	PF_R16F_FILTER,
 	PF_BC5,
-	// SNORM red, green (-1..1)
+	/** SNORM red, green (-1..1). */
 	PF_V8U8,
 	PF_A1,
-	// A low precision floating point format.
+	/** A low precision floating point format. */
 	PF_FloatR11G11B10,
 	PF_A8,
 	PF_R32_UINT,
@@ -130,28 +131,35 @@ enum EPixelFormat
 	PF_R16G16B16A16_SINT,
 	PF_R5G6B5_UNORM,
 	PF_R8G8B8A8,
-	// Only used for legacy loading; do NOT use!
+	/** Only used for legacy loading; do NOT use! */
 	PF_A8R8G8B8,
-	// High precision single channel block compressed, equivalent to a single channel BC5, 8 bytes per 4x4 block.
+	/** High precision single channel block compressed, equivalent to a single channel BC5, 8 bytes per 4x4 block. */
 	PF_BC4,
-	// UNORM red, green (0..1)
+	/** UNORM red, green (0..1). */
 	PF_R8G8,
-	// ATITC formats
+	/** ATITC format. */
 	PF_ATC_RGB,
+	/** ATITC format. */
 	PF_ATC_RGBA_E,
+	/** ATITC format. */
 	PF_ATC_RGBA_I,
-	// Used for creating SRVs to alias a DepthStencil buffer to read Stencil.  Don't use for creating textures.
+	/** Used for creating SRVs to alias a DepthStencil buffer to read Stencil.  Don't use for creating textures. */
 	PF_X24_G8,
 	PF_ETC1,
 	PF_ETC2_RGB,
 	PF_ETC2_RGBA,
 	PF_R32G32B32A32_UINT,
 	PF_R16G16_UINT,
-	PF_ASTC_4x4,	// 8.00 bpp
-	PF_ASTC_6x6,	// 3.56 bpp
-	PF_ASTC_8x8,	// 2.00 bpp
-	PF_ASTC_10x10,	// 1.28 bpp
-	PF_ASTC_12x12,	// 0.89 bpp
+	/** 8.00 bpp */
+	PF_ASTC_4x4,
+	/** 3.56 bpp */
+	PF_ASTC_6x6,
+	/** 2.00 bpp */
+	PF_ASTC_8x8,
+	/** 1.28 bpp */
+	PF_ASTC_10x10,
+	/** 0.89 bpp */
+	PF_ASTC_12x12,
 	PF_BC6H,
 	PF_BC7,
 	PF_MAX,
@@ -162,46 +170,46 @@ namespace EMouseCursor
 {
 	enum Type
 	{
-		/** Causes no mouse cursor to be visible */
+		/** Causes no mouse cursor to be visible. */
 		None,
 
-		/** Default cursor (arrow) */
+		/** Default cursor (arrow). */
 		Default,
 
-		/** Text edit beam */
+		/** Text edit beam. */
 		TextEditBeam,
 
-		/** Resize horizontal */
+		/** Resize horizontal. */
 		ResizeLeftRight,
 
-		/** Resize vertical */
+		/** Resize vertical. */
 		ResizeUpDown,
 
-		/** Resize diagonal */
+		/** Resize diagonal. */
 		ResizeSouthEast,
 
-		/** Resize other diagonal */
+		/** Resize other diagonal. */
 		ResizeSouthWest,
 
-		/** MoveItem */
+		/** MoveItem. */
 		CardinalCross,
 
-		/** Target Cross */
+		/** Target Cross. */
 		Crosshairs,
 
-		/** Hand cursor */
+		/** Hand cursor. */
 		Hand,
 
-		/** Grab Hand cursor */
+		/** Grab Hand cursor. */
 		GrabHand,
 
-		/** Grab Hand cursor closed */
+		/** Grab Hand cursor closed. */
 		GrabHandClosed,
 
-		/** a circle with a diagonal line through it */
+		/** a circle with a diagonal line through it. */
 		SlashedCircle,
 
-		/** Eye-dropper cursor for picking colors */
+		/** Eye-dropper cursor for picking colors. */
 		EyeDropper,
 	};
 }
@@ -210,44 +218,65 @@ namespace EMouseCursor
 UENUM(BlueprintType)
 enum class EUnit
 {
-	/** Scalar distance/length units */
+	/** Scalar distance/length unit. */
+
 	Micrometers, Millimeters, Centimeters, Meters, Kilometers,
 	Inches, Feet, Yards, Miles,
 	Lightyears,
 
+
 	/** Angular units */
+
 	Degrees, Radians,
 
+
 	/** Speed units */
+
 	MetersPerSecond, KilometersPerHour, MilesPerHour,
 
+
 	/** Temperature units */
+
 	Celsius, Farenheit, Kelvin,
 
+
 	/** Mass units */
+
 	Micrograms, Milligrams, Grams, Kilograms, MetricTons,
 	Ounces, Pounds, Stones,
 
+
 	/** Force units */
+
 	Newtons, PoundsForce, KilogramsForce,
 
+
 	/** Frequency units */
+
 	Hertz, Kilohertz, Megahertz, Gigahertz, RevolutionsPerMinute,
 
+
 	/** Data Size units */
+	
 	Bytes, Kilobytes, Megabytes, Gigabytes, Terabytes,
 
+
 	/** Luminous flux units */
+	
 	Lumens,
 
+
 	/** Time units */
+	
 	Milliseconds, Seconds, Minutes, Hours, Days, Months, Years,
 
-	/** Symbolic entry, not specifyable on meta data */
+
+	/** Symbolic entry, not specifiable on meta data. */
+	
 	Unspecified
 };
 
-// A globally unique identifier.
+/** A globally unique identifier. */
 USTRUCT(immutable, noexport, BlueprintType)
 struct FGuid
 {
diff --git a/Engine/Source/Runtime/Engine/Classes/AI/Navigation/RecastNavMesh.h b/Engine/Source/Runtime/Engine/Classes/AI/Navigation/RecastNavMesh.h
--- a/Engine/Source/Runtime/Engine/Classes/AI/Navigation/RecastNavMesh.h
+++ b/Engine/Source/Runtime/Engine/Classes/AI/Navigation/RecastNavMesh.h
@@ -48,6 +48,7 @@ UENUM()
 namespace ERecastPartitioning
 {
 	// keep in sync with rcRegionPartitioning enum!
+
 	enum Type
 	{
 		Monotone,
@@ -58,21 +59,21 @@ namespace ERecastPartitioning
 
 namespace ERecastPathFlags
 {
-	/** if set, path won't be post processed */
+	/** If set, path won't be post processed. */
 	const int32 SkipStringPulling = (1 << 0);
 
-	/** if set, path will contain navigation corridor */
+	/** If set, path will contain navigation corridor. */
 	const int32 GenerateCorridor = (1 << 1);
 }
 
-/** helper to translate FNavPathPoint.Flags */
+/** Helper to translate FNavPathPoint.Flags. */
 struct ENGINE_API FNavMeshNodeFlags
 {
-	/** Extra node information (like "path start", "off-mesh connection") */
+	/** Extra node information (like "path start", "off-mesh connection"). */
 	uint8 PathFlags;
-	/** Area type after this node */
+	/** Area type after this node. */
 	uint8 Area;
-	/** Area flags for this node */
+	/** Area flags for this node. */
 	uint16 AreaFlags;
 
 	FNavMeshNodeFlags() : PathFlags(0), Area(0), AreaFlags(0) {}
diff --git a/Engine/Source/Runtime/Engine/Classes/Animation/AnimInstance.h b/Engine/Source/Runtime/Engine/Classes/Animation/AnimInstance.h
--- a/Engine/Source/Runtime/Engine/Classes/Animation/AnimInstance.h
+++ b/Engine/Source/Runtime/Engine/Classes/Animation/AnimInstance.h
@@ -92,7 +92,7 @@ struct FA2Pose
 	FA2Pose() {}
 };
 
-/** component space poses **/
+/** Component space poses. */
 USTRUCT()
 struct ENGINE_API FA2CSPose : public FA2Pose
 {
@@ -102,7 +102,7 @@ private:
 	/** Pointer to current BoneContainer. */
 	const struct FBoneContainer * BoneContainer;
 
-	/** once evaluated to be mesh space, this flag will be set **/
+	/** Once evaluated to be mesh space, this flag will be set. */
 	UPROPERTY()
 	TArray<uint8> ComponentSpaceFlags;
 
@@ -112,25 +112,22 @@ public:
 	{
 	}
 
-	/** constructor - needs LocalPoses **/
+	/** Constructor - needs LocalPoses. */
 	void AllocateLocalPoses(const FBoneContainer& InBoneContainer, const FA2Pose & LocalPose);
 
-	/** constructor - needs LocalPoses **/
+	/** Constructor - needs LocalPoses. */
 	void AllocateLocalPoses(const FBoneContainer& InBoneContainer, const FTransformArrayA2 & LocalBones);
 
-	/** Returns if this struct is valid */
+	/** Returns if this struct is valid. */
 	bool IsValid() const;
 
 	/** Get parent bone index for given bone index. */
 	int32 GetParentBoneIndex(const int32& BoneIndex) const;
 
-	/** Returns local transform for the boneindex **/
+	/** Returns local transform for the bone index. **/
 	FTransform GetLocalSpaceTransform(int32 BoneIndex);
 
-	/**
-	 * Do not access Bones array directly but via this 
-	 * This will fill up gradually mesh space bases 
-	 */
+	/** Do not access Bones array directly; use this instead. This will fill up gradually mesh space bases. */
 	FTransform GetComponentSpaceTransform(int32 BoneIndex);
 
 	/** convert to local poses **/
@@ -149,8 +146,7 @@ private:
 
 	void SetLocalSpaceTransform(int32 Index, const FTransform& NewTransform);
 
-	/** this is not really best way to protect SetComponentSpaceTransform, but we'd like to make sure
-	 that isn't called by anywhere else */
+	// This is not really best way to protect SetComponentSpaceTransform, but we'd like to make sure that isn't called by anywhere else.
 	friend class FAnimationRuntime;
 };
 
@@ -158,10 +154,13 @@ USTRUCT(BlueprintType)
 struct FPerBoneBlendWeight
 {
 	GENERATED_USTRUCT_BODY()
+
+	/** Source index of the buffer. */
 	UPROPERTY()
-	int32 SourceIndex; // source index of the buffer
+	int32 SourceIndex;
+
 	UPROPERTY()
-	float BlendWeight; // how much blend weight
+	float BlendWeight;
 
 	FPerBoneBlendWeight()
 		: SourceIndex(0)
diff --git a/Engine/Source/Runtime/Engine/Classes/Animation/AnimLinkableElement.h b/Engine/Source/Runtime/Engine/Classes/Animation/AnimLinkableElement.h
--- a/Engine/Source/Runtime/Engine/Classes/Animation/AnimLinkableElement.h
+++ b/Engine/Source/Runtime/Engine/Classes/Animation/AnimLinkableElement.h
@@ -14,11 +14,11 @@ namespace EAnimLinkMethod
 {
 	enum Type
 	{
-		// Element stays at a specific time without moving
+		/** Element stays at a specific time without moving. */
 		Absolute,
-		// Element moves with its segment, but not when the segment changes size
+		/** Element moves with its segment, but not when the segment changes size. */
 		Relative,
-		// Element moves with its segment and will stay at a certain proportion through the segment
+		/** Element moves with its segment and will stay at a certain proportion through the segment. */
 		Proportional
 	};
 }
diff --git a/Engine/Source/Runtime/Engine/Classes/Animation/AnimNode_TransitionPoseEvaluator.h b/Engine/Source/Runtime/Engine/Classes/Animation/AnimNode_TransitionPoseEvaluator.h
--- a/Engine/Source/Runtime/Engine/Classes/Animation/AnimNode_TransitionPoseEvaluator.h
+++ b/Engine/Source/Runtime/Engine/Classes/Animation/AnimNode_TransitionPoseEvaluator.h
@@ -5,7 +5,7 @@
 #include "BonePose.h"
 #include "AnimNode_TransitionPoseEvaluator.generated.h"
 
-// Indicates which state is being evaluated by this node (source or destination)
+/** Indicates which state is being evaluated by this node (source or destination). */
 UENUM()
 namespace EEvaluatorDataSource
 {
@@ -16,26 +16,27 @@ namespace EEvaluatorDataSource
 	};
 }
 
-// Determines the behavior this node will use when updating and evaluating.
+/** Determines the behavior this node will use when updating and evaluating. */
 UENUM()
 namespace EEvaluatorMode
 {
 	enum Mode
 	{
-		// DataSource is ticked and evaluated every frame
+		/** DataSource is ticked and evaluated every frame. */
 		EM_Standard UMETA(DisplayName="Standard"),
 
-		// DataSource is never ticked and only evaluated on the first frame. Every frame after uses the cached pose from the first frame.
+		/** DataSource is never ticked and only evaluated on the first frame. Every frame after uses the cached pose from the first frame. */
 		EM_Freeze UMETA(DisplayName="Freeze"),
 
-		// DataSource is ticked and evaluated for a given number of frames, then freezes after and uses the cached pose for future frames.
+		/** DataSource is ticked and evaluated for a given number of frames, then freezes after and uses the cached pose for future frames. */
 		EM_DelayedFreeze UMETA(DisplayName="Delayed Freeze")
 	};
 }
 
 
-// Animation data node for state machine transitions.
-// Can be set to supply either the animation data from the transition source (From State) or the transition destination (To State).
+/** Animation data node for state machine transitions.
+ * Can be set to supply either the animation data from the transition source (From State) or the transition destination (To State).
+*/
 USTRUCT()
 struct ENGINE_API FAnimNode_TransitionPoseEvaluator : public FAnimNode_Base
 {
diff --git a/Engine/Source/Runtime/Engine/Classes/Animation/AnimSequence.h b/Engine/Source/Runtime/Engine/Classes/Animation/AnimSequence.h
--- a/Engine/Source/Runtime/Engine/Classes/Animation/AnimSequence.h
+++ b/Engine/Source/Runtime/Engine/Classes/Animation/AnimSequence.h
@@ -33,7 +33,7 @@ enum AnimationCompressionFormat
 };
 
 /**
- * Indicates animation data compression format.
+ * Indicates animation data key format.
  */
 UENUM()
 enum AnimationKeyFormat
@@ -45,35 +45,33 @@ enum AnimationKeyFormat
 };
 
 /** 
- * AdditiveAnimationType
- * 
- * This dictates if this is additive or not, If it is, what kind of additive
+ * Indicates whether an animation is additive, and what kind.
  */
 UENUM()
 enum EAdditiveAnimationType
 {
-	/** No additive */
+	/** No additive. */
 	AAT_None  UMETA(DisplayName="No additive"),
-	/* Create Additive from LocalSpace Base */
+	/* Create Additive from LocalSpace Base. */
 	AAT_LocalSpaceBase UMETA(DisplayName="Local Space"),
-	/* Create Additive from MeshSpace Rotation Only, Translation still will be LocalSpace */
+	/* Create Additive from MeshSpace Rotation Only, Translation still will be LocalSpace. */
 	AAT_RotationOffsetMeshSpace UMETA(DisplayName="Mesh Space"),
 	AAT_MAX,
 };
 
 /** 
- * What kind of additive animation
+ * For an additive animation, indicates what the animation is relative to.
  */
 UENUM()
 enum EAdditiveBasePoseType
 {
-	// will be deprecated
+	/** Will be deprecated. */
 	ABPT_None UMETA(DisplayName="None"),
-	// use ref pose of Skeleton as base
+	/** Use the Skeleton's ref pose as base. */
 	ABPT_RefPose UMETA(DisplayName="Reference Pose"),
-	// use whole animation as a base pose. Need BasePoseSeq.
+	/** Use a whole animation as a base pose. BasePoseSeq must be set. */
 	ABPT_AnimScaled UMETA(DisplayName="Selected animation scaled"),
-	// use animation as a base pose. Need BasePoseSeq and RefFrameIndex (will clamp if outside).
+	/** Use one frame of an animation as a base pose. BasePoseSeq and RefFrameIndex must be set (RefFrameIndex will be clamped). */
 	ABPT_AnimFrame UMETA(DisplayName="Selected animation frame"),
 	ABPT_MAX,
 };
@@ -115,11 +113,9 @@ struct FRawAnimSequenceTrack
 	}
 };
 
-/**
- * These two always should go together, but it is not right now. 
- * I wonder in the future, we change all compressed to be inside as well, so they all stay together
- * When remove tracks, it should be handled together
- */
+// These two always should go together, but it is not right now. 
+// I wonder in the future, we change all compressed to be inside as well, so they all stay together
+// When remove tracks, it should be handled together 
 USTRUCT()
 struct FAnimSequenceTrackContainer
 {
@@ -154,20 +150,20 @@ struct FAnimSequenceTrackContainer
 		return (AnimationTracks.Num());
 	}
 };
-/**
- *	@note we have plan to support skeleton hierarchy. 
- *	when that happens, we'd like to keep skeleton indexing
- */
+
+// @note We have a plan to support skeletal hierarchy. When that happens, we'd like to keep skeleton indexing.
 USTRUCT()
 struct FTrackToSkeletonMap
 {
 	GENERATED_USTRUCT_BODY()
 
+	// 0 is the current Skeleton (the one above), and N is the Nth parent.
 	UPROPERTY()
-	int32 SkeletonIndex_DEPRECATED;    // 0 as current Skeleton (the one above), and N is the N'th parent
+	int32 SkeletonIndex_DEPRECATED;
 
+	// Index of Skeleton.BoneTree this Track belongs to.
 	UPROPERTY()
-	int32 BoneTreeIndex;    // Index of Skeleton.BoneTree this Track belongs to
+	int32 BoneTreeIndex;
 
 
 	FTrackToSkeletonMap()
@@ -471,7 +467,7 @@ class UAnimSequence : public UAnimSequenceBase
 
 	class AnimEncoding* ScaleCodec;
 
-	/** Additive animation type **/
+	/** Additive animation type. **/
 	UPROPERTY(EditAnywhere, Category=AdditiveSettings, AssetRegistrySearchable)
 	TEnumAsByte<enum EAdditiveAnimationType> AdditiveAnimType;
 
diff --git a/Engine/Source/Runtime/Engine/Classes/Animation/AnimSequenceBase.h b/Engine/Source/Runtime/Engine/Classes/Animation/AnimSequenceBase.h
--- a/Engine/Source/Runtime/Engine/Classes/Animation/AnimSequenceBase.h
+++ b/Engine/Source/Runtime/Engine/Classes/Animation/AnimSequenceBase.h
@@ -29,7 +29,7 @@ namespace EAnimEventTriggerOffsets
 
 ENGINE_API float GetTriggerTimeOffsetForType(EAnimEventTriggerOffsets::Type OffsetType);
 
-/** Ticking method for AnimNotifies in AnimMontages */
+/** Ticking method for AnimNotifies in AnimMontages. */
 UENUM()
 namespace EMontageNotifyTickType
 {
@@ -42,16 +42,16 @@ namespace EMontageNotifyTickType
 	};
 }
 
-/** Filtering method for deciding whether to trigger a notify */
+/** Filtering method for deciding whether to trigger a notify. */
 UENUM()
 namespace ENotifyFilterType
 {
 	enum Type
 	{
-		/** No filtering*/
+		/** No filtering. */
 		NoFiltering,
 
-		/** Filter By Skeletal Mesh LOD */
+		/** Filter By Skeletal Mesh LOD. */
 		LOD,
 	};
 }
diff --git a/Engine/Source/Runtime/Engine/Classes/Animation/AnimationAsset.h b/Engine/Source/Runtime/Engine/Classes/Animation/AnimationAsset.h
--- a/Engine/Source/Runtime/Engine/Classes/Animation/AnimationAsset.h
+++ b/Engine/Source/Runtime/Engine/Classes/Animation/AnimationAsset.h
@@ -73,19 +73,19 @@ struct FBlendFilter
 	}
 };
 
-// Root Bone Lock options when extracting Root Motion
 UENUM()
+/** Root Bone Lock options when extracting Root Motion. */
 namespace ERootMotionRootLock
 {
 	enum Type
 	{
-		// Use reference pose root bone position
+		/** Use reference pose root bone position. */
 		RefPose,
 
-		// Use root bone position on first frame of animation.
+		/** Use root bone position on first frame of animation. */
 		AnimFirstFrame,
 
-		// FTransform::Identity
+		/** FTransform::Identity. */
 		Zero
 	};
 }
@@ -95,16 +95,16 @@ namespace ERootMotionMode
 {
 	enum Type
 	{
-		// Leave root motion in animation.
+		/** Leave root motion in animation. */
 		NoRootMotionExtraction,
 
-		// Extract root motion but do not apply it.
+		/** Extract root motion but do not apply it. */
 		IgnoreRootMotion,
 
-		// Root motion is taken from all animations contributing to the final pose, not suitable for network multiplayer setups
+		/** Root motion is taken from all animations contributing to the final pose, not suitable for network multiplayer setups. */
 		RootMotionFromEverything,
 
-		// Root motion is only taken from montages, suitable for network multiplayer setups
+		/** Root motion is only taken from montages, suitable for network multiplayer setups. */
 		RootMotionFromMontagesOnly,
 	};
 }
@@ -173,13 +173,13 @@ namespace EAnimGroupRole
 {
 	enum Type
 	{
-		// This node can be the leader, as long as it has a higher blend weight than the previous best leader
+		/** This node can be the leader, as long as it has a higher blend weight than the previous best leader. */
 		CanBeLeader,
 		
-		// This node will always be a follower (unless there are only followers, in which case the first one ticked wins)
+		/** This node will always be a follower (unless there are only followers, in which case the first one ticked wins). */
 		AlwaysFollower,
 
-		// This node will always be a leader (if more than one node is AlwaysLeader, the last one ticked wins)
+		/** This node will always be a leader (if more than one node is AlwaysLeader, the last one ticked wins). */
 		AlwaysLeader
 	};
 }
diff --git a/Engine/Source/Runtime/Engine/Classes/Animation/BlendSpaceBase.h b/Engine/Source/Runtime/Engine/Classes/Animation/BlendSpaceBase.h
--- a/Engine/Source/Runtime/Engine/Classes/Animation/BlendSpaceBase.h
+++ b/Engine/Source/Runtime/Engine/Classes/Animation/BlendSpaceBase.h
@@ -10,7 +10,7 @@
 #include "AnimSequence.h"
 #include "BlendSpaceBase.generated.h"
 
-// Interpolation data types.
+/** Interpolation data types. */
 UENUM()
 enum EBlendSpaceAxis
 {
@@ -25,11 +25,11 @@ struct FInterpolationParameter
 {
 	GENERATED_USTRUCT_BODY()
 
-	/** Interpolation Time for input, when it gets input, it will use this time to interpolate to target, used for smoother interpolation **/
+	/** Interpolation Time for input, when it gets input, it will use this time to interpolate to target, used for smoother interpolation. */
 	UPROPERTY(EditAnywhere, Category=Parameter)
 	float InterpolationTime;
 
-	/** Interpolation Type for input, when it gets input, it will use this filter to decide how to get to target **/
+	/** Interpolation Type for input, when it gets input, it will use this filter to decide how to get to target. */
 	UPROPERTY(EditAnywhere, Category=Parameter)
 	TEnumAsByte<EFilterInterpolationType> InterpolationType;
 };
@@ -42,15 +42,15 @@ struct FBlendParameter
 	UPROPERTY(EditAnywhere, Category=BlendParameter)
 	FString DisplayName;
 
-	// min value for this parameter
+	/** Min value for this parameter. */
 	UPROPERTY(EditAnywhere, Category=BlendParameter)
 	float Min;
 
-	// max value for this parameter
+	/** Max value for this parameter. */
 	UPROPERTY(EditAnywhere, Category=BlendParameter)
 	float Max;
 
-	// how many grid for this parameter
+	/** how many grid for this parameter. */
 	UPROPERTY(EditAnywhere, Category=BlendParameter)
 	int32 GridNum;
 
@@ -66,7 +66,7 @@ struct FBlendParameter
 	{
 		return Max-Min;
 	}
-	/** return size of each grid **/
+	/** Return size of each grid. */
 	float GetGridSize() const
 	{
 		return GetRange()/(float)GridNum;
@@ -74,7 +74,7 @@ struct FBlendParameter
 	
 };
 
-/** Sample data **/
+/** Sample data */
 USTRUCT()
 struct FBlendSample
 {
@@ -83,7 +83,8 @@ struct FBlendSample
 	UPROPERTY(EditAnywhere, Category=BlendSample)
 	class UAnimSequence* Animation;
 
-	// blend 0->x, blend 1->y, blend 2->z
+	//blend 0->x, blend 1->y, blend 2->z
+
 	UPROPERTY(EditAnywhere, Category=BlendSample)
 	FVector SampleValue;
 
diff --git a/Engine/Source/Runtime/Engine/Classes/Animation/Rig.h b/Engine/Source/Runtime/Engine/Classes/Animation/Rig.h
--- a/Engine/Source/Runtime/Engine/Classes/Animation/Rig.h
+++ b/Engine/Source/Runtime/Engine/Classes/Animation/Rig.h
@@ -65,12 +65,12 @@ namespace EControlConstraint
 		Orientation,
 		/** Translation constraint. */
 		Translation,
-		/** Max Number */
+		/** Max Number. */
 		Max
 	};
 }
 
-/** Constraint Transform Type. - currently unused*/
+/** Constraint Transform Type. - currently unused */
 UENUM()
 namespace EConstraintTransform
 {
diff --git a/Engine/Source/Runtime/Engine/Classes/Animation/Skeleton.h b/Engine/Source/Runtime/Engine/Classes/Animation/Skeleton.h
--- a/Engine/Source/Runtime/Engine/Classes/Animation/Skeleton.h
+++ b/Engine/Source/Runtime/Engine/Classes/Animation/Skeleton.h
@@ -54,21 +54,21 @@ namespace EBoneTranslationRetargetingMode
 	};
 }
 
-/** Each Bone node in BoneTree **/
+/** Each Bone node in BoneTree */
 USTRUCT()
 struct FBoneNode
 {
 	GENERATED_USTRUCT_BODY()
 
-	/** Name of bone, this is the search criteria to match with mesh bone. This will be NAME_None if deleted **/
+	/** Name of bone, this is the search criteria to match with mesh bone. This will be NAME_None if deleted. */
 	UPROPERTY()
 	FName Name_DEPRECATED;
 
-	/** Parent Index. -1 if not used. The root has 0 as its parent. Do not delete the element but set this to -1. If it is revived by other reason, fix up this link. **/
+	/** Parent Index. -1 if not used. The root has 0 as its parent. Do not delete the element but set this to -1. If it is revived by other reason, fix up this link. */
 	UPROPERTY()
 	int32 ParentIndex_DEPRECATED;
 
-	/** Retargeting Mode for Translation Component */
+	/** Retargeting Mode for Translation Component. */
 	UPROPERTY(EditAnywhere, Category=BoneNode)
 	TEnumAsByte<EBoneTranslationRetargetingMode::Type> TranslationRetargetingMode;
 
diff --git a/Engine/Source/Runtime/Engine/Classes/Camera/CameraShake.h b/Engine/Source/Runtime/Engine/Classes/Camera/CameraShake.h
--- a/Engine/Source/Runtime/Engine/Classes/Camera/CameraShake.h
+++ b/Engine/Source/Runtime/Engine/Classes/Camera/CameraShake.h
@@ -12,7 +12,7 @@
 UENUM()
 enum EInitialOscillatorOffset
 {
-	/** Start with random offset (default) */
+	/** Start with random offset (default). */
 	EOO_OffsetRandom,
 	/** Start with zero offset. */
 	EOO_OffsetZero,
diff --git a/Engine/Source/Runtime/Engine/Classes/Camera/CameraTypes.h b/Engine/Source/Runtime/Engine/Classes/Camera/CameraTypes.h
--- a/Engine/Source/Runtime/Engine/Classes/Camera/CameraTypes.h
+++ b/Engine/Source/Runtime/Engine/Classes/Camera/CameraTypes.h
@@ -22,11 +22,11 @@ namespace ECameraAnimPlaySpace
 {
 	enum Type
 	{
- 		/** This anim is applied in camera space */
+ 		/** This anim is applied in camera space. */
  		CameraLocal,
- 		/** This anim is applied in world space */
+ 		/** This anim is applied in world space. */
  		World,
- 		/** This anim is applied in a user-specified space (defined by UserPlaySpaceMatrix) */
+ 		/** This anim is applied in a user-specified space (defined by UserPlaySpaceMatrix). */
  		UserDefined,
  	};
 }
diff --git a/Engine/Source/Runtime/Engine/Classes/Components/SceneCaptureComponent2D.h b/Engine/Source/Runtime/Engine/Classes/Components/SceneCaptureComponent2D.h
--- a/Engine/Source/Runtime/Engine/Classes/Components/SceneCaptureComponent2D.h
+++ b/Engine/Source/Runtime/Engine/Classes/Components/SceneCaptureComponent2D.h
@@ -22,7 +22,7 @@ class USceneCaptureComponent2D : public USceneCaptureComponent
 {
 	GENERATED_UCLASS_BODY()
 
-	/** Camera field of view (in degrees) */
+	/** Camera field of view (in degrees). */
 	UPROPERTY(interp, Category=SceneCapture, meta=(DisplayName = "Field of View", UIMin = "5.0", UIMax = "170", ClampMin = "0.001", ClampMax = "360.0"))
 	float FOVAngle;
 
@@ -33,11 +33,10 @@ class USceneCaptureComponent2D : public USceneCaptureComponent
 	UPROPERTY(interp, Category=SceneCapture, meta=(DisplayName = "Capture Source"))
 	TEnumAsByte<enum ESceneCaptureSource> CaptureSource;
 
-	// post process stuff
 	UPROPERTY(interp, Category=PostProcessVolume, meta=(ShowOnlyInnerProperties))
 	struct FPostProcessSettings PostProcessSettings;
 
-	/** 0:no effect, 1:full effect */
+	/** Range (0.0, 1.0) where 0 indicates no effect, 1 indicates full effect. */
 	UPROPERTY(interp, Category=PostProcessVolume, BlueprintReadWrite, meta=(UIMin = "0.0", UIMax = "1.0"))
 	float PostProcessBlendWeight;
 
diff --git a/Engine/Source/Runtime/Engine/Classes/Components/SkeletalMeshComponent.h b/Engine/Source/Runtime/Engine/Classes/Components/SkeletalMeshComponent.h
--- a/Engine/Source/Runtime/Engine/Classes/Components/SkeletalMeshComponent.h
+++ b/Engine/Source/Runtime/Engine/Classes/Components/SkeletalMeshComponent.h
@@ -145,37 +145,37 @@ public:
 	ClothCollisionPrimType	PrimType;
 };
 
-/** to interact with environments */
+/** Used to define what type of objects cloth will react to. */
 struct FApexClothCollisionInfo
 {
 	enum OverlappedComponentType
 	{
-		// static objects
+		/** Interact with static objects only. */
 		OCT_STATIC,
-		// clothing objects
+		/** Interact with clothing objects. */
 		OCT_CLOTH,
 		OCT_MAX
 	};
 
 	OverlappedComponentType OverlapCompType;
-	// to verify validation of collision info
+	/** To verify validation of collision info. */
 	uint32 Revision;
-	// ClothingCollisions will be all released when clothing doesn't intersect with this component anymore
+	/** ClothingCollisions will be all released when clothing doesn't intersect with this component anymore. */
 	TArray<physx::apex::NxClothingCollision*> ClothingCollisions;			
 };
 #endif // #if WITH_CLOTH_COLLISION_DETECTION
 
-/** This enum defines how you'd like to update bones to physics world
-	If bone is simulating, you don't have to waste time on updating bone transform from kinematic
-	But also sometimes you don't like fixed bones to be updated by animation data **/
+/** This enum defines how you'd like to update bones to physics world.
+	If bone is simulating, you don't have to waste time on updating bone transform from kinematic.
+	But also sometimes you don't like fixed bones to be updated by animation data. */
 UENUM()
 namespace EKinematicBonesUpdateToPhysics
 {
 	enum Type
 	{
-		/** Update any bones that are not simulating*/
+		/** Update any bones that are not simulating. */
 		SkipSimulatingBones,
-		/** Skip physics update from kinematic changes **/
+		/** Skip physics update from kinematic changes. */
 		SkipAllBones
 	};
 }
@@ -206,19 +206,19 @@ struct FSingleAnimationPlayData
 	UPROPERTY()//EditAnywhere, BlueprintReadWrite, Category=Animation)
 	class UVertexAnimation* VertexAnimToPlay;
 
-	// Default setting for looping for SequenceToPlay. This is not current state of looping.
+	/** Default setting for looping for SequenceToPlay. This is not current state of looping. */
 	UPROPERTY(EditAnywhere, BlueprintReadWrite, Category=Animation, meta=(DisplayName="Looping"))
 	uint32 bSavedLooping:1;
 
-	// Default setting for playing for SequenceToPlay. This is not current state of playing.
+	/** Default setting for playing for SequenceToPlay. This is not current state of playing. */
 	UPROPERTY(EditAnywhere, BlueprintReadWrite, Category=Animation, meta=(DisplayName="Playing"))
 	uint32 bSavedPlaying:1;
 
-	// Default setting for position of SequenceToPlay to play. 
+	/** Default setting for position of SequenceToPlay to play. */
 	UPROPERTY(EditAnywhere, BlueprintReadWrite, Category=Animation, meta=(DisplayName="Initial Position"))
 	float SavedPosition;
 
-	// Default setting for playrate of SequenceToPlay to play. 
+	/** Default setting for play rate of SequenceToPlay to play. */
 	UPROPERTY(EditAnywhere, AdvancedDisplay, BlueprintReadWrite, Category=Animation, meta=(DisplayName="PlayRate"))
 	float SavedPlayRate;
 
@@ -230,10 +230,10 @@ struct FSingleAnimationPlayData
 		SavedPosition = 0.f;
 	}
 
-	/** Called when initialized **/
+	/** Called when initialized. */
 	ENGINE_API void Initialize(class UAnimSingleNodeInstance* Instance);
 
-	/** Populates this play data with the current state of the supplied instance */
+	/** Populates this play data with the current state of the supplied instance. */
 	ENGINE_API void PopulateFrom(UAnimSingleNodeInstance* Instance);
 
 	void ValidatePosition();
@@ -247,18 +247,18 @@ struct FSkeletalMeshComponentPreClothTickFunction : public FTickFunction
 {
 	GENERATED_USTRUCT_BODY()
 
-	/** World this tick function belongs to **/
+	/** World this tick function belongs to. */
 	class USkeletalMeshComponent*	Target;
 
 	/**
-	* Abstract function actually execute the tick.
-	* @param DeltaTime - frame time to advance, in seconds
-	* @param TickType - kind of tick for this frame
-	* @param CurrentThread - thread we are executing on, useful to pass along as new tasks are created
+	* Abstract function to execute the tick.
+	* @param DeltaTime - frame time to advance, in seconds.
+	* @param TickType - kind of tick for this frame.
+	* @param CurrentThread - thread we are executing on, useful to pass along as new tasks are created.
 	* @param MyCompletionGraphEvent - completion event for this task. Useful for holding the completetion of this task until certain child tasks are complete.
-	**/
+	*/
 	virtual void ExecuteTick(float DeltaTime, enum ELevelTick TickType, ENamedThreads::Type CurrentThread, const FGraphEventRef& MyCompletionGraphEvent) override;
-	/** Abstract function to describe this tick. Used to print messages about illegal cycles in the dependency graph **/
+	/** Abstract function to describe this tick. Used to print messages about illegal cycles in the dependency graph. */
 	virtual FString DiagnosticMessage() override;
 };
 
@@ -278,39 +278,34 @@ class ENGINE_API USkeletalMeshComponent : public USkinnedMeshComponent, public I
 	 * Animation 
 	 */
 	
-	/** @Todo anim: Matinee related data start - this needs to be replaced to new system **/
+	/** @Todo anim: Matinee related data start - this needs to be replaced to new system. */
 	
-	/** @Todo anim: Matinee related data end - this needs to be replaced to new system **/
+	/** @Todo anim: Matinee related data end - this needs to be replaced to new system. */
 
 protected:
-	/** Whether to use Animation Blueprint or play Single Animation Asset*/
+	/** Whether to use Animation Blueprint or play Single Animation Asset. */
 	UPROPERTY(EditAnywhere, BlueprintReadOnly, Category=Animation)
 	TEnumAsByte<EAnimationMode::Type>	AnimationMode;
 
 public:
 #if WITH_EDITORONLY_DATA
-	/**
-	 * The blueprint for creating an AnimationScript
-	 */
+	/** The blueprint for creating an AnimationScript. */
 	UPROPERTY()
 	class UAnimBlueprint* AnimationBlueprint_DEPRECATED;
 #endif
 
-	/** 
-	 * The AnimBlueprint class to use
-	 * Use 'SetAnimInstanceClass' to change at runtime. 
-	 */
+	/* The AnimBlueprint class to use. Use 'SetAnimInstanceClass' to change at runtime. */
 	UPROPERTY(EditAnywhere, BlueprintReadOnly, Category = Animation)
 	class UAnimBlueprintGeneratedClass* AnimBlueprintGeneratedClass;
 
-	/** The active animation graph program instance */
+	/** The active animation graph program instance. */
 	UPROPERTY(transient)
 	class UAnimInstance* AnimScriptInstance;
 
 	UPROPERTY(EditAnywhere, BlueprintReadWrite, Category=Animation, meta=(ShowOnlyInnerProperties))
 	struct FSingleAnimationPlayData AnimationData;
 
-	/** Temporary array of local-space (ie relative to parent bone) rotation/translation for each bone. */
+	/** Temporary array of local-space (relative to parent bone) rotation/translation for each bone. */
 	TArray<FTransform> LocalAtoms;
 	
 	// Update Rate
diff --git a/Engine/Source/Runtime/Engine/Classes/Components/SkinnedMeshComponent.h b/Engine/Source/Runtime/Engine/Classes/Components/SkinnedMeshComponent.h
--- a/Engine/Source/Runtime/Engine/Classes/Components/SkinnedMeshComponent.h
+++ b/Engine/Source/Runtime/Engine/Classes/Components/SkinnedMeshComponent.h
@@ -22,9 +22,12 @@ DECLARE_DELEGATE_OneParam(FOnAnimUpdateRateParamsCreated, FAnimUpdateRateParamet
 UENUM()
 enum EBoneVisibilityStatus
 {
-	BVS_HiddenByParent,    // Bone is hidden because it's parent is hidden
-	BVS_Visible,    // Bone is visible
-	BVS_ExplicitlyHidden,    // Bone is hidden directly
+	/** Bone is hidden because it's parent is hidden. */
+	BVS_HiddenByParent
+	/** Bone is visible. */
+	BVS_Visible,
+	/** Bone is hidden directly. */
+	BVS_ExplicitlyHidden,
 	BVS_MAX,
 };
 
@@ -32,41 +35,41 @@ enum EBoneVisibilityStatus
 UENUM()
 enum EPhysBodyOp
 {
-	// don't do anything
+	/** Don't do anything. */
 	PBO_None,
-	// terminate - if you terminate, you won't be able to re-init when unhidden
+	/** Terminate - if you terminate, you won't be able to re-init when unhidden. */
 	PBO_Term,
-	// disable collision - it will enable collision when unhidden
+	/** Disable collision - it will enable collision when unhidden. */
 	PBO_Disable,
 	PBO_MAX,
 };
 
-/** Skinned Mesh Update Flag based on rendered or not */
+/** Skinned Mesh Update Flag based on rendered or not. */
 UENUM()
 namespace EMeshComponentUpdateFlag
 {
 	enum Type
 	{
-		// Always Tick and Refresh BoneTransforms whether rendered or not
+		/** Always Tick and Refresh BoneTransforms whether rendered or not. */
 		AlwaysTickPoseAndRefreshBones,
-		// Always Tick, but Refresh BoneTransforms only when rendered
+		/** Always Tick, but Refresh BoneTransforms only when rendered. */
 		AlwaysTickPose,
-		// Tick only when rendered, and it will only RefreshBoneTransforms when rendered
+		/** Tick only when rendered, and it will only RefreshBoneTransforms when rendered. */
 		OnlyTickPoseWhenRendered,
 	};
 }
 
-/** Flag for specifying bone space */
+/** Values for specifying bone space. */
 UENUM()
 namespace EBoneSpaces
 {
 	enum Type
 	{
-		/** Set absolute position of bone in world space */
+		/** Set absolute position of bone in world space. */
 		WorldSpace		UMETA(DisplayName = "World Space"),
-		/** Set position of bone in components reference frame */
+		/** Set position of bone in components reference frame. */
 		ComponentSpace	UMETA(DisplayName = "Component Space"),
-		/** Set position of bone relative to parent bone */
+		/** Set position of bone relative to parent bone. */
 		//LocalSpace		UMETA( DisplayName = "Parent Bone Space" ),
 	};
 }
@@ -81,15 +84,15 @@ struct FActiveVertexAnim
 {
 	GENERATED_USTRUCT_BODY()
 
-	/** The anim that we want to apply. */
+	/** The animation that we want to apply. */
 	UPROPERTY()
 	class UVertexAnimBase* VertAnim;
 
-	/** Strength of the vertex animation, between 0.0 and 1.0 */
+	/** Strength of the vertex animation, between 0.0 and 1.0. */
 	UPROPERTY()
 	float Weight;
 
-	/** Time to evaluate the anim at */
+	/** Time to evaluate the animation at. */
 	UPROPERTY()
 	float Time;
 
@@ -121,13 +124,13 @@ struct FActiveVertexAnim
 	}
 };
 
-/** LOD specific setup for the skeletal mesh component */
+/** LOD specific setup for the skeletal mesh component. */
 USTRUCT()
 struct FSkelMeshComponentLODInfo
 {
 	GENERATED_USTRUCT_BODY()
 
-	/** Material corresponds to section. To show/hide each section, use this **/
+	/** Material corresponds to section. To show/hide each section, use this. */
 	UPROPERTY()
 	TArray<bool> HiddenMaterials;
 
diff --git a/Engine/Source/Runtime/Engine/Classes/Curves/CurveBase.h b/Engine/Source/Runtime/Engine/Classes/Curves/CurveBase.h
--- a/Engine/Source/Runtime/Engine/Classes/Curves/CurveBase.h
+++ b/Engine/Source/Runtime/Engine/Classes/Curves/CurveBase.h
@@ -87,8 +87,8 @@ struct TStructOpsTypeTraits< FKeyHandleMap > : public TStructOpsTypeTraitsBase
 };
 
 
+// @todo Some heavy refactoring can be done here. Much more stuff can go in this base class.
 /** A curve base class which enables key handles to index lookups */
-// @todo Some heavy refactoring can be done here. Much more stuff can go in this base class
 USTRUCT()
 struct ENGINE_API FIndexedCurve
 {
@@ -99,28 +99,28 @@ public:
 	/** Get number of keys in curve. */
 	virtual int32 GetNumKeys() const PURE_VIRTUAL(FIndexedCurve::GetNumKeys,return 0;);
 
-	/** Const iterator for the handles */
+	/** Const iterator for the handles. */
 	TMap<FKeyHandle, int32>::TConstIterator GetKeyHandleIterator() const;
 	
-	/** Gets the index of a handle, checks if the key handle is valid first */
+	/** Gets the index of a handle, checks if the key handle is valid first. */
 	int32 GetIndexSafe(FKeyHandle KeyHandle) const;
 
-	/** Checks to see if the key handle is valid for this curve */
+	/** Checks to see if the key handle is valid for this curve. */
 	virtual bool IsKeyHandleValid(FKeyHandle KeyHandle) const;
 
 protected:
-	/** Internal tool to get a handle from an index */
+	/** Internal tool to get a handle from an index. */
 	FKeyHandle GetKeyHandle(int32 KeyIndex) const;
 	
-	/** Gets the index of a handle */
+	/** Gets the index of a handle. */
 	int32 GetIndex(FKeyHandle KeyHandle) const;
 
-	/** Makes sure our handles are all valid and correct */
+	/** Makes sure our handles are all valid and correct. */
 	void EnsureIndexHasAHandle(int32 KeyIndex) const;
 	void EnsureAllIndicesHaveHandles() const;
 
 protected:
-	/** Map of which key handles go to which indices */
+	/** Map of which key handles go to which indices. */
 	UPROPERTY(transient)
 	mutable FKeyHandleMap KeyHandlesToIndices;
 };
@@ -128,7 +128,7 @@ protected:
 //////////////////////////////////////////////////////////////////////////
 // Rich curve data
 
-/** Method of interpolation between this key and the next */
+/** Method of interpolation between this key and the next. */
 UENUM()
 enum ERichCurveInterpMode
 {
@@ -137,7 +137,7 @@ enum ERichCurveInterpMode
 	RCIM_Cubic
 };
 
-/** If using RCIM_Cubic, this enum describes how the tangents should be controlled in editor */
+/** If using RCIM_Cubic, this enum describes how the tangents should be controlled in editor. */
 UENUM()
 enum ERichCurveTangentMode
 {
@@ -146,7 +146,7 @@ enum ERichCurveTangentMode
 	RCTM_Break
 };
 
-/** Enum to indicate whether if a tangent is 'weighted' (ie can be stretched) */
+/** Enum to indicate whether if a tangent is 'weighted' (ie can be stretched). */
 UENUM()
 enum ERichCurveTangentWeightMode
 {
diff --git a/Engine/Source/Runtime/Engine/Classes/Debug/ReporterBase.h b/Engine/Source/Runtime/Engine/Classes/Debug/ReporterBase.h
--- a/Engine/Source/Runtime/Engine/Classes/Debug/ReporterBase.h
+++ b/Engine/Source/Runtime/Engine/Classes/Debug/ReporterBase.h
@@ -2,7 +2,7 @@
 #pragma once
 #include "ReporterBase.generated.h"
 
-/** Draw styles for lines */
+/** Draw styles for lines. */
 UENUM()
 namespace EReporterLineStyle
 {
diff --git a/Engine/Source/Runtime/Engine/Classes/Debug/ReporterGraph.h b/Engine/Source/Runtime/Engine/Classes/Debug/ReporterGraph.h
--- a/Engine/Source/Runtime/Engine/Classes/Debug/ReporterGraph.h
+++ b/Engine/Source/Runtime/Engine/Classes/Debug/ReporterGraph.h
@@ -3,7 +3,7 @@
 #include "ReporterBase.h"
 #include "ReporterGraph.generated.h"
 
-/** Draw styles for axes */
+/** Draw styles for axes. */
 UENUM()
 namespace EGraphAxisStyle
 {
@@ -15,7 +15,7 @@ namespace EGraphAxisStyle
 	};
 }
 
-/** Draw styles for data */
+/** Draw styles for data. */
 UENUM()
 namespace EGraphDataStyle
 {
@@ -36,23 +36,23 @@ namespace ELegendPosition
 	};
 }
 
-/** Helper rectangle struct */
+/** Helper rectangle struct. */
 struct FRect
 {
 	FVector2D Min;
 	FVector2D Max;
 };
 
-/** Graph Line data */
+/** Graph Line data. */
 struct FGraphThreshold
 {
-	/** The threshold amount */
+	/** The threshold amount. */
 	float Threshold;
 
-	/** The colour of the threshold */
+	/** The color of the threshold. */
 	FLinearColor Color;
 
-	/** The Threshold name */
+	/** The Threshold name. */
 	UPROPERTY(EditAnywhere, BlueprintReadWrite, Category=FCombatEventData)
 	FString ThresholdName;
 
@@ -73,22 +73,22 @@ struct FGraphThreshold
 	}
 };
 
-/** Graph Line data */
+/** Graph Line data. */
 struct FGraphLine
 {
-	/** The list of data to plot */
+	/** The list of data to plot. */
 	TArray<FVector2D> Data;
 
-	/** The colour to graph this data with */
+	/** The color to graph this data with. */
 	FLinearColor Color;
 
-	/** Left extreme value */
+	/** Left extreme value. */
 	FVector2D LeftExtreme;
 
-	/** Right extreme value */
+	/** Right extreme value. */
 	FVector2D RightExtreme;
 
-	/** The event name */
+	/** The event name. */
 	UPROPERTY(EditAnywhere, BlueprintReadWrite, Category=FCombatEventData)
 	FString LineName;
 };
@@ -100,206 +100,205 @@ class ENGINE_API UReporterGraph : public UReporterBase
 
 public:
 
-	/** Set the styles of the axes
-	 * @param NewAxisStyle The new style to draw the axes with
-	 * @param NewDataStyle The new style to draw the data with
+	/** Set the styles of the axes.
+	 * @param NewAxisStyle The new style to draw the axes with.
+	 * @param NewDataStyle The new style to draw the data with.
 	 */
 	FORCEINLINE void SetStyles(EGraphAxisStyle::Type NewAxisStyle, EGraphDataStyle::Type NewDataStyle) { AxisStyle = NewAxisStyle; DataStyle = NewDataStyle; }
 
-	/** Set the color of the axes
-	 * @param NewAxesColor The colour to set the axes to
+	/** Set the color of the axes.
+	 * @param NewAxesColor The colour to set the axes to.
 	*/
 	FORCEINLINE void SetAxesColor(FLinearColor NewAxesColor) { AxesColor = NewAxesColor; }
 
-	/** Set The size of the graph
-	 * @param MinX The normalized minimum X extent of the graph
-	 * @param MaxX The normalized maximum X extent of the graph
-	 * @param MinY The normalized minimum Y extent of the graph
-	 * @param MaxY The normalized maximum Y extent of the graph
+	/** Set The size of the graph.
+	 * @param MinX The normalized minimum X extent of the graph.
+	 * @param MaxX The normalized maximum X extent of the graph.
+	 * @param MinY The normalized minimum Y extent of the graph.
+	 * @param MaxY The normalized maximum Y extent of the graph.
 	 */
 	void SetGraphScreenSize(float MinX, float MaxX, float MinY, float MaxY);
 	
-	/** Set The size of the graph
-	 * @param Min The normalized minimum extent of the graph
-	 * @param Max The normalized maximum extent of the graph
+	/** Set The size of the graph.
+	 * @param Min The normalized minimum extent of the graph.
+	 * @param Max The normalized maximum extent of the graph.
 	 */
 	void SetGraphScreenSize(const FVector2D& Min, const FVector2D& Max);
 
 	/** Main draw call for the Graph.
-	 * @param Canvas The canvas to draw this graph to
+	 * @param Canvas The canvas to draw this graph to.
 	 */
 	virtual void Draw(UCanvas* Canvas) override;
 
-	/** Set the axis min and max data for both axes
-	 * @param MinX The normalized minimum X data
-	 * @param MaxX The normalized maximum X data
-	 * @param MinY The normalized minimum Y data
-	 * @param MaxY The normalized maximum Y data
+	/** Set the axis min and max data for both axes.
+	 * @param MinX The normalized minimum X data.
+	 * @param MaxX The normalized maximum X data.
+	 * @param MinY The normalized minimum Y data.
+	 * @param MaxY The normalized maximum Y data.
 	 */
 	void SetAxesMinMax(float MinX, float MaxX, float MinY, float MaxY);
 
-	/** Set the axis min and max data for both axes
-	 * @param MinX The normalized minimum X data
-	 * @param MinY The normalized minimum Y data
+	/** Set the axis min and max data for both axes.
+	 * @param MinX The normalized minimum X data.
+	 * @param MinY The normalized minimum Y data.
 	 */
 	void SetAxesMinMax(const FVector2D& Min, const FVector2D& Max);
 
-	/** Set notches per axis
-	 * @param NewNumXNotches The number of notches/divisions to display on the graph on the X axis
-	 * @param NewNumYNotches The number of notches/divisions to display on the graph on the Y axis
+	/** Set notches per axis.
+	 * @param NewNumXNotches The number of notches/divisions to display on the graph on the X axis.
+	 * @param NewNumYNotches The number of notches/divisions to display on the graph on the Y axis.
 	 */
 	FORCEINLINE void SetNotchesPerAxis(int NewNumXNotches, int NewNumYNotches) { NumXNotches = NewNumXNotches; NumYNotches = NewNumYNotches; }
 
-	/** Set the number of graph lines 
-	 * @param NumDataLines The number of lines to plot on this graph
+	/** Set the number of graph lines.
+	 * @param NumDataLines The number of lines to plot on this graph.
 	 */
 	void SetNumGraphLines(int32 NumDataLines);
 
-	/** Get a pointer to a graph line
-	 * @param LineIndex Index of the graph line
-	 * @return Pointer to the graph line. Do not cache. May return NULL;
+	/** Get a pointer to a graph line.
+	 * @param LineIndex Index of the graph line.
+	 * @return Pointer to the graph line. Do not cache. May return NULL.
 	 */
 	FGraphLine* GetGraphLine(int32 LineIndex);
 
-	/** Set the number of Thresholds to display on this graph
-	 * @param NumThresholds The number of Thresholds
+	/** Set the number of Thresholds to display on this graph.
+	 * @param NumThresholds The number of Thresholds.
 	 */
 	void SetNumThresholds(int32 NumThresholds);
 
-	/** Get a pointer to a threshold
-	 * @param ThresholdIndex Index of the threshold
-	 * @return Pointer to the threshold. Do not cache. May return NULL;
+	/** Get a pointer to a threshold.
+	 * @param ThresholdIndex Index of the threshold.
+	 * @return Pointer to the threshold. Do not cache. May return NULL.
 	 */
 	FGraphThreshold* GetThreshold(int32 ThresholdIndex);
 
-	/** Sets background color
-	 * @param Color Color to set as background
+	/** Sets background color.
+	 * @param Color Color to set as background.
 	 */
 	void SetBackgroundColor(FColor Color);
 
-	/** Sets position where to draw legend
-	 * @param Position Actual position for legend
+	/** Sets position where to draw legend.
+	 * @param Position Actual position for legend.
 	 */
 	void SetLegendPosition(ELegendPosition::Type Position);
 
-	/** Enables small effset for data sets to make it easier to read
-	 * @Param Enable - set to true to enable offsets
+	/** Enables small offset for data sets to make it easier to read.
+	 * @Param Enable - set to true to enable offsets.
 	 */
 	void OffsetDataSets(bool Enable) { bOffsetDataSets = Enable; }
 
-	/** Checks if we have enabled offset for data sets on graph
-	 */
+	/** Checks if we have enabled offset for data sets on graph. */
 	bool IsOffsetForDataSetsEnabled() { return !!bOffsetDataSets; }
 
-	/** Sets cursor location on line graph to show value at specyfic place */
+	/** Sets cursor location on line graph to show value at specific place. */
 	void SetCursorLocation(float InValue) { CursorLocation = InValue; }
 
-	/** setter to force tiny font instead of small font */
+	/** setter to force tiny font instead of small font. */
 	void UseTinyFont(bool InUseTinyFont) { bUseTinyFont = InUseTinyFont;  }
 
-	/** setter to enable or disable cursor for line graphs */
+	/** setter to enable or disable cursor for line graphs. */
 	void DrawCursorOnGraph(bool InDrawCursorOnGraph) { bDrawCursorOnGraph = InDrawCursorOnGraph;  }
 
 	void DrawExtremesOnGraph(bool InDrawExtremes) { bDrawExtremes = InDrawExtremes; }
 
-	/** Draw background under graph
-	* @param Canvas The canvas to draw on
+	/** Draw background under graph.
+	* @param Canvas The canvas to draw on.
 	*/
 	void DrawBackground(UCanvas* Canvas);
 
-	/** Draw the Legend
-	 * @param Canvas The canvas to draw on
+	/** Draw the Legend.
+	 * @param Canvas The canvas to draw on.
 	 */
 	void DrawLegend(UCanvas* Canvas);
 
-	/** Draw the Axes
-	 * @param Canvas The canvas to draw on
+	/** Draw the Axes.
+	 * @param Canvas The canvas to draw on.
 	 */
 	void DrawAxes(UCanvas* Canvas);
 
-	/** Draw an individual axis
-	 * @param Canvas The canvas to draw on
-	 * @param Start Start position of the axis
-	 * @param End End position of the axis
-	 * @param NumNotches The number of notches to draw on this axis
+	/** Draw an individual axis.
+	 * @param Canvas The canvas to draw on.
+	 * @param Start Start position of the axis.
+	 * @param End End position of the axis.
+	 * @param NumNotches The number of notches to draw on this axis.
 	 * @param bIsVerticalAxis Is this the vertical axis of the graph?
 	 */
 	void DrawAxis(UCanvas* Canvas, FVector2D Start, FVector2D End, float NumNotches, bool bIsVerticalAxis);
 
-	/** Draw the thresholds
-	 * @param Canvas The canvas to draw on
+	/** Draw the thresholds.
+	 * @param Canvas The canvas to draw on.
 	 */
 	void DrawThresholds(UCanvas* Canvas);
 
-	/** Draw the data
-	 * @param Canvas The canvas to draw on
+	/** Draw the data.
+	 * @param Canvas The canvas to draw on.
 	 */
 	void DrawData(UCanvas* Canvas);
 
-	/** Helper to convert vectors from normalized into screen space
-	 * @param InVector Normalized screen space coordinates to convert to screen space
-	 * @param Canvas The canvas we wish to use during the conversion
-	 * @return The screen space coordinates based on the size of the canvas
+	/** Helper to convert vectors from normalized into screen space.
+	 * @param InVector Normalized screen space coordinates to convert to screen space.
+	 * @param Canvas The canvas we wish to use during the conversion.
+	 * @return The screen space coordinates based on the size of the canvas.
 	 */
 	virtual FVector2D ToScreenSpace(const FVector2D& InVector, UCanvas* Canvas) override;
 
-	/** Helper to convert data from raw data into screen space
-	 * @param InVector 2D plot point in data space
-	 * @return The data space coordinate converted to normalized screen coordinates
+	/** Helper to convert data from raw data into screen space.
+	 * @param InVector 2D plot point in data space.
+	 * @return The data space coordinate converted to normalized screen coordinates.
 	*/
 	FVector2D DataToNormalized(const FVector2D& InVector);
 
-	/** Returns default font used to print texts */
+	/** Returns default font used to print texts. */
 	UFont* GetDefaultFont();
 
-	/** The screen size of the graph */
+	/** The screen size of the graph. */
 	FRect GraphScreenSize;
 
-	/** The minimum and maximum for the graph data */
+	/** The minimum and maximum for the graph data. */
 	FRect GraphMinMaxData;
 
 	TArray<FGraphThreshold> Thresholds;
 
-	/** The data displayed on the graph */
+	/** The data displayed on the graph. */
 	TArray<FGraphLine> CurrentData;
 
-	/** The axes Color */
+	/** The color of the axes. */
 	FLinearColor AxesColor;
 
-	/** The number of notches on the X axis **/
+	/** The number of notches on the X axis. **/
 	int NumXNotches;
 		
-	/** The number of notches on the Y axis **/
+	/** The number of notches on the Y axis. **/
 	int NumYNotches;
 
-	/** The axis style */
+	/** The axis style. */
 	EGraphAxisStyle::Type AxisStyle;
 
-	/** The data style */
+	/** The data style. */
 	EGraphDataStyle::Type DataStyle;
 
-	/** Current legend position */
+	/** Current legend position. */
 	ELegendPosition::Type LegendPosition;
 
-	/** The maximum length of the legend names */
+	/** The maximum length of the legend names. */
 	float LegendWidth;
 
-	/** Background color to draw under graph */
+	/** Background color to draw under graph. */
 	FColor BackgroundColor;
 
-	/** Current location for cursor on line graphs */
+	/** Current location for cursor on line graphs. */
 	float CursorLocation;
 
-	/** if set it enables small offset for graphs to better visualize overlapping data sets  */
+	/** If set, enables small offset for graphs to better visualize overlapping data sets. */
 	int32 bOffsetDataSets : 1;
 
-	/** if set it forces tinyt font for texts */
+	/** If set, forces tiny font for texts. */
 	int32 bUseTinyFont : 1;
 
-	/** if set it enables cursor for line graphs */
+	/** If set, enables cursor for line graphs. */
 	int32 bDrawCursorOnGraph : 1;
 
-	/** draw extreams on vertical axes */
+	/** If set, draws extremes on vertical axes. */
 	int32 bDrawExtremes : 1;
 };
 
diff --git a/Engine/Source/Runtime/Engine/Classes/EdGraph/EdGraphNode.h b/Engine/Source/Runtime/Engine/Classes/EdGraph/EdGraphNode.h
--- a/Engine/Source/Runtime/Engine/Classes/EdGraph/EdGraphNode.h
+++ b/Engine/Source/Runtime/Engine/Classes/EdGraph/EdGraphNode.h
@@ -12,7 +12,7 @@ class UEdGraphPin;
 class SGraphNode;
 struct FEdGraphPinType;
 
-/** Enum used to define which way data flows into or out of this pin */
+/** Enum used to define which way data flows into or out of this pin. */
 UENUM()
 enum EEdGraphPinDirection
 {
@@ -21,63 +21,63 @@ enum EEdGraphPinDirection
 	EGPD_MAX,
 };
 
-/** Enum to indicate what sort of title we want */
+/** Enum to indicate what sort of title we want. */
 UENUM()
 namespace ENodeTitleType
 {
 	enum Type
 	{
-		// The full title, may be multiple lines
+		/** The full title, may be multiple lines. */
 		FullTitle,
-		// More concise, single line title
+		/** More concise, single line title. */
 		ListView,
-		// Returns the editable title (which might not be a title at all)
+		/** Returns the editable title (which might not be a title at all). */
 		EditableTitle,
-		// Menu Title for context menus to be displayed in context menus referencing the node
+		/** Menu Title for context menus to be displayed in context menus referencing the node. */
 		MenuTitle,
 
 		MAX_TitleTypes,
 	};
 }
 
-/** Enum to indicate if a node has advanced-display-pins, and if they are shown */
+/** Enum to indicate if a node has advanced-display-pins, and whether they are shown. */
 UENUM()
 namespace ENodeAdvancedPins
 {
 	enum Type
 	{
-		// No advanced pins
+		/** No advanced pins. */
 		NoPins,
-		// There are some advanced pins, and they are shown
+		/** There are some advanced pins, and they are shown. */
 		Shown,
-		// There are some advanced pins, and they are hidden
+		/** There are some advanced pins, and they are hidden. */
 		Hidden
 	};
 }
 
-/** Holds metadata keys, so as to discourage text duplication throughout the engine */
+/** Holds metadata keys, so as to discourage text duplication throughout the engine. */
 struct ENGINE_API FNodeMetadata
 {
-	/** Identifies nodes that are defaultly added to populate new graphs (helps determine if a graph has any user placed nodes) */
+	/** Identifies nodes that are added to populate new graphs by default (helps determine if a graph has any user-placed nodes). */
 	static const FName DefaultGraphNode;
 private: 
 	FNodeMetadata() {}
 };
 
-// This is the context for a GetContextMenuActions call into a specific node
+/** This is the context for a GetContextMenuActions call into a specific node. */
 struct FGraphNodeContextMenuBuilder
 {
-	// The blueprint associated with this context; may be NULL for non-Kismet related graphs.
+	/** The blueprint associated with this context; may be NULL for non-Kismet related graphs. */
 	const UBlueprint* Blueprint;
-	// The graph associated with this context.
+	/** The graph associated with this context. */
 	const UEdGraph* Graph;
-	// The node associated with this context.
+	/** The node associated with this context. */
 	const UEdGraphNode* Node;
-	// The pin associated with this context; may be NULL when over a node.
+	/** The pin associated with this context; may be NULL when over a node. */
 	const UEdGraphPin* Pin;
-	// The menu builder to append actions to.
+	/** The menu builder to append actions to. */
 	class FMenuBuilder* MenuBuilder;
-	// Whether the graph editor is currently part of a debugging session (any non-debugging commands should be disabled).
+	/** Whether the graph editor is currently part of a debugging session (any non-debugging commands should be disabled). */
 	bool bIsDebugging;
 
 	FGraphNodeContextMenuBuilder(const UEdGraph* InGraph, const UEdGraphNode* InNode, const UEdGraphPin* InPin, class FMenuBuilder* InMenuBuilder, bool bInDebuggingMode);
diff --git a/Engine/Source/Runtime/Engine/Classes/EdGraph/EdGraphSchema.h b/Engine/Source/Runtime/Engine/Classes/EdGraph/EdGraphSchema.h
--- a/Engine/Source/Runtime/Engine/Classes/EdGraph/EdGraphSchema.h
+++ b/Engine/Source/Runtime/Engine/Classes/EdGraph/EdGraphSchema.h
@@ -25,70 +25,70 @@ enum EGraphType
 UENUM()
 enum ECanCreateConnectionResponse
 {
-	// Make the connection; there are no issues (message string is displayed if not empty)
+	/** Make the connection; there are no issues (message string is displayed if not empty). */
 	CONNECT_RESPONSE_MAKE,
 
-	// Cannot make this connection; display the message string as an error
+	/** Cannot make this connection; display the message string as an error. */
 	CONNECT_RESPONSE_DISALLOW,
 
-	// Break all existing connections on A and make the new connection (it's exclusive); display the message string as a warning/notice
+	/** Break all existing connections on A and make the new connection (it's exclusive); display the message string as a warning/notice. */
 	CONNECT_RESPONSE_BREAK_OTHERS_A,
 
-	// Break all existing connections on B and make the new connection (it's exclusive); display the message string as a warning/notice
+	/** Break all existing connections on B and make the new connection (it's exclusive); display the message string as a warning/notice. */
 	CONNECT_RESPONSE_BREAK_OTHERS_B,
 
-	// Break all existing connections on A and B, and make the new connection (it's exclusive); display the message string as a warning/notice
+	/** Break all existing connections on A and B, and make the new connection (it's exclusive); display the message string as a warning/notice. */
 	CONNECT_RESPONSE_BREAK_OTHERS_AB,
 
-	// Make the connection via an intermediate cast node, or some other conversion node
+	/** Make the connection via an intermediate cast node, or some other conversion node. */
 	CONNECT_RESPONSE_MAKE_WITH_CONVERSION_NODE,
 
 	CONNECT_RESPONSE_MAX,
 };
 
-/** This structure represents a context dependent action, with sufficient information for the schema to perform it */
+/** This structure represents a context dependent action, with sufficient information for the schema to perform it. */
 USTRUCT()
 struct ENGINE_API FEdGraphSchemaAction
 {
 	GENERATED_USTRUCT_BODY()
 
-	// Simple type info
+	/** Simple type info. */
 	static FName StaticGetTypeId() {static FName Type("FEdGraphSchemaAction"); return Type;}
 	virtual FName GetTypeId() const { return StaticGetTypeId(); }
 
-	/** The menu text that should be displayed for this node in the creation menu */
+	/** The menu text that should be displayed for this node in the creation menu. */
 	UPROPERTY()
 	FText MenuDescription;
 
-	/** The tooltip text that should be displayed for this node in the creation menu */
+	/** The tooltip text that should be displayed for this node in the creation menu. */
 	UPROPERTY()
 	FString TooltipDescription;
 
-	/** This is the UI centric category the action fits in (e.g., Functions, Variables).  Use this instead of the NodeType.NodeCategory because multiple NodeCategories might visually belong together. */
+	/** This is the UI centric category the action fits in (e.g., Functions, Variables). Use this instead of the NodeType.NodeCategory because multiple NodeCategories might visually belong together. */
 	UPROPERTY()
 	FText Category;
 
-	/** This is just an arbitrary dump of extra text that search will match on, in addition to the description and tooltip, e.g., Add might have the keyword Math */
+	/** This is just an arbitrary dump of extra text that search will match on, in addition to the description and tooltip, e.g., Add might have the keyword Math. */
 	UPROPERTY()
 	FText Keywords;
 
-	/** This is a priority number for overriding alphabetical order in the action list (higher value  == higher in the list) */
+	/** This is a priority number for overriding alphabetical order in the action list (higher value  == higher in the list). */
 	UPROPERTY()
 	int32 Grouping;
 
-	/** Section ID of the action list in which this action belongs */
+	/** Section ID of the action list in which this action belongs. */
 	UPROPERTY()
 	int32 SectionID;
 
-	/** Search title for the action (doesn't have to be set when instantiated, will be constructed by GetSearchTitle() if left empty)*/
+	/** Search title for the action (doesn't have to be set when instantiated, will be constructed by GetSearchTitle() if left empty). */
 	UPROPERTY()
 	FText CachedSearchTitle;
 
-	/** Search keywords for the action (doesn't have to be set when instantiated, will be constructed by GetSearchTitle() if left empty)*/
+	/** Search keywords for the action (doesn't have to be set when instantiated, will be constructed by GetSearchTitle() if left empty). */
 	UPROPERTY()
 	FText CachedSearchKeywords;
 
-	/** Search categories for the action (doesn't have to be set when instantiated, will be constructed by GetSearchTitle() if left empty)*/
+	/** Search categories for the action (doesn't have to be set when instantiated, will be constructed by GetSearchTitle() if left empty). */
 	UPROPERTY()
 	FText CachedSearchCategories;
 
@@ -108,13 +108,13 @@ struct ENGINE_API FEdGraphSchemaAction
 	{
 	}
 
-	/** Whether or not this action can be parented to other actions of the same type */
+	/** Whether or not this action can be parented to other actions of the same type. */
 	virtual bool IsParentable() const { return false; }
 
-	/** Execute this action, given the graph and schema, and possibly a pin that we were dragged from. Returns a node that was created by this action (if any).  */
+	/** Execute this action, given the graph and schema, and possibly a pin that we were dragged from. Returns a node that was created by this action (if any). */
 	virtual UEdGraphNode* PerformAction(class UEdGraph* ParentGraph, UEdGraphPin* FromPin, const FVector2D Location, bool bSelectNewNode = true) { return NULL; }
 
-	/** Execute this action, given the graph and schema, and possibly a pin that we were dragged from. Returns a node that was created by this action (if any).  */
+	/** Execute this action, given the graph and schema, and possibly a pin that we were dragged from. Returns a node that was created by this action (if any). */
 	virtual UEdGraphNode* PerformAction(class UEdGraph* ParentGraph, TArray<UEdGraphPin*>& FromPins, const FVector2D Location, bool bSelectNewNode = true)
 	{
 		UEdGraphNode* NewNode = NULL;
@@ -130,7 +130,7 @@ struct ENGINE_API FEdGraphSchemaAction
 		return NewNode;
 	}
 
-	/** Retrieves the full searchable title for this action */
+	/** Retrieves the full searchable title for this action. */
 	FText GetSearchTitle()
 	{
 		if(CachedSearchTitle.IsEmpty())
@@ -143,7 +143,7 @@ struct ENGINE_API FEdGraphSchemaAction
 		return CachedSearchTitle;
 	}
 
-	/** Retrieves the full searchable keywords for this action */
+	/** Retrieves the full searchable keywords for this action. */
 	FText GetSearchKeywords()
 	{
 		if(CachedSearchKeywords.IsEmpty())
@@ -156,7 +156,7 @@ struct ENGINE_API FEdGraphSchemaAction
 		return CachedSearchKeywords;
 	}
 
-	/** Retrieves the full searchable categories for this action */
+	/** Retrieves the full searchable categories for this action. */
 	FText GetSearchCategory()
 	{
 		if(CachedSearchCategories.IsEmpty())
diff --git a/Engine/Source/Runtime/Engine/Classes/Engine/BlendableInterface.h b/Engine/Source/Runtime/Engine/Classes/Engine/BlendableInterface.h
--- a/Engine/Source/Runtime/Engine/Classes/Engine/BlendableInterface.h
+++ b/Engine/Source/Runtime/Engine/Classes/Engine/BlendableInterface.h
@@ -7,28 +7,30 @@
 class UMaterialInterface;
 class UMaterialInstanceDynamic;
 
-/** Where to place a material node in the post processing graph */
+/** Where to place a material node in the post processing graph. */
 UENUM()
 enum EBlendableLocation
 {
-	// Input0:former pass color, Input1:SeparateTranslucency
+	/** Input0:former pass color, Input1:SeparateTranslucency. */
 	BL_AfterTonemapping UMETA(DisplayName="After Tonemapping"),
-	// Input0:former pass color, Input1:SeparateTranslucency
+	/** Input0:former pass color, Input1:SeparateTranslucency. */
 	BL_BeforeTonemapping UMETA(DisplayName="Before Tonemapping"),
-	// Input0:former pass color, Input1:SeparateTranslucency
+	/** Input0:former pass color, Input1:SeparateTranslucency. */
 	BL_BeforeTranslucency UMETA(DisplayName="Before Translucency"),
-	// Input0:former pass color, Input1:SeparateTranslucency, Input2: BloomOutput
-	// vector parameters: Engine.FilmWhitePoint
-	// scalar parameters: Engine.FilmSaturation, Engine.FilmContrast
+
+	/** Input0:former pass color, Input1:SeparateTranslucency, Input2: BloomOutput
+	/** vector parameters: Engine.FilmWhitePoint
+	/** scalar parameters: Engine.FilmSaturation, Engine.FilmContrast
 	BL_ReplacingTonemapper UMETA(DisplayName="Replacing the Tonemapper"),
 //	BL_AfterOpaque,
 //	BL_AfterFog,
 //	BL_AfterTranslucency,
 //	BL_AfterPostProcessAA,
+
 	BL_MAX,
 };
 
-/** Dummy class needed to support Cast<IBlendableInterface>(Object) */
+/** Dummy class needed to support Cast<IBlendableInterface>(Object). */
 UINTERFACE()
 class UBlendableInterface : public UInterface
 {
@@ -36,33 +38,33 @@ class UBlendableInterface : public UInterface
 };
 
 /**
- * derive from this class if you want to be blended by the PostProcess blending e.g. PostproceessVolume
+ * Derive from this class if you want to be blended by the PostProcess blending e.g. PostproceessVolume
  */
 class IBlendableInterface
 {
 	GENERATED_IINTERFACE_BODY()
 		
-	// @param Weight 0..1, 1=fully take the values from this object, crash if outside the valid range
+	/** @param Weight 0..1, 1=fully take the values from this object, crash if outside the valid range. */
 	virtual void OverrideBlendableSettings(class FSceneView& View, float Weight) const = 0;
 };
 
 struct FPostProcessMaterialNode
 {
-	// default constructor
+	// Default constructor
 	FPostProcessMaterialNode()
 		: MaterialInterface(0), bIsMID(false), Location(BL_MAX), Priority(0)
 	{
 		check(!IsValid());
 	}
 
-	// constructor
+	// Constructor
 	FPostProcessMaterialNode(UMaterialInterface* InMaterialInterface, EBlendableLocation InLocation, int32 InPriority)
 		: MaterialInterface(InMaterialInterface), bIsMID(false), Location(InLocation), Priority(InPriority)
 	{
 		check(IsValid());
 	}
 
-	// constructor
+	// Constructor
 	FPostProcessMaterialNode(UMaterialInstanceDynamic* InMID, EBlendableLocation InLocation, int32 InPriority)
 		: MaterialInterface((UMaterialInterface*)InMID), bIsMID(true), Location(InLocation), Priority(InPriority)
 	{
@@ -70,10 +72,10 @@ struct FPostProcessMaterialNode
 	}
 
 	UMaterialInterface* GetMaterialInterface() const { return MaterialInterface; }
-	// only call if you are sure it's an MID
+	/** Only call if you are sure it's an MID. */
 	UMaterialInstanceDynamic* GetMID() const { check(bIsMID); return (UMaterialInstanceDynamic*)MaterialInterface; }
 
-	// for type safety in FBlendableManager
+	/** For type safety in FBlendableManager. */
 	static const FName& GetFName()
 	{
 		static const FName Name = FName(TEXT("FPostProcessMaterialNode"));
@@ -103,11 +105,12 @@ struct FPostProcessMaterialNode
 private:
 	UMaterialInterface* MaterialInterface;
 
-	// if MaterialInterface is an MID, needed for GetMID()
+	/** if MaterialInterface is an MID, needed for GetMID(). */
 	bool bIsMID;
-	// e.g. BL_BeforeTonemapping
+
 	EBlendableLocation Location;
-	// default 0
+
+	/** Default is 0. */
 	int32 Priority;
 };
 
diff --git a/Engine/Source/Runtime/Engine/Classes/Engine/Blueprint.h b/Engine/Source/Runtime/Engine/Classes/Engine/Blueprint.h
--- a/Engine/Source/Runtime/Engine/Classes/Engine/Blueprint.h
+++ b/Engine/Source/Runtime/Engine/Classes/Engine/Blueprint.h
@@ -14,46 +14,44 @@
 UENUM()
 enum EBlueprintStatus
 {
-	// Blueprint is in an unknown state
+	/** Blueprint is in an unknown state. */
 	BS_Unknown,
-	// Blueprint has been modified but not recompiled
+	/** Blueprint has been modified but not recompiled. */
 	BS_Dirty,
-	// Blueprint tried but failed to be compiled
+	/** Blueprint tried but failed to be compiled. */
 	BS_Error,
-	// Blueprint has been compiled since it was last modified
+	/** Blueprint has been compiled since it was last modified. */
 	BS_UpToDate,
-	// Blueprint is in the process of being created for the first time
+	/** Blueprint is in the process of being created for the first time. */
 	BS_BeingCreated,
-	// Blueprint has been compiled since it was last modified. There are warnings.
+	/** Blueprint has been compiled since it was last modified. There are warnings. */
 	BS_UpToDateWithWarnings,
 	BS_MAX,
 };
 
 
-/**
- * Enumerates types of blueprints.
- */
+/** Enumerates types of blueprints. */
 UENUM()
 enum EBlueprintType
 {
-	// Normal blueprint
+	/** Normal blueprint. */
 	BPTYPE_Normal				UMETA(DisplayName="Blueprint Class"),
-	// Blueprint that is const during execution (no state graph and methods cannot modify member variables)
+	/** Blueprint that is const during execution (no state graph and methods cannot modify member variables). */
 	BPTYPE_Const				UMETA(DisplayName="Const Blueprint Class"),
-	// Blueprint that serves as a container for macros to be used in other blueprints
+	/** Blueprint that serves as a container for macros to be used in other blueprints. */
 	BPTYPE_MacroLibrary			UMETA(DisplayName="Blueprint Macro Library"),
-	// Blueprint that serves as an interface to be implemented by other blueprints
+	/** Blueprint that serves as an interface to be implemented by other blueprints. */
 	BPTYPE_Interface			UMETA(DisplayName="Blueprint Interface"),
-	// Blueprint that handles level scripting
+	/** Blueprint that handles level scripting. */
 	BPTYPE_LevelScript			UMETA(DisplayName="Level Blueprint"),
-	// Blueprint that servers as a container for functions to be used in other blueprints
+	/** Blueprint that servers as a container for functions to be used in other blueprints. */
 	BPTYPE_FunctionLibrary		UMETA(DisplayName="Blueprint Function Library"),
 
 	BPTYPE_MAX,
 };
 
 
-// Type of compilation
+/** Type of compilation. */
 namespace EKismetCompileType
 {
 	enum Type
diff --git a/Engine/Source/Runtime/Engine/Classes/Engine/Brush.h b/Engine/Source/Runtime/Engine/Classes/Engine/Brush.h
--- a/Engine/Source/Runtime/Engine/Classes/Engine/Brush.h
+++ b/Engine/Source/Runtime/Engine/Classes/Engine/Brush.h
@@ -13,15 +13,15 @@
 UENUM()
 enum ECsgOper
 {
-	// Active brush. (deprecated do not use)
+	/** Active brush. (deprecated, do not use.) */
 	CSG_Active,
-	// Add to world. (deprecated do not use)
+	/** Add to world. (deprecated, do not use.) */
 	CSG_Add,
-	// Subtract from world. (deprecated do not use)
+	/** Subtract from world. (deprecated, do not use.) */
 	CSG_Subtract,
-	// Form from intersection with world.
+	/** Form from intersection with world. */
 	CSG_Intersect,
-	// Form from negative intersection with world.
+	/** Form from negative intersection with world. */
 	CSG_Deintersect,
 	CSG_None,
 	CSG_MAX,
@@ -31,11 +31,11 @@ enum ECsgOper
 UENUM()
 enum EBrushType
 {
-	// Default/builder brush.
+	/** Default/builder brush. */
 	Brush_Default UMETA(Hidden),
-	// Add to world.
+	/** Add to world. */
 	Brush_Add UMETA(DisplayName=Additive),
-	// Subtract from world.
+	/** Subtract from world. */
 	Brush_Subtract UMETA(DisplayName=Subtractive),
 	Brush_MAX,
 };
diff --git a/Engine/Source/Runtime/Engine/Classes/Engine/CloudStorageBase.h b/Engine/Source/Runtime/Engine/Classes/Engine/CloudStorageBase.h
--- a/Engine/Source/Runtime/Engine/Classes/Engine/CloudStorageBase.h
+++ b/Engine/Source/Runtime/Engine/Classes/Engine/CloudStorageBase.h
@@ -16,11 +16,9 @@ enum ECloudStorageDelegate
 	CSD_DocumentQueryComplete,
 	CSD_DocumentReadComplete,
 	CSD_DocumentWriteComplete,
-	// Data:Document index that has the conflict
-	// Type:Int
-	// Desc:Called when multiple machines have 
-	//      updated the document, and script needs to determine which one to use, via the Resolve*
-	//      functions
+	// Data: Document index that has the conflict
+	// Type: Int
+	/** Desc: Called when multiple machines have updated the document, and script needs to determine which one to use, via the Resolve functions. */
 	CSD_DocumentConflictDetected,
 	CSD_MAX,
 };
@@ -35,11 +33,11 @@ class UCloudStorageBase
 {
 	GENERATED_UCLASS_BODY()
 
-	/** When using local storage (aka "cloud emulation"), this maintains a list of the file paths */
+	/** When using local storage (aka "cloud emulation"), this maintains a list of the file paths. */
 	UPROPERTY()
 	TArray<FString> LocalCloudFiles;
 
-	/** If true, delegate callbacks should be skipped */
+	/** If true, delegate callbacks should be skipped. */
 	UPROPERTY()
 	uint32 bSuppressDelegateCalls:1;
 
@@ -47,7 +45,7 @@ class UCloudStorageBase
 	virtual void Init();
 
 	/**
-	 * Initiate reading a key/value pair from cloud storage. A CSD_KeyValueReadComplete
+	 * Initiate reading a key/value pair from cloud storage. A CSD_KeyValueReadComplete.
 	 * delegate will be called when it completes (if this function returns true).
 	 * 
 	 * @param KeyName String name of the key to retrieve.
@@ -58,7 +56,7 @@ class UCloudStorageBase
 	virtual bool ReadKeyValue(const FString& KeyName, EPlatformInterfaceDataType Type, class UObject* SerializedObj = nullptr);
 
 	/**
-	 * Write a key/value pair to the cloud. A CSD_KeyValueWriteComplete
+	 * Write a key/value pair to the cloud. A CSD_KeyValueWriteComplete.
 	 * delegate will be called when it completes (if this function returns true).
 	 *
 	 * @param KeyName String name of the key to write.
diff --git a/Engine/Source/Runtime/Engine/Classes/Engine/Engine.h b/Engine/Source/Runtime/Engine/Classes/Engine/Engine.h
--- a/Engine/Source/Runtime/Engine/Classes/Engine/Engine.h
+++ b/Engine/Source/Runtime/Engine/Classes/Engine/Engine.h
@@ -14,20 +14,20 @@ class FCanvas;
 
 
 /**
- * Enumerats types of fully loaded packages.
+ * Enumerates types of fully loaded packages.
  */
 UENUM()
 enum EFullyLoadPackageType
 {
-	/** Load the packages when the map in Tag is loaded */
+	/** Load the packages when the map in Tag is loaded. */
 	FULLYLOAD_Map,
-	/** Load the packages before the game class in Tag is loaded. The Game name MUST be specified in the URL (game=Package.GameName). Useful for loading packages needed to load the game type (a DLC game type, for instance) */
+	/** Load the packages before the game class in Tag is loaded. The Game name MUST be specified in the URL (game=Package.GameName). Useful for loading packages needed to load the game type (a DLC game type, for instance). */
 	FULLYLOAD_Game_PreLoadClass,
-	/** Load the packages after the game class in Tag is loaded. Will work no matter how game is specified in UWorld::SetGameMode. Useful for modifying shipping gametypes by loading more packages (mutators, for instance) */
+	/** Load the packages after the game class in Tag is loaded. Will work no matter how game is specified in UWorld::SetGameMode. Useful for modifying shipping gametypes by loading more packages (mutators, for instance). */
 	FULLYLOAD_Game_PostLoadClass,
-	/** Fully load the package as long as the DLC is loaded */
+	/** Fully load the package as long as the DLC is loaded. */
 	FULLYLOAD_Always,
-	/** Load the package for a mutator that is active */
+	/** Load the package for a mutator that is active. */
 	FULLYLOAD_Mutator,
 	FULLYLOAD_MAX,
 };
@@ -59,7 +59,7 @@ enum EConsoleType
 };
 
 
-/** Struct to help hold information about packages needing to be fully-loaded for DLC, etc */
+/** Struct to help hold information about packages needing to be fully-loaded for DLC, etc. */
 USTRUCT()
 struct FFullyLoadedPackagesInfo
 {
diff --git a/Engine/Source/Runtime/Engine/Classes/Engine/EngineTypes.h b/Engine/Source/Runtime/Engine/Classes/Engine/EngineTypes.h
--- a/Engine/Source/Runtime/Engine/Classes/Engine/EngineTypes.h
+++ b/Engine/Source/Runtime/Engine/Classes/Engine/EngineTypes.h
@@ -37,14 +37,17 @@ enum EActorMetricsType
 };
 
 
-// Return values for UEngine::Browse
+/** Return values for UEngine::Browse. */
 namespace EBrowseReturnVal
 {
 	enum Type
 	{
-		Success,		// Successfully browsed to a new map
-		Failure,		// Immediately failed to browse
-		Pending,		// A connection is pending
+		/** Successfully browsed to a new map. */
+		Success,
+		/** Immediately failed to browse. */
+		Failure,
+		/** A connection is pending. */
+		Pending,
 	};
 }
 
@@ -76,9 +79,9 @@ namespace EAttachLocation
 UENUM()
 enum ESceneDepthPriorityGroup
 {
-	// world scene DPG
+	/** World scene DPG. */
 	SDPG_World,
-	// foreground scene DPG
+	/** Foreground scene DPG. */
 	SDPG_Foreground,
 	SDPG_MAX,
 };
@@ -162,19 +165,19 @@ namespace ETranslucentSortPolicy
 {
 	enum Type
 	{
-		// Sort based on distance from camera centerpoint to bounding sphere centerpoint. (default, best for 3D games)
+		/** Sort based on distance from camera centerpoint to bounding sphere centerpoint. (Default, best for 3D games.) */
 		SortByDistance = 0,
 
-		// Sort based on the post-projection Z distance to the camera.
+		/** Sort based on the post-projection Z distance to the camera. */
 		SortByProjectedZ = 1,
 
-		// Sort based on the projection onto a fixed axis. (best for 2D games)
+		/** Sort based on the projection onto a fixed axis. (Best for 2D games.) */
 		SortAlongAxis = 2,
 	};
 }
 
 
-/** Controls the way that the width scale property affects anim trails */
+/** Controls the way that the width scale property affects animation trails. */
 UENUM()
 enum ETrailWidthMode
 {
@@ -193,7 +196,9 @@ namespace EParticleCollisionMode
 	};
 }
 
-// Note: Check UMaterialInstance::Serialize if changed!!
+
+// Note: Check UMaterialInstance::Serialize if changed!
+
 UENUM()
 enum EMaterialShadingModel
 {
@@ -208,15 +213,15 @@ enum EMaterialShadingModel
 };
 
 
-// This is used by the drawing passes to determine tessellation policy, so changes here need to be supported in native code.
+/** This is used by the drawing passes to determine tessellation policy, so changes here need to be supported in native code. */
 UENUM()
 enum EMaterialTessellationMode
 {
-	/** Tessellation disabled */
+	/** Tessellation disabled. */
 	MTM_NoTessellation UMETA(DisplayName="No Tessellation"),
-	/** Simple tessellation */
+	/** Simple tessellation. */
 	MTM_FlatTessellation UMETA(DisplayName="Flat Tessellation"),
-	/** Simple Spline based tessellation */
+	/** Simple spline based tessellation. */
 	MTM_PNTriangles UMETA(DisplayName="PN Triangles"),
 	MTM_MAX,
 };
@@ -252,23 +257,17 @@ enum ELightingBuildQuality
 UENUM()
 enum ETriangleSortOption
 {
-	//0
 	TRISORT_None,
-	//1
 	TRISORT_CenterRadialDistance,
-	//2
 	TRISORT_Random,
-	//3
 	TRISORT_MergeContiguous,
-	//4
 	TRISORT_Custom,
-	//5
 	TRISORT_CustomLeftRight,
 	TRISORT_MAX,
 };
 
 
-/** Enum to specify which axis to use for the forward vector when using TRISORT_CustomLeftRight sort mode */
+/** Enum to specify which axis to use for the forward vector when using TRISORT_CustomLeftRight sort mode. */
 UENUM()
 enum ETriangleSortAxis
 {
@@ -493,11 +492,11 @@ enum ETraceTypeQuery
 UENUM()
 enum EPhysicsSceneType
 {
-	// The synchronous scene, which must finish before Unreal sim code is run
+	/** The synchronous scene, which must finish before Unreal simulation code is run. */
 	PST_Sync,
-	// The cloth scene, which may run while Unreal sim code runs
+	/** The cloth scene, which may run while Unreal simulation code runs. */
 	PST_Cloth,
-	// The asynchronous scene, which may run while Unreal sim code runs
+	/** The asynchronous scene, which may run while Unreal simulation code runs. */
 	PST_Async,
 	PST_MAX,
 };
@@ -818,7 +817,7 @@ enum class ESleepFamily : uint8
 };
 
 
-/** Enum used to indicate what type of timeline signature a function matches */
+/** Enum used to indicate what type of timeline signature a function matches. */
 UENUM()
 enum ETimelineSigType
 {
@@ -831,25 +830,25 @@ enum ETimelineSigType
 };
 
 
-/** Enum used to describe what type of collision is enabled on a body */
+/** Enum used to describe what type of collision is enabled on a body. */
 UENUM()
 namespace ECollisionEnabled 
 { 
 	enum Type 
 	{ 
-		/** No collision is enabled for this body */ 
+		/** No collision is enabled for this body. */
 		NoCollision UMETA(DisplayName="No Collision"), 
-		/** This body is used only for collision queries (raycasts, sweeps, and overlaps) */ 
+		/** This body is used only for collision queries (raycasts, sweeps, and overlaps). */
 		QueryOnly UMETA(DisplayName="Query Only (No Physics Collision)"),
-		/** This body is used only for physics collision */
+		/** This body is used only for physics collision. */
 		PhysicsOnly UMETA(DisplayName="Physics Only (No Query Collision)"),
-		/** This body interacts with all collision (Query and Physics) */ 
+		/** This body interacts with all collision (Query and Physics). */
 		QueryAndPhysics UMETA(DisplayName="Collision Enabled (Query and Physics)") 
 	}; 
 } 
 
 
-/** describes the physical state of a rigid body */
+/** Describes the physical state of a rigid body. */
 USTRUCT()
 struct FRigidBodyState
 {
@@ -1478,20 +1477,18 @@ enum ELightMapPaddingType
 };
 
 
-/**
- * Bit-field flags that affects storage (e.g. packing, streaming) and other info about a shadowmap.
- */
+/** Bit-field flags that affects storage (e.g. packing, streaming) and other info about a shadowmap. */
 UENUM()
 enum EShadowMapFlags
 {
-	// No flags
+	/** No flags. */
 	SMF_None			= 0,
-	// Shadowmap should be placed in a streaming texture
+	/** Shadowmap should be placed in a streaming texture. */
 	SMF_Streamed		= 0x00000001
 };
 
 
-/** reference to a specific material in a PrimitiveComponent */
+/** Reference to a specific material in a PrimitiveComponent. */
 USTRUCT()
 struct FPrimitiveMaterialRef
 {
@@ -2026,9 +2023,7 @@ struct FPOV
 };
 
 
-/**
- * The importance of a mesh feature when automatically generating mesh LODs.
- */
+/** The importance of a mesh feature when automatically generating mesh LODs. */
 UENUM()
 namespace EMeshFeatureImportance
 {
@@ -2044,9 +2039,7 @@ namespace EMeshFeatureImportance
 }
 
 
-/**
- * Settings used to reduce a mesh.
- */
+/** Settings used to reduce a mesh. */
 USTRUCT()
 struct FMeshReductionSettings
 {
@@ -2526,13 +2519,13 @@ struct ENGINE_API FRadialDamageEvent : public FDamageEvent
 UENUM()
 enum ENetRole
 {
-	// No role at all.
+	/** No role at all. */
 	ROLE_None,
-	// Locally simulated proxy of this actor.
+	/** Locally simulated proxy of this actor. */
 	ROLE_SimulatedProxy,
-	// Locally autonomous proxy of this actor.
+	/** Locally autonomous proxy of this actor. */
 	ROLE_AutonomousProxy,
-	// Authoritative control over the actor.
+	/** Authoritative control over the actor. */
 	ROLE_Authority,
 	ROLE_MAX,
 };
@@ -2541,15 +2534,15 @@ enum ENetRole
 UENUM()
 enum ENetDormancy
 {
-	// This actor can never go network dormant
+	/** This actor can never go network dormant. */
 	DORM_Never,
-	// This actor can go dormant, but is not currently dormant. Game code will tell it when it go dormant
+	/** This actor can go dormant, but is not currently dormant. Game code will tell it when it go dormant. */
 	DORM_Awake,
-	// This actor wants to go fully dormant for all connections
+	/** This actor wants to go fully dormant for all connections. */
 	DORM_DormantAll,
-	// This actor may want to go dormant for some connections, GetNetDormancy() will be called to find out which
+	/** This actor may want to go dormant for some connections, GetNetDormancy() will be called to find out which. */
 	DORM_DormantPartial,
-	// This actor is initially dormant for all connection if it was placed in map.
+	/** This actor is initially dormant for all connection if it was placed in map. */
 	DORM_Initial,
 	DORN_MAX,
 };
@@ -2576,10 +2569,14 @@ namespace EAutoReceiveInput
 UENUM()
 enum class EAutoPossessAI : uint8
 {
-	Disabled,				// Feature is disabled (do not automatically possess AI).
-	PlacedInWorld,			// Only possess by an AI Controller if Pawn is placed in the world.
-	Spawned,				// Only possess by an AI Controller if Pawn is spawned after the world has loaded.
-	PlacedInWorldOrSpawned, // Pawn is automatically possessed by an AI Controller whenever it is created.
+	/** Feature is disabled (do not automatically possess AI). */
+	Disabled,
+	/** Only possess by an AI Controller if Pawn is placed in the world. */
+	PlacedInWorld,
+	/** Only possess by an AI Controller if Pawn is spawned after the world has loaded. */
+	Spawned,
+	/** Pawn is automatically possessed by an AI Controller whenever it is created. */
+	PlacedInWorldOrSpawned,
 };
 
 
@@ -2588,11 +2585,16 @@ namespace EEndPlayReason
 {
 	enum Type
 	{
-		Destroyed,			// When the Actor or Component is explicitly destroyed
-		LevelTransition,	// When the world is being unloaded for a level transition
-		EndPlayInEditor,	// When the world is being unloaded because PIE is ending
-		RemovedFromWorld,	// When the level it is a member of is streamed out
-		Quit,				// When the application is being exited
+		/** When the Actor or Component is explicitly destroyed. */
+		Destroyed,
+		/** When the world is being unloaded for a level transition. */
+		LevelTransition,
+		/** When the world is being unloaded because PIE is ending. */
+		EndPlayInEditor,
+		/** When the level it is a member of is streamed out. */
+		RemovedFromWorld,
+		/** When the application is being exited. */
+		Quit,
 	};
 
 }
@@ -2621,11 +2623,11 @@ struct FRepMovement
 	UPROPERTY()
 	FRotator Rotation;
 
-	/** if set, RootComponent should be sleeping */
+	/** If set, RootComponent should be sleeping. */
 	UPROPERTY()
 	uint8 bSimulatedPhysicSleep : 1;
 
-	/** if set, additional physic data (angular velocity) will be replicated */
+	/** If set, additional physic data (angular velocity) will be replicated. */
 	UPROPERTY()
 	uint8 bRepPhysics : 1;
 
@@ -3048,33 +3050,33 @@ class ENGINE_API UEngineTypes : public UObject
 	GENERATED_UCLASS_BODY()
 
 public:
-	/** Convert a trace type to a collision channel */
+	/** Convert a trace type to a collision channel. */
 	static ECollisionChannel ConvertToCollisionChannel(ETraceTypeQuery TraceType);
 
-	/** Convert an object type to a collision channel */
+	/** Convert an object type to a collision channel. */
 	static ECollisionChannel ConvertToCollisionChannel(EObjectTypeQuery ObjectType);
 
-	/** Convert a collision channel to an object type. Note: performs a search of object types */
+	/** Convert a collision channel to an object type. Note: performs a search of object types. */
 	static EObjectTypeQuery ConvertToObjectType(ECollisionChannel CollisionChannel);
 
-	/** Convert a collision channel to a trace type. Note: performs a search of trace types */
+	/** Convert a collision channel to a trace type. Note: performs a search of trace types. */
 	static ETraceTypeQuery ConvertToTraceType(ECollisionChannel CollisionChannel);
 };
 
 
-/** Type of a socket on a scene component */
+/** Type of a socket on a scene component. */
 UENUM()
 namespace EComponentSocketType
 {
 	enum Type
 	{
-		/** Not a valid socket or bone name */
+		/** Not a valid socket or bone name. */
 		Invalid,
 
-		/** Skeletal bone */
+		/** Skeletal bone. */
 		Bone,
 
-		/** Socket */
+		/** Socket. */
 		Socket,
 	};
 }
@@ -3107,11 +3109,11 @@ struct FComponentSocketDescription
 UENUM()
 enum EAngularConstraintMotion
 {
-	/** No constraint against this axis */ 
+	/** No constraint against this axis. */ 
 	ACM_Free		UMETA(DisplayName="Free"),
-	/** Limited freedom along this axis */ 
+	/** Limited freedom along this axis. */ 
 	ACM_Limited		UMETA(DisplayName="Limited"),
-	/** Fully constraint against this axis */
+	/** Fully constraint against this axis. */
 	ACM_Locked		UMETA(DisplayName="Locked"),
 
 	ACM_MAX,
diff --git a/Engine/Source/Runtime/Engine/Classes/Engine/Font.h b/Engine/Source/Runtime/Engine/Classes/Engine/Font.h
--- a/Engine/Source/Runtime/Engine/Classes/Engine/Font.h
+++ b/Engine/Source/Runtime/Engine/Classes/Engine/Font.h
@@ -8,21 +8,19 @@
 #include "Font.generated.h"
 
 
-/**
- * Enumerates supported font caching types.
- */
+/** Enumerates supported font caching types. */
 UENUM()
 enum class EFontCacheType
 {
-	/** The font is using offline caching (this is how UFont traditionally worked) */
+	/** The font is using offline caching (this is how UFont traditionally worked). */
 	Offline,
 
-	/** The font is using runtime caching (this is how Slate fonts work) */
+	/** The font is using runtime caching (this is how Slate fonts work). */
 	Runtime,
 };
 
 
-/** this struct is serialized using native serialization so any changes to it require a package version bump */
+/** This struct is serialized using native serialization so any changes to it require a package version bump. */
 USTRUCT()
 struct FFontCharacter
 {
diff --git a/Engine/Source/Runtime/Engine/Classes/Engine/FontImportOptions.h b/Engine/Source/Runtime/Engine/Classes/Engine/FontImportOptions.h
--- a/Engine/Source/Runtime/Engine/Classes/Engine/FontImportOptions.h
+++ b/Engine/Source/Runtime/Engine/Classes/Engine/FontImportOptions.h
@@ -5,7 +5,7 @@
 #include "FontImportOptions.generated.h"
 
 
-/** Font character set type for importing TrueType fonts */
+/** Font character set type for importing TrueType fonts. */
 UENUM()
 enum EFontImportCharacterSet
 {
diff --git a/Engine/Source/Runtime/Engine/Classes/Engine/MicroTransactionBase.h b/Engine/Source/Runtime/Engine/Classes/Engine/MicroTransactionBase.h
--- a/Engine/Source/Runtime/Engine/Classes/Engine/MicroTransactionBase.h
+++ b/Engine/Source/Runtime/Engine/Classes/Engine/MicroTransactionBase.h
@@ -8,24 +8,28 @@
 #pragma once
 #include "MicroTransactionBase.generated.h"
 
-/** All the types of delegate callbacks that a MicroTransaction subclass may receive from C++ */
+/** All the types of delegate callbacks that a MicroTransaction subclass may receive from C++. */
 UENUM()
 enum EMicroTransactionDelegate
 {
-	// Data:None
-	// Desc:QueryForAvailablePurchases() is complete and AvailableProducts is ready for use
+	/**
+	 * Data:None
+	 * Desc:QueryForAvailablePurchases() is complete and AvailableProducts is ready for use.
+	 */
 	MTD_PurchaseQueryComplete,
-	// Data:Result code, and identifier of the product that completed
-	// Type:Custom
-	// Desc:IntValue will have one of the enums in EMicroTransactionResult, and StringValue
-	//      will have the Identifier from the PurchaseInfo that was bought with BeginPurchase
-	//      If MTR_Failed was returned, then LastError and LastErrorSolution should be filled
-	//		out with the most recent localized and possible resolutions
+	/**
+	 * Data:Result code, and identifier of the product that completed.
+	 * Type:Custom
+	 * Desc:IntValue will have one of the enums in EMicroTransactionResult, and StringValue
+	 *      will have the Identifier from the PurchaseInfo that was bought with BeginPurchase
+	 *      If MTR_Failed was returned, then LastError and LastErrorSolution should be filled
+	 *		out with the most recent localized and possible resolutions.
+	 */
 	MTD_PurchaseComplete,
 	MTD_MAX,
 };
 
-/** Result of a purchase */
+/** Result of a purchase. */
 UENUM()
 enum EMicroTransactionResult
 {
diff --git a/Engine/Source/Runtime/Engine/Classes/Engine/PlatformInterfaceBase.h b/Engine/Source/Runtime/Engine/Classes/Engine/PlatformInterfaceBase.h
--- a/Engine/Source/Runtime/Engine/Classes/Engine/PlatformInterfaceBase.h
+++ b/Engine/Source/Runtime/Engine/Classes/Engine/PlatformInterfaceBase.h
@@ -11,17 +11,17 @@
 #pragma once
 #include "PlatformInterfaceBase.generated.h"
 
-/** An enum for the types of data used in a PlatformInterfaceData struct, below */
+/** An enum for the types of data used in a PlatformInterfaceData struct. */
 UENUM()
 enum EPlatformInterfaceDataType
 {
-	// no data type specified
+	/** No data type specified. */
 	PIDT_None,
 	PIDT_Int,
 	PIDT_Float,
 	PIDT_String,
 	PIDT_Object,
-	// a custom type where more than one value may be filled out
+	/** A custom type where more than one value may be filled out. */
 	PIDT_Custom,
 	PIDT_MAX,
 };
@@ -29,7 +29,7 @@ enum EPlatformInterfaceDataType
 
 /** 
  * Struct that encompasses the most common types of data. This is the data payload
- * of PlatformInterfaceDelegateResult
+ * of PlatformInterfaceDelegateResult.
  */
 USTRUCT()
 struct FPlatformInterfaceData
diff --git a/Engine/Source/Runtime/Engine/Classes/Engine/Scene.h b/Engine/Source/Runtime/Engine/Classes/Engine/Scene.h
--- a/Engine/Source/Runtime/Engine/Classes/Engine/Scene.h
+++ b/Engine/Source/Runtime/Engine/Classes/Engine/Scene.h
@@ -7,7 +7,7 @@
 #pragma once
 #include "Scene.generated.h"
 
-/** used by FPostProcessSettings Depth of Field */
+/** Used by FPostProcessSettings Depth of Fields */
 UENUM()
 enum EDepthOfFieldMethod
 {
@@ -17,7 +17,7 @@ enum EDepthOfFieldMethod
 	DOFM_MAX,
 };
 
-/** used by FPostProcessSettings Anti-aliasing */
+/** Used by FPostProcessSettings Anti-aliasings */
 UENUM()
 enum EAntiAliasingMethod
 {
diff --git a/Engine/Source/Runtime/Engine/Classes/Engine/Texture.h b/Engine/Source/Runtime/Engine/Classes/Engine/Texture.h
--- a/Engine/Source/Runtime/Engine/Classes/Engine/Texture.h
+++ b/Engine/Source/Runtime/Engine/Classes/Engine/Texture.h
@@ -33,7 +33,7 @@ enum TextureFilter
 	TF_Nearest UMETA(DisplayName="Nearest"),
 	TF_Bilinear UMETA(DisplayName="Bi-linear"),
 	TF_Trilinear UMETA(DisplayName="Tri-linear"),
-	/** use setting from the Texture Group */
+	/** Use setting from the Texture Group. */
 	TF_Default UMETA(DisplayName="Default (from Texture Group)"),
 	TF_MAX,
 };
@@ -51,13 +51,13 @@ UENUM()
 enum ECompositeTextureMode
 {
 	CTM_Disabled UMETA(DisplayName="Disabled"),
-	// CompositingTexture needs to be a normal map with the same or larger size
+	/** CompositingTexture needs to be a normal map with the same or larger size. */
 	CTM_NormalRoughnessToRed UMETA(DisplayName="Add Normal Roughness To Red"),
-	// CompositingTexture needs to be a normal map with the same or larger size
+	/** CompositingTexture needs to be a normal map with the same or larger size. */
 	CTM_NormalRoughnessToGreen UMETA(DisplayName="Add Normal Roughness To Green"),
-	// CompositingTexture needs to be a normal map with the same or larger size
+	/** CompositingTexture needs to be a normal map with the same or larger size. */
 	CTM_NormalRoughnessToBlue UMETA(DisplayName="Add Normal Roughness To Blue"),
-	// CompositingTexture needs to be a normal map with the same or larger size
+	/** CompositingTexture needs to be a normal map with the same or larger size. */
 	CTM_NormalRoughnessToAlpha UMETA(DisplayName="Add Normal Roughness To Alpha"),
 	CTM_MAX,
 
@@ -78,11 +78,11 @@ enum ETextureMipCount
 UENUM()
 enum ETextureSourceArtType
 {
-	// FColor Data[SrcWidth * SrcHeight]
+	/** FColor Data[SrcWidth * SrcHeight]. */
 	TSAT_Uncompressed,
-	// PNG compresed version of FColor Data[SrcWidth * SrcHeight]
+	/** PNG compresed version of FColor Data[SrcWidth * SrcHeight]. */
 	TSAT_PNGCompressed,
-	// DDS file with header
+	/** DDS file with header. */
 	TSAT_DDSFile,
 	TSAT_MAX,
 };
diff --git a/Engine/Source/Runtime/Engine/Classes/Engine/TextureDefines.h b/Engine/Source/Runtime/Engine/Classes/Engine/TextureDefines.h
--- a/Engine/Source/Runtime/Engine/Classes/Engine/TextureDefines.h
+++ b/Engine/Source/Runtime/Engine/Classes/Engine/TextureDefines.h
@@ -36,20 +36,20 @@ enum TextureGroup
 	TEXTUREGROUP_Lightmap UMETA(DisplayName="Lightmap"),
 	TEXTUREGROUP_RenderTarget UMETA(DisplayName="RenderTarget"),
 	TEXTUREGROUP_MobileFlattened UMETA(DisplayName="MobileFlattened"),
-	//obsolete - kept for back-compat
-	TEXTUREGROUP_ProcBuilding_Face UMETA(DisplayName="ProcBuilding_Face"),
-	//obsolete - kept for back-compat
+	/** Obsolete - kept for backwards compatibility. */
+	TEXTUREGROUP_ProcBuilding_Face UMETA(DisplayName = "ProcBuilding_Face"),
+	/** Obsolete - kept for backwards compatibility. */
 	TEXTUREGROUP_ProcBuilding_LightMap UMETA(DisplayName="ProcBuilding_LightMap"),
 	TEXTUREGROUP_Shadowmap UMETA(DisplayName="Shadowmap"),
-	// no compression, no mips
+	/** No compression, no mips. */
 	TEXTUREGROUP_ColorLookupTable UMETA(DisplayName="ColorLookupTable"),
 	TEXTUREGROUP_Terrain_Heightmap UMETA(DisplayName="Terrain_Heightmap"),
 	TEXTUREGROUP_Terrain_Weightmap UMETA(DisplayName="Terrain_Weightmap"),
-	// using this TextureGroup triggers special mip map generation code only useful for the BokehDOF post process
+	/** Using this TextureGroup triggers special mip map generation code only useful for the BokehDOF post process. */
 	TEXTUREGROUP_Bokeh UMETA(DisplayName="Bokeh"),
-	// no compression, created on import of a .IES file
+	/** No compression, created on import of a .IES file. */
 	TEXTUREGROUP_IESLightProfile UMETA(DisplayName="IESLightProfile"),
-	// non-filtered, useful for 2D rendering
+	/** Non-filtered, useful for 2D rendering. */
 	TEXTUREGROUP_Pixels2D UMETA(DisplayName="2D Pixels (unfiltered)"),
 	TEXTUREGROUP_MAX,
 };
@@ -57,26 +57,36 @@ enum TextureGroup
 UENUM()
 enum TextureMipGenSettings
 {
-	// default for the "texture"
+	/** Default for the "texture". */
 	TMGS_FromTextureGroup UMETA(DisplayName="FromTextureGroup"),
-	// 2x2 average, default for the "texture group"
+	/** 2x2 average, default for the "texture group". */
 	TMGS_SimpleAverage UMETA(DisplayName="SimpleAverage"),
-	// 8x8 with sharpening: 0=no sharpening but better quality which is softer, 1..little, 5=medium, 10=extreme
+	/** 8x8 with sharpening: 0=no sharpening but better quality which is softer, 1=little, 5=medium, 10=extreme. */
 	TMGS_Sharpen0 UMETA(DisplayName="Sharpen0"),
-	TMGS_Sharpen1 UMETA(DisplayName="Sharpen1"),
-	TMGS_Sharpen2 UMETA(DisplayName="Sharpen2"),
-	TMGS_Sharpen3 UMETA(DisplayName="Sharpen3"),
-	TMGS_Sharpen4 UMETA(DisplayName="Sharpen4"),
-	TMGS_Sharpen5 UMETA(DisplayName="Sharpen5"),
-	TMGS_Sharpen6 UMETA(DisplayName="Sharpen6"),
-	TMGS_Sharpen7 UMETA(DisplayName="Sharpen7"),
-	TMGS_Sharpen8 UMETA(DisplayName="Sharpen8"),
-	TMGS_Sharpen9 UMETA(DisplayName="Sharpen9"),
-	TMGS_Sharpen10 UMETA(DisplayName="Sharpen10"),
+	/** 8x8 with sharpening: 0=no sharpening but better quality which is softer, 1=little, 5=medium, 10=extreme. */
+	TMGS_Sharpen1 UMETA(DisplayName = "Sharpen1"),
+	/** 8x8 with sharpening: 0=no sharpening but better quality which is softer, 1=little, 5=medium, 10=extreme. */
+	TMGS_Sharpen2 UMETA(DisplayName = "Sharpen2"),
+	/** 8x8 with sharpening: 0=no sharpening but better quality which is softer, 1=little, 5=medium, 10=extreme. */
+	TMGS_Sharpen3 UMETA(DisplayName = "Sharpen3"),
+	/** 8x8 with sharpening: 0=no sharpening but better quality which is softer, 1=little, 5=medium, 10=extreme. */
+	TMGS_Sharpen4 UMETA(DisplayName = "Sharpen4"),
+	/** 8x8 with sharpening: 0=no sharpening but better quality which is softer, 1=little, 5=medium, 10=extreme. */
+	TMGS_Sharpen5 UMETA(DisplayName = "Sharpen5"),
+	/** 8x8 with sharpening: 0=no sharpening but better quality which is softer, 1=little, 5=medium, 10=extreme. */
+	TMGS_Sharpen6 UMETA(DisplayName = "Sharpen6"),
+	/** 8x8 with sharpening: 0=no sharpening but better quality which is softer, 1=little, 5=medium, 10=extreme. */
+	TMGS_Sharpen7 UMETA(DisplayName = "Sharpen7"),
+	/** 8x8 with sharpening: 0=no sharpening but better quality which is softer, 1=little, 5=medium, 10=extreme. */
+	TMGS_Sharpen8 UMETA(DisplayName = "Sharpen8"),
+	/** 8x8 with sharpening: 0=no sharpening but better quality which is softer, 1=little, 5=medium, 10=extreme. */
+	TMGS_Sharpen9 UMETA(DisplayName = "Sharpen9"),
+	/** 8x8 with sharpening: 0=no sharpening but better quality which is softer, 1=little, 5=medium, 10=extreme. */
+	TMGS_Sharpen10 UMETA(DisplayName = "Sharpen10"),
 	TMGS_NoMipmaps UMETA(DisplayName="NoMipmaps"),
-	// Do not touch existing mip chain as it contains generated data
+	/** Do not touch existing mip chain as it contains generated data. */
 	TMGS_LeaveExistingMips UMETA(DisplayName="LeaveExistingMips"),
-	// blur further (useful for image based reflections)
+	/** Blur further (useful for image based reflections). */
 	TMGS_Blur1 UMETA(DisplayName="Blur1"),
 	TMGS_Blur2 UMETA(DisplayName="Blur2"),
 	TMGS_Blur3 UMETA(DisplayName="Blur3"),
@@ -89,19 +99,19 @@ enum TextureMipGenSettings
 	// in the texture compressor DDC key
 };
 
-// Options for texture padding mode
+/** Options for texture padding mode. */
 UENUM()
 namespace ETexturePowerOfTwoSetting
 {
 	enum Type
 	{
-		// Do not modify the texture dimensions
+		/** Do not modify the texture dimensions. */
 		None,
 
-		// Pad the texture to the nearest power of two size
+		/** Pad the texture to the nearest power of two size. */
 		PadToPowerOfTwo,
 
-		// Pad the texture to the nearest square power of two size
+		/** Pad the texture to the nearest square power of two size. */
 		PadToSquarePowerOfTwo
 
 		// Note: These are serialized as as raw values in the texture DDC key, so additional entries
diff --git a/Engine/Source/Runtime/Engine/Classes/Engine/UserDefinedStruct.h b/Engine/Source/Runtime/Engine/Classes/Engine/UserDefinedStruct.h
--- a/Engine/Source/Runtime/Engine/Classes/Engine/UserDefinedStruct.h
+++ b/Engine/Source/Runtime/Engine/Classes/Engine/UserDefinedStruct.h
@@ -7,13 +7,13 @@
 UENUM()
 enum EUserDefinedStructureStatus
 {
-	// Struct is in an unknown state
+	/** Struct is in an unknown state. */
 	UDSS_UpToDate,
-	// Struct has been modified but not recompiled
+	/** Struct has been modified but not recompiled. */
 	UDSS_Dirty,
-	// Struct tried but failed to be compiled
+	/** Struct tried but failed to be compiled. */
 	UDSS_Error,
-	// Struct is a duplicate, the original one was changed
+	/** Struct is a duplicate, the original one was changed. */
 	UDSS_Duplicate,
 
 	UDSS_MAX,
diff --git a/Engine/Source/Runtime/Engine/Classes/Engine/UserInterfaceSettings.h b/Engine/Source/Runtime/Engine/Classes/Engine/UserInterfaceSettings.h
--- a/Engine/Source/Runtime/Engine/Classes/Engine/UserInterfaceSettings.h
+++ b/Engine/Source/Runtime/Engine/Classes/Engine/UserInterfaceSettings.h
@@ -9,13 +9,13 @@
 UENUM()
 enum class ERenderFocusRule : uint8
 {
-	/** Focus Brush will always be rendered for widgets that have user focus */
+	/** Focus Brush will always be rendered for widgets that have user focus. */
 	Always,
-	/** Focus Brush will be rendered for widgets that have user focus not set based on pointer causes */
+	/** Focus Brush will be rendered for widgets that have user focus not set based on pointer causes. */
 	NonPointer,
 	/** Focus Brush will be rendered for widgets that have user focus only if the focus was set by navigation. */
 	NavigationOnly,
-	/** Focus Brush will not be rendered */
+	/** Focus Brush will not be rendered. */
 	Never,
 };
 
@@ -23,15 +23,15 @@ enum class ERenderFocusRule : uint8
 UENUM()
 enum class EUIScalingRule : uint8
 {
-	/** Evaluates the scale curve based on the shortest side of the viewport */
+	/** Evaluates the scale curve based on the shortest side of the viewport. */
 	ShortestSide,
-	/** Evaluates the scale curve based on the longest side of the viewport */
+	/** Evaluates the scale curve based on the longest side of the viewport. */
 	LongestSide,
-	/** Evaluates the scale curve based on the X axis of the viewport */
+	/** Evaluates the scale curve based on the X axis of the viewport. */
 	Horizontal,
-	/** Evaluates the scale curve based on the Y axis of the viewport */
+	/** Evaluates the scale curve based on the Y axis of the viewport. */
 	Vertical,
-	/** Custom - Allows custom rule interpretation */
+	/** Custom - Allows custom rule interpretation. */
 	Custom
 };
 
diff --git a/Engine/Source/Runtime/Engine/Classes/GameFramework/GameNetworkManager.h b/Engine/Source/Runtime/Engine/Classes/GameFramework/GameNetworkManager.h
--- a/Engine/Source/Runtime/Engine/Classes/GameFramework/GameNetworkManager.h
+++ b/Engine/Source/Runtime/Engine/Classes/GameFramework/GameNetworkManager.h
@@ -4,7 +4,7 @@
 #include "TimerManager.h"
 #include "GameNetworkManager.generated.h"
 
-/** Describes which standby detection event occured so the game can take appropriate action */
+/** Describes which standby detection event occured so the game can take appropriate action. */
 UENUM()
 enum EStandbyType
 {
diff --git a/Engine/Source/Runtime/Engine/Classes/Kismet/DataTableFunctionLibrary.h b/Engine/Source/Runtime/Engine/Classes/Kismet/DataTableFunctionLibrary.h
--- a/Engine/Source/Runtime/Engine/Classes/Kismet/DataTableFunctionLibrary.h
+++ b/Engine/Source/Runtime/Engine/Classes/Kismet/DataTableFunctionLibrary.h
@@ -7,15 +7,15 @@
 
 class UCurveTable;
 
-/** Enum used to indicate success or failure of EvaluateCurveTableRow */
+/** Enum used to indicate success or failure of EvaluateCurveTableRow. */
 UENUM()
 namespace EEvaluateCurveTableResult
 {
     enum Type
     {
-        /** Found the row successfully */
+        /** Found the row successfully. */
         RowFound,
-        /** Failed to find row */
+        /** Failed to find the row. */
         RowNotFound,
     };
 }
diff --git a/Engine/Source/Runtime/Engine/Classes/Kismet/KismetSystemLibrary.h b/Engine/Source/Runtime/Engine/Classes/Kismet/KismetSystemLibrary.h
--- a/Engine/Source/Runtime/Engine/Classes/Kismet/KismetSystemLibrary.h
+++ b/Engine/Source/Runtime/Engine/Classes/Kismet/KismetSystemLibrary.h
@@ -20,7 +20,7 @@ namespace EDrawDebugTrace
 	};
 }
 
-/** Enum used to indicate desired behavior for MoveComponentTo latent function */
+/** Enum used to indicate desired behavior for MoveComponentTo latent function. */
 UENUM()
 namespace EMoveComponentAction
 {
@@ -40,9 +40,9 @@ namespace EQuitPreference
 {
 	enum Type
 	{
-		/** Exit the game completely */
+		/** Exit the game completely. */
 		Quit,
-		/** Move the application to the background */
+		/** Move the application to the background. */
 		Background,
 	};
 }
diff --git a/Engine/Source/Runtime/Engine/Classes/Materials/Material.h b/Engine/Source/Runtime/Engine/Classes/Materials/Material.h
--- a/Engine/Source/Runtime/Engine/Classes/Materials/Material.h
+++ b/Engine/Source/Runtime/Engine/Classes/Materials/Material.h
@@ -44,32 +44,32 @@ struct FMaterialsWithDirtyUsageFlags
 };
 #endif
 
-/** Defines how the GBuffer channels are getting manipulated by a decal material pass. Actual index is used to control shader parameters so don't change order */
+/** Defines how the GBuffer channels are getting manipulated by a decal material pass. Actual index is used to control shader parameters so don't change order. */
 UENUM()
 enum EDecalBlendMode
 {
-	// Blend full material, updating the GBuffer, does not work for baked lighting
+	/** Blend full material, updating the GBuffer, does not work for baked lighting. */
 	DBM_Translucent UMETA(DisplayName="Translucent"),
-	// Modulate BaseColor, blend rest, updating the GBuffer, does not work for baked lighting
+	/** Modulate BaseColor, blend rest, updating the GBuffer, does not work for baked lighting. */
 	DBM_Stain UMETA(DisplayName="Stain"),
-	// Only blend normal, updating the GBuffer, does not work for baked lighting
+	/** Only blend normal, updating the GBuffer, does not work for baked lighting. */
 	DBM_Normal UMETA(DisplayName="Normal"),
-	// Additive emissive only
+	/** Additive emissive only. */
 	DBM_Emissive UMETA(DisplayName="Emissive"),
 
-	// Non metal, put into DBuffer to work for baked lighting as well (becomes DBM_TranslucentNormal if normal is not hooked up)
+	/** Non metal, put into DBuffer to work for baked lighting as well (becomes DBM_TranslucentNormal if normal is not hooked up). */
 	DBM_DBuffer_ColorNormalRoughness UMETA(DisplayName="DBuffer Translucent Color,Normal,Roughness"),
-	// Non metal, put into DBuffer to work for baked lighting as well
+	/** Non metal, put into DBuffer to work for baked lighting as well. */
 	DBM_DBuffer_Color UMETA(DisplayName="DBuffer Translucent Color"),
-	// Non metal, put into DBuffer to work for baked lighting as well (becomes DBM_DBuffer_Color if normal is not hooked up)
+	/** Non metal, put into DBuffer to work for baked lighting as well (becomes DBM_DBuffer_Color if normal is not hooked up). */
 	DBM_DBuffer_ColorNormal UMETA(DisplayName="DBuffer Translucent Color,Normal"),
-	// Non metal, put into DBuffer to work for baked lighting as well
+	/** Non metal, put into DBuffer to work for baked lighting as well. */
 	DBM_DBuffer_ColorRoughness UMETA(DisplayName="DBuffer Translucent Color,Roughness"),
-	// Non metal, put into DBuffer to work for baked lighting as well
+	/** Non metal, put into DBuffer to work for baked lighting as well. */
 	DBM_DBuffer_Normal UMETA(DisplayName="DBuffer Translucent Normal"),
-	// Non metal, put into DBuffer to work for baked lighting as well (becomes DBM_DBuffer_Roughness if normal is not hooked up)
+	/** Non metal, put into DBuffer to work for baked lighting as well (becomes DBM_DBuffer_Roughness if normal is not hooked up). */
 	DBM_DBuffer_NormalRoughness UMETA(DisplayName="DBuffer Translucent Normal,Roughness"),
-	// Non metal, put into DBuffer to work for baked lighting as well
+	/** Non metal, put into DBuffer to work for baked lighting as well. */
 	DBM_DBuffer_Roughness UMETA(DisplayName="DBuffer Translucent Roughness"),
 
 	DBM_MAX,
@@ -91,26 +91,26 @@ enum EMaterialDomain
 };
 
 
-/** Defines how the material reacts on DBuffer decals, later we can expose more variants between None and Default */
+/** Defines how the material reacts on DBuffer decals, later we can expose more variants between None and Default. */
 UENUM()
 enum EMaterialDecalResponse
 {
 	/** Do not receive decals (Later we still can read the DBuffer channels to customize the effect, this frees up some interpolators). */
 	MDR_None UMETA(DisplayName="None"),
 
-	// Receive Decals, applies all DBuffer channels, assumes the decal is non metal and mask the subsurface scattering
+	/** Receive Decals, applies all DBuffer channels, assumes the decal is non metal and mask the subsurface scattering. */
 	MDR_ColorNormalRoughness UMETA(DisplayName="Color Normal Roughness"),
-	// Receive Decals, applies color DBuffer channels, assumes the decal is non metal and mask the subsurface scattering
+	/** Receive Decals, applies color DBuffer channels, assumes the decal is non metal and mask the subsurface scattering. */
 	MDR_Color UMETA(DisplayName="Color"),
-	// Receive Decals, applies all DBuffer channels, assumes the decal is non metal and mask the subsurface scattering
+	/** Receive Decals, applies all DBuffer channels, assumes the decal is non metal and mask the subsurface scattering. */
 	MDR_ColorNormal UMETA(DisplayName="Color Normal"),
-	// Receive Decals, applies all DBuffer channels, assumes the decal is non metal and mask the subsurface scattering
+	/** Receive Decals, applies all DBuffer channels, assumes the decal is non metal and mask the subsurface scattering. */
 	MDR_ColorRoughness UMETA(DisplayName="Color Roughness"),
-	// Receive Decals, applies all DBuffer channels, assumes the decal is non metal and mask the subsurface scattering
+	/** Receive Decals, applies all DBuffer channels, assumes the decal is non metal and mask the subsurface scattering. */
 	MDR_Normal UMETA(DisplayName="Normal"),
-	// Receive Decals, applies all DBuffer channels, assumes the decal is non metal and mask the subsurface scattering
+	/** Receive Decals, applies all DBuffer channels, assumes the decal is non metal and mask the subsurface scattering. */
 	MDR_NormalRoughness UMETA(DisplayName="Normal Roughness"),
-	// Receive Decals, applies all DBuffer channels, assumes the decal is non metal and mask the subsurface scattering
+	/** Receive Decals, applies all DBuffer channels, assumes the decal is non metal and mask the subsurface scattering. */
 	MDR_Roughness UMETA(DisplayName="Roughness"),
 	MDR_MAX
 };
diff --git a/Engine/Source/Runtime/Engine/Classes/Materials/MaterialExpressionClamp.h b/Engine/Source/Runtime/Engine/Classes/Materials/MaterialExpressionClamp.h
--- a/Engine/Source/Runtime/Engine/Classes/Materials/MaterialExpressionClamp.h
+++ b/Engine/Source/Runtime/Engine/Classes/Materials/MaterialExpressionClamp.h
@@ -10,7 +10,6 @@ enum EClampMode
 {
 	CMODE_Clamp,
 	CMODE_ClampMin,
-	// Extra x to differentiate from the special value CMODE_Max, which is ignored by the property window
 	CMODE_ClampMax,
 };
 
diff --git a/Engine/Source/Runtime/Engine/Classes/Materials/MaterialExpressionDepthOfFieldFunction.h b/Engine/Source/Runtime/Engine/Classes/Materials/MaterialExpressionDepthOfFieldFunction.h
--- a/Engine/Source/Runtime/Engine/Classes/Materials/MaterialExpressionDepthOfFieldFunction.h
+++ b/Engine/Source/Runtime/Engine/Classes/Materials/MaterialExpressionDepthOfFieldFunction.h
@@ -8,11 +8,11 @@
 UENUM()
 enum EDepthOfFieldFunctionValue
 {
-	// 0:in Focus .. 1:Near or Far
+	/** 0:in Focus .. 1:Near or Far. */
 	TDOF_NearAndFarMask,
-	// 0:in Focus or Far .. 1:Near
+	/** 0:in Focus or Far .. 1:Near. */
 	TDOF_NearMask,
-	// 0:in Focus or Near .. 1:Far
+	/** 0:in Focus or Near .. 1:Far. */
 	TDOF_FarMask,
 	TDOF_MAX,
 };
