Compile fix - replaced a missing comma.

diff --git a/Engine/Source/Runtime/Engine/Classes/Components/SkinnedMeshComponent.h b/Engine/Source/Runtime/Engine/Classes/Components/SkinnedMeshComponent.h
--- a/Engine/Source/Runtime/Engine/Classes/Components/SkinnedMeshComponent.h
+++ b/Engine/Source/Runtime/Engine/Classes/Components/SkinnedMeshComponent.h
@@ -23,7 +23,7 @@ UENUM()
 enum EBoneVisibilityStatus
 {
 	/** Bone is hidden because it's parent is hidden. */
-	BVS_HiddenByParent
+	BVS_HiddenByParent,
 	/** Bone is visible. */
 	BVS_Visible,
 	/** Bone is hidden directly. */
