Comment fix.

diff --git a/Engine/Source/Runtime/Engine/Classes/Engine/BlendableInterface.h b/Engine/Source/Runtime/Engine/Classes/Engine/BlendableInterface.h
--- a/Engine/Source/Runtime/Engine/Classes/Engine/BlendableInterface.h
+++ b/Engine/Source/Runtime/Engine/Classes/Engine/BlendableInterface.h
@@ -17,10 +17,11 @@ enum EBlendableLocation
 	BL_BeforeTonemapping UMETA(DisplayName="Before Tonemapping"),
 	/** Input0:former pass color, Input1:SeparateTranslucency. */
 	BL_BeforeTranslucency UMETA(DisplayName="Before Translucency"),
-
-	/** Input0:former pass color, Input1:SeparateTranslucency, Input2: BloomOutput
-	/** vector parameters: Engine.FilmWhitePoint
-	/** scalar parameters: Engine.FilmSaturation, Engine.FilmContrast
+	/**
+	* Input0:former pass color, Input1:SeparateTranslucency, Input2: BloomOutput
+	* vector parameters: Engine.FilmWhitePoint
+	* scalar parameters: Engine.FilmSaturation, Engine.FilmContrast
+	*/
 	BL_ReplacingTonemapper UMETA(DisplayName="Replacing the Tonemapper"),
 //	BL_AfterOpaque,
 //	BL_AfterFog,
