Comment/doc changes.

diff --git a/Engine/Source/Runtime/Engine/Classes/Engine/EngineBaseTypes.h b/Engine/Source/Runtime/Engine/Classes/Engine/EngineBaseTypes.h
--- a/Engine/Source/Runtime/Engine/Classes/Engine/EngineBaseTypes.h
+++ b/Engine/Source/Runtime/Engine/Classes/Engine/EngineBaseTypes.h
@@ -27,44 +27,48 @@ enum EInputEvent
 /** Type of tick we wish to perform on the level */
 enum ELevelTick
 {
-	LEVELTICK_TimeOnly = 0,	// Update the level time only.
-	LEVELTICK_ViewportsOnly = 1,	// Update time and viewports.
-	LEVELTICK_All = 2,	// Update all.
-	LEVELTICK_PauseTick = 3,	// Delta time is zero, we are paused. Components don't tick.
+	/** Update the level time only. */
+	LEVELTICK_TimeOnly = 0,
+	/** Update time and viewports. */
+	LEVELTICK_ViewportsOnly = 1,
+	/** Update all. */
+	LEVELTICK_All = 2,
+	/** Delta time is zero, we are paused. Components don't tick. */
+	LEVELTICK_PauseTick = 3,
 };
 
-/** Determines which ticking group a tick function belongs to */
+/** Determines which ticking group a tick function belongs to. */
 UENUM()
 enum ETickingGroup
 {
-	/** Any item that needs to be executed before physics simulation starts */
+	/** Any item that needs to be executed before physics simulation starts. */
 	TG_PrePhysics,
 
-	/** Special tick group that starts physics simulation */							
+	/** Special tick group that starts physics simulation. */							
 	TG_StartPhysics,
 
-	/** Any item that can be run in parallel with our physics simulation work */
+	/** Any item that can be run in parallel with our physics simulation work. */
 	TG_DuringPhysics,
 
-	/** Special tick group that ends physics simulation */
+	/** Special tick group that ends physics simulation. */
 	TG_EndPhysics,
 
-	/** Any item that needs physics to be complete before being executed */
+	/** Any item that needs physics to be complete before being executed. */
 	TG_PreCloth,
 
-	/** Any item that needs to be updated after rigid body simulation is done, but before cloth is simulation is done */
+	/** Any item that needs to be updated after rigid body simulation is done, but before cloth is simulation is done. */
 	TG_StartCloth,
 
-	/** Any item that can be run during cloth simulation */	
+	/** Any item that can be run during cloth simulation. */	
 	TG_EndCloth,
 
-	/** Any item that needs rigid body and cloth sim to be complete before being executed */
+	/** Any item that needs rigid body and cloth simulation to be complete before being executed. */
 	TG_PostPhysics,
 
-	/** Any item that needs the update work to be done before being ticked */
+	/** Any item that needs the update work to be done before being ticked. */
 	TG_PostUpdateWork,
 
-	/** Special tick group that is not actually a tick group. After every tick group this is repeatedly re-run until there are no more newly spawned items to run */
+	/** Special tick group that is not actually a tick group. After every tick group this is repeatedly re-run until there are no more newly spawned items to run. */
 	TG_NewlySpawned,
 
 	TG_MAX,
@@ -78,15 +82,15 @@ struct FTickPrerequisite
 {
 	GENERATED_USTRUCT_BODY()
 
-	/** Tick functions live inside of UObjects, so we need a separate weak pointer to the UObject solely for the purpose of determining if PrerequisiteTickFunction is still valid **/
+	/** Tick functions live inside of UObjects, so we need a separate weak pointer to the UObject solely for the purpose of determining if PrerequisiteTickFunction is still valid. */
 	TWeakObjectPtr<class UObject> PrerequisiteObject;
 
 
 
-		/** Pointer to the actual tick function and must be completed prior to our tick running **/
+		/** Pointer to the actual tick function and must be completed prior to our tick running. */
 		struct FTickFunction*		PrerequisiteTickFunction;
 
-		/** Noop constructor **/
+		/** Noop constructor. */
 		FTickPrerequisite()
 		: PrerequisiteTickFunction(nullptr)
 		{
@@ -102,13 +106,13 @@ struct FTickPrerequisite
 		{
 			check(PrerequisiteTickFunction);
 		}
-		/** Equality operator, used to prevent duplicates and allow removal by value **/
+		/** Equality operator, used to prevent duplicates and allow removal by value. */
 		bool operator==(const FTickPrerequisite& Other) const
 		{
 			return PrerequisiteObject == Other.PrerequisiteObject &&
 				PrerequisiteTickFunction == Other.PrerequisiteTickFunction;
 		}
-		/** Return the tick function, if it is still valid. Can be null if the tick function was null or the containing UObject has been garbage collected. **/
+		/** Return the tick function, if it is still valid. Can be null if the tick function was null or the containing UObject has been garbage collected. */
 		struct FTickFunction* Get()
 		{
 			if (PrerequisiteObject.IsValid(true))
@@ -528,11 +532,11 @@ namespace ETravelFailure
 UENUM()
 enum ETravelType
 {
-	// Absolute URL.
+	/** Absolute URL. */
 	TRAVEL_Absolute,
-	// Partial (carry name, reset server).
+	/** Partial (carry name, reset server). */
 	TRAVEL_Partial,
-	// Relative URL.
+	/** Relative URL. */
 	TRAVEL_Relative,
 	TRAVEL_MAX,
 };
@@ -543,13 +547,13 @@ namespace EDemoPlayFailure
 {
 	enum Type
 	{
-		/** A Generic failure */
+		/** A Generic failure. */
 		Generic,
-		/** Demo was not found */
+		/** Demo was not found. */
 		DemoNotFound,
-		/** Demo is corrupt */
+		/** Demo is corrupt. */
 		Corrupt,
-		/** Invalid version */
+		/** Invalid version. */
 		InvalidVersion,
 	};
 }
@@ -670,31 +674,33 @@ enum ENetMode
 UENUM()
 enum EViewModeIndex 
 {
-	// Wireframe w/ brushes
+	/** Wireframe w/ brushes. */
 	VMI_BrushWireframe = 0,
-	// Wireframe w/ BSP
+	/** Wireframe w/ BSP. */
 	VMI_Wireframe = 1,
-	// Unlit
+	/** Unlit. */
 	VMI_Unlit = 2,
-	// Lit
+	/** Lit. */
 	VMI_Lit = 3,
 	VMI_Lit_DetailLighting = 4,
-	// Lit wo/ materials
+	/** Lit wo/ materials. */
 	VMI_LightingOnly = 5,
-	// Colored according to light count.
+	/** Colored according to light count. */
 	VMI_LightComplexity = 6,
-	// Colored according to shader complexity.
+	/** Colored according to shader complexity. */
 	VMI_ShaderComplexity = 8,
-	// Colored according to world-space LightMap texture density.
+	/** Colored according to world-space LightMap texture density. */
 	VMI_LightmapDensity = 9,
-	// Colored according to light count - showing lightmap texel density on texture mapped objects.
+	/** Colored according to light count - showing lightmap texel density on texture mapped objects. */
 	VMI_LitLightmapDensity = 10,
 	VMI_ReflectionOverride = 11,
 	VMI_VisualizeBuffer = 12,
 //	VMI_VoxelLighting = 13,
-	// Colored according to stationary light overlap.
+
+	/** Colored according to stationary light overlap. */
 	VMI_StationaryLightOverlap = 14,
 //	VMI_VertexDensity = 15,
+
 	VMI_CollisionPawn = 15, 
 	VMI_CollisionVisibility = 16, 
 	VMI_Max UMETA(Hidden),
