Fix bug where welding to a skeletal mesh meant the response, collision, and notification were overridden by skeletal mesh.

#UE-16864

diff --git a/Engine/Source/Runtime/Engine/Private/PhysicsEngine/BodyInstance.cpp b/Engine/Source/Runtime/Engine/Private/PhysicsEngine/BodyInstance.cpp
--- a/Engine/Source/Runtime/Engine/Private/PhysicsEngine/BodyInstance.cpp
+++ b/Engine/Source/Runtime/Engine/Private/PhysicsEngine/BodyInstance.cpp
@@ -682,11 +682,12 @@ void FBodyInstance::UpdatePhysicsShapeFilterData(uint32 SkelMeshCompID, bool bUs
 		{
 			PxShape* PShape = AllShapes[ShapeIdx];
 			FBodyInstance* BI = FPhysxUserData::Get<FBodyInstance>(PShape->userData);
+			const bool bWelded = BI != nullptr;
 			BI = BI ? BI : this;
 
-			const TEnumAsByte<ECollisionEnabled::Type> UseCollisionEnabled = CollisionEnabledOverride ? *CollisionEnabledOverride : (TEnumAsByte<ECollisionEnabled::Type>)BI->GetCollisionEnabled();
-			const FCollisionResponseContainer& UseResponse = ResponseOverride ? *ResponseOverride : BI->CollisionResponses.GetResponseContainer();
-			bool bUseNotify = bNotifyOverride ? *bNotifyOverride : BI->bNotifyRigidBodyCollision;
+			const TEnumAsByte<ECollisionEnabled::Type> UseCollisionEnabled = CollisionEnabledOverride && !bWelded ? *CollisionEnabledOverride : (TEnumAsByte<ECollisionEnabled::Type>)BI->GetCollisionEnabled();
+			const FCollisionResponseContainer& UseResponse = ResponseOverride && !bWelded ? *ResponseOverride : BI->CollisionResponses.GetResponseContainer();
+			bool bUseNotify = bNotifyOverride && !bWelded ? *bNotifyOverride : BI->bNotifyRigidBodyCollision;
 
 
 			UPrimitiveComponent* OwnerPrimitiveComponent = BI->OwnerComponent.Get();
