#ue4 - Fix blocking hit notifications being triggered in invalid cases (pending kill or null actors/components). Fix SceneComponent trying to move when pending kill (PrimitiveComponent already aborts in this case).

diff --git a/Engine/Source/Runtime/Engine/Private/Actor.cpp b/Engine/Source/Runtime/Engine/Private/Actor.cpp
--- a/Engine/Source/Runtime/Engine/Private/Actor.cpp
+++ b/Engine/Source/Runtime/Engine/Private/Actor.cpp
@@ -1856,24 +1856,27 @@ void AActor::InternalDispatchBlockingHit(UPrimitiveComponent* MyComp, UPrimitive
 {
 	check(MyComp);
 
-	AActor* OtherActor = (OtherComp != NULL) ? OtherComp->GetOwner() : NULL;
-
-	// Call virtual
-	if(IsActorValidToNotify(this))
+	if (OtherComp != nullptr)
 	{
-		NotifyHit(MyComp, OtherActor, OtherComp, bSelfMoved, Hit.ImpactPoint, Hit.ImpactNormal, FVector(0,0,0), Hit);
-	}
+		AActor* OtherActor = OtherComp->GetOwner();
 
-	// If we are still ok, call delegate on actor
-	if(IsActorValidToNotify(this))
-	{
-		OnActorHit.Broadcast(this, OtherActor, FVector(0,0,0), Hit);
-	}
+		// Call virtual
+		if(IsActorValidToNotify(this))
+		{
+			NotifyHit(MyComp, OtherActor, OtherComp, bSelfMoved, Hit.ImpactPoint, Hit.ImpactNormal, FVector(0,0,0), Hit);
+		}
 
-	// If component is still alive, call delegate on component
-	if(!MyComp->IsPendingKill())
-	{
-		MyComp->OnComponentHit.Broadcast(OtherActor, OtherComp, FVector(0,0,0), Hit);
+		// If we are still ok, call delegate on actor
+		if(IsActorValidToNotify(this))
+		{
+			OnActorHit.Broadcast(this, OtherActor, FVector(0,0,0), Hit);
+		}
+
+		// If component is still alive, call delegate on component
+		if(!MyComp->IsPendingKill())
+		{
+			MyComp->OnComponentHit.Broadcast(OtherActor, OtherComp, FVector(0,0,0), Hit);
+		}
 	}
 }
 
diff --git a/Engine/Source/Runtime/Engine/Private/Components/PrimitiveComponent.cpp b/Engine/Source/Runtime/Engine/Private/Components/PrimitiveComponent.cpp
--- a/Engine/Source/Runtime/Engine/Private/Components/PrimitiveComponent.cpp
+++ b/Engine/Source/Runtime/Engine/Private/Components/PrimitiveComponent.cpp
@@ -1360,18 +1360,8 @@ bool UPrimitiveComponent::MoveComponentImpl( const FVector& Delta, const FQuat&
 	CLOCK_CYCLES(MoveCompTakingLongTime);
 #endif
 
-	if ( IsPendingKill() )
-	{
-		//UE_LOG(LogPrimitiveComponent, Log, TEXT("%s deleted move physics %d"),*Actor->GetName(),Actor->Physics);
-		if (OutHit)
-		{
-			*OutHit = FHitResult();
-		}
-		return false;
-	}
-
 	// static things can move before they are registered (e.g. immediately after streaming), but not after.
-	if (CheckStaticMobilityAndWarn(PrimitiveComponentStatics::MobilityWarnText))
+	if (IsPendingKill() || CheckStaticMobilityAndWarn(PrimitiveComponentStatics::MobilityWarnText))
 	{
 		if (OutHit)
 		{
@@ -1635,8 +1625,8 @@ bool UPrimitiveComponent::MoveComponentImpl( const FVector& Delta, const FQuat&
 		}
 	}
 
-	// Handle blocking hit notifications.
-	if (BlockingHit.bBlockingHit)
+	// Handle blocking hit notifications. Avoid if pending kill (which could happen after overlaps).
+	if (BlockingHit.bBlockingHit && !IsPendingKill())
 	{
 		if (IsDeferringMovementUpdates())
 		{
diff --git a/Engine/Source/Runtime/Engine/Private/Components/SceneComponent.cpp b/Engine/Source/Runtime/Engine/Private/Components/SceneComponent.cpp
--- a/Engine/Source/Runtime/Engine/Private/Components/SceneComponent.cpp
+++ b/Engine/Source/Runtime/Engine/Private/Components/SceneComponent.cpp
@@ -515,6 +515,12 @@ void USceneComponent::EndScopedMovementUpdate(class FScopedMovementUpdate& Compl
 					UPrimitiveComponent* PrimitiveThis = CastChecked<UPrimitiveComponent>(this);
 					for (const FHitResult& Hit : CurrentScopedUpdate->BlockingHits)
 					{
+						// Overlaps may have caused us to be destroyed, as could other queued blocking hits.
+						if (PrimitiveThis->IsPendingKill())
+						{
+							break;
+						}
+
 						// Collision response may change (due to overlaps or multiple blocking hits), make sure it's still considered blocking.
 						if (PrimitiveThis->GetCollisionResponseToComponent(Hit.GetComponent()) == ECR_Block)
 						{
@@ -1908,7 +1914,7 @@ bool USceneComponent::MoveComponentImpl( const FVector& Delta, const FQuat& NewR
 	SCOPE_CYCLE_COUNTER(STAT_MoveComponentSceneComponentTime);
 
 	// static things can move before they are registered (e.g. immediately after streaming), but not after.
-	if (CheckStaticMobilityAndWarn(SceneComponentStatics::MobilityWarnText))
+	if (IsPendingKill() || CheckStaticMobilityAndWarn(SceneComponentStatics::MobilityWarnText))
 	{
 		if (OutHit)
 		{
