Exposing more package streaming and garbage collection settings to project settings in the editor.

Also, added new console variables for the existing settings:

New Package Streaming console variables:
- WarnIfTimeLimitExceeded
- TimeLimitExceededMultiplier
- TimeLimitExceededMinTime
- MinBulkDataSizeForAsyncLoading
- AsyncIOBandwidthLimit

New GC console variables:
- FlushStreamingOnGC
- NumRetriesBeforeForcingGC

diff --git a/Engine/Config/BaseEngine.ini b/Engine/Config/BaseEngine.ini
--- a/Engine/Config/BaseEngine.ini
+++ b/Engine/Config/BaseEngine.ini
@@ -117,10 +117,6 @@ LightingOnlyBrightness=(R=0.3,G=0.3,B=0.3,A=1.0)
 +StationaryLightOverlapColors=(R=1.0,G=0.9,B=0.9,A=1.0)
 MaxPixelShaderAdditiveComplexityCount=2000
 MaxES2PixelShaderAdditiveComplexityCount=600
-TimeBetweenPurgingPendingKillObjects=60
-bUseBackgroundLevelStreaming=True
-LevelStreamingActorsUpdateTimeLimit = 5.0
-LevelStreamingComponentsRegistrationGranularity = 10
 bSubtitlesEnabled=True
 bSubtitlesForcedOff=false
 DefaultSoundName=/Engine/EngineSounds/WhiteNoise.WhiteNoise
@@ -1330,22 +1326,32 @@ bFullyCompressStartupPackages=false
 +Package=/Engine/EngineFonts/RobotoDistanceField
 
 [Core.System]
-MaxObjectsNotConsideredByGC=0
-SizeOfPermanentObjectPool=0
-AsyncIOBandwidthLimit=0
 +Paths=../../../Engine/Content
 +Paths=%GAMEDIR%Content
 CutdownPaths=%GAMEDIR%CutdownPackages
 ZeroEngineVersionWarning=True
 UseStrictEngineVersioning=True
-FlushStreamingOnGC=False
+UseVTableConstructors=True
+
+[/Script/Engine.StreamingSettings]
+AsyncIOBandwidthLimit=0
 MinBulkDataSizeForAsyncLoading=131072
 AsyncLoadingThreadEnabled=False
-WarnIfTimeLimitExceeded=True
+WarnIfTimeLimitExceeded=False
 TimeLimitExceededMultiplier=1.5
 TimeLimitExceededMinTime=0.005
+bUseBackgroundLevelStreaming=True
+PriorityAsyncLoadingExtraTime=20.0
+LevelStreamingActorsUpdateTimeLimit = 5.0
+LevelStreamingComponentsRegistrationGranularity = 10
+
+[/Script/Engine.GarbageCollectionSettings]
+MaxObjectsNotConsideredByGC=0
+SizeOfPermanentObjectPool=0
+FlushStreamingOnGC=0
 NumRetriesBeforeForcingGC=0
-UseVTableConstructors=True
+AllowParallelGC=True
+TimeBetweenPurgingPendingKillObjects=60
 
 [Internationalization]
 +LocalizationPaths=../../../Engine/Content/Localization/Engine
diff --git a/Engine/Source/Editor/UnrealEd/Private/EditorEngine.cpp b/Engine/Source/Editor/UnrealEd/Private/EditorEngine.cpp
--- a/Engine/Source/Editor/UnrealEd/Private/EditorEngine.cpp
+++ b/Engine/Source/Editor/UnrealEd/Private/EditorEngine.cpp
@@ -113,6 +113,7 @@
 #include "Engine/SimpleConstructionScript.h"
 
 #include "PhysicsPublic.h"
+#include "Engine/CoreSettings.h"
 
 DEFINE_LOG_CATEGORY_STATIC(LogEditor, Log, All);
 
@@ -981,7 +982,7 @@ void UEditorEngine::Tick( float DeltaSeconds, bool bIdleMode )
 	// Update subsystems.
 	{
 		// This assumes that UObject::StaticTick only calls ProcessAsyncLoading.	
-		StaticTick(DeltaSeconds, bAsyncLoadingUseFullTimeLimit, AsyncLoadingTimeLimit / 1000.f);
+		StaticTick(DeltaSeconds, GetDefault<UStreamingSettings>()->bAsyncLoadingUseFullTimeLimit, GetDefault<UStreamingSettings>()->AsyncLoadingTimeLimit / 1000.f);
 	}
 
 	// Look for realtime flags.
diff --git a/Engine/Source/Runtime/Core/Private/Serialization/AsyncIOSystemBase.cpp b/Engine/Source/Runtime/Core/Private/Serialization/AsyncIOSystemBase.cpp
--- a/Engine/Source/Runtime/Core/Private/Serialization/AsyncIOSystemBase.cpp
+++ b/Engine/Source/Runtime/Core/Private/Serialization/AsyncIOSystemBase.cpp
@@ -31,6 +31,12 @@ DECLARE_FLOAT_COUNTER_STAT( TEXT( "Bandwidth (MByte/ sec)" ), STAT_AsyncIO_Bandw
 
 // Constrain bandwidth if wanted. Value is in MByte/ sec.
 float GAsyncIOBandwidthLimit = 0.0f;
+static FAutoConsoleVariableRef CVarAsyncIOBandwidthLimit(
+	TEXT("AsyncIOBandwidthLimit"),
+	GAsyncIOBandwidthLimit,
+	TEXT("Constrain bandwidth if wanted. Value is in MByte/ sec."),
+	ECVF_Default
+	);
 
 CORE_API bool GbLogAsyncLoading = false;
 
diff --git a/Engine/Source/Runtime/CoreUObject/Private/Serialization/AsyncLoading.cpp b/Engine/Source/Runtime/CoreUObject/Private/Serialization/AsyncLoading.cpp
--- a/Engine/Source/Runtime/CoreUObject/Private/Serialization/AsyncLoading.cpp
+++ b/Engine/Source/Runtime/CoreUObject/Private/Serialization/AsyncLoading.cpp
@@ -213,6 +213,30 @@ public:
 static TAutoPtr<FAsyncLoadingExec> GAsyncLoadingExec;
 #endif
 
+static int32 GWarnIfTimeLimitExceeded = 0;
+static FAutoConsoleVariableRef CVarWarnIfTimeLimitExceeded(
+	TEXT("WarnIfTimeLimitExceeded"),
+	GWarnIfTimeLimitExceeded,
+	TEXT("Enables log warning if time limit for time-sliced package streaming has been exceeded."),
+	ECVF_Default
+	);
+
+static float GTimeLimitExceededMultiplier = 1.5f;
+static FAutoConsoleVariableRef CVarTimeLimitExceededMultiplier(
+	TEXT("TimeLimitExceededMultiplier"),
+	GTimeLimitExceededMultiplier,
+	TEXT("Multiplier for time limit exceeded warning time threshold."),
+	ECVF_Default
+	);
+
+static float GTimeLimitExceededMinTime = 0.005f;
+static FAutoConsoleVariableRef CVarTimeLimitExceededMinTime(
+	TEXT("TimeLimitExceededMinTime"),
+	GTimeLimitExceededMinTime,
+	TEXT("Multiplier for time limit exceeded warning time threshold."),
+	ECVF_Default
+	);
+
 static FORCEINLINE bool IsTimeLimitExceeded(double InTickStartTime, bool bUseTimeLimit, float InTimeLimit, const TCHAR* InLastTypeOfWorkPerformed = nullptr, UObject* InLastObjectWorkWasPerformedOn = nullptr)
 {
 	bool bTimeLimitExceeded = false;
@@ -221,29 +245,10 @@ static FORCEINLINE bool IsTimeLimitExceeded(double InTickStartTime, bool bUseTim
 		double CurrentTime = FPlatformTime::Seconds();
 		bTimeLimitExceeded = CurrentTime - InTickStartTime > InTimeLimit;
 
-		// One time init for ini override if we should log a warning when time limit is exceeded
-		static struct FWarnIfTimeLimitExceeded
-		{
-			bool WarnignEnabled;
-			double TimeMultiplier;
-			double MinTime;
-			FWarnIfTimeLimitExceeded()
-				: WarnignEnabled(true)
-				, TimeMultiplier(1.0)
-				, MinTime(0.005)
-			{
-				check(GConfig);
-				GConfig->GetBool(TEXT("Core.System"), TEXT("WarnIfTimeLimitExceeded"), WarnignEnabled, GEngineIni);
-				GConfig->GetDouble(TEXT("Core.System"), TEXT("TimeLimitExceededMultiplier"), TimeMultiplier, GEngineIni);
-				GConfig->GetDouble(TEXT("Core.System"), TEXT("TimeLimitExceededMinTime"), MinTime, GEngineIni);
-			}
-		} WarnIfTimeLimitExceeded;
-
 		// Log single operations that take longer than time limit (but only in cooked builds)
-		if (FPlatformProperties::RequiresCookedData() && 
-			WarnIfTimeLimitExceeded.WarnignEnabled && 
-			(CurrentTime - InTickStartTime) > WarnIfTimeLimitExceeded.MinTime &&
-			(CurrentTime - InTickStartTime) > (WarnIfTimeLimitExceeded.TimeMultiplier * InTimeLimit))
+		if (GWarnIfTimeLimitExceeded && 
+			(CurrentTime - InTickStartTime) > GTimeLimitExceededMinTime &&
+			(CurrentTime - InTickStartTime) > (GTimeLimitExceededMultiplier * InTimeLimit))
 		{			
 			UE_LOG(LogStreaming, Warning, TEXT("IsTimeLimitExceeded: %s %s took (less than) %5.2f ms"),
 				InLastTypeOfWorkPerformed ? InLastTypeOfWorkPerformed : TEXT("unknown"),
diff --git a/Engine/Source/Runtime/CoreUObject/Private/Serialization/AsyncLoadingThread.h b/Engine/Source/Runtime/CoreUObject/Private/Serialization/AsyncLoadingThread.h
--- a/Engine/Source/Runtime/CoreUObject/Private/Serialization/AsyncLoadingThread.h
+++ b/Engine/Source/Runtime/CoreUObject/Private/Serialization/AsyncLoadingThread.h
@@ -109,7 +109,7 @@ public:
 				{
 					check(GConfig);
 					bool bConfigValue = true;
-					GConfig->GetBool(TEXT("Core.System"), TEXT("AsyncLoadingThreadEnabled"), bConfigValue, GEngineIni);
+					GConfig->GetBool(TEXT("/Script/Engine.StreamingSettings"), TEXT("AsyncLoadingThreadEnabled"), bConfigValue, GEngineIni);
 					bool bCommandLineNoAsyncThread = false;
 					Value = bConfigValue && FApp::ShouldUseThreadingForPerformance() && !FParse::Param(FCommandLine::Get(), TEXT("NoAsyncLoadingThread"));
 				}
diff --git a/Engine/Source/Runtime/CoreUObject/Private/Serialization/BulkData.cpp b/Engine/Source/Runtime/CoreUObject/Private/Serialization/BulkData.cpp
--- a/Engine/Source/Runtime/CoreUObject/Private/Serialization/BulkData.cpp
+++ b/Engine/Source/Runtime/CoreUObject/Private/Serialization/BulkData.cpp
@@ -690,25 +690,22 @@ void FUntypedBulkData::StartSerializingBulkData(FArchive& Ar, UObject* Owner, in
 	}
 }
 
+int32 GMinimumBulkDataSizeForAsyncLoading = 131072;
+static FAutoConsoleVariableRef CVarMinimumBulkDataSizeForAsyncLoading(
+	TEXT("MinBulkDataSizeForAsyncLoading"),
+	GMinimumBulkDataSizeForAsyncLoading,
+	TEXT("Minimum time the time limit exceeded warning will be triggered by."),
+	ECVF_Default
+	);
+
 bool FUntypedBulkData::ShouldStreamBulkData()
 {
-	// Minimum bulk data size to start async loading
-	static struct FMinBulkDataSizeForAsyncLoadingSetting
-	{
-		int32 Value;
-		FMinBulkDataSizeForAsyncLoadingSetting()
-			: Value(131072) // 128K by default
-		{
-			GConfig->GetInt(TEXT("Core.System"), TEXT("MinBulkDataSizeForAsyncLoading"), Value, GEngineIni);
-		}
-	} MinBulkDataSizeForAsyncLoading;
-
 	const bool bForceStream = !!(BulkDataFlags & BULKDATA_ForceStreamPayload);
 
 	return (FPlatformProperties::RequiresCookedData() && !Filename.IsEmpty() &&
 		FPlatformProcess::SupportsMultithreading() && IsInGameThread() &&
-		(bForceStream || GetBulkDataSize() > MinBulkDataSizeForAsyncLoading.Value) &&
-		MinBulkDataSizeForAsyncLoading.Value >= 0);
+		(bForceStream || GetBulkDataSize() > GMinimumBulkDataSizeForAsyncLoading) &&
+		GMinimumBulkDataSizeForAsyncLoading >= 0);
 }
 
 /**
diff --git a/Engine/Source/Runtime/CoreUObject/Private/UObject/GarbageCollection.cpp b/Engine/Source/Runtime/CoreUObject/Private/UObject/GarbageCollection.cpp
--- a/Engine/Source/Runtime/CoreUObject/Private/UObject/GarbageCollection.cpp
+++ b/Engine/Source/Runtime/CoreUObject/Private/UObject/GarbageCollection.cpp
@@ -1187,12 +1187,35 @@ typedef void (*EditorPostReachabilityAnalysisCallbackType)();
 COREUOBJECT_API EditorPostReachabilityAnalysisCallbackType EditorPostReachabilityAnalysisCallback = NULL;
 
 // Allow parralel GC to be overriden to single threaded via console command.
-static const auto CVarAllowParallelGC = 
-IConsoleManager::Get().RegisterConsoleVariable(TEXT("AllowParallelGC"), (!PLATFORM_MAC || !WITH_EDITORONLY_DATA) ? 1 : 0, TEXT("Used to control parallel GC."))->AsVariableInt();
+static int32 GAllowParallelGC = (!PLATFORM_MAC || !WITH_EDITORONLY_DATA) ? 1 : 0;
+static FAutoConsoleVariableRef CVarAllowParallelGC(
+	TEXT("AllowParallelGC"),
+	GAllowParallelGC,
+	TEXT("sed to control parallel GC."),
+	ECVF_Default
+	);
 
 // This counts how many times GC was skipped
 static int32 GNumAttemptsSinceLastGC = 0;
 
+// Number of times GC can be skipped.
+static int32 GNumRetriesBeforeForcingGC = 0;
+static FAutoConsoleVariableRef CVarNumRetriesBeforeForcingGC(
+	TEXT("NumRetriesBeforeForcingGC"),
+	GNumRetriesBeforeForcingGC,
+	TEXT("Maximum number of times GC can be skipped if worker threads are currently modifying UObject state."),
+	ECVF_Default
+	);
+
+// Force flush streaming on GC console variable
+static int32 GFlushStreamingOnGC = 0;
+static FAutoConsoleVariableRef CVarFlushStreamingOnGC(
+	TEXT("FlushStreamingOnGC"),
+	GFlushStreamingOnGC,
+	TEXT("If enabled, streaming will be flushed each time garbage collection is triggered."),
+	ECVF_Default
+	);
+
 /** 
  * Deletes all unreferenced objects, keeping objects that have any of the passed in KeepFlags set
  *
@@ -1207,26 +1230,11 @@ void CollectGarbageInternal(EObjectFlags KeepFlags, bool bPerformFullPurge)
 	// Reset GC skip counter
 	GNumAttemptsSinceLastGC = 0;
 
-	// Helper class to register FlushAsyncLoadingCallback on first GC run.
-	struct FAddFlushAsyncLoadingCallback
+	// Flush streaming before GC if requested
+	if (GFlushStreamingOnGC)
 	{
-		FAddFlushAsyncLoadingCallback()
-		{
-			bool bFlushStreaming = false;
-			GConfig->GetBool(TEXT("Core.System"), TEXT("FlushStreamingOnGC"), bFlushStreaming, GEngineIni);
-			if (bFlushStreaming)
-			{
-				FCoreUObjectDelegates::PreGarbageCollect.AddStatic(FlushAsyncLoadingCallback);
-			}
-		}
-		/** Wrapper function to handle default parameter when used as function pointer */
-		static void FlushAsyncLoadingCallback()
-		{
-			FlushAsyncLoading();
-		}
-	};
-	// Add FlushAsyncLoadingCallback the first time CollectGarbage is called if requested by ini settings
-	static FAddFlushAsyncLoadingCallback MaybeAddFlushAsyncLoadingCallback;
+		FlushAsyncLoading();
+	}
 
 	// Route callbacks so we can ensure that we are e.g. not in the middle of loading something by flushing
 	// the async loading, etc...
@@ -1292,7 +1300,7 @@ void CollectGarbageInternal(EObjectFlags KeepFlags, bool bPerformFullPurge)
 	// Temporarily forcing single-threaded GC in the editor until Modify() can be safely removed from HandleObjectReference.
 	const bool bForceSingleThreadedGC = !FApp::ShouldUseThreadingForPerformance() || !FPlatformProcess::SupportsMultithreading() ||
 #if PLATFORM_SUPPORTS_MULTITHREADED_GC
-		( FPlatformMisc::NumberOfCores() < 2 || CVarAllowParallelGC->GetValueOnGameThread() == 0 || PERF_DETAILED_PER_CLASS_GC_STATS );
+		( FPlatformMisc::NumberOfCores() < 2 || GAllowParallelGC == 0 || PERF_DETAILED_PER_CLASS_GC_STATS );
 #else	//PLATFORM_SUPPORTS_MULTITHREADED_GC
 		true;
 #endif	//PLATFORM_SUPPORTS_MULTITHREADED_GC
@@ -1360,18 +1368,7 @@ bool TryCollectGarbage(EObjectFlags KeepFlags, bool bPerformFullPurge)
 	bool bCanRunGC = GGarbageCollectionGuardCritical.TryGCLock();
 	if (!bCanRunGC)
 	{
-		// One time init for ini override if we should log a warning when time limit is exceeded
-		static struct FForceGCSettings
-		{
-			int32 NumRetriesBeforeForcingGC;
-			FForceGCSettings()
-				: NumRetriesBeforeForcingGC(0) // 0 means 'never force GC'
-			{
-				check(GConfig);
-				GConfig->GetInt(TEXT("Core.System"), TEXT("NumRetriesBeforeForcingGC"), NumRetriesBeforeForcingGC, GEngineIni);
-			}
-		} ForceGCSettings;
-		if (ForceGCSettings.NumRetriesBeforeForcingGC > 0 && GNumAttemptsSinceLastGC > ForceGCSettings.NumRetriesBeforeForcingGC)
+		if (GNumRetriesBeforeForcingGC > 0 && GNumAttemptsSinceLastGC > GNumRetriesBeforeForcingGC)
 		{
 			// Force GC and block main thread
 			bCanRunGC = true;
diff --git a/Engine/Source/Runtime/Engine/Classes/Engine/CoreSettings.h b/Engine/Source/Runtime/Engine/Classes/Engine/CoreSettings.h
new file mode 100644
--- /dev/null
+++ b/Engine/Source/Runtime/Engine/Classes/Engine/CoreSettings.h
@@ -0,0 +1,145 @@
+// Copyright 1998-2015 Epic Games, Inc. All Rights Reserved.
+
+#pragma once
+
+#include "Curves/CurveFloat.h"
+#include "UserInterfaceSettings.h"
+
+#include "CoreSettings.generated.h"
+
+
+/**
+ * Rendering settings.
+ */
+UCLASS(config=Engine, defaultconfig, meta=(DisplayName="Streaming"))
+class ENGINE_API UStreamingSettings : public UDeveloperSettings
+{
+	GENERATED_BODY()
+
+public:
+
+	UStreamingSettings();
+
+	UPROPERTY(config, EditAnywhere, Category = PackageStreaming, meta = (
+		DisplayName = "Async Loading Thread Enabled",
+		ToolTip = "Enables separate thread for package streaming. Requires restart to take effect."))
+	uint32 AsyncLoadingThreadEnabled : 1;
+
+	UPROPERTY(config, EditAnywhere, Category = PackageStreaming, AdvancedDisplay, meta = (
+		ConsoleVariable = "WarnIfTimeLimitExceeded", DisplayName = "Warn If Time Limit Has Been Exceeded",
+		ToolTip = "Enables log warning if time limit for time-sliced package streaming has been exceeded."))
+	uint32 WarnIfTimeLimitExceeded : 1;
+
+	UPROPERTY(config, EditAnywhere, Category = PackageStreaming, AdvancedDisplay, meta = (
+		ConsoleVariable = "TimeLimitExceededMultiplier", DisplayName = "Time Limit Exceeded Warning Multiplier",
+		ToolTip = "Multiplier for time limit exceeded warning time threshold."))
+	float TimeLimitExceededMultiplier;
+
+	UPROPERTY(config, EditAnywhere, Category = PackageStreaming, AdvancedDisplay, meta = (
+		ConsoleVariable = "TimeLimitExceededMinTime", DisplayName = "Minimum Time Limit For Time Limit Exceeded Warning",
+		ToolTip = "Minimum time the time limit exceeded warning will be triggered by."))
+	float TimeLimitExceededMinTime;
+
+	UPROPERTY(config, EditAnywhere, Category = PackageStreaming, AdvancedDisplay, meta = (
+		ConsoleVariable = "MinBulkDataSizeForAsyncLoading", DisplayName = "Minimum Bulk Data Size For Async Loading",
+		ToolTip = "Minimum time the time limit exceeded warning will be triggered by."))
+	int32 MinBulkDataSizeForAsyncLoading;
+
+	UPROPERTY(config, EditAnywhere, Category = IO, meta = (
+		ConsoleVariable = "AsyncIOBandwidthLimit", DisplayName = "Asynchronous IO Bandwidth Limit",
+		ToolTip = "Constrain bandwidth if wanted. Value is in MByte/ sec."))
+	float AsyncIOBandwidthLimit;
+
+	UPROPERTY(EditAnywhere, config, Category = LevelStreaming, meta = (
+		DisplayName = "Use Background Level Streaming",
+		ToolTip = "Whether to allow background level streaming."))
+	bool bUseBackgroundLevelStreaming;
+
+	UPROPERTY(EditAnywhere, config, Category = LevelStreaming, AdvancedDisplay, meta = (
+		DisplayName = "Async Loading Time Limit",
+		ToolTip = "Maximum amount of time to spend doing asynchronous loading (ms per frame)."))
+	float AsyncLoadingTimeLimit;
+
+	/** Whether to use the entire time limit even if blocked on I/O */
+	UPROPERTY(EditAnywhere, config, Category = LevelStreaming, AdvancedDisplay, meta = (
+		DisplayName = "Async Loading Use Full Time Limit",
+		ToolTip = "Whether to use the entire time limit even if blocked on I/O."))
+	bool bAsyncLoadingUseFullTimeLimit;
+
+	UPROPERTY(EditAnywhere, config, Category = LevelStreaming, AdvancedDisplay, meta = (
+		DisplayName = "Priority Async Loading Extra Time",
+		ToolTip = "Additional time to spend asynchronous loading during a high priority load."))
+	float PriorityAsyncLoadingExtraTime;
+
+	/** Maximum allowed time to spend for actor registration steps during level streaming (ms per frame)*/
+	UPROPERTY(EditAnywhere, config, Category = LevelStreaming, AdvancedDisplay, meta = (
+		DisplayName = "Level Streaming Actors Update Time Limit",
+		ToolTip = "Maximum allowed time to spend for actor registration steps during level streaming (ms per frame)."))
+	float LevelStreamingActorsUpdateTimeLimit;
+
+	/** Batching granularity used to register actor components during level streaming */
+	UPROPERTY(EditAnywhere, config, Category = LevelStreaming, AdvancedDisplay, meta = (
+		DisplayName = "Level Streaming Components Registration Granularity",
+		ToolTip = "Batching granularity used to register actor components during level streaming."))
+	int32 LevelStreamingComponentsRegistrationGranularity;
+
+	// Begin UObject interface
+	virtual void PostInitProperties() override;
+
+#if WITH_EDITOR
+	virtual void PostEditChangeProperty(FPropertyChangedEvent& PropertyChangedEvent) override;
+#endif
+	// End UObject interface
+};
+
+
+/**
+* Implements the settings for garbage collection.
+*/
+UCLASS(config = Engine, defaultconfig, meta = (DisplayName = "Garbage Collection"))
+class ENGINE_API UGarbageCollectionSettings : public UDeveloperSettings
+{
+	GENERATED_BODY()
+
+public:
+
+	UGarbageCollectionSettings();
+
+	UPROPERTY(EditAnywhere, config, Category = General, meta = (
+		DisplayName = "Time Between Purging Pending Kill Objects",
+		ToolTip = "Time in seconds (game time) we should wait between purging object references to objects that are pending kill."))
+	float TimeBetweenPurgingPendingKillObjects;
+
+	UPROPERTY(EditAnywhere, config, Category = General, meta = (
+		ConsoleVariable = "FlushStreamingOnGC", DisplayName = "Flush Streaming On GC",
+		ToolTip = "If enabled, streaming will be flushed each time garbage collection is triggered."))
+	uint32 FlushStreamingOnGC : 1;
+
+	UPROPERTY(EditAnywhere, config, Category = Optimization, meta = (
+		ConsoleVariable = "AllowParallelGC", DisplayName = "Allow Parallel GC",
+		ToolTip = "If enabled, garbage collection will use multiple threads."))
+	uint32 AllowParallelGC : 1;
+
+	UPROPERTY(EditAnywhere, config, Category = General, meta = (
+		ConsoleVariable = "NumRetriesBeforeForcingGC", DisplayName = "Number Of Retries Before Forcing GC",
+		ToolTip = "Maximum number of times GC can be skipped if worker threads are currently modifying UObject state. 0 = never force GC"))
+	int32 NumRetriesBeforeForcingGC;
+
+	UPROPERTY(EditAnywhere, config, Category = Optimization, meta = (
+		DisplayName = "Maximum Object Count Not Considered By GC",
+		ToolTip = "Maximum Object Count Not Considered By GC. Works only in cooked builds."))
+	int32 MaxObjectsNotConsideredByGC;
+
+	UPROPERTY(EditAnywhere, config, Category = Optimization, meta = (
+		DisplayName = "Size Of Permanent Object Pool",
+		ToolTip = "Size Of Permanent Object Pool (bytes). Works only in cooked builds."))
+	int32 SizeOfPermanentObjectPool;
+
+	// Begin UObject interface
+	virtual void PostInitProperties() override;
+
+#if WITH_EDITOR
+	virtual void PostEditChangeProperty(FPropertyChangedEvent& PropertyChangedEvent) override;
+#endif
+	// End UObject interface
+};
\ No newline at end of file
diff --git a/Engine/Source/Runtime/Engine/Classes/Engine/DeveloperSettings.h b/Engine/Source/Runtime/Engine/Classes/Engine/DeveloperSettings.h
--- a/Engine/Source/Runtime/Engine/Classes/Engine/DeveloperSettings.h
+++ b/Engine/Source/Runtime/Engine/Classes/Engine/DeveloperSettings.h
@@ -13,7 +13,7 @@ class ENGINE_API UDeveloperSettings : public UObject
 	GENERATED_BODY()
 
 public:
-	UDeveloperSettings(const FObjectInitializer& ObjectInitializer);
+	UDeveloperSettings(const FObjectInitializer& ObjectInitializer = FObjectInitializer::Get());
 
 	/** Gets the settings container name for the settings, either Project or Editor */
 	virtual FName GetContainerName() const;
@@ -42,4 +42,13 @@ protected:
 
 	/** The Section name, is the short name for the settings.  If not filled in, will be the FName of the class. */
 	FName SectionName;
+
+#if WITH_EDITOR
+
+	/** Populates all properties that have 'ConsoleVariable' meta data with the respective console variable values */
+	void ImportConsoleVariableValues();
+	/** If property has 'ConsoleVariable' meta data, exports the property value to the specified console variable */
+	void ExportValuesToConsoleVariables(UProperty* PropertyThatChanged);
+
+#endif
 };
diff --git a/Engine/Source/Runtime/Engine/Classes/Engine/Engine.h b/Engine/Source/Runtime/Engine/Classes/Engine/Engine.h
--- a/Engine/Source/Runtime/Engine/Classes/Engine/Engine.h
+++ b/Engine/Source/Runtime/Engine/Classes/Engine/Engine.h
@@ -1132,34 +1132,6 @@ public:
 	UPROPERTY(EditAnywhere, config, Category=Subtitles)
 	uint32 bSubtitlesForcedOff:1;
 
-	/** Time in seconds (game time) we should wait between purging object references to objects that are pending kill */
-	UPROPERTY(EditAnywhere, config, Category=Settings, AdvancedDisplay)
-	float TimeBetweenPurgingPendingKillObjects;
-
-	/** Whether to allow background level streaming. */
-	UPROPERTY(EditAnywhere, config, Category=LevelStreaming)
-	uint32 bUseBackgroundLevelStreaming:1;
-
-	/** Maximum amount of time to spend doing asynchronous loading (ms per frame) */
-	UPROPERTY(EditAnywhere, config, Category = LevelStreaming, AdvancedDisplay)
-	float AsyncLoadingTimeLimit;
-
-	/** Whether to use the entire time limit even if blocked on I/O */
-	UPROPERTY(EditAnywhere, config, Category = LevelStreaming, AdvancedDisplay)
-	uint32 bAsyncLoadingUseFullTimeLimit:1;
-	
-	/** Additional time to spend asynchronous loading during a "high priority" load */
-	UPROPERTY(EditAnywhere, config, Category = LevelStreaming, AdvancedDisplay)
-	float PriorityAsyncLoadingExtraTime;
-
-	/** Maximum allowed time to spend for actor registration steps during level streaming (ms per frame)*/
-	UPROPERTY(EditAnywhere, config, Category=LevelStreaming, AdvancedDisplay)
-	float LevelStreamingActorsUpdateTimeLimit;
-	
-	/** Batching granularity used to register actor components during level streaming */
-	UPROPERTY(EditAnywhere, config, Category=LevelStreaming, AdvancedDisplay)
-	int32 LevelStreamingComponentsRegistrationGranularity;
-
 	/** Script maximum loop iteration count used as a threshold to warn users about script execution runaway */
 	UPROPERTY(EditAnywhere, config, Category=Blueprints)
 	int32 MaximumLoopIterationCount;
diff --git a/Engine/Source/Runtime/Engine/Private/CoreSettings.cpp b/Engine/Source/Runtime/Engine/Private/CoreSettings.cpp
new file mode 100644
--- /dev/null
+++ b/Engine/Source/Runtime/Engine/Private/CoreSettings.cpp
@@ -0,0 +1,84 @@
+// Copyright 1998-2015 Epic Games, Inc. All Rights Reserved.
+
+#include "EnginePrivate.h"
+#include "Engine/CoreSettings.h"
+
+UStreamingSettings::UStreamingSettings()
+	: Super()
+{
+	SectionName = TEXT("Streaming");
+
+	AsyncLoadingThreadEnabled = false;
+	WarnIfTimeLimitExceeded = false;
+	TimeLimitExceededMultiplier = 1.5f;
+	TimeLimitExceededMinTime = 0.005f;
+	MinBulkDataSizeForAsyncLoading = 131072;
+	AsyncIOBandwidthLimit = 0;
+	bUseBackgroundLevelStreaming = true;
+	AsyncLoadingTimeLimit = 5.0f;
+	bAsyncLoadingUseFullTimeLimit = true;
+	PriorityAsyncLoadingExtraTime = 20.0f;
+	LevelStreamingActorsUpdateTimeLimit = 5.0f;
+	LevelStreamingComponentsRegistrationGranularity = 10;
+}
+
+void UStreamingSettings::PostInitProperties()
+{
+	Super::PostInitProperties();
+
+#if WITH_EDITOR
+	if (IsTemplate())
+	{
+		ImportConsoleVariableValues();
+	}
+#endif // #if WITH_EDITOR
+}
+
+#if WITH_EDITOR
+void UStreamingSettings::PostEditChangeProperty(FPropertyChangedEvent& PropertyChangedEvent)
+{
+	Super::PostEditChangeProperty(PropertyChangedEvent);
+
+	if (PropertyChangedEvent.Property)
+	{
+		ExportValuesToConsoleVariables(PropertyChangedEvent.Property);
+	}
+}
+#endif // #if WITH_EDITOR
+
+UGarbageCollectionSettings::UGarbageCollectionSettings()
+: Super()
+{
+	SectionName = TEXT("Garbage Collection");
+
+	TimeBetweenPurgingPendingKillObjects = 60.0f;
+	FlushStreamingOnGC = false;
+	AllowParallelGC = true;
+	NumRetriesBeforeForcingGC = 0;
+	MaxObjectsNotConsideredByGC = 0;
+	SizeOfPermanentObjectPool = 0;
+}
+
+void UGarbageCollectionSettings::PostInitProperties()
+{
+	Super::PostInitProperties();
+
+#if WITH_EDITOR
+	if (IsTemplate())
+	{
+		ImportConsoleVariableValues();
+	}
+#endif // #if WITH_EDITOR
+}
+
+#if WITH_EDITOR
+void UGarbageCollectionSettings::PostEditChangeProperty(FPropertyChangedEvent& PropertyChangedEvent)
+{
+	Super::PostEditChangeProperty(PropertyChangedEvent);
+
+	if (PropertyChangedEvent.Property)
+	{
+		ExportValuesToConsoleVariables(PropertyChangedEvent.Property);
+	}
+}
+#endif // #if WITH_EDITOR
diff --git a/Engine/Source/Runtime/Engine/Private/DeveloperSettings.cpp b/Engine/Source/Runtime/Engine/Private/DeveloperSettings.cpp
--- a/Engine/Source/Runtime/Engine/Private/DeveloperSettings.cpp
+++ b/Engine/Source/Runtime/Engine/Private/DeveloperSettings.cpp
@@ -86,3 +86,71 @@ TSharedPtr<SWidget> UDeveloperSettings::GetCustomSettingsWidget() const
 {
 	return TSharedPtr<SWidget>();
 }
+
+#if WITH_EDITOR
+
+static FName DeveloperSettingsConsoleVariableMetaFName(TEXT("ConsoleVariable"));
+
+void UDeveloperSettings::ImportConsoleVariableValues()
+{
+	for (UProperty* Property = GetClass()->PropertyLink; Property; Property = Property->PropertyLinkNext)
+	{
+		if (!Property->HasAnyPropertyFlags(CPF_Config))
+		{
+			continue;
+		}
+
+		FString CVarName = Property->GetMetaData(DeveloperSettingsConsoleVariableMetaFName);
+		if (!CVarName.IsEmpty())
+		{
+			IConsoleVariable* CVar = IConsoleManager::Get().FindConsoleVariable(*CVarName);
+			if (CVar)
+			{
+				if (Property->ImportText(*CVar->GetString(), Property->ContainerPtrToValuePtr<uint8>(this, 0), PPF_ConsoleVariable, this) == NULL)
+				{
+					UE_LOG(LogTemp, Error, TEXT("%s import failed for %s on console variable %s (=%s)"), *GetClass()->GetName(), *Property->GetName(), *CVarName, *CVar->GetString());
+				}
+			}
+			else
+			{
+				UE_LOG(LogTemp, Fatal, TEXT("%s failed to find console variable %s for %s"), *GetClass()->GetName(), *CVarName, *Property->GetName());
+			}
+		}
+	}
+}
+
+void UDeveloperSettings::ExportValuesToConsoleVariables(UProperty* PropertyThatChanged)
+{
+	check(PropertyThatChanged);
+	FString CVarName = PropertyThatChanged->GetMetaData(DeveloperSettingsConsoleVariableMetaFName);
+	if (!CVarName.IsEmpty())
+	{
+		IConsoleVariable* CVar = IConsoleManager::Get().FindConsoleVariable(*CVarName);
+		if (CVar && (CVar->GetFlags() & ECVF_ReadOnly) == 0)
+		{
+			UByteProperty* ByteProperty = Cast<UByteProperty>(PropertyThatChanged);
+			if (ByteProperty != NULL && ByteProperty->Enum != NULL)
+			{
+				CVar->Set(ByteProperty->GetPropertyValue_InContainer(this), ECVF_SetByProjectSetting);
+			}
+			else if (UBoolProperty* BoolProperty = Cast<UBoolProperty>(PropertyThatChanged))
+			{
+				CVar->Set((int32)BoolProperty->GetPropertyValue_InContainer(this), ECVF_SetByProjectSetting);
+			}
+			else if (UIntProperty* IntProperty = Cast<UIntProperty>(PropertyThatChanged))
+			{
+				CVar->Set(IntProperty->GetPropertyValue_InContainer(this), ECVF_SetByProjectSetting);
+			}
+			else if (UFloatProperty* FloatProperty = Cast<UFloatProperty>(PropertyThatChanged))
+			{
+				CVar->Set(FloatProperty->GetPropertyValue_InContainer(this), ECVF_SetByProjectSetting);
+			}
+		}
+		else
+		{
+			UE_LOG(LogInit, Warning, TEXT("CVar named '%s' marked up in %s was not found or is set to read-only"), *CVarName, *GetClass()->GetName());
+		}
+	}
+}
+
+#endif
\ No newline at end of file
diff --git a/Engine/Source/Runtime/Engine/Private/GameEngine.cpp b/Engine/Source/Runtime/Engine/Private/GameEngine.cpp
--- a/Engine/Source/Runtime/Engine/Private/GameEngine.cpp
+++ b/Engine/Source/Runtime/Engine/Private/GameEngine.cpp
@@ -38,6 +38,7 @@
 #include "GameFramework/GameUserSettings.h"
 #include "GameFramework/GameMode.h"
 #include "GameDelegates.h"
+#include "Engine/CoreSettings.h"
 
 ENGINE_API bool GDisallowNetworkTravel = false;
 
@@ -376,10 +377,6 @@ void UGameEngine::RedrawViewports( bool bShouldPresent /*= true*/ )
 UEngine::UEngine(const FObjectInitializer& ObjectInitializer)
 	: Super(ObjectInitializer)
 {
-	AsyncLoadingTimeLimit = 5.0f;
-	bAsyncLoadingUseFullTimeLimit = true;
-	PriorityAsyncLoadingExtraTime = 20.0f;
-
 	C_WorldBox = FColor(0, 0, 40, 255);
 	C_BrushWire = FColor(192, 0, 0, 255);
 	C_AddWire = FColor(127, 127, 255, 255);
@@ -849,7 +846,7 @@ void UGameEngine::Tick( float DeltaSeconds, bool bIdleMode )
 	// Update subsystems.
 	{
 		// This assumes that UObject::StaticTick only calls ProcessAsyncLoading.
-		StaticTick(DeltaSeconds, bAsyncLoadingUseFullTimeLimit, AsyncLoadingTimeLimit / 1000.f);
+		StaticTick(DeltaSeconds, GetDefault<UStreamingSettings>()->bAsyncLoadingUseFullTimeLimit, GetDefault<UStreamingSettings>()->AsyncLoadingTimeLimit / 1000.f);
 	}
 
 	// -----------------------------------------------------
diff --git a/Engine/Source/Runtime/Engine/Private/LevelStreaming.cpp b/Engine/Source/Runtime/Engine/Private/LevelStreaming.cpp
--- a/Engine/Source/Runtime/Engine/Private/LevelStreaming.cpp
+++ b/Engine/Source/Runtime/Engine/Private/LevelStreaming.cpp
@@ -13,6 +13,7 @@
 #endif
 #include "Engine/LevelStreamingKismet.h"
 #include "Components/BrushComponent.h"
+#include "Engine/CoreSettings.h"
 
 DEFINE_LOG_CATEGORY_STATIC(LogLevelStreaming, Log, All);
 
@@ -171,7 +172,7 @@ bool FStreamLevelAction::UpdateLevel( ULevelStreaming* LevelStreamingObject )
 		return true;
 	}
 	// Level shouldn't be loaded but is as background level streaming is enabled so we need to fire finished event regardless.
-	else if( LevelStreamingObject->GetLoadedLevel() && !LevelStreamingObject->bShouldBeLoaded && !GEngine->bUseBackgroundLevelStreaming )
+	else if (LevelStreamingObject->GetLoadedLevel() && !LevelStreamingObject->bShouldBeLoaded && !GetDefault<UStreamingSettings>()->bUseBackgroundLevelStreaming)
 	{
 		return true;
 	}
diff --git a/Engine/Source/Runtime/Engine/Private/LevelTick.cpp b/Engine/Source/Runtime/Engine/Private/LevelTick.cpp
--- a/Engine/Source/Runtime/Engine/Private/LevelTick.cpp
+++ b/Engine/Source/Runtime/Engine/Private/LevelTick.cpp
@@ -26,6 +26,7 @@
 #include "STaskGraph.h"
 #endif
 #include "ParallelFor.h"
+#include "Engine/CoreSettings.h"
 
 // this will log out all of the objects that were ticked in the FDetailedTickStats struct so you can isolate what is expensive
 #define LOG_DETAILED_DUMPSTATS 0
@@ -1081,7 +1082,7 @@ void UWorld::Tick( ELevelTick TickType, float DeltaSeconds )
 	if (Info->bHighPriorityLoading || Info->bHighPriorityLoadingLocal || IsInSeamlessTravel())
 	{
 		// Force it to use the entire time slice, even if blocked on I/O
-		ProcessAsyncLoading(true, true, GEngine->PriorityAsyncLoadingExtraTime / 1000.0f);
+		ProcessAsyncLoading(true, true, GetDefault<UStreamingSettings>()->PriorityAsyncLoadingExtraTime / 1000.0f);
 	}
 
 	// Translate world origin if requested
@@ -1314,7 +1315,9 @@ void UWorld::Tick( ELevelTick TickType, float DeltaSeconds )
 
 		const bool bAtLeastOnePlayerConnected = NetDriver && NetDriver->ClientConnections.Num() > 0;
 		const bool bShouldUseLowFrequencyGC = IsRunningDedicatedServer() && !bAtLeastOnePlayerConnected;
-		const float TimeBetweenPurgingPendingKillObjects = bShouldUseLowFrequencyGC ? GEngine->TimeBetweenPurgingPendingKillObjects * 10 : GEngine->TimeBetweenPurgingPendingKillObjects;
+		const float TimeBetweenPurgingPendingKillObjects = bShouldUseLowFrequencyGC ? 
+			GetDefault<UGarbageCollectionSettings>()->TimeBetweenPurgingPendingKillObjects * 10 : 
+			GetDefault<UGarbageCollectionSettings>()->TimeBetweenPurgingPendingKillObjects;
 
 		// See if we should delay garbage collect for this frame
 		if (bShouldDelayGarbageCollect)
@@ -1463,7 +1466,7 @@ void UWorld::CleanupActors()
 
 void UWorld::ForceGarbageCollection( bool bForcePurge/*=false*/ )
 {
-	TimeSinceLastPendingKillPurge = 1.f + GEngine->TimeBetweenPurgingPendingKillObjects;
+	TimeSinceLastPendingKillPurge = 1.f + GetDefault<UGarbageCollectionSettings>()->TimeBetweenPurgingPendingKillObjects;
 	FullPurgeTriggered = FullPurgeTriggered || bForcePurge;
 }
 
diff --git a/Engine/Source/Runtime/Engine/Private/RenderingSettings.cpp b/Engine/Source/Runtime/Engine/Private/RenderingSettings.cpp
--- a/Engine/Source/Runtime/Engine/Private/RenderingSettings.cpp
+++ b/Engine/Source/Runtime/Engine/Private/RenderingSettings.cpp
@@ -10,8 +10,6 @@ URendererSettings::URendererSettings(const FObjectInitializer& ObjectInitializer
 	TranslucentSortAxis = FVector(0.0f, -1.0f, 0.0f);
 }
 
-static FName RenderingConsoleVariableFName(TEXT("ConsoleVariable"));
-
 void URendererSettings::PostInitProperties()
 {
 	Super::PostInitProperties();
@@ -36,30 +34,7 @@ void URendererSettings::PostInitProperties()
 #if WITH_EDITOR
 	if (IsTemplate())
 	{
-		for (UProperty* Property = GetClass()->PropertyLink; Property; Property = Property->PropertyLinkNext)
-		{
-			if (!Property->HasAnyPropertyFlags(CPF_Config))
-			{
-				continue;
-			}
-
-			FString CVarName = Property->GetMetaData(RenderingConsoleVariableFName);
-			if (!CVarName.IsEmpty())
-			{
-				IConsoleVariable* CVar = IConsoleManager::Get().FindConsoleVariable(*CVarName);
-				if (CVar)
-				{
-					if (Property->ImportText(*CVar->GetString(), Property->ContainerPtrToValuePtr<uint8>(this, 0), PPF_ConsoleVariable, this) == NULL)
-					{
-						UE_LOG(LogTemp, Error, TEXT("URendererSettings import failed for %s on console variable %s (=%s)"), *Property->GetName(), *CVarName, *CVar->GetString());
-					}
-				}
-				else
-				{
-					UE_LOG(LogTemp, Fatal, TEXT("URendererSettings failed to find console variable %s for %s"), *CVarName, *Property->GetName());
-				}
-			}
-		}
+		ImportConsoleVariableValues();
 	}
 #endif // #if WITH_EDITOR
 }
@@ -71,35 +46,7 @@ void URendererSettings::PostEditChangeProperty(FPropertyChangedEvent& PropertyCh
 
 	if (PropertyChangedEvent.Property)
 	{
-		FString CVarName = PropertyChangedEvent.Property->GetMetaData(RenderingConsoleVariableFName);
-		if (!CVarName.IsEmpty())
-		{
-			IConsoleVariable* CVar = IConsoleManager::Get().FindConsoleVariable(*CVarName);
-			if (CVar && (CVar->GetFlags() & ECVF_ReadOnly) == 0)
-			{
-				UByteProperty* ByteProperty = Cast<UByteProperty>(PropertyChangedEvent.Property);
-				if (ByteProperty != NULL && ByteProperty->Enum != NULL)
-				{
-					CVar->Set(ByteProperty->GetPropertyValue_InContainer(this), ECVF_SetByProjectSetting);
-				}
-				else if (UBoolProperty* BoolProperty = Cast<UBoolProperty>(PropertyChangedEvent.Property))
-				{
-					CVar->Set((int32)BoolProperty->GetPropertyValue_InContainer(this), ECVF_SetByProjectSetting);
-				}
-				else if (UIntProperty* IntProperty = Cast<UIntProperty>(PropertyChangedEvent.Property))
-				{
-					CVar->Set(IntProperty->GetPropertyValue_InContainer(this), ECVF_SetByProjectSetting);
-				}
-				else if (UFloatProperty* FloatProperty = Cast<UFloatProperty>(PropertyChangedEvent.Property))
-				{
-					CVar->Set(FloatProperty->GetPropertyValue_InContainer(this), ECVF_SetByProjectSetting);
-				}
-			}
-			else
-			{
-				UE_LOG(LogInit, Warning, TEXT("CVar named '%s' marked up in URenderingSettings was not found or is set to read-only"), *CVarName);
-			}
-		}
+		ExportValuesToConsoleVariables(PropertyChangedEvent.Property);
 
 		// If the renderer settings are updated and saved, we need to update the user interface settings too, to prevent data loss
 		// from old curve data for UI settings.
diff --git a/Engine/Source/Runtime/Engine/Private/World.cpp b/Engine/Source/Runtime/Engine/Private/World.cpp
--- a/Engine/Source/Runtime/Engine/Private/World.cpp
+++ b/Engine/Source/Runtime/Engine/Private/World.cpp
@@ -79,6 +79,7 @@
 #include "Engine/LightMapTexture2D.h"
 #include "Engine/GameInstance.h"
 #include "UObject/UObjectThreadContext.h"
+#include "Engine/CoreSettings.h"
 
 DEFINE_LOG_CATEGORY_STATIC(LogWorld, Log, All);
 DEFINE_LOG_CATEGORY(LogSpawn);
@@ -1582,7 +1583,7 @@ static bool IsTimeLimitExceeded( const TCHAR* CurrentTask, double StartTime, ULe
 	// We don't spread work across several frames in the Editor to avoid potential side effects.
 	if( Level->OwningWorld->IsGameWorld() == true )
 	{
-		double TimeLimit = GEngine->LevelStreamingActorsUpdateTimeLimit;
+		double TimeLimit = GetDefault<UStreamingSettings>()->LevelStreamingActorsUpdateTimeLimit;
 		double CurrentTime	= FPlatformTime::Seconds();
 		// Delta time in ms.
 		double DeltaTime	= (CurrentTime - StartTime) * 1000;
@@ -1784,7 +1785,7 @@ void UWorld::AddToWorld( ULevel* Level, const FTransform& LevelTransform )
 		const bool bRerunConstructionScript = !(FPlatformProperties::RequiresCookedData() || (IsGameWorld() && (Level->bWasDuplicatedForPIE || !bRerunConstructionDuringEditorStreaming)));
 		
 		// Incrementally update components.
-		int32 NumComponentsToUpdate = GEngine->LevelStreamingComponentsRegistrationGranularity;
+		int32 NumComponentsToUpdate = GetDefault<UStreamingSettings>()->LevelStreamingComponentsRegistrationGranularity;
 		do
 		{
 			Level->IncrementalUpdateComponents( (!IsGameWorld() || IsRunningCommandlet()) ? 0 : NumComponentsToUpdate, bRerunConstructionScript );
@@ -2362,7 +2363,7 @@ void UWorld::UpdateLevelStreamingInner(ULevelStreaming* StreamingLevel)
 	bool bHasVisibilityRequestPending	= StreamingLevel->GetLoadedLevel() && StreamingLevel->GetLoadedLevel() == CurrentLevelPendingVisibility;
 		
 	// Figure out whether level should be loaded, visible and block on load if it should be loaded but currently isn't.
-	bool bShouldBeLoaded	= bHasVisibilityRequestPending || (!GEngine->bUseBackgroundLevelStreaming && !bShouldForceUnloadStreamingLevels && !StreamingLevel->bIsRequestingUnloadAndRemoval);
+	bool bShouldBeLoaded = bHasVisibilityRequestPending || (!GetDefault<UStreamingSettings>()->bUseBackgroundLevelStreaming && !bShouldForceUnloadStreamingLevels && !StreamingLevel->bIsRequestingUnloadAndRemoval);
 	bool bShouldBeVisible	= bHasVisibilityRequestPending || bShouldForceVisibleStreamingLevels;
 	bool bShouldBlockOnLoad	= StreamingLevel->bShouldBlockOnLoad || StreamingLevel->ShouldBeAlwaysLoaded();
 
@@ -2395,7 +2396,7 @@ void UWorld::UpdateLevelStreamingInner(ULevelStreaming* StreamingLevel)
 	// See whether level is already loaded
 	if (bShouldBeLoaded)
 	{
-		const bool bBlockOnLoad = (!IsGameWorld() || !GEngine->bUseBackgroundLevelStreaming || bShouldBlockOnLoad);
+		const bool bBlockOnLoad = (!IsGameWorld() || !GetDefault<UStreamingSettings>()->bUseBackgroundLevelStreaming || bShouldBlockOnLoad);
 		// Request to load or duplicate existing level
 		StreamingLevel->RequestLevel(this, bAllowLevelLoadRequests, bBlockOnLoad);
 	}
diff --git a/Engine/Source/Runtime/Launch/Private/LaunchEngineLoop.cpp b/Engine/Source/Runtime/Launch/Private/LaunchEngineLoop.cpp
--- a/Engine/Source/Runtime/Launch/Private/LaunchEngineLoop.cpp
+++ b/Engine/Source/Runtime/Launch/Private/LaunchEngineLoop.cpp
@@ -1117,6 +1117,8 @@ int32 FEngineLoop::PreInit( const TCHAR* CmdLine )
 
 	// Apply renderer settings from console variables stored in the INI.
 	ApplyCVarSettingsFromIni(TEXT("/Script/Engine.RendererSettings"),*GEngineIni, ECVF_SetByProjectSetting);
+	ApplyCVarSettingsFromIni(TEXT("/Script/Engine.StreamingSettings"), *GEngineIni, ECVF_SetByProjectSetting);
+	ApplyCVarSettingsFromIni(TEXT("/Script/Engine.GarbageCollectionSettings"), *GEngineIni, ECVF_SetByProjectSetting);
 
 #if !UE_SERVER
 	if (!IsRunningDedicatedServer())
