Added "magnitude" keyword to vector length functions.

diff --git a/Engine/Source/Runtime/Engine/Classes/Kismet/KismetMathLibrary.h b/Engine/Source/Runtime/Engine/Classes/Kismet/KismetMathLibrary.h
--- a/Engine/Source/Runtime/Engine/Classes/Kismet/KismetMathLibrary.h
+++ b/Engine/Source/Runtime/Engine/Classes/Kismet/KismetMathLibrary.h
@@ -607,19 +607,19 @@ class ENGINE_API UKismetMathLibrary : public UBlueprintFunctionLibrary
 	static float CrossProduct2D(FVector2D A, FVector2D B);
 
 	/* Returns the length of the FVector */
-	UFUNCTION(BlueprintPure, meta=(DisplayName = "VectorLength"), Category="Math|Vector")
+	UFUNCTION(BlueprintPure, meta=(DisplayName = "VectorLength", Keywords="magnitude"), Category="Math|Vector")
 	static float VSize(FVector A);
 
 	/** Returns the length of a 2d FVector. */
-	UFUNCTION(BlueprintPure, meta=(DisplayName = "Vector2dLength"), Category="Math|Vector2D")
+	UFUNCTION(BlueprintPure, meta=(DisplayName = "Vector2dLength", Keywords="magnitude"), Category="Math|Vector2D")
 	static float VSize2D(FVector2D A);
 
 	/* Returns the squared length of the FVector */
-	UFUNCTION(BlueprintPure, meta=(DisplayName = "VectorLengthSquared"), Category="Math|Vector")
+	UFUNCTION(BlueprintPure, meta=(DisplayName = "VectorLengthSquared", Keywords="magnitude"), Category="Math|Vector")
 	static float VSizeSquared(FVector A);
 
 	/** Returns the squared length of a 2d FVector. */
-	UFUNCTION(BlueprintPure, meta=(DisplayName = "Vector2dLengthSquared"), Category="Math|Vector2D")
+	UFUNCTION(BlueprintPure, meta=(DisplayName = "Vector2dLengthSquared", Keywords="magnitude"), Category="Math|Vector2D")
 	static float VSize2DSquared(FVector2D A);
 
 	/* Returns a unit normal version of the FVector A */
