Enum changes to improve documentation.

diff --git a/Engine/Source/Runtime/Engine/Classes/Animation/AnimNodeBase.h b/Engine/Source/Runtime/Engine/Classes/Animation/AnimNodeBase.h
--- a/Engine/Source/Runtime/Engine/Classes/Animation/AnimNodeBase.h
+++ b/Engine/Source/Runtime/Engine/Classes/Animation/AnimNodeBase.h
@@ -183,20 +183,20 @@ private:
 	{
 		DebugItem(FString Data, bool bInPoseSource) : DebugData(Data), bPoseSource(bInPoseSource) {}
 		
-		// This node items debug text to display
+		/** This node item's debug text to display. */
 		FString DebugData;
 
-		// Whether we are supplying a pose instead of modifying one (e.g. an playing animation)
+		/** Whether we are supplying a pose instead of modifying one (e.g. an playing animation). */
 		bool bPoseSource;
 
-		// Nodes that we are connected to
+		/** Nodes that we are connected to. */
 		TArray<FNodeDebugData> ChildNodeChain;
 	};
 
-	// This nodes final contribution weight (based on its own weight and the weight of its parents)
+	/** This nodes final contribution weight (based on its own weight and the weight of its parents). */
 	float AbsoluteWeight;
 
-	// Nodes that we are dependent on
+	/** Nodes that we are dependent on. */
 	TArray<DebugItem> NodeChain;
 
 public:
@@ -238,22 +238,22 @@ public:
 	const UAnimInstance* AnimInstance;
 };
 
-// The display mode of editable values on an animation node
+/** The display mode of editable values on an animation node. */
 UENUM()
 namespace EPinHidingMode
 {
 	enum Type
 	{
-		// Never show this property as a pin, it is only editable in the details panel (default for everything but FPoseLink properties)
+		/** Never show this property as a pin, it is only editable in the details panel (default for everything but FPoseLink properties). */
 		NeverAsPin,
 
-		// Hide this property by default, but allow the user to expose it as a pin via the details panel
+		/** Hide this property by default, but allow the user to expose it as a pin via the details panel. */
 		PinHiddenByDefault,
 
-		// Show this property as a pin by default, but allow the user to hide it via the details panel
+		/** Show this property as a pin by default, but allow the user to hide it via the details panel. */
 		PinShownByDefault,
 
-		// Always show this property as a pin; it never makes sense to edit it in the details panel (default for FPoseLink properties)
+		/** Always show this property as a pin; it never makes sense to edit it in the details panel (default for FPoseLink properties). */
 		AlwaysAsPin
 	};
 }
@@ -264,21 +264,21 @@ struct ENGINE_API FPoseLinkBase
 {
 	GENERATED_USTRUCT_BODY()
 
-	// Serialized link ID, used to build the non-serialized pointer map
+	/** Serialized link ID, used to build the non-serialized pointer map. */
 	UPROPERTY()
 	int32 LinkID;
 
 #if WITH_EDITORONLY_DATA
-	// The source link ID, used for debug visualization
+	/** The source link ID, used for debug visualization. */
 	UPROPERTY()
 	int32 SourceLinkID;
 #endif
 
 protected:
-	// The non serialized node pointer
+	/** The non serialized node pointer. */
 	struct FAnimNode_Base* LinkedNode;
 
-	// Flag to prevent reentry when dealing with circular trees
+	/** Flag to prevent reentry when dealing with circular trees. */
 	bool bProcessed;
 
 public:
@@ -293,12 +293,13 @@ public:
 	}
 
 	// Interface
+
 	void Initialize(const FAnimationInitializeContext& Context);
 	void CacheBones(const FAnimationCacheBonesContext& Context) ;
 	void Update(const FAnimationUpdateContext& Context);
 	void GatherDebugData(FNodeDebugData& DebugData);
 
-	// Try to re-establish the linked node pointer
+	/** Try to re-establish the linked node pointer. */
 	void AttemptRelink(const FAnimationBaseContext& Context);
 };
 
diff --git a/Engine/Source/Runtime/Engine/Classes/Components/ActorComponent.h b/Engine/Source/Runtime/Engine/Classes/Components/ActorComponent.h
--- a/Engine/Source/Runtime/Engine/Classes/Components/ActorComponent.h
+++ b/Engine/Source/Runtime/Engine/Classes/Components/ActorComponent.h
@@ -12,10 +12,14 @@ struct FReplicationFlags;
 UENUM()
 enum class EComponentCreationMethod : uint8
 {
-	Native,						// A component that is part of a native class
-	SimpleConstructionScript,	// A component that is created from a template defined in the Components section of the Blueprint
-	UserConstructionScript,		// A dynamically created component, either from the UserConstructionScript or from a Add Component node in a Blueprint event graph
-	Instance,					// A component added to a single Actor instance via the Component section of the Actor's details panel
+	/** A component that is part of a native class. */
+	Native,
+	/** A component that is created from a template defined in the Components section of the Blueprint. */
+	SimpleConstructionScript,	
+	/**A dynamically created component, either from the UserConstructionScript or from a Add Component node in a Blueprint event graph. */
+	UserConstructionScript,
+	/** A component added to a single Actor instance via the Component section of the Actor's details panel. */
+	Instance,
 };
 
 /**
diff --git a/Engine/Source/Runtime/Engine/Classes/Components/SceneComponent.h b/Engine/Source/Runtime/Engine/Classes/Components/SceneComponent.h
--- a/Engine/Source/Runtime/Engine/Classes/Components/SceneComponent.h
+++ b/Engine/Source/Runtime/Engine/Classes/Components/SceneComponent.h
@@ -49,9 +49,12 @@ enum EDetailMode
 UENUM()
 enum ERelativeTransformSpace
 {
-	RTS_World, // World space transform
-	RTS_Actor, // Actor space transform
-	RTS_Component, // Component space transform
+	/** World space transform. */
+	RTS_World,
+	/** Actor space transform. */
+	RTS_Actor,
+	/** Component space transform. */
+	RTS_Component,
 };
 
 //
diff --git a/Engine/Source/Runtime/Engine/Classes/Components/SplineComponent.h b/Engine/Source/Runtime/Engine/Classes/Components/SplineComponent.h
--- a/Engine/Source/Runtime/Engine/Classes/Components/SplineComponent.h
+++ b/Engine/Source/Runtime/Engine/Classes/Components/SplineComponent.h
@@ -6,7 +6,7 @@
 
 
 
-/** Permitted spline point types for SplineComponent */
+/** Permitted spline point types for SplineComponent. */
 UENUM()
 namespace ESplinePointType
 {
@@ -20,7 +20,7 @@ namespace ESplinePointType
 	};
 }
 
-/** Types of coordinate space accepted by the functions */
+/** Types of coordinate space accepted by the functions. */
 UENUM()
 namespace ESplineCoordinateSpace
 {
diff --git a/Engine/Source/Runtime/Engine/Classes/Components/SplineMeshComponent.h b/Engine/Source/Runtime/Engine/Classes/Components/SplineMeshComponent.h
--- a/Engine/Source/Runtime/Engine/Classes/Components/SplineMeshComponent.h
+++ b/Engine/Source/Runtime/Engine/Classes/Components/SplineMeshComponent.h
@@ -27,15 +27,15 @@ struct FSplineMeshParams
 {
 	GENERATED_USTRUCT_BODY()
 
-	/** Start location of spline, in component space */
+	/** Start location of spline, in component space. */
 	UPROPERTY(EditAnywhere, Category=SplineMesh)
 	FVector StartPos;
 
-	/** Start tangent of spline, in component space */
+	/** Start tangent of spline, in component space. */
 	UPROPERTY(EditAnywhere, Category=SplineMesh)
 	FVector StartTangent;
 
-	/** X and Y scale applied to mesh at start of spline */
+	/** X and Y scale applied to mesh at start of spline. */
 	UPROPERTY(EditAnywhere, Category=SplineMesh, AdvancedDisplay)
 	FVector2D StartScale;
 
@@ -43,27 +43,27 @@ struct FSplineMeshParams
 	UPROPERTY(EditAnywhere, Category=SplineMesh, AdvancedDisplay)
 	float StartRoll;
 
-	/** Starting offset of the mesh from the spline, in component space */
+	/** Starting offset of the mesh from the spline, in component space. */
 	UPROPERTY(EditAnywhere, Category=SplineMesh, AdvancedDisplay)
 	FVector2D StartOffset;
 
-	/** End location of spline, in component space */
+	/** End location of spline, in component space. */
 	UPROPERTY(EditAnywhere, Category=SplineMesh)
 	FVector EndPos;
 
-	/** End tangent of spline, in component space */
+	/** End tangent of spline, in component space. */
 	UPROPERTY(EditAnywhere, Category=SplineMesh)
 	FVector EndTangent;
 
-	/** X and Y scale applied to mesh at end of spline */
+	/** X and Y scale applied to mesh at end of spline. */
 	UPROPERTY(EditAnywhere, Category=SplineMesh, AdvancedDisplay)
 	FVector2D EndScale;
 
-	/** Roll around spline applied at end */
+	/** Roll around spline applied at end. */
 	UPROPERTY(EditAnywhere, Category=SplineMesh, AdvancedDisplay)
 	float EndRoll;
 
-	/** Ending offset of the mesh from the spline, in component space */
+	/** Ending offset of the mesh from the spline, in component space. */
 	UPROPERTY(EditAnywhere, Category=SplineMesh, AdvancedDisplay)
 	FVector2D EndOffset;
 
diff --git a/Engine/Source/Runtime/Engine/Classes/Materials/MaterialExpressionNoise.h b/Engine/Source/Runtime/Engine/Classes/Materials/MaterialExpressionNoise.h
--- a/Engine/Source/Runtime/Engine/Classes/Materials/MaterialExpressionNoise.h
+++ b/Engine/Source/Runtime/Engine/Classes/Materials/MaterialExpressionNoise.h
@@ -8,13 +8,13 @@
 UENUM()
 enum ENoiseFunction
 {
-	/** fast (~94 instructions per level) */
+	/** Fast (~94 instructions per level). */
 	NOISEFUNCTION_Simplex UMETA(DisplayName="Simplex"),
-	/** fast (~77 instructions per level) but low quality*/
+	/** Fast (~77 instructions per level) but low quality. */
 	NOISEFUNCTION_Perlin UMETA(DisplayName="Perlin"),
-	/** very slow (~393 instructions per level) */
+	/** Very slow (~393 instructions per level). */
 	NOISEFUNCTION_Gradient UMETA(DisplayName="Gradient"),
-	/** very fast (1 texture lookup, ~33 instructions per level), need to test more on every hardware, requires high quality texture filtering for bump mapping */
+	/** Very fast (1 texture lookup, ~33 instructions per level), need to test more on every hardware, requires high quality texture filtering for bump mapping. */
 	NOISEFUNCTION_FastGradient UMETA(DisplayName="FastGradient"),
 	NOISEFUNCTION_MAX,
 };
diff --git a/Engine/Source/Runtime/Engine/Classes/Matinee/InterpTrackMove.h b/Engine/Source/Runtime/Engine/Classes/Matinee/InterpTrackMove.h
--- a/Engine/Source/Runtime/Engine/Classes/Matinee/InterpTrackMove.h
+++ b/Engine/Source/Runtime/Engine/Classes/Matinee/InterpTrackMove.h
@@ -9,9 +9,9 @@ class UInterpGroup;
 UENUM()
 enum EInterpTrackMoveRotMode
 {
-	/** Should take orientation from the . */
+	/** Should take orientation from the keyframe. */
 	IMR_Keyframed,
-	/** Point the X-Axis of the controlled  AActor  at the group specified by LookAtGroupName. */
+	/** Point the X-Axis of the controlled Actor at the group specified by LookAtGroupName. */
 	IMR_LookAtGroup,
 	/** Should look along the direction of the translation path, with Z always up. */
 	// IMR_LookAlongPath // TODO!
@@ -22,7 +22,7 @@ enum EInterpTrackMoveRotMode
 };
 
 /**
- * Array of group names to retrieve position and rotation data from instead of using the datastored in the keyframe.
+ * Array of group names to retrieve position and rotation data from instead of using the data stored in the keyframe.
  * A value of NAME_None means to use the PosTrack and EulerTrack data for the keyframe.
  * There needs to be the same amount of elements in this array as there are keyframes.
  */
diff --git a/Engine/Source/Runtime/Engine/Classes/PhysicsEngine/BodySetupEnums.h b/Engine/Source/Runtime/Engine/Classes/PhysicsEngine/BodySetupEnums.h
--- a/Engine/Source/Runtime/Engine/Classes/PhysicsEngine/BodySetupEnums.h
+++ b/Engine/Source/Runtime/Engine/Classes/PhysicsEngine/BodySetupEnums.h
@@ -7,11 +7,11 @@
 UENUM()
 enum ECollisionTraceFlag
 {
-	// default, we keep simple/complex separate for each test
+	/** Keep simple/complex separate for each test. This is the default. */
 	CTF_UseDefault UMETA(DisplayName="Default"),
-	// use simple collision for complex collision test
+	/** Use simple collision for all collision tests. */
 	CTF_UseSimpleAsComplex UMETA(DisplayName="Use Simple Collision As Complex"),
-	// use complex collision (per poly) for simple collision test
+	/** Use complex collision (per poly) for all collision tests. */
 	CTF_UseComplexAsSimple UMETA(DisplayName="Use Complex Collision As Simple"),
 	CTF_MAX,
 };
@@ -19,11 +19,11 @@ enum ECollisionTraceFlag
 UENUM()
 enum EPhysicsType
 {
-	// follow owner option
+	/** Follow owner. */
 	PhysType_Default UMETA(DisplayName="Default"),	
-	// Do not follow owner, but make kinematic
+	/** Do not follow owner, but make kinematic. */
 	PhysType_Kinematic	UMETA(DisplayName="Kinematic"),		
-	// Do not follow owner, but simulate
+	/** Do not follow owner, but simulate. */
 	PhysType_Simulated	UMETA(DisplayName="Simulated")	
 };
 
diff --git a/Engine/Source/Runtime/Engine/Classes/PhysicsEngine/ConstraintInstance.h b/Engine/Source/Runtime/Engine/Classes/PhysicsEngine/ConstraintInstance.h
--- a/Engine/Source/Runtime/Engine/Classes/PhysicsEngine/ConstraintInstance.h
+++ b/Engine/Source/Runtime/Engine/Classes/PhysicsEngine/ConstraintInstance.h
@@ -16,17 +16,17 @@ namespace physx
 UENUM()
 enum ELinearConstraintMotion
 {
-	/** No constraint against this axis */ 
+	/** No constraint against this axis. */ 
 	LCM_Free	UMETA(DisplayName="Free"),
-	/** Limited freedom along this axis */ 
+	/** Limited freedom along this axis. */ 
 	LCM_Limited UMETA(DisplayName="Limited"),
-	/** Fully constraint against this axis */
+	/** Fully constraint against this axis. */
 	LCM_Locked UMETA(DisplayName="Locked"),
 
 	LCM_MAX,
 };
 
-/** Enum to indicate which frame we want */
+/** Enum to indicate which frame we want. */
 UENUM()
 namespace EConstraintFrame
 {
@@ -42,14 +42,14 @@ namespace EAngularDriveMode
 {
 	enum Type
 	{
-		//Follows the shortest arc between a pair of anuglar configurations (Ignored if any angular limits/locks are used)
+		/** Follows the shortest arc between a pair of anuglar configurations (Ignored if any angular limits/locks are used). */
 		SLERP,
-		//Path is decomposed into twist and swing. Doesn't follow shortest arc and may have gimbal lock. (Works with angular limits/locks)
+		/** Path is decomposed into twist and swing. Doesn't follow shortest arc and may have gimbal lock. (Works with angular limits/locks.) */
 		TwistAndSwing
 	};
 }
 
-/** Container for a physics representation of an object */
+/** Container for a physics representation of an object. */
 USTRUCT()
 struct ENGINE_API FConstraintInstance
 {
@@ -62,7 +62,7 @@ struct ENGINE_API FConstraintInstance
 	 */
 	int32 ConstraintIndex;
 
-	/** The component that created this instance */
+	/** The component that created this instance. */
 	UPROPERTY()
 	USceneComponent* OwnerComponent;
 
diff --git a/Engine/Source/Runtime/Engine/Classes/PhysicsEngine/PhysicsSettings.h b/Engine/Source/Runtime/Engine/Classes/PhysicsEngine/PhysicsSettings.h
--- a/Engine/Source/Runtime/Engine/Classes/PhysicsEngine/PhysicsSettings.h
+++ b/Engine/Source/Runtime/Engine/Classes/PhysicsEngine/PhysicsSettings.h
@@ -38,13 +38,13 @@ namespace EFrictionCombineMode
 {
 	enum Type
 	{
-		//Uses the average value of the materials touching: (a+b) /2
+		/** Uses the average value of the materials touching: (a+b)/2 */
 		Average = 0,	
-		//Uses the minimum value of the materials touching: min(a,b)
+		/** Uses the minimum value of the materials touching: min(a,b) */
 		Min = 1,		
-		//Uses the product of the values of the materials touching: a*b
+		/** Uses the product of the values of the materials touching: a*b */
 		Multiply = 2,	
-		//Uses the maximum value of materials touching: max(a,b)
+		/** Uses the maximum value of materials touching: max(a,b) */
 		Max = 3
 	};
 }
@@ -54,13 +54,13 @@ namespace ESettingsDOF
 {
 	enum Type
 	{
-		/*Allows for full 3D movement and rotation*/
+		/** Allows for full 3D movement and rotation. */
 		Full3D,
-		/*Allows 2D movement along the Y-Z plane.*/
+		/** Allows 2D movement along the Y-Z plane. */
 		YZPlane,
-		/*Allows 2D movement along the X-Z plane.*/
+		/** Allows 2D movement along the X-Z plane. */
 		XZPlane,
-		/*Allows 2D movement along the X-Y plane.*/
+		/** Allows 2D movement along the X-Y plane. */
 		XYPlane,
 	};
 }
@@ -70,15 +70,15 @@ namespace ESettingsLockedAxis
 {
 	enum Type
 	{
-		/*No axis is locked*/
+		/** No axis is locked. */
 		None,
-		/*Lock movement along the x-axis*/
+		/** Lock movement along the x-axis. */
 		X,
-		/*Lock movement along the y-axis*/
+		/** Lock movement along the y-axis. */
 		Y,
-		/*Lock movement along the z-axis*/
+		/** Lock movement along the z-axis. */
 		Z,
-		/* Used for backwards compatibility. Indicates that we've updated into the new struct*/
+		/** Used for backwards compatibility. Indicates that we've updated into the new struct. */
 		Invalid
 	};
 }
diff --git a/Engine/Source/Runtime/EngineSettings/Classes/GameMapsSettings.h b/Engine/Source/Runtime/EngineSettings/Classes/GameMapsSettings.h
--- a/Engine/Source/Runtime/EngineSettings/Classes/GameMapsSettings.h
+++ b/Engine/Source/Runtime/EngineSettings/Classes/GameMapsSettings.h
@@ -5,7 +5,7 @@
 #include "GameMapsSettings.generated.h"
 
 
-// Ways the screen can be split with two players
+/** Ways the screen can be split with two players. */
 UENUM()
 namespace ETwoPlayerSplitScreenType
 {
@@ -17,7 +17,7 @@ namespace ETwoPlayerSplitScreenType
 }
 
 
-// Ways the screen can be split with three players
+/** Ways the screen can be split with three players. */
 UENUM()
 namespace EThreePlayerSplitScreenType
 {
diff --git a/Engine/Source/Runtime/Foliage/Public/FoliageType.h b/Engine/Source/Runtime/Foliage/Public/FoliageType.h
--- a/Engine/Source/Runtime/Foliage/Public/FoliageType.h
+++ b/Engine/Source/Runtime/Foliage/Public/FoliageType.h
@@ -22,15 +22,15 @@ namespace EHasCustomNavigableGeometry
 UENUM()
 enum class EFoliageScaling : uint8
 {
-	// Foliage instances will have uniform X,Y and Z scales
+	/** Foliage instances will have uniform X,Y and Z scales. */
 	Uniform,
-	// Foliage instances will have random X,Y and Z scales
+	/** Foliage instances will have random X,Y and Z scales. */
 	Free,
-	// Locks the X and Y axis scale
+	/** Locks the X and Y axis scale. */
 	LockXY,
-	// Locks the X and Z axis scale
+	/** Locks the X and Z axis scale. */
 	LockXZ,
-	// Locks the Y and Z axis scale
+	/** Locks the Y and Z axis scale. */
 	LockYZ
 };
 
diff --git a/Engine/Source/Runtime/GameplayAbilities/Public/GameplayAbilitySpec.h b/Engine/Source/Runtime/GameplayAbilities/Public/GameplayAbilitySpec.h
--- a/Engine/Source/Runtime/GameplayAbilities/Public/GameplayAbilitySpec.h
+++ b/Engine/Source/Runtime/GameplayAbilities/Public/GameplayAbilitySpec.h
@@ -81,7 +81,7 @@ enum class EGameplayEffectGrantedAbilityRemovePolicy : uint8
 	CancelAbilityImmediately,
 	/** Active abilities are allowed to finish, and then removed. */
 	RemoveAbilityOnEnd,
-	/** Granted abilties are left lone when the granting GameplayEffect is removed */
+	/** Granted abilties are left lone when the granting GameplayEffect is removed. */
 	DoNothing,
 };
 
diff --git a/Engine/Source/Runtime/GameplayAbilities/Public/GameplayCueManager.h b/Engine/Source/Runtime/GameplayAbilities/Public/GameplayCueManager.h
--- a/Engine/Source/Runtime/GameplayAbilities/Public/GameplayCueManager.h
+++ b/Engine/Source/Runtime/GameplayAbilities/Public/GameplayCueManager.h
@@ -46,7 +46,7 @@ class UGameplayCueSet;
  *	
  */
 
-/** Type of payload to pass along with this cue */
+/** Type of payload to pass along with this cue. */
 UENUM()
 enum class EGameplayCuePayloadType : uint8
 {
@@ -55,7 +55,7 @@ enum class EGameplayCuePayloadType : uint8
 	FromSpec,
 };
 
-/** Structure to keep track of pending gameplay cues that haven't been applied yet */
+/** Structure to keep track of pending gameplay cues that haven't been applied yet. */
 USTRUCT()
 struct FGameplayCuePendingExecute
 {
diff --git a/Engine/Source/Runtime/GameplayAbilities/Public/GameplayEffect.h b/Engine/Source/Runtime/GameplayAbilities/Public/GameplayEffect.h
--- a/Engine/Source/Runtime/GameplayAbilities/Public/GameplayEffect.h
+++ b/Engine/Source/Runtime/GameplayAbilities/Public/GameplayEffect.h
@@ -23,27 +23,30 @@ class UGameplayEffectExecutionCalculation;
 
 DECLARE_DELEGATE_RetVal_OneParam(bool, FActiveGameplayEffectQueryCustomMatch, const FActiveGameplayEffect&);
 
-/** Enumeration outlining the possible gameplay effect magnitude calculation policies */
+/** Enumeration outlining the possible gameplay effect magnitude calculation policies. */
 UENUM()
 enum class EGameplayEffectMagnitudeCalculation : uint8
 {
-	/** Use a simple, scalable float for the calculation */
+	/** Use a simple, scalable float for the calculation. */
 	ScalableFloat,
-	/** Perform a calculation based upon an attribute */
+	/** Perform a calculation based upon an attribute. */
 	AttributeBased,
-	/** Perform a custom calculation, capable of capturing and acting on multiple attributes, in either BP or native */
+	/** Perform a custom calculation, capable of capturing and acting on multiple attributes, in either BP or native. */
 	CustomCalculationClass,	
 	/** This magnitude will be set explicity by that code/blueprint that creates the spec. */
 	SetByCaller,
 };
 
-/** Enumeration outlining the possible attribute based float calculation policies */
+/** Enumeration outlining the possible attribute based float calculation policies. */
 UENUM()
 enum class EAttributeBasedFloatCalculationType : uint8
 {
-	AttributeMagnitude,			// Use the final evaluated magnitude of the attribute
-	AttributeBaseValue,			// Use the base value of the attribute
-	AttributeBonusMagnitude		// Use the "bonus" evaluated magnitude of the attribute: Equivalent to (FinalMag - BaseValue)
+	/** Use the final evaluated magnitude of the attribute. */
+	AttributeMagnitude,
+	/** Use the base value of the attribute. */
+	AttributeBaseValue,
+	/** Use the "bonus" evaluated magnitude of the attribute: Equivalent to (FinalMag - BaseValue). */
+	AttributeBonusMagnitude
 };
 
 /** 
diff --git a/Engine/Source/Runtime/GameplayAbilities/Public/GameplayEffectTypes.h b/Engine/Source/Runtime/GameplayAbilities/Public/GameplayEffectTypes.h
--- a/Engine/Source/Runtime/GameplayAbilities/Public/GameplayEffectTypes.h
+++ b/Engine/Source/Runtime/GameplayAbilities/Public/GameplayEffectTypes.h
@@ -36,15 +36,17 @@ namespace EGameplayModOp
 {
 	enum Type
 	{		
-		// Numeric
+		/** Numeric. */
 		Additive = 0		UMETA(DisplayName="Add"),
-		Multiplicitive		UMETA(DisplayName="Multiply"),
-		Division			UMETA(DisplayName="Divide"),
+		/** Numeric. */
+		Multiplicitive		UMETA(DisplayName = "Multiply"),
+		/** Numeric. */
+		Division			UMETA(DisplayName = "Divide"),
 
-		// Other
+		/** Other. */
 		Override 			UMETA(DisplayName="Override"),	// This should always be the first non numeric ModOp
 
-		// This must always be at the end
+		// This must always be at the end.
 		Max					UMETA(DisplayName="Invalid")
 	};
 }
@@ -73,25 +75,25 @@ namespace GameplayEffectUtilities
 }
 
 
-/** Enumeration for options of where to capture gameplay attributes from for gameplay effects */
+/** Enumeration for options of where to capture gameplay attributes from for gameplay effects. */
 UENUM()
 enum class EGameplayEffectAttributeCaptureSource : uint8
 {
-	// Source (caster) of the gameplay effect
+	/** Source (caster) of the gameplay effect. */
 	Source,	
-	// Target (recipient) of the gameplay effect
+	/** Target (recipient) of the gameplay effect. */
 	Target	
 };
 
-/** Enumeration for ways a single GameplayEffect asset can stack */
+/** Enumeration for ways a single GameplayEffect asset can stack. */
 UENUM()
 enum class EGameplayEffectStackingType : uint8
 {
-	// No stacking. Multiple applications of this GameplayEffect are treated as separate instances.
+	/** No stacking. Multiple applications of this GameplayEffect are treated as separate instances. */
 	None,
-	// Each caster has its own stack
+	/** Each caster has its own stack. */
 	AggregateBySource,
-	// Each target has its own stack
+	/** Each target has its own stack. */
 	AggregateByTarget,
 };
 
diff --git a/Engine/Source/Runtime/GameplayTasks/Classes/GameplayTask.h b/Engine/Source/Runtime/GameplayTasks/Classes/GameplayTask.h
--- a/Engine/Source/Runtime/GameplayTasks/Classes/GameplayTask.h
+++ b/Engine/Source/Runtime/GameplayTasks/Classes/GameplayTask.h
@@ -24,9 +24,9 @@ enum class EGameplayTaskState : uint8
 UENUM()
 enum class ETaskResourceOverlapPolicy : uint8
 {
-	// Pause overlapping same-priority tasks
+	/** Pause overlapping same-priority tasks. */
 	StartOnTop,
-	// Wait for other same-priority tasks to finish
+	/** Wait for other same-priority tasks to finish. */
 	StartAtEnd,
 };
 	
diff --git a/Engine/Source/Runtime/IOS/IOSRuntimeSettings/Classes/IOSRuntimeSettings.h b/Engine/Source/Runtime/IOS/IOSRuntimeSettings/Classes/IOSRuntimeSettings.h
--- a/Engine/Source/Runtime/IOS/IOSRuntimeSettings/Classes/IOSRuntimeSettings.h
+++ b/Engine/Source/Runtime/IOS/IOSRuntimeSettings/Classes/IOSRuntimeSettings.h
@@ -8,16 +8,16 @@
 UENUM()
 enum class EPowerUsageFrameRateLock : uint8
 {
-    /** Frame rate is not limited */
+    /** Frame rate is not limited. */
     PUFRL_None = 0 UMETA(DisplayName="None"),
         
-    /** Frame rate is limited to a maximum of 20 frames per second */
+    /** Frame rate is limited to a maximum of 20 frames per second. */
     PUFRL_20 = 20 UMETA(DisplayName="20 FPS"),
     
-    /** Frame rate is limited to a maximum of 30 frames per second */
+    /** Frame rate is limited to a maximum of 30 frames per second. */
     PUFRL_30 = 30 UMETA(DisplayName="30 FPS"),
     
-    /** Frame rate is limited to a maximum of 60 frames per second */
+    /** Frame rate is limited to a maximum of 60 frames per second. */
     PUFRL_60 = 60 UMETA(DisplayName="60 FPS"),
 };
 
diff --git a/Engine/Source/Runtime/InputCore/Classes/InputCoreTypes.h b/Engine/Source/Runtime/InputCore/Classes/InputCoreTypes.h
--- a/Engine/Source/Runtime/InputCore/Classes/InputCoreTypes.h
+++ b/Engine/Source/Runtime/InputCore/Classes/InputCoreTypes.h
@@ -132,6 +132,7 @@ private:
 UENUM(BlueprintType)
 namespace ETouchIndex
 {
+	// The number of entries in ETouchIndex must match the number of touch keys defined in EKeys and NUM_TOUCH_KEYS above
 	enum Type
 	{
 		Touch1,
@@ -143,7 +144,7 @@ namespace ETouchIndex
 		Touch7,
 		Touch8,
 		Touch9,
-		Touch10 // The number of entries in ETouchIndex must match the number of touch keys defined in EKeys and NUM_TOUCH_KEYS above
+		Touch10
 	};
 }
 
@@ -419,7 +420,7 @@ private:
 
 };
 
-// various states of touch inputs
+/** Various states of touch inputs. */
 UENUM()
 namespace ETouchType
 {
diff --git a/Engine/Source/Runtime/Landscape/Classes/Landscape.h b/Engine/Source/Runtime/Landscape/Classes/Landscape.h
--- a/Engine/Source/Runtime/Landscape/Classes/Landscape.h
+++ b/Engine/Source/Runtime/Landscape/Classes/Landscape.h
@@ -17,11 +17,11 @@ UENUM()
 enum ELandscapeSetupErrors
 {
 	LSE_None,
-	// No Landscape Info available
+	/** No Landscape Info available. */
 	LSE_NoLandscapeInfo,
-	// There was already component with same X,Y
+	/** There was already component with same X,Y. */
 	LSE_CollsionXY,
-	// No Layer Info, need to add proper layers
+	/** No Layer Info, need to add proper layers. */
 	LSE_NoLayerInfo,
 	LSE_MAX,
 };
diff --git a/Engine/Source/Runtime/Landscape/Classes/LandscapeProxy.h b/Engine/Source/Runtime/Landscape/Classes/LandscapeProxy.h
--- a/Engine/Source/Runtime/Landscape/Classes/LandscapeProxy.h
+++ b/Engine/Source/Runtime/Landscape/Classes/LandscapeProxy.h
@@ -170,13 +170,13 @@ namespace ELandscapeLayerPaintingRestriction
 {
 	enum Type
 	{
-		// No restriction, can paint anywhere (default)
+		/** No restriction, can paint anywhere (default). */
 		None         UMETA(DisplayName="None"),
 
-		// Uses the MaxPaintedLayersPerComponent setting from the LandscapeProxy
+		/** Uses the MaxPaintedLayersPerComponent setting from the LandscapeProxy. */
 		UseMaxLayers UMETA(DisplayName="Limit Layer Count"),
 
-		// Restricts painting to only components that already have this layer
+		/** Restricts painting to only components that already have this layer. */
 		ExistingOnly UMETA(DisplayName="Existing Layers Only"),
 	};
 }
@@ -186,9 +186,9 @@ namespace ELandscapeLODFalloff
 {
 	enum Type
 	{
-		// Default mode
+		/** Default mode. */
 		Linear			UMETA(DisplayName = "Linear"),
-		// Square Root give more natural transition, and also keep the same LOD 
+		/** Square Root give more natural transition, and also keep the same LOD. */
 		SquareRoot		UMETA(DisplayName = "Square Root"),
 	};
 }
diff --git a/Engine/Source/Runtime/Landscape/Classes/Materials/MaterialExpressionLandscapeLayerCoords.h b/Engine/Source/Runtime/Landscape/Classes/Materials/MaterialExpressionLandscapeLayerCoords.h
--- a/Engine/Source/Runtime/Landscape/Classes/Materials/MaterialExpressionLandscapeLayerCoords.h
+++ b/Engine/Source/Runtime/Landscape/Classes/Materials/MaterialExpressionLandscapeLayerCoords.h
@@ -18,11 +18,13 @@ enum ETerrainCoordMappingType
 UENUM()
 enum ELandscapeCustomizedCoordType
 {
-	LCCT_None, // Don't use customized UV, just use original UV
+	/** Don't use customized UV, just use original UV. */
+	LCCT_None,
 	LCCT_CustomUV0,
 	LCCT_CustomUV1,
 	LCCT_CustomUV2,
-	LCCT_WeightMapUV, // Use original WeightMapUV, which could not be customized
+	/** Use original WeightMapUV, which could not be customized. */
+	LCCT_WeightMapUV,
 	LCCT_MAX,
 };
 
diff --git a/Engine/Source/Runtime/Online/BuildPatchServices/Private/BuildPatchManifest.h b/Engine/Source/Runtime/Online/BuildPatchServices/Private/BuildPatchManifest.h
--- a/Engine/Source/Runtime/Online/BuildPatchServices/Private/BuildPatchManifest.h
+++ b/Engine/Source/Runtime/Online/BuildPatchServices/Private/BuildPatchManifest.h
@@ -88,8 +88,11 @@ namespace EManifestFileHeader
 	enum Type
 	{
 		// Storage flags, can be raw or a combination of others.
-		STORED_RAW = 0x0, // Zero means raw data
-		STORED_COMPRESSED = 0x1, // Flag for compressed
+
+		/** Zero means raw data. */
+		STORED_RAW = 0x0,
+		/** Flag for compressed. */
+		STORED_COMPRESSED = 0x1,
 	};
 }
 
diff --git a/Engine/Source/Runtime/Online/OnlineSubsystemUtils/Public/PartyBeaconState.h b/Engine/Source/Runtime/Online/OnlineSubsystemUtils/Public/PartyBeaconState.h
--- a/Engine/Source/Runtime/Online/OnlineSubsystemUtils/Public/PartyBeaconState.h
+++ b/Engine/Source/Runtime/Online/OnlineSubsystemUtils/Public/PartyBeaconState.h
@@ -9,29 +9,29 @@ namespace EPartyReservationResult
 {
 	enum Type
 	{
-		// Empty state
+		/** Empty state. */
 		NoResult,
-		// Pending request due to async operation, server will contact client shortly
+		/** Pending request due to async operation, server will contact client shortly. */
 		RequestPending,
-		// An unknown error happened
+		/** An unknown error happened. */
 		GeneralError,
-		// All available reservations are booked
+		/** All available reservations are booked. */
 		PartyLimitReached,
-		// Wrong number of players to join the session
+		/** Wrong number of players to join the session. */
 		IncorrectPlayerCount,
-		// No response from the host
+		/** No response from the host. */
 		RequestTimedOut,
-		// Already have a reservation entry for the requesting party leader
+		/** Already have a reservation entry for the requesting party leader. */
 		ReservationDuplicate,
-		// Couldn't find the party leader specified for a reservation update request 
+		/** Couldn't find the party leader specified for a reservation update request. */
 		ReservationNotFound,
-		// Space was available and it's time to join
+		/** Space was available and it's time to join. */
 		ReservationAccepted,
-		// The beacon is paused and not accepting new connections
+		/** The beacon is paused and not accepting new connections. */
 		ReservationDenied,
-		// This player is banned
+		/** This player is banned. */
 		ReservationDenied_Banned,
-		// The reservation request was canceled before being sent
+		/** The reservation request was canceled before being sent. */
 		ReservationRequestCanceled
 	};
 }
diff --git a/Engine/Source/Runtime/Slate/Public/Framework/Views/ITypedTableView.h b/Engine/Source/Runtime/Slate/Public/Framework/Views/ITypedTableView.h
--- a/Engine/Source/Runtime/Slate/Public/Framework/Views/ITypedTableView.h
+++ b/Engine/Source/Runtime/Slate/Public/Framework/Views/ITypedTableView.h
@@ -15,7 +15,7 @@ namespace ESelectionMode
 		/** A single item can be selected at once, or no item may be selected. */
 		Single,
 
-		/** A single item can be selected at once, or no item may be selected.  You can click the item to toggle selection on and off */
+		/** A single item can be selected at once, or no item may be selected.  You can click the item to toggle selection on and off. */
 		SingleToggle,
 
 		/** Multiple items can be selected at the same time. */
@@ -23,7 +23,7 @@ namespace ESelectionMode
 	};
 }
 
-/** Lists can exist as regular lists or as backing widgets for a tree */
+/** Lists can exist as regular lists or as backing widgets for a tree. */
 UENUM()
 namespace ETableViewMode
 {
diff --git a/Engine/Source/Runtime/Slate/Public/Widgets/Layout/SScaleBox.h b/Engine/Source/Runtime/Slate/Public/Widgets/Layout/SScaleBox.h
--- a/Engine/Source/Runtime/Slate/Public/Widgets/Layout/SScaleBox.h
+++ b/Engine/Source/Runtime/Slate/Public/Widgets/Layout/SScaleBox.h
@@ -9,7 +9,7 @@ namespace EStretchDirection
 {
 	enum Type
 	{
-		/** Will scale the content up or down */
+		/** Will scale the content up or down. */
 		Both,
 		/** Will only make the content smaller, will never scale it larger than the content's desired size. */
 		DownOnly,
@@ -25,13 +25,13 @@ namespace EStretch
 	{
 		/** Does not scale the content. */
 		None,
-		/** Scales the content non-uniformly filling the entire space of the area */
+		/** Scales the content non-uniformly filling the entire space of the area. */
 		Fill,
 		/** Scales the content uniformly (preserving aspect ratio) until it can no longer scale the content without clipping it. */
 		ScaleToFit,
 		/** Scales the content uniformly (preserving aspect ratio), until all sides meet or exceed the size of the area.  Will result in clipping longer sides. */
 		ScaleToFill,
-		/** Scales the content by the scale specified by the user */
+		/** Scales the content by the scale specified by the user. */
 		UserSpecified
 	};
 }
diff --git a/Engine/Source/Runtime/SlateCore/Public/Fonts/CompositeFont.h b/Engine/Source/Runtime/SlateCore/Public/Fonts/CompositeFont.h
--- a/Engine/Source/Runtime/SlateCore/Public/Fonts/CompositeFont.h
+++ b/Engine/Source/Runtime/SlateCore/Public/Fonts/CompositeFont.h
@@ -8,19 +8,19 @@
 UENUM()
 enum class EFontHinting : uint8
 {
-	/** Use the default hinting specified in the font */
+	/** Use the default hinting specified in the font. */
 	Default,
-	/** Force the use of an automatic hinting algorithm */
+	/** Force the use of an automatic hinting algorithm. */
 	Auto,
-	/** Force the use of an automatic light hinting algorithm, optimized for non-monochrome displays */
+	/** Force the use of an automatic light hinting algorithm, optimized for non-monochrome displays. */
 	AutoLight,
-	/** Force the use of an automatic hinting algorithm optimized for monochrome displays */
+	/** Force the use of an automatic hinting algorithm optimized for monochrome displays. */
 	Monochrome,
-	/** Do not use hinting */
+	/** Do not use hinting. */
 	None,
 };
 
-/** Payload data describing an individual font in a typeface */
+/** Payload data describing an individual font in a typeface. */
 USTRUCT(BlueprintType)
 struct SLATECORE_API FFontData
 {
diff --git a/Engine/Source/Runtime/SlateCore/Public/Input/Events.h b/Engine/Source/Runtime/SlateCore/Public/Input/Events.h
--- a/Engine/Source/Runtime/SlateCore/Public/Input/Events.h
+++ b/Engine/Source/Runtime/SlateCore/Public/Input/Events.h
@@ -15,7 +15,7 @@ class SWindow;
 UENUM()
 enum class EFocusCause : uint8
 {
-	/** Focus was changed because of a mouse action */
+	/** Focus was changed because of a mouse action. */
 	Mouse,
 
 	/** Focus was changed in response to a navigation, such as the arrow keys, TAB key, controller DPad, ... */
@@ -24,13 +24,13 @@ enum class EFocusCause : uint8
 	/** Focus was changed because someone asked the application to change it. */
 	SetDirectly,
 
-	/** Focus was explicitly cleared via the escape key or other similar action */
+	/** Focus was explicitly cleared via the escape key or other similar action. */
 	Cleared,
 
-	/** Focus was changed because another widget lost focus, and focus moved to a new widget */
+	/** Focus was changed because another widget lost focus, and focus moved to a new widget. */
 	OtherWidgetLostFocus,
 
-	/** Focus was set in response to the owning window being activated */
+	/** Focus was set in response to the owning window being activated. */
 	WindowActivate,
 };
 
diff --git a/Engine/Source/Runtime/SlateCore/Public/Styling/SlateBrush.h b/Engine/Source/Runtime/SlateCore/Public/Styling/SlateBrush.h
--- a/Engine/Source/Runtime/SlateCore/Public/Styling/SlateBrush.h
+++ b/Engine/Source/Runtime/SlateCore/Public/Styling/SlateBrush.h
@@ -62,16 +62,16 @@ namespace ESlateBrushMirrorType
 {
 	enum Type
 	{
-		/** Don't mirror anything, just draw the texture as it is */
+		/** Don't mirror anything, just draw the texture as it is. */
 		NoMirror,
 
-		/** Mirror the image horizontally */
+		/** Mirror the image horizontally. */
 		Horizontal,
 
-		/** Mirror the image vertically */
+		/** Mirror the image vertically. */
 		Vertical,
 
-		/** Mirror in both directions */
+		/** Mirror in both directions. */
 		Both
 	};
 }
@@ -85,13 +85,13 @@ namespace ESlateBrushImageType
 {
 	enum Type
 	{
-		/** No image is loaded.  Color only brushes, transparent brushes etc */
+		/** No image is loaded.  Color only brushes, transparent brushes etc. */
 		NoImage,
 
-		/** The image to be loaded is in full color */
+		/** The image to be loaded is in full color. */
 		FullColor,
 
-		/** The image is a special texture in linear space (usually a rendering resource such as a lookup table)*/
+		/** The image is a special texture in linear space (usually a rendering resource such as a lookup table). */
 		Linear,
 	};
 }
diff --git a/Engine/Source/Runtime/UMG/Public/Components/WidgetComponent.h b/Engine/Source/Runtime/UMG/Public/Components/WidgetComponent.h
--- a/Engine/Source/Runtime/UMG/Public/Components/WidgetComponent.h
+++ b/Engine/Source/Runtime/UMG/Public/Components/WidgetComponent.h
@@ -9,9 +9,9 @@ struct FVirtualPointerPosition;
 UENUM()
 enum class EWidgetSpace : uint8
 {
-	// The widget is rendered in the world as mesh, it can be occluded like any other mesh in the world.
+	/** The widget is rendered in the world as mesh, it can be occluded like any other mesh in the world. */
 	World,
-	// The widget is rendered in the screen, completely outside of the world, never occluded.
+	/** The widget is rendered in the screen, completely outside of the world, never occluded. */
 	Screen
 };
 
