Improve how ComponentInstanceDataCache matches up components, particularly for User Construction Script created components.

diff --git a/Engine/Source/Runtime/Engine/Private/ComponentInstanceDataCache.cpp b/Engine/Source/Runtime/Engine/Private/ComponentInstanceDataCache.cpp
--- a/Engine/Source/Runtime/Engine/Private/ComponentInstanceDataCache.cpp
+++ b/Engine/Source/Runtime/Engine/Private/ComponentInstanceDataCache.cpp
@@ -3,11 +3,19 @@
 #include "EnginePrivate.h"
 #include "ComponentInstanceDataCache.h"
 
+FActorComponentInstanceData::FActorComponentInstanceData()
+	: SourceComponentClass(nullptr)
+	, SourceComponentTypeSerializedIndex(-1)
+	, SourceComponentCreationMethod(EComponentCreationMethod::Native)
+{
+}
+
 FActorComponentInstanceData::FActorComponentInstanceData(const UActorComponent* SourceComponent)
 {
 	check(SourceComponent);
 	SourceComponentName = SourceComponent->GetFName();
 	SourceComponentClass = SourceComponent->GetClass();
+	SourceComponentCreationMethod = SourceComponent->CreationMethod;
 	SourceComponentTypeSerializedIndex = -1;
 	
 	AActor* ComponentOwner = SourceComponent->GetOwner();
@@ -24,7 +32,8 @@ FActorComponentInstanceData::FActorComponentInstanceData(const UActorComponent*
 					bFound = true;
 					break;
 				}
-				else if (BlueprintCreatedComponent->GetClass() == SourceComponentClass)
+				else if (   BlueprintCreatedComponent->GetClass() == SourceComponentClass
+						 && BlueprintCreatedComponent->CreationMethod == SourceComponentCreationMethod)
 				{
 					++SourceComponentTypeSerializedIndex;
 				}
@@ -67,9 +76,12 @@ FActorComponentInstanceData::FActorComponentInstanceData(const UActorComponent*
 bool FActorComponentInstanceData::MatchesComponent(const UActorComponent* Component) const
 {
 	bool bMatches = false;
-	if (Component && Component->GetClass() == SourceComponentClass)
+	if (   Component 
+		&& (Component->CreationMethod == SourceComponentCreationMethod)
+		&& (Component->GetClass() == SourceComponentClass))
 	{
-		if (Component->GetFName() == SourceComponentName)
+		if (   (SourceComponentCreationMethod != EComponentCreationMethod::UserConstructionScript)
+			&& (Component->GetFName() == SourceComponentName))
 		{
 			bMatches = true;
 		}
@@ -83,6 +95,7 @@ bool FActorComponentInstanceData::MatchesComponent(const UActorComponent* Compon
 				{
 					if (   BlueprintCreatedComponent
 						&& (BlueprintCreatedComponent->GetClass() == SourceComponentClass)
+						&& (BlueprintCreatedComponent->CreationMethod == SourceComponentCreationMethod)
 						&& (++FoundSerializedComponentsOfType == SourceComponentTypeSerializedIndex))
 					{
 						bMatches = (BlueprintCreatedComponent == Component);
diff --git a/Engine/Source/Runtime/Engine/Public/ComponentInstanceDataCache.h b/Engine/Source/Runtime/Engine/Public/ComponentInstanceDataCache.h
--- a/Engine/Source/Runtime/Engine/Public/ComponentInstanceDataCache.h
+++ b/Engine/Source/Runtime/Engine/Public/ComponentInstanceDataCache.h
@@ -4,6 +4,7 @@
 
 class UActorComponent;
 class AActor;
+enum class EComponentCreationMethod : uint8;
 
 /** At what point in the rerun construction script process is ApplyToActor being called for */
 enum class ECacheApplyPhase
@@ -16,11 +17,7 @@ enum class ECacheApplyPhase
 class ENGINE_API FActorComponentInstanceData
 {
 public:
-	FActorComponentInstanceData()
-		: SourceComponentClass(nullptr)
-		, SourceComponentTypeSerializedIndex(-1)
-	{}
-
+	FActorComponentInstanceData();
 	FActorComponentInstanceData(const UActorComponent* SourceComponent);
 
 	virtual ~FActorComponentInstanceData()
@@ -52,6 +49,10 @@ protected:
 		when filtered to just that component type */
 	int32 SourceComponentTypeSerializedIndex;
 
+	/** The index of the source component in its owner's serialized array 
+		when filtered to just that component type */
+	EComponentCreationMethod SourceComponentCreationMethod;
+
 	TArray<uint8> SavedProperties;
 };
 
