Display a graceful exit message instead of crashing when unable to find the map or fall back to the default map.
Improve error reporting in these cases.

diff --git a/Engine/Source/Runtime/Engine/Private/GameInstance.cpp b/Engine/Source/Runtime/Engine/Private/GameInstance.cpp
--- a/Engine/Source/Runtime/Engine/Private/GameInstance.cpp
+++ b/Engine/Source/Runtime/Engine/Private/GameInstance.cpp
@@ -321,28 +321,42 @@ void UGameInstance::StartGameInstance()
 	}
 
 	// If waiting for a network connection, go into the starting level.
-	if (BrowseRet != EBrowseReturnVal::Success && FCString::Stricmp(Parm, *DefaultMap) != 0)
+	if (BrowseRet != EBrowseReturnVal::Success)
 	{
-		const FText Message = FText::Format(NSLOCTEXT("Engine", "MapNotFound", "The map specified on the commandline '{0}' could not be found. Would you like to load the default map instead?"), FText::FromString(URL.Map));
+		UE_LOG(LogLoad, Error, TEXT("%s"), *FString::Printf(TEXT("Failed to enter %s: %s. Please check the log for errors."), *URL.Map, *Error));
 
 		// the map specified on the command-line couldn't be loaded.  ask the user if we should load the default map instead
-		if (FCString::Stricmp(*URL.Map, *DefaultMap) != 0 &&
-			FMessageDialog::Open(EAppMsgType::OkCancel, Message) != EAppReturnType::Ok)
+		if (FCString::Stricmp(*URL.Map, *DefaultMap) != 0)
 		{
-			// user canceled (maybe a typo while attempting to run a commandlet)
-			FPlatformMisc::RequestExit(false);
-			return;
+			const FText Message = FText::Format(NSLOCTEXT("Engine", "MapNotFound", "The map specified on the commandline '{0}' could not be found. Would you like to load the default map instead?"), FText::FromString(URL.Map));
+			if (FMessageDialog::Open(EAppMsgType::OkCancel, Message) != EAppReturnType::Ok)
+			{
+				// user canceled (maybe a typo while attempting to run a commandlet)
+				FPlatformMisc::RequestExit(false);
+				return;
+			}
+			else
+			{
+				BrowseRet = Engine->Browse(*WorldContext, FURL(&DefaultURL, *(DefaultMap + GameMapsSettings->LocalMapOptions), TRAVEL_Partial), Error);
+			}
 		}
 		else
 		{
-			BrowseRet = Engine->Browse(*WorldContext, FURL(&DefaultURL, *(DefaultMap + GameMapsSettings->LocalMapOptions), TRAVEL_Partial), Error);
+			const FText Message = FText::Format(NSLOCTEXT("Engine", "MapNotFoundNoFallback", "The map specified on the commandline '{0}' could not be found. Exiting."), FText::FromString(URL.Map));
+			FMessageDialog::Open(EAppMsgType::Ok, Message);
+			FPlatformMisc::RequestExit(false);
+			return;
 		}
 	}
 
 	// Handle failure.
 	if (BrowseRet != EBrowseReturnVal::Success)
 	{
-		UE_LOG(LogLoad, Fatal, TEXT("%s"), *FString::Printf(TEXT("Failed to enter %s: %s. Please check the log for errors."), Parm, *Error));
+		UE_LOG(LogLoad, Error, TEXT("%s"), *FString::Printf(TEXT("Failed to enter %s: %s. Please check the log for errors."), *DefaultMap, *Error));
+		const FText Message = FText::Format(NSLOCTEXT("Engine", "DefaultMapNotFound", "The default map '{0}' could not be found. Exiting."), FText::FromString(DefaultMap));
+		FMessageDialog::Open(EAppMsgType::Ok, Message);
+		FPlatformMisc::RequestExit(false);
+		return;
 	}
 
 }
