Fix shadow variable

diff --git a/Engine/Source/Runtime/Engine/Private/PhysicsEngine/BodyInstance.cpp b/Engine/Source/Runtime/Engine/Private/PhysicsEngine/BodyInstance.cpp
--- a/Engine/Source/Runtime/Engine/Private/PhysicsEngine/BodyInstance.cpp
+++ b/Engine/Source/Runtime/Engine/Private/PhysicsEngine/BodyInstance.cpp
@@ -682,12 +682,12 @@ void FBodyInstance::UpdatePhysicsShapeFilterData(uint32 SkelMeshCompID, bool bUs
 		{
 			PxShape* PShape = AllShapes[ShapeIdx];
 			FBodyInstance* BI = FPhysxUserData::Get<FBodyInstance>(PShape->userData);
-			const bool bWelded = BI != nullptr;
+			const bool bIsWelded = BI != nullptr;
 			BI = BI ? BI : this;
 
-			const TEnumAsByte<ECollisionEnabled::Type> UseCollisionEnabled = CollisionEnabledOverride && !bWelded ? *CollisionEnabledOverride : (TEnumAsByte<ECollisionEnabled::Type>)BI->GetCollisionEnabled();
-			const FCollisionResponseContainer& UseResponse = ResponseOverride && !bWelded ? *ResponseOverride : BI->CollisionResponses.GetResponseContainer();
-			bool bUseNotify = bNotifyOverride && !bWelded ? *bNotifyOverride : BI->bNotifyRigidBodyCollision;
+			const TEnumAsByte<ECollisionEnabled::Type> UseCollisionEnabled = CollisionEnabledOverride && !bIsWelded ? *CollisionEnabledOverride : (TEnumAsByte<ECollisionEnabled::Type>)BI->GetCollisionEnabled();
+			const FCollisionResponseContainer& UseResponse = ResponseOverride && !bIsWelded ? *ResponseOverride : BI->CollisionResponses.GetResponseContainer();
+			bool bUseNotify = bNotifyOverride && !bIsWelded ? *bNotifyOverride : BI->bNotifyRigidBodyCollision;
 
 
 			UPrimitiveComponent* OwnerPrimitiveComponent = BI->OwnerComponent.Get();
