Turn constraint projection on by default and set their values to something reasonable

#UE-12541

diff --git a/Engine/Source/Runtime/Engine/Private/PhysicsEngine/ConstraintInstance.cpp b/Engine/Source/Runtime/Engine/Private/PhysicsEngine/ConstraintInstance.cpp
--- a/Engine/Source/Runtime/Engine/Private/PhysicsEngine/ConstraintInstance.cpp
+++ b/Engine/Source/Runtime/Engine/Private/PhysicsEngine/ConstraintInstance.cpp
@@ -225,6 +225,7 @@ FConstraintInstance::FConstraintInstance()
 	, ConstraintData(NULL)
 #endif	//WITH_PHYSX
 	, SceneIndex(0)
+	, bEnableProjection(true)
 	, bSwingLimitSoft(true)
 	, bTwistLimitSoft(true)
 	, Swing1LimitAngle(45)
@@ -270,8 +271,8 @@ FConstraintInstance::FConstraintInstance()
 	LinearBreakThreshold = 300.0f;
 	AngularBreakThreshold = 500.0f;
 
-	ProjectionLinearTolerance = 0.5; // Linear projection when error > 0.5 unreal units
-	ProjectionAngularTolerance = 10.f;// Angular projection when error > 10 degrees
+	ProjectionLinearTolerance = 5.f; // Linear projection when error > 5 cm
+	ProjectionAngularTolerance = 180.f;// Angular projection when error > 10 degrees
 }
 
 void FConstraintInstance::SetDisableCollision(bool InDisableCollision)
