Temporarily turn off destructible awake state. This is due to partial integration of apex libs. Will turn back on once new libs are in.

diff --git a/Engine/Source/Runtime/Engine/Private/Components/DestructibleComponent.cpp b/Engine/Source/Runtime/Engine/Private/Components/DestructibleComponent.cpp
--- a/Engine/Source/Runtime/Engine/Private/Components/DestructibleComponent.cpp
+++ b/Engine/Source/Runtime/Engine/Private/Components/DestructibleComponent.cpp
@@ -393,9 +393,9 @@ void UDestructibleComponent::CreatePhysicsState()
 		// Sleep/wake up as appropriate
 		if (!BodyInstance.bStartAwake)
 		{
-			ApexDestructibleActor->setChunkPhysXActorAwakeState(0, false);
-		}
+			//ApexDestructibleActor->setChunkPhysXActorAwakeState(0, false);	//TODO: broke during bad integration of apex. Will turn on once new libs are in
 		}
+	}
 
 	UpdateBounds();
 #endif	// #if WITH_APEX
@@ -1314,7 +1314,7 @@ void UDestructibleComponent::WakeRigidBody(FName BoneName /* = NAME_None */)
 	ExecuteOnPhysicsReadWrite([&]
 	{
 		const int32 ChunkIdx = BoneIdxToChunkIdx(GetBoneIndex(BoneName));
-		ApexDestructibleActor->setChunkPhysXActorAwakeState(ChunkIdx, true);
+		//ApexDestructibleActor->setChunkPhysXActorAwakeState(ChunkIdx, true); //TODO: broke during bad integration of apex. Will turn on once new libs are in
 	});
 #endif
 }
