Force PVD to disconnect when exiting PIE.

#UE-16942

diff --git a/Engine/Source/Editor/UnrealEd/Private/PlayLevel.cpp b/Engine/Source/Editor/UnrealEd/Private/PlayLevel.cpp
--- a/Engine/Source/Editor/UnrealEd/Private/PlayLevel.cpp
+++ b/Engine/Source/Editor/UnrealEd/Private/PlayLevel.cpp
@@ -531,6 +531,7 @@ void UEditorEngine::TeardownPlaySession(FWorldContext &PieWorldContext)
 	if (PhysScene)
 	{
 		PhysScene->WaitPhysScenes();
+		PhysScene->KillVisualDebugger();
 	}
 
 	// Clean up the temporary play level.
diff --git a/Engine/Source/Runtime/Engine/Private/PhysicsEngine/PhysScene.cpp b/Engine/Source/Runtime/Engine/Private/PhysicsEngine/PhysScene.cpp
--- a/Engine/Source/Runtime/Engine/Private/PhysicsEngine/PhysScene.cpp
+++ b/Engine/Source/Runtime/Engine/Private/PhysicsEngine/PhysScene.cpp
@@ -665,6 +665,14 @@ void FPhysScene::TickPhysScene(uint32 SceneType, FGraphEventRef& InOutCompletion
 #endif
 }
 
+void FPhysScene::KillVisualDebugger()
+{
+	if (GPhysXSDK)
+	{
+		GPhysXSDK->getPvdConnectionManager()->disconnect();
+	}
+}
+
 void FPhysScene::WaitPhysScenes()
 {
 	FGraphEventArray ThingsToComplete;
diff --git a/Engine/Source/Runtime/Engine/Public/PhysicsPublic.h b/Engine/Source/Runtime/Engine/Public/PhysicsPublic.h
--- a/Engine/Source/Runtime/Engine/Public/PhysicsPublic.h
+++ b/Engine/Source/Runtime/Engine/Public/PhysicsPublic.h
@@ -430,6 +430,9 @@ public:
 	/** Waits for all physics scenes to complete */
 	ENGINE_API void WaitPhysScenes();
 
+	/** Kill the visual debugger */
+	ENGINE_API void KillVisualDebugger();
+
 	/** Waits for cloth scene to complete */
 	ENGINE_API void WaitClothScene();
 
