Fix welded complex collision getting incorrect position when scaled.

#UE-16946

diff --git a/Engine/Source/Runtime/Engine/Private/PhysicsEngine/BodySetup.cpp b/Engine/Source/Runtime/Engine/Private/PhysicsEngine/BodySetup.cpp
--- a/Engine/Source/Runtime/Engine/Private/PhysicsEngine/BodySetup.cpp
+++ b/Engine/Source/Runtime/Engine/Private/PhysicsEngine/BodySetup.cpp
@@ -537,10 +537,17 @@ public:
 				PTriMeshGeom.meshFlags |= PxMeshGeometryFlag::eDOUBLE_SIDED;
 			}
 
+			PxTransform PElementTransform = U2PTransform(RelativeTM);
+			
+
 
 			if (PTriMeshGeom.isValid())
 			{
 				PxTransform PElementTransform = U2PTransform(RelativeTM);
+				PElementTransform.p.x *= Scale3D.X;
+				PElementTransform.p.y *= Scale3D.Y;
+				PElementTransform.p.z *= Scale3D.Z;
+
 				// Create without 'sim shape' flag, problematic if it's kinematic, and it gets set later anyway.
 				if (!AttachShape_AssumesLocked(PTriMeshGeom, PElementTransform, MaxContactOffset, PxShapeFlag::eSCENE_QUERY_SHAPE | PxShapeFlag::eVISUALIZATION))
 				{
