#vr Fix for stereo disparity in Canvas rendering on non-oculus platforms

diff --git a/Engine/Source/Runtime/Engine/Private/GameViewportClient.cpp b/Engine/Source/Runtime/Engine/Private/GameViewportClient.cpp
--- a/Engine/Source/Runtime/Engine/Private/GameViewportClient.cpp
+++ b/Engine/Source/Runtime/Engine/Private/GameViewportClient.cpp
@@ -788,7 +788,7 @@ void UGameViewportClient::Draw(FViewport* InViewport, FCanvas* SceneCanvas)
 	DebugCanvasObject->Canvas = DebugCanvas;	
 	DebugCanvasObject->Init(InViewport->GetSizeXY().X, InViewport->GetSizeXY().Y, NULL);
 
-	const bool bStereoRendering = GEngine->HMDDevice.IsValid() && GEngine->IsStereoscopic3D(InViewport);
+	const bool bStereoRendering = GEngine->IsStereoscopic3D(InViewport);
 	if (DebugCanvas)
 	{
 		DebugCanvas->SetScaledToRenderTarget(bStereoRendering);
@@ -821,7 +821,7 @@ void UGameViewportClient::Draw(FViewport* InViewport, FCanvas* SceneCanvas)
 		}
 	}
 
-	if (bStereoRendering)
+	if (bStereoRendering && GEngine->HMDDevice.IsValid())
 	{
 		// Allow HMD to modify screen settings
 		GEngine->HMDDevice->UpdateScreenSettings(Viewport);
