Adding a console variable for TimeBetweenPurgingPendingKillObjects.

diff --git a/Engine/Source/Runtime/CoreUObject/Private/Serialization/AsyncLoading.cpp b/Engine/Source/Runtime/CoreUObject/Private/Serialization/AsyncLoading.cpp
--- a/Engine/Source/Runtime/CoreUObject/Private/Serialization/AsyncLoading.cpp
+++ b/Engine/Source/Runtime/CoreUObject/Private/Serialization/AsyncLoading.cpp
@@ -233,7 +233,7 @@ static float GTimeLimitExceededMinTime = 0.005f;
 static FAutoConsoleVariableRef CVarTimeLimitExceededMinTime(
 	TEXT("TimeLimitExceededMinTime"),
 	GTimeLimitExceededMinTime,
-	TEXT("Multiplier for time limit exceeded warning time threshold."),
+	TEXT("Minimum time the time limit exceeded warning will be triggered by."),
 	ECVF_Default
 	);
 
diff --git a/Engine/Source/Runtime/Engine/Classes/Engine/CoreSettings.h b/Engine/Source/Runtime/Engine/Classes/Engine/CoreSettings.h
--- a/Engine/Source/Runtime/Engine/Classes/Engine/CoreSettings.h
+++ b/Engine/Source/Runtime/Engine/Classes/Engine/CoreSettings.h
@@ -106,7 +106,7 @@ public:
 	UGarbageCollectionSettings();
 
 	UPROPERTY(EditAnywhere, config, Category = General, meta = (
-		DisplayName = "Time Between Purging Pending Kill Objects",
+		ConsoleVariable = "TimeBetweenPurgingPendingKillObjects", DisplayName = "Time Between Purging Pending Kill Objects",
 		ToolTip = "Time in seconds (game time) we should wait between purging object references to objects that are pending kill."))
 	float TimeBetweenPurgingPendingKillObjects;
 
diff --git a/Engine/Source/Runtime/Engine/Private/LevelTick.cpp b/Engine/Source/Runtime/Engine/Private/LevelTick.cpp
--- a/Engine/Source/Runtime/Engine/Private/LevelTick.cpp
+++ b/Engine/Source/Runtime/Engine/Private/LevelTick.cpp
@@ -997,6 +997,15 @@ static TAutoConsoleVariable<int32> CVarStressTestGCWhileStreaming(
 	TEXT("If set to 1, the engine will attempt to trigger GC each frame while async loading."));
 #endif
 
+static float GTimeBetweenPurgingPendingKillObjects = 60.0f;
+static FAutoConsoleVariableRef CVarTimeBetweenPurgingPendingKillObjects(
+	TEXT("TimeBetweenPurgingPendingKillObjects"),
+	GTimeBetweenPurgingPendingKillObjects,
+	TEXT("Time in seconds (game time) we should wait between purging object references to objects that are pending kill."),
+	ECVF_Default
+	);
+
+
 /**
  * Update the level after a variable amount of time, DeltaSeconds, has passed.
  * All child actors are ticked after their owners have been ticked.
@@ -1316,8 +1325,8 @@ void UWorld::Tick( ELevelTick TickType, float DeltaSeconds )
 		const bool bAtLeastOnePlayerConnected = NetDriver && NetDriver->ClientConnections.Num() > 0;
 		const bool bShouldUseLowFrequencyGC = IsRunningDedicatedServer() && !bAtLeastOnePlayerConnected;
 		const float TimeBetweenPurgingPendingKillObjects = bShouldUseLowFrequencyGC ? 
-			GetDefault<UGarbageCollectionSettings>()->TimeBetweenPurgingPendingKillObjects * 10 : 
-			GetDefault<UGarbageCollectionSettings>()->TimeBetweenPurgingPendingKillObjects;
+			GTimeBetweenPurgingPendingKillObjects * 10 :
+			GTimeBetweenPurgingPendingKillObjects;
 
 		// See if we should delay garbage collect for this frame
 		if (bShouldDelayGarbageCollect)
@@ -1466,7 +1475,7 @@ void UWorld::CleanupActors()
 
 void UWorld::ForceGarbageCollection( bool bForcePurge/*=false*/ )
 {
-	TimeSinceLastPendingKillPurge = 1.f + GetDefault<UGarbageCollectionSettings>()->TimeBetweenPurgingPendingKillObjects;
+	TimeSinceLastPendingKillPurge = 1.f + GTimeBetweenPurgingPendingKillObjects;
 	FullPurgeTriggered = FullPurgeTriggered || bForcePurge;
 }
 
