FIX UE-16715 Note Actor Text field does not accept Multiline Text

diff --git a/Engine/Source/Runtime/Engine/Classes/Engine/Note.h b/Engine/Source/Runtime/Engine/Classes/Engine/Note.h
--- a/Engine/Source/Runtime/Engine/Classes/Engine/Note.h
+++ b/Engine/Source/Runtime/Engine/Classes/Engine/Note.h
@@ -14,7 +14,7 @@ class ANote : public AActor
 	GENERATED_UCLASS_BODY()
 
 #if WITH_EDITORONLY_DATA
-	UPROPERTY(EditAnywhere, BlueprintReadWrite, Category=Note)
+	UPROPERTY(EditAnywhere, BlueprintReadWrite, Category = Note, meta = (MultiLine = "true"))
 	FString Text;
 
 	// Reference to sprite visualization component
