FEATURE UE-15219 A user would like to have Set BConstrainToPlane available in blueprints

diff --git a/Engine/Source/Runtime/Engine/Classes/GameFramework/MovementComponent.h b/Engine/Source/Runtime/Engine/Classes/GameFramework/MovementComponent.h
--- a/Engine/Source/Runtime/Engine/Classes/GameFramework/MovementComponent.h
+++ b/Engine/Source/Runtime/Engine/Classes/GameFramework/MovementComponent.h
@@ -383,6 +383,10 @@ public:
 	/** Sets the origin of the plane that constrains movement, enforced if the plane constraint is enabled. */
 	UFUNCTION(BlueprintCallable, Category="Components|Movement|Planar")
 	virtual void SetPlaneConstraintOrigin(FVector PlaneOrigin);
+	
+	/** Sets whether or not the plane constraint is enabled. */
+	UFUNCTION(BlueprintCallable, Category = "Components|Movement|Planar")
+	virtual void SetPlaneConstraintEnabled(bool bEnabled);
 
 	/** @return The normal of the plane that constrains movement, enforced if the plane constraint is enabled. */
 	UFUNCTION(BlueprintCallable, Category="Components|Movement|Planar")
diff --git a/Engine/Source/Runtime/Engine/Private/Components/MovementComponent.cpp b/Engine/Source/Runtime/Engine/Private/Components/MovementComponent.cpp
--- a/Engine/Source/Runtime/Engine/Private/Components/MovementComponent.cpp
+++ b/Engine/Source/Runtime/Engine/Private/Components/MovementComponent.cpp
@@ -400,6 +400,11 @@ void UMovementComponent::SetPlaneConstraintOrigin(FVector PlaneOrigin)
 	PlaneConstraintOrigin = PlaneOrigin;
 }
 
+void UMovementComponent::SetPlaneConstraintEnabled(bool bEnabled)
+{
+	bConstrainToPlane = bEnabled;
+}
+
 const FVector& UMovementComponent::GetPlaneConstraintOrigin() const
 {
 	return PlaneConstraintOrigin;
