Added a function to GameplayStatics that returns the name of the current level.

diff --git a/Engine/Source/Runtime/Engine/Classes/Kismet/GameplayStatics.h b/Engine/Source/Runtime/Engine/Classes/Kismet/GameplayStatics.h
--- a/Engine/Source/Runtime/Engine/Classes/Kismet/GameplayStatics.h
+++ b/Engine/Source/Runtime/Engine/Classes/Kismet/GameplayStatics.h
@@ -140,6 +140,14 @@ class ENGINE_API UGameplayStatics : public UBlueprintFunctionLibrary
 	UFUNCTION(BlueprintCallable, meta=(WorldContext="WorldContextObject", AdvancedDisplay = "2"), Category="Game")
 	static void OpenLevel(UObject* WorldContextObject, FName LevelName, bool bAbsolute = true, FString Options = FString(TEXT("")));
 
+	/**
+	* Get the name of the currently-open level.
+	*
+	* @param bRemovePrefixString	remove any streaming- or editor- added prefixes from the level name.
+	*/
+	UFUNCTION(BlueprintCallable, meta = (WorldContext = "WorldContextObject", AdvancedDisplay = "1"), Category = "Game")
+	static FString GetCurrentLevelName(UObject* WorldContextObject, bool bRemovePrefixString = true);
+
 	// --- Global functions ------------------------------
 
 	/** Returns the current GameMode or NULL if the GameMode can't be retrieved */
diff --git a/Engine/Source/Runtime/Engine/Private/GameplayStatics.cpp b/Engine/Source/Runtime/Engine/Private/GameplayStatics.cpp
--- a/Engine/Source/Runtime/Engine/Private/GameplayStatics.cpp
+++ b/Engine/Source/Runtime/Engine/Private/GameplayStatics.cpp
@@ -518,6 +518,20 @@ void UGameplayStatics::OpenLevel(UObject* WorldContextObject, FName LevelName, b
 	GEngine->SetClientTravel( World, *Cmd, TravelType );
 }
 
+FString UGameplayStatics::GetCurrentLevelName(UObject* WorldContextObject, bool bRemovePrefixString)
+{
+	if (UWorld* World = GEngine->GetWorldFromContextObject(WorldContextObject))
+	{
+		FString LevelName = World->GetMapName();
+		if (bRemovePrefixString)
+		{
+			LevelName.RemoveFromStart(World->StreamingLevelsPrefix);
+		}
+		return LevelName;
+	}
+	return FString();
+}
+
 FVector UGameplayStatics::GetActorArrayAverageLocation(const TArray<AActor*>& Actors)
 {
 	FVector LocationSum(0,0,0); // sum of locations
