Merge: Fix issue where the last checkpoint wouldn't load if there wasn't any stream info after it

diff --git a/Engine/Source/Runtime/Engine/Private/DemoNetDriver.cpp b/Engine/Source/Runtime/Engine/Private/DemoNetDriver.cpp
--- a/Engine/Source/Runtime/Engine/Private/DemoNetDriver.cpp
+++ b/Engine/Source/Runtime/Engine/Private/DemoNetDriver.cpp
@@ -1325,6 +1325,12 @@ void UDemoNetDriver::TickDemoPlayback( float DeltaSeconds )
 		TimeToSkip = 0.0f;
 	}
 
+	// If we have a checkpoint but and we don't need to fine scrub, we can load it now
+	if ( GotoCheckpointArchive != NULL && GotoCheckpointSkipExtraTimeInMS == -1 )
+	{
+		LoadCheckpoint();
+	}
+
 	// Make sure there is data available to read
 	if ( !ReplayStreamer->IsDataAvailable() )
 	{
@@ -1338,6 +1344,8 @@ void UDemoNetDriver::TickDemoPlayback( float DeltaSeconds )
 		LoadCheckpoint();
 	}
 
+	check( GotoCheckpointArchive == NULL );
+
 	// If we're at the end of the demo, just pause channels and return
 	if ( bDemoPlaybackDone )
 	{
@@ -1494,8 +1502,6 @@ void UDemoNetDriver::LoadCheckpoint()
 		return;
 	}
 
-	bIsLoadingCheckpoint = true;
-
 	// Reset the never-queue GUID list, we'll rebuild it
 	NonQueuedGUIDsForScrubbing.Empty();
 
@@ -1630,11 +1636,14 @@ void UDemoNetDriver::LoadCheckpoint()
 		if ( GotoCheckpointSkipExtraTimeInMS != -1 )
 		{
 			DemoCurrentTime += (float)GotoCheckpointSkipExtraTimeInMS / 1000;
-			bIsFastForwarding = true;
+
+			bIsFastForwarding		= true;
+			bIsLoadingCheckpoint	= true;
 		}
 		else
 		{
-			bIsLoadingCheckpoint = false;
+			bIsLoadingCheckpoint	= false;
+			bIsFastForwarding		= false;
 		}
 
 		GotoCheckpointArchive			= NULL;
@@ -1676,7 +1685,16 @@ void UDemoNetDriver::LoadCheckpoint()
 
 	DemoCurrentTime = (float)SavedAbsTimeMS / 1000.0f;
 
-	bIsFastForwarding = true;
+	if ( GotoCheckpointSkipExtraTimeInMS != -1 )
+	{
+		bIsLoadingCheckpoint = true;
+		bIsFastForwarding = true;
+	}
+	else
+	{
+		bIsLoadingCheckpoint = false;
+		bIsFastForwarding = false;
+	}
 
 	ReadDemoFrame( GotoCheckpointArchive );
 
@@ -1685,10 +1703,6 @@ void UDemoNetDriver::LoadCheckpoint()
 	{
 		DemoCurrentTime += (float)GotoCheckpointSkipExtraTimeInMS / 1000;
 	}
-	else
-	{
-		bIsLoadingCheckpoint = false;
-	}
 
 	GotoCheckpointSkipExtraTimeInMS = -1;
 	GotoCheckpointArchive			= NULL;
