Merge: Handle when live replays timeout when first connecting

diff --git a/Engine/Source/Runtime/Engine/Classes/Engine/DemoNetDriver.h b/Engine/Source/Runtime/Engine/Classes/Engine/DemoNetDriver.h
--- a/Engine/Source/Runtime/Engine/Classes/Engine/DemoNetDriver.h
+++ b/Engine/Source/Runtime/Engine/Classes/Engine/DemoNetDriver.h
@@ -60,6 +60,7 @@ class ENGINE_API UDemoNetDriver : public UNetDriver
 	float		TimeToSkip;
 	float		QueuedGotoTimeInSeconds;
 	uint32		InitialLiveDemoTime;
+	double		InitialLiveDemoTimeRealtime;
 
 	bool		bSavingCheckpoint;
 	double		LastCheckpointTime;
diff --git a/Engine/Source/Runtime/Engine/Private/DemoNetDriver.cpp b/Engine/Source/Runtime/Engine/Private/DemoNetDriver.cpp
--- a/Engine/Source/Runtime/Engine/Private/DemoNetDriver.cpp
+++ b/Engine/Source/Runtime/Engine/Private/DemoNetDriver.cpp
@@ -53,18 +53,19 @@ bool UDemoNetDriver::InitBase( bool bInitAsClient, FNetworkNotify* InNotify, con
 {
 	if ( Super::InitBase( bInitAsClient, InNotify, URL, bReuseAddressAndPort, Error ) )
 	{
-		DemoFilename			= URL.Map;
-		Time					= 0;
-		bIsRecordingDemoFrame	= false;
-		bDemoPlaybackDone		= false;
-		bChannelsArePaused		= false;
-		TimeToSkip				= 0.0f;
-		bIsFastForwarding		= false;
+		DemoFilename					= URL.Map;
+		Time							= 0;
+		bIsRecordingDemoFrame			= false;
+		bDemoPlaybackDone				= false;
+		bChannelsArePaused				= false;
+		TimeToSkip						= 0.0f;
+		bIsFastForwarding				= false;
 		GotoCheckpointSkipExtraTimeInMS = -1;
-		QueuedGotoTimeInSeconds	= -1.0f;
-		bIsLoadingCheckpoint	= false;
-		InitialLiveDemoTime		= 0;
-		bWasStartStreamingSuccessful = true;
+		QueuedGotoTimeInSeconds			= -1.0f;
+		bIsLoadingCheckpoint			= false;
+		InitialLiveDemoTime				= 0;
+		InitialLiveDemoTimeRealtime		= 0;
+		bWasStartStreamingSuccessful	= true;
 
 		ResetDemoState();
 
@@ -1257,10 +1258,18 @@ void UDemoNetDriver::TickDemoPlayback( float DeltaSeconds )
 		}
 		else
 		{
-			const uint32 TotalDemoTimeInMS = ReplayStreamer->GetTotalDemoTime();
-
 			// Wait for the most recent live time
-			if ( ReplayStreamer->GetTotalDemoTime() == InitialLiveDemoTime )
+			const bool bHasNewReplayTime = ( ReplayStreamer->GetTotalDemoTime() != InitialLiveDemoTime );
+				
+			// If we haven't gotten a new time from the demo by now, assume it might not be live, and just jump to the end now so we don't hang forever
+			const bool bTimeExpired = ( FPlatformTime::Seconds() - InitialLiveDemoTimeRealtime >= 15 );
+
+			if ( bTimeExpired )
+			{
+				UE_LOG( LogDemo, Warning, TEXT( "UDemoNetDriver::TickDemoPlayback: Too much time since last live update." ) );
+			}
+
+			if ( !bHasNewReplayTime && !bTimeExpired )
 			{
 				return;
 			}
@@ -1471,6 +1480,7 @@ void UDemoNetDriver::ReplayStreamingReady( bool bSuccess, bool bRecord )
 			{
 				UE_LOG( LogDemo, Log, TEXT( "UDemoNetConnection::ReplayStreamingReady: Deferring checkpoint until next available time." ) );
 				InitialLiveDemoTime = ReplayStreamer->GetTotalDemoTime();
+				InitialLiveDemoTimeRealtime = FPlatformTime::Seconds();
 			}
 		}
 	}
