Merge: Report replay streamer errors + stop demo when it fails

diff --git a/Engine/Source/Runtime/Engine/Private/DemoNetDriver.cpp b/Engine/Source/Runtime/Engine/Private/DemoNetDriver.cpp
--- a/Engine/Source/Runtime/Engine/Private/DemoNetDriver.cpp
+++ b/Engine/Source/Runtime/Engine/Private/DemoNetDriver.cpp
@@ -467,6 +467,7 @@ void UDemoNetDriver::TickDispatch(float DeltaSeconds)
 	{
 		UE_LOG( LogDemo, Error, TEXT( "UDemoNetDriver::TickFlush: ReplayStreamer ERROR: %s" ), ENetworkReplayError::ToString( ReplayStreamer->GetLastError() ) );
 		StopDemo();
+		World->GetGameInstance()->HandleDemoPlaybackFailure( EDemoPlayFailure::Generic, FString( EDemoPlayFailure::ToString( EDemoPlayFailure::Generic ) ) );
 		return;
 	}
 
@@ -1454,6 +1455,7 @@ void UDemoNetDriver::ReplayStreamingReady( bool bSuccess, bool bRecord )
 	if ( !bSuccess )
 	{
 		UE_LOG( LogDemo, Warning, TEXT( "UDemoNetConnection::ReplayStreamingReady: Failed." ) );
+		StopDemo();
 		GetWorld()->GetGameInstance()->HandleDemoPlaybackFailure( EDemoPlayFailure::DemoNotFound, FString( EDemoPlayFailure::ToString( EDemoPlayFailure::DemoNotFound ) ) );
 		return;
 	}
