Refactor welding code so that the original body is stored in a map instead of the shape userData. This is needed for shapes to easily get at their original UE counterpart.

diff --git a/Engine/Source/Runtime/Engine/Classes/PhysicsEngine/BodyInstance.h b/Engine/Source/Runtime/Engine/Classes/PhysicsEngine/BodyInstance.h
--- a/Engine/Source/Runtime/Engine/Classes/PhysicsEngine/BodyInstance.h
+++ b/Engine/Source/Runtime/Engine/Classes/PhysicsEngine/BodyInstance.h
@@ -992,6 +992,9 @@ public:
 	// Current state of the physx body for tracking deferred addition and removal.
 	BodyInstanceSceneState CurrentSceneState;
 
+	/** Returns the original owning body instance. This is needed for welding */
+	const FBodyInstance* GetOriginalBodyInstance(const physx::PxShape* PShape) const;
+
 private:
 	/**
 	 *  Trace a box against just this bodyinstance
@@ -1009,6 +1012,9 @@ private:
 	 * Helper function to update per shape filtering info. This should interface is not very friendly and should only be used from inside FBodyInstance
 	 */
 	void UpdatePhysicsShapeFilterData(uint32 SkelMeshCompID, bool bUseComplexAsSimple, bool bUseSimpleAsComplex, bool bPhysicsStatic, const TEnumAsByte<ECollisionEnabled::Type> * CollisionEnabledOverride, FCollisionResponseContainer * ResponseOverride, bool * bNotifyOverride);
+
+	/** Check if the shape is owned by this body instance */
+	bool IsShapeBoundToBody(const physx::PxShape* PShape) const;
 #endif 
 	/**
 	 * Invalidate Collision Profile Name
@@ -1040,6 +1046,12 @@ public:
 #endif	//WITH_BOX2D
 
 private:
+
+#if WITH_PHYSX
+	/** Used to map between shapes and welded bodies. We do not create entries if the owning body instance is root*/
+	TMap<physx::PxShape*, FBodyInstance*> ShapeToBodiesMap;
+
+#endif
 	bool bHasSharedShapes;
 };
 
diff --git a/Engine/Source/Runtime/Engine/Private/Collision/CollisionConversions.cpp b/Engine/Source/Runtime/Engine/Private/Collision/CollisionConversions.cpp
--- a/Engine/Source/Runtime/Engine/Private/Collision/CollisionConversions.cpp
+++ b/Engine/Source/Runtime/Engine/Private/Collision/CollisionConversions.cpp
@@ -326,13 +326,17 @@ static FVector FindGeomOpposingNormal(PxGeometryType::Enum QueryGeomType, const
 /** Set info in the HitResult (Actor, Component, PhysMaterial, BoneName, Item) based on the supplied shape and face index */
 static void SetHitResultFromShapeAndFaceIndex(const PxShape* PShape,  const PxRigidActor* PActor, const uint32 FaceIndex, FHitResult& OutResult, bool bReturnPhysMat)
 {
-	FBodyInstance* BodyInst = FPhysxUserData::Get<FBodyInstance>(PShape->userData);
+	const FBodyInstance* BodyInst = FPhysxUserData::Get<FBodyInstance>(PActor->userData);
 	FDestructibleChunkInfo* ChunkInfo = FPhysxUserData::Get<FDestructibleChunkInfo>(PShape->userData);
 	UPrimitiveComponent* PrimComp = FPhysxUserData::Get<UPrimitiveComponent>(PShape->userData);
 
-	if ((BodyInst == NULL && ChunkInfo == NULL))
+	if(BodyInst)
+	{
+		BodyInst = BodyInst->GetOriginalBodyInstance(PShape);
+	}
+
+	if (ChunkInfo == NULL)
 	{
-		BodyInst = FPhysxUserData::Get<FBodyInstance>(PActor->userData);
 		ChunkInfo = FPhysxUserData::Get<FDestructibleChunkInfo>(PActor->userData);
 	}
 
diff --git a/Engine/Source/Runtime/Engine/Private/PhysicsEngine/BodyInstance.cpp b/Engine/Source/Runtime/Engine/Private/PhysicsEngine/BodyInstance.cpp
--- a/Engine/Source/Runtime/Engine/Private/PhysicsEngine/BodyInstance.cpp
+++ b/Engine/Source/Runtime/Engine/Private/PhysicsEngine/BodyInstance.cpp
@@ -319,10 +319,10 @@ FPhysScene* GetPhysicsScene(const FBodyInstance* BodyInstance)
 
 #if WITH_PHYSX
 //Determine that the shape is associated with this subbody (or root body)
-bool ShapeBoundToBody(const PxShape * PShape, const FBodyInstance* SubBody)
+bool FBodyInstance::IsShapeBoundToBody(const PxShape * PShape) const
 {
-	FBodyInstance* BI = FPhysxUserData::Get<FBodyInstance>(PShape->userData);
-	return (SubBody->WeldParent == NULL && BI == NULL) || (BI == SubBody && BI->WeldParent != NULL);
+	const FBodyInstance* BI = GetOriginalBodyInstance(PShape);
+	return BI == this;
 }
 
 int32 FBodyInstance::GetAllShapes_AssumesLocked(TArray<PxShape*>& OutShapes) const
@@ -681,9 +681,8 @@ void FBodyInstance::UpdatePhysicsShapeFilterData(uint32 SkelMeshCompID, bool bUs
 		for (int32 ShapeIdx = 0; ShapeIdx < AllShapes.Num(); ShapeIdx++)
 		{
 			PxShape* PShape = AllShapes[ShapeIdx];
-			FBodyInstance* BI = FPhysxUserData::Get<FBodyInstance>(PShape->userData);
-			const bool bIsWelded = BI != nullptr;
-			BI = BI ? BI : this;
+			const FBodyInstance* BI = GetOriginalBodyInstance(PShape);
+			const bool bIsWelded = BI != this;
 
 			const TEnumAsByte<ECollisionEnabled::Type> UseCollisionEnabled = CollisionEnabledOverride && !bIsWelded ? *CollisionEnabledOverride : (TEnumAsByte<ECollisionEnabled::Type>)BI->GetCollisionEnabled();
 			const FCollisionResponseContainer& UseResponse = ResponseOverride && !bIsWelded ? *ResponseOverride : BI->CollisionResponses.GetResponseContainer();
@@ -1657,6 +1656,14 @@ FVector GetInitialLinearVelocity(const AActor* OwningActor, bool& bComponentAwak
 #endif // UE_WITH_PHYSICS
 
 #if WITH_PHYSX
+
+const FBodyInstance* FBodyInstance::GetOriginalBodyInstance(const PxShape* PShape) const
+{
+	const FBodyInstance* BI = WeldParent ? WeldParent : this;
+	FBodyInstance*const * Result = BI->ShapeToBodiesMap.Find(PShape);
+	return Result ? *Result : this;
+}
+
 TArray<int32> FBodyInstance::AddCollisionNotifyInfo(const FBodyInstance* Body0, const FBodyInstance* Body1, const physx::PxContactPair * Pairs, uint32 NumPairs, TArray<FCollisionNotifyInfo> & PendingNotifyInfos)
 {
 	TArray<int32> PairNotifyMapping;
@@ -1674,11 +1681,8 @@ TArray<int32> FBodyInstance::AddCollisionNotifyInfo(const FBodyInstance* Body0,
 
 		PairNotifyMapping.Add(-1);	//start as -1 because we can have collisions that we don't want to actually record collision
 
-		const FBodyInstance* SubBody0 = FPhysxUserData::Get<FBodyInstance>(Shape0->userData);
-		const FBodyInstance* SubBody1 = FPhysxUserData::Get<FBodyInstance>(Shape1->userData);
-
-		if (SubBody0 == NULL) { SubBody0 = Body0; }
-		if (SubBody1 == NULL) { SubBody1 = Body1; }
+		const FBodyInstance* SubBody0 = Body0->GetOriginalBodyInstance(Shape0);
+		const FBodyInstance* SubBody1 = Body1->GetOriginalBodyInstance(Shape1);
 		
 		if (SubBody0->bNotifyRigidBodyCollision || SubBody1->bNotifyRigidBodyCollision)
 		{
@@ -1833,6 +1837,8 @@ bool FBodyInstance::Weld(FBodyInstance* TheirBody, const FTransform& TheirTM)
 	{
 	SCOPE_CYCLE_COUNTER(STAT_UpdatePhysMats);
 
+		WeldParent = this;
+
 	UPhysicalMaterial* SimplePhysMat = GetSimplePhysicalMaterial();
 	TArray<UPhysicalMaterial*> ComplexPhysMats = GetComplexPhysicalMaterials();
 	PxMaterial* PSimpleMat = SimplePhysMat->GetPhysXMaterial();
@@ -1856,10 +1862,7 @@ bool FBodyInstance::Weld(FBodyInstance* TheirBody, const FTransform& TheirTM)
 	for (int32 ShapeIdx = 0; ShapeIdx < PNewShapes.Num(); ++ShapeIdx)
 	{
 		PxShape* PShape = PNewShapes[ShapeIdx];
-		//FBodyInstance *& BI = ShapeToBodyMap.FindOrAdd(PShape);
-		//BI = TheirBody;
-
-		PShape->userData = &TheirBody->PhysxUserData;
+			ShapeToBodiesMap.Add(PShape, TheirBody);
 	}
 
 	PostShapeChange();
@@ -1883,7 +1886,6 @@ void FBodyInstance::UnWeld(FBodyInstance* TheirBI)
 #if WITH_PHYSX
 
 	bool bShapesChanged = false;
-	bool bNeedsNotification = false;
 
 	ExecuteOnPhysicsReadWrite([&]
 	{
@@ -1893,40 +1895,34 @@ void FBodyInstance::UnWeld(FBodyInstance* TheirBI)
 	for (int32 ShapeIdx = 0; ShapeIdx < NumSyncShapes; ++ShapeIdx)
 	{
 		PxShape* PShape = PShapes[ShapeIdx];
-		if (FBodyInstance* BI = FPhysxUserData::Get<FBodyInstance>(PShape->userData))
-		{
-			bNeedsNotification |= BI->bNotifyRigidBodyCollision;
-
+			const FBodyInstance* BI = GetOriginalBodyInstance(PShape);
 			if (TheirBI == BI)
 			{
-				PShape->userData = NULL;
+				ShapeToBodiesMap.Remove(PShape);
 				RigidActorSync->detachShape(*PShape);
 				bShapesChanged = true;
 			}
 		}
-	}
 
 	for (int32 ShapeIdx = NumSyncShapes; ShapeIdx <PShapes.Num(); ++ShapeIdx)
 	{
 		PxShape* PShape = PShapes[ShapeIdx];
-		if (FBodyInstance* BI = FPhysxUserData::Get<FBodyInstance>(PShape->userData))
-		{
-			bNeedsNotification |= BI->bNotifyRigidBodyCollision;
-
+			const FBodyInstance* BI = GetOriginalBodyInstance(PShape);
 			if (TheirBI == BI)
 			{
-				PShape->userData = NULL;
+				ShapeToBodiesMap.Remove(PShape);
 				RigidActorAsync->detachShape(*PShape);
 				bShapesChanged = true;
 			}
 		}
-	}
-	});
 
 	if (bShapesChanged)
 	{
 		PostShapeChange();
 	}
+
+		WeldParent = nullptr;
+	});
 #endif
 }
 
@@ -2444,7 +2440,6 @@ void FBodyInstance::SetInstanceSimulatePhysics(bool bSimulate, bool bMaintainPhy
 					if (ChildBI != this)
 					{
 						Weld(ChildBI, ChildBI->OwnerComponent->GetSocketTransform(ChildrenLabels[ChildIdx]));
-						ChildBI->WeldParent = this;
 					}
 				}
 			}
@@ -3762,7 +3757,7 @@ bool FBodyInstance::LineTrace(struct FHitResult& OutHit, const FVector& Start, c
 				PxShape* PShape = PShapes[ShapeIdx];
 				check(PShape);
 
-				if (ShapeBoundToBody(PShape, this) == false) { continue;  }
+				if (IsShapeBoundToBody(PShape) == false) { continue;  }
 
 				const PxU32 HitBufferSize = 1;
 				PxRaycastHit PHits[HitBufferSize];
@@ -3890,7 +3885,7 @@ bool FBodyInstance::InternalSweepPhysX(struct FHitResult& OutHit, const FVector&
 			PxShape* PShape = PShapes[ShapeIdx];
 			check(PShape);
 
-			if (ShapeBoundToBody(PShape, this) == false){ continue; }
+			if (IsShapeBoundToBody(PShape) == false){ continue; }
 
 			// Filter so we trace against the right kind of collision
 			PxFilterData ShapeFilter = PShape->getQueryFilterData();
@@ -4132,7 +4127,7 @@ bool FBodyInstance::OverlapMulti(TArray<struct FOverlapResult>& InOutOverlaps, c
 				PxShape* PShape = PShapes[ShapeIdx];
 				check(PShape);
 
-				if (ShapeBoundToBody(PShape, this) == false)
+				if (IsShapeBoundToBody(PShape) == false)
 				{
 					continue;
 				}
@@ -4189,7 +4184,7 @@ bool FBodyInstance::OverlapPhysX_AssumesLocked(const PxGeometry& PGeom, const Px
 		const PxShape* PShape = PShapes[ShapeIdx];
 		check(PShape);
 
-		if (ShapeBoundToBody(PShape, this) == true)
+		if (IsShapeBoundToBody(PShape) == true)
 		{
 			PxVec3 POutDirection;
 			float OutDistance;
diff --git a/Engine/Source/Runtime/Engine/Private/PrimitiveComponentPhysics.cpp b/Engine/Source/Runtime/Engine/Private/PrimitiveComponentPhysics.cpp
--- a/Engine/Source/Runtime/Engine/Private/PrimitiveComponentPhysics.cpp
+++ b/Engine/Source/Runtime/Engine/Private/PrimitiveComponentPhysics.cpp
@@ -710,7 +710,6 @@ bool UPrimitiveComponent::WeldToImplementation(USceneComponent * InParent, FName
 			//root is simulated so we actually weld the body
 			FTransform RelativeTM = RootComponent == AttachParent ? GetRelativeTransform() : GetComponentToWorld().GetRelativeTransform(RootComponent->GetComponentToWorld());	//if direct parent we already have relative. Otherwise compute it
 			RootBI->Weld(BI, GetComponentToWorld());
-			BI->WeldParent = RootBI;
 
 			return true;
 		}
