[UE-15598] Fix crash on reload of maps saved with constructed component instances that match archetypes by name but no longer by type.

change summary:
- modified UBlueprintGeneratedClass::FindArchetype() to also check matched SCS component template types against the given archetype class

diff --git a/Engine/Source/Runtime/Engine/Private/BlueprintGeneratedClass.cpp b/Engine/Source/Runtime/Engine/Private/BlueprintGeneratedClass.cpp
--- a/Engine/Source/Runtime/Engine/Private/BlueprintGeneratedClass.cpp
+++ b/Engine/Source/Runtime/Engine/Private/BlueprintGeneratedClass.cpp
@@ -297,7 +297,15 @@ UObject* UBlueprintGeneratedClass::FindArchetype(UClass* ArchetypeClass, const F
 
 			if (SCSNode)
 			{
-				Archetype = SCSNode->ComponentTemplate;
+				// Ensure that the stored template is of the same type as the serialized object. Since
+				// we match these by name, this handles the case where the Blueprint class was updated
+				// after having previously serialized an instanced into another package (e.g. map). In
+				// that case, the Blueprint class might contain an SCS node with the same name as the
+				// previously-serialized object, but it might also have been switched to a different type.
+				if (SCSNode->ComponentTemplate && SCSNode->ComponentTemplate->IsA(ArchetypeClass))
+				{
+					Archetype = SCSNode->ComponentTemplate;
+				}
 			}
 			else if (UInheritableComponentHandler* ICH = Class->GetInheritableComponentHandler())
 			{
