#ue4 - Fix bReturnPhysMat not being respected for heightfields. Also make sure to init Hit.FaceIndex to INDEX_NONE when it's not filled in.

diff --git a/Engine/Source/Runtime/Engine/Private/Collision/CollisionConversions.cpp b/Engine/Source/Runtime/Engine/Private/Collision/CollisionConversions.cpp
--- a/Engine/Source/Runtime/Engine/Private/Collision/CollisionConversions.cpp
+++ b/Engine/Source/Runtime/Engine/Private/Collision/CollisionConversions.cpp
@@ -400,6 +400,8 @@ static void SetHitResultFromShapeAndFaceIndex(const PxShape* PShape,  const PxRi
 		OutResult.Item = INDEX_NONE;
 		OutResult.BoneName = NAME_None;
 	}
+
+	OutResult.FaceIndex = INDEX_NONE;
 }
 
 void ConvertQueryImpactHit(const UWorld* World, const PxLocationHit& PHit, FHitResult& OutResult, float CheckLength, const PxFilterData& QueryFilter, const FVector& StartLoc, const FVector& EndLoc, const PxGeometry* const Geom, const PxTransform& QueryTM, bool bReturnFaceIndex, bool bReturnPhysMat)
@@ -463,7 +465,7 @@ void ConvertQueryImpactHit(const UWorld* World, const PxLocationHit& PHit, FHitR
 	if( PHit.shape->getGeometryType() == PxGeometryType::eHEIGHTFIELD)
 	{
 		// Lookup physical material for heightfields
-		if (PHit.faceIndex != InvalidQueryHit.faceIndex)
+		if (bReturnPhysMat && PHit.faceIndex != InvalidQueryHit.faceIndex)
 		{
 			PxMaterial* HitMaterial = PHit.shape->getMaterialFromInternalFaceIndex(PHit.faceIndex);
 			if (HitMaterial != NULL)
@@ -482,14 +484,6 @@ void ConvertQueryImpactHit(const UWorld* World, const PxLocationHit& PHit, FHitR
 		{
 			OutResult.FaceIndex	= PTriMeshGeom.triangleMesh->getTrianglesRemap()[PHit.faceIndex];
 		}
-		else
-		{
-			OutResult.FaceIndex	= INDEX_NONE;
-		}
-	}
-	else
-	{
-		OutResult.FaceIndex	= INDEX_NONE;
 	}
 }
 
