Change 2586270 by [REDACTED_PERSON][REDACTED_HANDLE] on 2015/06/12 15:45:41

Made UCharacterMovementComponent::CheckFall() virtual

Merging 2586270

//depot/UE4-UT/...

to //depot/UE4/...

diff --git a/Engine/Source/Runtime/Engine/Classes/GameFramework/CharacterMovementComponent.h b/Engine/Source/Runtime/Engine/Classes/GameFramework/CharacterMovementComponent.h
--- a/Engine/Source/Runtime/Engine/Classes/GameFramework/CharacterMovementComponent.h
+++ b/Engine/Source/Runtime/Engine/Classes/GameFramework/CharacterMovementComponent.h
@@ -1211,7 +1211,7 @@ public:
 	FVector GetLedgeMove(const FVector& OldLocation, const FVector& Delta, const FVector& GravDir);
 
 	/** Check if pawn is falling */
-	bool CheckFall(const FFindFloorResult& OldFloor, const FHitResult& Hit, const FVector& Delta, const FVector& OldLocation, float remainingTime, float timeTick, int32 Iterations, bool bMustJump);
+	virtual bool CheckFall(const FFindFloorResult& OldFloor, const FHitResult& Hit, const FVector& Delta, const FVector& OldLocation, float remainingTime, float timeTick, int32 Iterations, bool bMustJump);
 	
 	/** 
 	 *  Revert to previous position OldLocation, return to being based on OldBase.
