fixed typo

diff --git a/Engine/Source/Runtime/Engine/Classes/Engine/SubsurfaceProfile.h b/Engine/Source/Runtime/Engine/Classes/Engine/SubsurfaceProfile.h
--- a/Engine/Source/Runtime/Engine/Classes/Engine/SubsurfaceProfile.h
+++ b/Engine/Source/Runtime/Engine/Classes/Engine/SubsurfaceProfile.h
@@ -157,4 +157,4 @@ private:
 };
 
 // lives on the render thread
-extern ENGINE_API TGlobalResource<FSubsurfaceProfileTexture> GSubsufaceProfileTextureObject;
+extern ENGINE_API TGlobalResource<FSubsurfaceProfileTexture> GSubsurfaceProfileTextureObject;
diff --git a/Engine/Source/Runtime/Engine/Private/Materials/Material.cpp b/Engine/Source/Runtime/Engine/Private/Materials/Material.cpp
--- a/Engine/Source/Runtime/Engine/Private/Materials/Material.cpp
+++ b/Engine/Source/Runtime/Engine/Private/Materials/Material.cpp
@@ -234,7 +234,7 @@ public:
 				if(SubsurfaceProfileRT)
 				{
 					// can be optimized (cached)
-					AllocationId = GSubsufaceProfileTextureObject.FindAllocationId(SubsurfaceProfileRT);
+					AllocationId = GSubsurfaceProfileTextureObject.FindAllocationId(SubsurfaceProfileRT);
 				}
 				else
 				{
diff --git a/Engine/Source/Runtime/Engine/Private/Materials/MaterialInstance.cpp b/Engine/Source/Runtime/Engine/Private/Materials/MaterialInstance.cpp
--- a/Engine/Source/Runtime/Engine/Private/Materials/MaterialInstance.cpp
+++ b/Engine/Source/Runtime/Engine/Private/Materials/MaterialInstance.cpp
@@ -164,7 +164,7 @@ bool FMaterialInstanceResource::GetScalarValue(
 		if (SubsurfaceProfileRT)
 		{
 			// can be optimized (cached)
-			AllocationId = GSubsufaceProfileTextureObject.FindAllocationId(SubsurfaceProfileRT);
+			AllocationId = GSubsurfaceProfileTextureObject.FindAllocationId(SubsurfaceProfileRT);
 		}
 		else
 		{
diff --git a/Engine/Source/Runtime/Engine/Private/Materials/MaterialInterface.cpp b/Engine/Source/Runtime/Engine/Private/Materials/MaterialInterface.cpp
--- a/Engine/Source/Runtime/Engine/Private/Materials/MaterialInterface.cpp
+++ b/Engine/Source/Runtime/Engine/Private/Materials/MaterialInterface.cpp
@@ -348,7 +348,7 @@ void UMaterialInterface::UpdateMaterialRenderProxy(FMaterialRenderProxy& Proxy)
 
 			if (LocalSubsurfaceProfile)
 			{
-				AllocationId = GSubsufaceProfileTextureObject.AddOrUpdateProfile(Settings, LocalSubsurfaceProfile);
+				AllocationId = GSubsurfaceProfileTextureObject.AddOrUpdateProfile(Settings, LocalSubsurfaceProfile);
 
 				check(AllocationId >= 0 && AllocationId <= 255);
 			}
diff --git a/Engine/Source/Runtime/Engine/Private/Rendering/SubsurfaceProfile.cpp b/Engine/Source/Runtime/Engine/Private/Rendering/SubsurfaceProfile.cpp
--- a/Engine/Source/Runtime/Engine/Private/Rendering/SubsurfaceProfile.cpp
+++ b/Engine/Source/Runtime/Engine/Private/Rendering/SubsurfaceProfile.cpp
@@ -9,7 +9,7 @@
 DEFINE_LOG_CATEGORY_STATIC(LogSubsurfaceProfile, Log, All);
 
 // lives on the render thread
-ENGINE_API TGlobalResource<FSubsurfaceProfileTexture> GSubsufaceProfileTextureObject;
+ENGINE_API TGlobalResource<FSubsurfaceProfileTexture> GSubsurfaceProfileTextureObject;
 
 // Texture with one or more SubSurfaceProfiles or 0 if there is no user
 static TRefCountPtr<IPooledRenderTarget> GSSProfiles;
@@ -299,7 +299,7 @@ ENGINE_API const IPooledRenderTarget* GetSubsufaceProfileTexture_RT(FRHICommandL
 {
 	check(IsInRenderingThread());
 
-	return GSubsufaceProfileTextureObject.GetTexture(RHICmdList);
+	return GSubsurfaceProfileTextureObject.GetTexture(RHICmdList);
 }
 
 // ------------------------------------------------------
@@ -315,7 +315,7 @@ void USubsurfaceProfile::BeginDestroy()
 		RemoveSubsurfaceProfile,
 		USubsurfaceProfile*, Ref, this,
 		{
-			GSubsufaceProfileTextureObject.RemoveProfile(Ref);
+			GSubsurfaceProfileTextureObject.RemoveProfile(Ref);
 		});
 
 	Super::BeginDestroy();
@@ -329,6 +329,6 @@ void USubsurfaceProfile::PostEditChangeProperty(struct FPropertyChangedEvent& Pr
 		USubsurfaceProfile*, Profile, this,
 	{
 		// any changes to the setting require an update of the texture
-		GSubsufaceProfileTextureObject.UpdateProfile(Settings, Profile);
+		GSubsurfaceProfileTextureObject.UpdateProfile(Settings, Profile);
 	});
 }
\ No newline at end of file
diff --git a/Engine/Source/Runtime/Renderer/Private/PostProcess/PostProcessSubsurface.cpp b/Engine/Source/Runtime/Renderer/Private/PostProcess/PostProcessSubsurface.cpp
--- a/Engine/Source/Runtime/Renderer/Private/PostProcess/PostProcessSubsurface.cpp
+++ b/Engine/Source/Runtime/Renderer/Private/PostProcess/PostProcessSubsurface.cpp
@@ -293,7 +293,7 @@ void FRCPassPostProcessSubsurfaceVisualize::Process(FRenderingCompositePassConte
 		Y += YStep;
 
 		uint32 Index = 0;
-		while (GSubsufaceProfileTextureObject.GetEntryString(Index++, Line))
+		while (GSubsurfaceProfileTextureObject.GetEntryString(Index++, Line))
 		{
 			Canvas.DrawShadowedString(X, Y += YStep, *Line, GetStatsFont(), FLinearColor(1, 1, 1));
 		}
