Fix crash dereferencing the game viewport if it has already been closed

https://jira.ol.epicgames.net/browse/UE-9485

diff --git a/Engine/Source/Runtime/Engine/Private/GameViewportClient.cpp b/Engine/Source/Runtime/Engine/Private/GameViewportClient.cpp
--- a/Engine/Source/Runtime/Engine/Private/GameViewportClient.cpp
+++ b/Engine/Source/Runtime/Engine/Private/GameViewportClient.cpp
@@ -1356,7 +1356,7 @@ void UGameViewportClient::LostFocus(FViewport* InViewport)
 
 void UGameViewportClient::ReceivedFocus(FViewport* InViewport)
 {
-	if (GetDefault<UInputSettings>()->bUseMouseForTouch && !GetGameViewport()->GetPlayInEditorIsSimulate())
+	if (GetDefault<UInputSettings>()->bUseMouseForTouch && GetGameViewport() && !GetGameViewport()->GetPlayInEditorIsSimulate())
 	{
 		FSlateApplication::Get().SetGameIsFakingTouchEvents(true);
 	}
