Fix scaling code so that it uses the original geometry data instead of relative scale. Fixes several Nan issues

#UE-16847

diff --git a/Engine/Source/Runtime/Engine/Classes/PhysicsEngine/AggregateGeom.h b/Engine/Source/Runtime/Engine/Classes/PhysicsEngine/AggregateGeom.h
--- a/Engine/Source/Runtime/Engine/Classes/PhysicsEngine/AggregateGeom.h
+++ b/Engine/Source/Runtime/Engine/Classes/PhysicsEngine/AggregateGeom.h
@@ -2,6 +2,7 @@
 
 #pragma once
 
+#include "ShapeElem.h"
 #include "ConvexElem.h"
 #include "BoxElem.h"
 #include "SphereElem.h"
diff --git a/Engine/Source/Runtime/Engine/Classes/PhysicsEngine/BodyInstance.h b/Engine/Source/Runtime/Engine/Classes/PhysicsEngine/BodyInstance.h
--- a/Engine/Source/Runtime/Engine/Classes/PhysicsEngine/BodyInstance.h
+++ b/Engine/Source/Runtime/Engine/Classes/PhysicsEngine/BodyInstance.h
@@ -995,6 +995,9 @@ public:
 	/** Returns the original owning body instance. This is needed for welding */
 	const FBodyInstance* GetOriginalBodyInstance(const physx::PxShape* PShape) const;
 
+	/** Returns the relative transform between root body and welded instance owned by the shape.*/
+	const FTransform& GetRelativeBodyTransform(const physx::PxShape* PShape) const;
+
 private:
 	/**
 	 *  Trace a box against just this bodyinstance
@@ -1047,9 +1050,20 @@ public:
 
 private:
 
+	struct FWeldInfo
+	{
+		FWeldInfo(FBodyInstance* InChildBI, const FTransform& InRelativeTM)
+		: ChildBI(InChildBI)
+		, RelativeTM(InRelativeTM)
+		{}
+
+		FBodyInstance* ChildBI;
+		FTransform RelativeTM;
+	};
+
 #if WITH_PHYSX
 	/** Used to map between shapes and welded bodies. We do not create entries if the owning body instance is root*/
-	TMap<physx::PxShape*, FBodyInstance*> ShapeToBodiesMap;
+	TMap<physx::PxShape*, FWeldInfo> ShapeToBodiesMap;
 
 #endif
 	bool bHasSharedShapes;
diff --git a/Engine/Source/Runtime/Engine/Classes/PhysicsEngine/BoxElem.h b/Engine/Source/Runtime/Engine/Classes/PhysicsEngine/BoxElem.h
--- a/Engine/Source/Runtime/Engine/Classes/PhysicsEngine/BoxElem.h
+++ b/Engine/Source/Runtime/Engine/Classes/PhysicsEngine/BoxElem.h
@@ -2,13 +2,14 @@
 
 #pragma once
 
+#include "ShapeElem.h"
 #include "BoxElem.generated.h"
 
 class FMeshElementCollector;
 
 /** Box shape used for collision */
 USTRUCT()
-struct FKBoxElem
+struct FKBoxElem : public FKShapeElem
 {
 	GENERATED_USTRUCT_BODY()
 
@@ -33,7 +34,8 @@ struct FKBoxElem
 
 
 	FKBoxElem()
-	: Center( FVector::ZeroVector )
+	: FKShapeElem(EAggCollisionShape::Box)
+	, Center( FVector::ZeroVector )
 	, Orientation( FQuat::Identity )
 	, X(1), Y(1), Z(1)
 	{
@@ -41,7 +43,8 @@ struct FKBoxElem
 	}
 
 	FKBoxElem( float s )
-	: Center( FVector::ZeroVector )
+	: FKShapeElem(EAggCollisionShape::Box)
+	, Center( FVector::ZeroVector )
 	, Orientation( FQuat::Identity )
 	, X(s), Y(s), Z(s)
 	{
@@ -49,7 +52,8 @@ struct FKBoxElem
 	}
 
 	FKBoxElem( float InX, float InY, float InZ ) 
-	: Center( FVector::ZeroVector )
+	: FKShapeElem(EAggCollisionShape::Box)
+	, Center( FVector::ZeroVector )
 	, Orientation( FQuat::Identity )
 	, X(InX), Y(InY), Z(InZ)
 
@@ -89,4 +93,6 @@ struct FKBoxElem
 	ENGINE_API FBox CalcAABB(const FTransform& BoneTM, float Scale) const;
 
 	ENGINE_API void ScaleElem(FVector DeltaSize, float MinSize);
+
+	ENGINE_API static EAggCollisionShape::Type StaticShapeType;
 };
diff --git a/Engine/Source/Runtime/Engine/Classes/PhysicsEngine/ConvexElem.h b/Engine/Source/Runtime/Engine/Classes/PhysicsEngine/ConvexElem.h
--- a/Engine/Source/Runtime/Engine/Classes/PhysicsEngine/ConvexElem.h
+++ b/Engine/Source/Runtime/Engine/Classes/PhysicsEngine/ConvexElem.h
@@ -2,6 +2,7 @@
 
 #pragma once
 
+#include "ShapeElem.h"
 #include "ConvexElem.generated.h"
 
 struct FDynamicMeshVertex;
@@ -13,7 +14,7 @@ namespace physx
 
 /** One convex hull, used for simplified collision. */
 USTRUCT()
-struct FKConvexElem
+struct FKConvexElem : public FKShapeElem
 {
 	GENERATED_USTRUCT_BODY()
 
@@ -36,7 +37,8 @@ struct FKConvexElem
 	physx::PxConvexMesh*   ConvexMeshNegX;
 
 	FKConvexElem()
-		: ElemBox(0)
+		: FKShapeElem(EAggCollisionShape::Convex)
+		, ElemBox(0)
 		, Transform(FTransform::Identity)
 		, ConvexMesh(NULL)
 		, ConvexMeshNegX(NULL)
@@ -78,4 +80,6 @@ struct FKConvexElem
 	friend FArchive& operator<<(FArchive& Ar, FKConvexElem& Elem);
 
 	ENGINE_API void ScaleElem(FVector DeltaSize, float MinSize);
+
+	ENGINE_API static EAggCollisionShape::Type StaticShapeType;
 };
\ No newline at end of file
diff --git a/Engine/Source/Runtime/Engine/Classes/PhysicsEngine/ShapeElem.h b/Engine/Source/Runtime/Engine/Classes/PhysicsEngine/ShapeElem.h
new file mode 100644
--- /dev/null
+++ b/Engine/Source/Runtime/Engine/Classes/PhysicsEngine/ShapeElem.h
@@ -0,0 +1,65 @@
+// Copyright 1998-2015 Epic Games, Inc. All Rights Reserved.
+
+#pragma once
+
+#include "PhysxUserData.h"
+#include "ShapeElem.generated.h"
+
+namespace EAggCollisionShape
+{
+	enum Type
+	{
+		Sphere,
+		Box,
+		Sphyl,
+		Convex,
+		Unknown
+	};
+}
+
+/** Sphere shape used for collision */
+USTRUCT()
+struct FKShapeElem
+{
+	GENERATED_USTRUCT_BODY()
+
+	FKShapeElem()
+	: ShapeType(EAggCollisionShape::Unknown)
+#if WITH_PHYSX
+	, UserData(this)
+#endif
+	
+	{}
+
+	FKShapeElem(EAggCollisionShape::Type InShapeType)
+	: ShapeType(InShapeType)
+#if WITH_PHYSX
+	, UserData(this)
+#endif
+	{}
+
+	FKShapeElem(const FKShapeElem& Copy)
+	: ShapeType(Copy.ShapeType)
+#if WITH_PHYSX
+	, UserData(this)
+#endif
+	{
+	}
+
+	virtual ~FKShapeElem(){}
+
+	template <typename T>
+	T* GetShapeCheck()
+	{
+		check(T::StaticShapeType == ShapeType);
+		return (T*)this;
+	}
+
+	const FPhysxUserData* GetUserData() const { return &UserData; }
+
+	ENGINE_API static EAggCollisionShape::Type StaticShapeType;
+
+private:
+	EAggCollisionShape::Type ShapeType;
+	FPhysxUserData UserData;
+};
\ No newline at end of file
diff --git a/Engine/Source/Runtime/Engine/Classes/PhysicsEngine/SphereElem.h b/Engine/Source/Runtime/Engine/Classes/PhysicsEngine/SphereElem.h
--- a/Engine/Source/Runtime/Engine/Classes/PhysicsEngine/SphereElem.h
+++ b/Engine/Source/Runtime/Engine/Classes/PhysicsEngine/SphereElem.h
@@ -2,12 +2,13 @@
 
 #pragma once
 
+#include "ShapeElem.h"
 #include "SphereElem.generated.h"
 
 
 /** Sphere shape used for collision */
 USTRUCT()
-struct FKSphereElem
+struct FKSphereElem : public FKShapeElem
 {
 	GENERATED_USTRUCT_BODY()
 
@@ -21,14 +22,16 @@ struct FKSphereElem
 	float Radius;
 
 	FKSphereElem() 
-	: Center( FVector::ZeroVector )
+	: FKShapeElem(EAggCollisionShape::Sphere)
+	, Center( FVector::ZeroVector )
 	, Radius(1)
 	{
 
 	}
 
 	FKSphereElem( float r ) 
-	: Center( FVector::ZeroVector )
+	: FKShapeElem(EAggCollisionShape::Sphere)
+	, Center( FVector::ZeroVector )
 	, Radius(r)
 	{
 
@@ -62,4 +65,6 @@ struct FKSphereElem
 	ENGINE_API FBox CalcAABB(const FTransform& BoneTM, float Scale) const;
 
 	ENGINE_API void ScaleElem(FVector DeltaSize, float MinSize);
+
+	ENGINE_API static EAggCollisionShape::Type StaticShapeType;
 };
\ No newline at end of file
diff --git a/Engine/Source/Runtime/Engine/Classes/PhysicsEngine/SphylElem.h b/Engine/Source/Runtime/Engine/Classes/PhysicsEngine/SphylElem.h
--- a/Engine/Source/Runtime/Engine/Classes/PhysicsEngine/SphylElem.h
+++ b/Engine/Source/Runtime/Engine/Classes/PhysicsEngine/SphylElem.h
@@ -2,6 +2,7 @@
 
 #pragma once
 
+#include "ShapeElem.h"
 #include "SphylElem.generated.h"
 
 class FMaterialRenderProxy;
@@ -9,7 +10,7 @@ class FPrimitiveDrawInterface;
 
 /** Capsule shape used for collision */
 USTRUCT()
-struct FKSphylElem
+struct FKSphylElem : public FKShapeElem
 {
 	GENERATED_USTRUCT_BODY()
 
@@ -31,7 +32,8 @@ struct FKSphylElem
 	float Length;
 
 	FKSphylElem()
-	: Center( FVector::ZeroVector )
+	: FKShapeElem(EAggCollisionShape::Sphyl)
+	, Center( FVector::ZeroVector )
 	, Orientation( FQuat::Identity )
 	, Radius(1), Length(1)
 
@@ -40,7 +42,8 @@ struct FKSphylElem
 	}
 
 	FKSphylElem( float InRadius, float InLength )
-	: Center( FVector::ZeroVector )
+	: FKShapeElem(EAggCollisionShape::Sphyl)
+	, Center( FVector::ZeroVector )
 	, Orientation( FQuat::Identity )
 	, Radius(InRadius), Length(InLength)
 	{
@@ -78,4 +81,6 @@ struct FKSphylElem
 	ENGINE_API FBox CalcAABB(const FTransform& BoneTM, float Scale) const;
 
 	ENGINE_API void ScaleElem(FVector DeltaSize, float MinSize);
+
+	ENGINE_API static EAggCollisionShape::Type StaticShapeType;
 };
\ No newline at end of file
diff --git a/Engine/Source/Runtime/Engine/Private/Collision/KAggregateGeom.cpp b/Engine/Source/Runtime/Engine/Private/Collision/KAggregateGeom.cpp
--- a/Engine/Source/Runtime/Engine/Private/Collision/KAggregateGeom.cpp
+++ b/Engine/Source/Runtime/Engine/Private/Collision/KAggregateGeom.cpp
@@ -5,6 +5,7 @@
 =============================================================================*/ 
 
 #include "EnginePrivate.h"
+#include "PhysicsEngine/AggregateGeom.h"
 #include "Collision.h"
 #include "Engine/Polys.h"
 
@@ -252,11 +253,19 @@ static bool EnsureHullIsValid(TArray<FVector>& InVerts)
 	return true;
 }
 
+///////////////////////////////////////
+///////////// FKShapeElem ////////////
+///////////////////////////////////////
+
+EAggCollisionShape::Type FKShapeElem::StaticShapeType = EAggCollisionShape::Unknown;
+
 
 ///////////////////////////////////////
 ///////////// FKSphereElem ////////////
 ///////////////////////////////////////
 
+EAggCollisionShape::Type FKSphereElem::StaticShapeType = EAggCollisionShape::Sphere;
+
 FBox FKSphereElem::CalcAABB(const FTransform& BoneTM, float Scale) const
 {
 	FTransform ElemTM = GetTransform();
@@ -274,6 +283,7 @@ FBox FKSphereElem::CalcAABB(const FTransform& BoneTM, float Scale) const
 ////////////// FKBoxElem //////////////
 ///////////////////////////////////////
 
+EAggCollisionShape::Type FKBoxElem::StaticShapeType = EAggCollisionShape::Box;
 
 FBox FKBoxElem::CalcAABB(const FTransform& BoneTM, float Scale) const
 {
@@ -292,6 +302,8 @@ FBox FKBoxElem::CalcAABB(const FTransform& BoneTM, float Scale) const
 ////////////// FKSphylElem ////////////
 ///////////////////////////////////////
 
+EAggCollisionShape::Type FKSphylElem::StaticShapeType = EAggCollisionShape::Sphyl;
+
 FBox FKSphylElem::CalcAABB(const FTransform& BoneTM, float Scale) const
 {
 	FTransform ElemTM = GetTransform();
@@ -321,6 +333,8 @@ FBox FKSphylElem::CalcAABB(const FTransform& BoneTM, float Scale) const
 ///////////// FKConvexElem ////////////
 ///////////////////////////////////////
 
+EAggCollisionShape::Type FKConvexElem::StaticShapeType = EAggCollisionShape::Convex;
+
 /** Reset the hull to empty all arrays */
 void FKConvexElem::Reset()
 {
diff --git a/Engine/Source/Runtime/Engine/Private/PhysicsEngine/BodyInstance.cpp b/Engine/Source/Runtime/Engine/Private/PhysicsEngine/BodyInstance.cpp
--- a/Engine/Source/Runtime/Engine/Private/PhysicsEngine/BodyInstance.cpp
+++ b/Engine/Source/Runtime/Engine/Private/PhysicsEngine/BodyInstance.cpp
@@ -1660,8 +1660,15 @@ FVector GetInitialLinearVelocity(const AActor* OwningActor, bool& bComponentAwak
 const FBodyInstance* FBodyInstance::GetOriginalBodyInstance(const PxShape* PShape) const
 {
 	const FBodyInstance* BI = WeldParent ? WeldParent : this;
-	FBodyInstance*const * Result = BI->ShapeToBodiesMap.Find(PShape);
-	return Result ? *Result : this;
+	const FWeldInfo* Result = BI->ShapeToBodiesMap.Find(PShape);
+	return Result ? Result->ChildBI : this;
+}
+
+const FTransform& FBodyInstance::GetRelativeBodyTransform(const physx::PxShape* PShape) const
+{
+	const FBodyInstance* BI = WeldParent ? WeldParent : this;
+	const FWeldInfo* Result = BI->ShapeToBodiesMap.Find(PShape);
+	return Result ? Result->RelativeTM : FTransform::Identity;
 }
 
 TArray<int32> FBodyInstance::AddCollisionNotifyInfo(const FBodyInstance* Body0, const FBodyInstance* Body1, const physx::PxContactPair * Pairs, uint32 NumPairs, TArray<FCollisionNotifyInfo> & PendingNotifyInfos)
@@ -1862,7 +1869,7 @@ bool FBodyInstance::Weld(FBodyInstance* TheirBody, const FTransform& TheirTM)
 	for (int32 ShapeIdx = 0; ShapeIdx < PNewShapes.Num(); ++ShapeIdx)
 	{
 		PxShape* PShape = PNewShapes[ShapeIdx];
-			ShapeToBodiesMap.Add(PShape, TheirBody);
+		ShapeToBodiesMap.Add(PShape, FWeldInfo(TheirBody, RelativeTM));
 	}
 
 	PostShapeChange();
@@ -1968,23 +1975,14 @@ namespace EScaleMode
 }
 
 //computes the relative scaling vectors based on scale mode used
-void ComputeScalingVectors(EScaleMode::Type ScaleMode, const FVector& NewScale3D, const FVector& OldScale3D, FVector& RelativeScale3D, FVector& RelativeScale3DAbs, FVector& OutScale3D)
+void ComputeScalingVectors(EScaleMode::Type ScaleMode, const FVector& NewScale3D, FVector& OutScale3D, FVector& OutScale3DAbs)
 {
-	FVector NewScale3DAbs = NewScale3D.GetAbs();
-	FVector OldScale3DAbs = OldScale3D.GetAbs();
-
+	const FVector NewScale3DAbs = NewScale3D.GetAbs();
 	switch (ScaleMode)
 	{
 	case EScaleMode::Free:
 	{
 		OutScale3D = NewScale3D;
-		RelativeScale3DAbs.X = NewScale3DAbs.X / OldScale3DAbs.X;
-		RelativeScale3DAbs.Y = NewScale3DAbs.Y / OldScale3DAbs.Y;
-		RelativeScale3DAbs.Z = NewScale3DAbs.Z / OldScale3DAbs.Z;
-
-		RelativeScale3D.X = NewScale3D.X / OldScale3D.X;
-		RelativeScale3D.Y = NewScale3D.Y / OldScale3D.Y;
-		RelativeScale3D.Z = NewScale3D.Z / OldScale3D.Z;
 		break;
 	}
 	case EScaleMode::LockedXY:
@@ -1992,18 +1990,6 @@ void ComputeScalingVectors(EScaleMode::Type ScaleMode, const FVector& NewScale3D
 		float XYScaleAbs = FMath::Max(NewScale3DAbs.X, NewScale3DAbs.Y);
 		float XYScale = FMath::Max(NewScale3D.X, NewScale3D.Y) < 0.f ? -XYScaleAbs : XYScaleAbs;	//if both xy are negative we should make the xy scale negative
 
-		float OldXYScaleAbs = FMath::Max(OldScale3DAbs.X, OldScale3DAbs.Y);
-		float OldScaleXY = FMath::Max(OldScale3D.X, OldScale3D.Y) < 0.f ? -OldXYScaleAbs : OldXYScaleAbs;
-
-		float RelativeScaleAbs = XYScaleAbs / OldXYScaleAbs;
-		float RelativeScale = XYScale / OldScaleXY;
-
-		RelativeScale3DAbs.X = RelativeScale3DAbs.Y = RelativeScaleAbs;
-		RelativeScale3DAbs.Z = NewScale3DAbs.Z / OldScale3DAbs.Z;
-
-		RelativeScale3D.X = RelativeScale3D.Y = RelativeScale;
-		RelativeScale3D.Z = NewScale3D.Z / OldScale3D.Z;
-		
 		OutScale3D = NewScale3D;
 		OutScale3D.X = OutScale3D.Y = XYScale;
 
@@ -2014,15 +2000,6 @@ void ComputeScalingVectors(EScaleMode::Type ScaleMode, const FVector& NewScale3D
 		float UniformScaleAbs = NewScale3DAbs.GetMin();	//uniform scale uses the smallest magnitude
 		float UniformScale = FMath::Max3(NewScale3D.X, NewScale3D.Y, NewScale3D.Z) < 0.f ? -UniformScaleAbs : UniformScaleAbs;	//if all three values are negative we should make uniform scale negative
 
-		float OldUniformScaleAbs = OldScale3D.GetAbs().GetMin();
-		float OldUniformScale = FMath::Max3(OldScale3D.X, OldScale3D.Y, OldScale3D.Z) < 0.f ? -OldUniformScaleAbs : OldUniformScaleAbs;
-
-		float RelativeScale = UniformScale / OldUniformScale;
-		float RelativeScaleAbs = UniformScaleAbs / OldUniformScaleAbs;
-
-		RelativeScale3DAbs = FVector(RelativeScaleAbs);
-		RelativeScale3D = FVector(RelativeScale);
-
 		OutScale3D = FVector(UniformScale);
 		break;
 	}
@@ -2031,12 +2008,35 @@ void ComputeScalingVectors(EScaleMode::Type ScaleMode, const FVector& NewScale3D
 		check(false);	//invalid scale mode
 	}
 	}
+
+	OutScale3DAbs = OutScale3D.GetAbs();
 }
 
-bool FBodyInstance::UpdateBodyScale(const FVector& InScale3D)
+EScaleMode::Type ComputeScaleMode(const TArray<PxShape*>& PShapes)
 {
-	FVector InScale3DAdjusted = InScale3D;
+	EScaleMode::Type ScaleMode = EScaleMode::Free;
 
+	for (int32 ShapeIdx = 0; ShapeIdx < PShapes.Num(); ++ShapeIdx)
+	{
+		PxShape* PShape = PShapes[ShapeIdx];
+		PxGeometryType::Enum GeomType = PShape->getGeometryType();
+
+		if (GeomType == PxGeometryType::eSPHERE)
+		{
+			ScaleMode = EScaleMode::LockedXYZ;	//sphere is most restrictive so we can stop
+			break;
+		}
+		else if (GeomType == PxGeometryType::eCAPSULE)
+		{
+			ScaleMode = EScaleMode::LockedXY;
+		}
+	}
+
+	return ScaleMode;
+}
+
+bool FBodyInstance::UpdateBodyScale(const FVector& InScale3D)
+{
 	if (!IsValidBodyInstance())
 	{
 		//UE_LOG(LogPhysics, Log, TEXT("Body hasn't been initialized. Call InitBody to initialize."));
@@ -2054,44 +2054,16 @@ bool FBodyInstance::UpdateBodyScale(const FVector& InScale3D)
 #if !(UE_BUILD_SHIPPING || UE_BUILD_TEST)
 	ensure ( !Scale3D.ContainsNaN() && !InScale3D.ContainsNaN() );
 #endif
-	FVector OldScale3D = Scale3D;
-	
-	//we never want to hit a scale of 0
-	//But we still want to be able to cross from positive to negative
-	InScale3DAdjusted.X = AdjustForSmallThreshold(InScale3D.X, OldScale3D.X);
-	InScale3DAdjusted.Y = AdjustForSmallThreshold(InScale3D.Y, OldScale3D.Y);
-	InScale3DAdjusted.Z = AdjustForSmallThreshold(InScale3D.Z, OldScale3D.Z);
-	
-	//Make sure OldScale3D is not too small or NaNs can happen
-	OldScale3D.X = OldScale3D.X < 0.1f ? 0.1f : OldScale3D.X;
-	OldScale3D.Y = OldScale3D.Y < 0.1f ? 0.1f : OldScale3D.Y;
-	OldScale3D.Z = OldScale3D.Z < 0.1f ? 0.1f : OldScale3D.Z;
 
-	// Determine the scaling mode
-	EScaleMode::Type ScaleMode = EScaleMode::Free;
 	FVector UpdatedScale3D;
 #if WITH_PHYSX
 	//Get all shapes
+	EScaleMode::Type ScaleMode = EScaleMode::Free;
 	ExecuteOnPhysicsReadWrite([&]
 	{
 		TArray<PxShape *> PShapes;
 		GetAllShapes_AssumesLocked(PShapes);
-
-		for (int32 ShapeIdx = 0; ShapeIdx < PShapes.Num(); ++ShapeIdx)
-		{
-			PxShape* PShape = PShapes[ShapeIdx];
-			PxGeometryType::Enum GeomType = PShape->getGeometryType();
-
-			if (GeomType == PxGeometryType::eSPHERE)
-			{
-				ScaleMode = EScaleMode::LockedXYZ;	//sphere is most restrictive so we can stop
-				break;
-			}
-			else if (GeomType == PxGeometryType::eCAPSULE)
-			{
-				ScaleMode = EScaleMode::LockedXY;
-			}
-		}
+		ScaleMode = ComputeScaleMode(PShapes);
 #endif
 #if WITH_BOX2D
 		if (BodyInstancePtr)
@@ -2099,10 +2071,9 @@ bool FBodyInstance::UpdateBodyScale(const FVector& InScale3D)
 			//@TODO: BOX2D: UpdateBodyScale is not implemented yet
 		}
 #endif
-
-		FVector RelativeScale3D;
-		FVector RelativeScale3DAbs;
-		ComputeScalingVectors(ScaleMode, InScale3DAdjusted, OldScale3D, RelativeScale3D, RelativeScale3DAbs, UpdatedScale3D);
+		FVector AdjustedScale3D;
+		FVector AdjustedScale3DAbs;
+		ComputeScalingVectors(ScaleMode, InScale3D, AdjustedScale3D, AdjustedScale3DAbs);
 
 		// Apply scaling
 #if WITH_PHYSX
@@ -2113,26 +2084,28 @@ bool FBodyInstance::UpdateBodyScale(const FVector& InScale3D)
 		PxConvexMeshGeometry PConvexGeom;
 		PxTriangleMeshGeometry PTriMeshGeom;
 
-		for (int32 ShapeIdx = 0; ShapeIdx < PShapes.Num(); ShapeIdx++)
+		for (PxShape* PShape : PShapes)
 		{
 			bool bInvalid = false;	//we only mark invalid if actually found geom and it's invalid scale
 			PxGeometry* UpdatedGeometry = NULL;
-			PxShape* PShape = PShapes[ShapeIdx];
-
 			PxTransform PLocalPose = PShape->getLocalPose();
-			PLocalPose.q.normalize();
+
 			PxGeometryType::Enum GeomType = PShape->getGeometryType();
+			FKShapeElem* ShapeElem = FPhysxUserData::Get<FKShapeElem>(PShape->userData);
+			const FTransform& RelativeTM = GetRelativeBodyTransform(PShape);
 
 			switch (GeomType)
 			{
 				case PxGeometryType::eSPHERE:
 				{
+					FKSphereElem* SphereElem = ShapeElem->GetShapeCheck<FKSphereElem>();
 					ensure(ScaleMode == EScaleMode::LockedXYZ);
 
 					PShape->getSphereGeometry(PSphereGeom);
-
-					PSphereGeom.radius *= RelativeScale3DAbs.X;
-					PLocalPose.p *= RelativeScale3D.X;
+					 
+					PSphereGeom.radius = SphereElem->Radius * AdjustedScale3DAbs.X;
+					PLocalPose.p = U2PVector(RelativeTM.TransformPosition(SphereElem->Center));
+					PLocalPose.p *= AdjustedScale3D.X;
 
 					if (PSphereGeom.isValid())
 					{
@@ -2147,14 +2120,18 @@ bool FBodyInstance::UpdateBodyScale(const FVector& InScale3D)
 				}
 				case PxGeometryType::eBOX:
 				{
+					FKBoxElem* BoxElem = ShapeElem->GetShapeCheck<FKBoxElem>();
 					PShape->getBoxGeometry(PBoxGeom);
 
-					PBoxGeom.halfExtents.x *= RelativeScale3DAbs.X;
-					PBoxGeom.halfExtents.y *= RelativeScale3DAbs.Y;
-					PBoxGeom.halfExtents.z *= RelativeScale3DAbs.Z;
-					PLocalPose.p.x *= RelativeScale3D.X;
-					PLocalPose.p.y *= RelativeScale3D.Y;
-					PLocalPose.p.z *= RelativeScale3D.Z;
+					PBoxGeom.halfExtents.x = (0.5f * BoxElem->X * AdjustedScale3DAbs.X);
+					PBoxGeom.halfExtents.y = (0.5f * BoxElem->Y * AdjustedScale3DAbs.Y);
+					PBoxGeom.halfExtents.z = (0.5f * BoxElem->Z * AdjustedScale3DAbs.Z);
+
+					FTransform BoxTransform = BoxElem->GetTransform() * RelativeTM;
+					PLocalPose = PxTransform(U2PTransform(BoxTransform));
+					PLocalPose.p.x *= AdjustedScale3D.X;
+					PLocalPose.p.y *= AdjustedScale3D.Y;
+					PLocalPose.p.z *= AdjustedScale3D.Z;
 
 					if (PBoxGeom.isValid())
 					{
@@ -2169,16 +2146,30 @@ bool FBodyInstance::UpdateBodyScale(const FVector& InScale3D)
 				}
 				case PxGeometryType::eCAPSULE:
 				{
+					FKSphylElem* SphylElem = ShapeElem->GetShapeCheck<FKSphylElem>();
 					ensure(ScaleMode == EScaleMode::LockedXY || ScaleMode == EScaleMode::LockedXYZ);
 
+					float ScaleRadius = FMath::Max(AdjustedScale3DAbs.X, AdjustedScale3DAbs.Y);
+					float ScaleLength = AdjustedScale3DAbs.Z;
+
 					PShape->getCapsuleGeometry(PCapsuleGeom);
 
-					PCapsuleGeom.halfHeight *= RelativeScale3DAbs.Z;
-					PCapsuleGeom.radius *= RelativeScale3DAbs.X;
+					// this is a bit confusing since radius and height is scaled
+					// first apply the scale first 
+					float Radius = FMath::Max(SphylElem->Radius * ScaleRadius, 0.1f);
+					float Length = SphylElem->Length + SphylElem->Radius * 2.f;
+					float HalfLength = Length * ScaleLength * 0.5f;
+					Radius = FMath::Clamp(Radius, 0.1f, HalfLength);	//radius is capped by half length
+					float HalfHeight = HalfLength - Radius;
+					HalfHeight = FMath::Max(0.1f, HalfHeight);
 
-					PLocalPose.p.x *= RelativeScale3D.X;
-					PLocalPose.p.y *= RelativeScale3D.Y;
-					PLocalPose.p.z *= RelativeScale3D.Z;
+					PCapsuleGeom.halfHeight = HalfHeight;
+					PCapsuleGeom.radius = Radius;
+
+					PLocalPose = PxTransform(U2PVector(RelativeTM.TransformPosition(SphylElem->Center)), U2PQuat(SphylElem->Orientation) * U2PSphylBasis);
+					PLocalPose.p.x *= AdjustedScale3D.X;
+					PLocalPose.p.y *= AdjustedScale3D.Y;
+					PLocalPose.p.z *= AdjustedScale3D.Z;
 
 					if (PCapsuleGeom.isValid())
 					{
@@ -2194,67 +2185,54 @@ bool FBodyInstance::UpdateBodyScale(const FVector& InScale3D)
 				}
 				case PxGeometryType::eCONVEXMESH:
 				{
+					FKConvexElem* ConvexElem = ShapeElem->GetShapeCheck<FKConvexElem>();
 					PShape->getConvexMeshGeometry(PConvexGeom);
 
-					// find which convex elems it is
-					// it would be nice to know if the order of PShapes array index is in the order of createShape
-					// Create convex shapes
-					if (BodySetup.IsValid())
-					{
-						for (int32 i = 0; i < BodySetup->AggGeom.ConvexElems.Num(); i++)
-						{
-							FKConvexElem* ConvexElem = &(BodySetup->AggGeom.ConvexElems[i]);
+					bool bUseNegX = CalcMeshNegScaleCompensation(AdjustedScale3D, PLocalPose);
 
-							// found it
-							if (ConvexElem->ConvexMesh == PConvexGeom.convexMesh)
-							{
-								// Please note that this one we don't inverse old scale, but just set new one (but we still follow scale mode restriction)
-								FVector NewScale3D = RelativeScale3D * OldScale3D;
-								FVector Scale3DAbs(FMath::Abs(NewScale3D.X), FMath::Abs(NewScale3D.Y), FMath::Abs(NewScale3D.Z)); // magnitude of scale (sign removed)
+					PConvexGeom.convexMesh = bUseNegX ? ConvexElem->ConvexMeshNegX : ConvexElem->ConvexMesh;
+					PConvexGeom.scale.scale = U2PVector(AdjustedScale3DAbs * ConvexElem->GetTransform().GetScale3D().GetAbs());
+					FTransform ConvexTransform = ConvexElem->GetTransform();
 
-								PxTransform PNewLocalPose;
-								bool bUseNegX = CalcMeshNegScaleCompensation(NewScale3D, PNewLocalPose);
+					PxTransform PElementTransform = U2PTransform(ConvexTransform * RelativeTM);
+					PLocalPose.q *= PElementTransform.q;
+					PLocalPose.p = PElementTransform.p;
+					PLocalPose.p.x *= AdjustedScale3D.X;
+					PLocalPose.p.y *= AdjustedScale3D.Y;
+					PLocalPose.p.z *= AdjustedScale3D.Z;
 
-								PxTransform PElementTransform = U2PTransform(ConvexElem->GetTransform());
-								PNewLocalPose.q *= PElementTransform.q;
-								PNewLocalPose.p += PElementTransform.p;
-
-								PConvexGeom.convexMesh = bUseNegX ? ConvexElem->ConvexMeshNegX : ConvexElem->ConvexMesh;
-								PConvexGeom.scale.scale = U2PVector(Scale3DAbs);
-
-								if (PConvexGeom.isValid())
-								{
-									UpdatedGeometry = &PConvexGeom;
-									bSuccess = true;
-								}
-								else
-								{
-									bInvalid = true;
-								}
-								break;
-							}
-						}
+					if (PConvexGeom.isValid())
+					{
+						UpdatedGeometry = &PConvexGeom;
+						bSuccess = true;
+					}
+					else
+					{
+						bInvalid = true;
 					}
 
 					break;
 				}
 				case PxGeometryType::eTRIANGLEMESH:
 				{
+					check(ShapeElem == nullptr);	//trimesh shape doesn't have userData
 					PShape->getTriangleMeshGeometry(PTriMeshGeom);
 
 					// find which trimesh elems it is
 					// it would be nice to know if the order of PShapes array index is in the order of createShape
 					if (BodySetup.IsValid())
 					{
-						// Please note that this one we don't inverse old scale, but just set new one (but still adjust for scale mode)
-						FVector NewScale3D = RelativeScale3D * OldScale3D;
-
 						for (PxTriangleMesh* TriMesh : BodySetup->TriMeshes)
 						{
 							// found it
 							if (TriMesh == PTriMeshGeom.triangleMesh)
 							{
-								PTriMeshGeom.scale.scale = U2PVector(Scale3D);
+								PTriMeshGeom.scale.scale = U2PVector(AdjustedScale3D);
+
+								PLocalPose = U2PTransform(RelativeTM);
+								PLocalPose.p.x *= AdjustedScale3D.X;
+								PLocalPose.p.y *= AdjustedScale3D.Y;
+								PLocalPose.p.z *= AdjustedScale3D.Z;
 
 								if (PTriMeshGeom.isValid())
 								{
@@ -2292,7 +2270,7 @@ bool FBodyInstance::UpdateBodyScale(const FVector& InScale3D)
 			}
 			else if (bInvalid)
 			{
-				FMessageLog("PIE").Warning(FText::Format(LOCTEXT("PhysicsInvalidScale", "Scale ''{0}'' is not valid on object '{1}'."), FText::FromString(InScale3DAdjusted.ToString()), FText::FromString(GetBodyDebugName())));
+				FMessageLog("PIE").Warning(FText::Format(LOCTEXT("PhysicsInvalidScale", "Scale ''{0}'' is not valid on object '{1}'."), FText::FromString(AdjustedScale3D.ToString()), FText::FromString(GetBodyDebugName())));
 			}
 		}
 	});
diff --git a/Engine/Source/Runtime/Engine/Private/PhysicsEngine/BodySetup.cpp b/Engine/Source/Runtime/Engine/Private/PhysicsEngine/BodySetup.cpp
--- a/Engine/Source/Runtime/Engine/Private/PhysicsEngine/BodySetup.cpp
+++ b/Engine/Source/Runtime/Engine/Private/PhysicsEngine/BodySetup.cpp
@@ -6,6 +6,7 @@
 
 #include "EnginePrivate.h"
 #include "PhysicsPublic.h"
+#include "PhysicsEngine/BodySetup.h"
 #include "TargetPlatform.h"
 #include "AnimTree.h"
 
@@ -18,12 +19,9 @@
 #include "PhysicalMaterials/PhysicalMaterial.h"
 
 #if WITH_PHYSX
-	namespace
-	{
-		// Quaternion that converts Sphyls from UE space to PhysX space (negate Y, swap X & Z)
-		// This is equivalent to a 180 degree rotation around the normalized (1, 0, 1) axis
-		static const PxQuat U2PSphylBasis( PI, PxVec3( 1.0f / FMath::Sqrt( 2.0f ), 0.0f, 1.0f / FMath::Sqrt( 2.0f ) ) );
-	}
+	// Quaternion that converts Sphyls from UE space to PhysX space (negate Y, swap X & Z)
+	// This is equivalent to a 180 degree rotation around the normalized (1, 0, 1) axis
+	const physx::PxQuat U2PSphylBasis( PI, PxVec3( 1.0f / FMath::Sqrt( 2.0f ), 0.0f, 1.0f / FMath::Sqrt( 2.0f ) ) );
 #endif // WITH_PHYSX
 
 // CVars
@@ -372,7 +370,7 @@ public:
 				ensure(PLocalPose.isValid());
 				{
 					const float ContactOffset = FMath::Min(MaxContactOffset, ContactOffsetFactor * PSphereGeom.radius);
-					AttachShape_AssumesLocked(PSphereGeom, PLocalPose, ContactOffset);
+					PxShape* PShape = AttachShape_AssumesLocked(PSphereGeom, PLocalPose, ContactOffset, &SphereElem);
 				}
 			}
 			else
@@ -407,7 +405,7 @@ public:
 				ensure(PLocalPose.isValid());
 				{
 					const float ContactOffset = FMath::Min(MaxContactOffset, ContactOffsetFactor * PBoxGeom.halfExtents.minElement());
-					AttachShape_AssumesLocked(PBoxGeom, PLocalPose, ContactOffset);
+					AttachShape_AssumesLocked(PBoxGeom, PLocalPose, ContactOffset, &BoxElem);
 				}
 			}
 			else
@@ -453,7 +451,7 @@ public:
 				ensure(PLocalPose.isValid());
 				{
 					const float ContactOffset = FMath::Min(MaxContactOffset, ContactOffsetFactor * PCapsuleGeom.radius);
-					AttachShape_AssumesLocked(PCapsuleGeom, PLocalPose, ContactOffset);
+					AttachShape_AssumesLocked(PCapsuleGeom, PLocalPose, ContactOffset, &SphylElem);
 				}
 			}
 			else
@@ -501,7 +499,7 @@ public:
 						ensure(PLocalPose.isValid());
 						{
 							const float ContactOffset = FMath::Min(MaxContactOffset, ContactOffsetFactor * PBoundsExtents.minElement());
-							AttachShape_AssumesLocked(PConvexGeom, PLocalPose, ContactOffset);
+							AttachShape_AssumesLocked(PConvexGeom, PLocalPose, ContactOffset, &ConvexElem);
 						}
 					}
 					else
@@ -549,7 +547,7 @@ public:
 				PElementTransform.p.z *= Scale3D.Z;
 
 				// Create without 'sim shape' flag, problematic if it's kinematic, and it gets set later anyway.
-				if (!AttachShape_AssumesLocked(PTriMeshGeom, PElementTransform, MaxContactOffset, PxShapeFlag::eSCENE_QUERY_SHAPE | PxShapeFlag::eVISUALIZATION))
+				if (!AttachShape_AssumesLocked(PTriMeshGeom, PElementTransform, MaxContactOffset, nullptr, PxShapeFlag::eSCENE_QUERY_SHAPE | PxShapeFlag::eVISUALIZATION))
 				{
 					UE_LOG(LogPhysics, Log, TEXT("Can't create new mesh shape in AddShapesToRigidActor"));
 				}
@@ -563,13 +561,14 @@ public:
 
 private:
 
-	PxShape* AttachShape_AssumesLocked(const PxGeometry& PGeom, const PxTransform& PLocalPose, const float ContactOffset, PxShapeFlags PShapeFlags = PxShapeFlag::eVISUALIZATION | PxShapeFlag::eSCENE_QUERY_SHAPE | PxShapeFlag::eSIMULATION_SHAPE) const
+	PxShape* AttachShape_AssumesLocked(const PxGeometry& PGeom, const PxTransform& PLocalPose, const float ContactOffset, const FKShapeElem* ShapeElem, PxShapeFlags PShapeFlags = PxShapeFlag::eVISUALIZATION | PxShapeFlag::eSCENE_QUERY_SHAPE | PxShapeFlag::eSIMULATION_SHAPE) const
 	{
 		const PxMaterial* PMaterial = GetDefaultPhysMaterial(); 
 		PxShape* PNewShape = bShapeSharing ? GPhysXSDK->createShape(PGeom, *PMaterial, /*isExclusive =*/ false, PShapeFlags) : PDestActor->createShape(PGeom, *PMaterial, PShapeFlags);
 
 		if (PNewShape)
 		{
+			PNewShape->userData = ShapeElem ? ((void*) ShapeElem->GetUserData()) : nullptr;
 			PNewShape->setLocalPose(PLocalPose);
 
 			if (NewShapes)
diff --git a/Engine/Source/Runtime/Engine/Private/PhysicsEngine/PhysXSupport.h b/Engine/Source/Runtime/Engine/Private/PhysicsEngine/PhysXSupport.h
--- a/Engine/Source/Runtime/Engine/Private/PhysicsEngine/PhysXSupport.h
+++ b/Engine/Source/Runtime/Engine/Private/PhysicsEngine/PhysXSupport.h
@@ -448,6 +448,11 @@ extern ENGINE_API TArray<PxHeightField*>	GPhysXPendingKillHeightfield;
 /** Array of PxMaterial objects which are awaiting cleaning up. */
 extern TArray<PxMaterial*>		GPhysXPendingKillMaterial;
 
+
+#if WITH_PHYSX
+extern const physx::PxQuat U2PSphylBasis;
+#endif // WITH_PHYSX
+
 /** Utility class to keep track of shared physics data */
 class FPhysxSharedData
 {
diff --git a/Engine/Source/Runtime/Engine/Public/PhysxUserData.h b/Engine/Source/Runtime/Engine/Public/PhysxUserData.h
--- a/Engine/Source/Runtime/Engine/Public/PhysxUserData.h
+++ b/Engine/Source/Runtime/Engine/Public/PhysxUserData.h
@@ -11,6 +11,7 @@ struct FDestructibleChunkInfo;
 class FPhysScene;
 class UPhysicalMaterial;
 class UPrimitiveComponent;
+struct FKShapeElem;
 
 /** PhysX user data type*/
 namespace EPhysxUserDataType
@@ -24,6 +25,7 @@ namespace EPhysxUserDataType
 		ConstraintInstance,
 		PrimitiveComponent,
 		DestructibleChunk,
+		AggShape,
 	};
 };
 
@@ -42,12 +44,13 @@ public:
 	FPhysxUserData(FConstraintInstance* InPayload)		:Type(EPhysxUserDataType::ConstraintInstance), Payload(InPayload) {}
 	FPhysxUserData(UPrimitiveComponent* InPayload)		:Type(EPhysxUserDataType::PrimitiveComponent), Payload(InPayload) {}
 	FPhysxUserData(FDestructibleChunkInfo* InPayload)	:Type(EPhysxUserDataType::DestructibleChunk), Payload(InPayload) {}
+	FPhysxUserData(FKShapeElem* InPayload)				:Type(EPhysxUserDataType::AggShape), Payload(InPayload) {}
 	
 	template <class T> static T* Get(void* UserData);
 	template <class T> static void Set(void* UserData, T* Payload);
 
 	//helper function to determine if userData is garbage (maybe dangling pointer)
-	static bool IsGarbage(void* UserData){ return ((FPhysxUserData*)UserData)->Type < EPhysxUserDataType::Invalid || ((FPhysxUserData*)UserData)->Type > EPhysxUserDataType::DestructibleChunk; }
+	static bool IsGarbage(void* UserData){ return ((FPhysxUserData*)UserData)->Type < EPhysxUserDataType::Invalid || ((FPhysxUserData*)UserData)->Type > EPhysxUserDataType::AggShape; }
 };
 
 template <> FORCEINLINE FBodyInstance* FPhysxUserData::Get(void* UserData)			{ if (!UserData || ((FPhysxUserData*)UserData)->Type != EPhysxUserDataType::BodyInstance) { return NULL; } return (FBodyInstance*)((FPhysxUserData*)UserData)->Payload; }
@@ -56,6 +59,7 @@ template <> FORCEINLINE FPhysScene* FPhysxUserData::Get(void* UserData)				{ if
 template <> FORCEINLINE FConstraintInstance* FPhysxUserData::Get(void* UserData)	{ if (!UserData || ((FPhysxUserData*)UserData)->Type != EPhysxUserDataType::ConstraintInstance) { return NULL; } return (FConstraintInstance*)((FPhysxUserData*)UserData)->Payload; }
 template <> FORCEINLINE UPrimitiveComponent* FPhysxUserData::Get(void* UserData)	{ if (!UserData || ((FPhysxUserData*)UserData)->Type != EPhysxUserDataType::PrimitiveComponent) { return NULL; } return (UPrimitiveComponent*)((FPhysxUserData*)UserData)->Payload; }
 template <> FORCEINLINE FDestructibleChunkInfo* FPhysxUserData::Get(void* UserData)	{ if (!UserData || ((FPhysxUserData*)UserData)->Type != EPhysxUserDataType::DestructibleChunk) { return NULL; } return (FDestructibleChunkInfo*)((FPhysxUserData*)UserData)->Payload; }
+template <> FORCEINLINE FKShapeElem* FPhysxUserData::Get(void* UserData)	{ if (!UserData || ((FPhysxUserData*)UserData)->Type != EPhysxUserDataType::AggShape) { return NULL; } return (FKShapeElem*)((FPhysxUserData*)UserData)->Payload; }
 
 template <> FORCEINLINE void FPhysxUserData::Set(void* UserData, FBodyInstance* Payload)			{ check(UserData); ((FPhysxUserData*)UserData)->Type = EPhysxUserDataType::BodyInstance; ((FPhysxUserData*)UserData)->Payload = Payload; }
 template <> FORCEINLINE void FPhysxUserData::Set(void* UserData, UPhysicalMaterial* Payload)		{ check(UserData); ((FPhysxUserData*)UserData)->Type = EPhysxUserDataType::PhysicalMaterial; ((FPhysxUserData*)UserData)->Payload = Payload; }
@@ -63,6 +67,7 @@ template <> FORCEINLINE void FPhysxUserData::Set(void* UserData, FPhysScene* Pay
 template <> FORCEINLINE void FPhysxUserData::Set(void* UserData, FConstraintInstance* Payload)		{ check(UserData); ((FPhysxUserData*)UserData)->Type = EPhysxUserDataType::ConstraintInstance; ((FPhysxUserData*)UserData)->Payload = Payload; }
 template <> FORCEINLINE void FPhysxUserData::Set(void* UserData, UPrimitiveComponent* Payload)		{ check(UserData); ((FPhysxUserData*)UserData)->Type = EPhysxUserDataType::PrimitiveComponent; ((FPhysxUserData*)UserData)->Payload = Payload; }
 template <> FORCEINLINE void FPhysxUserData::Set(void* UserData, FDestructibleChunkInfo* Payload)	{ check(UserData); ((FPhysxUserData*)UserData)->Type = EPhysxUserDataType::DestructibleChunk; ((FPhysxUserData*)UserData)->Payload = Payload; }
+template <> FORCEINLINE void FPhysxUserData::Set(void* UserData, FKShapeElem* Payload)	{ check(UserData); ((FPhysxUserData*)UserData)->Type = EPhysxUserDataType::AggShape; ((FPhysxUserData*)UserData)->Payload = Payload; }
 
 
 #endif // WITH_PHYSICS
\ No newline at end of file
