Fixed navmesh rendering flicker for a moment when switching between tabs of the editor #UE4

diff --git a/Engine/Source/Runtime/Engine/Private/Components/NavMeshRenderingComponent.cpp b/Engine/Source/Runtime/Engine/Private/Components/NavMeshRenderingComponent.cpp
--- a/Engine/Source/Runtime/Engine/Private/Components/NavMeshRenderingComponent.cpp
+++ b/Engine/Source/Runtime/Engine/Private/Components/NavMeshRenderingComponent.cpp
@@ -103,7 +103,7 @@ void UNavMeshRenderingComponent::TimerFunction()
 
 	const bool bShowNavigation = IsNavigationShowFlagSet(World);
 
-	if (bShowNavigation != !!bCollectNavigationData)
+	if (bShowNavigation != !!bCollectNavigationData && bShowNavigation == true)
 	{
 		bCollectNavigationData = bShowNavigation;
 		MarkRenderStateDirty();
