UE-16863 Crash upon using the undo shortcut twice after expanding a Math Expression node #2

diff --git a/Engine/Source/Editor/Kismet/Private/FindInBlueprintManager.cpp b/Engine/Source/Editor/Kismet/Private/FindInBlueprintManager.cpp
--- a/Engine/Source/Editor/Kismet/Private/FindInBlueprintManager.cpp
+++ b/Engine/Source/Editor/Kismet/Private/FindInBlueprintManager.cpp
@@ -615,7 +615,10 @@ namespace BlueprintSearchMetaDataHelpers
 					InWriter->WriteObjectStart();
 
 					FGraphDisplayInfo DisplayInfo;
-					Graph->GetSchema()->GetGraphDisplayInformation(*Graph, DisplayInfo);
+					if (auto GraphSchema = Graph->GetSchema())
+					{
+						GraphSchema->GetGraphDisplayInformation(*Graph, DisplayInfo);
+					}
 					InWriter->WriteValue(FFindInBlueprintSearchTags::FiB_Name, DisplayInfo.PlainName);
 
 					FText GraphDescription = FBlueprintEditorUtils::GetGraphDescription(Graph);
diff --git a/Engine/Source/Runtime/Engine/Private/EdGraph/EdGraph.cpp b/Engine/Source/Runtime/Engine/Private/EdGraph/EdGraph.cpp
--- a/Engine/Source/Runtime/Engine/Private/EdGraph/EdGraph.cpp
+++ b/Engine/Source/Runtime/Engine/Private/EdGraph/EdGraph.cpp
@@ -231,9 +231,11 @@ void UEdGraph::GetAllChildrenGraphs(TArray<UEdGraph*>& Graphs) const
 	for (int32 i = 0; i < SubGraphs.Num(); ++i)
 	{
 		UEdGraph* Graph = SubGraphs[i];
-		checkf(Graph, *FString::Printf(TEXT("%s has invalid SubGraph array entry at %d"), *GetFullName(), i));
-		Graphs.Add(Graph);
-		Graph->GetAllChildrenGraphs(Graphs);
+		if (ensureMsgf(Graph, TEXT("%s has invalid SubGraph array entry at %d"), *GetFullName(), i))
+		{
+			Graphs.Add(Graph);
+			Graph->GetAllChildrenGraphs(Graphs);
+		}
 	}
 #endif // WITH_EDITORONLY_DATA
 }
