Fix crash trying to display subtiles for sounds not associated with any UWorld.  This can happen frequently in the editor

diff --git a/Engine/Source/Runtime/Engine/Private/SoundWave.cpp b/Engine/Source/Runtime/Engine/Private/SoundWave.cpp
--- a/Engine/Source/Runtime/Engine/Private/SoundWave.cpp
+++ b/Engine/Source/Runtime/Engine/Private/SoundWave.cpp
@@ -529,7 +529,7 @@ FWaveInstance* USoundWave::HandleStart( FActiveSound& ActiveSound, const UPTRINT
 				// intercept the subtitles if the delegate is set
 				ActiveSound.AudioComponent->OnQueueSubtitles.ExecuteIfBound( Subtitles, Duration );
 			}
-			else
+			else if( ActiveSound.World.IsValid() )
 			{
 				// otherwise, pass them on to the subtitle manager for display
 				// Subtitles are hashed based on the associated sound (wave instance).
