Backed out changelist 2545467. Matinee stop is no longer a deferred operation; instead extra checks are performed while updating a group to ensure that the sequence is still being played.

diff --git a/Engine/Source/Runtime/Engine/Classes/Matinee/MatineeActor.h b/Engine/Source/Runtime/Engine/Classes/Matinee/MatineeActor.h
--- a/Engine/Source/Runtime/Engine/Classes/Matinee/MatineeActor.h
+++ b/Engine/Source/Runtime/Engine/Classes/Matinee/MatineeActor.h
@@ -188,9 +188,6 @@ public:
 	UPROPERTY(replicated,transient)
 	uint32 bPaused:1;
 
-	UPROPERTY(replicated,transient)
-	uint32 bPendingStop:1;
-
 	UPROPERTY(replicated)
 	float InterpPosition;
 
diff --git a/Engine/Source/Runtime/Engine/Private/Interpolation.cpp b/Engine/Source/Runtime/Engine/Private/Interpolation.cpp
--- a/Engine/Source/Runtime/Engine/Private/Interpolation.cpp
+++ b/Engine/Source/Runtime/Engine/Private/Interpolation.cpp
@@ -482,9 +482,16 @@ void AMatineeActor::Reverse()
 
 void AMatineeActor::Stop()
 {
-	if (bIsPlaying)
+	// Re-enable the radio filter
+	EnableRadioFilter();
+
+	bIsPlaying = false;
+	bPaused = false;
+
+	if (GetWorld()->IsGameWorld())
 	{
-		bPendingStop = true;
+		// We should only terminate the interp in the game.  The editor handles this from inside the matinee editor
+		TermInterp();
 	}
 }
 
@@ -610,22 +617,6 @@ void AMatineeActor::Tick(float DeltaTime)
 	{
 		StepInterp(DeltaTime, false);
 	}
-
-	if (bPendingStop)
-	{
-		// Re-enable the radio filter
-		EnableRadioFilter();
-
-		bIsPlaying = false;
-		bPaused = false;
-		bPendingStop = false;
-
-		if (GetWorld()->IsGameWorld())
-		{
-			// We should only terminate the interp in the game.  The editor handles this from inside the matinee editor
-			TermInterp();
-		}
-	}
 }
 
 
@@ -750,9 +741,22 @@ void AMatineeActor::UpdateInterp( float NewPosition, bool bPreview, bool bJump )
 			for( int32 GroupIndex = 0; GroupIndex < Groups.Num(); ++GroupIndex )
 			{
 				Groups[GroupIndex]->Group->UpdateGroup( NewPosition, Groups[GroupIndex], bPreview, bJump );
+
+				const bool bhasBeenTerminated = (GroupInst.Num() == 0);
+#if WITH_EDITORONLY_DATA
+				if (bhasBeenTerminated && !bIsBeingEdited)
+#else
+				if (bhasBeenTerminated)
+#endif
+				{
+					UE_LOG(LogMatinee, Log, TEXT("WARNING: A matinee was stopped while updating group '%s'; the next groups will not be updated."), *Groups[GroupIndex]->Group->GetFullGroupName(true));
+					InterpPosition = NewPosition;
+					return;
+				}
 			}
 		}
 
+
 		InterpPosition = NewPosition;
 	}
 }
@@ -2038,7 +2042,10 @@ void UInterpGroup::UpdateGroup(float NewPosition, UInterpGroupInst* GrInst, bool
 		UpdateAnimWeights(NewPosition, GrInst, bPreview, bJump);
 	}
 #endif
-	for(int32 i=0; i<InterpTracks.Num(); i++)
+
+	AMatineeActor* MatineeActor = Cast<AMatineeActor>(GrInst->GetOuter());
+
+	for(int32 i=0; i<InterpTracks.Num() && MatineeActor->bIsPlaying; i++)
 	{
 		UInterpTrack* Track = InterpTracks[i];
 		UInterpTrackInst* TrInst = GrInst->TrackInst[i];
