Start sounds from the beginning when playing matinee.  Only use time when jumping around.

diff --git a/Engine/Source/Runtime/Engine/Private/Interpolation.cpp b/Engine/Source/Runtime/Engine/Private/Interpolation.cpp
--- a/Engine/Source/Runtime/Engine/Private/Interpolation.cpp
+++ b/Engine/Source/Runtime/Engine/Private/Interpolation.cpp
@@ -7921,7 +7921,7 @@ void UInterpTrackSound::UpdateTrack(float NewPosition, UInterpTrackInst* TrInst,
 					SoundInst->PlayAudioComp->SetVolumeMultiplier(VolumePitchValue.X);
 					SoundInst->PlayAudioComp->SetPitchMultiplier(VolumePitchValue.Y);
 					SoundInst->PlayAudioComp->SubtitlePriority = bSuppressSubtitles ? 0.f : SUBTITLE_PRIORITY_MATINEE;
-					SoundInst->PlayAudioComp->Play(NewPosition - SoundTrackKey.Time);
+					SoundInst->PlayAudioComp->Play((bJump ? NewPosition - SoundTrackKey.Time : 0.f));
 				}
 				else
 				{
@@ -7951,7 +7951,7 @@ void UInterpTrackSound::UpdateTrack(float NewPosition, UInterpTrackInst* TrInst,
 						SoundInst->PlayAudioComp->SetVolumeMultiplier(VolumePitchValue.X);
 						SoundInst->PlayAudioComp->SetPitchMultiplier(VolumePitchValue.Y);
 						SoundInst->PlayAudioComp->SubtitlePriority = bSuppressSubtitles ? 0.f : SUBTITLE_PRIORITY_MATINEE;
-						SoundInst->PlayAudioComp->Play(NewPosition - SoundTrackKey.Time);
+						SoundInst->PlayAudioComp->Play((bJump ? NewPosition - SoundTrackKey.Time : 0.f));
 					}
 				}
 			}
