#ue4 - Fix SlideAlongSurface pushing out away from the original direction of movement when hitting multiple surfaces.

UDN: https://udn.unrealengine.com/questions/249142/474-character-jitters-while-moving-towards-certain.html

diff --git a/Engine/Source/Runtime/Engine/Private/Components/MovementComponent.cpp b/Engine/Source/Runtime/Engine/Private/Components/MovementComponent.cpp
--- a/Engine/Source/Runtime/Engine/Private/Components/MovementComponent.cpp
+++ b/Engine/Source/Runtime/Engine/Private/Components/MovementComponent.cpp
@@ -639,8 +639,8 @@ float UMovementComponent::SlideAlongSurface(const FVector& Delta, float Time, co
 			// Compute new slide normal when hitting multiple surfaces.
 			TwoWallAdjust(SlideDelta, Hit, OldHitNormal);
 
-			// Only proceed if the new direction is of significant length.
-			if (!SlideDelta.IsNearlyZero(1e-3f))
+			// Only proceed if the new direction is of significant length and not in reverse of original attempted move.
+			if (!SlideDelta.IsNearlyZero(1e-3f) && (SlideDelta | Delta) > 0.f)
 			{
 				// Perform second move
 				SafeMoveUpdatedComponent(SlideDelta, Rotation, true, Hit);
