#ue4 - Perf: UMovementComponent::SlideAlongSurface internally uses FQuat instead of FRotator.

diff --git a/Engine/Source/Runtime/Engine/Private/Components/MovementComponent.cpp b/Engine/Source/Runtime/Engine/Private/Components/MovementComponent.cpp
--- a/Engine/Source/Runtime/Engine/Private/Components/MovementComponent.cpp
+++ b/Engine/Source/Runtime/Engine/Private/Components/MovementComponent.cpp
@@ -623,7 +623,7 @@ float UMovementComponent::SlideAlongSurface(const FVector& Delta, float Time, co
 
 	if ((SlideDelta | Delta) > 0.f)
 	{
-		const FRotator Rotation = UpdatedComponent->GetComponentRotation();
+		const FQuat Rotation = UpdatedComponent->GetComponentQuat();
 		SafeMoveUpdatedComponent(SlideDelta, Rotation, true, Hit);
 
 		const float FirstHitPercent = Hit.Time;
