#ue4 - Fix NavWalking movement mode replicating to proxies as Custom movement. Added GetNavData() helper to see if there is NavMesh available.

diff --git a/Engine/Source/Runtime/Engine/Classes/GameFramework/CharacterMovementComponent.h b/Engine/Source/Runtime/Engine/Classes/GameFramework/CharacterMovementComponent.h
--- a/Engine/Source/Runtime/Engine/Classes/GameFramework/CharacterMovementComponent.h
+++ b/Engine/Source/Runtime/Engine/Classes/GameFramework/CharacterMovementComponent.h
@@ -1165,6 +1165,9 @@ protected:
 	/** Switch collision settings for NavWalking mode (ignore world collisions) */
 	virtual void SetNavWalkingPhysics(bool bEnable);
 
+	/** Get Navigation data for the Character. Returns null if there is no associated nav data. */
+	const class ANavigationData* UCharacterMovementComponent::GetNavData() const;
+
 	/** 
 	 * Checks to see if the current location is not encroaching blocking geometry so the character can leave NavWalking.
 	 * Restores collision settings and adjusts character location to avoid getting stuck in geometry.
diff --git a/Engine/Source/Runtime/Engine/Private/Components/CharacterMovementComponent.cpp b/Engine/Source/Runtime/Engine/Private/Components/CharacterMovementComponent.cpp
--- a/Engine/Source/Runtime/Engine/Private/Components/CharacterMovementComponent.cpp
+++ b/Engine/Source/Runtime/Engine/Private/Components/CharacterMovementComponent.cpp
@@ -627,6 +627,15 @@ void UCharacterMovementComponent::SetMovementMode(EMovementMode NewMovementMode,
 		NewCustomMode = 0;
 	}
 
+	// If trying to use NavWalking but there is no navmesh, use walking instead.
+	if (NewMovementMode == MOVE_NavWalking)
+	{
+		if (GetNavData() == nullptr)
+		{
+			NewMovementMode = MOVE_Walking;
+		}
+	}
+
 	// Do nothing if nothing is changing.
 	if (MovementMode == NewMovementMode)
 	{
@@ -728,12 +737,13 @@ void UCharacterMovementComponent::OnMovementModeChanged(EMovementMode PreviousMo
 
 uint8 UCharacterMovementComponent::PackNetworkMovementMode() const
 {
-	const uint32 GroundShift = FMath::FloorLog2(MOVE_MAX) + 1;
+	const uint32 GroundShift = FMath::CeilLogTwo(MOVE_MAX);
 	const uint8 CustomModeThr = 2 * (1 << GroundShift);
 
 	if (MovementMode != MOVE_Custom)
 	{
-		return uint8(MovementMode.GetValue()) | (GroundMovementMode.GetValue() << GroundShift);
+		const uint8 GroundModeBit = (GroundMovementMode == MOVE_Walking ? 0 : 1);
+		return uint8(MovementMode.GetValue()) | (GroundModeBit << GroundShift);
 	}
 	else
 	{
@@ -744,7 +754,7 @@ uint8 UCharacterMovementComponent::PackNetworkMovementMode() const
 
 void UCharacterMovementComponent::UnpackNetworkMovementMode(const uint8 ReceivedMode, TEnumAsByte<EMovementMode>& OutMode, uint8& OutCustomMode, TEnumAsByte<EMovementMode>& OutGroundMode) const
 {
-	const uint32 GroundShift = FMath::FloorLog2(MOVE_MAX) + 1;
+	const uint32 GroundShift = FMath::CeilLogTwo(MOVE_MAX);
 	const uint8 CustomModeThr = 2 * (1 << GroundShift);
 
 	if (ReceivedMode < CustomModeThr)
@@ -753,7 +763,8 @@ void UCharacterMovementComponent::UnpackNetworkMovementMode(const uint8 Received
 
 		OutMode = TEnumAsByte<EMovementMode>(ReceivedMode & GroundMask);
 		OutCustomMode = 0;
-		OutGroundMode = TEnumAsByte<EMovementMode>(ReceivedMode >> GroundShift);
+		const uint8 GroundModeBit = (ReceivedMode >> GroundShift);
+		OutGroundMode = TEnumAsByte<EMovementMode>(GroundModeBit == 0 ? MOVE_Walking : MOVE_NavWalking);
 	}
 	else
 	{
@@ -3921,33 +3932,21 @@ void UCharacterMovementComponent::PhysNavWalking(float deltaTime, int32 Iteratio
 
 bool UCharacterMovementComponent::FindNavFloor(const FVector& TestLocation, FNavLocation& NavFloorLocation) const
 {
-	UNavigationSystem* NavSys = UNavigationSystem::GetCurrent(GetWorld());
-	if (NavSys == nullptr)
+	const ANavigationData* NavData = GetNavData();
+	if (NavData == nullptr)
 	{
 		return false;
 	}
 
-	const ANavigationData* NavData = nullptr;
-	INavAgentInterface* MyNavAgent = Cast<INavAgentInterface>(CharacterOwner);
+	INavAgentInterface* MyNavAgent = CastChecked<INavAgentInterface>(CharacterOwner);
 	float SearchRadius = 0.0f;
 	float SearchHeight = 100.0f;
 	if (MyNavAgent)
 	{
 		const FNavAgentProperties& AgentProps = MyNavAgent->GetNavAgentPropertiesRef();
-		NavData = NavSys->GetNavDataForProps(AgentProps);
 		SearchRadius = AgentProps.AgentRadius * 2.0f;
 		SearchHeight = AgentProps.AgentHeight * AgentProps.NavWalkingSearchHeightScale;
 	}
-	if (NavData == nullptr)
-	{
-		NavData = NavSys->GetMainNavData();
-	}
-
-	const ARecastNavMesh* NavMeshData = Cast<const ARecastNavMesh>(NavData);
-	if (NavMeshData == nullptr)
-	{
-		return false;
-	}
 
 	NavData->ProjectPoint(TestLocation, NavFloorLocation, FVector(SearchRadius, SearchRadius, SearchHeight));
 	return true;
@@ -4044,6 +4043,38 @@ FVector UCharacterMovementComponent::ProjectLocationFromNavMesh(float DeltaSecon
 	return NewLocation;
 }
 
+
+const ANavigationData* UCharacterMovementComponent::GetNavData() const
+{
+	UNavigationSystem* NavSys = UNavigationSystem::GetCurrent(GetWorld());
+	if (NavSys == nullptr || !HasValidData())
+	{
+		return nullptr;
+	}
+
+	const ANavigationData* NavData = nullptr;
+	INavAgentInterface* MyNavAgent = CastChecked<INavAgentInterface>(CharacterOwner);
+	if (MyNavAgent)
+	{
+		const FNavAgentProperties& AgentProps = MyNavAgent->GetNavAgentPropertiesRef();
+		NavData = NavSys->GetNavDataForProps(AgentProps);
+	}
+	if (NavData == nullptr)
+	{
+		NavData = NavSys->GetMainNavData();
+	}
+
+	// Only RecastNavMesh supported
+	const ARecastNavMesh* NavMeshData = Cast<const ARecastNavMesh>(NavData);
+	if (NavMeshData == nullptr)
+	{
+		return nullptr;
+	}
+
+	return NavData;
+}
+
+
 void UCharacterMovementComponent::PhysCustom(float deltaTime, int32 Iterations)
 {
 	if (CharacterOwner)
