Fix nan caused by not setting the UpdatedScale

diff --git a/Engine/Source/Runtime/Engine/Private/PhysicsEngine/BodyInstance.cpp b/Engine/Source/Runtime/Engine/Private/PhysicsEngine/BodyInstance.cpp
--- a/Engine/Source/Runtime/Engine/Private/PhysicsEngine/BodyInstance.cpp
+++ b/Engine/Source/Runtime/Engine/Private/PhysicsEngine/BodyInstance.cpp
@@ -2052,7 +2052,7 @@ bool FBodyInstance::UpdateBodyScale(const FVector& InScale3D)
 	bool bSuccess = false;
 
 #if !(UE_BUILD_SHIPPING || UE_BUILD_TEST)
-	ensure ( !Scale3D.ContainsNaN() && !InScale3D.ContainsNaN() );
+	ensureMsgf ( !Scale3D.ContainsNaN() && !InScale3D.ContainsNaN(), TEXT("Scale3D = (%f,%f,%f) InScale3D = (%f,%f,%f)"), Scale3D.X, Scale3D.Y, Scale3D.Z, InScale3D.X, InScale3D.Y, InScale3D.Z );
 #endif
 
 	FVector UpdatedScale3D;
@@ -2267,6 +2267,7 @@ bool FBodyInstance::UpdateBodyScale(const FVector& InScale3D)
 					PGivenShape->setLocalPose(PLocalPose);
 					PGivenShape->setGeometry(*UpdatedGeometry);
 				});
+				UpdatedScale3D = AdjustedScale3D;
 			}
 			else if (bInvalid)
 			{
