Implements the view property material input node, exposing the field of view, view size and view texel size.

diff --git a/Engine/Documentation/Source/Engine/Rendering/Materials/ExpressionReference/Constant/ConstantExpressions.INT.udn b/Engine/Documentation/Source/Engine/Rendering/Materials/ExpressionReference/Constant/ConstantExpressions.INT.udn
--- a/Engine/Documentation/Source/Engine/Rendering/Materials/ExpressionReference/Constant/ConstantExpressions.INT.udn
+++ b/Engine/Documentation/Source/Engine/Rendering/Materials/ExpressionReference/Constant/ConstantExpressions.INT.udn
@@ -199,4 +199,21 @@ The VertexColor expression is the access point for the material to the outputs o
 
 The network above would utilize color and alpha data directly from a Particle System. Such networks are ideal for special effects such as sparks.
 
+(#ViewProperty)
+## View Property
 
+The ViewProperty expression outputs a view dependant constant property such as the field of view or the render target size. The view property to be accessed can be configured, and the type of the output depends on the configured accessed property.
+
+[REGION:fullwidth]
+![ViewPropertyExample.png](ViewPropertyExample.png)
+[/REGION]
+
+Here is the example with a field of view of 50 degrees.
+[REGION:fullwidth]
+![ViewPropertyExampleFOV50.png](ViewPropertyExampleFOV50.png)
+[/REGION]
+
+Here is the same example with a field of view of 90 degrees.
+[REGION:fullwidth]
+![ViewPropertyExampleFOV90.png](ViewPropertyExampleFOV90.png)
+[/REGION]
diff --git a/Engine/Documentation/Source/Shared/GraphNodes/Material/Material.INT.udn b/Engine/Documentation/Source/Shared/GraphNodes/Material/Material.INT.udn
--- a/Engine/Documentation/Source/Shared/GraphNodes/Material/Material.INT.udn
+++ b/Engine/Documentation/Source/Shared/GraphNodes/Material/Material.INT.udn
@@ -278,6 +278,14 @@ Engine/Rendering/Materials/ExpressionReference/Coordinates#vertexnormalws
 [/VAR]
 [/EXCERPT:UMaterialExpressionVertexNormalWS]
 
+[EXCERPT:UMaterialExpressionViewProperty]
+The ViewProperty expression outputs a view dependant constant property. The view property to be accessed can be configured, and the type of the output depends on the configured accessed property.
+
+[VAR:ToolTipFullLink]
+Engine/Rendering/Materials/ExpressionReference/Constant#viewproperty
+[/VAR]
+[/EXCERPT:UMaterialExpressionViewProperty]
+
 [EXCERPT:UMaterialExpressionViewSize]
 The ViewSize expression outputs a 2D vector giving the size of the current view in pixels. This is useful for causing various changes in your materials based on the current resolution of the screen. 
 
diff --git a/Engine/Source/Runtime/Engine/Classes/Materials/MaterialExpressionViewProperty.h b/Engine/Source/Runtime/Engine/Classes/Materials/MaterialExpressionViewProperty.h
new file mode 100644
--- /dev/null
+++ b/Engine/Source/Runtime/Engine/Classes/Materials/MaterialExpressionViewProperty.h
@@ -0,0 +1,34 @@
+// Copyright 1998-2015 Epic Games, Inc. All Rights Reserved.
+
+
+#pragma once
+#include "Materials/MaterialExpression.h"
+#include "MaterialExpressionViewProperty.generated.h"
+
+UENUM()
+enum EMaterialExposedViewProperty
+{	
+	/** Horizontal and vertical fields of view in radian */
+	MEVP_FieldOfView UMETA(DisplayName="FieldOfView"),
+	/** Horizontal and vertical size of the view in pixels */
+	MEVP_ViewSize UMETA(DisplayName="ViewSize"),
+	/** Horizontal and vertical size of view's render target's texel in the uv space (1 / ViewSize) */
+	MEVP_TexelSize UMETA(DisplayName="TexelSize"),
+
+	MEVP_MAX,
+};
+
+UCLASS(collapsecategories, hidecategories=Object, MinimalAPI)
+class UMaterialExpressionViewProperty : public UMaterialExpression
+{
+	GENERATED_UCLASS_BODY()
+	
+	/** View input property to be accessed */
+	UPROPERTY(EditAnywhere, Category=UMaterialExpressionViewProperty, meta=(DisplayName = "View Property"))
+	TEnumAsByte<EMaterialExposedViewProperty> Property;
+	
+	// Begin UMaterialExpression Interface
+	virtual int32 Compile(class FMaterialCompiler* Compiler, int32 OutputIndex, int32 MultiplexIndex) override;
+	virtual void GetCaption(TArray<FString>& OutCaptions) const override;
+	// End UMaterialExpression Interface
+};
diff --git a/Engine/Source/Runtime/Engine/Private/Materials/HLSLMaterialTranslator.h b/Engine/Source/Runtime/Engine/Private/Materials/HLSLMaterialTranslator.h
--- a/Engine/Source/Runtime/Engine/Private/Materials/HLSLMaterialTranslator.h
+++ b/Engine/Source/Runtime/Engine/Private/Materials/HLSLMaterialTranslator.h
@@ -1632,6 +1632,39 @@ protected:
 	{
 		return AddUniformExpression(new FMaterialUniformExpressionConstant(FLinearColor(X,Y,Z,W),MCT_Float4),MCT_Float4,TEXT("MaterialFloat4(%0.8f,%0.8f,%0.8f,%0.8f)"),X,Y,Z,W);
 	}
+	
+	virtual int32 ViewProperty(EMaterialExposedViewProperty Property) override
+	{
+		check(Property < MEVP_MAX);
+
+		// Compile time struct storing all EMaterialExposedViewProperty's enumerations' HLSL compilation specific meta information
+		struct EMaterialExposedViewPropertyMeta
+		{
+			EMaterialExposedViewProperty EnumValue;
+			EMaterialValueType Type;
+			const TCHAR * CurrentFrameCode;
+			const TCHAR * PreviousFrameCode;
+		};
+
+		static const EMaterialExposedViewPropertyMeta ViewPropertyMetaArray[] = {
+			{MEVP_FieldOfView, MCT_Float2, TEXT("View.FieldOfViewWideAngles"), TEXT("View.PrevFieldOfViewWideAngles")},
+			{MEVP_ViewSize, MCT_Float2, TEXT("View.ViewSizeAndSceneTexelSize.xy"), nullptr},
+			{MEVP_TexelSize, MCT_Float2, TEXT("View.ViewSizeAndSceneTexelSize.zw"), nullptr},
+		};
+		static_assert((sizeof(ViewPropertyMetaArray) / sizeof(ViewPropertyMetaArray[0])) == MEVP_MAX, "incoherency between EMaterialExposedViewProperty and ViewPropertyMetaArray");
+
+		auto& PropertyMeta = ViewPropertyMetaArray[Property];
+		check(Property == PropertyMeta.EnumValue);
+
+		auto Code = PropertyMeta.CurrentFrameCode;
+
+		if (bCompilingPreviousFrame && PropertyMeta.PreviousFrameCode)
+		{
+			Code = PropertyMeta.PreviousFrameCode;
+		}
+
+		return AddCodeChunk(PropertyMeta.Type, Code);
+	}
 
 	virtual int32 GameTime(bool bPeriodic, float Period) override
 	{
@@ -1900,16 +1933,6 @@ protected:
 		return AddCodeChunk(MCT_Float2,TEXT("ScreenAlignedPosition(Parameters.ScreenPosition).xy"));		
 	}
 
-	virtual int32 ViewSize() override
-	{
-		return AddCodeChunk(MCT_Float2,TEXT("View.ViewSizeAndSceneTexelSize.xy"));
-	}
-
-	virtual int32 SceneTexelSize() override
-	{
-		return AddCodeChunk(MCT_Float2,TEXT("View.ViewSizeAndSceneTexelSize.zw"));
-	}
-
 	virtual int32 ParticleMacroUV() override 
 	{
 		if (ShaderFrequency != SF_Pixel && ShaderFrequency != SF_Compute)
diff --git a/Engine/Source/Runtime/Engine/Private/Materials/MaterialExpressions.cpp b/Engine/Source/Runtime/Engine/Private/Materials/MaterialExpressions.cpp
--- a/Engine/Source/Runtime/Engine/Private/Materials/MaterialExpressions.cpp
+++ b/Engine/Source/Runtime/Engine/Private/Materials/MaterialExpressions.cpp
@@ -126,6 +126,7 @@
 #include "Materials/MaterialExpressionTwoSidedSign.h"
 #include "Materials/MaterialExpressionVertexColor.h"
 #include "Materials/MaterialExpressionVertexNormalWS.h"
+#include "Materials/MaterialExpressionViewProperty.h"
 #include "Materials/MaterialExpressionViewSize.h"
 #include "Materials/MaterialExpressionWorldPosition.h"
 #include "Materials/MaterialExpressionDistanceToNearestSurface.h"
@@ -4608,6 +4609,44 @@ void UMaterialExpressionScreenPosition::GetCaption(TArray<FString>& OutCaptions)
 	OutCaptions.Add(TEXT("ScreenPosition"));
 }
 
+UMaterialExpressionViewProperty::UMaterialExpressionViewProperty(const FObjectInitializer& ObjectInitializer)
+	: Super(ObjectInitializer)
+{
+	// Structure to hold one-time initialization
+	struct FConstructorStatics
+	{
+		FText NAME_Constants;
+		FConstructorStatics()
+			: NAME_Constants(LOCTEXT( "Constants", "Constants" ))
+		{
+		}
+	};
+	static FConstructorStatics ConstructorStatics;
+
+	Property = MEVP_FieldOfView;
+	MenuCategories.Add(ConstructorStatics.NAME_Constants);
+	bShaderInputData = true;
+}
+
+int32 UMaterialExpressionViewProperty::Compile(class FMaterialCompiler* Compiler, int32 OutputIndex, int32 MultiplexIndex)
+{
+	return Compiler->ViewProperty(Property);
+}
+
+void UMaterialExpressionViewProperty::GetCaption(TArray<FString>& OutCaptions) const
+{
+#if WITH_EDITOR
+	const UEnum* ViewPropertyEnum = FindObject<UEnum>(NULL, TEXT("Engine.EMaterialExposedViewProperty"));
+	check(ViewPropertyEnum);
+
+	const FString PropertyDisplayName = ViewPropertyEnum->GetDisplayNameText(Property).ToString();
+#else
+	const FString PropertyDisplayName = TEXT("");
+#endif
+
+	OutCaptions.Add(FString(TEXT("View.")) + PropertyDisplayName);
+}
+
 UMaterialExpressionViewSize::UMaterialExpressionViewSize(const FObjectInitializer& ObjectInitializer)
 	: Super(ObjectInitializer)
 {
@@ -4628,7 +4667,7 @@ UMaterialExpressionViewSize::UMaterialExpressionViewSize(const FObjectInitialize
 
 int32 UMaterialExpressionViewSize::Compile(class FMaterialCompiler* Compiler, int32 OutputIndex, int32 MultiplexIndex)
 {
-	return Compiler->ViewSize();
+	return Compiler->ViewProperty(MEVP_ViewSize);
 }
 
 void UMaterialExpressionViewSize::GetCaption(TArray<FString>& OutCaptions) const
@@ -4656,7 +4695,7 @@ UMaterialExpressionSceneTexelSize::UMaterialExpressionSceneTexelSize(const FObje
 
 int32 UMaterialExpressionSceneTexelSize::Compile(class FMaterialCompiler* Compiler, int32 OutputIndex, int32 MultiplexIndex)
 {
-	return Compiler->SceneTexelSize();
+	return Compiler->ViewProperty(MEVP_TexelSize);
 }
 
 void UMaterialExpressionSceneTexelSize::GetCaption(TArray<FString>& OutCaptions) const
diff --git a/Engine/Source/Runtime/Engine/Public/MaterialCompiler.h b/Engine/Source/Runtime/Engine/Public/MaterialCompiler.h
--- a/Engine/Source/Runtime/Engine/Public/MaterialCompiler.h
+++ b/Engine/Source/Runtime/Engine/Public/MaterialCompiler.h
@@ -10,6 +10,7 @@
 #include "Materials/MaterialExpressionSpeedTree.h"
 #include "Materials/MaterialExpressionTextureSample.h"
 #include "Materials/MaterialExpressionWorldPosition.h"
+#include "Materials/MaterialExpressionViewProperty.h"
 
 /** 
  * The interface used to translate material expressions into executable code. 
@@ -57,6 +58,8 @@ public:
 	virtual int32 Constant3(float X,float Y,float Z) = 0;
 	virtual int32 Constant4(float X,float Y,float Z,float W) = 0;
 
+	virtual	int32 ViewProperty(EMaterialExposedViewProperty Property) = 0;
+
 	virtual int32 GameTime(bool bPeriodic, float Period) = 0;
 	virtual int32 RealTime(bool bPeriodic, float Period) = 0;
 	virtual int32 PeriodicHint(int32 PeriodicCode) { return PeriodicCode; }
@@ -77,8 +80,6 @@ public:
 	virtual int32 LightVector() = 0;
 
 	virtual int32 ScreenPosition() = 0;
-	virtual int32 ViewSize() = 0;
-	virtual int32 SceneTexelSize() = 0;
 	virtual int32 WorldPosition(EWorldPositionIncludedOffsets WorldPositionIncludedOffsets) = 0;
 	virtual int32 ObjectWorldPosition() = 0;
 	virtual int32 ObjectRadius() = 0;
@@ -221,6 +222,8 @@ public:
 	virtual int32 Constant2(float X,float Y) override { return Compiler->Constant2(X,Y); }
 	virtual int32 Constant3(float X,float Y,float Z) override { return Compiler->Constant3(X,Y,Z); }
 	virtual int32 Constant4(float X,float Y,float Z,float W) override { return Compiler->Constant4(X,Y,Z,W); }
+	
+	virtual	int32 ViewProperty(EMaterialExposedViewProperty Property) override { return Compiler->ViewProperty(Property); }
 
 	virtual int32 GameTime(bool bPeriodic, float Period) override { return Compiler->GameTime(bPeriodic, Period); }
 	virtual int32 RealTime(bool bPeriodic, float Period) override { return Compiler->RealTime(bPeriodic, Period); }
@@ -242,8 +245,6 @@ public:
 	virtual int32 LightVector() override { return Compiler->LightVector(); }
 
 	virtual int32 ScreenPosition() override { return Compiler->ScreenPosition(); }
-	virtual int32 ViewSize() override { return Compiler->ViewSize(); }
-	virtual int32 SceneTexelSize() override { return Compiler->SceneTexelSize(); }
 	virtual int32 WorldPosition(EWorldPositionIncludedOffsets WorldPositionIncludedOffsets) override { return Compiler->WorldPosition(WorldPositionIncludedOffsets); }
 	virtual int32 ObjectWorldPosition() override { return Compiler->ObjectWorldPosition(); }
 	virtual int32 ObjectRadius() override { return Compiler->ObjectRadius(); }
diff --git a/Engine/Source/Runtime/Engine/Public/SceneView.h b/Engine/Source/Runtime/Engine/Public/SceneView.h
--- a/Engine/Source/Runtime/Engine/Public/SceneView.h
+++ b/Engine/Source/Runtime/Engine/Public/SceneView.h
@@ -154,7 +154,7 @@ struct FViewMatrices
 		PreShadowTranslation = FVector::ZeroVector;
 		PreViewTranslation = FVector::ZeroVector;
 		ViewOrigin = FVector::ZeroVector;
-		TemporalAASample = FVector2D::ZeroVector;
+		TemporalAAProjJitter = FVector2D::ZeroVector;
 	}
 
 	/** ViewToClip : UE4 projection matrix projects such that clip space Z=1 is the near plane, and Z=0 is the infinite far plane. */
@@ -177,7 +177,9 @@ struct FViewMatrices
 	FVector		ViewOrigin;
 	/** Scale applied by the projection matrix in X and Y. */
 	FVector2D	ProjectionScale;
-	FVector2D	TemporalAASample;
+	/** TemporalAA jitter offset currently stored in the projection matrix */
+	FVector2D	TemporalAAProjJitter;
+
 	/**
 	 * Scale factor to use when computing the size of a sphere in pixels.
 	 * 
@@ -201,6 +203,22 @@ struct FViewMatrices
 		return ProjMatrix.M[3][3] < 1.0f;
 	}
 
+	FMatrix GetProjNoAAMatrix() const
+	{
+		FMatrix ProjNoAAMatrix = ProjMatrix;
+
+		ProjNoAAMatrix.M[2][0] -= TemporalAAProjJitter.X;
+		ProjNoAAMatrix.M[2][1] -= TemporalAAProjJitter.Y;
+
+		return ProjNoAAMatrix;
+	}
+
+	void RemoveTemporalJitter()
+	{
+		ProjMatrix = GetProjNoAAMatrix();
+		TemporalAAProjJitter = FVector2D::ZeroVector;
+	}
+
 	FMatrix GetViewProjMatrix() const
 	{
 		return ViewMatrix * ProjMatrix;
@@ -215,6 +233,11 @@ struct FViewMatrices
 	{
 		return ProjMatrix.Inverse();
 	}
+	
+	FMatrix GetInvProjNoAAMatrix() const
+	{
+		return GetProjNoAAMatrix().Inverse();
+	}
 
 	FMatrix GetInvViewMatrix() const
 	{
@@ -227,9 +250,9 @@ struct FViewMatrices
 		return GetInvProjMatrix() * GetInvViewMatrix();
 	}
 
-	FVector2D GetFieldOfViewPerAxis() const
+	FVector2D GetHalfFieldOfViewPerAxis() const
 	{
-		const FMatrix ClipToView = GetInvProjMatrix();
+		const FMatrix ClipToView = GetInvProjNoAAMatrix();
 
 		FVector VCenter = FVector(ClipToView.TransformPosition(FVector(0.0, 0.0, 0.0)));
 		FVector VUp = FVector(ClipToView.TransformPosition(FVector(0.0, 1.0, 0.0)));
@@ -296,6 +319,8 @@ BEGIN_UNIFORM_BUFFER_STRUCT_WITH_CONSTRUCTOR(FViewUniformShaderParameters,ENGINE
 	DECLARE_UNIFORM_BUFFER_STRUCT_MEMBER_EX(FVector,ViewForward, EShaderPrecisionModifier::Half)
 	DECLARE_UNIFORM_BUFFER_STRUCT_MEMBER_EX(FVector,ViewUp, EShaderPrecisionModifier::Half)
 	DECLARE_UNIFORM_BUFFER_STRUCT_MEMBER_EX(FVector,ViewRight, EShaderPrecisionModifier::Half)
+	DECLARE_UNIFORM_BUFFER_STRUCT_MEMBER(FVector2D,FieldOfViewWideAngles)
+	DECLARE_UNIFORM_BUFFER_STRUCT_MEMBER(FVector2D,PrevFieldOfViewWideAngles)
 	DECLARE_UNIFORM_BUFFER_STRUCT_MEMBER(FVector4,InvDeviceZToWorldZTransform)
 	DECLARE_UNIFORM_BUFFER_STRUCT_MEMBER_EX(FVector4,ScreenPositionScaleBias, EShaderPrecisionModifier::Half)
 	DECLARE_UNIFORM_BUFFER_STRUCT_MEMBER_EX(FVector4,ViewRectMin, EShaderPrecisionModifier::Half)
diff --git a/Engine/Source/Runtime/Renderer/Private/CompositionLighting/PostProcessAmbientOcclusion.cpp b/Engine/Source/Runtime/Renderer/Private/CompositionLighting/PostProcessAmbientOcclusion.cpp
--- a/Engine/Source/Runtime/Renderer/Private/CompositionLighting/PostProcessAmbientOcclusion.cpp
+++ b/Engine/Source/Runtime/Renderer/Private/CompositionLighting/PostProcessAmbientOcclusion.cpp
@@ -596,14 +596,8 @@ TUniformBufferRef<FCameraMotionParameters> CreateCameraMotionParametersUniformBu
 {
 	FSceneViewState* ViewState = (FSceneViewState*)View.State;
 
-	FMatrix Proj = View.ViewMatrices.ProjMatrix;
-	FMatrix PrevProj = ViewState->PrevViewMatrices.ProjMatrix;
-
-	// Remove jitter
-	Proj.M[2][0] -= View.ViewMatrices.TemporalAASample.X * 2.0f / View.ViewRect.Width();
-	Proj.M[2][1] -= View.ViewMatrices.TemporalAASample.Y * 2.0f / View.ViewRect.Height();
-	PrevProj.M[2][0] -= ViewState->PrevViewMatrices.TemporalAASample.X * 2.0f / View.ViewRect.Width();
-	PrevProj.M[2][1] -= ViewState->PrevViewMatrices.TemporalAASample.Y * 2.0f / View.ViewRect.Height();
+	FMatrix Proj = View.ViewMatrices.GetProjNoAAMatrix();
+	FMatrix PrevProj = ViewState->PrevViewMatrices.GetProjNoAAMatrix();
 
 	FVector DeltaTranslation = ViewState->PrevViewMatrices.PreViewTranslation - View.ViewMatrices.PreViewTranslation;
 	FMatrix ViewProj = ( View.ViewMatrices.TranslatedViewMatrix * Proj ).GetTransposed();
diff --git a/Engine/Source/Runtime/Renderer/Private/PostProcess/PostProcessMobile.cpp b/Engine/Source/Runtime/Renderer/Private/PostProcess/PostProcessMobile.cpp
--- a/Engine/Source/Runtime/Renderer/Private/PostProcess/PostProcessMobile.cpp
+++ b/Engine/Source/Runtime/Renderer/Private/PostProcess/PostProcessMobile.cpp
@@ -2601,14 +2601,8 @@ public:
 		{
 			const FViewInfo& View = Context.View;
 
-			FMatrix Proj = View.ViewMatrices.ProjMatrix;
-			FMatrix PrevProj = ViewState->PrevViewMatrices.ProjMatrix;
-
-			// Remove jitter
-			Proj.M[2][0] -= View.ViewMatrices.TemporalAASample.X * 2.0f / View.ViewRect.Width();
-			Proj.M[2][1] -= View.ViewMatrices.TemporalAASample.Y * 2.0f / View.ViewRect.Height();
-			PrevProj.M[2][0] -= ViewState->PrevViewMatrices.TemporalAASample.X * 2.0f / View.ViewRect.Width();
-			PrevProj.M[2][1] -= ViewState->PrevViewMatrices.TemporalAASample.Y * 2.0f / View.ViewRect.Height();
+			FMatrix Proj = View.ViewMatrices.GetProjNoAAMatrix();
+			FMatrix PrevProj = ViewState->PrevViewMatrices.GetProjNoAAMatrix();
 
 			FMatrix ViewProj = ( Context.View.ViewMatrices.ViewMatrix * Proj ).GetTransposed();
 			FMatrix PrevViewProj = ( ViewState->PrevViewMatrices.ViewMatrix * PrevProj ).GetTransposed();
@@ -2784,8 +2778,7 @@ void FRCPassPostProcessAaES2::Process(FRenderingCompositePassContext& Context)
 	if (FSceneRenderer::ShouldCompositeEditorPrimitives(View))
 	{
 		// Remove jitter (ensures editor prims are stable.)
-		View.ViewMatrices.ProjMatrix.M[2][0] -= View.ViewMatrices.TemporalAASample.X * 2.0f / View.ViewRect.Width();
-		View.ViewMatrices.ProjMatrix.M[2][1] -= View.ViewMatrices.TemporalAASample.Y * 2.0f / View.ViewRect.Height();
+		View.ViewMatrices.RemoveTemporalJitter();
 
 		// Compute the view projection matrix and its inverse.
 		View.ViewProjectionMatrix = View.ViewMatrices.ViewMatrix * View.ViewMatrices.ProjMatrix;
diff --git a/Engine/Source/Runtime/Renderer/Private/PostProcess/PostProcessTemporalAA.cpp b/Engine/Source/Runtime/Renderer/Private/PostProcess/PostProcessTemporalAA.cpp
--- a/Engine/Source/Runtime/Renderer/Private/PostProcess/PostProcessTemporalAA.cpp
+++ b/Engine/Source/Runtime/Renderer/Private/PostProcess/PostProcessTemporalAA.cpp
@@ -707,8 +707,7 @@ void FRCPassPostProcessTemporalAA::Process(FRenderingCompositePassContext& Conte
 	// TODO draw separate translucency after jitter has been removed
 
 	// Remove jitter
-	View.ViewMatrices.ProjMatrix.M[2][0] -= View.ViewMatrices.TemporalAASample.X * 2.0f / View.ViewRect.Width();
-	View.ViewMatrices.ProjMatrix.M[2][1] -= View.ViewMatrices.TemporalAASample.Y * 2.0f / View.ViewRect.Height();
+	View.ViewMatrices.RemoveTemporalJitter();
 
 	// Compute the view projection matrix and its inverse.
 	View.ViewProjectionMatrix = View.ViewMatrices.ViewMatrix * View.ViewMatrices.ProjMatrix;
diff --git a/Engine/Source/Runtime/Renderer/Private/PostProcess/PostProcessUpscale.cpp b/Engine/Source/Runtime/Renderer/Private/PostProcess/PostProcessUpscale.cpp
--- a/Engine/Source/Runtime/Renderer/Private/PostProcess/PostProcessUpscale.cpp
+++ b/Engine/Source/Runtime/Renderer/Private/PostProcess/PostProcessUpscale.cpp
@@ -70,7 +70,7 @@ public:
 		FGlobalShader::SetParameters(Context.RHICmdList, ShaderRHI, Context.View);
 
 		{
-			const FVector2D FOVPerAxis = Context.View.ViewMatrices.GetFieldOfViewPerAxis();
+			const FVector2D FOVPerAxis = Context.View.ViewMatrices.GetHalfFieldOfViewPerAxis();
 			const FVector2D ScreenPosToPaniniFactor = FVector2D(FMath::Tan(FOVPerAxis.X), FMath::Tan(FOVPerAxis.Y));
 			const FVector2D PaniniDirection = FVector2D(1.0f, 0.0f) * ScreenPosToPaniniFactor;
 			const FVector2D PaniniPosition = PaniniProjection(PaniniDirection, InPaniniConfig.D, InPaniniConfig.S);
diff --git a/Engine/Source/Runtime/Renderer/Private/SceneRendering.cpp b/Engine/Source/Runtime/Renderer/Private/SceneRendering.cpp
--- a/Engine/Source/Runtime/Renderer/Private/SceneRendering.cpp
+++ b/Engine/Source/Runtime/Renderer/Private/SceneRendering.cpp
@@ -452,6 +452,7 @@ TUniformBufferRef<FViewUniformShaderParameters> FViewInfo::CreateUniformBuffer(
 	const bool bIsUnlitView = !Family->EngineShowFlags.Lighting;
 
 	// Create the view's uniform buffer.
+	// TODO: We should use a view and previous view uniform buffer to avoid code duplication and keep consistency
 	FViewUniformShaderParameters ViewUniformShaderParameters;
 	ViewUniformShaderParameters.TranslatedWorldToClip = ViewMatrices.TranslatedViewProjectionMatrix;
 	ViewUniformShaderParameters.WorldToClip = ViewProjectionMatrix;
@@ -463,6 +464,8 @@ TUniformBufferRef<FViewUniformShaderParameters> FViewInfo::CreateUniformBuffer(
 	ViewUniformShaderParameters.ViewForward = EffectiveTranslatedViewMatrix.GetColumn(2);
 	ViewUniformShaderParameters.ViewUp = EffectiveTranslatedViewMatrix.GetColumn(1);
 	ViewUniformShaderParameters.ViewRight = EffectiveTranslatedViewMatrix.GetColumn(0);
+	ViewUniformShaderParameters.FieldOfViewWideAngles = 2.f * ViewMatrices.GetHalfFieldOfViewPerAxis();
+	ViewUniformShaderParameters.PrevFieldOfViewWideAngles = 2.f * PrevViewMatrices.GetHalfFieldOfViewPerAxis();
 	ViewUniformShaderParameters.InvDeviceZToWorldZTransform = InvDeviceZToWorldZTransform;
 	ViewUniformShaderParameters.ScreenPositionScaleBias = ScreenPositionScaleBias;
 	ViewUniformShaderParameters.ViewRectMin = FVector4(ViewRect.Min.X, ViewRect.Min.Y, 0.0f, 0.0f);
diff --git a/Engine/Source/Runtime/Renderer/Private/SceneVisibility.cpp b/Engine/Source/Runtime/Renderer/Private/SceneVisibility.cpp
--- a/Engine/Source/Runtime/Renderer/Private/SceneVisibility.cpp
+++ b/Engine/Source/Runtime/Renderer/Private/SceneVisibility.cpp
@@ -1507,11 +1507,10 @@ void FSceneRenderer::PreVisibilityFrameSetup(FRHICommandListImmediate& RHICmdLis
 					SampleY = Halton( Index, 3 ) - 0.5f;
 				}
 
-				View.ViewMatrices.TemporalAASample.X = SampleX;
-				View.ViewMatrices.TemporalAASample.Y = SampleY;
-
-				View.ViewMatrices.ProjMatrix.M[2][0] += View.ViewMatrices.TemporalAASample.X * 2.0f / View.ViewRect.Width();
-				View.ViewMatrices.ProjMatrix.M[2][1] += View.ViewMatrices.TemporalAASample.Y * 2.0f / View.ViewRect.Height();;
+				View.ViewMatrices.TemporalAAProjJitter.X = SampleX * 2.0f / View.ViewRect.Width();
+				View.ViewMatrices.TemporalAAProjJitter.Y = SampleY * 2.0f / View.ViewRect.Height();
+				View.ViewMatrices.ProjMatrix.M[2][0] += View.ViewMatrices.TemporalAAProjJitter.X;
+				View.ViewMatrices.ProjMatrix.M[2][1] += View.ViewMatrices.TemporalAAProjJitter.Y;
 
 				// Compute the view projection matrix and its inverse.
 				View.ViewProjectionMatrix = View.ViewMatrices.ViewMatrix * View.ViewMatrices.ProjMatrix;
