fixing build

diff --git a/Engine/Source/Runtime/Engine/Classes/Engine/Texture.h b/Engine/Source/Runtime/Engine/Classes/Engine/Texture.h
--- a/Engine/Source/Runtime/Engine/Classes/Engine/Texture.h
+++ b/Engine/Source/Runtime/Engine/Classes/Engine/Texture.h
@@ -181,7 +181,7 @@ struct FTextureSource
 	ENGINE_API int32 GetBytesPerPixel() const;
 
 	/** Return true if the source data is power-of-2. */
-	bool IsPowerOfTwo() const;
+	ENGINE_API bool IsPowerOfTwo() const;
 
 	/** Returns true if source art is available. */
 	ENGINE_API bool IsValid() const;
