fixing build

diff --git a/Engine/Source/Runtime/Engine/Classes/GameFramework/CharacterMovementComponent.h b/Engine/Source/Runtime/Engine/Classes/GameFramework/CharacterMovementComponent.h
--- a/Engine/Source/Runtime/Engine/Classes/GameFramework/CharacterMovementComponent.h
+++ b/Engine/Source/Runtime/Engine/Classes/GameFramework/CharacterMovementComponent.h
@@ -1166,7 +1166,7 @@ protected:
 	virtual void SetNavWalkingPhysics(bool bEnable);
 
 	/** Get Navigation data for the Character. Returns null if there is no associated nav data. */
-	const class ANavigationData* UCharacterMovementComponent::GetNavData() const;
+	const class ANavigationData* GetNavData() const;
 
 	/** 
 	 * Checks to see if the current location is not encroaching blocking geometry so the character can leave NavWalking.
