Encapsulate mark pending kill loops for world unloading to use ForEachObjectWithOuter instead of the object iterator
Mark objects pending kill when unloading a streaming level
Move mark pending kill loop until after other cleanup has been done in the end PIE case

diff --git a/Engine/Source/Editor/UnrealEd/Private/PlayLevel.cpp b/Engine/Source/Editor/UnrealEd/Private/PlayLevel.cpp
--- a/Engine/Source/Editor/UnrealEd/Private/PlayLevel.cpp
+++ b/Engine/Source/Editor/UnrealEd/Private/PlayLevel.cpp
@@ -158,24 +158,7 @@ void UEditorEngine::EndPlayMap()
 	}
 	CleanupGameViewport();
 
-	// Gather all the worlds relevant to the PIE levels we are ending
-	TSet<UWorld*> CurrentPlayWorlds;
-	for (const FWorldContext& ThisContext : WorldList)
-	{
-		if (ThisContext.WorldType == EWorldType::PIE && ThisContext.World())
-		{
-			for (auto LevelIt(ThisContext.World()->GetLevelIterator()); LevelIt; ++LevelIt)
-			{
-				if (const ULevel* Level = *LevelIt)
-				{
-					CurrentPlayWorlds.Add(CastChecked<UWorld>(Level->GetOuter()));
-				}
-			}
-		}
-	}
-
 	// find objects like Textures in the playworld levels that won't get garbage collected as they are marked RF_Standalone
-	// mark everything else contained in the PIE worlds to be deleted
 	for( FObjectIterator It; It; ++It )
 	{
 		UObject* Object = *It;
@@ -190,23 +173,21 @@ void UEditorEngine::EndPlayMap()
 			// Close any asset editors that are currently editing this object
 			FAssetEditorManager::Get().CloseAllEditorsForAsset(Object);
 		}
-		else if (!Object->IsPendingKill())
-		{
-			UWorld* InWorld = Object->GetTypedOuter<UWorld>();
-			if (InWorld && CurrentPlayWorlds.Contains(InWorld))
-			{
-				Object->MarkPendingKill();
-			}
-		}
 	}
 
 	// Clean up each world individually
 	TArray<FName> OnlineIdentifiers;
+	TArray<UWorld*> WorldsBeingCleanedUp;
 	for (int32 WorldIdx = WorldList.Num()-1; WorldIdx >= 0; --WorldIdx)
 	{
 		FWorldContext &ThisContext = WorldList[WorldIdx];
 		if (ThisContext.WorldType == EWorldType::PIE)
 		{
+			if (ThisContext.World())
+			{
+				WorldsBeingCleanedUp.Add(ThisContext.World());
+			}
+
 			TeardownPlaySession(ThisContext);
 			
 			// Cleanup online subsystems instantiated during PIE
@@ -262,6 +243,17 @@ void UEditorEngine::EndPlayMap()
 	EditorWorld->bAllowAudioPlayback = true;
 	EditorWorld = NULL;
 
+	// mark everything contained in the PIE worlds to be deleted
+	for (UWorld* World : WorldsBeingCleanedUp)
+	{
+		for (auto LevelIt(World->GetLevelIterator()); LevelIt; ++LevelIt)
+		{
+			if (const ULevel* Level = *LevelIt)
+			{
+				CastChecked<UWorld>(Level->GetOuter())->MarkObjectsPendingKill();
+			}
+		}
+	}
 	// Garbage Collect
 	CollectGarbage( GARBAGE_COLLECTION_KEEPFLAGS );
 
diff --git a/Engine/Source/Runtime/Engine/Classes/Engine/World.h b/Engine/Source/Runtime/Engine/Classes/Engine/World.h
--- a/Engine/Source/Runtime/Engine/Classes/Engine/World.h
+++ b/Engine/Source/Runtime/Engine/Classes/Engine/World.h
@@ -2548,6 +2548,11 @@ public:
 	 */
 	void DestroyWorld( bool bInformEngineOfWorld, UWorld* NewWorld = nullptr );
 
+	/** 
+	 * Marks all objects that have this World as an Outer as pending kill
+	 */
+	void MarkObjectsPendingKill();
+
 	/**
 	 *  Interface to allow WorldSettings to request immediate garbage collection
 	 */
diff --git a/Engine/Source/Runtime/Engine/Private/UnrealEngine.cpp b/Engine/Source/Runtime/Engine/Private/UnrealEngine.cpp
--- a/Engine/Source/Runtime/Engine/Private/UnrealEngine.cpp
+++ b/Engine/Source/Runtime/Engine/Private/UnrealEngine.cpp
@@ -8921,22 +8921,12 @@ bool UEngine::LoadMap( FWorldContext& WorldContext, FURL URL, class UPendingNetG
 		WorldContext.World()->RemoveFromRoot();
 
 		// mark everything else contained in the world to be deleted
-		TSet<UWorld*> CurrentWorlds;
 		for (auto LevelIt(WorldContext.World()->GetLevelIterator()); LevelIt; ++LevelIt)
 		{
 			const ULevel* Level = *LevelIt;
 			if (Level)
 			{
-				CurrentWorlds.Add(CastChecked<UWorld>(Level->GetOuter()));
-			}
-		}
-
-		for (TObjectIterator<UObject> It(RF_PendingKill); It; ++It)
-		{
-			UWorld* InWorld = It->GetTypedOuter<UWorld>();
-			if (InWorld && CurrentWorlds.Contains(InWorld))
-			{
-				It->MarkPendingKill();
+				CastChecked<UWorld>(Level->GetOuter())->MarkObjectsPendingKill();
 			}
 		}
 
diff --git a/Engine/Source/Runtime/Engine/Private/World.cpp b/Engine/Source/Runtime/Engine/Private/World.cpp
--- a/Engine/Source/Runtime/Engine/Private/World.cpp
+++ b/Engine/Source/Runtime/Engine/Private/World.cpp
@@ -1082,6 +1082,15 @@ void UWorld::DestroyWorld( bool bInformEngineOfWorld, UWorld* NewWorld )
 	}
 }
 
+void UWorld::MarkObjectsPendingKill()
+{
+	auto MarkObjectPendingKill = [](UObject* Object)
+	{
+		Object->MarkPendingKill();
+	};
+	ForEachObjectWithOuter(this, MarkObjectPendingKill, true, RF_PendingKill);
+}
+
 UWorld* UWorld::CreateWorld(const EWorldType::Type InWorldType, bool bInformEngineOfWorld, FName WorldName, UPackage* InWorldPackage, bool bAddToRoot, ERHIFeatureLevel::Type InFeatureLevel)
 {
 	if (InFeatureLevel >= ERHIFeatureLevel::Num)
@@ -2097,33 +2106,13 @@ void FLevelStreamingGCHelper::PrepareStreamedOutLevelsForGC()
 			LevelPackageNames.Add( LevelPackage->GetFName() );
 
 			// Mark level as pending kill so references to it get deleted.
-			Level->GetOuter()->MarkPendingKill();
-			Level->MarkPendingKill();
+			UWorld* LevelWorld = CastChecked<UWorld>(Level->GetOuter());
+			LevelWorld->MarkObjectsPendingKill();
+			LevelWorld->MarkPendingKill();
 			if (LevelPackage->MetaData)
 			{
 				LevelPackage->MetaData->MarkPendingKill();
 			}
-
-			// Mark all model components as pending kill so GC deletes references to them.
-			for( int32 ModelComponentIndex=0; ModelComponentIndex<Level->ModelComponents.Num(); ModelComponentIndex++ )
-			{
-				UModelComponent* ModelComponent = Level->ModelComponents[ModelComponentIndex];
-				if( ModelComponent )
-				{
-					ModelComponent->MarkPendingKill();
-				}
-			}
-
-			// Mark all actors and their components as pending kill so GC will delete references to them.
-			for( int32 ActorIndex=0; ActorIndex<Level->Actors.Num(); ActorIndex++ )
-			{
-				AActor* Actor = Level->Actors[ActorIndex];
-				if (Actor != NULL)
-				{
-					Actor->MarkComponentsAsPendingKill();
-					Actor->MarkPendingKill();
-				}
-			}
 		}
 	}
 
@@ -4770,22 +4759,12 @@ UWorld* FSeamlessTravelHandler::Tick()
 			GWorld = NULL;
 
 			// mark everything else contained in the world to be deleted
-			TSet<UWorld*> CurrentWorlds;
 			for (auto LevelIt(CurrentWorld->GetLevelIterator()); LevelIt; ++LevelIt)
 			{
 				const ULevel* Level = *LevelIt;
 				if (Level)
 				{
-					CurrentWorlds.Add(CastChecked<UWorld>(Level->GetOuter()));
-				}
-			}
-
-			for (TObjectIterator<UObject> It(RF_PendingKill); It; ++It)
-			{
-				UWorld* InWorld = It->GetTypedOuter<UWorld>();
-				if (InWorld && CurrentWorlds.Contains(InWorld))
-				{
-					It->MarkPendingKill();
+					CastChecked<UWorld>(Level->GetOuter())->MarkObjectsPendingKill();
 				}
 			}
 
