Stats - Added markers for load map

diff --git a/Engine/Source/Runtime/Engine/Private/UnrealEngine.cpp b/Engine/Source/Runtime/Engine/Private/UnrealEngine.cpp
--- a/Engine/Source/Runtime/Engine/Private/UnrealEngine.cpp
+++ b/Engine/Source/Runtime/Engine/Private/UnrealEngine.cpp
@@ -4698,7 +4698,7 @@ struct FHierarchy
 	}
 };
 
-// #TODO_YRX 2014-09-15 Move to ObjectCommads.cpp or ObjectExec.cpp
+// #YRX_UObject 2014-09-15 Move to ObjectCommads.cpp or ObjectExec.cpp
 bool UEngine::HandleObjCommand( const TCHAR* Cmd, FOutputDevice& Ar )
 {
 	if( FParse::Command(&Cmd,TEXT("GARBAGE")) || FParse::Command(&Cmd,TEXT("GC")) )
@@ -8785,6 +8785,8 @@ void LoadGametypeContent(FWorldContext &Context, const FURL& URL)
 
 bool UEngine::LoadMap( FWorldContext& WorldContext, FURL URL, class UPendingNetGame* Pending, FString& Error )
 {
+	STAT_ADD_CUSTOMMESSAGE_NAME( STAT_NamedMarker, *(URL.Map + TEXT( " - LoadMapStart" )) );
+
 	DECLARE_SCOPE_CYCLE_COUNTER(TEXT("UEngine::LoadMap"), STAT_LoadMap, STATGROUP_LoadTime);
 
 	NETWORK_PROFILER(GNetworkProfiler.TrackSessionChange(true,URL));
@@ -8982,6 +8984,7 @@ bool UEngine::LoadMap( FWorldContext& WorldContext, FURL URL, class UPendingNetG
 	
 #endif
 
+	STAT_ADD_CUSTOMMESSAGE_NAME( STAT_NamedMarker, *(URL.Map + TEXT( " - LoadMapMid" )) );
 	MALLOC_PROFILER( FMallocProfiler::SnapshotMemoryLoadMapMid( URL.Map ); )
 
 	if( GUseSeekFreeLoading )
@@ -9308,6 +9311,7 @@ bool UEngine::LoadMap( FWorldContext& WorldContext, FURL URL, class UPendingNetG
 		}
 	}
 
+	STAT_ADD_CUSTOMMESSAGE_NAME( STAT_NamedMarker, *(URL.Map + TEXT( " - LoadMapEnd" )) );
 	MALLOC_PROFILER( FMallocProfiler::SnapshotMemoryLoadMapEnd( URL.Map ); )
 
 	// Successfully started local level.
diff --git a/Engine/Source/Runtime/Engine/Private/World.cpp b/Engine/Source/Runtime/Engine/Private/World.cpp
--- a/Engine/Source/Runtime/Engine/Private/World.cpp
+++ b/Engine/Source/Runtime/Engine/Private/World.cpp
@@ -3029,6 +3029,8 @@ void UWorld::InitializeActorsForPlay(const FURL& InURL, bool bResetTime)
 	{
 		UE_LOG(LogWorld, Log, TEXT("Bringing up level for play took: %f"), FPlatformTime::Seconds() - StartTime );
 	}
+
+	STAT_ADD_CUSTOMMESSAGE_NAME( STAT_NamedMarker, *(InURL.Map + TEXT( " - Bringing Up" )) );
 }
 
 void UWorld::BeginPlay()
