Fix limited iteration loop node playing one more iteration than specified
Fix limited iteration loop node claiming notify event when looping has completed
Disable loop count edit box when loop indefinitely is true

diff --git a/Engine/Source/Runtime/Engine/Classes/Sound/SoundNodeLooping.h b/Engine/Source/Runtime/Engine/Classes/Sound/SoundNodeLooping.h
--- a/Engine/Source/Runtime/Engine/Classes/Sound/SoundNodeLooping.h
+++ b/Engine/Source/Runtime/Engine/Classes/Sound/SoundNodeLooping.h
@@ -14,7 +14,7 @@ class USoundNodeLooping : public USoundNode
 	GENERATED_UCLASS_BODY()
 
 	/* The amount of times to loop */
-	UPROPERTY(EditAnywhere, Category = Looping, meta = (ClampMin = 1))
+	UPROPERTY(EditAnywhere, Category = Looping, meta = (ClampMin = 1, EditCondition = "!bLoopIndefinitely"))
 	int32 LoopCount;
 
 	/* If enabled, the node will continue to loop indefinitely regardless of the Loop Count value. */
diff --git a/Engine/Source/Runtime/Engine/Private/SoundNodeLooping.cpp b/Engine/Source/Runtime/Engine/Private/SoundNodeLooping.cpp
--- a/Engine/Source/Runtime/Engine/Private/SoundNodeLooping.cpp
+++ b/Engine/Source/Runtime/Engine/Private/SoundNodeLooping.cpp
@@ -42,7 +42,7 @@ bool USoundNodeLooping::NotifyWaveInstanceFinished( FWaveInstance* InWaveInstanc
 	DECLARE_SOUNDNODE_ELEMENT(int32, CurrentLoopCount);
 	check(*RequiresInitialization == 0);
 
-	if (bLoopIndefinitely == 1 || CurrentLoopCount < LoopCount)
+	if (bLoopIndefinitely || ++CurrentLoopCount < LoopCount)
 	{
 		struct FNodeHashPairs
 		{
@@ -107,10 +107,10 @@ bool USoundNodeLooping::NotifyWaveInstanceFinished( FWaveInstance* InWaveInstanc
 		InWaveInstance->bIsStarted = false;
 		InWaveInstance->bIsFinished = false;
 
-		CurrentLoopCount++;
+		return true;
 	}
 
-	return true;
+	return false;
 }
 
 float USoundNodeLooping::GetDuration()
