Fixed issue which could cause a crash if a camera anim terminated the sequence during InterpGroup update.

diff --git a/Engine/Source/Runtime/Engine/Private/Interpolation.cpp b/Engine/Source/Runtime/Engine/Private/Interpolation.cpp
--- a/Engine/Source/Runtime/Engine/Private/Interpolation.cpp
+++ b/Engine/Source/Runtime/Engine/Private/Interpolation.cpp
@@ -2043,10 +2043,17 @@ void UInterpGroup::UpdateGroup(float NewPosition, UInterpGroupInst* GrInst, bool
 	}
 #endif
 
-	AMatineeActor* MatineeActor = Cast<AMatineeActor>(GrInst->GetOuter());
-
-	for(int32 i=0; i<InterpTracks.Num() && MatineeActor->bIsPlaying; i++)
+	for(int32 i=0; i<InterpTracks.Num(); i++)
 	{
+		// If the track instances have been removed from the group instance, this means that a previous track update has terminated the sequence.
+		// The group instance itself will still be valid, but unreferenced.
+		const bool bHasBeenTerminated = (GrInst->TrackInst.Num() == 0);
+
+		if (bHasBeenTerminated)
+		{
+			break;
+		}
+
 		UInterpTrack* Track = InterpTracks[i];
 		UInterpTrackInst* TrInst = GrInst->TrackInst[i];
 
