Fixed additive animation playback in Persona.

diff --git a/Engine/Source/Runtime/Engine/Private/Animation/AnimSingleNodeInstance.cpp b/Engine/Source/Runtime/Engine/Private/Animation/AnimSingleNodeInstance.cpp
--- a/Engine/Source/Runtime/Engine/Private/Animation/AnimSingleNodeInstance.cpp
+++ b/Engine/Source/Runtime/Engine/Private/Animation/AnimSingleNodeInstance.cpp
@@ -319,12 +319,14 @@ bool UAnimSingleNodeInstance::NativeEvaluateAnimation(FPoseContext& Output)
 		{
 			if (Sequence->IsValidAdditive())
 			{
+				FAnimExtractContext ExtractionContext(CurrentTime, Sequence->bEnableRootMotion);
+
+				Sequence->GetAdditiveBasePose(Output.Pose, ExtractionContext);
+
 				FCompactPose AdditivePose;
 				AdditivePose.SetBoneContainer(&Output.Pose.GetBoneContainer());
+				Sequence->GetAnimationPose(AdditivePose, ExtractionContext);
 
-				FAnimExtractContext ExtractionContext(CurrentTime, Sequence->bEnableRootMotion);
-				Sequence->GetAdditiveBasePose(AdditivePose, ExtractionContext);
-				Sequence->GetAnimationPose(Output.Pose, ExtractionContext);
 				if (Sequence->AdditiveAnimType == AAT_LocalSpaceBase)
 				{
 					FAnimationRuntime::AccumulateAdditivePose(Output.Pose, AdditivePose, 1.f);
