Fixed shadowed variable compile error on Mac.  It looks like the variable definition was just accidentally copy/pasted outside of the if() scope where it's used.

#UE4 #BuildWatcher

diff --git a/Engine/Source/Runtime/Engine/Private/PhysicsEngine/BodySetup.cpp b/Engine/Source/Runtime/Engine/Private/PhysicsEngine/BodySetup.cpp
--- a/Engine/Source/Runtime/Engine/Private/PhysicsEngine/BodySetup.cpp
+++ b/Engine/Source/Runtime/Engine/Private/PhysicsEngine/BodySetup.cpp
@@ -535,10 +535,6 @@ public:
 				PTriMeshGeom.meshFlags |= PxMeshGeometryFlag::eDOUBLE_SIDED;
 			}
 
-			PxTransform PElementTransform = U2PTransform(RelativeTM);
-			
-
-
 			if (PTriMeshGeom.isValid())
 			{
 				PxTransform PElementTransform = U2PTransform(RelativeTM);
