fixed UE-16626 SceneTexture:Scene Color node renders black in Standalone, Launch On and Packaged Game
the material need to use SceneTexture PostprocessInput0 instead

diff --git a/Engine/Source/Runtime/Engine/Private/Materials/HLSLMaterialTranslator.h b/Engine/Source/Runtime/Engine/Private/Materials/HLSLMaterialTranslator.h
--- a/Engine/Source/Runtime/Engine/Private/Materials/HLSLMaterialTranslator.h
+++ b/Engine/Source/Runtime/Engine/Private/Materials/HLSLMaterialTranslator.h
@@ -2591,6 +2591,18 @@ protected:
 	{
 		MaterialCompilationOutput.bNeedsSceneTextures = true;
 
+		if(SceneTextureId == PPI_SceneColor && Material->GetMaterialDomain() != MD_Surface)
+		{
+			if(Material->GetMaterialDomain() == MD_PostProcess)
+			{
+				Errorf(TEXT("SceneColor lookups are only available when MaterialDomain = Surface. PostProcessMaterials should use the SceneTexture PostProcessInput0."));
+			}
+			else
+			{
+				Errorf(TEXT("SceneColor lookups are only available when MaterialDomain = Surface."));
+			}
+		}
+
 		if(bTextureLookup)
 		{
 			bNeedsSceneTexturePostProcessInputs = bNeedsSceneTexturePostProcessInputs
@@ -2628,6 +2640,11 @@ protected:
 			return NonPixelShaderExpressionError();
 		}
 
+		if(Material->GetMaterialDomain() != MD_Surface)
+		{
+			Errorf(TEXT("SceneColor lookups are only available when MaterialDomain = Surface."));
+		}
+
 		if (ErrorUnlessFeatureLevelSupported(ERHIFeatureLevel::SM4) == INDEX_NONE)
 		{
 			return INDEX_NONE;
