fixed typo

diff --git a/Engine/Source/Runtime/Engine/Private/LevelTick.cpp b/Engine/Source/Runtime/Engine/Private/LevelTick.cpp
--- a/Engine/Source/Runtime/Engine/Private/LevelTick.cpp
+++ b/Engine/Source/Runtime/Engine/Private/LevelTick.cpp
@@ -401,7 +401,7 @@ bool UWorld::IsPaused()
 
 
 /**
- * Streaming settings for levels which are detrmined visible by level streaming volumes.
+ * Streaming settings for levels which are determined visible by level streaming volumes.
  */
 class FVisibleLevelStreamingSettings
 {
