Exposes translated world position for the WorldPosition node in the vertex and tessellation shaders

diff --git a/Engine/Shaders/MaterialTemplate.usf b/Engine/Shaders/MaterialTemplate.usf
--- a/Engine/Shaders/MaterialTemplate.usf
+++ b/Engine/Shaders/MaterialTemplate.usf
@@ -292,11 +292,36 @@ float3 GetPrevWorldPosition_NoOffsets(FMaterialVertexParameters Parameters)
 }
 
 //TODO(bug UE-17131): We should compute world displacement for the previous frame
+float3 GetWorldPosition(FMaterialTessellationParameters Parameters)
+{
+	return Parameters.WorldPosition + View.PreViewTranslation;
+}
+
 float3 GetWorldPosition_NoOffsets(FMaterialTessellationParameters Parameters)
 {
 	return Parameters.WorldPosition;
 }
 
+float3 GetWorldPosition(FMaterialPixelParameters Parameters)
+{
+	return Parameters.WorldPosition;
+}
+
+float3 GetWorldPosition_NoOffsets(FMaterialPixelParameters Parameters)
+{
+	return Parameters.WorldPosition_NoOffsets;
+}
+
+float3 GetWorldPosition_CamRelative(FMaterialPixelParameters Parameters)
+{
+	return Parameters.WorldPosition_CamRelative;
+}
+
+float3 GetWorldPosition_NoOffsets_CamRelative(FMaterialPixelParameters Parameters)
+{
+	return Parameters.WorldPosition_NoOffsets_CamRelative;
+}
+
 #if DECAL_RENDERSTAGE && NUM_MATERIAL_TEXCOORDS
 	/*
 	 * Material node DecalMipmapLevel's code designed to avoid the 2x2 pixels artefacts on the edges around where the decal
diff --git a/Engine/Source/Runtime/Engine/Private/Materials/HLSLMaterialTranslator.h b/Engine/Source/Runtime/Engine/Private/Materials/HLSLMaterialTranslator.h
--- a/Engine/Source/Runtime/Engine/Private/Materials/HLSLMaterialTranslator.h
+++ b/Engine/Source/Runtime/Engine/Private/Materials/HLSLMaterialTranslator.h
@@ -2081,57 +2081,67 @@ protected:
 
 	virtual int32 WorldPosition(EWorldPositionIncludedOffsets WorldPositionIncludedOffsets) override
 	{
-		if (ShaderFrequency == SF_Pixel)
+		const TCHAR * WorldPositionSuffix = nullptr; 
+
+		// If this material has no expressions for world position offset or world displacement, the non-offset world position will
+		// be exactly the same as the offset one, so there is no point bringing in the extra code.
+		// Also, we can't access the full offset world position in anything other than the pixel shader, because it won't have
+		// been calculated yet
+		switch (WorldPositionIncludedOffsets)
 		{
-			// If this material has no expressions for world position offset or world displacement, the non-offset world position will
-			// be exactly the same as the offset one, so there is no point bringing in the extra code.
-			// Also, we can't access the full offset world position in anything other than the pixel shader, because it won't have
-			// been calculated yet
-			bool bNonOffsetWorldPositionAvailable = Material->MaterialMayModifyMeshPosition();
+		case WPT_Default:
+			{
+				WorldPositionSuffix = TEXT("");
+				break;
+			}
 
-			switch (WorldPositionIncludedOffsets)
+		case WPT_ExcludeAllShaderOffsets:
 			{
-			case WPT_Default:
-				{
-					return AddInlinedCodeChunk(MCT_Float3, TEXT("Parameters.WorldPosition"));
-				}
+				bNeedsWorldPositionExcludingShaderOffsets = true;
+				WorldPositionSuffix = TEXT("_NoOffsets");
+				break;
+			}
 
-			case WPT_ExcludeAllShaderOffsets:
+		case WPT_CameraRelative:
+			{
+				if (ShaderFrequency != SF_Pixel)
 				{
-					bNeedsWorldPositionExcludingShaderOffsets = true;
-					return AddInlinedCodeChunk(MCT_Float3, TEXT("Parameters.WorldPosition_NoOffsets"));
+					return Errorf(TEXT("Cannot access camera relative world position outside of the pixel shader"));
 				}
 
-			case WPT_CameraRelative:
-				{
-					return AddInlinedCodeChunk(MCT_Float3, TEXT("Parameters.WorldPosition_CamRelative"));
-				}
+				WorldPositionSuffix = TEXT("_CamRelative");
+				break;
+			}
 
-			case WPT_CameraRelativeNoOffsets:
+		case WPT_CameraRelativeNoOffsets:
+			{
+				if (ShaderFrequency != SF_Pixel)
 				{
-					bNeedsWorldPositionExcludingShaderOffsets = true;
-					return AddInlinedCodeChunk(MCT_Float3, TEXT("Parameters.WorldPosition_NoOffsets_CamRelative"));
+					return Errorf(TEXT("Cannot access camera relative world position outside of the pixel shader"));
 				}
 
-			default:
-				{
-					Errorf(TEXT("Encountered unknown world position type '%d'"), WorldPositionIncludedOffsets);
-					return INDEX_NONE;
-				}
+				bNeedsWorldPositionExcludingShaderOffsets = true;
+				WorldPositionSuffix = TEXT("_NoOffsets_CamRelative");
+				break;
+			}
+
+		default:
+			{
+				Errorf(TEXT("Encountered unknown world position type '%d'"), WorldPositionIncludedOffsets);
+				return INDEX_NONE;
 			}
 		}
+
+		check(WorldPositionSuffix);
+		
+		if (bCompilingPreviousFrame && ShaderFrequency == SF_Vertex)
+		{
+			//TODO(bug UE-17131): We should compute world displacement for the previous frame
+			return AddInlinedCodeChunk(MCT_Float3, TEXT("GetPrevWorldPosition%s(Parameters)"), WorldPositionSuffix);
+		}
 		else
 		{
-			if (bCompilingPreviousFrame && ShaderFrequency == SF_Vertex)
-			{
-				// Only if not in a pixel shader...
-				return AddInlinedCodeChunk(MCT_Float3, TEXT("GetPrevWorldPosition_NoOffsets(Parameters)"));
-			}
-			else
-			{
-				// If not in a pixel shader, only the normal WorldPosition is available
-				return AddInlinedCodeChunk(MCT_Float3, TEXT("GetWorldPosition_NoOffsets(Parameters)"));
-			}
+			return AddInlinedCodeChunk(MCT_Float3, TEXT("GetWorldPosition%s(Parameters)"), WorldPositionSuffix);
 		}
 	}
 
