fixed UE-17047 Motion Blur exaggeration when Paused in Simulate

diff --git a/Engine/Source/Runtime/Engine/Private/SceneView.cpp b/Engine/Source/Runtime/Engine/Private/SceneView.cpp
--- a/Engine/Source/Runtime/Engine/Private/SceneView.cpp
+++ b/Engine/Source/Runtime/Engine/Private/SceneView.cpp
@@ -19,6 +19,10 @@
 #include "Classes/Engine/RendererSettings.h"
 #include "LightPropagationVolumeBlendable.h"
 
+#if WITH_EDITOR
+	#include "UnrealEd.h"
+#endif
+
 DEFINE_LOG_CATEGORY(LogBufferVisualization);
 
 IMPLEMENT_UNIFORM_BUFFER_STRUCT(FPrimitiveUniformShaderParameters,TEXT("Primitive"));
@@ -1543,7 +1547,16 @@ FSceneViewFamily::FSceneViewFamily( const ConstructionValues& CVS )
 
 			if (World->IsPaused())
 			{
-				bWorldIsPaused = true;
+				bool bCameraCanMove = false;
+
+#if WITH_EDITOR
+				// to fix UE-17047 Motion Blur exaggeration when Paused in Simulate
+				bCameraCanMove = GEditor && GEditor->bIsSimulatingInEditor;
+#endif
+				if(!bCameraCanMove)
+				{
+					bWorldIsPaused = true;
+				}
 			}
 		}
 	}
diff --git a/Engine/Source/Runtime/Engine/Public/SceneView.h b/Engine/Source/Runtime/Engine/Public/SceneView.h
--- a/Engine/Source/Runtime/Engine/Public/SceneView.h
+++ b/Engine/Source/Runtime/Engine/Public/SceneView.h
@@ -844,7 +844,10 @@ public:
 	/** if true then results of scene rendering are copied/resolved to the RenderTarget. */
 	bool bResolveScene;
 
-	/** from GetWorld->IsPaused() */
+	/**
+	 * GetWorld->IsPaused() && !Simulate
+	 * Simulate is excluded as the camera can move which invalidates motionblur
+	 * /
 	bool bWorldIsPaused;
 
 	/** Gamma correction used when rendering this family. Default is 1.0 */
