#ue4 - Make UCharacterMovementComponent::GroundMovementMode private, and provide Get/Set methods.

- Enforces that GroundMovementMode is only ever Walking or NavWalking.
- Set method will switch movement modes if currently a different ground walking mode.
- Cleaned up constants in Pack/UnpackNetworkMovementMode() and made sure we only have a valid ground mode (Walking or NavWalking).

diff --git a/Engine/Source/Runtime/Engine/Classes/GameFramework/CharacterMovementComponent.h b/Engine/Source/Runtime/Engine/Classes/GameFramework/CharacterMovementComponent.h
--- a/Engine/Source/Runtime/Engine/Classes/GameFramework/CharacterMovementComponent.h
+++ b/Engine/Source/Runtime/Engine/Classes/GameFramework/CharacterMovementComponent.h
@@ -566,10 +566,16 @@ public:
 	UPROPERTY(Category="Character Movement (General Settings)", EditAnywhere, BlueprintReadWrite)
 	TEnumAsByte<enum EMovementMode> DefaultWaterMovementMode;
 
-	/** Ground movement mode to switch to after falling */
+private:
+	/**
+	 * Ground movement mode to switch to after falling and resuming ground movement.
+	 * Only allowed values are: MOVE_Walking, MOVE_NavWalking.
+	 * @see SetGroundMovementMode(), GetGroundMovementMode()
+	 */
 	UPROPERTY(Transient)
 	TEnumAsByte<enum EMovementMode> GroundMovementMode;
 
+public:
 	/**
 	 * If true, walking movement always maintains horizontal velocity when moving up ramps, which causes movement up ramps to be faster parallel to the ramp surface.
 	 * If false, then walking movement maintains velocity magnitude parallel to the ramp surface.
@@ -801,6 +807,23 @@ public:
 	UFUNCTION(BlueprintCallable, Category="Pawn|Components|CharacterMovement")
 	virtual void SetMovementMode(EMovementMode NewMovementMode, uint8 NewCustomMode = 0);
 
+	/**
+	 * Set movement mode to use when returning to walking movement (either MOVE_Walking or MOVE_NavWalking).
+	 * If movement mode is currently one of Walking or NavWalking, this will also change the current movement mode (via SetMovementMode())
+	 * if the new mode is not the current ground mode.
+	 * 
+	 * @param  NewGroundMovementMode New ground movement mode. Must be either MOVE_Walking or MOVE_NavWalking, other values are ignored.
+	 * @see GroundMovementMode
+	 */
+	 void SetGroundMovementMode(EMovementMode NewGroundMovementMode);
+
+	/**
+	 * Get current GroundMovementMode value.
+	 * @return current GroundMovementMode
+	 * @see GroundMovementMode, SetGroundMovementMode()
+	 */
+	 EMovementMode GetGroundMovementMode() const { return GroundMovementMode; }
+
 protected:
 
 	/** Called after MovementMode has changed. Base implementation does special handling for starting certain modes, then notifies the CharacterOwner. */
diff --git a/Engine/Source/Runtime/Engine/Private/Components/CharacterMovementComponent.cpp b/Engine/Source/Runtime/Engine/Private/Components/CharacterMovementComponent.cpp
--- a/Engine/Source/Runtime/Engine/Private/Components/CharacterMovementComponent.cpp
+++ b/Engine/Source/Runtime/Engine/Private/Components/CharacterMovementComponent.cpp
@@ -620,6 +620,25 @@ void UCharacterMovementComponent::SetDefaultMovementMode()
 	}
 }
 
+void UCharacterMovementComponent::SetGroundMovementMode(EMovementMode NewGroundMovementMode)
+{
+	// Enforce restriction that it's either Walking or NavWalking.
+	if (NewGroundMovementMode != MOVE_Walking && NewGroundMovementMode != MOVE_NavWalking)
+	{
+		return;
+	}
+
+	// Set new value
+	GroundMovementMode = NewGroundMovementMode;
+
+	// Possibly change movement modes if already on ground and choosing the other ground mode.
+	const bool bOnGround = (MovementMode == MOVE_Walking || MovementMode == MOVE_NavWalking);
+	if (bOnGround && MovementMode != NewGroundMovementMode)
+	{
+		SetMovementMode(NewGroundMovementMode);
+	}
+}
+
 void UCharacterMovementComponent::SetMovementMode(EMovementMode NewMovementMode, uint8 NewCustomMode)
 {
 	if (NewMovementMode != MOVE_Custom)
@@ -732,45 +751,46 @@ void UCharacterMovementComponent::OnMovementModeChanged(EMovementMode PreviousMo
 	}
 
 	CharacterOwner->OnMovementModeChanged(PreviousMovementMode, PreviousCustomMode);
+	ensureOnce(GroundMovementMode == MOVE_Walking || GroundMovementMode == MOVE_NavWalking);
 };
 
 
-uint8 UCharacterMovementComponent::PackNetworkMovementMode() const
+namespace PackedMovementModeConstants
 {
 	const uint32 GroundShift = FMath::CeilLogTwo(MOVE_MAX);
 	const uint8 CustomModeThr = 2 * (1 << GroundShift);
+	const uint8 GroundMask = (1 << GroundShift) - 1;
+}
 
+uint8 UCharacterMovementComponent::PackNetworkMovementMode() const
+{
 	if (MovementMode != MOVE_Custom)
 	{
+		ensureOnce(GroundMovementMode == MOVE_Walking || GroundMovementMode == MOVE_NavWalking);
 		const uint8 GroundModeBit = (GroundMovementMode == MOVE_Walking ? 0 : 1);
-		return uint8(MovementMode.GetValue()) | (GroundModeBit << GroundShift);
+		return uint8(MovementMode.GetValue()) | (GroundModeBit << PackedMovementModeConstants::GroundShift);
 	}
 	else
 	{
-		return CustomMovementMode + CustomModeThr;
+		return CustomMovementMode + PackedMovementModeConstants::CustomModeThr;
 	}
 }
 
 
 void UCharacterMovementComponent::UnpackNetworkMovementMode(const uint8 ReceivedMode, TEnumAsByte<EMovementMode>& OutMode, uint8& OutCustomMode, TEnumAsByte<EMovementMode>& OutGroundMode) const
 {
-	const uint32 GroundShift = FMath::CeilLogTwo(MOVE_MAX);
-	const uint8 CustomModeThr = 2 * (1 << GroundShift);
-
-	if (ReceivedMode < CustomModeThr)
+	if (ReceivedMode < PackedMovementModeConstants::CustomModeThr)
 	{
-		const uint8 GroundMask = (1 << GroundShift) - 1;
-
-		OutMode = TEnumAsByte<EMovementMode>(ReceivedMode & GroundMask);
+		OutMode = TEnumAsByte<EMovementMode>(ReceivedMode & PackedMovementModeConstants::GroundMask);
 		OutCustomMode = 0;
-		const uint8 GroundModeBit = (ReceivedMode >> GroundShift);
+		const uint8 GroundModeBit = (ReceivedMode >> PackedMovementModeConstants::GroundShift);
 		OutGroundMode = TEnumAsByte<EMovementMode>(GroundModeBit == 0 ? MOVE_Walking : MOVE_NavWalking);
 	}
 	else
 	{
 		OutMode = MOVE_Custom;
-		OutCustomMode = ReceivedMode - CustomModeThr;
-		OutGroundMode = DefaultLandMovementMode;
+		OutCustomMode = ReceivedMode - PackedMovementModeConstants::CustomModeThr;
+		OutGroundMode = MOVE_Walking;
 	}
 }
 
@@ -781,6 +801,7 @@ void UCharacterMovementComponent::ApplyNetworkMovementMode(const uint8 ReceivedM
 	TEnumAsByte<EMovementMode> NetGroundMode(MOVE_None);
 	uint8 NetCustomMode(0);
 	UnpackNetworkMovementMode(ReceivedMode, NetMovementMode, NetCustomMode, NetGroundMode);
+	ensureOnce(NetGroundMode == MOVE_Walking || NetGroundMode == MOVE_NavWalking);
 
 	GroundMovementMode = NetGroundMode;
 	SetMovementMode(NetMovementMode, NetCustomMode);
