KismetMathLibrary
- added InvertTransform() blueprint node
- updated comment on ConvertTransformToRelative for clarity

diff --git a/Engine/Source/Runtime/Engine/Classes/Kismet/KismetMathLibrary.h b/Engine/Source/Runtime/Engine/Classes/Kismet/KismetMathLibrary.h
--- a/Engine/Source/Runtime/Engine/Classes/Kismet/KismetMathLibrary.h
+++ b/Engine/Source/Runtime/Engine/Classes/Kismet/KismetMathLibrary.h
@@ -1381,15 +1381,29 @@ class ENGINE_API UKismetMathLibrary : public UBlueprintFunctionLibrary
 	static FTransform ComposeTransforms(const FTransform& A, const FTransform& B);
 
 	/** 
-	 *  Convert a world-transform from world-space into local-space.
-	 *  @param		WorldTransform	The transform you wish to convert
-	 *  @param		LocalTransform	The transform the conversion is relative to
-	 *  @return		A new relative transform
+	 * Returns the given transform, converted to be relative to the given ParentTransform.
+	 *
+	 * Example: AToB = ConvertTransformToRelative(AToWorld, BToWorld) to compute A relative to B.
+	 *
+	 * @param		Transform		The transform you wish to convert
+	 * @param		ParentTransform	The transform the conversion is relative to (in the same space as Transform)
+	 * @return		The new relative transform
 	 */
 	UFUNCTION(BlueprintPure, Category="Math|Transform", meta=(Keywords="cast convert"))
-	static FTransform ConvertTransformToRelative(const FTransform& WorldTransform, const FTransform& LocalTransform);
+	static FTransform ConvertTransformToRelative(const FTransform& Transform, const FTransform& ParentTransform);
+
+	/** 
+	 * Returns the inverse of the given transform T.
+	 * 
+	 * Example: Given a LocalToWorld transform, WorldToLocal will be returned.
+	 *
+	 * @param	T	The transform you wish to invert
+	 * @return	The inverse of T.
+	 */
+	UFUNCTION(BlueprintPure, Category = "Math|Transform", meta = (Keywords = "inverse"))
+	static FTransform InvertTransform(const FTransform& T);
 
-	/* Linearly interpolates between A and B based on Alpha (100% of A when Alpha=0 and 100% of B when Alpha=1). */
+	/** Linearly interpolates between A and B based on Alpha (100% of A when Alpha=0 and 100% of B when Alpha=1). */
 	UFUNCTION(BlueprintPure, meta=(DisplayName = "Lerp (Transform)"), Category="Math|Transform")
 	static FTransform TLerp(const FTransform& A, const FTransform& B, float Alpha);
 
