missing files for:
KismetMathLibrary
- added InvertTransform() blueprint node
- updated comment on ConvertTransformToRelative for clarity

diff --git a/Engine/Source/Runtime/Engine/Private/KismetMathLibrary.cpp b/Engine/Source/Runtime/Engine/Private/KismetMathLibrary.cpp
--- a/Engine/Source/Runtime/Engine/Private/KismetMathLibrary.cpp
+++ b/Engine/Source/Runtime/Engine/Private/KismetMathLibrary.cpp
@@ -1082,6 +1082,11 @@ FTransform UKismetMathLibrary::ConvertTransformToRelative(const FTransform& Worl
 	return LocalTransform.GetRelativeTransformReverse(WorldTransform);
 }
 
+FTransform UKismetMathLibrary::InvertTransform(const FTransform& T)
+{
+	return T.Inverse();
+}
+
 FTransform UKismetMathLibrary::TLerp(const FTransform& A, const FTransform& B, float Alpha)
 {
 	FTransform Result;
