PR #1250: Allow DefaultPhysicsVolume to be overridden in WorldSettings (Contributed by [REDACTED_PERSON])

diff --git a/Engine/Source/Runtime/Engine/Classes/GameFramework/WorldSettings.h b/Engine/Source/Runtime/Engine/Classes/GameFramework/WorldSettings.h
--- a/Engine/Source/Runtime/Engine/Classes/GameFramework/WorldSettings.h
+++ b/Engine/Source/Runtime/Engine/Classes/GameFramework/WorldSettings.h
@@ -294,6 +294,10 @@ class ENGINE_API AWorldSettings : public AInfo, public IInterface_AssetUserData
 	UPROPERTY(EditAnywhere, BlueprintReadOnly, Category=Physics, meta=(editcondition = "bGlobalGravitySet"))
 	float GlobalGravityZ;
 
+	// level specific default physics volume
+	UPROPERTY(EditAnywhere, noclear, BlueprintReadOnly, Category=Physics, AdvancedDisplay)
+	TSubclassOf<class ADefaultPhysicsVolume> DefaultPhysicsVolumeClass;
+
 	// optional level specific collision handler
 	UPROPERTY(EditAnywhere, BlueprintReadOnly, Category=Physics, AdvancedDisplay)
 	TSubclassOf<class UPhysicsCollisionHandler>	PhysicsCollisionHandlerClass;
diff --git a/Engine/Source/Runtime/Engine/Private/DefaultPhysicsVolume.cpp b/Engine/Source/Runtime/Engine/Private/DefaultPhysicsVolume.cpp
--- a/Engine/Source/Runtime/Engine/Private/DefaultPhysicsVolume.cpp
+++ b/Engine/Source/Runtime/Engine/Private/DefaultPhysicsVolume.cpp
@@ -2,6 +2,7 @@
 
 #include "EnginePrivate.h"
 #include "GameFramework/DefaultPhysicsVolume.h"
+#include "PhysicsEngine/PhysicsSettings.h"
 
 ADefaultPhysicsVolume::ADefaultPhysicsVolume(const FObjectInitializer& ObjectInitializer)
 	: Super(ObjectInitializer)
@@ -11,6 +12,10 @@ ADefaultPhysicsVolume::ADefaultPhysicsVolume(const FObjectInitializer& ObjectIni
 	// Not allowed to be selected or edited within Unreal Editor
 	bEditable = false;
 #endif // WITH_EDITORONLY_DATA
+
+	// update default values when world is restarted
+	TerminalVelocity = UPhysicsSettings::Get()->DefaultTerminalVelocity;
+	FluidFriction = UPhysicsSettings::Get()->DefaultFluidFriction;
 }
 
 void ADefaultPhysicsVolume::Destroyed()
diff --git a/Engine/Source/Runtime/Engine/Private/World.cpp b/Engine/Source/Runtime/Engine/Private/World.cpp
--- a/Engine/Source/Runtime/Engine/Private/World.cpp
+++ b/Engine/Source/Runtime/Engine/Private/World.cpp
@@ -895,10 +895,6 @@ void UWorld::InitWorld(const InitializationValues IVS)
 	// Spawned on demand by this function.
 	DefaultPhysicsVolume = GetDefaultPhysicsVolume();
 
-	// update default values when world is restarted
-	DefaultPhysicsVolume->TerminalVelocity = UPhysicsSettings::Get()->DefaultTerminalVelocity;
-	DefaultPhysicsVolume->FluidFriction = UPhysicsSettings::Get()->DefaultFluidFriction;
-
 	// Find gravity
 	if (GetPhysicsScene())
 	{
@@ -3369,14 +3365,24 @@ void UWorld::SetPhysicsScene(FPhysScene* InScene)
 	}
 }
 
-
 APhysicsVolume* UWorld::GetDefaultPhysicsVolume() const
 {
 	// Create on demand.
-	if (!DefaultPhysicsVolume)
+	if (DefaultPhysicsVolume == nullptr)
 	{
+		// Try WorldSettings first
+		AWorldSettings* WorldSettings = GetWorldSettings();
+		UClass* DefaultPhysicsVolumeClass = WorldSettings->DefaultPhysicsVolumeClass;
+
+		// Fallback on DefaultPhysicsVolume static
+		if (DefaultPhysicsVolumeClass == nullptr)
+		{
+			DefaultPhysicsVolumeClass = ADefaultPhysicsVolume::StaticClass();
+		}
+
+		// Spawn volume
 		UWorld* MutableThis = const_cast<UWorld*>(this);
-		MutableThis->DefaultPhysicsVolume = MutableThis->SpawnActor<APhysicsVolume>(ADefaultPhysicsVolume::StaticClass());
+		MutableThis->DefaultPhysicsVolume = MutableThis->SpawnActor<APhysicsVolume>(DefaultPhysicsVolumeClass);
 		MutableThis->DefaultPhysicsVolume->Priority = -1000000;
 	}
 	return DefaultPhysicsVolume;
diff --git a/Engine/Source/Runtime/Engine/Private/WorldSettings.cpp b/Engine/Source/Runtime/Engine/Private/WorldSettings.cpp
--- a/Engine/Source/Runtime/Engine/Private/WorldSettings.cpp
+++ b/Engine/Source/Runtime/Engine/Private/WorldSettings.cpp
@@ -12,6 +12,7 @@
 #include "MapErrors.h"
 #include "Particles/ParticleEventManager.h"
 #include "PhysicsEngine/PhysicsSettings.h"
+#include "GameFramework/DefaultPhysicsVolume.h"
 
 #define LOCTEXT_NAMESPACE "ErrorChecking"
 
@@ -45,6 +46,7 @@ AWorldSettings::AWorldSettings(const FObjectInitializer& ObjectInitializer)
 
 	WorldToMeters = 100.f;
 
+	DefaultPhysicsVolumeClass = ADefaultPhysicsVolume::StaticClass();
 	GameNetworkManagerClass = AGameNetworkManager::StaticClass();
 	SetRemoteRoleForBackwardsCompat(ROLE_SimulatedProxy);
 	bReplicates = true;
