#ue4 - UCharacterMovement::StepUp() should not use SafeMoveUpdatedComponent for its movement, this could cause movements to move in reverse directions or violate step-up height rules.

diff --git a/Engine/Source/Runtime/Engine/Private/Components/CharacterMovementComponent.cpp b/Engine/Source/Runtime/Engine/Private/Components/CharacterMovementComponent.cpp
--- a/Engine/Source/Runtime/Engine/Private/Components/CharacterMovementComponent.cpp
+++ b/Engine/Source/Runtime/Engine/Private/Components/CharacterMovementComponent.cpp
@@ -5255,20 +5255,19 @@ bool UCharacterMovementComponent::StepUp(const FVector& InGravDir, const FVector
 	// step up - treat as vertical wall
 	FHitResult SweepUpHit(1.f);
 	const FQuat PawnRotation = UpdatedComponent->GetComponentQuat();
-	SafeMoveUpdatedComponent(-GravDir * StepTravelUpHeight, PawnRotation, true, SweepUpHit);
+	MoveUpdatedComponent(-GravDir * StepTravelUpHeight, PawnRotation, true, &SweepUpHit);
 
-	// step fwd
-	FHitResult Hit(1.f);
-	SafeMoveUpdatedComponent( Delta, PawnRotation, true, Hit);
-
-	// If we hit something above us and also something ahead of us, we should notify about the upward hit as well.
-	// The forward hit will be handled later (in the bSteppedOver case below).
-	// In the case of hitting something above but not forward, we are not blocked from moving so we don't need the notification.
-	if (SweepUpHit.bBlockingHit && Hit.bBlockingHit)
+	if (SweepUpHit.bStartPenetrating)
 	{
-		HandleImpact(SweepUpHit);
+		// Undo movement
+		ScopedStepUpMovement.RevertMove();
+		return false;
 	}
 
+	// step fwd
+	FHitResult Hit(1.f);
+	MoveUpdatedComponent( Delta, PawnRotation, true, &Hit);
+
 	// Check result of forward movement
 	if (Hit.bBlockingHit)
 	{
@@ -5279,9 +5278,17 @@ bool UCharacterMovementComponent::StepUp(const FVector& InGravDir, const FVector
 			return false;
 		}
 
+		// If we hit something above us and also something ahead of us, we should notify about the upward hit as well.
+		// The forward hit will be handled later (in the bSteppedOver case below).
+		// In the case of hitting something above but not forward, we are not blocked from moving so we don't need the notification.
+		if (SweepUpHit.bBlockingHit && Hit.bBlockingHit)
+		{
+			HandleImpact(SweepUpHit);
+		}
+
 		// pawn ran into a wall
 		HandleImpact(Hit);
-		if ( IsFalling() )
+		if (IsFalling())
 		{
 			return true;
 		}
@@ -5305,7 +5312,7 @@ bool UCharacterMovementComponent::StepUp(const FVector& InGravDir, const FVector
 	}
 	
 	// Step down
-	SafeMoveUpdatedComponent(GravDir * StepTravelDownHeight, UpdatedComponent->GetComponentQuat(), true, Hit);
+	MoveUpdatedComponent(GravDir * StepTravelDownHeight, UpdatedComponent->GetComponentQuat(), true, &Hit);
 
 	// If step down was initially penetrating abort the step up
 	if (Hit.bStartPenetrating)
