UE4 - Add extra info on message (UE-8973)

diff --git a/Engine/Source/Runtime/Engine/Private/GlobalShader.cpp b/Engine/Source/Runtime/Engine/Private/GlobalShader.cpp
--- a/Engine/Source/Runtime/Engine/Private/GlobalShader.cpp
+++ b/Engine/Source/Runtime/Engine/Private/GlobalShader.cpp
@@ -395,7 +395,7 @@ TShaderMap<FGlobalShaderType>* GetGlobalShaderMap(EShaderPlatform Platform, bool
 				// Handle this gracefully and exit.
 				FString SandboxPath = IFileManager::Get().ConvertToAbsolutePathForExternalAppForWrite(*GlobalShaderCacheFilename);				
 				// This can be too early to localize in some situations.
-				const FText Message = FText::Format( NSLOCTEXT("Engine", "GlobalShaderCacheFileMissing", "The global shader cache file '{0}' is missing.\n\nYou're running a version of the application built to load COOKED content only, however no COOKED content was found. Consider cooking content for this build, or build and run the UNCOOKED version of the application instead."), FText::FromString( SandboxPath ) );
+				const FText Message = FText::Format(NSLOCTEXT("Engine", "GlobalShaderCacheFileMissing", "The global shader cache file '{0}' is missing.\n\nYour application is built to load COOKED content. No COOKED content was found; This usually means you did not cook content for this build.\nIt also may indicate missing cooked data for a shader platform(e.g., OpenGL under Windows): Make sure your platform's packaging settings include this Targeted RHI.\n\nAlternatively build and run the UNCOOKED version instead."), FText::FromString( SandboxPath ) );
 				if (FPlatformProperties::SupportsWindowedMode())
 				{
 					UE_LOG(LogMaterial, Error, TEXT("%s"), *Message.ToString());
