Animation: Added a virtual dtor to FAnimNode_Base
- Converted/improved some comments in AnimNodeBase.h

diff --git a/Engine/Source/Runtime/Engine/Classes/Animation/AnimNodeBase.h b/Engine/Source/Runtime/Engine/Classes/Animation/AnimNodeBase.h
--- a/Engine/Source/Runtime/Engine/Classes/Animation/AnimNodeBase.h
+++ b/Engine/Source/Runtime/Engine/Classes/Animation/AnimNodeBase.h
@@ -9,7 +9,7 @@
 #include "BonePose.h"
 #include "AnimNodeBase.generated.h"
 
-// Base class for update/evaluate contexts
+/** Base class for update/evaluate contexts */
 struct FAnimationBaseContext
 {
 public:
@@ -43,7 +43,7 @@ public:
 };
 
 
-// Initialization context passed around during animation tree initialization
+/** Initialization context passed around during animation tree initialization */
 struct FAnimationInitializeContext : public FAnimationBaseContext
 {
 public:
@@ -53,8 +53,10 @@ public:
 	}
 };
 
-// Context passed around when RequiredBones array changed and cached bones indices have to be refreshed.
-// (RequiredBones array changed because of an LOD switch for example)
+/**
+ * Context passed around when RequiredBones array changed and cached bones indices have to be refreshed.
+ * (RequiredBones array changed because of an LOD switch for example)
+ */
 struct FAnimationCacheBonesContext : public FAnimationBaseContext
 {
 public:
@@ -64,7 +66,7 @@ public:
 	}
 };
 
-// Update context passed around during animation tree update
+/** Update context passed around during animation tree update */
 struct FAnimationUpdateContext : public FAnimationBaseContext
 {
 private:
@@ -100,7 +102,7 @@ public:
 };
 
 
-// Evaluation context passed around during animation tree evaluation
+/** Evaluation context passed around during animation tree evaluation */
 struct FPoseContext : public FAnimationBaseContext
 {
 public:
@@ -145,7 +147,7 @@ public:
 };
 
 
-// Evaluation context passed around during animation tree evaluation
+/** Evaluation context passed around during animation tree evaluation */
 struct FComponentSpacePoseContext : public FAnimationBaseContext
 {
 public:
@@ -258,7 +260,7 @@ namespace EPinHidingMode
 	};
 }
 
-// A pose link to another node
+/** A pose link to another node */
 USTRUCT()
 struct ENGINE_API FPoseLinkBase
 {
@@ -303,7 +305,7 @@ public:
 	void AttemptRelink(const FAnimationBaseContext& Context);
 };
 
-// A local-space pose link to another node
+/** A local-space pose link to another node */
 USTRUCT()
 struct ENGINE_API FPoseLink : public FPoseLinkBase
 {
@@ -314,7 +316,7 @@ public:
 	void Evaluate(FPoseContext& Output);
 };
 
-// A component-space pose link to another node
+/** A component-space pose link to another node */
 USTRUCT()
 struct ENGINE_API FComponentSpacePoseLink : public FPoseLinkBase
 {
@@ -348,11 +350,13 @@ struct FExposedValueHandler
 	}
 };
 
-// To create a new animation node:
-//   Create a struct derived from FAnimNode_Base - this is your runtime node
-//   Create a class derived from UAnimGraphNode_Base, containing an instance of your runtime as a member - this is your visual/editor-only node
-
-// This is the base of all runtime animation nodes
+/**
+ * This is the base of all runtime animation nodes
+ *
+ * To create a new animation node:
+ *   Create a struct derived from FAnimNode_Base - this is your runtime node
+ *   Create a class derived from UAnimGraphNode_Base, containing an instance of your runtime node as a member - this is your visual/editor-only node
+ */
 USTRUCT()
 struct ENGINE_API FAnimNode_Base
 {
@@ -379,4 +383,6 @@ struct ENGINE_API FAnimNode_Base
 		DebugData.AddDebugItem(TEXT("Non Overriden GatherDebugData")); 
 	}
 	// End of interface to implement
+
+	virtual ~FAnimNode_Base() {}
 };
