Made GetRandomPoint and GetRandomPointInRadius deprecated in BP #UE4

UE-17138

diff --git a/Engine/Source/Runtime/Engine/Classes/AI/Navigation/NavigationSystem.h b/Engine/Source/Runtime/Engine/Classes/AI/Navigation/NavigationSystem.h
--- a/Engine/Source/Runtime/Engine/Classes/AI/Navigation/NavigationSystem.h
+++ b/Engine/Source/Runtime/Engine/Classes/AI/Navigation/NavigationSystem.h
@@ -961,12 +961,12 @@ public:
 	DEPRECATED(4.8, "This version is deprecated. Please use the one taking reference to INavLinkCustomInterface rather than a pointer instead.")
 	void UnregisterCustomLink(INavLinkCustomInterface* CustomLink);
 	DEPRECATED(4.8, "GetRandomPointInRadius is deprecated. Use either GetRandomReachablePointInRadius or GetRandomPointInNavigableSpace")
-	UFUNCTION(BlueprintPure, Category = "AI|Navigation", meta = (WorldContext = "WorldContext"))
+	UFUNCTION(BlueprintPure, Category = "AI|Navigation", meta = (WorldContext = "WorldContext", DeprecatedFunction, DeprecatedMessage = "GetRandomPointInRadius is deprecated. Use either GetRandomReachablePointInRadius or GetRandomPointInNavigableSpace"))
 	static FVector GetRandomPointInRadius(UObject* WorldContext, const FVector& Origin, float Radius, ANavigationData* NavData = NULL, TSubclassOf<UNavigationQueryFilter> FilterClass = NULL);
 	DEPRECATED(4.8, "GetRandomPointInRadius is deprecated. Use either GetRandomReachablePointInRadius or GetRandomPointInNavigableSpace")
 	bool GetRandomPointInRadius(const FVector& Origin, float Radius, FNavLocation& ResultLocation, ANavigationData* NavData = NULL, TSharedPtr<const FNavigationQueryFilter> QueryFilter = NULL) const;
-	DEPRECATED(4.8, "GetRandomPointInRadius is deprecated. Use either GetRandomReachablePointInRadius or GetRandomPointInNavigableSpace")
-	UFUNCTION(BlueprintPure, Category = "AI|Navigation", meta = (WorldContext = "WorldContext"))
+	DEPRECATED(4.8, "GetRandomPoint is deprecated. Use either GetRandomReachablePointInRadius or GetRandomPointInNavigableSpace")
+	UFUNCTION(BlueprintPure, Category = "AI|Navigation", meta = (WorldContext = "WorldContext", DeprecatedFunction, DeprecatedMessage = "GetRandomPoint is deprecated. Use either GetRandomReachablePointInRadius or GetRandomPointInNavigableSpace"))
 	static FVector GetRandomPoint(UObject* WorldContext, ANavigationData* NavData = NULL, TSubclassOf<UNavigationQueryFilter> FilterClass = NULL);
 };
 
