PR #1249: [UDestructibleComponent::SetEnableGravity] Added ReadWrite Protection since setting gravity flag (Contributed by [REDACTED_PERSON])

diff --git a/Engine/Source/Runtime/Engine/Private/Components/DestructibleComponent.cpp b/Engine/Source/Runtime/Engine/Private/Components/DestructibleComponent.cpp
--- a/Engine/Source/Runtime/Engine/Private/Components/DestructibleComponent.cpp
+++ b/Engine/Source/Runtime/Engine/Private/Components/DestructibleComponent.cpp
@@ -1324,14 +1324,17 @@ void UDestructibleComponent::SetEnableGravity(bool bGravityEnabled)
 	Super::SetEnableGravity(bGravityEnabled);
 	
 #if WITH_APEX
-	uint32 ChunkCount = ApexDestructibleActor->getNumVisibleChunks();
-	const uint16* ChunkIndices = ApexDestructibleActor->getVisibleChunks();
-	for(uint32 c = 0; c < ChunkCount; c++)
+	ExecuteOnPhysicsReadWrite([&]
+	{
+		uint32 ChunkCount = ApexDestructibleActor->getNumVisibleChunks();
+		const uint16* ChunkIndices = ApexDestructibleActor->getVisibleChunks();
+		for(uint32 c = 0; c < ChunkCount; c++)
 		{
-		PxActor* PActor = ApexDestructibleActor->getChunkPhysXActor(ChunkIndices[c]);
-		check(PActor);
-		PActor->setActorFlag(PxActorFlag::eDISABLE_GRAVITY, !bGravityEnabled);
-	}
+			PxActor* PActor = ApexDestructibleActor->getChunkPhysXActor(ChunkIndices[c]);
+			check(PActor);
+			PActor->setActorFlag(PxActorFlag::eDISABLE_GRAVITY, !bGravityEnabled);
+		}
+	});
 #endif //WITH_APEX
 }
 
