Fixed bAccurateWindVelocities display name

diff --git a/Engine/Source/Runtime/Engine/Classes/Materials/MaterialExpressionSpeedTree.h b/Engine/Source/Runtime/Engine/Classes/Materials/MaterialExpressionSpeedTree.h
--- a/Engine/Source/Runtime/Engine/Classes/Materials/MaterialExpressionSpeedTree.h
+++ b/Engine/Source/Runtime/Engine/Classes/Materials/MaterialExpressionSpeedTree.h
@@ -51,7 +51,7 @@ class UMaterialExpressionSpeedTree : public UMaterialExpression
 	UPROPERTY(EditAnywhere, Category=MaterialExpressionSpeedTree, meta=(DisplayName = "Billboard Threshold", ToolTip="The threshold for triangles to be removed from the bilboard mesh when not facing the camera (0 = none pass, 1 = all pass).", ClampMin = "0.0", ClampMax = "1.0", UIMin = "0.0", UIMax = "1.0"))
 	float BillboardThreshold;
 
-	UPROPERTY(EditAnywhere, Category=MaterialExpressionSpeedTree, AdvancedDisplay, meta=(DisplayName = "Support accurate velocities from wind. This will incur extra cost per vertex."))
+	UPROPERTY(EditAnywhere, Category=MaterialExpressionSpeedTree, AdvancedDisplay, meta=(DisplayName = "Accurate Wind Velocities", ToolTip="Support accurate velocities from wind. This will incur extra cost per vertex."))
 	bool bAccurateWindVelocities;
 
 	// Begin UObject Interface
