Fixed LightSourceAngle being greyed out incorrectly

diff --git a/Engine/Source/Runtime/Engine/Private/Components/DirectionalLightComponent.cpp b/Engine/Source/Runtime/Engine/Private/Components/DirectionalLightComponent.cpp
--- a/Engine/Source/Runtime/Engine/Private/Components/DirectionalLightComponent.cpp
+++ b/Engine/Source/Runtime/Engine/Private/Components/DirectionalLightComponent.cpp
@@ -713,7 +713,7 @@ bool UDirectionalLightComponent::CanEditChange(const UProperty* InProperty) cons
 	{
 		FString PropertyName = InProperty->GetName();
 		
-		bool bShadowCascases = CastShadows
+		bool bShadowCascades = CastShadows
 			&& CastDynamicShadows 
 			&& ((DynamicShadowDistanceMovableLight > 0 && Mobility == EComponentMobility::Movable)
 			|| (DynamicShadowDistanceStationaryLight > 0 && Mobility == EComponentMobility::Stationary));
@@ -735,16 +735,18 @@ bool UDirectionalLightComponent::CanEditChange(const UProperty* InProperty) cons
 			|| PropertyName == GET_MEMBER_NAME_STRING_CHECKED(UDirectionalLightComponent, bUseInsetShadowsForMovableObjects)
 			|| PropertyName == GET_MEMBER_NAME_STRING_CHECKED(UDirectionalLightComponent, FarShadowCascadeCount))
 		{
-			return bShadowCascases;
+			return bShadowCascades;
 		}
 
 		if (PropertyName == GET_MEMBER_NAME_STRING_CHECKED(UDirectionalLightComponent, FarShadowDistance))
 		{
-			return bShadowCascases && FarShadowCascadeCount > 0;
+			return bShadowCascades && FarShadowCascadeCount > 0;
 		}
 
 		if (PropertyName == GET_MEMBER_NAME_STRING_CHECKED(UDirectionalLightComponent, DistanceFieldShadowDistance)
-			|| PropertyName == GET_MEMBER_NAME_STRING_CHECKED(UDirectionalLightComponent, LightSourceAngle))
+			|| (PropertyName == GET_MEMBER_NAME_STRING_CHECKED(UDirectionalLightComponent, LightSourceAngle)
+				// Make sure we don't accidentally affect the LightSourceAngle property in LightmassSettings
+				&& InProperty && InProperty->GetOuter() && InProperty->GetOuter()->GetName() == TEXT("DirectionalLightComponent")))
 		{
 			static const auto CVar = IConsoleManager::Get().FindTConsoleVariableDataInt(TEXT("r.GenerateMeshDistanceFields"));
 			bool bCanEdit = CastShadows && CastDynamicShadows && bUseRayTracedDistanceFieldShadows && Mobility != EComponentMobility::Static && CVar->GetValueOnGameThread() != 0;
