Merging

2589134
Fix issues with checkpoint scrubbing not retaining which actor we are following and refactor the view target for the replay spectator camera to always be the actor we are looking at (instead of always the spectator pawn)
Bart.Bressler

diff --git a/Engine/Source/Runtime/Engine/Classes/Camera/PlayerCameraManager.h b/Engine/Source/Runtime/Engine/Classes/Camera/PlayerCameraManager.h
--- a/Engine/Source/Runtime/Engine/Classes/Camera/PlayerCameraManager.h
+++ b/Engine/Source/Runtime/Engine/Classes/Camera/PlayerCameraManager.h
@@ -595,7 +595,7 @@ public:
 	 * @param NewViewTarget - New viewtarget actor.
 	 * @param TransitionParams - Optional parameters to define the interpolation from the old viewtarget to the new. Transition will be instant by default.
 	 */
-	void SetViewTarget(class AActor* NewViewTarget, FViewTargetTransitionParams TransitionParams = FViewTargetTransitionParams());
+	virtual void SetViewTarget(class AActor* NewViewTarget, FViewTargetTransitionParams TransitionParams = FViewTargetTransitionParams());
 	
 	/** 
 	 * Called to give PlayerCameraManager a chance to adjust view rotation updates before they are applied. 
