Logging for when blueprint component instance data is tossed, which causes a 'lighting needs to be rebuilt' message, but is usually a bug in how cached instance data is matched after rerunning the construction script.

diff --git a/Engine/Source/Runtime/Engine/Private/Components/StaticMeshComponent.cpp b/Engine/Source/Runtime/Engine/Private/Components/StaticMeshComponent.cpp
--- a/Engine/Source/Runtime/Engine/Private/Components/StaticMeshComponent.cpp
+++ b/Engine/Source/Runtime/Engine/Private/Components/StaticMeshComponent.cpp
@@ -1587,49 +1587,61 @@ void UStaticMeshComponent::ApplyComponentInstanceData(FStaticMeshComponentInstan
 
 	// Note: ApplyComponentInstanceData is called while the component is registered so the rendering thread is already using this component
 	// That means all component state that is modified here must be mirrored on the scene proxy, which will be recreated to receive the changes later due to MarkRenderStateDirty.
+	// Changing refcounted pointers like LightMap works because those are deferred cleanup resources (deleted next frame)
 
 	if (StaticMesh != StaticMeshInstanceData->StaticMesh)
 	{
 		return;
 	}
 
-	// See if data matches current state
-	if(	StaticMeshInstanceData->bHasCachedStaticLighting && StaticMeshInstanceData->CachedStaticLighting.Transform.Equals(ComponentToWorld, 1.e-3f) )
+	if (StaticMeshInstanceData->bHasCachedStaticLighting)
 	{
-		const int32 NumLODLightMaps = StaticMeshInstanceData->CachedStaticLighting.LODDataLightMap.Num();
-		SetLODDataCount(NumLODLightMaps, NumLODLightMaps);
-		for (int32 i = 0; i < NumLODLightMaps; ++i)
+		// See if data matches current state
+		if (StaticMeshInstanceData->CachedStaticLighting.Transform.Equals(ComponentToWorld, 1.e-3f))
 		{
-			LODData[i].LightMap = StaticMeshInstanceData->CachedStaticLighting.LODDataLightMap[i];
-			LODData[i].ShadowMap = StaticMeshInstanceData->CachedStaticLighting.LODDataShadowMap[i];
+			const int32 NumLODLightMaps = StaticMeshInstanceData->CachedStaticLighting.LODDataLightMap.Num();
+			SetLODDataCount(NumLODLightMaps, NumLODLightMaps);
 
-// this code is to try to track down a mystery GC crash from crash reporter...if this code crashes, and then we know it is component reinstancing
-			if (LODData[i].LightMap.GetReference())
+			for (int32 i = 0; i < NumLODLightMaps; ++i)
 			{
-				FLightMap2D* LightMap = LODData[i].LightMap->GetLightMap2D();
-				if (LightMap)
+				LODData[i].LightMap = StaticMeshInstanceData->CachedStaticLighting.LODDataLightMap[i];
+				LODData[i].ShadowMap = StaticMeshInstanceData->CachedStaticLighting.LODDataShadowMap[i];
+
+				// this code is to try to track down a mystery GC crash from crash reporter...if this code crashes, and then we know it is component reinstancing
+				if (LODData[i].LightMap.GetReference())
 				{
-					if (LightMap->IsValid(0))
-					{
-						UTexture2D* Tex = LightMap->GetTexture(0);
-						Tex->GetResourceSize(EResourceSizeMode::Exclusive);
-					}
-					if (LightMap->IsValid(1))
-					{
-						UTexture2D* Tex = LightMap->GetTexture(1);
-						Tex->GetResourceSize(EResourceSizeMode::Exclusive);
-					}
-					if (LightMap->GetSkyOcclusionTexture())
+					FLightMap2D* LightMap = LODData[i].LightMap->GetLightMap2D();
+					if (LightMap)
 					{
-						UTexture2D* Tex = LightMap->GetSkyOcclusionTexture();
-						Tex->GetResourceSize(EResourceSizeMode::Exclusive);
+						if (LightMap->IsValid(0))
+						{
+							UTexture2D* Tex = LightMap->GetTexture(0);
+							Tex->GetResourceSize(EResourceSizeMode::Exclusive);
+						}
+						if (LightMap->IsValid(1))
+						{
+							UTexture2D* Tex = LightMap->GetTexture(1);
+							Tex->GetResourceSize(EResourceSizeMode::Exclusive);
+						}
+						if (LightMap->GetSkyOcclusionTexture())
+						{
+							UTexture2D* Tex = LightMap->GetSkyOcclusionTexture();
+							Tex->GetResourceSize(EResourceSizeMode::Exclusive);
+						}
 					}
 				}
 			}
-		}
 
-		IrrelevantLights = StaticMeshInstanceData->CachedStaticLighting.IrrelevantLights;
-		bHasCachedStaticLighting = true;
+			IrrelevantLights = StaticMeshInstanceData->CachedStaticLighting.IrrelevantLights;
+			bHasCachedStaticLighting = true;
+		}
+		else
+		{
+			UE_LOG(LogStaticMesh, Warning, TEXT("Cached component instance data transform did not match!  Discarding cached lighting data which will cause lighting to be unbuilt.\n%s\nCurrent X=%f Y=%f Z=%f Cached X=%f Y=%f Z=%f"), 
+				*GetPathName(),
+				ComponentToWorld.GetLocation().X, ComponentToWorld.GetLocation().Y, ComponentToWorld.GetLocation().Z,
+				StaticMeshInstanceData->CachedStaticLighting.Transform.GetLocation().X, StaticMeshInstanceData->CachedStaticLighting.Transform.GetLocation().Y, StaticMeshInstanceData->CachedStaticLighting.Transform.GetLocation().Z);
+		}
 	}
 
 	{
