Fixed issue with camera shake oscillations not working in non-local playspaces. UE-16905

diff --git a/Engine/Source/Runtime/Engine/Private/Camera/CameraShake.cpp b/Engine/Source/Runtime/Engine/Private/Camera/CameraShake.cpp
--- a/Engine/Source/Runtime/Engine/Private/Camera/CameraShake.cpp
+++ b/Engine/Source/Runtime/Engine/Private/Camera/CameraShake.cpp
@@ -135,6 +135,7 @@ void UCameraShake::PlayShake(APlayerCameraManager* Camera, float Scale, ECameraA
 		}
 	}
 
+	PlaySpace = InPlaySpace;
 	if (InPlaySpace == ECameraAnimPlaySpace::UserDefined)
 	{
 		UserPlaySpaceMatrix = FRotationMatrix(UserPlaySpaceRot);
