#ue4 - Fix crash when Character capsule is marked pending kill and set null by MovementComponent::TickComponent.

UE-17167

diff --git a/Engine/Source/Runtime/Engine/Private/Components/CharacterMovementComponent.cpp b/Engine/Source/Runtime/Engine/Private/Components/CharacterMovementComponent.cpp
--- a/Engine/Source/Runtime/Engine/Private/Components/CharacterMovementComponent.cpp
+++ b/Engine/Source/Runtime/Engine/Private/Components/CharacterMovementComponent.cpp
@@ -876,6 +876,12 @@ void UCharacterMovementComponent::TickComponent(float DeltaTime, enum ELevelTick
 
 	Super::TickComponent(DeltaTime, TickType, ThisTickFunction);
 
+	// Super tick may destroy/invalidate CharacterOwner or UpdatedComponent, so we need to re-check.
+	if (!HasValidData())
+	{
+		return;
+	}
+
 	// See if we fell out of the world.
 	const bool bIsSimulatingPhysics = UpdatedComponent->IsSimulatingPhysics();
 	if (CharacterOwner->Role == ROLE_Authority)
