Fixed allowing number of custom uvs in a material to become negative

diff --git a/Engine/Source/Runtime/Engine/Classes/Materials/Material.h b/Engine/Source/Runtime/Engine/Classes/Materials/Material.h
--- a/Engine/Source/Runtime/Engine/Classes/Materials/Material.h
+++ b/Engine/Source/Runtime/Engine/Classes/Materials/Material.h
@@ -413,7 +413,7 @@ public:
 	uint32 DitheredLODTransition:1;
 
 	/** Number of customized UV inputs to display.  Unconnected customized UV inputs will just pass through the vertex UVs. */
-	UPROPERTY(EditAnywhere, Category=Material, AdvancedDisplay)
+	UPROPERTY(EditAnywhere, Category=Material, AdvancedDisplay, meta=(ClampMin=0))
 	int32 NumCustomizedUVs;
 
 	/** Sets the lighting mode that will be used on this material if it is translucent. */
