Move read lock out of inner loop in BlendPhysics.

diff --git a/Engine/Source/Runtime/Engine/Private/PhysicsEngine/PhysAnim.cpp b/Engine/Source/Runtime/Engine/Private/PhysicsEngine/PhysAnim.cpp
--- a/Engine/Source/Runtime/Engine/Private/PhysicsEngine/PhysAnim.cpp
+++ b/Engine/Source/Runtime/Engine/Private/PhysicsEngine/PhysAnim.cpp
@@ -72,6 +72,18 @@ void USkeletalMeshComponent::BlendPhysicsBones( TArray<FBoneIndexType>& InRequir
 		return;
 	}
 
+	// Get the scene, and do nothing if we can't get one.
+	FPhysScene* PhysScene = nullptr;
+	if (GetWorld() != nullptr)
+	{
+		PhysScene = GetWorld()->GetPhysicsScene();
+	}
+
+	if (PhysScene == nullptr)
+	{
+		return;
+	}
+
 	// Make sure scratch space is big enough.
 	TArray<FAssetWorldBoneTM> WorldBoneTMs;
 	WorldBoneTMs.Reset();
@@ -82,6 +94,20 @@ void USkeletalMeshComponent::BlendPhysicsBones( TArray<FBoneIndexType>& InRequir
 
 	TArray<FTransform>& EditableSpaceBases = GetEditableSpaceBases();
 
+#if WITH_PHYSX
+	//TODO: we need to use a write lock because later down we might update body instance. This should be going away, at which point we can switch to a read lock (physx supports multiple readers)
+	// Lock the scenes we need (flags set in InitArticulated)
+	if (bHasBodiesInSyncScene)
+	{
+		SCENE_LOCK_WRITE(PhysScene->GetPhysXScene(PST_Sync))
+	}
+
+	if (bHasBodiesInAsyncScene)
+	{
+		SCENE_LOCK_WRITE(PhysScene->GetPhysXScene(PST_Async))
+	}
+#endif
+
 	// For each bone - see if we need to provide some data for it.
 	for(int32 i=0; i<InRequiredBones.Num(); i++)
 	{
@@ -115,7 +141,7 @@ void USkeletalMeshComponent::BlendPhysicsBones( TArray<FBoneIndexType>& InRequir
 			//if simulated body copy back and blend with animation
 			if(BodyInstance->IsInstanceSimulatingPhysics())
 			{
-				FTransform PhysTM = BodyInstance->GetUnrealWorldTransform();
+				FTransform PhysTM = BodyInstance->GetUnrealWorldTransform_AssumesLocked();
 
 				// Store this world-space transform in cache.
 				WorldBoneTMs[BoneIndex].TM = PhysTM;
@@ -193,6 +219,20 @@ void USkeletalMeshComponent::BlendPhysicsBones( TArray<FBoneIndexType>& InRequir
 		}
 	}
 
+#if WITH_PHYSX
+	//See above for read lock instead of write lock
+	// Unlock the scenes 
+	if (bHasBodiesInSyncScene)
+	{
+		SCENE_UNLOCK_WRITE(PhysScene->GetPhysXScene(PST_Sync))
+	}
+
+	if (bHasBodiesInAsyncScene)
+	{
+		SCENE_UNLOCK_WRITE(PhysScene->GetPhysXScene(PST_Async))
+	}
+#endif
+
 	// Transforms updated, cached local bounds are now out of date.
 	InvalidateCachedBounds();
 }
