Add a "View Target Offset" for debug cameras in the player camera manager

--------
Integrated using branch UE4-To-UE4-Orion (reversed) of change#2557713 by [REDACTED_PERSON] on 2015/05/19 16:55:44.

diff --git a/Engine/Source/Runtime/Engine/Classes/Camera/PlayerCameraManager.h b/Engine/Source/Runtime/Engine/Classes/Camera/PlayerCameraManager.h
--- a/Engine/Source/Runtime/Engine/Classes/Camera/PlayerCameraManager.h
+++ b/Engine/Source/Runtime/Engine/Classes/Camera/PlayerCameraManager.h
@@ -264,6 +264,10 @@ public:
 	UPROPERTY(EditAnywhere, BlueprintReadWrite, Category=Debug)
 	FVector FreeCamOffset;
 
+	/** Offset to view target (used in certain CameraStyles) */
+	UPROPERTY(EditAnywhere, BlueprintReadWrite, Category = Debug)
+	FVector ViewTargetOffset;
+
 	/** Current camera fade alpha range, where X = starting alpha and Y = final alpha (when bEnableFading == true) */
 	FVector2D FadeAlpha;
 
