FIX UE-7141 Set Ignore Look Input stacks the resulting value each time it is called

diff --git a/Engine/Source/Runtime/Engine/Classes/GameFramework/PlayerController.h b/Engine/Source/Runtime/Engine/Classes/GameFramework/PlayerController.h
--- a/Engine/Source/Runtime/Engine/Classes/GameFramework/PlayerController.h
+++ b/Engine/Source/Runtime/Engine/Classes/GameFramework/PlayerController.h
@@ -409,22 +409,30 @@ public:
 	virtual void SetCinematicMode( bool bInCinematicMode, bool bAffectsMovement, bool bAffectsTurning);
 
 	/**
-	  * Toggles move input ignoring.
+	  * Locks or unlocks movement input, consecutive calls stack up and require the same amount of calls to undo, or can all be undone using ResetIgnoreMoveInput.
 	  * @param bNewMoveInput	If true, move input is ignored. If false, input is not ignored.
 	  */
 	UFUNCTION(BlueprintCallable, Category="Input")
 	virtual void SetIgnoreMoveInput( bool bNewMoveInput );
 
+	/** Stops ignoring move input by resetting the ignore move input state. */
+	UFUNCTION(BlueprintCallable, Category = "Input", meta = (Keywords = "ClearIgnoreMoveInput"))
+	virtual void ResetIgnoreMoveInput();
+
 	/** Returns true if movement input is ignored. */
 	UFUNCTION(BlueprintCallable, Category="Input")
 	virtual bool IsMoveInputIgnored() const;
 
 	/**
-	  * Toggles look input ignoring.
+	  * Locks or unlocks look input, consecutive calls stack up and require the same amount of calls to undo, or can all be undone using ResetIgnoreLookInput.
 	  * @param bNewLookInput	If true, look input is ignored. If false, input is not ignored.
 	  */
 	UFUNCTION(BlueprintCallable, Category="Input")
 	virtual void SetIgnoreLookInput( bool bNewLookInput );
+	
+	/** Stops ignoring look input by resetting the ignore look input state. */
+	UFUNCTION(BlueprintCallable, Category = "Input", meta = (Keywords = "ClearIgnoreLookInput"))
+	virtual void ResetIgnoreLookInput();
 
 	/** Returns true if look input is ignored. */
 	UFUNCTION(BlueprintCallable, Category="Input")
diff --git a/Engine/Source/Runtime/Engine/Private/PlayerController.cpp b/Engine/Source/Runtime/Engine/Private/PlayerController.cpp
--- a/Engine/Source/Runtime/Engine/Private/PlayerController.cpp
+++ b/Engine/Source/Runtime/Engine/Private/PlayerController.cpp
@@ -2388,24 +2388,32 @@ void APlayerController::SetIgnoreMoveInput( bool bNewMoveInput )
 	//`Log("IgnoreMove: " $ IgnoreMoveInput);
 }
 
+void APlayerController::ResetIgnoreMoveInput()
+{
+	IgnoreMoveInput = 0;
+}
+
 bool APlayerController::IsMoveInputIgnored() const
 {
 	return (IgnoreMoveInput > 0);
 }
 
-
 void APlayerController::SetIgnoreLookInput( bool bNewLookInput )
 {
 	IgnoreLookInput = FMath::Max( IgnoreLookInput + (bNewLookInput ? +1 : -1), 0 );
 	//`Log("IgnoreLook: " $ IgnoreLookInput);
 }
 
+void APlayerController::ResetIgnoreLookInput()
+{
+	IgnoreLookInput = 0;
+}
+
 bool APlayerController::IsLookInputIgnored() const
 {
 	return (IgnoreLookInput > 0);
 }
 
-
 void APlayerController::SetViewTargetWithBlend(AActor* NewViewTarget, float BlendTime, EViewTargetBlendFunction BlendFunc, float BlendExp, bool bLockOutgoing)
 {
 	FViewTargetTransitionParams TransitionParams;
