FIX UE-17149 Physical Material unavailable as Blueprint variable type

diff --git a/Engine/Source/Runtime/Engine/Classes/PhysicalMaterials/PhysicalMaterial.h b/Engine/Source/Runtime/Engine/Classes/PhysicalMaterials/PhysicalMaterial.h
--- a/Engine/Source/Runtime/Engine/Classes/PhysicalMaterials/PhysicalMaterial.h
+++ b/Engine/Source/Runtime/Engine/Classes/PhysicalMaterials/PhysicalMaterial.h
@@ -33,7 +33,7 @@ struct FTireFrictionScalePair
 /**
  * Physical materials are used to define the response of a physical object when interacting dynamically with the world.
  */
-UCLASS(collapsecategories, hidecategories=Object)
+UCLASS(BlueprintType, collapsecategories, hidecategories = Object)
 class ENGINE_API UPhysicalMaterial : public UObject
 {
 	GENERATED_UCLASS_BODY()
