Fix missing async trace param on component overlap from blueprint

diff --git a/Engine/Source/Runtime/Engine/Private/KismetSystemLibrary.cpp b/Engine/Source/Runtime/Engine/Private/KismetSystemLibrary.cpp
--- a/Engine/Source/Runtime/Engine/Private/KismetSystemLibrary.cpp
+++ b/Engine/Source/Runtime/Engine/Private/KismetSystemLibrary.cpp
@@ -1165,6 +1165,7 @@ bool UKismetSystemLibrary::ComponentOverlapComponents_NEW(UPrimitiveComponent* C
 
 	static FName ComponentOverlapComponentsName(TEXT("ComponentOverlapComponents"));
 	FComponentQueryParams Params(ComponentOverlapComponentsName);	
+	Params.bTraceAsyncScene = true;
 	Params.AddIgnoredActors(ActorsToIgnore);
 
 	TArray<FOverlapResult> Overlaps;
