Fix crash in substepping when removing a body during physics in between sub-steps.

diff --git a/Engine/Source/Runtime/Engine/Private/PhysicsEngine/PhysSubstepTasks.cpp b/Engine/Source/Runtime/Engine/Private/PhysicsEngine/PhysSubstepTasks.cpp
--- a/Engine/Source/Runtime/Engine/Private/PhysicsEngine/PhysSubstepTasks.cpp
+++ b/Engine/Source/Runtime/Engine/Private/PhysicsEngine/PhysSubstepTasks.cpp
@@ -37,8 +37,8 @@ void FPhysSubstepTask::SwapBuffers()
 
 void FPhysSubstepTask::RemoveBodyInstance_AssumesLocked(FBodyInstance* BodyInstance)
 {
-	
 	PhysTargetBuffers[External].Remove(BodyInstance);
+	PhysTargetBuffers[!External].Remove(BodyInstance);
 }
 
 void FPhysSubstepTask::SetKinematicTarget_AssumesLocked(FBodyInstance* Body, const FTransform& TM)
@@ -259,11 +259,11 @@ void FPhysSubstepTask::SubstepInterpolation(float InAlpha, float DeltaTime)
 #else
 	PxScene * PScene = PAScene;
 #endif
-
-	PhysTargetMap& Targets = PhysTargetBuffers[!External];
-
+	
 	/** Note: We lock the entire scene before iterating. The assumption is that removing an FBodyInstance from the map will also be wrapped by this lock */
 	SCOPED_SCENE_WRITE_LOCK(PScene);
+	
+	PhysTargetMap& Targets = PhysTargetBuffers[!External];
 
 	for (PhysTargetMap::TIterator Itr = Targets.CreateIterator(); Itr; ++Itr)
 	{
