Introduce apex scene locking macro. Haven't switched it everywhere in engine, still evaluating this

diff --git a/Engine/Source/Runtime/Engine/Private/PhysicsEngine/PhysSubstepTasks.cpp b/Engine/Source/Runtime/Engine/Private/PhysicsEngine/PhysSubstepTasks.cpp
--- a/Engine/Source/Runtime/Engine/Private/PhysicsEngine/PhysSubstepTasks.cpp
+++ b/Engine/Source/Runtime/Engine/Private/PhysicsEngine/PhysSubstepTasks.cpp
@@ -255,13 +255,15 @@ void FPhysSubstepTask::SubstepInterpolation(float InAlpha, float DeltaTime)
 {
 #if WITH_PHYSX
 #if WITH_APEX
+	SCOPED_APEX_SCENE_WRITE_LOCK(PAScene);
 	PxScene * PScene = PAScene->getPhysXScene();
 #else
 	PxScene * PScene = PAScene;
+	SCOPED_SCENE_WRITE_LOCK(PScene);
 #endif
 	
 	/** Note: We lock the entire scene before iterating. The assumption is that removing an FBodyInstance from the map will also be wrapped by this lock */
-	SCOPED_SCENE_WRITE_LOCK(PScene);
+	
 	
 	PhysTargetMap& Targets = PhysTargetBuffers[!External];
 
diff --git a/Engine/Source/Runtime/Engine/Private/PhysicsEngine/PhysXSupport.h b/Engine/Source/Runtime/Engine/Private/PhysicsEngine/PhysXSupport.h
--- a/Engine/Source/Runtime/Engine/Private/PhysicsEngine/PhysXSupport.h
+++ b/Engine/Source/Runtime/Engine/Private/PhysicsEngine/PhysXSupport.h
@@ -40,13 +40,13 @@ class FPhysXSceneReadLock
 {
 public:
 	
-	FPhysXSceneReadLock(PxScene* PInScene)
+	FPhysXSceneReadLock(PxScene* PInScene, const char* filename, PxU32 lineno)
 		: PScene(PInScene)
 	{
 		SCOPE_CYCLE_COUNTER(STAT_PhysSceneReadLock);
 		if(PScene)
 		{
-			PScene->lockRead();
+			PScene->lockRead(filename, lineno);
 		}
 	}
 
@@ -62,17 +62,46 @@ private:
 	PxScene* PScene;
 };
 
+#if WITH_APEX
+/** Scoped scene read lock - we use this instead of PxSceneReadLock because it handles NULL scene */
+class FApexSceneReadLock
+{
+public:
+
+	FApexSceneReadLock(NxApexScene* PInScene, const char* filename, PxU32 lineno)
+		: PScene(PInScene)
+	{
+		SCOPE_CYCLE_COUNTER(STAT_PhysSceneReadLock);
+		if (PScene)
+		{
+			PScene->lockRead(filename, lineno);
+		}
+	}
+
+	~FApexSceneReadLock()
+	{
+		if (PScene)
+		{
+			PScene->unlockRead();
+		}
+	}
+
+private:
+	NxApexScene* PScene;
+};
+#endif
+
 /** Scoped scene write lock - we use this instead of PxSceneReadLock because it handles NULL scene */
 class FPhysXSceneWriteLock
 {
 public:
-	FPhysXSceneWriteLock(PxScene* PInScene)
+	FPhysXSceneWriteLock(PxScene* PInScene, const char* filename, PxU32 lineno)
 		: PScene(PInScene)
 	{
 		SCOPE_CYCLE_COUNTER(STAT_PhysSceneWriteLock);
 		if(PScene)
 		{
-			PScene->lockWrite();
+			PScene->lockWrite(filename, lineno);
 		}
 	}
 
@@ -88,12 +117,44 @@ private:
 	PxScene* PScene;
 };
 
-#define SCOPED_SCENE_READ_LOCK( _scene ) FPhysXSceneReadLock PREPROCESSOR_JOIN(_rlock,__LINE__)(_scene)
-#define SCOPED_SCENE_WRITE_LOCK( _scene ) FPhysXSceneWriteLock PREPROCESSOR_JOIN(_wlock,__LINE__)(_scene)
+#if WITH_APEX
+/** Scoped scene write lock - we use this instead of PxSceneReadLock because it handles NULL scene */
+class FApexSceneWriteLock
+{
+public:
+	FApexSceneWriteLock(NxApexScene* PInScene, const char* filename, PxU32 lineno)
+		: PScene(PInScene)
+	{
+		SCOPE_CYCLE_COUNTER(STAT_PhysSceneWriteLock);
+		if (PScene)
+		{
+			PScene->lockWrite(filename, lineno);
+		}
+	}
+
+	~FApexSceneWriteLock()
+	{
+		if (PScene)
+		{
+			PScene->unlockWrite();
+		}
+	}
+
+private:
+	NxApexScene* PScene;
+};
+#endif
+
+#define SCOPED_SCENE_READ_LOCK( _scene ) FPhysXSceneReadLock PREPROCESSOR_JOIN(_rlock,__LINE__)(_scene, __FILE__, __LINE__)
+#define SCOPED_SCENE_WRITE_LOCK( _scene ) FPhysXSceneWriteLock PREPROCESSOR_JOIN(_wlock,__LINE__)(_scene, __FILE__, __LINE__)
+#if WITH_APEX
+#define SCOPED_APEX_SCENE_READ_LOCK( _scene ) FApexSceneReadLock PREPROCESSOR_JOIN(_rlock,__LINE__)(_scene, __FILE__, __LINE__)
+#define SCOPED_APEX_SCENE_WRITE_LOCK( _scene ) FApexSceneWriteLock PREPROCESSOR_JOIN(_wlock,__LINE__)(_scene, __FILE__, __LINE__)
+#endif
 
-#define SCENE_LOCK_READ( _scene )		{ SCOPE_CYCLE_COUNTER(STAT_PhysSceneReadLock); if((_scene) != NULL) { (_scene)->lockRead(); } }
+#define SCENE_LOCK_READ( _scene )		{ SCOPE_CYCLE_COUNTER(STAT_PhysSceneReadLock); if((_scene) != NULL) { (_scene)->lockRead(__FILE__, __LINE__); } }
 #define SCENE_UNLOCK_READ( _scene )		{ if((_scene) != NULL) { (_scene)->unlockRead(); } }
-#define SCENE_LOCK_WRITE( _scene )		{ SCOPE_CYCLE_COUNTER(STAT_PhysSceneWriteLock); if((_scene) != NULL) { (_scene)->lockWrite(); } }
+#define SCENE_LOCK_WRITE( _scene )		{ SCOPE_CYCLE_COUNTER(STAT_PhysSceneWriteLock); if((_scene) != NULL) { (_scene)->lockWrite(__FILE__, __LINE__); } }
 #define SCENE_UNLOCK_WRITE( _scene )	{ if((_scene) != NULL) { (_scene)->unlockWrite(); } }
 #else
 #define SCOPED_SCENE_READ_LOCK_INDEXED( _scene, _index )
@@ -104,6 +165,11 @@ private:
 #define SCENE_UNLOCK_READ( _scene )
 #define SCENE_LOCK_WRITE( _scene )
 #define SCENE_UNLOCK_WRITE( _scene )
+
+#if WITH_APEX
+#define SCOPED_APEX_SCENE_READ_LOCK( _scene )
+#define SCOPED_APEX_SCENE_WRITE_LOCK( _scene )
+#endif
 #endif
 
 /** Get a pointer to the PxScene from an SceneIndex (will be NULL if scene already shut down) */
