FIX UE-17148 Get Mass Scale not available in Blueprints

diff --git a/Engine/Source/Runtime/Engine/Classes/Components/PrimitiveComponent.h b/Engine/Source/Runtime/Engine/Classes/Components/PrimitiveComponent.h
--- a/Engine/Source/Runtime/Engine/Classes/Components/PrimitiveComponent.h
+++ b/Engine/Source/Runtime/Engine/Classes/Components/PrimitiveComponent.h
@@ -1500,6 +1500,10 @@ public:
 	UFUNCTION(BlueprintCallable, Category="Physics")
 	virtual void SetMassScale(FName BoneName = NAME_None, float InMassScale = 1.f);
 
+	/** Returns the mass scale used to calculate the mass of a single physics body */
+	UFUNCTION(BlueprintCallable, Category = "Physics")
+	virtual float GetMassScale(FName BoneName = NAME_None) const;
+
 	/** Change the mass scale used fo all bodies in this component */
 	UFUNCTION(BlueprintCallable, Category="Physics")
 	virtual void SetAllMassScale(float InMassScale = 1.f);
diff --git a/Engine/Source/Runtime/Engine/Private/PrimitiveComponentPhysics.cpp b/Engine/Source/Runtime/Engine/Private/PrimitiveComponentPhysics.cpp
--- a/Engine/Source/Runtime/Engine/Private/PrimitiveComponentPhysics.cpp
+++ b/Engine/Source/Runtime/Engine/Private/PrimitiveComponentPhysics.cpp
@@ -453,6 +453,16 @@ void UPrimitiveComponent::SetMassScale(FName BoneName, float InMassScale)
 	}
 }
 
+float UPrimitiveComponent::GetMassScale(FName BoneName /*= NAME_None*/) const
+{
+	if (FBodyInstance* BI = GetBodyInstance(BoneName))
+	{
+		return BI->GetMassScale();
+	}
+
+	return 0.0f;
+}
+
 void UPrimitiveComponent::SetAllMassScale(float InMassScale)
 {
 	SetMassScale(NAME_None, InMassScale);
