Fix CIS (non-existent method GetMassScale()).

diff --git a/Engine/Source/Runtime/Engine/Private/PrimitiveComponentPhysics.cpp b/Engine/Source/Runtime/Engine/Private/PrimitiveComponentPhysics.cpp
--- a/Engine/Source/Runtime/Engine/Private/PrimitiveComponentPhysics.cpp
+++ b/Engine/Source/Runtime/Engine/Private/PrimitiveComponentPhysics.cpp
@@ -457,7 +457,7 @@ float UPrimitiveComponent::GetMassScale(FName BoneName /*= NAME_None*/) const
 {
 	if (FBodyInstance* BI = GetBodyInstance(BoneName))
 	{
-		return BI->GetMassScale();
+		return BI->MassScale;
 	}
 
 	return 0.0f;
@@ -1012,4 +1012,4 @@ void UPrimitiveComponent::UpdatePhysicsToRBChannels()
 	}
 }
 
-#undef LOCTEXT_NAMESPACE
\ No newline at end of file
+#undef LOCTEXT_NAMESPACE
