Allow matinees to be played client-side only.

--------
Integrated using branch UE4-To-UE4-Orion (reversed) of change#2550738 by [REDACTED_PERSON] on 2015/05/14 10:41:58.

diff --git a/Engine/Source/Runtime/Engine/Classes/Matinee/MatineeActor.h b/Engine/Source/Runtime/Engine/Classes/Matinee/MatineeActor.h
--- a/Engine/Source/Runtime/Engine/Classes/Matinee/MatineeActor.h
+++ b/Engine/Source/Runtime/Engine/Classes/Matinee/MatineeActor.h
@@ -227,19 +227,19 @@ public:
 	 * Begin playback of the matinee. Only called in game.
 	 * Will then advance Position by (PlayRate * Deltatime) each time the matinee is ticked.
 	 */
-	UFUNCTION(BlueprintCallable, BlueprintAuthorityOnly, Category="Cinematic")
+	UFUNCTION(BlueprintCallable, Category="Cinematic")
 	virtual void Play();
 
 	/** Stops playback at the current position */
-	UFUNCTION(BlueprintCallable, BlueprintAuthorityOnly, Category="Cinematic")
+	UFUNCTION(BlueprintCallable, Category="Cinematic")
 	virtual void Stop();
 
 	/** Similar to play, but the playback will go backwards until the beginning of the sequence is reached. */
-	UFUNCTION(BlueprintCallable, BlueprintAuthorityOnly, Category="Cinematic")
+	UFUNCTION(BlueprintCallable, Category="Cinematic")
 	virtual void Reverse();
 
 	/** Hold playback at its current position. Calling Pause again will continue playback in its current direction. */
-	UFUNCTION(BlueprintCallable, BlueprintAuthorityOnly, Category="Cinematic")
+	UFUNCTION(BlueprintCallable, Category="Cinematic")
 	virtual void Pause();
 
 	/** 
@@ -248,15 +248,15 @@ public:
 	 * @param NewPosition the new position to set the interpolation to
 	 * @param bJump if true, teleport to the new position (don't trigger any events between the old and new positions, etc)
 	 */
-	UFUNCTION(BlueprintCallable, BlueprintAuthorityOnly, Category="Cinematic")
+	UFUNCTION(BlueprintCallable, Category="Cinematic")
 	ENGINE_API void SetPosition(float NewPosition, bool bJump = false);
 
 	/** Changes the direction of playback (go in reverse if it was going forward, or vice versa) */
-	UFUNCTION(BlueprintCallable, BlueprintAuthorityOnly, Category="Cinematic")
+	UFUNCTION(BlueprintCallable, Category="Cinematic")
 	virtual void ChangePlaybackDirection();
 
 	/** Change the looping behaviour of this matinee */
-	UFUNCTION(BlueprintCallable, BlueprintAuthorityOnly, Category="Cinematic")
+	UFUNCTION(BlueprintCallable, Category="Cinematic")
 	virtual void SetLoopingState(bool bNewLooping);
 
 #if WITH_EDITOR
@@ -316,6 +316,8 @@ public:
 	virtual void PostNetReceive() override;
 	virtual void BeginPlay() override;
 	virtual void ApplyWorldOffset(const FVector& InOffset, bool bWorldShift) override;
+	virtual void PostLoad() override;
+
 #if WITH_EDITOR
 	virtual bool GetReferencedContentObjects( TArray<UObject*>& Objects ) const override;
 #endif
diff --git a/Engine/Source/Runtime/Engine/Private/Interpolation.cpp b/Engine/Source/Runtime/Engine/Private/Interpolation.cpp
--- a/Engine/Source/Runtime/Engine/Private/Interpolation.cpp
+++ b/Engine/Source/Runtime/Engine/Private/Interpolation.cpp
@@ -377,6 +377,12 @@ AMatineeActor::AMatineeActor(const FObjectInitializer& ObjectInitializer)
 	ReplicationForceIsPlaying = 0;
 }
 
+void AMatineeActor::PostLoad()
+{
+	SetReplicates(!bClientSideOnly);
+	Super::PostLoad();
+}
+
 FName AMatineeActor::GetFunctionNameForEvent(FName EventName,bool bUseCustomEventName)
 {
 	FName EventFuncName;
@@ -1019,6 +1025,11 @@ bool AMatineeActor::IsMatineeCompatibleWithPlayer( APlayerController* InPC ) con
 		}
 	}
 
+	if (bReplicates == false && InPC->IsLocalController() == false)
+	{
+		bBindPlayerToMatinee = false;
+	}
+
 	return bBindPlayerToMatinee;
 }
 
@@ -1176,7 +1187,10 @@ void AMatineeActor::EnableCinematicMode(bool bEnable)
 		for (FConstPlayerControllerIterator Iterator = GetWorld()->GetPlayerControllerIterator(); Iterator; ++Iterator)
 		{
 			APlayerController *PC = *Iterator;
-			PC->SetCinematicMode(bEnable, bHidePlayer, bHideHud, bDisableMovementInput, bDisableLookAtInput);
+			if (bReplicates || PC->IsLocalController())
+			{
+				PC->SetCinematicMode(bEnable, bHidePlayer, bHideHud, bDisableMovementInput, bDisableLookAtInput);
+			}
 		}
 	}
 }
