#UE4

- Adding Clamp Axis and Normalize Axis to Kismet Math Library

--------
Integrated using branch UE4-To-UE4-Orion (reversed) of change#2590728 by [REDACTED_PERSON] on 2015/06/17 15:32:26.

diff --git a/Engine/Source/Runtime/Engine/Classes/Kismet/KismetMathLibrary.h b/Engine/Source/Runtime/Engine/Classes/Kismet/KismetMathLibrary.h
--- a/Engine/Source/Runtime/Engine/Classes/Kismet/KismetMathLibrary.h
+++ b/Engine/Source/Runtime/Engine/Classes/Kismet/KismetMathLibrary.h
@@ -762,6 +762,26 @@ class ENGINE_API UKismetMathLibrary : public UBlueprintFunctionLibrary
 	UFUNCTION(BlueprintPure, Category="Math|Rotator", meta=(Keywords="make construct build rotate rotation"))
 	static FRotator RotatorFromAxisAndAngle(FVector Axis, float Angle);
 
+	/**
+	* Clamps an angle to the range of [0, 360].
+	*
+	* @param Angle The angle to clamp.
+	* @return The clamped angle.
+	*/
+	UFUNCTION(BlueprintPure, Category = "Math|Rotator")
+	static float ClampAxis(float Angle);
+
+	/**
+	* Clamps an angle to the range of [-180, 180].
+	*
+	* @param Angle The Angle to clamp.
+	* @return The clamped angle.
+	*/
+	UFUNCTION(BlueprintPure, Category="Math|Rotator")
+	static float NormalizeAxis(float Angle);
+
+
+
 	//
 	//	LinearColor functions
 	//
diff --git a/Engine/Source/Runtime/Engine/Private/KismetMathLibrary.cpp b/Engine/Source/Runtime/Engine/Private/KismetMathLibrary.cpp
--- a/Engine/Source/Runtime/Engine/Private/KismetMathLibrary.cpp
+++ b/Engine/Source/Runtime/Engine/Private/KismetMathLibrary.cpp
@@ -1017,6 +1017,16 @@ FRotator UKismetMathLibrary::RotatorFromAxisAndAngle(FVector Axis, float Angle)
 }
 
 
+float UKismetMathLibrary::ClampAxis(float Angle)
+{
+	return FRotator::ClampAxis(Angle);
+}
+
+float UKismetMathLibrary::NormalizeAxis(float Angle)
+{
+	return FRotator::NormalizeAxis(Angle);
+}
+
 FRotator UKismetMathLibrary::RInterpTo(FRotator Current, FRotator Target, float DeltaTime, float InterpSpeed)
 {
 	return FMath::RInterpTo( Current, Target, DeltaTime, InterpSpeed);
