#UE4

- Ignoring "Ignore Input" in the case of looking for the tilde key

--------
Integrated using branch UE4-To-UE4-Orion (reversed) of change#2578871 by [REDACTED_PERSON] on 2015/06/05 15:55:27.

diff --git a/Engine/Source/Runtime/Engine/Private/GameViewportClient.cpp b/Engine/Source/Runtime/Engine/Private/GameViewportClient.cpp
--- a/Engine/Source/Runtime/Engine/Private/GameViewportClient.cpp
+++ b/Engine/Source/Runtime/Engine/Private/GameViewportClient.cpp
@@ -297,7 +297,7 @@ bool UGameViewportClient::InputKey(FViewport* InViewport, int32 ControllerId, FK
 {
 	if (IgnoreInput())
 	{
-		return false;
+		return ViewportConsole ? ViewportConsole->InputKey(ControllerId, Key, EventType, AmountDepressed, bGamepad) : false;
 	}
 
 	if (Key == EKeys::Enter && EventType == EInputEvent::IE_Pressed && FSlateApplication::Get().GetModifierKeys().IsAltDown() && GetDefault<UInputSettings>()->bAltEnterTogglesFullscreen)
