#Orion

- Being able to type into the console is nice....

--------
Integrated using branch UE4-To-UE4-Orion (reversed) of change#2578930 by [REDACTED_PERSON] on 2015/06/05 16:29:52.

diff --git a/Engine/Source/Runtime/Engine/Private/GameViewportClient.cpp b/Engine/Source/Runtime/Engine/Private/GameViewportClient.cpp
--- a/Engine/Source/Runtime/Engine/Private/GameViewportClient.cpp
+++ b/Engine/Source/Runtime/Engine/Private/GameViewportClient.cpp
@@ -400,19 +400,20 @@ bool UGameViewportClient::InputAxis(FViewport* InViewport, int32 ControllerId, F
 
 bool UGameViewportClient::InputChar(FViewport* InViewport, int32 ControllerId, TCHAR Character)
 {
-	if (IgnoreInput())
-	{
-		return false;
-	}
-
 	// should probably just add a ctor to FString that takes a TCHAR
 	FString CharacterString;
 	CharacterString += Character;
 
-	// route to subsystems that care
+	//Always route to the console
 	bool bResult = (ViewportConsole ? ViewportConsole->InputChar(ControllerId, CharacterString) : false);
 
-	if( InViewport->IsSlateViewport() && InViewport->IsPlayInEditorViewport() )
+	if (IgnoreInput())
+	{
+		return bResult;
+	}
+
+	// route to subsystems that care
+	if (!bResult && InViewport->IsSlateViewport() && InViewport->IsPlayInEditorViewport())
 	{
 		// Absorb all keys so game input events are not routed to the Slate editor frame
 		bResult = true;
