Add a "View Target Offset" for debug cameras in the player camera manager

--------
Integrated using branch UE4-To-UE4-Orion (reversed) of change#2557713 by [REDACTED_PERSON] on 2015/05/19 16:55:44.

diff --git a/Engine/Source/Runtime/Engine/Private/PlayerCameraManager.cpp b/Engine/Source/Runtime/Engine/Private/PlayerCameraManager.cpp
--- a/Engine/Source/Runtime/Engine/Private/PlayerCameraManager.cpp
+++ b/Engine/Source/Runtime/Engine/Private/PlayerCameraManager.cpp
@@ -541,7 +541,7 @@ void APlayerCameraManager::UpdateViewTarget(FTViewTarget& OutVT, float DeltaTime
 				Rotator = PCOwner->GetControlRotation();
 			}
 
-			FVector Pos = Loc + FRotationMatrix(Rotator).TransformVector(FreeCamOffset) - Rotator.Vector() * FreeCamDistance;
+			FVector Pos = Loc + ViewTargetOffset + FRotationMatrix(Rotator).TransformVector(FreeCamOffset) - Rotator.Vector() * FreeCamDistance;
 			FCollisionQueryParams BoxParams(NAME_FreeCam, false, this);
 			BoxParams.AddIgnoredActor(OutVT.Target);
 			FHitResult Result;
