Add profiling scopes for AnimInstance NativeUpdateAnimation() and BlueprintUpdateAnimation() calls

--------
Integrated using branch UE4-To-UE4-Orion (reversed) of change#2578356 by [REDACTED_PERSON] on 2015/06/05 09:00:49.

diff --git a/Engine/Source/Runtime/Engine/Private/Animation/AnimInstance.cpp b/Engine/Source/Runtime/Engine/Private/Animation/AnimInstance.cpp
--- a/Engine/Source/Runtime/Engine/Private/Animation/AnimInstance.cpp
+++ b/Engine/Source/Runtime/Engine/Private/Animation/AnimInstance.cpp
@@ -49,6 +49,8 @@ DEFINE_STAT(STAT_SkinnedMeshCompTick);
 DEFINE_STAT(STAT_TickUpdateRate);
 DEFINE_STAT(STAT_PerformAnimEvaluation);
 DEFINE_STAT(STAT_PostAnimEvaluation);
+DEFINE_STAT(STAT_BlueprintUpdateAnimation);
+DEFINE_STAT(STAT_NativeUpdateAnimation);
 
 DECLARE_CYCLE_STAT_EXTERN(TEXT("Anim Init Time"), STAT_AnimInitTime, STATGROUP_Anim, );
 DEFINE_STAT(STAT_AnimInitTime);
@@ -413,8 +415,14 @@ void UAnimInstance::UpdateAnimation(float DeltaSeconds)
 		SyncGroups[GroupIndex].Reset();
 	}
 
-	NativeUpdateAnimation(DeltaSeconds);
-	BlueprintUpdateAnimation(DeltaSeconds);
+	{
+		SCOPE_CYCLE_COUNTER(STAT_NativeUpdateAnimation);
+		NativeUpdateAnimation(DeltaSeconds);
+	}
+	{
+		SCOPE_CYCLE_COUNTER(STAT_BlueprintUpdateAnimation);
+		BlueprintUpdateAnimation(DeltaSeconds);
+	}
 
 	// update weight before all nodes update comes in
 	Montage_UpdateWeight(DeltaSeconds);
diff --git a/Engine/Source/Runtime/Engine/Public/AnimTree.h b/Engine/Source/Runtime/Engine/Public/AnimTree.h
--- a/Engine/Source/Runtime/Engine/Public/AnimTree.h
+++ b/Engine/Source/Runtime/Engine/Public/AnimTree.h
@@ -38,5 +38,7 @@ DECLARE_CYCLE_STAT_EXTERN(TEXT("BlendInPhysics"), STAT_BlendInPhysics, STATGROUP
 DECLARE_CYCLE_STAT_EXTERN(TEXT("SkinPerPolyVertices);"), STAT_SkinPerPolyVertices, STATGROUP_Anim, );
 DECLARE_CYCLE_STAT_EXTERN(TEXT("UpdateTriMeshVertices"), STAT_UpdateTriMeshVertices, STATGROUP_Anim, );
 DECLARE_CYCLE_STAT_EXTERN(TEXT("AnimGameThreadTime"), STAT_AnimGameThreadTime, STATGROUP_Anim, );
+DECLARE_CYCLE_STAT_EXTERN(TEXT("BlueprintUpdateAnimation"), STAT_BlueprintUpdateAnimation, STATGROUP_Anim, );
+DECLARE_CYCLE_STAT_EXTERN(TEXT("NativeUpdateAnimation"), STAT_NativeUpdateAnimation, STATGROUP_Anim, );
 
 
