Changelist approved by [REDACTED_PERSON]

- Particle system components that run on a dedicated server and are set to bAutoDestroy == true are automatically destroyed before the first tick to prevent the components from lingering around forever.

--------
Integrated using branch UE4-To-UE4-Orion (reversed) of change#2510472 by [REDACTED_PERSON] on 2015/04/13 12:57:06.

diff --git a/Engine/Source/Runtime/Engine/Private/Particles/ParticleComponents.cpp b/Engine/Source/Runtime/Engine/Private/Particles/ParticleComponents.cpp
--- a/Engine/Source/Runtime/Engine/Private/Particles/ParticleComponents.cpp
+++ b/Engine/Source/Runtime/Engine/Private/Particles/ParticleComponents.cpp
@@ -3490,13 +3490,23 @@ void UParticleSystemComponent::TickComponent(float DeltaTime, enum ELevelTick Ti
 
 	UWorld* World = GetWorld();
 	check(World);
-	// Bail out if there is no template, there are no instances, or we're running a dedicated server and we don't update on those
-	if ((Template == NULL) || (EmitterInstances.Num() == 0) || (Template->Emitters.Num() == 0) || 
-		((bUpdateOnDedicatedServer == false) && (GetNetMode() == NM_DedicatedServer)))
+	// Bail out if there is no template or there are no instances, or we're running a dedicated server and we don't update on those
+	if ((Template == NULL) || (EmitterInstances.Num() == 0) || (Template->Emitters.Num() == 0))
 	{
 		return;
 	}
 
+	// Bail out if we are running on a dedicated server and we don't want to update on those
+	if ((bUpdateOnDedicatedServer == false) && (GetNetMode() == NM_DedicatedServer))
+	{
+		if (bAutoDestroy)
+		{
+			// We need to destroy the component if the user is expecting us to do it automatically otherwise this component will live forever because HasCompleted() will never get checked
+			DestroyComponent();
+		}
+		return;
+	}
+
 	// System settings may have been lowered. Support late deactivation.
 	int32 DetailModeCVar = GetCachedScalabilityCVars().DetailMode;
 	const bool bDetailModeAllowsRendering	= DetailMode <= DetailModeCVar;
