Reinstated a deprecated member in order that the public facing class does not change for the hot fix.

diff --git a/Engine/Source/Runtime/Engine/Classes/Matinee/MatineeActor.h b/Engine/Source/Runtime/Engine/Classes/Matinee/MatineeActor.h
--- a/Engine/Source/Runtime/Engine/Classes/Matinee/MatineeActor.h
+++ b/Engine/Source/Runtime/Engine/Classes/Matinee/MatineeActor.h
@@ -188,6 +188,10 @@ public:
 	UPROPERTY(replicated,transient)
 	uint32 bPaused:1;
 
+	// The below property is deprecated and will be removed in 4.9.
+	UPROPERTY(replicated,transient)
+	uint32 bPendingStop:1;
+
 	UPROPERTY(replicated)
 	float InterpPosition;
 
