FIX UE-16522 Set Material for a DM in the Construction script will result in a crash

diff --git a/Engine/Source/Runtime/Engine/Private/Components/DestructibleComponent.cpp b/Engine/Source/Runtime/Engine/Private/Components/DestructibleComponent.cpp
--- a/Engine/Source/Runtime/Engine/Private/Components/DestructibleComponent.cpp
+++ b/Engine/Source/Runtime/Engine/Private/Components/DestructibleComponent.cpp
@@ -1492,17 +1492,28 @@ void UDestructibleComponent::SetMaterial(int32 ElementIndex, UMaterialInterface*
 	// Mesh component handles render side materials
 	Super::SetMaterial(ElementIndex, Material);
 
-	// Update physical properties of the chunks in the mesh
-	GetBodyInstance()->UpdatePhysicalMaterials();
-	int32 NumBones = SkeletalMesh->RefSkeleton.GetNum();
-	for(int32 BoneIdx = 0 ; BoneIdx < NumBones ; ++BoneIdx)
+	// Update physical properties of the chunks in the mesh if the body instance is valid
+	FBodyInstance* BodyInst = GetBodyInstance();
+	if (BodyInst)
 	{
-		FName BoneName = SkeletalMesh->RefSkeleton.GetBoneName(BoneIdx);
-		if(FBodyInstance* Instance = GetBodyInstance(BoneName))
+		BodyInst->UpdatePhysicalMaterials();
+	}
+	
+	// Update physical properties for individual bone instances as well
+	if (SkeletalMesh)
+	{
+		int32 NumBones = SkeletalMesh->RefSkeleton.GetNum();
+		for (int32 BoneIdx = 0; BoneIdx < NumBones; ++BoneIdx)
 		{
-			Instance->UpdatePhysicalMaterials();
+			FName BoneName = SkeletalMesh->RefSkeleton.GetBoneName(BoneIdx);
+			FBodyInstance* Instance = GetBodyInstance(BoneName);
+			if (Instance)
+			{
+				Instance->UpdatePhysicalMaterials();
+			}
 		}
 	}
+	
 
 #if WITH_APEX
 	// Set new template parameters for the apex actor, so they take effect before fracturing too.
