- Add BP-exposed GetBoundingBox function to StaticMesh (forum request)

diff --git a/Engine/Source/Runtime/Engine/Classes/Engine/StaticMesh.h b/Engine/Source/Runtime/Engine/Classes/Engine/StaticMesh.h
--- a/Engine/Source/Runtime/Engine/Classes/Engine/StaticMesh.h
+++ b/Engine/Source/Runtime/Engine/Classes/Engine/StaticMesh.h
@@ -502,6 +502,10 @@ public:
 	 */
 	ENGINE_API FBoxSphereBounds GetBounds() const;
 
+	/** Returns the bounding box, in local space, of the StaticMesh asset */
+	UFUNCTION(BlueprintCallable, Category="StaticMesh")
+	ENGINE_API FBox GetBoundingBox() const;
+
 	/**
 	 * Gets a Material given a Material Index and an LOD number
 	 *
diff --git a/Engine/Source/Runtime/Engine/Private/StaticMesh.cpp b/Engine/Source/Runtime/Engine/Private/StaticMesh.cpp
--- a/Engine/Source/Runtime/Engine/Private/StaticMesh.cpp
+++ b/Engine/Source/Runtime/Engine/Private/StaticMesh.cpp
@@ -1376,6 +1376,16 @@ FBoxSphereBounds UStaticMesh::GetBounds() const
 	return Bounds;
 }
 
+FBox UStaticMesh::GetBoundingBox() const
+{
+	FBox BoundingBox(FVector(0,0,0), FVector(0,0,0));
+	if (RenderData != nullptr)
+	{
+		BoundingBox = RenderData->Bounds.GetBox();
+	}
+	return BoundingBox;
+}
+
 float UStaticMesh::GetStreamingTextureFactor(int32 RequestedUVIndex)
 {
 	check(RequestedUVIndex >= 0);
