Remove unused EnableParent

diff --git a/Engine/Source/Runtime/Engine/Classes/Components/ActorComponent.h b/Engine/Source/Runtime/Engine/Classes/Components/ActorComponent.h
--- a/Engine/Source/Runtime/Engine/Classes/Components/ActorComponent.h
+++ b/Engine/Source/Runtime/Engine/Classes/Components/ActorComponent.h
@@ -415,7 +415,6 @@ public:
 	/** 
 	 * Set up a tick function for a component in the standard way. 
 	 * Tick after the actor. Don't tick if the actor is static, or if the actor is a template or if this is a "NeverTick" component.
-	 * Owner is the EnableParent.
 	 * I tick while paused if only if my owner does.
 	 * @param	TickFunction - structure holding the specific tick function
 	 * @return  true if this component met the criteria for actually being ticked.
diff --git a/Engine/Source/Runtime/Engine/Classes/Engine/EngineBaseTypes.h b/Engine/Source/Runtime/Engine/Classes/Engine/EngineBaseTypes.h
--- a/Engine/Source/Runtime/Engine/Classes/Engine/EngineBaseTypes.h
+++ b/Engine/Source/Runtime/Engine/Classes/Engine/EngineBaseTypes.h
@@ -174,7 +174,7 @@ private:
 	uint32 bRegistered:1;
 
 	/** 
-	 * If false, this tick will not fire, nor will any other tick that has this tick function as an EnableParent. 
+	 * If false, this tick will not fire
 	 * CAUTION: Do not set this directly
 	 **/
 	uint32 bTickEnabled:1;
@@ -193,11 +193,6 @@ private:
 	TArray<struct FTickPrerequisite> Prerequisites;
 
 public:
-	/** 
-	 * If the "EnableParent" is not enabled, then this tick function is implicitly disabled as well. 
-	 * Caution, there is no protection on this strong pointer. It is assumed your enable parent will not be destroyed before you are.
-	 **/
-	FTickFunction*								EnableParent;
 	/** Back pointer to the FTickTaskLevel containing this tick function if it is registered **/
 	class FTickTaskLevel*						TickTaskLevel;
 
diff --git a/Engine/Source/Runtime/Engine/Private/Components/ActorComponent.cpp b/Engine/Source/Runtime/Engine/Private/Components/ActorComponent.cpp
--- a/Engine/Source/Runtime/Engine/Private/Components/ActorComponent.cpp
+++ b/Engine/Source/Runtime/Engine/Private/Components/ActorComponent.cpp
@@ -653,7 +653,6 @@ void FActorComponentTickFunction::ExecuteTick(float DeltaTime, enum ELevelTick T
 	{
 		FScopeCycleCounterUObject ComponentScope(Target);
 		FScopeCycleCounterUObject AdditionalScope(Target->AdditionalStatObject());
-	    checkSlow(Target && (!EnableParent || Target->IsPendingKill() || ((FActorTickFunction*)EnableParent)->Target == Target->GetOwner())); // components that get renamed into other outers will have this wrong and hence will not necessarily tick after their actor, or use their actor as an enable parent
 	    Target->ConditionalTickComponent(DeltaTime, TickType, *this);	
 	}
 }
diff --git a/Engine/Source/Runtime/Engine/Private/TickTaskManager.cpp b/Engine/Source/Runtime/Engine/Private/TickTaskManager.cpp
--- a/Engine/Source/Runtime/Engine/Private/TickTaskManager.cpp
+++ b/Engine/Source/Runtime/Engine/Private/TickTaskManager.cpp
@@ -425,7 +425,7 @@ public:
 		{
 			FTickFunction* TickFunction = *It;
 			TickFunction->CompletionHandle = NULL; // might as well NULL this out to allow these handles to be recycled
-			if (TickFunction->bTickEvenWhenPaused && TickFunction->bTickEnabled && (!TickFunction->EnableParent || TickFunction->EnableParent->bTickEnabled))
+			if (TickFunction->bTickEvenWhenPaused && TickFunction->bTickEnabled)
 			{
 				TickFunction->TickVisitedGFrameCounter = GFrameCounter;
 				TickFunction->TickQueuedGFrameCounter = GFrameCounter;
@@ -789,7 +789,6 @@ FTickFunction::FTickFunction()
 	, bTickEnabled(true)
 	, TickVisitedGFrameCounter(0)
 	, TickQueuedGFrameCounter(0)
-	, EnableParent(NULL)
 	, TickTaskLevel(NULL)
 {
 }
@@ -886,7 +885,7 @@ void FTickFunction::QueueTickFunction(FTickTaskSequencer& TTS, const struct FTic
 	if (TickVisitedGFrameCounter != GFrameCounter)
 	{
 		TickVisitedGFrameCounter = GFrameCounter;
-		if (bTickEnabled && (!EnableParent || EnableParent->bTickEnabled))
+		if (bTickEnabled)
 		{
 			ETickingGroup MaxPrerequisiteTickGroup =  ETickingGroup(0);
 
