Fixed crash caused because Custom Output material expression was not excluded from the material expression list. UE-17118

diff --git a/Engine/Source/Runtime/Engine/Classes/Materials/MaterialExpressionCustomOutput.h b/Engine/Source/Runtime/Engine/Classes/Materials/MaterialExpressionCustomOutput.h
--- a/Engine/Source/Runtime/Engine/Classes/Materials/MaterialExpressionCustomOutput.h
+++ b/Engine/Source/Runtime/Engine/Classes/Materials/MaterialExpressionCustomOutput.h
@@ -5,7 +5,7 @@
 #include "Materials/MaterialExpression.h"
 #include "MaterialExpressionCustomOutput.generated.h"
 
-UCLASS(collapsecategories, hidecategories=Object, MinimalAPI)
+UCLASS(abstract,collapsecategories, hidecategories = Object, MinimalAPI)
 class UMaterialExpressionCustomOutput : public UMaterialExpression
 {
 	GENERATED_UCLASS_BODY()
