Added file modied date and file size to map file asset tooltips. Allows sorting by these values in column view mode.

Note: Map files must be saved in order for their tooltip to be updated with the new data.

UE-950

diff --git a/Engine/Source/Runtime/Engine/Private/World.cpp b/Engine/Source/Runtime/Engine/Private/World.cpp
--- a/Engine/Source/Runtime/Engine/Private/World.cpp
+++ b/Engine/Source/Runtime/Engine/Private/World.cpp
@@ -5615,6 +5615,15 @@ void UWorld::GetAssetRegistryTags(TArray<FAssetRegistryTag>& OutTags) const
 			OutTags.Add(FAssetRegistryTag("FiB", FString(), FAssetRegistryTag::TT_Hidden));
 		}
 	}
+
+	// Get the full file path with extension
+	const FString FullFilePath = FPackageName::LongPackageNameToFilename(GetOutermost()->GetName(), FPackageName::GetMapPackageExtension());
+
+	// Save/Display the file size and modify date
+	FDateTime AssetDateModified = IFileManager::Get().GetTimeStamp(*FullFilePath);
+	OutTags.Add(FAssetRegistryTag("DateModified", FText::AsDate(AssetDateModified, EDateTimeStyle::Short).ToString(), FAssetRegistryTag::TT_Dimensional));
+	OutTags.Add(FAssetRegistryTag("MapFileSize", FText::AsMemory(IFileManager::Get().FileSize(*FullFilePath)).ToString(), FAssetRegistryTag::TT_Numerical));
+
 	FWorldDelegates::GetAssetTags.Broadcast(this, OutTags);
 }
 #endif
