More Gamma Correction - Didn't catch this on the previous pass, apparently there was an implicit constructor allowing FLinearColor to FColor that was doing pow(2.2) gamma conversion inversion.  Rather than leave the implicit constructor, I'm making it private and making people use ToFColor(true).  Which is slightly more expensive, but performs the proper sRGB conversion.  While fixing this, I found several terrible uses of the implicit constructor, when Hashing FLinearColors we were converting them to FColors first, when clearing FCanvas we were manually gamma correcting but leaving it as an FLinearColor, then implicitly converting to FColor, double gamma corrrecting.  Neither of which should even be required as the RHI Clear command expects an FLinearColor.  Additionally fixing a myriad of Slate widgets that were all doing FColor conversions needlessly only to convert back to FLinearColor when queuing slate draw commands.

diff --git a/Engine/Plugins/2D/Paper2D/Source/Paper2D/Private/PaperFlipbookComponent.cpp b/Engine/Plugins/2D/Paper2D/Source/Paper2D/Private/PaperFlipbookComponent.cpp
--- a/Engine/Plugins/2D/Paper2D/Source/Paper2D/Private/PaperFlipbookComponent.cpp
+++ b/Engine/Plugins/2D/Paper2D/Source/Paper2D/Private/PaperFlipbookComponent.cpp
@@ -87,7 +87,7 @@ FPrimitiveSceneProxy* UPaperFlipbookComponent::CreateSceneProxy()
 
 	FSpriteDrawCallRecord DrawCall;
 	DrawCall.BuildFromSprite(SpriteToSend);
-	DrawCall.Color = SpriteColor;
+	DrawCall.Color = SpriteColor.ToFColor(true);
 	NewProxy->SetDrawCall_RenderThread(DrawCall);
 	return NewProxy;
 }
@@ -342,7 +342,7 @@ void UPaperFlipbookComponent::SendRenderDynamicData_Concurrent()
 
 		FSpriteDrawCallRecord DrawCall;
 		DrawCall.BuildFromSprite(SpriteToSend);
-		DrawCall.Color = SpriteColor;
+		DrawCall.Color = SpriteColor.ToFColor(true);
 
 		ENQUEUE_UNIQUE_RENDER_COMMAND_THREEPARAMETER(
 				FSendPaperRenderComponentDynamicData,
diff --git a/Engine/Plugins/2D/Paper2D/Source/Paper2D/Private/PaperGroupedSpriteComponent.cpp b/Engine/Plugins/2D/Paper2D/Source/Paper2D/Private/PaperGroupedSpriteComponent.cpp
--- a/Engine/Plugins/2D/Paper2D/Source/Paper2D/Private/PaperGroupedSpriteComponent.cpp
+++ b/Engine/Plugins/2D/Paper2D/Source/Paper2D/Private/PaperGroupedSpriteComponent.cpp
@@ -42,7 +42,7 @@ int32 UPaperGroupedSpriteComponent::AddInstanceWithMaterial(const FTransform& Tr
 	const FTransform LocalTransform(bWorldSpace ? Transform.GetRelativeTransform(ComponentToWorld) : Transform);
 
 	FSpriteInstanceData& NewInstanceData = *new (PerInstanceSpriteData)FSpriteInstanceData();
-	SetupNewInstanceData(NewInstanceData, NewInstanceIndex, LocalTransform, Sprite, MaterialOverride, Color);
+	SetupNewInstanceData(NewInstanceData, NewInstanceIndex, LocalTransform, Sprite, MaterialOverride, Color.ToFColor(true));
 
 	MarkRenderStateDirty();
 
@@ -113,7 +113,7 @@ bool UPaperGroupedSpriteComponent::UpdateInstanceColor(int32 InstanceIndex, FLin
 	if (PerInstanceSpriteData.IsValidIndex(InstanceIndex))
 	{
 		FSpriteInstanceData& InstanceData = PerInstanceSpriteData[InstanceIndex];
-		InstanceData.VertexColor = NewInstanceColor;
+		InstanceData.VertexColor = NewInstanceColor.ToFColor(true);
 
 		if (bMarkRenderStateDirty)
 		{
diff --git a/Engine/Plugins/2D/Paper2D/Source/Paper2D/Private/PaperRenderSceneProxy.cpp b/Engine/Plugins/2D/Paper2D/Source/Paper2D/Private/PaperRenderSceneProxy.cpp
--- a/Engine/Plugins/2D/Paper2D/Source/Paper2D/Private/PaperRenderSceneProxy.cpp
+++ b/Engine/Plugins/2D/Paper2D/Source/Paper2D/Private/PaperRenderSceneProxy.cpp
@@ -106,7 +106,7 @@ class FDummyPaperSpriteVertexBuffer : public FPaperSpriteVertexBuffer
 public:
 	FDummyPaperSpriteVertexBuffer()
 	{
-		Vertices.Emplace(FVector::ZeroVector, FVector2D::ZeroVector, FLinearColor::Black);
+		Vertices.Emplace(FVector::ZeroVector, FVector2D::ZeroVector, FColor::Black);
 	}
 };
 static TGlobalResource<FDummyPaperSpriteVertexBuffer> GDummyMaterialSpriteVertexBuffer;
@@ -320,12 +320,12 @@ void FPaperRenderSceneProxy::DebugDrawBodySetup(const FSceneView* View, int32 Vi
 
 			Collector.RegisterOneFrameMaterialProxy(SolidMaterialInstance);
 
-			BodySetup->AggGeom.GetAggGeom(GeomTransform, WireframeColor, SolidMaterialInstance, false, true, UseEditorDepthTest(), ViewIndex, Collector);
+			BodySetup->AggGeom.GetAggGeom(GeomTransform, WireframeColor.ToFColor(true), SolidMaterialInstance, false, true, UseEditorDepthTest(), ViewIndex, Collector);
 		}
 		else
 		{
 			// wireframe
-			BodySetup->AggGeom.GetAggGeom(GeomTransform, GetSelectionColor(CollisionColor, IsSelected(), IsHovered()), nullptr, (Owner == nullptr), false, UseEditorDepthTest(), ViewIndex, Collector);
+			BodySetup->AggGeom.GetAggGeom(GeomTransform, GetSelectionColor(CollisionColor, IsSelected(), IsHovered()).ToFColor(true), nullptr, ( Owner == nullptr ), false, UseEditorDepthTest(), ViewIndex, Collector);
 		}
 	}
 }
diff --git a/Engine/Plugins/2D/Paper2D/Source/Paper2D/Private/PaperSprite.cpp b/Engine/Plugins/2D/Paper2D/Source/Paper2D/Private/PaperSprite.cpp
--- a/Engine/Plugins/2D/Paper2D/Source/Paper2D/Private/PaperSprite.cpp
+++ b/Engine/Plugins/2D/Paper2D/Source/Paper2D/Private/PaperSprite.cpp
@@ -356,7 +356,7 @@ void FSpriteDrawCallRecord::BuildFromSprite(const UPaperSprite* Sprite)
 		BaseTexture = Sprite->GetBakedTexture();
 		Sprite->GetBakedAdditionalSourceTextures(/*out*/ AdditionalTextures);
 
-		Color = FLinearColor::White;
+		Color = FLinearColor::White.ToFColor(true);
 
 		RenderVerts = Sprite->BakedRenderData;
 	}
@@ -1293,17 +1293,17 @@ void UPaperSprite::FindContours(const FIntPoint& ScanPos, const FIntPoint& ScanS
 								// Move to the next pixel
 
 								// Check to see if we closed the loop
- 								if ((CX == X) && (CY == Y))
- 								{
+								if ((CX == X) && (CY == Y))
+								{
 // 									// If we went thru the boundary pixel more than two times, or
 // 									// entered it from the same way we started, then we're done
 // 									++EnteredStartingSquareCount;
 // 									if ((EnteredStartingSquareCount > 2) || (NeighborPhase == 0))
 // 									{
 									//@TODO: Not good enough, will early out too soon some of the time!
- 										bInsideBoundary = true;
- 										break;
- 									}
+										bInsideBoundary = true;
+										break;
+									}
 // 								}
 
 								BoundaryImage.SetPixel(CX, CY, NeighborPhase+1);
diff --git a/Engine/Plugins/2D/Paper2D/Source/Paper2D/Private/PaperSpriteComponent.cpp b/Engine/Plugins/2D/Paper2D/Source/Paper2D/Private/PaperSpriteComponent.cpp
--- a/Engine/Plugins/2D/Paper2D/Source/Paper2D/Private/PaperSpriteComponent.cpp
+++ b/Engine/Plugins/2D/Paper2D/Source/Paper2D/Private/PaperSpriteComponent.cpp
@@ -67,7 +67,7 @@ FPrimitiveSceneProxy* UPaperSpriteComponent::CreateSceneProxy()
 	{
 		FSpriteDrawCallRecord DrawCall;
 		DrawCall.BuildFromSprite(SourceSprite);
-		DrawCall.Color = SpriteColor;
+		DrawCall.Color = SpriteColor.ToFColor(true);
 		NewProxy->SetSprite_RenderThread(DrawCall, SourceSprite->AlternateMaterialSplitIndex);
 	}
 	return NewProxy;
@@ -108,7 +108,7 @@ void UPaperSpriteComponent::SendRenderDynamicData_Concurrent()
 	{
 		FSpriteDrawCallRecord DrawCall;
 		DrawCall.BuildFromSprite(SourceSprite);
-		DrawCall.Color = SpriteColor;
+		DrawCall.Color = SpriteColor.ToFColor(true);
 		const int32 SplitIndex = (SourceSprite != nullptr) ? SourceSprite->AlternateMaterialSplitIndex : INDEX_NONE;
 
 		ENQUEUE_UNIQUE_RENDER_COMMAND_THREEPARAMETER(
diff --git a/Engine/Plugins/2D/Paper2D/Source/Paper2D/Private/PaperSpriteSceneProxy.cpp b/Engine/Plugins/2D/Paper2D/Source/Paper2D/Private/PaperSpriteSceneProxy.cpp
--- a/Engine/Plugins/2D/Paper2D/Source/Paper2D/Private/PaperSpriteSceneProxy.cpp
+++ b/Engine/Plugins/2D/Paper2D/Source/Paper2D/Private/PaperSpriteSceneProxy.cpp
@@ -88,7 +88,7 @@ void FPaperSpriteSceneProxy::GetDynamicMeshElements(const TArray<const FSceneVie
 
 							const bool bUseSeparateColorPerHull = (Owner == nullptr);
 							const bool bDrawSolid = false;
-							BodySetup->AggGeom.GetAggGeom(GeomTransform, GetSelectionColor(CollisionColor, bDrawWireSelected, IsHovered()), CollisionMaterialInstance, bUseSeparateColorPerHull, bDrawSolid, false, ViewIndex, Collector);
+							BodySetup->AggGeom.GetAggGeom(GeomTransform, GetSelectionColor(CollisionColor, bDrawWireSelected, IsHovered()).ToFColor(true), CollisionMaterialInstance, bUseSeparateColorPerHull, bDrawSolid, false, ViewIndex, Collector);
 						}
 					}
 				}
diff --git a/Engine/Plugins/2D/Paper2D/Source/Paper2D/Private/PaperTileMapComponent.cpp b/Engine/Plugins/2D/Paper2D/Source/Paper2D/Private/PaperTileMapComponent.cpp
--- a/Engine/Plugins/2D/Paper2D/Source/Paper2D/Private/PaperTileMapComponent.cpp
+++ b/Engine/Plugins/2D/Paper2D/Source/Paper2D/Private/PaperTileMapComponent.cpp
@@ -271,7 +271,7 @@ void UPaperTileMapComponent::RebuildRenderData(TArray<FSpriteRenderSection>& Sec
 		}
 	}
 
- 	Vertices.Empty(EstimatedNumVerts);
+	Vertices.Empty(EstimatedNumVerts);
 
 	UMaterialInterface* TileMapMaterial = GetMaterial(0);
 	if (TileMapMaterial == nullptr)
@@ -297,7 +297,7 @@ void UPaperTileMapComponent::RebuildRenderData(TArray<FSpriteRenderSection>& Sec
 			}
 		}
 
-		const FColor DrawColor(TileMapColor * Layer->GetLayerColor());
+		const FColor DrawColor(( TileMapColor * Layer->GetLayerColor() ).ToFColor(true));
 
 #if WITH_EDITORONLY_DATA
 		if (!Layer->ShouldRenderInEditor())
diff --git a/Engine/Plugins/2D/Paper2D/Source/Paper2D/Private/PaperTileMapRenderSceneProxy.cpp b/Engine/Plugins/2D/Paper2D/Source/Paper2D/Private/PaperTileMapRenderSceneProxy.cpp
--- a/Engine/Plugins/2D/Paper2D/Source/Paper2D/Private/PaperTileMapRenderSceneProxy.cpp
+++ b/Engine/Plugins/2D/Paper2D/Source/Paper2D/Private/PaperTileMapRenderSceneProxy.cpp
@@ -157,10 +157,10 @@ void FPaperTileMapRenderSceneProxy::DrawStaggeredGridLines(FPrimitiveDrawInterfa
 		XBottom -= XExcess;
 		YBottom -= XExcess * 2;
 
- 		if (XBottom == TileMap->MapWidth)
- 		{
+		if (XBottom == TileMap->MapWidth)
+		{
 			YBottom -= ((TileMap->MapHeight & 1) != 0) ? 0 : 1;
- 		}
+		}
 
 		// Bottom center
 		Poly.Reset();
@@ -287,7 +287,7 @@ void FPaperTileMapRenderSceneProxy::GetDynamicMeshElements(const TArray<const FS
 
 							const bool bPerHullColor = false;
 							const bool bDrawSimpleSolid = false;
-							BodySetup->AggGeom.GetAggGeom(GeomTransform, GetSelectionColor(CollisionColor, bDrawWireSelected, IsHovered()), CollisionMaterialInstance, bPerHullColor, bDrawSimpleSolid, UseEditorDepthTest(), ViewIndex, Collector);
+							BodySetup->AggGeom.GetAggGeom(GeomTransform, GetSelectionColor(CollisionColor, bDrawWireSelected, IsHovered()).ToFColor(true), CollisionMaterialInstance, bPerHullColor, bDrawSimpleSolid, UseEditorDepthTest(), ViewIndex, Collector);
 						}
 					}
 				}
diff --git a/Engine/Plugins/2D/Paper2D/Source/Paper2D/Private/Terrain/PaperTerrainComponent.cpp b/Engine/Plugins/2D/Paper2D/Source/Paper2D/Private/Terrain/PaperTerrainComponent.cpp
--- a/Engine/Plugins/2D/Paper2D/Source/Paper2D/Private/Terrain/PaperTerrainComponent.cpp
+++ b/Engine/Plugins/2D/Paper2D/Source/Paper2D/Private/Terrain/PaperTerrainComponent.cpp
@@ -539,10 +539,10 @@ void UPaperTerrainComponent::OnSplineEdited()
 					const UPaperSprite* Sprite = RuleHelper.ValidBodies[BodyIndex];
 					const float Width = RuleHelper.ValidBodyWidths[BodyIndex];
 
- 					if ((NumSegments > 0) && ((Width * 0.5f) > DistanceBudget))
- 					{
- 						break;
- 					}
+					if ((NumSegments > 0) && ((Width * 0.5f) > DistanceBudget))
+					{
+						break;
+					}
 					new (Segment.Stamps) FTerrainSpriteStamp(Sprite, Position + (Width * 0.5f), /*bIsEndCap=*/ false);
 
 					DistanceBudget -= Width;
@@ -553,10 +553,10 @@ void UPaperTerrainComponent::OnSplineEdited()
 				const float UsedSpace = (BodyDistance - DistanceBudget);
 				const float OverallScaleFactor = BodyDistance / UsedSpace;
 				
- 				// Stretch body segments
+				// Stretch body segments
 				float PositionCorrectionSum = 0.0f;
- 				for (int32 Index = 0; Index < NumSegments; ++Index)
- 				{
+				for (int32 Index = 0; Index < NumSegments; ++Index)
+				{
 					FTerrainSpriteStamp& Stamp = Segment.Stamps[Index + (Segment.Stamps.Num() - NumSegments)];
 					
 					const float WidthChange = (OverallScaleFactor - 1.0f) * Stamp.NominalWidth;
@@ -565,7 +565,7 @@ void UPaperTerrainComponent::OnSplineEdited()
 
 					Stamp.Scale = OverallScaleFactor;
 					Stamp.Time += PositionCorrectionSum;
- 				}
+				}
 			}
 			else
 			{
@@ -637,7 +637,7 @@ void UPaperTerrainComponent::OnSplineEdited()
 				FSpriteDrawCallRecord& FillDrawCall = *new (MaterialBatch.Records) FSpriteDrawCallRecord();
 				FillDrawCall.BuildFromSprite(FillSprite);
 				FillDrawCall.RenderVerts.Empty();
-				FillDrawCall.Color = TerrainColor;
+				FillDrawCall.Color = TerrainColor.ToFColor(true);
 				FillDrawCall.Destination = PaperAxisZ * 0.1f;
 
 				const FVector2D TextureSize = GetSpriteRenderDataBounds2D(FillSprite->BakedRenderData).GetSize();
@@ -818,7 +818,7 @@ void UPaperTerrainComponent::SpawnSegments(const TArray<FTerrainSegment>& Terrai
 
 				FSpriteDrawCallRecord& Record = *new (MaterialBatch.Records) FSpriteDrawCallRecord();
 				Record.BuildFromSprite(NewSprite);
-				Record.Color = TerrainColor;
+				Record.Color = TerrainColor.ToFColor(true);
 
 				// Work out if the sprite is likely to be bent > X deg (folded over itself)
 				const FVector ForwardVector(1, 0, 0);
@@ -1015,12 +1015,14 @@ void UPaperTerrainComponent::SetTerrainColor(FLinearColor NewColor)
 	{
 		TerrainColor = NewColor;
 
+		FColor TerrainGammaSpaceColor = TerrainColor.ToFColor(true);
+
 		// Update the color in the game-thread copy of the render geometry
 		for (FPaperTerrainSpriteGeometry& Batch : GeneratedSpriteGeometry)
 		{
 			for (FSpriteDrawCallRecord& DrawCall : Batch.Records)
 			{
-				DrawCall.Color = TerrainColor;
+				DrawCall.Color = TerrainGammaSpaceColor;
 			}
 		}
 
diff --git a/Engine/Plugins/2D/Paper2D/Source/Paper2DEditor/Private/PaperEditorShared/SpriteGeometryEditing.cpp b/Engine/Plugins/2D/Paper2D/Source/Paper2DEditor/Private/PaperEditorShared/SpriteGeometryEditing.cpp
--- a/Engine/Plugins/2D/Paper2D/Source/Paper2DEditor/Private/PaperEditorShared/SpriteGeometryEditing.cpp
+++ b/Engine/Plugins/2D/Paper2D/Source/Paper2DEditor/Private/PaperEditorShared/SpriteGeometryEditing.cpp
@@ -350,7 +350,7 @@ void FSpriteGeometryEditingHelper::DrawGeometry(const FSceneView& View, FPrimiti
 			PDI.SetHitProxy(new HSpriteSelectableObjectHitProxy(Data));
 		}
 
-		FColor BackgroundColor(bIsShapeSelected ? SpriteEditingConstantsEX::GeometrySelectedColor : LineColor);
+		FColor BackgroundColor(( bIsShapeSelected ? SpriteEditingConstantsEX::GeometrySelectedColor : LineColor ).ToFColor(true));
 		BackgroundColor.A = 4;
 		FMaterialRenderProxy* ShapeMaterialProxy = WidgetVertexColorMaterial->GetRenderProxy(bIsShapeSelected);
 
diff --git a/Engine/Source/Developer/LogVisualizer/Private/TimeSliderController.cpp b/Engine/Source/Developer/LogVisualizer/Private/TimeSliderController.cpp
--- a/Engine/Source/Developer/LogVisualizer/Private/TimeSliderController.cpp
+++ b/Engine/Source/Developer/LogVisualizer/Private/TimeSliderController.cpp
@@ -664,7 +664,7 @@ int32 FVisualLoggerTimeSliderController::OnPaintSectionView( const FGeometry& Al
 			CursorBackground,
 			MyClippingRect,
 			DrawEffects,
-			FColor(FLinearColor::White).WithAlpha(20)
+			FLinearColor::White.CopyWithNewOpacity(0.08f)
 			);
 
 		// Draw a line for the scrub position
@@ -680,7 +680,7 @@ int32 FVisualLoggerTimeSliderController::OnPaintSectionView( const FGeometry& Al
 			LinePoints,
 			MyClippingRect,
 			DrawEffects,
-			FColor(FLinearColor::White).WithAlpha(100),
+			FLinearColor::White.CopyWithNewOpacity(0.39f),
 			false
 			);
 
diff --git a/Engine/Source/Developer/TaskGraph/Private/SGraphBar.cpp b/Engine/Source/Developer/TaskGraph/Private/SGraphBar.cpp
--- a/Engine/Source/Developer/TaskGraph/Private/SGraphBar.cpp
+++ b/Engine/Source/Developer/TaskGraph/Private/SGraphBar.cpp
@@ -63,8 +63,8 @@ int32 SGraphBar::OnPaint( const FPaintArgs& Args, const FGeometry& AllottedGeome
 	bool bEnabled = ShouldBeEnabled( bParentEnabled );
 	const ESlateDrawEffect::Type DrawEffects = bEnabled ? ESlateDrawEffect::None : ESlateDrawEffect::DisabledEffect;
 		
-	const FColor ColorAndOpacitySRGB = InWidgetStyle.GetColorAndOpacityTint();
-	const FColor SelectedBarColor( 255, 255, 255, 255 );
+	const FLinearColor ColorAndOpacitySRGB = InWidgetStyle.GetColorAndOpacityTint();
+	static const FLinearColor SelectedBarColor(FLinearColor::White);
 
 	// Paint inside the border only. 
 	const FVector2D BorderPadding = FTaskGraphStyle::Get()->GetVector("TaskGraph.ProgressBar.BorderPadding");
diff --git a/Engine/Source/Developer/TaskGraph/Private/STimeline.cpp b/Engine/Source/Developer/TaskGraph/Private/STimeline.cpp
--- a/Engine/Source/Developer/TaskGraph/Private/STimeline.cpp
+++ b/Engine/Source/Developer/TaskGraph/Private/STimeline.cpp
@@ -27,8 +27,8 @@ int32 STimeline::OnPaint( const FPaintArgs& Args, const FGeometry& AllottedGeome
 	bool bEnabled = ShouldBeEnabled( bParentEnabled );
 	const ESlateDrawEffect::Type DrawEffects = bEnabled ? ESlateDrawEffect::None : ESlateDrawEffect::DisabledEffect;
 		
-	const FColor ColorAndOpacitySRGB = InWidgetStyle.GetColorAndOpacityTint();
-	const FColor SelectedBarColor( 255, 255, 255, 255 );
+	const FLinearColor ColorAndOpacitySRGB = InWidgetStyle.GetColorAndOpacityTint();
+	static const FLinearColor SelectedBarColor(FLinearColor::White);
 
 	// Paint inside the border only. 
 	const FVector2D BorderPadding = FTaskGraphStyle::Get()->GetVector("TaskGraph.ProgressBar.BorderPadding");
diff --git a/Engine/Source/Editor/ComponentVisualizers/Private/SplineComponentVisualizer.cpp b/Engine/Source/Editor/ComponentVisualizers/Private/SplineComponentVisualizer.cpp
--- a/Engine/Source/Editor/ComponentVisualizers/Private/SplineComponentVisualizer.cpp
+++ b/Engine/Source/Editor/ComponentVisualizers/Private/SplineComponentVisualizer.cpp
@@ -161,8 +161,8 @@ void FSplineComponentVisualizer::DrawVisualization(const UActorComponent* Compon
 		const bool bIsSplineEditable = SplineComp->bSplineHasBeenEdited || SplineComp->SplineInfo.Points == Archetype->SplineInfo.Points;
 
 		const FColor ReadOnlyColor = FColor(255, 0, 255, 255);
-		const FColor NormalColor = bIsSplineEditable ? FColor(SplineComp->EditorUnselectedSplineSegmentColor) : ReadOnlyColor;
-		const FColor SelectedColor = bIsSplineEditable ? FColor(SplineComp->EditorSelectedSplineSegmentColor) : ReadOnlyColor;
+		const FColor NormalColor = bIsSplineEditable ? FColor(SplineComp->EditorUnselectedSplineSegmentColor.ToFColor(true)) : ReadOnlyColor;
+		const FColor SelectedColor = bIsSplineEditable ? FColor(SplineComp->EditorSelectedSplineSegmentColor.ToFColor(true)) : ReadOnlyColor;
 		const float GrabHandleSize = 12.0f;
 		const float TangentHandleSize = 10.0f;
 
diff --git a/Engine/Source/Editor/DistCurveEditor/Private/CurveEditorViewportClient.cpp b/Engine/Source/Editor/DistCurveEditor/Private/CurveEditorViewportClient.cpp
--- a/Engine/Source/Editor/DistCurveEditor/Private/CurveEditorViewportClient.cpp
+++ b/Engine/Source/Editor/DistCurveEditor/Private/CurveEditorViewportClient.cpp
@@ -1106,7 +1106,7 @@ void FCurveEditorViewportClient::DrawEntry(FViewport* Viewport, FCanvas* Canvas,
 			FCurveEditorSelectedKey TestKey(CurveIndex, SubIdx, KeyIdx);
 			bool bSelectedKey = SharedData->SelectedKeys.Contains(TestKey);
 			FColor BorderColor = EdInterface->GetKeyColor(SubIdx, KeyIdx, Entry.CurveColor);
-			FColor CenterColor = bSelectedKey ? (FColor)SelectedKeyColor : Entry.CurveColor;
+			FColor CenterColor = bSelectedKey ? SelectedKeyColor.ToFColor(true) : Entry.CurveColor;
 
 			if(Canvas->IsHitTesting()) Canvas->SetHitProxy(new HCurveEditorKeyProxy(CurveIndex, SubIdx, KeyIdx));
 			Canvas->DrawTile( NewKeyPos.X-3, NewKeyPos.Y-3, 7, 7, 0.f, 0.f, 1.f, 1.f, BorderColor);
@@ -1278,26 +1278,29 @@ void FCurveEditorViewportClient::DrawGrid(FViewport* Viewport, FCanvas* Canvas)
 		AuxFrameStep = CurveEditor::CalculateBestFrameStep(InSnapAmount, PixelsPerIn, 6);
 	}
 
+	const static FLinearColor NormalLine(FColor(80, 80, 80));
+	const static FLinearColor ImportantLine(FColor(110, 110, 110));
+
 	// Draw input grid
 	FCanvasLineItem LineItem;
 	int32 InNum = FMath::FloorToInt(SharedData->StartIn/InGridSpacing);
-	while(InNum*InGridSpacing < SharedData->EndIn)
+	while(InNum * InGridSpacing < SharedData->EndIn)
 	{
-		FColor LineColor = GridColor;
+		FLinearColor LineColor = GridColor;
 
 		// Change line color for important frames.
 		if (bSnapToFrames)
 		{
-			LineColor = FColor(80,80,80);
+			LineColor = NormalLine;
 			if (InNum % FrameStep == 0)
 			{
-				LineColor = FColor(110,110,110);
+				LineColor = ImportantLine;
 			}
 		}
 
 		// Draw grid line.
 		// In frames mode auxiliary lines cannot be too close.
-		FIntPoint GridPos = CalcScreenPos(FVector2D(InNum*InGridSpacing, 0.f));		
+		FIntPoint GridPos = CalcScreenPos(FVector2D(InNum * InGridSpacing, 0.f));		
 		if (!bSnapToFrames || (FMath::Abs(InNum) % AuxFrameStep == 0))
 		{
 			LineItem.SetColor( LineColor );
diff --git a/Engine/Source/Editor/EditorWidgets/Private/SObjectNameEditableTextBox.cpp b/Engine/Source/Editor/EditorWidgets/Private/SObjectNameEditableTextBox.cpp
--- a/Engine/Source/Editor/EditorWidgets/Private/SObjectNameEditableTextBox.cpp
+++ b/Engine/Source/Editor/EditorWidgets/Private/SObjectNameEditableTextBox.cpp
@@ -89,7 +89,7 @@ int32 SObjectNameEditableTextBox::OnPaint( const FPaintArgs& Args, const FGeomet
 		float EffectOpacity = EffectAlpha;
 
 		// Figure out a universally visible highlight color.
-		FColor HighlightTargetColorAndOpacity = ( (FLinearColor::White - ColorAndOpacity.Get())*0.5f + FLinearColor(+0.4f, +0.1f, -0.2f)) * InWidgetStyle.GetColorAndOpacityTint();
+		FLinearColor HighlightTargetColorAndOpacity = ( (FLinearColor::White - ColorAndOpacity.Get())*0.5f + FLinearColor(+0.4f, +0.1f, -0.2f)) * InWidgetStyle.GetColorAndOpacityTint();
 		HighlightTargetColorAndOpacity.A = HighlightTargetOpacity * EffectOpacity * 255.0f;
 
 		// Compute the bounds offset of the highlight target from where the highlight target spring
diff --git a/Engine/Source/Editor/FoliageEdit/Private/FoliageEdMode.cpp b/Engine/Source/Editor/FoliageEdit/Private/FoliageEdMode.cpp
--- a/Engine/Source/Editor/FoliageEdit/Private/FoliageEdMode.cpp
+++ b/Engine/Source/Editor/FoliageEdit/Private/FoliageEdMode.cpp
@@ -1977,7 +1977,7 @@ struct FFoliagePaintBucketTriangle
 			y = 1.f - y;
 		}
 
-		OutBaryVertexColor = (1.f - x - y) * VertexColor[0] + x * VertexColor[1] + y * VertexColor[2];
+		OutBaryVertexColor = ( ( 1.f - x - y ) * VertexColor[0] + x * VertexColor[1] + y * VertexColor[2] ).ToFColor(true);
 		OutPoint = Vertex + x * Vector1 + y * Vector2;
 	}
 
diff --git a/Engine/Source/Editor/Matinee/Private/MatineeMenus.cpp b/Engine/Source/Editor/Matinee/Private/MatineeMenus.cpp
--- a/Engine/Source/Editor/Matinee/Private/MatineeMenus.cpp
+++ b/Engine/Source/Editor/Matinee/Private/MatineeMenus.cpp
@@ -2338,7 +2338,7 @@ void FMatinee::OnContextSetColor()
 		FVector CurrentColorVector = ColorPropTrack->VectorTrack.Points[SelKey.KeyIndex].OutVal;
 		static FColor CurrentColor;
 
-		CurrentColor = FColor(FLinearColor(CurrentColorVector.X, CurrentColorVector.Y, CurrentColorVector.Z));
+		CurrentColor = FLinearColor(CurrentColorVector.X, CurrentColorVector.Y, CurrentColorVector.Z).ToFColor(true);
 
 		TArray<FColor*> FColorArray;
 		FColorArray.Add(&CurrentColor);
diff --git a/Engine/Source/Editor/SceneOutliner/Private/SOutlinerTreeView.cpp b/Engine/Source/Editor/SceneOutliner/Private/SOutlinerTreeView.cpp
--- a/Engine/Source/Editor/SceneOutliner/Private/SOutlinerTreeView.cpp
+++ b/Engine/Source/Editor/SceneOutliner/Private/SOutlinerTreeView.cpp
@@ -343,7 +343,7 @@ namespace SceneOutliner
 			float EffectOpacity = EffectAlpha;
 
 			// Figure out a universally visible highlight color.
-			FColor HighlightTargetColorAndOpacity = ( (FLinearColor::White - ColorAndOpacity.Get())*0.5f + FLinearColor(+0.4f, +0.1f, -0.2f)) * InWidgetStyle.GetColorAndOpacityTint();
+			FLinearColor HighlightTargetColorAndOpacity = ( (FLinearColor::White - ColorAndOpacity.Get())*0.5f + FLinearColor(+0.4f, +0.1f, -0.2f)) * InWidgetStyle.GetColorAndOpacityTint();
 			HighlightTargetColorAndOpacity.A = HighlightTargetOpacity * EffectOpacity * 255.0f;
 
 			// Compute the bounds offset of the highlight target from where the highlight target spring
diff --git a/Engine/Source/Editor/UnrealEd/Private/Factories/TTFontImport.cpp b/Engine/Source/Editor/UnrealEd/Private/Factories/TTFontImport.cpp
--- a/Engine/Source/Editor/UnrealEd/Private/Factories/TTFontImport.cpp
+++ b/Engine/Source/Editor/UnrealEd/Private/Factories/TTFontImport.cpp
@@ -165,8 +165,8 @@ UObject* UTrueTypeFontFactory::FactoryCreateNew(
 		ImportOptions->Data.DistanceFieldScaleFactor = FMath::RoundUpToPowerOfTwo(FMath::Max<int32>(ImportOptions->Data.DistanceFieldScaleFactor,2));
 		ImportOptions->Data.DistanceFieldScanRadiusScale = FMath::Clamp<float>(ImportOptions->Data.DistanceFieldScanRadiusScale,0.f,8.f);
 		// need a minimum padding of 4,4 to prevent bleeding of distance values across characters
- 		ImportOptions->Data.XPadding = FMath::Max<int32>(4,ImportOptions->Data.XPadding);
- 		ImportOptions->Data.YPadding = FMath::Max<int32>(4,ImportOptions->Data.YPadding);
+		ImportOptions->Data.XPadding = FMath::Max<int32>(4,ImportOptions->Data.XPadding);
+		ImportOptions->Data.YPadding = FMath::Max<int32>(4,ImportOptions->Data.YPadding);
 	}
 
 	// Copy the import settings into the font for later reference
@@ -626,7 +626,7 @@ float FTextureAlphaToDistanceField::FBuildDistanceFieldTask::CalcSignedDistanceT
 	//        2 1 * 1 2
 	//        2 1 1 1 2
 	//        2 2 2 2 2
-    //  
+	//  
 	// Note that the "rings" are not actually circular, so
 	// we may find a sample that is up to Sqrt(2*(RingSize^2)) away.
 	// Once we have found such a sample, we must search a few more
@@ -759,7 +759,7 @@ UTexture2D* UTrueTypeFontFactory::CreateTextureFromDC( UFont* Font, HDC dc, int3
 		TextureString = TextureString + FString::Printf(TEXT("%c%c"), 'A'+TextureNum/26, 'A'+TextureNum%26 );
 	}
 
- 	if( StaticFindObject( NULL, Font, *TextureString ) )
+	if( StaticFindObject( NULL, Font, *TextureString ) )
 	{
 		UE_LOG(LogTTFontImport, Warning, TEXT("A texture named %s already exists!"), *TextureString );
 	}
@@ -788,7 +788,7 @@ UTexture2D* UTrueTypeFontFactory::CreateTextureFromDC( UFont* Font, HDC dc, int3
 	Texture->NeverStream = true;
 
 	// Copy bitmap data to texture page.
-	FColor FontColor8Bit( ImportOptions->Data.ForegroundColor );
+	FColor FontColor8Bit(ImportOptions->Data.ForegroundColor.ToFColor(true));
 
 	// restart progress bar for font bitmap generation since this takes a while
 	int32 TotalProgress = BitmapWidth-1;
@@ -1143,7 +1143,7 @@ bool UTrueTypeFontFactory::CreateFontTexture(
 
 		for( int32 CurCharIndex = 0; CurCharIndex < CharsPerPage; ++CurCharIndex )
 		{
-  			GWarn->UpdateProgress( Page * CharsPerPage + CurCharIndex, TotalProgress );
+			GWarn->UpdateProgress( Page * CharsPerPage + CurCharIndex, TotalProgress );
 
 			// Remap the character if we need to
 			TCHAR Char = ( TCHAR )CurCharIndex;
@@ -1654,7 +1654,7 @@ UTexture2D* UTrueTypeFontFactory::CreateTextureFromBitmap( UFont* Font, uint8* B
 		TextureString = TextureString + FString::Printf(TEXT("%c%c"), 'A'+TextureNum/26, 'A'+TextureNum%26 );
 	}
 
- 	if( StaticFindObject( NULL, Font, *TextureString ) )
+	if( StaticFindObject( NULL, Font, *TextureString ) )
 	{
 		UE_LOG(LogTTFontImport, Warning, TEXT("A texture named %s already exists!"), *TextureString );
 	}
@@ -1903,7 +1903,7 @@ bool UTrueTypeFontFactory::CreateFontTexture(
 
 		for( int32 CurCharIndex = 0; CurCharIndex < CharsPerPage; ++CurCharIndex )
 		{
-  			GWarn->UpdateProgress( Page * CharsPerPage + CurCharIndex, TotalProgress );
+			GWarn->UpdateProgress( Page * CharsPerPage + CurCharIndex, TotalProgress );
 
 			// Remap the character if we need to
 			TCHAR Char = ( TCHAR )CurCharIndex;
@@ -2221,21 +2221,21 @@ bool UTrueTypeFontFactory::ImportTrueTypeFont(
 
 	// Add space for characters.
 	Font->Characters.AddZeroed(CharsPerPage * NumResolutions );
-    
+	
 	// If all upper case chars have lower case char counterparts no mapping is required.   
 	if( !Font->IsRemapped )
 	{
 		bool NeedToRemap = false;
-        
+		
 		for( const TCHAR* p = *ImportOptions->Data.Chars; *p; p++ )
 		{
 			TCHAR c;
-            
+			
 			if( !FChar::IsAlpha( *p ) )
 			{
 				continue;
 			}
-            
+			
 			if( FChar::IsUpper( *p ) )
 			{
 				c = FChar::ToLower( *p );
@@ -2249,11 +2249,11 @@ bool UTrueTypeFontFactory::ImportTrueTypeFont(
 			{
 				continue;
 			}
-            
+			
 			NeedToRemap = true;
 			break;
 		}
-        
+		
 		if( NeedToRemap )
 		{
 			Font->IsRemapped = 1;
@@ -2268,7 +2268,7 @@ bool UTrueTypeFontFactory::ImportTrueTypeFont(
 					InverseMap.Add( *p, *p );
 					continue;
 				}
-                
+				
 				if( FChar::IsUpper( *p ) )
 				{
 					c = FChar::ToLower( *p );
diff --git a/Engine/Source/Editor/UnrealEd/Private/Fbx/FbxMainExport.cpp b/Engine/Source/Editor/UnrealEd/Private/Fbx/FbxMainExport.cpp
--- a/Engine/Source/Editor/UnrealEd/Private/Fbx/FbxMainExport.cpp
+++ b/Engine/Source/Editor/UnrealEd/Private/Fbx/FbxMainExport.cpp
@@ -887,7 +887,7 @@ FbxDouble3 SetMaterialComponent(FColorMaterialInput& MatInput)
 		else if (Cast<UMaterialExpressionVectorParameter>(MatInput.Expression))
 		{
 			UMaterialExpressionVectorParameter* Expr = Cast<UMaterialExpressionVectorParameter>(MatInput.Expression);
-			FinalColor = Expr->DefaultValue;
+			FinalColor = Expr->DefaultValue.ToFColor(true);
 		}
 		else if (Cast<UMaterialExpressionConstant3Vector>(MatInput.Expression))
 		{
diff --git a/Engine/Source/Editor/UnrealEd/Private/UnrealWidget.cpp b/Engine/Source/Editor/UnrealEd/Private/UnrealWidget.cpp
--- a/Engine/Source/Editor/UnrealEd/Private/UnrealWidget.cpp
+++ b/Engine/Source/Editor/UnrealEd/Private/UnrealWidget.cpp
@@ -532,17 +532,17 @@ void FWidget::Render_Rotate( const FSceneView* View,FPrimitiveDrawInterface* PDI
 		//no draw the arc segments
 		if( DrawAxis&EAxisList::X )
 		{
-			DrawRotationArc(View, PDI, EAxisList::X, InLocation, ZAxis, YAxis, DirectionToWidget, AxisColorX, Scale, XAxisDir);
+			DrawRotationArc(View, PDI, EAxisList::X, InLocation, ZAxis, YAxis, DirectionToWidget, AxisColorX.ToFColor(true), Scale, XAxisDir);
 		}
 
 		if( DrawAxis&EAxisList::Y )
 		{
-			DrawRotationArc(View, PDI, EAxisList::Y, InLocation, XAxis, ZAxis, DirectionToWidget, AxisColorY, Scale, YAxisDir);
+			DrawRotationArc(View, PDI, EAxisList::Y, InLocation, XAxis, ZAxis, DirectionToWidget, AxisColorY.ToFColor(true), Scale, YAxisDir);
 		}
 
 		if( DrawAxis&EAxisList::Z )
 		{
-			DrawRotationArc(View, PDI, EAxisList::Z, InLocation, XAxis, YAxis, DirectionToWidget, AxisColorZ, Scale, ZAxisDir);
+			DrawRotationArc(View, PDI, EAxisList::Z, InLocation, XAxis, YAxis, DirectionToWidget, AxisColorZ.ToFColor(true), Scale, ZAxisDir);
 		}
 	}
 }
@@ -669,10 +669,10 @@ void FWidget::Render_TranslateRotateZ( const FSceneView* View, FPrimitiveDrawInt
 	// Figure out axis colors
 
 	FColor XYPlaneColor  = ( (CurrentAxis&EAxisList::XY) ==  EAxisList::XY) ? CurrentColor : PlaneColorXY;
-	FColor ZRotateColor  = ( (CurrentAxis&EAxisList::ZRotation) == EAxisList::ZRotation ) ? CurrentColor : (FColor)AxisColorZ ;
-	FColor XColor        = ( (CurrentAxis&EAxisList::X ) == EAxisList::X )? CurrentColor : (FColor)AxisColorX;
-	FColor YColor        = ( (CurrentAxis&EAxisList::Y) == EAxisList::Y && CurrentAxis != EAxisList::ZRotation ) ? CurrentColor : (FColor)AxisColorY;
-	FColor ZColor        = ( (CurrentAxis&EAxisList::Z) == EAxisList::Z) ? CurrentColor : (FColor)AxisColorZ;
+	FColor ZRotateColor  = ( ( CurrentAxis&EAxisList::ZRotation ) == EAxisList::ZRotation ) ? CurrentColor : AxisColorZ.ToFColor(true);
+	FColor XColor        = ( ( CurrentAxis&EAxisList::X ) == EAxisList::X ) ? CurrentColor : AxisColorX.ToFColor(true);
+	FColor YColor        = ( ( CurrentAxis&EAxisList::Y ) == EAxisList::Y && CurrentAxis != EAxisList::ZRotation ) ? CurrentColor : AxisColorY.ToFColor(true);
+	FColor ZColor        = ( ( CurrentAxis&EAxisList::Z ) == EAxisList::Z ) ? CurrentColor : AxisColorZ.ToFColor(true);
 
 	// Figure out axis materials
 	UMaterialInterface* ZRotateMaterial = (CurrentAxis&EAxisList::ZRotation) == EAxisList::ZRotation? CurrentAxisMaterial : AxisMaterialZ;
@@ -912,7 +912,7 @@ void FWidget::Render_2D(const FSceneView* View, FPrimitiveDrawInterface* PDI, FE
 			uint8 Alpha = ((CurrentAxis&EAxisList::XZ) == EAxisList::XZ ? HoverAlpha : NormalAlpha);
 			PDI->SetHitProxy(new HWidgetAxis(EAxisList::XZ, bDisabled));
 			{
-				FColor Color = YColor;
+				FColor Color = YColor.ToFColor(true);
 				DrawCircle(PDI, InLocation, CustomCoordSystem.TransformPosition(FVector(1, 0, 0)), CustomCoordSystem.TransformPosition(FVector(0, 0, 1)), Color, ScaledRadius, CircleSides, SDPG_Foreground);
 				Color.A = Alpha;
 				DrawDisc(PDI, InLocation, CustomCoordSystem.TransformPosition(FVector(1, 0, 0)), CustomCoordSystem.TransformPosition(FVector(0, 0, 1)), Color, ScaledRadius, CircleSides, TransparentPlaneMaterialXY->GetRenderProxy(false), SDPG_Foreground);
@@ -921,7 +921,7 @@ void FWidget::Render_2D(const FSceneView* View, FPrimitiveDrawInterface* PDI, FE
 
 			PDI->SetHitProxy(new HWidgetAxis(EAxisList::Rotate2D, bDisabled));
 			{
-				FColor Color = YColor;
+				FColor Color = YColor.ToFColor(true);
 				Color.A = ((CurrentAxis&EAxisList::Rotate2D) == EAxisList::Rotate2D ? HoverAlpha : NormalAlpha);
 
 				FVector XAxis = CustomCoordSystem.TransformPosition(FVector(1, 0, 0).RotateAngleAxis((EditorModeTools ? EditorModeTools->TranslateRotate2DAngle : 0), FVector(0, -1, 0)));
diff --git a/Engine/Source/Runtime/AIModule/Private/AIController.cpp b/Engine/Source/Runtime/AIModule/Private/AIController.cpp
--- a/Engine/Source/Runtime/AIModule/Private/AIController.cpp
+++ b/Engine/Source/Runtime/AIModule/Private/AIController.cpp
@@ -559,7 +559,7 @@ EPathFollowingRequestResult::Type AAIController::MoveTo(const FAIMoveRequest& Mo
 
 			if (NavSys && !NavSys->ProjectPointToNavigation(MoveRequest.GetGoalLocation(), ProjectedLocation, AgentProps.GetExtent(), &AgentProps))
 			{
-				UE_VLOG_LOCATION(this, LogAINavigation, Error, MoveRequest.GetGoalLocation(), 30.f, FLinearColor::Red, TEXT("AAIController::MoveTo failed to project destination location to navmesh"));
+				UE_VLOG_LOCATION(this, LogAINavigation, Error, MoveRequest.GetGoalLocation(), 30.f, FColor::Red, TEXT("AAIController::MoveTo failed to project destination location to navmesh"));
 				bCanRequestMove = false;
 			}
 
diff --git a/Engine/Source/Runtime/Core/Private/Math/Color.cpp b/Engine/Source/Runtime/Core/Private/Math/Color.cpp
--- a/Engine/Source/Runtime/Core/Private/Math/Color.cpp
+++ b/Engine/Source/Runtime/Core/Private/Math/Color.cpp
@@ -18,6 +18,7 @@ const FLinearColor FLinearColor::Yellow(1.f,1.f,0);
 
 const FColor FColor::White(255,255,255);
 const FColor FColor::Black(0,0,0);
+const FColor FColor::Transparent(0, 0, 0, 0);
 const FColor FColor::Red(255,0,0);
 const FColor FColor::Green(0,255,0);
 const FColor FColor::Blue(0,0,255);
@@ -374,7 +375,7 @@ FColor FColor::MakeRedToGreenColorFromScalar(float Scalar)
 void ComputeAndFixedColorAndIntensity(const FLinearColor& InLinearColor,FColor& OutColor,float& OutIntensity)
 {
 	float MaxComponent = FMath::Max(DELTA,FMath::Max(InLinearColor.R,FMath::Max(InLinearColor.G,InLinearColor.B)));
-	OutColor = InLinearColor / MaxComponent;
+	OutColor = ( InLinearColor / MaxComponent ).ToFColor(true);
 	OutIntensity = MaxComponent;
 }
 
diff --git a/Engine/Source/Runtime/Core/Public/Math/Color.h b/Engine/Source/Runtime/Core/Public/Math/Color.h
--- a/Engine/Source/Runtime/Core/Public/Math/Color.h
+++ b/Engine/Source/Runtime/Core/Public/Math/Color.h
@@ -427,18 +427,6 @@ public:
 		A = InA;
 
 	}
-	
-	// fast, for more accuracy use FLinearColor::ToFColor()
-	// TODO: doesn't handle negative colors well, implicit constructor can cause
-	// accidental conversion (better use .ToFColor(true))
-	FColor(const FLinearColor& C)
-		// put these into the body for proper ordering with INTEL vs non-INTEL_BYTE_ORDER
-	{
-		R = FMath::Clamp(FMath::TruncToInt(FMath::Pow(C.R,1.0f / 2.2f) * 255.0f),0,255);
-		G = FMath::Clamp(FMath::TruncToInt(FMath::Pow(C.G,1.0f / 2.2f) * 255.0f),0,255);
-		B = FMath::Clamp(FMath::TruncToInt(FMath::Pow(C.B,1.0f / 2.2f) * 255.0f),0,255);
-		A = FMath::Clamp(FMath::TruncToInt(       C.A              * 255.0f),0,255);
-	}
 
 	FORCEINLINE explicit FColor( uint32 InColor )
 	{ 
@@ -577,12 +565,24 @@ public:
 	/** Some pre-inited colors, useful for debug code */
 	static CORE_API const FColor White;
 	static CORE_API const FColor Black;
+	static CORE_API const FColor Transparent;
 	static CORE_API const FColor Red;
 	static CORE_API const FColor Green;
 	static CORE_API const FColor Blue;
 	static CORE_API const FColor Yellow;
 	static CORE_API const FColor Cyan;
 	static CORE_API const FColor Magenta;
+
+private:
+	/**
+	 * Please use .ToFColor(true) on FLinearColor if you wish to convert from FLinearColor to FColor,
+	 * with proper sRGB conversion applied.
+	 */
+	explicit FColor(const FLinearColor& LinearColor)
+	{
+		// Do not implement or make public.  We don't want people needlessly and implicitly converting between
+		// FLinearColor and FColor.  It's not a free conversion.
+	}
 };
 
 FORCEINLINE uint32 GetTypeHash( const FColor& Color )
@@ -590,9 +590,10 @@ FORCEINLINE uint32 GetTypeHash( const FColor& Color )
 	return Color.DWColor();
 }
 
-FORCEINLINE uint32 GetTypeHash( const FLinearColor& Color )
+FORCEINLINE uint32 GetTypeHash( const FLinearColor& LinearColor )
 {
-	return GetTypeHash(FColor(Color));
+	// Note: this assumes there's no padding in FLinearColor that could contain uncompared data.
+	return FCrc::MemCrc_DEPRECATED(&LinearColor, sizeof(FLinearColor));
 }
 
 
diff --git a/Engine/Source/Runtime/Engine/Private/AI/Navigation/NavigationPath.cpp b/Engine/Source/Runtime/Engine/Private/AI/Navigation/NavigationPath.cpp
--- a/Engine/Source/Runtime/Engine/Private/AI/Navigation/NavigationPath.cpp
+++ b/Engine/Source/Runtime/Engine/Private/AI/Navigation/NavigationPath.cpp
@@ -1263,7 +1263,7 @@ void UNavigationPath::DrawDebug(UCanvas* Canvas, APlayerController*)
 
 void UNavigationPath::EnableDebugDrawing(bool bShouldDrawDebugData, FLinearColor PathColor)
 {
-	DebugDrawingColor = PathColor;
+	DebugDrawingColor = PathColor.ToFColor(true);
 
 	if (bDebugDrawingEnabled == bShouldDrawDebugData)
 	{
diff --git a/Engine/Source/Runtime/Engine/Private/Atmosphere/Atmosphere.cpp b/Engine/Source/Runtime/Engine/Private/Atmosphere/Atmosphere.cpp
--- a/Engine/Source/Runtime/Engine/Private/Atmosphere/Atmosphere.cpp
+++ b/Engine/Source/Runtime/Engine/Private/Atmosphere/Atmosphere.cpp
@@ -362,7 +362,7 @@ void UAtmosphericFogComponent::SetDefaultBrightness(float NewBrightness)
 /** Set color of the light */
 void UAtmosphericFogComponent::SetDefaultLightColor(FLinearColor NewLightColor)
 {
-	FColor NewColor(NewLightColor);
+	FColor NewColor(NewLightColor.ToFColor(true));
 	if( DefaultLightColor != NewColor )
 	{
 		DefaultLightColor = NewColor;
diff --git a/Engine/Source/Runtime/Engine/Private/BatchedElements.cpp b/Engine/Source/Runtime/Engine/Private/BatchedElements.cpp
--- a/Engine/Source/Runtime/Engine/Private/BatchedElements.cpp
+++ b/Engine/Source/Runtime/Engine/Private/BatchedElements.cpp
@@ -72,7 +72,7 @@ void FBatchedElements::AddPoint(const FVector& Position,float Size,const FLinear
 	FBatchedPoint* Point = new(Points) FBatchedPoint;
 	Point->Position = Position;
 	Point->Size = Size;
-	Point->Color = OpaqueColor;
+	Point->Color = OpaqueColor.ToFColor(true);
 	Point->HitProxyId = HitProxyId;
 }
 
@@ -468,23 +468,23 @@ void FBatchedElements::PrepareShaders(
 	else
 	{
 		TShaderMapRef<FSimpleElementVS> VertexShader(GetGlobalShaderMap(FeatureLevel));
-		    
-	    if (bHitTesting)
-	    {
+			
+		if (bHitTesting)
+		{
 			TShaderMapRef<FSimpleElementHitProxyPS> HitTestingPixelShader(GetGlobalShaderMap(FeatureLevel));
 			SetGlobalBoundShaderState(RHICmdList, FeatureLevel, HitTestingBoundShaderState, GSimpleElementVertexDeclaration.VertexDeclarationRHI,
 				*VertexShader, *HitTestingPixelShader);
 
 			HitTestingPixelShader->SetParameters(RHICmdList, Texture);
 			SetHitTestingBlendState(RHICmdList, BlendMode);
-		    
-	    }
-	    else
-	    {
-		    if (BlendMode == SE_BLEND_Masked)
-		    {
-			    // use clip() in the shader instead of alpha testing as cards that don't support floating point blending
-			    // also don't support alpha testing to floating point render targets
+			
+		}
+		else
+		{
+			if (BlendMode == SE_BLEND_Masked)
+			{
+				// use clip() in the shader instead of alpha testing as cards that don't support floating point blending
+				// also don't support alpha testing to floating point render targets
 				RHICmdList.SetBlendState(TStaticBlendState<>::GetRHI());
 
 				if (Texture->bSRGB)
@@ -505,78 +505,78 @@ void FBatchedElements::PrepareShaders(
 					MaskedPixelShader->SetEditorCompositingParameters(RHICmdList, View, DepthTexture);
 					MaskedPixelShader->SetParameters(RHICmdList, Texture, Gamma, GBatchedElementAlphaRefVal / 255.0f, BlendMode);
 				}
-		    }
-		    // render distance field elements
-		    else if (
-			    BlendMode == SE_BLEND_MaskedDistanceField || 
-			    BlendMode == SE_BLEND_MaskedDistanceFieldShadowed ||
-			    BlendMode == SE_BLEND_TranslucentDistanceField	||
-			    BlendMode == SE_BLEND_TranslucentDistanceFieldShadowed)
-		    {
-			    float AlphaRefVal = GBatchedElementAlphaRefVal;
-			    if (BlendMode == SE_BLEND_TranslucentDistanceField ||
-				    BlendMode == SE_BLEND_TranslucentDistanceFieldShadowed)
-			    {
-				    // enable alpha blending and disable clip ref value for translucent rendering
+			}
+			// render distance field elements
+			else if (
+				BlendMode == SE_BLEND_MaskedDistanceField || 
+				BlendMode == SE_BLEND_MaskedDistanceFieldShadowed ||
+				BlendMode == SE_BLEND_TranslucentDistanceField	||
+				BlendMode == SE_BLEND_TranslucentDistanceFieldShadowed)
+			{
+				float AlphaRefVal = GBatchedElementAlphaRefVal;
+				if (BlendMode == SE_BLEND_TranslucentDistanceField ||
+					BlendMode == SE_BLEND_TranslucentDistanceFieldShadowed)
+				{
+					// enable alpha blending and disable clip ref value for translucent rendering
 					RHICmdList.SetBlendState(TStaticBlendState<CW_RGB, BO_Add, BF_SourceAlpha, BF_InverseSourceAlpha>::GetRHI());
-				    AlphaRefVal = 0.0f;
-			    }
-			    else
-			    {
-				    // clip is done in shader so just render opaque
+					AlphaRefVal = 0.0f;
+				}
+				else
+				{
+					// clip is done in shader so just render opaque
 					RHICmdList.SetBlendState(TStaticBlendState<>::GetRHI());
-			    }
-			    
+				}
+				
 				TShaderMapRef<FSimpleElementDistanceFieldGammaPS> DistanceFieldPixelShader(GetGlobalShaderMap(FeatureLevel));
 				SetGlobalBoundShaderState(RHICmdList, FeatureLevel, DistanceFieldBoundShaderState, GSimpleElementVertexDeclaration.VertexDeclarationRHI,
 					*VertexShader, *DistanceFieldPixelShader );			
 
-			    // @todo - expose these as options for batch rendering
-			    static FVector2D ShadowDirection(-1.0f/Texture->GetSizeX(),-1.0f/Texture->GetSizeY());
-			    static FLinearColor ShadowColor(FLinearColor::Black);
-			    const float ShadowSmoothWidth = (GBatchedElementSmoothWidth * 2) / Texture->GetSizeX();
-			    
-			    const bool EnableShadow = (
-				    BlendMode == SE_BLEND_MaskedDistanceFieldShadowed || 
-				    BlendMode == SE_BLEND_TranslucentDistanceFieldShadowed
-				    );
-			    
-			    DistanceFieldPixelShader->SetParameters(
+				// @todo - expose these as options for batch rendering
+				static FVector2D ShadowDirection(-1.0f/Texture->GetSizeX(),-1.0f/Texture->GetSizeY());
+				static FLinearColor ShadowColor(FLinearColor::Black);
+				const float ShadowSmoothWidth = (GBatchedElementSmoothWidth * 2) / Texture->GetSizeX();
+				
+				const bool EnableShadow = (
+					BlendMode == SE_BLEND_MaskedDistanceFieldShadowed || 
+					BlendMode == SE_BLEND_TranslucentDistanceFieldShadowed
+					);
+				
+				DistanceFieldPixelShader->SetParameters(
 					RHICmdList, 
-				    Texture,
-				    Gamma,
-				    AlphaRefVal / 255.0f,
-				    GBatchedElementSmoothWidth,
-				    EnableShadow,
-				    ShadowDirection,
-				    ShadowColor,
-				    ShadowSmoothWidth,
-				    (GlowInfo != NULL) ? *GlowInfo : FDepthFieldGlowInfo(),
-				    BlendMode
-				    );
-		    }
+					Texture,
+					Gamma,
+					AlphaRefVal / 255.0f,
+					GBatchedElementSmoothWidth,
+					EnableShadow,
+					ShadowDirection,
+					ShadowColor,
+					ShadowSmoothWidth,
+					(GlowInfo != NULL) ? *GlowInfo : FDepthFieldGlowInfo(),
+					BlendMode
+					);
+			}
 			else if(BlendMode == SE_BLEND_TranslucentAlphaOnly)
 			{
 				SetBlendState(RHICmdList, BlendMode);
 
 				if (FMath::Abs(Gamma - 1.0f) < KINDA_SMALL_NUMBER)
-			    {
+				{
 					TShaderMapRef<FSimpleElementAlphaOnlyPS> AlphaOnlyPixelShader(GetGlobalShaderMap(FeatureLevel));
 					SetGlobalBoundShaderState(RHICmdList, FeatureLevel, AlphaOnlyShaderState, GSimpleElementVertexDeclaration.VertexDeclarationRHI,
 						*VertexShader, *AlphaOnlyPixelShader);
 
 					AlphaOnlyPixelShader->SetParameters(RHICmdList, Texture);
 					AlphaOnlyPixelShader->SetEditorCompositingParameters(RHICmdList, View, DepthTexture);
-			    }
-			    else
-			    {
+				}
+				else
+				{
 					TShaderMapRef<FSimpleElementGammaAlphaOnlyPS> GammaAlphaOnlyPixelShader(GetGlobalShaderMap(FeatureLevel));
 					SetGlobalBoundShaderState(RHICmdList, FeatureLevel, GammaAlphaOnlyShaderState, GSimpleElementVertexDeclaration.VertexDeclarationRHI,
 						*VertexShader, *GammaAlphaOnlyPixelShader);
 
 					GammaAlphaOnlyPixelShader->SetParameters(RHICmdList, Texture, Gamma, BlendMode);
 					GammaAlphaOnlyPixelShader->SetEditorCompositingParameters(RHICmdList, View, DepthTexture);
-			    }
+				}
 			}
 			else if(BlendMode >= SE_BLEND_RGBA_MASK_START && BlendMode <= SE_BLEND_RGBA_MASK_END)
 			{
@@ -587,20 +587,20 @@ void FBatchedElements::PrepareShaders(
 				ColorChannelMaskPixelShader->SetParameters(RHICmdList, Texture, ColorWeights, GammaToUse );
 			}
 			else
-		    {
+			{
 				SetBlendState(RHICmdList, BlendMode);
-    
-			    if (FMath::Abs(Gamma - 1.0f) < KINDA_SMALL_NUMBER)
-			    {
+	
+				if (FMath::Abs(Gamma - 1.0f) < KINDA_SMALL_NUMBER)
+				{
 					TShaderMapRef<FSimpleElementPS> PixelShader(GetGlobalShaderMap(FeatureLevel));
 					SetGlobalBoundShaderState(RHICmdList, FeatureLevel, SimpleBoundShaderState, GSimpleElementVertexDeclaration.VertexDeclarationRHI,
 						*VertexShader, *PixelShader);
 
 					PixelShader->SetParameters(RHICmdList, Texture);
 					PixelShader->SetEditorCompositingParameters(RHICmdList, View, DepthTexture);
-			    }
-			    else
-			    {
+				}
+				else
+				{
 					FSimpleElementGammaBasePS* BasePixelShader;
 					FGlobalBoundShaderState* BoundShaderState;
 
@@ -622,9 +622,9 @@ void FBatchedElements::PrepareShaders(
 
 					BasePixelShader->SetParameters(RHICmdList, Texture, Gamma, BlendMode);
 					BasePixelShader->SetEditorCompositingParameters(RHICmdList, View, DepthTexture);
-			    }
-		    }
-	    }
+				}
+			}
+		}
 
 		// Set the simple element vertex shader parameters
 		VertexShader->SetParameters(RHICmdList, Transform, bSwitchVerticalAxis);
diff --git a/Engine/Source/Runtime/Engine/Private/Components/ArrowComponent.cpp b/Engine/Source/Runtime/Engine/Private/Components/ArrowComponent.cpp
--- a/Engine/Source/Runtime/Engine/Private/Components/ArrowComponent.cpp
+++ b/Engine/Source/Runtime/Engine/Private/Components/ArrowComponent.cpp
@@ -334,7 +334,7 @@ void UArrowComponent::SetArrowColor_DEPRECATED(FColor NewColor)
 
 void UArrowComponent::SetArrowColor_New(FLinearColor NewColor)
 {
-	ArrowColor = NewColor;
+	ArrowColor = NewColor.ToFColor(true);
 	MarkRenderStateDirty();
 }
 
diff --git a/Engine/Source/Runtime/Engine/Private/Components/BrushComponent.cpp b/Engine/Source/Runtime/Engine/Private/Components/BrushComponent.cpp
--- a/Engine/Source/Runtime/Engine/Private/Components/BrushComponent.cpp
+++ b/Engine/Source/Runtime/Engine/Private/Components/BrushComponent.cpp
@@ -282,7 +282,7 @@ public:
 							Collector.RegisterOneFrameMaterialProxy(SolidMaterialInstance);
 
 							FTransform GeomTransform(GetLocalToWorld());
-							BodySetup->AggGeom.GetAggGeom(GeomTransform, DrawColor, /*Material=*/SolidMaterialInstance, false, /*bSolid=*/ true, UseEditorDepthTest(), ViewIndex, Collector);
+							BodySetup->AggGeom.GetAggGeom(GeomTransform, DrawColor.ToFColor(true), /*Material=*/SolidMaterialInstance, false, /*bSolid=*/ true, UseEditorDepthTest(), ViewIndex, Collector);
 						}
 					}
 					// WIREFRAME
@@ -320,7 +320,7 @@ public:
 							// If not, use the body setup for wireframe
 						{
 							FTransform GeomTransform(GetLocalToWorld());
-							BodySetup->AggGeom.GetAggGeom(GeomTransform, GetSelectionColor(DrawColor, IsSelected(), IsHovered()), /* Material=*/ NULL, false, /* bSolid=*/ false, UseEditorDepthTest(), ViewIndex, Collector);
+							BodySetup->AggGeom.GetAggGeom(GeomTransform, GetSelectionColor(DrawColor, IsSelected(), IsHovered()).ToFColor(true), /* Material=*/ NULL, false, /* bSolid=*/ false, UseEditorDepthTest(), ViewIndex, Collector);
 						}
 
 					}
diff --git a/Engine/Source/Runtime/Engine/Private/Components/LightComponent.cpp b/Engine/Source/Runtime/Engine/Private/Components/LightComponent.cpp
--- a/Engine/Source/Runtime/Engine/Private/Components/LightComponent.cpp
+++ b/Engine/Source/Runtime/Engine/Private/Components/LightComponent.cpp
@@ -719,7 +719,7 @@ void ULightComponent::SetIndirectLightingIntensity(float NewIntensity)
 /** Set color of the light */
 void ULightComponent::SetLightColor(FLinearColor NewLightColor)
 {
-	FColor NewColor(NewLightColor);
+	FColor NewColor(NewLightColor.ToFColor(true));
 
 	// Can't set color on a static light
 	if (AreDynamicDataChangesAllowed()
diff --git a/Engine/Source/Runtime/Engine/Private/Components/MaterialBillboardComponent.cpp b/Engine/Source/Runtime/Engine/Private/Components/MaterialBillboardComponent.cpp
--- a/Engine/Source/Runtime/Engine/Private/Components/MaterialBillboardComponent.cpp
+++ b/Engine/Source/Runtime/Engine/Private/Components/MaterialBillboardComponent.cpp
@@ -89,7 +89,7 @@ public:
 	FMaterialSpriteSceneProxy(const UMaterialBillboardComponent* InComponent)
 	: FPrimitiveSceneProxy(InComponent)
 	, Elements(InComponent->Elements)
-	, BaseColor(FLinearColor::White)
+	, BaseColor(FColor::White)
 	{
 		AActor* Owner = InComponent->GetOwner();
 		if (Owner)
@@ -179,7 +179,7 @@ public:
 
 						for(uint32 VertexIndex = 0;VertexIndex < 4;++VertexIndex)
 						{
-							VertexArray.Vertices[VertexIndex].Color = Color;
+							VertexArray.Vertices[VertexIndex].Color = Color.ToFColor(true);
 							VertexArray.Vertices[VertexIndex].TangentX = FPackedNormal(LocalCameraRight.GetSafeNormal());
 							VertexArray.Vertices[VertexIndex].TangentZ = FPackedNormal(-LocalCameraForward.GetSafeNormal());
 						}
diff --git a/Engine/Source/Runtime/Engine/Private/Components/SkeletalMeshComponent.cpp b/Engine/Source/Runtime/Engine/Private/Components/SkeletalMeshComponent.cpp
--- a/Engine/Source/Runtime/Engine/Private/Components/SkeletalMeshComponent.cpp
+++ b/Engine/Source/Runtime/Engine/Private/Components/SkeletalMeshComponent.cpp
@@ -273,8 +273,8 @@ void USkeletalMeshComponent::InitAnim(bool bForceReinit)
 
 		InitializeAnimScriptInstance(bForceReinit);
 
-        //Make sure we have a valid pose		
-        TickAnimation(0.f); 
+		//Make sure we have a valid pose		
+		TickAnimation(0.f); 
 
 		RefreshBoneTransforms();
 		UpdateComponentToWorld();
@@ -1524,8 +1524,8 @@ void USkeletalMeshComponent::DebugDrawBones(UCanvas* Canvas, bool bSimpleBones)
 				float ConeLength = EndToStart.Size();
 				float Angle = FMath::RadiansToDegrees(FMath::Atan(SphereRadius / ConeLength));
 
-				DrawDebugCanvasWireSphere(Canvas, End, LineColor, SphereRadius, 10);
-				DrawDebugCanvasWireCone(Canvas, FTransform(FRotationMatrix::MakeFromX(EndToStart)*FTranslationMatrix(End)), ConeLength, Angle, 4, LineColor);
+				DrawDebugCanvasWireSphere(Canvas, End, LineColor.ToFColor(true), SphereRadius, 10);
+				DrawDebugCanvasWireCone(Canvas, FTransform(FRotationMatrix::MakeFromX(EndToStart)*FTranslationMatrix(End)), ConeLength, Angle, 4, LineColor.ToFColor(true));
 			}
 
 			RenderAxisGizmo(BoneTM, Canvas);
diff --git a/Engine/Source/Runtime/Engine/Private/Components/SkyLightComponent.cpp b/Engine/Source/Runtime/Engine/Private/Components/SkyLightComponent.cpp
--- a/Engine/Source/Runtime/Engine/Private/Components/SkyLightComponent.cpp
+++ b/Engine/Source/Runtime/Engine/Private/Components/SkyLightComponent.cpp
@@ -484,7 +484,7 @@ void USkyLightComponent::SetIndirectLightingIntensity(float NewIntensity)
 /** Set color of the light */
 void USkyLightComponent::SetLightColor(FLinearColor NewLightColor)
 {
-	FColor NewColor(NewLightColor);
+	FColor NewColor(NewLightColor.ToFColor(true));
 
 	// Can't set color on a static light
 	if (AreDynamicDataChangesAllowed()
diff --git a/Engine/Source/Runtime/Engine/Private/Components/TextRenderComponent.cpp b/Engine/Source/Runtime/Engine/Private/Components/TextRenderComponent.cpp
--- a/Engine/Source/Runtime/Engine/Private/Components/TextRenderComponent.cpp
+++ b/Engine/Source/Runtime/Engine/Private/Components/TextRenderComponent.cpp
@@ -581,8 +581,6 @@ bool  FTextRenderSceneProxy::BuildStringMesh( TArray<FDynamicMeshVertex>& OutVer
 	float FirstLineHeight = -1; // Only kept around for legacy positioning support
 	float StartY = 0;
 
-	FLinearColor ActualColor = TextRenderColor;
-
 	const float CharIncrement = ( (float)Font->Kerning + HorizSpacingAdjust ) * XScale;
 
 	float LineX = 0;
@@ -650,10 +648,10 @@ bool  FTextRenderSceneProxy::BuildStringMesh( TArray<FDynamicMeshVertex>& OutVer
 				FVector TangentY(0, 0, -1);
 				FVector TangentZ(1, 0, 0);
 
-				int32 V00 = OutVertices.Add( FDynamicMeshVertex( V0, TangentX, TangentZ, FVector2D(U, V), ActualColor));
-				int32 V10 = OutVertices.Add( FDynamicMeshVertex( V1, TangentX, TangentZ, FVector2D(U + SizeU,	V),	ActualColor));
-				int32 V01 = OutVertices.Add( FDynamicMeshVertex( V2, TangentX, TangentZ, FVector2D(U,	V + SizeV), ActualColor));
-				int32 V11 = OutVertices.Add( FDynamicMeshVertex( V3, TangentX, TangentZ, FVector2D(U + SizeU,	V + SizeV), ActualColor));
+				int32 V00 = OutVertices.Add(FDynamicMeshVertex(V0, TangentX, TangentZ, FVector2D(U, V), TextRenderColor));
+				int32 V10 = OutVertices.Add(FDynamicMeshVertex(V1, TangentX, TangentZ, FVector2D(U + SizeU, V), TextRenderColor));
+				int32 V01 = OutVertices.Add(FDynamicMeshVertex(V2, TangentX, TangentZ, FVector2D(U, V + SizeV), TextRenderColor));
+				int32 V11 = OutVertices.Add(FDynamicMeshVertex(V3, TangentX, TangentZ, FVector2D(U + SizeU, V + SizeV), TextRenderColor));
 
 				check(V00 < 65536);
 				check(V10 < 65536);
diff --git a/Engine/Source/Runtime/Engine/Private/HUD.cpp b/Engine/Source/Runtime/Engine/Private/HUD.cpp
--- a/Engine/Source/Runtime/Engine/Private/HUD.cpp
+++ b/Engine/Source/Runtime/Engine/Private/HUD.cpp
@@ -500,7 +500,7 @@ void AHUD::AddDebugText_Implementation(const FString& DebugText,
 										 )
 {
 	// set a default color
-	if (TextColor == FLinearColor::Transparent)
+	if (TextColor == FColor::Transparent)
 	{
 		TextColor = FColor::White;
 	}
diff --git a/Engine/Source/Runtime/Engine/Private/ImageUtils.cpp b/Engine/Source/Runtime/Engine/Private/ImageUtils.cpp
--- a/Engine/Source/Runtime/Engine/Private/ImageUtils.cpp
+++ b/Engine/Source/Runtime/Engine/Private/ImageUtils.cpp
@@ -73,7 +73,7 @@ void FImageUtils::ImageResize(int32 SrcWidth, int32 SrcHeight, const TArray<FCol
 				LinearStepColor /= (float)PixelCount;
 
 				// Convert back from linear space to gamma space.
-				FinalColor = FColor(LinearStepColor);
+				FinalColor = LinearStepColor.ToFColor(true);
 			}
 			else
 			{
diff --git a/Engine/Source/Runtime/Engine/Private/Interpolation.cpp b/Engine/Source/Runtime/Engine/Private/Interpolation.cpp
--- a/Engine/Source/Runtime/Engine/Private/Interpolation.cpp
+++ b/Engine/Source/Runtime/Engine/Private/Interpolation.cpp
@@ -6078,7 +6078,7 @@ void UInterpTrackColorProp::UpdateTrack(float NewPosition, UInterpTrackInst* TrI
 	FLinearColor DefaultLinearColor = DefaultColor;
 	FVector DefaultColorAsVector(DefaultLinearColor.R, DefaultLinearColor.G, DefaultLinearColor.B);
 	FVector NewVectorValue = VectorTrack.Eval( NewPosition, DefaultColorAsVector );
-	FColor NewColorValue = FLinearColor(NewVectorValue.X, NewVectorValue.Y, NewVectorValue.Z);
+	FColor NewColorValue = FLinearColor(NewVectorValue.X, NewVectorValue.Y, NewVectorValue.Z).ToFColor(true);
 	*PropInst->ColorProp = NewColorValue;
 
 	// If we have a custom callback for this property, call that
diff --git a/Engine/Source/Runtime/Engine/Private/KismetMathLibrary.cpp b/Engine/Source/Runtime/Engine/Private/KismetMathLibrary.cpp
--- a/Engine/Source/Runtime/Engine/Private/KismetMathLibrary.cpp
+++ b/Engine/Source/Runtime/Engine/Private/KismetMathLibrary.cpp
@@ -374,17 +374,17 @@ float UKismetMathLibrary::GridSnap_Float(float Location, float GridSize)
 
 float UKismetMathLibrary::GetPI()
 {
-    return PI;
+	return PI;
 }
 
 float UKismetMathLibrary::DegreesToRadians(float A)
 {
-    return FMath::DegreesToRadians(A);
+	return FMath::DegreesToRadians(A);
 }
 
 float UKismetMathLibrary::RadiansToDegrees(float A)
 {
-    return FMath::RadiansToDegrees(A);
+	return FMath::RadiansToDegrees(A);
 }
 
 
@@ -1679,7 +1679,7 @@ FLinearColor UKismetMathLibrary::Conv_ColorToLinearColor(FColor InColor)
 
 FColor UKismetMathLibrary::Conv_LinearColorToColor(FLinearColor InLinearColor)
 {
-	return FColor(InLinearColor);
+	return InLinearColor.ToFColor(true);
 }
 
 FTransform UKismetMathLibrary::Conv_VectorToTransform(FVector InTranslation)
diff --git a/Engine/Source/Runtime/Engine/Private/KismetSystemLibrary.cpp b/Engine/Source/Runtime/Engine/Private/KismetSystemLibrary.cpp
--- a/Engine/Source/Runtime/Engine/Private/KismetSystemLibrary.cpp
+++ b/Engine/Source/Runtime/Engine/Private/KismetSystemLibrary.cpp
@@ -1253,7 +1253,7 @@ bool UKismetSystemLibrary::LineTraceSingle_DEPRECATED(UObject* WorldContextObjec
 		else
 		{
 			// no hit means all red
-			::DrawDebugLine(World, Start, End, FLinearColor::Red, bPersistent, LifeTime);
+			::DrawDebugLine(World, Start, End, FColor::Red, bPersistent, LifeTime);
 		}
 	}
 
@@ -1316,7 +1316,7 @@ bool UKismetSystemLibrary::LineTraceMulti_DEPRECATED(UObject* WorldContextObject
 		else
 		{
 			// no hit means all red
-			::DrawDebugLine(World, Start, End, FLinearColor::Red, bPersistent, LifeTime);
+			::DrawDebugLine(World, Start, End, FColor::Red, bPersistent, LifeTime);
 		}
 
 		// draw hits
@@ -1852,7 +1852,7 @@ bool UKismetSystemLibrary::LineTraceSingleByObject_DEPRECATED(UObject* WorldCont
 		else
 		{
 			// no hit means all red
-			::DrawDebugLine(World, Start, End, FLinearColor::Red, bPersistent, LifeTime);
+			::DrawDebugLine(World, Start, End, FColor::Red, bPersistent, LifeTime);
 		}
 	}
 
@@ -1943,7 +1943,7 @@ bool UKismetSystemLibrary::LineTraceMultiByObject_DEPRECATED(UObject* WorldConte
 		else
 		{
 			// no hit means all red
-			::DrawDebugLine(World, Start, End, FLinearColor::Red, bPersistent, LifeTime);
+			::DrawDebugLine(World, Start, End, FColor::Red, bPersistent, LifeTime);
 		}
 
 		// draw hits
@@ -2529,7 +2529,7 @@ void UKismetSystemLibrary::DrawDebugLine(UObject* WorldContextObject, FVector co
 	UWorld* World = GEngine->GetWorldFromContextObject(WorldContextObject);
 	if(World != nullptr)
 	{
-		::DrawDebugLine(World, LineStart, LineEnd, Color, false, LifeTime, SDPG_World, Thickness);
+		::DrawDebugLine(World, LineStart, LineEnd, Color.ToFColor(true), false, LifeTime, SDPG_World, Thickness);
 	}
 }
 
@@ -2539,7 +2539,7 @@ void UKismetSystemLibrary::DrawDebugPoint(UObject* WorldContextObject, FVector c
 	UWorld* World = GEngine->GetWorldFromContextObject(WorldContextObject);
 	if (World != nullptr)
 	{
-		::DrawDebugPoint(World, Position, Size, PointColor, false, LifeTime, SDPG_World);
+		::DrawDebugPoint(World, Position, Size, PointColor.ToFColor(true), false, LifeTime, SDPG_World);
 	}
 }
 
@@ -2549,7 +2549,7 @@ void UKismetSystemLibrary::DrawDebugArrow(UObject* WorldContextObject, FVector c
 	UWorld* World = GEngine->GetWorldFromContextObject(WorldContextObject);
 	if (World != nullptr)
 	{
-		::DrawDebugDirectionalArrow(World, LineStart, LineEnd, ArrowSize, Color, false, LifeTime, SDPG_World);
+		::DrawDebugDirectionalArrow(World, LineStart, LineEnd, ArrowSize, Color.ToFColor(true), false, LifeTime, SDPG_World);
 	}
 }
 
@@ -2676,7 +2676,7 @@ void UKismetSystemLibrary::DrawDebugFrustum(UObject* WorldContextObject, const F
 	if( World != nullptr && FrustumTransform.IsRotationNormalized() )
 	{
 		FMatrix FrustumToWorld =  FrustumTransform.ToMatrixWithScale();
-		::DrawDebugFrustum(World, FrustumToWorld, FrustumColor, false, Duration, SDPG_World);
+		::DrawDebugFrustum(World, FrustumToWorld, FrustumColor.ToFColor(true), false, Duration, SDPG_World);
 	}
 }
 
@@ -2687,7 +2687,7 @@ void UKismetSystemLibrary::DrawDebugCamera(const ACameraActor* CameraActor, FLin
 	{
 		FVector CamLoc = CameraActor->GetActorLocation();
 		FRotator CamRot = CameraActor->GetActorRotation();
-		::DrawDebugCamera(CameraActor->GetWorld(), CameraActor->GetActorLocation(), CameraActor->GetActorRotation(), CameraActor->GetCameraComponent()->FieldOfView, 1.0f, CameraColor, false, Duration, SDPG_World);
+		::DrawDebugCamera(CameraActor->GetWorld(), CameraActor->GetActorLocation(), CameraActor->GetActorRotation(), CameraActor->GetCameraComponent()->FieldOfView, 1.0f, CameraColor.ToFColor(true), false, Duration, SDPG_World);
 	}
 }
 
diff --git a/Engine/Source/Runtime/Engine/Private/Light.cpp b/Engine/Source/Runtime/Engine/Private/Light.cpp
--- a/Engine/Source/Runtime/Engine/Private/Light.cpp
+++ b/Engine/Source/Runtime/Engine/Private/Light.cpp
@@ -262,7 +262,7 @@ void ADirectionalLight::PostLoad()
 #if WITH_EDITORONLY_DATA
 	if(ArrowComponent != nullptr)
 	{
-		ArrowComponent->ArrowColor = GetLightColor();
+		ArrowComponent->ArrowColor = GetLightColor().ToFColor(true);
 	}
 #endif
 }
@@ -308,7 +308,7 @@ void ADirectionalLight::PostEditChangeProperty(FPropertyChangedEvent& PropertyCh
 
 	if(ArrowComponent != nullptr)
 	{
-		ArrowComponent->ArrowColor = GetLightColor();
+		ArrowComponent->ArrowColor = GetLightColor().ToFColor(true);
 	}
 }
 #endif // WITH_EDITOR
diff --git a/Engine/Source/Runtime/Engine/Private/NavMeshRenderingHelpers.h b/Engine/Source/Runtime/Engine/Private/NavMeshRenderingHelpers.h
--- a/Engine/Source/Runtime/Engine/Private/NavMeshRenderingHelpers.h
+++ b/Engine/Source/Runtime/Engine/Private/NavMeshRenderingHelpers.h
@@ -644,7 +644,7 @@ FORCEINLINE void CacheArc(TArray<FDebugRenderSceneProxy::FDebugLine>& DebugLines
 		const float u = i * ArcPtsScale;
 		const FVector Pt = EvalArc(Start, Dir, Length*Height, u);
 
-		DebugLines.Add(FDebugRenderSceneProxy::FDebugLine(Prev, Pt, Color));
+		DebugLines.Add(FDebugRenderSceneProxy::FDebugLine(Prev, Pt, Color.ToFColor(true)));
 		Prev = Pt;
 	}
 }
@@ -657,8 +657,8 @@ FORCEINLINE void CacheArrowHead(TArray<FDebugRenderSceneProxy::FDebugLine>& Debu
 	Ax = FVector::CrossProduct(Az, Ay);
 
 	FHitProxyId HitProxyId;
-	DebugLines.Add(FDebugRenderSceneProxy::FDebugLine(Tip, FVector(Tip.X + Ay.X*Size + Ax.X*Size/3, Tip.Y + Ay.Y*Size + Ax.Y*Size/3, Tip.Z + Ay.Z*Size + Ax.Z*Size/3), Color));
-	DebugLines.Add(FDebugRenderSceneProxy::FDebugLine(Tip, FVector(Tip.X + Ay.X*Size - Ax.X*Size/3, Tip.Y + Ay.Y*Size - Ax.Y*Size/3, Tip.Z + Ay.Z*Size - Ax.Z*Size/3), Color));
+	DebugLines.Add(FDebugRenderSceneProxy::FDebugLine(Tip, FVector(Tip.X + Ay.X*Size + Ax.X*Size / 3, Tip.Y + Ay.Y*Size + Ax.Y*Size / 3, Tip.Z + Ay.Z*Size + Ax.Z*Size / 3), Color.ToFColor(true)));
+	DebugLines.Add(FDebugRenderSceneProxy::FDebugLine(Tip, FVector(Tip.X + Ay.X*Size - Ax.X*Size / 3, Tip.Y + Ay.Y*Size - Ax.Y*Size / 3, Tip.Z + Ay.Z*Size - Ax.Z*Size / 3), Color.ToFColor(true)));
 }
 
 FORCEINLINE void DrawWireCylinder(TArray<FDebugRenderSceneProxy::FDebugLine>& DebugLines,const FVector& Base,const FVector& X,const FVector& Y,const FVector& Z,FColor Color,float Radius,float HalfHeight,int32 NumSides,uint8 DepthPriority, float LineThickness=0)
diff --git a/Engine/Source/Runtime/Engine/Private/Particles/ParticleComponents.cpp b/Engine/Source/Runtime/Engine/Private/Particles/ParticleComponents.cpp
--- a/Engine/Source/Runtime/Engine/Private/Particles/ParticleComponents.cpp
+++ b/Engine/Source/Runtime/Engine/Private/Particles/ParticleComponents.cpp
@@ -3578,8 +3578,8 @@ void UParticleSystemComponent::TickComponent(float DeltaTime, enum ELevelTick Ti
 			if (ShouldComputeLODFromGameThread())
 			{
 				bool bCalculateLODLevel = 
- 					(bOverrideLODMethod == true) ? (LODMethod == PARTICLESYSTEMLODMETHOD_Automatic) : 
- 					 	(Template ? (Template->LODMethod == PARTICLESYSTEMLODMETHOD_Automatic) : false);
+					(bOverrideLODMethod == true) ? (LODMethod == PARTICLESYSTEMLODMETHOD_Automatic) : 
+						(Template ? (Template->LODMethod == PARTICLESYSTEMLODMETHOD_Automatic) : false);
 				if (bCalculateLODLevel == true)
 				{
 					FVector EffectPosition = GetComponentLocation();
@@ -4041,7 +4041,7 @@ void UParticleSystemComponent::ResetParticles(bool bEmptyInstances)
 
 	UWorld* OwningWorld = GetWorld();
 
- 	const bool bIsGameWorld = OwningWorld ? OwningWorld->IsGameWorld() : !GIsEditor;
+	const bool bIsGameWorld = OwningWorld ? OwningWorld->IsGameWorld() : !GIsEditor;
 
 	// Remove instances from scene.
 	for( int32 InstanceIndex=0; InstanceIndex<EmitterInstances.Num(); InstanceIndex++ )
@@ -5159,7 +5159,7 @@ void UParticleSystemComponent::SetColorParameter(FName Name, FLinearColor Param)
 	}
 	ForceAsyncWorkCompletion(STALL);
 
-	FColor NewColor(Param);
+	FColor NewColor(Param.ToFColor(true));
 
 	// First see if an entry for this name already exists
 	for (int32 i = 0; i < InstanceParameters.Num(); i++)
diff --git a/Engine/Source/Runtime/Engine/Private/SpotLight.cpp b/Engine/Source/Runtime/Engine/Private/SpotLight.cpp
--- a/Engine/Source/Runtime/Engine/Private/SpotLight.cpp
+++ b/Engine/Source/Runtime/Engine/Private/SpotLight.cpp
@@ -31,7 +31,7 @@ ASpotLight::ASpotLight(const FObjectInitializer& ObjectInitializer)
 	ArrowComponent = CreateEditorOnlyDefaultSubobject<UArrowComponent>(TEXT("ArrowComponent0"));
 	if (ArrowComponent)
 	{
-		ArrowComponent->ArrowColor = GetLightColor();
+		ArrowComponent->ArrowColor = GetLightColor().ToFColor(true);
 		ArrowComponent->bTreatAsASprite = true;
 		ArrowComponent->SpriteInfo.Category = ConstructorStatics.ID_Lighting;
 		ArrowComponent->SpriteInfo.DisplayName = ConstructorStatics.NAME_Lighting;
@@ -54,7 +54,7 @@ void ASpotLight::PostLoad()
 #if WITH_EDITORONLY_DATA
 	if(ArrowComponent)
 	{
-		ArrowComponent->ArrowColor = GetLightColor();
+		ArrowComponent->ArrowColor = GetLightColor().ToFColor(true);
 	}
 #endif
 }
@@ -135,7 +135,7 @@ void ASpotLight::PostEditChangeProperty(FPropertyChangedEvent& PropertyChangedEv
 
 	if(ArrowComponent)
 	{
-		ArrowComponent->ArrowColor = GetLightColor();
+		ArrowComponent->ArrowColor = GetLightColor().ToFColor(true);
 	}
 }
 #endif
diff --git a/Engine/Source/Runtime/Engine/Private/StaticMeshRender.cpp b/Engine/Source/Runtime/Engine/Private/StaticMeshRender.cpp
--- a/Engine/Source/Runtime/Engine/Private/StaticMeshRender.cpp
+++ b/Engine/Source/Runtime/Engine/Private/StaticMeshRender.cpp
@@ -851,14 +851,14 @@ void FStaticMeshSceneProxy::GetDynamicMeshElements(const TArray<const FSceneView
 							Collector.RegisterOneFrameMaterialProxy(SolidMaterialInstance);
 
 							FTransform GeomTransform(GetLocalToWorld());
-							BodySetup->AggGeom.GetAggGeom(GeomTransform, WireframeColor, SolidMaterialInstance, false, true, UseEditorDepthTest(), ViewIndex, Collector);
+							BodySetup->AggGeom.GetAggGeom(GeomTransform, WireframeColor.ToFColor(true), SolidMaterialInstance, false, true, UseEditorDepthTest(), ViewIndex, Collector);
 						}
 						// wireframe
 						else
 						{
 							FColor CollisionColor = FColor(157,149,223,255);
 							FTransform GeomTransform(GetLocalToWorld());
-							BodySetup->AggGeom.GetAggGeom(GeomTransform, GetSelectionColor(CollisionColor, bProxyIsSelected, IsHovered()), NULL, (Owner == NULL), false, UseEditorDepthTest(), ViewIndex, Collector);
+							BodySetup->AggGeom.GetAggGeom(GeomTransform, GetSelectionColor(CollisionColor, bProxyIsSelected, IsHovered()).ToFColor(true), NULL, ( Owner == NULL ), false, UseEditorDepthTest(), ViewIndex, Collector);
 						}
 
 
@@ -868,7 +868,7 @@ void FStaticMeshSceneProxy::GetDynamicMeshElements(const TArray<const FSceneView
 							// Draw the static mesh's body setup (simple collision)
 							FTransform GeomTransform(GetLocalToWorld());
 							FColor NavCollisionColor = FColor(118,84,255,255);
-							StaticMesh->NavCollision->DrawSimpleGeom(Collector.GetPDI(ViewIndex), GeomTransform, GetSelectionColor(NavCollisionColor, bProxyIsSelected, IsHovered()));
+							StaticMesh->NavCollision->DrawSimpleGeom(Collector.GetPDI(ViewIndex), GeomTransform, GetSelectionColor(NavCollisionColor, bProxyIsSelected, IsHovered()).ToFColor(true));
 						}
 					}
 				}
diff --git a/Engine/Source/Runtime/Engine/Private/StatsRender2.cpp b/Engine/Source/Runtime/Engine/Private/StatsRender2.cpp
--- a/Engine/Source/Runtime/Engine/Private/StatsRender2.cpp
+++ b/Engine/Source/Runtime/Engine/Private/StatsRender2.cpp
@@ -257,7 +257,7 @@ static void RightJustify(FCanvas* Canvas, int32 X, int32 Y, TCHAR const* Text, F
  */
 static int32 RenderCycle( const FComplexStatMessage& Item, class FCanvas* Canvas, int32 X, int32 Y, const int32 Indent, const bool bStackStat )
 {
-	FColor Color = GetStatRenderGlobals().StatColor;	
+	FColor Color = GetStatRenderGlobals().StatColor.ToFColor(true);
 
 	check(Item.NameAndInfo.GetFlag(EStatMetaFlags::IsCycle));
 
diff --git a/Engine/Source/Runtime/Engine/Private/UnrealEngine.cpp b/Engine/Source/Runtime/Engine/Private/UnrealEngine.cpp
--- a/Engine/Source/Runtime/Engine/Private/UnrealEngine.cpp
+++ b/Engine/Source/Runtime/Engine/Private/UnrealEngine.cpp
@@ -5723,7 +5723,7 @@ bool UEngine::GetStatValueColoration(const FString& StatName, float Value, FColo
 						FVector BV(B.R, B.G, B.B);
 
 						FVector OutColorV = FMath::Lerp( AV, BV, Alpha );
-						OutColor = FLinearColor(OutColorV.X, OutColorV.Y, OutColorV.Z);
+						OutColor = FLinearColor(OutColorV.X, OutColorV.Y, OutColorV.Z).ToFColor(true);
 					}
 
 					return true;
diff --git a/Engine/Source/Runtime/Engine/Private/UserInterface/Canvas.cpp b/Engine/Source/Runtime/Engine/Private/UserInterface/Canvas.cpp
--- a/Engine/Source/Runtime/Engine/Private/UserInterface/Canvas.cpp
+++ b/Engine/Source/Runtime/Engine/Private/UserInterface/Canvas.cpp
@@ -887,35 +887,20 @@ void FCanvas::SetRenderTargetRect( const FIntRect& InViewRect )
 	ViewRect = InViewRect;
 }
 
-void FCanvas::Clear(const FLinearColor& Color)
+void FCanvas::Clear(const FLinearColor& LinearColor)
 {
-	// desired display gamma space
-	const float DisplayGamma =  GEngine ? GEngine->GetDisplayGamma() : 2.2f;
-	// render target gamma space expected
-	float RenderTargetGamma = DisplayGamma;
-	if( GetRenderTarget() )
-	{
-		RenderTargetGamma = GetRenderTarget()->GetDisplayGamma();
-	}
-	// assume that the clear color specified is in 2.2 gamma space
-	// so convert to the render target's color space 
-	FLinearColor GammaCorrectedColor(Color);
-	GammaCorrectedColor.R = FMath::Pow(FMath::Clamp<float>(GammaCorrectedColor.R,0.0f,1.0f), DisplayGamma / RenderTargetGamma);
-	GammaCorrectedColor.G = FMath::Pow(FMath::Clamp<float>(GammaCorrectedColor.G,0.0f,1.0f), DisplayGamma / RenderTargetGamma);
-	GammaCorrectedColor.B = FMath::Pow(FMath::Clamp<float>(GammaCorrectedColor.B,0.0f,1.0f), DisplayGamma / RenderTargetGamma);
-
 	ENQUEUE_UNIQUE_RENDER_COMMAND_TWOPARAMETER(
 		ClearCommand,
-		FColor,Color,GammaCorrectedColor,
-		FRenderTarget*,CanvasRenderTarget,GetRenderTarget(),
+		FLinearColor, ClearColor, LinearColor,
+		FRenderTarget*, CanvasRenderTarget, GetRenderTarget(),
 	{
 		SCOPED_DRAW_EVENT(RHICmdList, CanvasClear);
 		if( CanvasRenderTarget )
 		{
-			::SetRenderTarget(RHICmdList, CanvasRenderTarget->GetRenderTargetTexture(),FTextureRHIRef());
-			RHICmdList.SetViewport(0,0,0.0f,CanvasRenderTarget->GetSizeXY().X,CanvasRenderTarget->GetSizeXY().Y,1.0f);
+			::SetRenderTarget(RHICmdList, CanvasRenderTarget->GetRenderTargetTexture(), FTextureRHIRef());
+			RHICmdList.SetViewport(0, 0, 0.0f, CanvasRenderTarget->GetSizeXY().X, CanvasRenderTarget->GetSizeXY().Y, 1.0f);
 		}
-		RHICmdList.Clear(true,Color,false,0.0f,false,0, FIntRect());
+		RHICmdList.Clear(true, ClearColor, false, 0.0f, false, 0, FIntRect());
 	});
 }
 
@@ -940,9 +925,9 @@ void FCanvas::DrawTile( float X, float Y, float SizeX,	float SizeY, float U, flo
 	Right = (X + SizeX) * Z;
 	Bottom = (Y + SizeY) * Z;
 
-	int32 V00 = BatchedElements->AddVertex(FVector4(Left,	    Top,	0,Z),FVector2D(U,			V),			ActualColor,HitProxyId);
+	int32 V00 = BatchedElements->AddVertex(FVector4(Left,	  Top,	0,Z),FVector2D(U,			V),			ActualColor,HitProxyId);
 	int32 V10 = BatchedElements->AddVertex(FVector4(Right,    Top,	0,Z),FVector2D(U + SizeU,	V),			ActualColor,HitProxyId);
-	int32 V01 = BatchedElements->AddVertex(FVector4(Left,	    Bottom,	0,Z),FVector2D(U,			V + SizeV),	ActualColor,HitProxyId);
+	int32 V01 = BatchedElements->AddVertex(FVector4(Left,	  Bottom,	0,Z),FVector2D(U,			V + SizeV),	ActualColor,HitProxyId);
 	int32 V11 = BatchedElements->AddVertex(FVector4(Right,    Bottom,	0,Z),FVector2D(U + SizeU,	V + SizeV),	ActualColor,HitProxyId);
 
 	BatchedElements->AddTriangle(V00,V10,V11,FinalTexture,BlendMode);
@@ -1344,7 +1329,7 @@ void UCanvas::Update()
 /** Set DrawColor with a FLinearColor and optional opacity override */
 void UCanvas::SetLinearDrawColor(FLinearColor InColor, float OpacityOverride)
 {
-	DrawColor = FColor(InColor);
+	DrawColor = InColor.ToFColor(true);
 
 	if( OpacityOverride != -1 )
 	{
diff --git a/Engine/Source/Runtime/Engine/Private/UserInterface/CanvasItem.cpp b/Engine/Source/Runtime/Engine/Private/UserInterface/CanvasItem.cpp
--- a/Engine/Source/Runtime/Engine/Private/UserInterface/CanvasItem.cpp
+++ b/Engine/Source/Runtime/Engine/Private/UserInterface/CanvasItem.cpp
@@ -507,7 +507,7 @@ void FCanvasTileItem::RenderMaterialTile( class FCanvas* InCanvas, const FVector
 	}
 	FHitProxyId HitProxyId = InCanvas->GetHitProxyId();
 	// add the quad to the tile render batch
-	RenderBatch->AddTile( InPosition.X, InPosition.Y ,Size.X, Size.Y, UV0.X, UV0.Y, UV1.X-UV0.X, UV1.Y-UV0.Y, HitProxyId, Color);
+	RenderBatch->AddTile( InPosition.X, InPosition.Y ,Size.X, Size.Y, UV0.X, UV0.Y, UV1.X-UV0.X, UV1.Y-UV0.Y, HitProxyId, Color.ToFColor(true));
 }
 
 
diff --git a/Engine/Source/Runtime/Engine/Private/VisualLogger/VisualLoggerKismetLibrary.cpp b/Engine/Source/Runtime/Engine/Private/VisualLogger/VisualLoggerKismetLibrary.cpp
--- a/Engine/Source/Runtime/Engine/Private/VisualLogger/VisualLoggerKismetLibrary.cpp
+++ b/Engine/Source/Runtime/Engine/Private/VisualLogger/VisualLoggerKismetLibrary.cpp
@@ -22,7 +22,7 @@ void UVisualLoggerKismetLibrary::LogLocation(UObject* WorldContextObject, FVecto
 {
 #if ENABLE_VISUAL_LOG
 	const ELogVerbosity::Type DefaultVerbosity = ELogVerbosity::Log;
-	FVisualLogger::GeometryShapeLogf(WorldContextObject, FLogCategoryBase(*CategoryName.ToString(), DefaultVerbosity, DefaultVerbosity), DefaultVerbosity, INDEX_NONE, Location, Radius, Color, *Text);
+	FVisualLogger::GeometryShapeLogf(WorldContextObject, FLogCategoryBase(*CategoryName.ToString(), DefaultVerbosity, DefaultVerbosity), DefaultVerbosity, INDEX_NONE, Location, Radius, Color.ToFColor(true), *Text);
 #endif
 }
 
@@ -30,6 +30,6 @@ void UVisualLoggerKismetLibrary::LogBox(UObject* WorldContextObject, FBox Box, F
 {
 #if ENABLE_VISUAL_LOG
 	const ELogVerbosity::Type DefaultVerbosity = ELogVerbosity::Log;
-	FVisualLogger::GeometryShapeLogf(WorldContextObject, FLogCategoryBase(*CategoryName.ToString(), DefaultVerbosity, DefaultVerbosity), DefaultVerbosity, INDEX_NONE, Box, FMatrix::Identity, ObjectColor, *Text);
+	FVisualLogger::GeometryShapeLogf(WorldContextObject, FLogCategoryBase(*CategoryName.ToString(), DefaultVerbosity, DefaultVerbosity), DefaultVerbosity, INDEX_NONE, Box, FMatrix::Identity, ObjectColor.ToFColor(true), *Text);
 #endif
 }
diff --git a/Engine/Source/Runtime/Engine/Public/DebugRenderSceneProxy.h b/Engine/Source/Runtime/Engine/Public/DebugRenderSceneProxy.h
--- a/Engine/Source/Runtime/Engine/Public/DebugRenderSceneProxy.h
+++ b/Engine/Source/Runtime/Engine/Public/DebugRenderSceneProxy.h
@@ -160,7 +160,7 @@ public:
 		FSphere(const float& InRadius, const FVector& InLocation, const FLinearColor& InColor) :
 		Radius(InRadius),
 		Location(InLocation),
-		Color(InColor) {}
+		Color(InColor.ToFColor(true)) {}
 
 		float Radius;
 		FVector Location;
@@ -174,7 +174,7 @@ public:
 		FText3d(const FString& InString, const FVector& InLocation, const FLinearColor& InColor) :
 		Text(InString),
 		Location(InLocation),
-		Color(InColor) {}
+		Color(InColor.ToFColor(true)) {}
 
 		FString Text;
 		FVector Location;
@@ -188,7 +188,7 @@ public:
 			ConeToWorld(InConeToWorld),
 			Angle1(InAngle1),
 			Angle2(InAngle2),
-			Color(InColor) {}
+			Color(InColor.ToFColor(true)) {}
 
 		FMatrix ConeToWorld;
 		float Angle1;
@@ -209,7 +209,7 @@ public:
 		FCapsule(const FVector& InLocation, const float& InRadius, const FVector& x, const FVector& y, const FVector &z, const float& InHalfHeight, const FLinearColor& InColor)
 			: Radius(InRadius)
 			, Location(InLocation)
-			, Color(InColor)
+			, Color(InColor.ToFColor(true))
 			, HalfHeight(InHalfHeight)
 			, X(x)
 			, Y(y)
diff --git a/Engine/Source/Runtime/GameplayAbilities/Private/Abilities/GameplayAbilityTargetActor_Trace.cpp b/Engine/Source/Runtime/GameplayAbilities/Private/Abilities/GameplayAbilityTargetActor_Trace.cpp
--- a/Engine/Source/Runtime/GameplayAbilities/Private/Abilities/GameplayAbilityTargetActor_Trace.cpp
+++ b/Engine/Source/Runtime/GameplayAbilities/Private/Abilities/GameplayAbilityTargetActor_Trace.cpp
@@ -184,8 +184,8 @@ void AGameplayAbilityTargetActor_Trace::Tick(float DeltaSeconds)
 
 		if (bDebug)
 		{
-			DrawDebugLine(GetWorld(), SourceActor->GetActorLocation(), EndPoint, FLinearColor::Green, false);
-			DrawDebugSphere(GetWorld(), EndPoint, 16, 10, FLinearColor::Green, false);
+			DrawDebugLine(GetWorld(), SourceActor->GetActorLocation(), EndPoint, FColor::Green, false);
+			DrawDebugSphere(GetWorld(), EndPoint, 16, 10, FColor::Green, false);
 		}
 
 		SetActorLocationAndRotation(EndPoint, SourceActor->GetActorRotation());
diff --git a/Engine/Source/Runtime/Renderer/Private/ShadowRendering.cpp b/Engine/Source/Runtime/Renderer/Private/ShadowRendering.cpp
--- a/Engine/Source/Runtime/Renderer/Private/ShadowRendering.cpp
+++ b/Engine/Source/Runtime/Renderer/Private/ShadowRendering.cpp
@@ -2465,7 +2465,7 @@ void FProjectedShadowInfo::RenderFrustumWireframe(FPrimitiveDrawInterface* PDI)
 	}
 	else
 	{
-		Color = FLinearColor::FGetHSV(((SubjectPrimitiveId + LightSceneInfo->Id) * 31) & 255,0,255);
+		Color = FLinearColor::FGetHSV(( ( SubjectPrimitiveId + LightSceneInfo->Id ) * 31 ) & 255, 0, 255).ToFColor(true);
 	}
 
 	// Render the wireframe for the frustum derived from ReceiverMatrix.
diff --git a/Engine/Source/Runtime/Slate/Private/Framework/Text/SlateImageRun.cpp b/Engine/Source/Runtime/Slate/Private/Framework/Text/SlateImageRun.cpp
--- a/Engine/Source/Runtime/Slate/Private/Framework/Text/SlateImageRun.cpp
+++ b/Engine/Source/Runtime/Slate/Private/Framework/Text/SlateImageRun.cpp
@@ -135,7 +135,7 @@ int32 FSlateImageRun::OnPaint( const FPaintArgs& Args, const FTextLayout::FLineV
 
 	if ( Image->DrawAs != ESlateBrushDrawType::NoDrawType )
 	{
-		const FColor FinalColorAndOpacity( InWidgetStyle.GetColorAndOpacityTint() * Image->GetTint( InWidgetStyle ) );
+		const FLinearColor FinalColorAndOpacity( InWidgetStyle.GetColorAndOpacityTint() * Image->GetTint( InWidgetStyle ) );
 		const uint32 DrawEffects = bParentEnabled ? ESlateDrawEffect::None : ESlateDrawEffect::DisabledEffect;
 		FSlateDrawElement::MakeBox(
 			OutDrawElements, 
diff --git a/Engine/Source/Runtime/Slate/Private/Widgets/Colors/SColorBlock.cpp b/Engine/Source/Runtime/Slate/Private/Widgets/Colors/SColorBlock.cpp
--- a/Engine/Source/Runtime/Slate/Private/Widgets/Colors/SColorBlock.cpp
+++ b/Engine/Source/Runtime/Slate/Private/Widgets/Colors/SColorBlock.cpp
@@ -49,11 +49,11 @@ int32 SColorBlock::OnPaint( const FPaintArgs& Args, const FGeometry& AllottedGeo
 	{
 		FLinearColor NormalizedLinearColor = InColor / MaxRGB;
 		NormalizedLinearColor.A = InColor.A;
-		const FColor DrawNormalizedColor = InWidgetStyle.GetColorAndOpacityTint() * NormalizedLinearColor.ToFColor(bUseSRGB.Get());
+		const FLinearColor DrawNormalizedColor = InWidgetStyle.GetColorAndOpacityTint() * NormalizedLinearColor.ToFColor(bUseSRGB.Get());
 
 		FLinearColor ClampedLinearColor = InColor;
 		ClampedLinearColor.A = InColor.A * MaxRGB;
-		const FColor DrawClampedColor = InWidgetStyle.GetColorAndOpacityTint() * ClampedLinearColor.ToFColor(bUseSRGB.Get());
+		const FLinearColor DrawClampedColor = InWidgetStyle.GetColorAndOpacityTint() * ClampedLinearColor.ToFColor(bUseSRGB.Get());
 
 		TArray<FSlateGradientStop> GradientStops;
 		
diff --git a/Engine/Source/Runtime/Slate/Private/Widgets/Images/SSpinningImage.cpp b/Engine/Source/Runtime/Slate/Private/Widgets/Images/SSpinningImage.cpp
--- a/Engine/Source/Runtime/Slate/Private/Widgets/Images/SSpinningImage.cpp
+++ b/Engine/Source/Runtime/Slate/Private/Widgets/Images/SSpinningImage.cpp
@@ -23,7 +23,7 @@ int32 SSpinningImage::OnPaint( const FPaintArgs& Args, const FGeometry& Allotted
 		const bool bIsEnabled = ShouldBeEnabled(bParentEnabled);
 		const uint32 DrawEffects = bIsEnabled ? ESlateDrawEffect::None : ESlateDrawEffect::DisabledEffect;
 
-		const FColor FinalColorAndOpacity( InWidgetStyle.GetColorAndOpacityTint() * ColorAndOpacity.Get().GetColor(InWidgetStyle) * ImageBrush->GetTint( InWidgetStyle ) );
+		const FLinearColor FinalColorAndOpacity( InWidgetStyle.GetColorAndOpacityTint() * ColorAndOpacity.Get().GetColor(InWidgetStyle) * ImageBrush->GetTint( InWidgetStyle ) );
 		
 		const float Angle = SpinAnimationSequence.GetLerp() * 2.0f * PI;
 
diff --git a/Engine/Source/Runtime/Slate/Private/Widgets/Images/SThrobber.cpp b/Engine/Source/Runtime/Slate/Private/Widgets/Images/SThrobber.cpp
--- a/Engine/Source/Runtime/Slate/Private/Widgets/Images/SThrobber.cpp
+++ b/Engine/Source/Runtime/Slate/Private/Widgets/Images/SThrobber.cpp
@@ -139,7 +139,7 @@ void SCircularThrobber::ConstructSequence()
 
 int32 SCircularThrobber::OnPaint( const FPaintArgs& Args, const FGeometry& AllottedGeometry, const FSlateRect& MyClippingRect, FSlateWindowElementList& OutDrawElements, int32 LayerId, const FWidgetStyle& InWidgetStyle, bool bParentEnabled ) const
 {
-	const FColor FinalColorAndOpacity( InWidgetStyle.GetColorAndOpacityTint() * PieceImage->GetTint( InWidgetStyle ) );
+	const FLinearColor FinalColorAndOpacity( InWidgetStyle.GetColorAndOpacityTint() * PieceImage->GetTint( InWidgetStyle ) );
 	const FVector2D LocalOffset = (AllottedGeometry.Size - PieceImage->ImageSize) * 0.5f;
 	const float DeltaAngle = NumPieces > 0 ? 2 * PI / NumPieces : 0;
 	const float Phase = Curve.GetLerp() * 2 * PI;
diff --git a/Engine/Source/Runtime/Slate/Private/Widgets/Input/SEditableText.cpp b/Engine/Source/Runtime/Slate/Private/Widgets/Input/SEditableText.cpp
--- a/Engine/Source/Runtime/Slate/Private/Widgets/Input/SEditableText.cpp
+++ b/Engine/Source/Runtime/Slate/Private/Widgets/Input/SEditableText.cpp
@@ -1201,7 +1201,7 @@ int32 SEditableText::OnPaint( const FPaintArgs& Args, const FGeometry& AllottedG
 	const FSlateFontInfo& FontInfo = Font.Get();
 	const FString VisibleText = GetStringToRender();
 	const FLinearColor ThisColorAndOpacity = ColorAndOpacity.Get().GetColor(InWidgetStyle);
-	const FColor ColorAndOpacitySRGB = ThisColorAndOpacity * InWidgetStyle.GetColorAndOpacityTint();
+	const FLinearColor ColorAndOpacitySRGB = ThisColorAndOpacity * InWidgetStyle.GetColorAndOpacityTint();
 	const float FontMaxCharHeight = FTextEditHelper::GetFontHeight(FontInfo) * AllottedGeometry.Scale;
 	const double CurrentTime = FSlateApplication::Get().GetCurrentTime();
 
@@ -1211,7 +1211,7 @@ int32 SEditableText::OnPaint( const FPaintArgs& Args, const FGeometry& AllottedG
 	if( AnyTextSelected() && ( bShouldAppearFocused || bIsReadonly ) )
 	{
 		// Figure out a universally visible selection color.
-		const FColor SelectionBackgroundColorAndOpacity = ( (FLinearColor::White - ThisColorAndOpacity)*0.5f + FLinearColor(-0.2f, -0.05f, 0.15f)) * InWidgetStyle.GetColorAndOpacityTint();
+		const FLinearColor SelectionBackgroundColorAndOpacity = ( (FLinearColor::White - ThisColorAndOpacity)*0.5f + FLinearColor(-0.2f, -0.05f, 0.15f)) * InWidgetStyle.GetColorAndOpacityTint();
 
 		const float SelectionLeftX = EditableTextDefs::SelectionRectLeftOffset + FontMeasureService->Measure( VisibleText, 0, Selection.GetMinIndex(), FontInfo, false, AllottedGeometry.Scale ).X;
 		const float SelectionRightX = EditableTextDefs::SelectionRectRightOffset + FontMeasureService->Measure( VisibleText, 0, Selection.GetMaxIndex(), FontInfo, false, AllottedGeometry.Scale ).X;
diff --git a/Engine/Source/Runtime/Slate/Private/Widgets/Input/SVirtualJoystick.cpp b/Engine/Source/Runtime/Slate/Private/Widgets/Input/SVirtualJoystick.cpp
--- a/Engine/Source/Runtime/Slate/Private/Widgets/Input/SVirtualJoystick.cpp
+++ b/Engine/Source/Runtime/Slate/Private/Widgets/Input/SVirtualJoystick.cpp
@@ -108,7 +108,7 @@ int32 SVirtualJoystick::OnPaint( const FPaintArgs& Args, const FGeometry& Allott
 
 	if (bVisible)
 	{
-		FColor ColorAndOpacitySRGB = InWidgetStyle.GetColorAndOpacityTint();
+		FLinearColor ColorAndOpacitySRGB = InWidgetStyle.GetColorAndOpacityTint();
 		ColorAndOpacitySRGB.A = FMath::RoundToInt(255.f * CurrentOpacity);
 
 		for (int32 ControlIndex = 0; ControlIndex < Controls.Num(); ControlIndex++)
diff --git a/Engine/Source/Runtime/Slate/Private/Widgets/Input/SVirtualKeyboardEntry.cpp b/Engine/Source/Runtime/Slate/Private/Widgets/Input/SVirtualKeyboardEntry.cpp
--- a/Engine/Source/Runtime/Slate/Private/Widgets/Input/SVirtualKeyboardEntry.cpp
+++ b/Engine/Source/Runtime/Slate/Private/Widgets/Input/SVirtualKeyboardEntry.cpp
@@ -172,7 +172,7 @@ int32 SVirtualKeyboardEntry::OnPaint( const FPaintArgs& Args, const FGeometry& A
 	const FSlateFontInfo& FontInfo = Font.Get();
 	const FString VisibleText = GetStringToRender();
 	const FLinearColor ThisColorAndOpacity = ColorAndOpacity.Get().GetColor(InWidgetStyle);
-	const FColor ColorAndOpacitySRGB = ThisColorAndOpacity * InWidgetStyle.GetColorAndOpacityTint();
+	const FLinearColor ColorAndOpacitySRGB = ThisColorAndOpacity * InWidgetStyle.GetColorAndOpacityTint();
 	const float FontMaxCharHeight = FTextEditHelper::GetFontHeight(FontInfo);
 	const double CurrentTime = FSlateApplication::Get().GetCurrentTime();
 
diff --git a/Engine/Source/Runtime/Slate/Private/Widgets/Notifications/SProgressBar.cpp b/Engine/Source/Runtime/Slate/Private/Widgets/Notifications/SProgressBar.cpp
--- a/Engine/Source/Runtime/Slate/Private/Widgets/Notifications/SProgressBar.cpp
+++ b/Engine/Source/Runtime/Slate/Private/Widgets/Notifications/SProgressBar.cpp
@@ -92,8 +92,8 @@ int32 SProgressBar::OnPaint( const FPaintArgs& Args, const FGeometry& AllottedGe
 	
 	const FSlateBrush* CurrentFillImage = GetFillImage();
 	
-	const FColor FillColorAndOpacitySRGB( InWidgetStyle.GetColorAndOpacityTint() * FillColorAndOpacity.Get().GetColor(InWidgetStyle) * CurrentFillImage->GetTint( InWidgetStyle ) );
-	const FColor ColorAndOpacitySRGB = InWidgetStyle.GetColorAndOpacityTint();
+	const FLinearColor FillColorAndOpacitySRGB(InWidgetStyle.GetColorAndOpacityTint() * FillColorAndOpacity.Get().GetColor(InWidgetStyle) * CurrentFillImage->GetTint(InWidgetStyle));
+	const FLinearColor ColorAndOpacitySRGB = InWidgetStyle.GetColorAndOpacityTint();
 
 	TOptional<float> ProgressFraction = Percent.Get();	
 
diff --git a/Engine/Source/Runtime/Slate/Private/Widgets/Text/SMultiLineEditableText.cpp b/Engine/Source/Runtime/Slate/Private/Widgets/Text/SMultiLineEditableText.cpp
--- a/Engine/Source/Runtime/Slate/Private/Widgets/Text/SMultiLineEditableText.cpp
+++ b/Engine/Source/Runtime/Slate/Private/Widgets/Text/SMultiLineEditableText.cpp
@@ -156,7 +156,7 @@ int32 SMultiLineEditableText::FTextSelectionRunRenderer::OnPaint( const FPaintAr
 	Location.Y = Line.Offset.Y;
 
 	// If we've not been set to an explicit color, calculate a suitable one from the linked color
-	const FColor SelectionBackgroundColorAndOpacity = DefaultStyle.SelectedBackgroundColor.IsColorSpecified() 
+	const FLinearColor SelectionBackgroundColorAndOpacity = DefaultStyle.SelectedBackgroundColor.IsColorSpecified() 
 		? DefaultStyle.SelectedBackgroundColor.GetSpecifiedColor() * InWidgetStyle.GetColorAndOpacityTint()
 		: ((FLinearColor::White - DefaultStyle.SelectedBackgroundColor.GetColor(InWidgetStyle))*0.5f + FLinearColor(-0.2f, -0.05f, 0.15f)) * InWidgetStyle.GetColorAndOpacityTint();
 
diff --git a/Engine/Source/Runtime/SlateCore/Private/Rendering/DrawElements.cpp b/Engine/Source/Runtime/SlateCore/Private/Rendering/DrawElements.cpp
--- a/Engine/Source/Runtime/SlateCore/Private/Rendering/DrawElements.cpp
+++ b/Engine/Source/Runtime/SlateCore/Private/Rendering/DrawElements.cpp
@@ -111,7 +111,7 @@ void FSlateDrawElement::MakeSpline( FSlateWindowElementList& ElementList, uint32
 }
 
 
-void FSlateDrawElement::MakeDrawSpaceSpline( FSlateWindowElementList& ElementList, uint32 InLayer, const FVector2D& InStart, const FVector2D& InStartDir, const FVector2D& InEnd, const FVector2D& InEndDir, const FSlateRect InClippingRect, float InThickness, ESlateDrawEffect::Type InDrawEffects, const FColor& InTint )
+void FSlateDrawElement::MakeDrawSpaceSpline( FSlateWindowElementList& ElementList, uint32 InLayer, const FVector2D& InStart, const FVector2D& InStartDir, const FVector2D& InEnd, const FVector2D& InEndDir, const FSlateRect InClippingRect, float InThickness, ESlateDrawEffect::Type InDrawEffects, const FLinearColor& InTint )
 {
 	MakeSpline( ElementList, InLayer, FPaintGeometry(), InStart, InStartDir, InEnd, InEndDir, InClippingRect, InThickness, InDrawEffects, InTint );
 }
diff --git a/Engine/Source/Runtime/SlateCore/Public/Rendering/DrawElements.h b/Engine/Source/Runtime/SlateCore/Public/Rendering/DrawElements.h
--- a/Engine/Source/Runtime/SlateCore/Public/Rendering/DrawElements.h
+++ b/Engine/Source/Runtime/SlateCore/Public/Rendering/DrawElements.h
@@ -274,7 +274,7 @@ public:
 	SLATECORE_API static void MakeSpline( FSlateWindowElementList& ElementList, uint32 InLayer, const FPaintGeometry& PaintGeometry, const FVector2D& InStart, const FVector2D& InStartDir, const FVector2D& InEnd, const FVector2D& InEndDir, const FSlateRect InClippingRect, float InThickness = 0.0f, ESlateDrawEffect::Type InDrawEffects = ESlateDrawEffect::None, const FLinearColor& InTint=FLinearColor::White );
 
 	/** Just like MakeSpline but in draw-space coordinates. This is useful for connecting already-transformed widgets together. */
-	SLATECORE_API static void MakeDrawSpaceSpline( FSlateWindowElementList& ElementList, uint32 InLayer, const FVector2D& InStart, const FVector2D& InStartDir, const FVector2D& InEnd, const FVector2D& InEndDir, const FSlateRect InClippingRect, float InThickness = 0.0f, ESlateDrawEffect::Type InDrawEffects = ESlateDrawEffect::None, const FColor& InTint=FColor(255,255,255) );
+	SLATECORE_API static void MakeDrawSpaceSpline(FSlateWindowElementList& ElementList, uint32 InLayer, const FVector2D& InStart, const FVector2D& InStartDir, const FVector2D& InEnd, const FVector2D& InEndDir, const FSlateRect InClippingRect, float InThickness = 0.0f, ESlateDrawEffect::Type InDrawEffects = ESlateDrawEffect::None, const FLinearColor& InTint=FLinearColor::White);
 
 	/**
 	 * Creates a line defined by the provided points
