Exposes in the TransformPosition material expression the transformations from and to translated world space

diff --git a/Engine/Shaders/MaterialTemplate.usf b/Engine/Shaders/MaterialTemplate.usf
--- a/Engine/Shaders/MaterialTemplate.usf
+++ b/Engine/Shaders/MaterialTemplate.usf
@@ -409,17 +409,6 @@ MaterialFloat3 TransformTangentVectorToWorld_PreScaled(FMaterialTessellationPara
 #endif // #if FEATURE_LEVEL_SM5
 }
 
-#if HAS_PRIMITIVE_UNIFORM_BUFFER
-
-/** Transforms a vector from tangent space to local space */
-MaterialFloat3 TransformTangentVectorToLocal(FMaterialPixelParameters Parameters, MaterialFloat3 InTangentVector)
-{
-	// Transform to world space and then to local space
-	return mul(mul(InTangentVector, Parameters.TangentToWorld), (MaterialFloat3x3)Primitive.WorldToLocal);
-}
-
-#endif
-
 /** Transforms a vector from tangent space to view space */
 MaterialFloat3 TransformTangentVectorToView(FMaterialPixelParameters Parameters, MaterialFloat3 InTangentVector)
 {
@@ -427,13 +416,6 @@ MaterialFloat3 TransformTangentVectorToView(FMaterialPixelParameters Parameters,
 	return mul(mul(InTangentVector, Parameters.TangentToWorld), (MaterialFloat3x3)View.TranslatedWorldToView);
 }
 
-/** Transforms a vector from local space to tangent space */
-MaterialFloat3 TransformLocalVectorToTangent(FMaterialPixelParameters Parameters, MaterialFloat3 InLocalVector)
-{
-	// Transform from local to world space, and then to tangent space.
-	return mul(Parameters.TangentToWorld, mul(InLocalVector, GetLocalToWorld3x3()));
-}
-
 /** Transforms a vector from local space to world space (VS version) */
 MaterialFloat3 TransformLocalVectorToWorld(FMaterialVertexParameters Parameters,MaterialFloat3 InLocalVector)
 {
@@ -450,12 +432,6 @@ MaterialFloat3 TransformLocalVectorToWorld(FMaterialPixelParameters Parameters,M
 	return mul(InLocalVector, GetLocalToWorld3x3());
 }
 
-/** Transforms a vector from local space to view space */
-MaterialFloat3 TransformLocalVectorToView(MaterialFloat3 InLocalVector)
-{
-	return mul(mul(InLocalVector, GetLocalToWorld3x3()), (MaterialFloat3x3)View.TranslatedWorldToView);
-}
-
 #if HAS_PRIMITIVE_UNIFORM_BUFFER
 
 /** Transforms a vector from world space to local space */
@@ -464,28 +440,6 @@ MaterialFloat3 TransformWorldVectorToLocal(MaterialFloat3 InWorldVector)
 	return mul(InWorldVector, (MaterialFloat3x3)Primitive.WorldToLocal);
 }
 
-/** Transforms a vector from view space to local space */
-MaterialFloat3 TransformViewVectorToLocal(MaterialFloat3 InViewVector)
-{
-	return mul(mul(InViewVector, (MaterialFloat3x3)View.ViewToTranslatedWorld), (MaterialFloat3x3)Primitive.WorldToLocal);
-}
-
-#endif
-
-/** Transforms a vector from view space to world space */
-MaterialFloat3 TransformViewVectorToWorld(MaterialFloat3 InViewVector)
-{
-	return mul(InViewVector, (MaterialFloat3x3)View.ViewToTranslatedWorld);
-}
-
-/** Transforms a vector from world space to view space */
-MaterialFloat3 TransformWorldVectorToView(MaterialFloat3 InWorldVector)
-{
-	return mul(InWorldVector, (MaterialFloat3x3)View.TranslatedWorldToView);
-}
-
-#if HAS_PRIMITIVE_UNIFORM_BUFFER
-
 /** Transforms a position from local space to absolute world space */
 float3 TransformLocalPositionToWorld(FMaterialPixelParameters Parameters,float3 InLocalPosition)
 {
@@ -502,26 +456,8 @@ float3 TransformLocalPositionToWorld(FMaterialVertexParameters Parameters,float3
 	#endif
 }
 
-/** Transforms a position from absolute world space to local space */
-float3 TransformWorldPositionToLocal(float3 InWorldPosition)
-{
-	return mul(float4(InWorldPosition, 1), Primitive.WorldToLocal).xyz;
-}
-
 #endif
 
-/** Transforms a position from view space to world space */
-MaterialFloat3 TransformViewPositionToWorld(MaterialFloat3 InViewPosition)
-{
-	return mul(MaterialFloat4(InViewPosition, 1), View.ViewToTranslatedWorld).xyz - View.PreViewTranslation.xyz;
-}
-
-/** Transforms a position from world space to view space */
-MaterialFloat3 TransformWorldPositionToView(MaterialFloat3 InWorldPosition)
-{
-	return mul(MaterialFloat4(InWorldPosition + View.PreViewTranslation.xyz, 1), View.TranslatedWorldToView).xyz;
-}
-
 #if HAS_PRIMITIVE_UNIFORM_BUFFER
 
 /** Return the object's position in world space */
diff --git a/Engine/Source/Editor/UnrealEd/Private/Lightmass/LightmassRender.cpp b/Engine/Source/Editor/UnrealEd/Private/Lightmass/LightmassRender.cpp
--- a/Engine/Source/Editor/UnrealEd/Private/Lightmass/LightmassRender.cpp
+++ b/Engine/Source/Editor/UnrealEd/Private/Lightmass/LightmassRender.cpp
@@ -128,22 +128,13 @@ struct FLightmassMaterialCompiler : public FProxyMaterialCompiler
 		return Compiler->Constant3(0.0f,0.0f,-1.0f);
 	}
 
-	/**
-	 *	Generate shader code for transforming a vector
-	 */
-	virtual int32 TransformVector(uint8 SourceCoordType,uint8 DestCoordType,int32 A) override
+	virtual int32 TransformVector(EMaterialCommonBasis SourceCoordBasis, EMaterialCommonBasis DestCoordBasis, int32 A) override
 	{
 		//UE_LOG(LogLightmassRender, Log, TEXT("Lightmass material compiler has encountered TransformVector... Passing thru source vector untouched."));
 		return A;
 	}
 
-	/**
-	 *	Generate shader code for transforming a position
-	 *
-	 *	@param	CoordType - type of transform to apply. see EMaterialExpressionTransformPosition 
-	 *	@param	A - index for input vector parameter's code
-	 */
-	virtual int32 TransformPosition(uint8 SourceCoordType,uint8 DestCoordType,int32 A) override
+	virtual int32 TransformPosition(EMaterialCommonBasis SourceCoordBasis, EMaterialCommonBasis DestCoordBasis, int32 A) override
 	{
 		//UE_LOG(LogLightmassRender, Log, TEXT("Lightmass material compiler has encountered TransformPosition... Passing thru source vector untouched."));
 		return A;
diff --git a/Engine/Source/Runtime/Engine/Classes/Materials/MaterialExpressionTransform.h b/Engine/Source/Runtime/Engine/Classes/Materials/MaterialExpressionTransform.h
--- a/Engine/Source/Runtime/Engine/Classes/Materials/MaterialExpressionTransform.h
+++ b/Engine/Source/Runtime/Engine/Classes/Materials/MaterialExpressionTransform.h
@@ -8,20 +8,36 @@
 UENUM()
 enum EMaterialVectorCoordTransformSource
 {
-	TRANSFORMSOURCE_World UMETA(DisplayName="World"),
-	TRANSFORMSOURCE_Local UMETA(DisplayName="Local"),
-	TRANSFORMSOURCE_Tangent UMETA(DisplayName="Tangent"),
-	TRANSFORMSOURCE_View UMETA(DisplayName="View"),
+	/** Tangent space */
+	TRANSFORMSOURCE_Tangent UMETA(DisplayName="Tangent Space"),
+	
+	/** Local space */
+	TRANSFORMSOURCE_Local UMETA(DisplayName="Local Space"),
+	
+	/** World space */
+	TRANSFORMSOURCE_World UMETA(DisplayName="World Space"),
+	
+	/** View space */
+	TRANSFORMSOURCE_View UMETA(DisplayName="View Space"),
+
 	TRANSFORMSOURCE_MAX,
 };
 
 UENUM()
 enum EMaterialVectorCoordTransform
 {
-	TRANSFORM_World UMETA(DisplayName="World"),
-	TRANSFORM_View UMETA(DisplayName="View"),
-	TRANSFORM_Local UMETA(DisplayName="Local"),
-	TRANSFORM_Tangent UMETA(DisplayName="Tangent"),
+	/** Tangent space */
+	TRANSFORM_Tangent UMETA(DisplayName="Tangent Space"),
+	
+	/** Local space */
+	TRANSFORM_Local UMETA(DisplayName="Local Space"),
+	
+	/** World space */
+	TRANSFORM_World UMETA(DisplayName="World Space"),
+	
+	/** View space */
+	TRANSFORM_View UMETA(DisplayName="View Space"),
+
 	TRANSFORM_MAX,
 };
 
diff --git a/Engine/Source/Runtime/Engine/Classes/Materials/MaterialExpressionTransformPosition.h b/Engine/Source/Runtime/Engine/Classes/Materials/MaterialExpressionTransformPosition.h
--- a/Engine/Source/Runtime/Engine/Classes/Materials/MaterialExpressionTransformPosition.h
+++ b/Engine/Source/Runtime/Engine/Classes/Materials/MaterialExpressionTransformPosition.h
@@ -8,9 +8,18 @@
 UENUM()
 enum EMaterialPositionTransformSource
 {
-	TRANSFORMPOSSOURCE_Local UMETA(DisplayName="Local"),
-	TRANSFORMPOSSOURCE_World UMETA(DisplayName="World"),
-	TRANSFORMPOSSOURCE_View  UMETA(DisplayName="View"),
+	/** Local space */
+	TRANSFORMPOSSOURCE_Local UMETA(DisplayName="Local Space"),
+	
+	/** World space with no shader offsets applied */
+	TRANSFORMPOSSOURCE_World UMETA(DisplayName="World Space (No Shader Offsets)"),
+	
+	/** World space with all shader offsets applied */
+	TRANSFORMPOSSOURCE_TranslatedWorld  UMETA(DisplayName="World Space"),
+
+	/** View space */
+	TRANSFORMPOSSOURCE_View  UMETA(DisplayName="View Space"),
+	
 	TRANSFORMPOSSOURCE_MAX,
 };
 
diff --git a/Engine/Source/Runtime/Engine/Private/Materials/HLSLMaterialTranslator.h b/Engine/Source/Runtime/Engine/Private/Materials/HLSLMaterialTranslator.h
--- a/Engine/Source/Runtime/Engine/Private/Materials/HLSLMaterialTranslator.h
+++ b/Engine/Source/Runtime/Engine/Private/Materials/HLSLMaterialTranslator.h
@@ -3163,279 +3163,196 @@ protected:
 		}
 	}
 
-	/**
-	* Generate shader code for transforming a vector
-	*/
-	virtual int32 TransformVector(uint8 SourceCoordType,uint8 DestCoordType,int32 A) override
+	int32 TransformBase(EMaterialCommonBasis SourceCoordBasis, EMaterialCommonBasis DestCoordBasis, int32 A, int AWComponent)
 	{
-		if (ShaderFrequency != SF_Pixel && ShaderFrequency != SF_Compute && ShaderFrequency != SF_Domain && ShaderFrequency != SF_Vertex)
-		{
-			return NonPixelShaderExpressionError();
-		}
-
-		const EMaterialVectorCoordTransformSource SourceCoordinateSpace = (EMaterialVectorCoordTransformSource)SourceCoordType;
-		const EMaterialVectorCoordTransform DestinationCoordinateSpace = (EMaterialVectorCoordTransform)DestCoordType;
-
-		// Construct float3(0,0,x) out of the input if it is a scalar
-		// This way artists can plug in a scalar and it will be treated as height, or a vector displacement
-		if(A != INDEX_NONE && (GetType(A) & MCT_Float1) && SourceCoordinateSpace == TRANSFORMSOURCE_Tangent)
-		{
-			A = AppendVector(Constant2(0, 0), A);
-		}
-
-		int32 Domain = Material->GetMaterialDomain();
-		if( (Domain != MD_Surface && Domain != MD_DeferredDecal) && 
-			(SourceCoordinateSpace == TRANSFORMSOURCE_Tangent || SourceCoordinateSpace == TRANSFORMSOURCE_Local || DestCoordType == TRANSFORM_Tangent || DestCoordType == TRANSFORM_Local))
+		if (A == INDEX_NONE)
 		{
-			return Errorf(TEXT("Local and tangent transforms are only supported in the Surface and Deferred Decal material domains!"));
+			// unable to compile
+			return INDEX_NONE;
 		}
+		
+		{ // validation
+			if (ShaderFrequency != SF_Pixel && ShaderFrequency != SF_Compute && ShaderFrequency != SF_Domain && ShaderFrequency != SF_Vertex)
+			{
+				return NonPixelShaderExpressionError();
+			}
 
-		if (ShaderFrequency != SF_Vertex)
-		{
-			bUsesTransformVector = true;
-		}
+			if (ShaderFrequency != SF_Pixel && ShaderFrequency != SF_Compute && ShaderFrequency != SF_Vertex)
+			{
+				if ((SourceCoordBasis == MCB_Local || DestCoordBasis == MCB_Local))
+				{
+					return Errorf(TEXT("Local space in only supported for vertex, compute or pixel shader"));
+				}
+			}
 
-		int32 Result = INDEX_NONE;
-		if(A != INDEX_NONE)
-		{
-			int32 NumInputComponents = GetNumComponents(GetParameterType(A));
-			// only allow float3/float4 transforms
-			if( NumInputComponents < 3 )
+			if (AWComponent != 0 && (SourceCoordBasis == MCB_Tangent || DestCoordBasis == MCB_Tangent))
 			{
-				Result = Errorf(TEXT("input must be a vector (%s: %s) or a scalar (if source is Tangent)"),*GetParameterCode(A),DescribeType(GetParameterType(A)));
+				return Errorf(TEXT("Tangent basis not available for position transformations"));
 			}
-			else if ((SourceCoordinateSpace == TRANSFORMSOURCE_World && DestinationCoordinateSpace == TRANSFORM_World)
-				|| (SourceCoordinateSpace == TRANSFORMSOURCE_Local && DestinationCoordinateSpace == TRANSFORM_Local)
-				|| (SourceCoordinateSpace == TRANSFORMSOURCE_Tangent && DestinationCoordinateSpace == TRANSFORM_Tangent)
-				|| (SourceCoordinateSpace == TRANSFORMSOURCE_View && DestinationCoordinateSpace == TRANSFORM_View))
+		
+			// Construct float3(0,0,x) out of the input if it is a scalar
+			// This way artists can plug in a scalar and it will be treated as height, or a vector displacement
+			if (GetType(A) == MCT_Float1 && SourceCoordBasis == MCB_Tangent)
 			{
-				// Pass through
-				Result = A;
+				A = AppendVector(Constant2(0, 0), A);
 			}
-			else 
+			else if (GetNumComponents(GetParameterType(A)) < 3)
 			{
-				// code string to transform the input vector
-				FString CodeStr;
-				if (SourceCoordinateSpace == TRANSFORMSOURCE_Tangent)
-				{
-					switch( DestinationCoordinateSpace )
-					{
-					case TRANSFORM_Local:
-						// transform from tangent to local space
-						if( ShaderFrequency != SF_Pixel && ShaderFrequency != SF_Compute && ShaderFrequency != SF_Vertex )
-						{
-							return Errorf(TEXT("Local space in only supported for vertex or pixel shader!"));
-						}
-						CodeStr = FString(TEXT("TransformTangentVectorToLocal(Parameters,%s)"));
-						break;
-					case TRANSFORM_World:
-						// transform from tangent to world space
-						if(ShaderFrequency == SF_Domain)
-						{
-							// domain shader uses a prescale value to preserve scaling factor on WorldTransform	when sampling a displacement map
-							CodeStr = FString(TEXT("TransformTangentVectorToWorld_PreScaled(Parameters,%s)"));
-						}
-						else
-						{
-							CodeStr = FString(TEXT("TransformTangentVectorToWorld(Parameters.TangentToWorld,%s)"));
-						}
-						break;
-					case TRANSFORM_View:
-						// transform from tangent to view space
-						if( ShaderFrequency != SF_Pixel && ShaderFrequency != SF_Compute && ShaderFrequency != SF_Vertex )
-						{
-							return Errorf(TEXT("View space in only supported for vertex or pixel shader!"));
-						}
-						CodeStr = FString(TEXT("TransformTangentVectorToView(Parameters,%s)"));
-						break;
-					default:
-						UE_LOG(LogMaterial, Fatal, TEXT("Invalid DestCoordType. See EMaterialVectorCoordTransform") );
-					}
-				}
-				else if (SourceCoordinateSpace == TRANSFORMSOURCE_Local)
-				{
-					if( ShaderFrequency != SF_Pixel && ShaderFrequency != SF_Compute && ShaderFrequency != SF_Vertex )
-					{
-						return Errorf(TEXT("Local space in only supported for vertex or pixel shader!"));
-					}
-
-					switch( DestinationCoordinateSpace )
-					{
-					case TRANSFORM_Tangent:
-						CodeStr = FString(TEXT("TransformLocalVectorToTangent(Parameters,%s)"));
-						break;
-					case TRANSFORM_World:
-						CodeStr = FString(TEXT("TransformLocalVectorToWorld(Parameters,%s)"));
-						break;
-					case TRANSFORM_View:
-						CodeStr = FString(TEXT("TransformLocalVectorToView(%s)"));
-						break;
-					default:
-						UE_LOG(LogMaterial, Fatal, TEXT("Invalid DestCoordType. See EMaterialVectorCoordTransform") );
-					}
-				}
-				else if (SourceCoordinateSpace == TRANSFORMSOURCE_World)
-				{
-					switch( DestinationCoordinateSpace )
-					{
-					case TRANSFORM_Tangent:
-						CodeStr = FString(TEXT("TransformWorldVectorToTangent(Parameters.TangentToWorld,%s)"));
-						break;
-					case TRANSFORM_Local:
-						if( ShaderFrequency != SF_Pixel && ShaderFrequency != SF_Compute && ShaderFrequency != SF_Vertex )
-						{
-							return Errorf(TEXT("Local space in only supported for vertex or pixel shader!"));
-						}
-						CodeStr = FString(TEXT("TransformWorldVectorToLocal(%s)"));
-						break;
-					case TRANSFORM_View:
-						if( ShaderFrequency != SF_Pixel && ShaderFrequency != SF_Compute && ShaderFrequency != SF_Vertex )
-						{
-							return Errorf(TEXT("View space in only supported for vertex or pixel shader!"));
-						}
-						CodeStr = FString(TEXT("TransformWorldVectorToView(%s)"));
-						break;
-					default:
-						UE_LOG(LogMaterial, Fatal, TEXT("Invalid DestCoordType. See EMaterialVectorCoordTransform") );
-					}
-				}
-				else if (SourceCoordinateSpace == TRANSFORMSOURCE_View)
-				{
-					switch( DestinationCoordinateSpace )
-					{
-					case TRANSFORM_Tangent:
-						CodeStr = FString(TEXT("TransformWorldVectorToTangent(Parameters.TangentToWorld,TransformViewVectorToWorld(%s))"));
-						break;
-					case TRANSFORM_Local:
-						if( ShaderFrequency != SF_Pixel && ShaderFrequency != SF_Compute && ShaderFrequency != SF_Vertex )
-						{
-							return Errorf(TEXT("Local space in only supported for vertex or pixel shader!"));
-						}
-						CodeStr = FString(TEXT("TransformViewVectorToLocal(%s)"));
-						break;
-					case TRANSFORM_World:
-						if( ShaderFrequency != SF_Pixel && ShaderFrequency != SF_Compute && ShaderFrequency != SF_Vertex )
-						{
-							return Errorf(TEXT("View space in only supported for vertex or pixel shader!"));
-						}
-						CodeStr = FString(TEXT("TransformViewVectorToWorld(%s)"));
-						break;
-					default:
-						UE_LOG(LogMaterial, Fatal, TEXT("Invalid DestCoordType. See EMaterialVectorCoordTransform") );
-					}
-				}
-				else
-				{
-					UE_LOG(LogMaterial, Fatal, TEXT("Invalid SourceCoordType. See EMaterialVectorCoordTransformSource") );
-				}
-
-				// we are only transforming vectors (not points) so only return a float3
-				Result = AddCodeChunk(
-					MCT_Float3,
-					*CodeStr,
-					*CoerceParameter(A,MCT_Float3)
-					);
+				return Errorf(TEXT("input must be a vector (%s: %s) or a scalar (if source is Tangent)"), *GetParameterCode(A), DescribeType(GetParameterType(A)));
 			}
 		}
-		return Result; 
-	}
-
-	/**
-	* Generate shader code for transforming a position
-	*
-	* @param	CoordType - type of transform to apply. see EMaterialExpressionTransformPosition 
-	* @param	A - index for input vector parameter's code
-	*/
-	virtual int32 TransformPosition(uint8 SourceCoordType, uint8 DestCoordType, int32 A) override
-	{
-		const EMaterialPositionTransformSource SourceCoordinateSpace = (EMaterialPositionTransformSource)SourceCoordType;
-		const EMaterialPositionTransformSource DestinationCoordinateSpace = (EMaterialPositionTransformSource)DestCoordType;
-
-		int32 Result = INDEX_NONE;
-
-		if (SourceCoordinateSpace == DestinationCoordinateSpace)
+		
+		if (SourceCoordBasis == DestCoordBasis)
 		{
-			Result = A;
+			// no transformation needed
+			return A;
 		}
-		else if (A != INDEX_NONE)
-		{
-			// code string to transform the input vector
-			FString CodeStr;
+		
+		FString CodeStr;
+		EMaterialCommonBasis IntermediaryBasis = MCB_World;
 
-			if (SourceCoordinateSpace == TRANSFORMPOSSOURCE_Local ||
-				DestinationCoordinateSpace == TRANSFORMPOSSOURCE_Local)
+		switch (SourceCoordBasis)
+		{
+			case MCB_Tangent:
 			{
-				if (ShaderFrequency != SF_Pixel && ShaderFrequency != SF_Compute && ShaderFrequency != SF_Vertex)
+				check(AWComponent == 0);
+				if (DestCoordBasis == MCB_World)
 				{
-					return Errorf(TEXT("Local space is only supported for vertex or pixel shader!"));
+					if (ShaderFrequency == SF_Domain)
+					{
+						// domain shader uses a prescale value to preserve scaling factor on WorldTransform	when sampling a displacement map
+						CodeStr = FString(TEXT("TransformTangent<TO>World_PreScaled(Parameters, <A>.xyz)"));
+					}
+					else
+					{
+						CodeStr = TEXT("mul(<A>, <MATRIX>(Parameters.TangentToWorld))");
+					}
 				}
+				// else use MCB_World as intermediary basis
+				break;
 			}
-			if (SourceCoordinateSpace == TRANSFORMPOSSOURCE_View ||
-				DestinationCoordinateSpace == TRANSFORMPOSSOURCE_View)
+			case MCB_Local:
 			{
-				if (ShaderFrequency != SF_Pixel && ShaderFrequency != SF_Compute && ShaderFrequency != SF_Vertex)
+				if (DestCoordBasis == MCB_World)
 				{
-					return Errorf(TEXT("View space in only supported for vertex or pixel shader!"));
+					 // TODO: need <PREV>
+					CodeStr = TEXT("TransformLocal<TO>World(Parameters, <A>.xyz)");
 				}
+				// else use MCB_World as intermediary basis
+				break;
 			}
-
-			if (SourceCoordinateSpace == TRANSFORMPOSSOURCE_Local)
+			case MCB_TranslatedWorld:
 			{
-				if (DestinationCoordinateSpace == TRANSFORMPOSSOURCE_World)
+				if (DestCoordBasis == MCB_World)
 				{
-					CodeStr = FString(TEXT("TransformLocalPositionToWorld(Parameters,%s)"));
-				}
-				else if (DestinationCoordinateSpace == TRANSFORMPOSSOURCE_View)
-				{
-					CodeStr = FString(TEXT("TransformWorldPositionToView(TransformLocalPositionToWorld(Parameters,%s))"));
+					if (AWComponent)
+					{
+						CodeStr = TEXT("(<A>.xyz - View.<PREV>PreViewTranslation.xyz)");
+					}
+					else
+					{
+						CodeStr = TEXT("<A>");
+					}
 				}
-				else
+				else if (DestCoordBasis == MCB_View)
 				{
-					UE_LOG(LogMaterial, Fatal, TEXT("Invalid DestCoordType. See EMaterialPositionTransformSource"));
+					CodeStr = TEXT("mul(<A>, <MATRIX>(View.<PREV>TranslatedWorldToView))");
 				}
+				// else use MCB_World as intermediary basis
+				break;
 			}
-			else if (SourceCoordinateSpace == TRANSFORMPOSSOURCE_World)
+			case MCB_World:
 			{
-				if (DestinationCoordinateSpace == TRANSFORMPOSSOURCE_Local)
+				if (DestCoordBasis == MCB_Tangent)
 				{
-					CodeStr = FString(TEXT("TransformWorldPositionToLocal(%s)"));
+					CodeStr = TEXT("mul(<MATRIX>(Parameters.TangentToWorld), <A>)");
 				}
-				else if (DestinationCoordinateSpace == TRANSFORMPOSSOURCE_View)
+				else if (DestCoordBasis == MCB_Local)
 				{
-					CodeStr = FString(TEXT("TransformWorldPositionToView(%s)"));
+					//TODO: need Primitive.PrevWorldToLocal
+					//TODO: inconsistent with TransformLocal<TO>World with instancing
+					CodeStr = TEXT("mul(<A>, <MATRIX>(Primitive.WorldToLocal))");
 				}
-				else
+				else if (DestCoordBasis == MCB_TranslatedWorld)
 				{
-					UE_LOG(LogMaterial, Fatal, TEXT("Invalid DestCoordType. See EMaterialPositionTransformSource"));
+					if (AWComponent)
+					{
+						CodeStr = TEXT("(<A>.xyz + View.<PREV>PreViewTranslation.xyz)");
+					}
+					else
+					{
+						CodeStr = TEXT("<A>");
+					}
 				}
+				// else use MCB_TranslatedWorld as intermediary basis
+				IntermediaryBasis = MCB_TranslatedWorld;
+				break;
 			}
-			else if (SourceCoordinateSpace == TRANSFORMPOSSOURCE_View)
+			case MCB_View:
 			{
-				if (DestinationCoordinateSpace == TRANSFORMPOSSOURCE_Local)
-				{
-					CodeStr = FString(TEXT("TransformWorldPositionToLocal(TransformViewPositionToWorld(%s))"));
-				}
-				else if (DestinationCoordinateSpace == TRANSFORMPOSSOURCE_World)
-				{
-					CodeStr = FString(TEXT("TransformViewPositionToWorld(%s)"));
-				}
-				else
+				if (DestCoordBasis == MCB_TranslatedWorld)
 				{
-					UE_LOG(LogMaterial, Fatal, TEXT("Invalid DestCoordType. See EMaterialPositionTransformSource"));
+					CodeStr = TEXT("mul(<A>, <MATRIX>(View.<PREV>ViewToTranslatedWorld))");
 				}
+				// else use MCB_TranslatedWorld as intermediary basis
+				IntermediaryBasis = MCB_TranslatedWorld;
+				break;
 			}
-			else
-			{
-				UE_LOG(LogMaterial, Fatal, TEXT("Invalid SourceCoordType. See EMaterialPositionTransformSource"));
-			}
+			default:
+				check(0);
+				break;
+		}
 
-			Result = AddCodeChunk(
-				MCT_Float3,
-				*CodeStr,
-				*CoerceParameter(A,MCT_Float3)
-				);
+		if (CodeStr.IsEmpty())
+		{
+			// check intermediary basis so we don't have infinite recursion
+			check(IntermediaryBasis != SourceCoordBasis);
+			check(IntermediaryBasis != DestCoordBasis);
+
+			// use intermediary basis
+			const int32 IntermediaryA = TransformBase(SourceCoordBasis, IntermediaryBasis, A, AWComponent);
+
+			return TransformBase(IntermediaryBasis, DestCoordBasis, IntermediaryA, AWComponent);
+		}
+		
+		if (AWComponent != 0)
+		{
+			A = AppendVector(A, Constant(1));
+			CodeStr.ReplaceInline(TEXT("<TO>"),TEXT("PositionTo"));
+			CodeStr.ReplaceInline(TEXT("<MATRIX>"),TEXT(""));
+			CodeStr += ".xyz";
+		}
+		else
+		{
+			CodeStr.ReplaceInline(TEXT("<TO>"),TEXT("VectorTo"));
+			CodeStr.ReplaceInline(TEXT("<MATRIX>"),TEXT("(MaterialFloat3x3)"));
 		}
-		return Result; 
+		
+		if (bCompilingPreviousFrame)
+		{
+			CodeStr.ReplaceInline(TEXT("<PREV>"),TEXT("Prev"));
+		}
+		else
+		{
+			CodeStr.ReplaceInline(TEXT("<PREV>"),TEXT(""));
+		}
+		
+		CodeStr.ReplaceInline(TEXT("<A>"), *GetParameterCode(A));
+
+		return AddCodeChunk(
+			MCT_Float3,
+			*CodeStr
+			);
+	}
+	
+	virtual int32 TransformVector(EMaterialCommonBasis SourceCoordBasis, EMaterialCommonBasis DestCoordBasis, int32 A) override
+	{
+		return TransformBase(SourceCoordBasis, DestCoordBasis, A, 0);
+	}
+
+	virtual int32 TransformPosition(EMaterialCommonBasis SourceCoordBasis, EMaterialCommonBasis DestCoordBasis, int32 A) override
+	{
+		return TransformBase(SourceCoordBasis, DestCoordBasis, A, 1);
 	}
 
 	virtual int32 DynamicParameter() override
diff --git a/Engine/Source/Runtime/Engine/Private/Materials/MaterialExpressions.cpp b/Engine/Source/Runtime/Engine/Private/Materials/MaterialExpressions.cpp
--- a/Engine/Source/Runtime/Engine/Private/Materials/MaterialExpressions.cpp
+++ b/Engine/Source/Runtime/Engine/Private/Materials/MaterialExpressions.cpp
@@ -3121,7 +3121,7 @@ int32 UMaterialExpressionBumpOffset::Compile(class FMaterialCompiler* Compiler,
 
 	return Compiler->Add(
 			Compiler->Mul(
-				Compiler->ComponentMask(Compiler->TransformVector(TRANSFORMSOURCE_World, TRANSFORM_Tangent, Compiler->CameraVector()),1,1,0,0),
+				Compiler->ComponentMask(Compiler->TransformVector(MCB_World, MCB_Tangent, Compiler->CameraVector()),1,1,0,0),
 				Compiler->Add(
 					Compiler->Mul(
 						HeightRatioInput.Expression ? Compiler->ForceCast(HeightRatioInput.Compile(Compiler),MCT_Float1) : Compiler->Constant(HeightRatio),
@@ -5256,6 +5256,27 @@ void UMaterialExpressionAbs::GetCaption(TArray<FString>& OutCaptions) const
 // UMaterialExpressionTransform
 ///////////////////////////////////////////////////////////////////////////////
 
+static EMaterialCommonBasis GetMaterialCommonBasis(EMaterialVectorCoordTransformSource X)
+{
+	static const EMaterialCommonBasis ConversionTable[TRANSFORMSOURCE_MAX] = {
+		MCB_Tangent,
+		MCB_Local,
+		MCB_World,
+		MCB_View,
+	};
+	return ConversionTable[X];
+}
+
+static EMaterialCommonBasis GetMaterialCommonBasis(EMaterialVectorCoordTransform X)
+{
+	static const EMaterialCommonBasis ConversionTable[TRANSFORM_MAX] = {
+		MCB_Tangent,
+		MCB_Local,
+		MCB_World,
+		MCB_View,
+	};
+	return ConversionTable[X];
+}
 
 int32 UMaterialExpressionTransform::Compile(class FMaterialCompiler* Compiler, int32 OutputIndex, int32 MultiplexIndex)
 {
@@ -5268,52 +5289,33 @@ int32 UMaterialExpressionTransform::Compile(class FMaterialCompiler* Compiler, i
 	else
 	{
 		int32 VecInputIdx = Input.Compile(Compiler);
-		Result = Compiler->TransformVector(TransformSourceType, TransformType, VecInputIdx);
+		const auto TransformSourceBasis = GetMaterialCommonBasis(TransformSourceType);
+		const auto TransformDestBasis = GetMaterialCommonBasis(TransformType);
+		Result = Compiler->TransformVector(TransformSourceBasis, TransformDestBasis, VecInputIdx);
 	}
 
 	return Result;
 }
 
-const TCHAR* GetTransformSourceEnumText(EMaterialVectorCoordTransformSource Source)
-{
-	switch (Source)
-	{
-		case TRANSFORMSOURCE_World: return TEXT("World");
-		case TRANSFORMSOURCE_Local: return TEXT("Local");
-		case TRANSFORMSOURCE_View: return TEXT("View");
-		case TRANSFORMSOURCE_Tangent: return TEXT("Tangent");
-	}
-	return TEXT("");
-}
-
-const TCHAR* GetTransformTypeEnumText(EMaterialVectorCoordTransform Type)
-{
-	switch (Type)
-	{
-		case TRANSFORM_World: return TEXT("World");
-		case TRANSFORM_Local: return TEXT("Local");
-		case TRANSFORM_View: return TEXT("View");
-		case TRANSFORM_Tangent: return TEXT("Tangent");
-	}
-	return TEXT("");
-}
-
-
 void UMaterialExpressionTransform::GetCaption(TArray<FString>& OutCaptions) const
 {
-	OutCaptions.Add(TEXT("Vector Transform"));
-
-	FString TransformDesc = TEXT("(");
-	TransformDesc += GetTransformSourceEnumText(TransformSourceType);
+#if WITH_EDITOR
+	const UEnum* MVCTSEnum = FindObject<UEnum>(NULL, TEXT("Engine.EMaterialVectorCoordTransformSource"));
+	const UEnum* MVCTEnum = FindObject<UEnum>(NULL, TEXT("Engine.EMaterialVectorCoordTransform"));
+	check(MVCTSEnum);
+	check(MVCTEnum);
+	
+	FString TransformDesc;
+	TransformDesc += MVCTSEnum->GetDisplayNameText(TransformSourceType).ToString();
 	TransformDesc += TEXT(" to ");
-	TransformDesc += GetTransformTypeEnumText(TransformType);
-	TransformDesc += TEXT(")");
+	TransformDesc += MVCTEnum->GetDisplayNameText(TransformType).ToString();
 	OutCaptions.Add(TransformDesc);
-}
+#else
+	OutCaptions.Add(TEXT(""));
+#endif
 
-///////////////////////////////////////////////////////////////////////////////
-// UMaterialExpressionTransform
-///////////////////////////////////////////////////////////////////////////////
+	OutCaptions.Add(TEXT("TransformVector"));
+}
 
 UMaterialExpressionTransform::UMaterialExpressionTransform(const FObjectInitializer& ObjectInitializer)
 	: Super(ObjectInitializer)
@@ -5331,6 +5333,7 @@ UMaterialExpressionTransform::UMaterialExpressionTransform(const FObjectInitiali
 
 	MenuCategories.Add(ConstructorStatics.NAME_VectorOps);
 	TransformSourceType = TRANSFORMSOURCE_Tangent;
+	TransformType = TRANSFORM_Local;
 }
 
 
@@ -5353,13 +5356,24 @@ UMaterialExpressionTransformPosition::UMaterialExpressionTransformPosition(const
 
 	MenuCategories.Add(ConstructorStatics.NAME_VectorOps);
 	TransformSourceType = TRANSFORMPOSSOURCE_Local;
+	TransformType = TRANSFORMPOSSOURCE_World;
+}
+
+static EMaterialCommonBasis GetMaterialCommonBasis(EMaterialPositionTransformSource X)
+{
+	static const EMaterialCommonBasis ConversionTable[TRANSFORMPOSSOURCE_MAX] = {
+		MCB_Local,
+		MCB_World,
+		MCB_TranslatedWorld,
+		MCB_View,
+	};
+	return ConversionTable[X];
 }
 
-	
 int32 UMaterialExpressionTransformPosition::Compile(class FMaterialCompiler* Compiler, int32 OutputIndex, int32 MultiplexIndex)
 {
 	int32 Result=INDEX_NONE;
-
+	
 	if( !Input.Expression )
 	{
 		Result = Compiler->Errorf(TEXT("Missing Transform Position input vector"));
@@ -5367,34 +5381,30 @@ int32 UMaterialExpressionTransformPosition::Compile(class FMaterialCompiler* Com
 	else
 	{
 		int32 VecInputIdx = Input.Compile(Compiler);
-		Result = Compiler->TransformPosition( TransformSourceType, TransformType, VecInputIdx );
+		const auto TransformSourceBasis = GetMaterialCommonBasis(TransformSourceType);
+		const auto TransformDestBasis = GetMaterialCommonBasis(TransformType);
+		Result = Compiler->TransformPosition(TransformSourceBasis, TransformDestBasis, VecInputIdx);
 	}
 
 	return Result;
 }
 
-const TCHAR* GetTransformPositionTypeEnumText(EMaterialPositionTransformSource Transform)
-{
-	switch (Transform)
-	{
-		case TRANSFORMPOSSOURCE_Local: return TEXT("Local");
-		case TRANSFORMPOSSOURCE_World: return TEXT("World");
-	}
-	return TEXT("");
-}
-
-
 void UMaterialExpressionTransformPosition::GetCaption(TArray<FString>& OutCaptions) const
 {
-	OutCaptions.Add(TEXT("Position Transform"));
-
+#if WITH_EDITOR
+	const UEnum* MPTSEnum = FindObject<UEnum>(NULL, TEXT("Engine.EMaterialPositionTransformSource"));
+	check(MPTSEnum);
+	
 	FString TransformDesc;
-	TransformDesc += TEXT("(");
-	TransformDesc += GetTransformPositionTypeEnumText(TransformSourceType);
+	TransformDesc += MPTSEnum->GetDisplayNameText(TransformSourceType).ToString();
 	TransformDesc += TEXT(" to ");
-	TransformDesc += GetTransformPositionTypeEnumText(TransformType);
-	TransformDesc += TEXT(")");
+	TransformDesc += MPTSEnum->GetDisplayNameText(TransformType).ToString();
 	OutCaptions.Add(TransformDesc);
+#else
+	OutCaptions.Add(TEXT(""));
+#endif
+	
+	OutCaptions.Add(TEXT("TransformPosition"));
 }
 
 ///////////////////////////////////////////////////////////////////////////////
diff --git a/Engine/Source/Runtime/Engine/Public/MaterialCompiler.h b/Engine/Source/Runtime/Engine/Public/MaterialCompiler.h
--- a/Engine/Source/Runtime/Engine/Public/MaterialCompiler.h
+++ b/Engine/Source/Runtime/Engine/Public/MaterialCompiler.h
@@ -148,8 +148,8 @@ public:
 
 	virtual int32 ComponentMask(int32 Vector,bool R,bool G,bool B,bool A) = 0;
 	virtual int32 AppendVector(int32 A,int32 B) = 0;
-	virtual int32 TransformVector(uint8 SourceCoordType,uint8 DestCoordType,int32 A) = 0;
-	virtual int32 TransformPosition(uint8 SourceCoordType,uint8 DestCoordType,int32 A) = 0;
+	virtual int32 TransformVector(EMaterialCommonBasis SourceCoordBasis, EMaterialCommonBasis DestCoordBasis, int32 A) = 0;
+	virtual int32 TransformPosition(EMaterialCommonBasis SourceCoordBasis, EMaterialCommonBasis DestCoordBasis, int32 A) = 0;
 
 	virtual int32 DynamicParameter() = 0;
 	virtual int32 LightmapUVs() = 0;
@@ -306,8 +306,14 @@ public:
 
 	virtual int32 ComponentMask(int32 Vector,bool R,bool G,bool B,bool A) override { return Compiler->ComponentMask(Vector,R,G,B,A); }
 	virtual int32 AppendVector(int32 A,int32 B) override { return Compiler->AppendVector(A,B); }
-	virtual int32 TransformVector(uint8 SourceCoordType,uint8 DestCoordType,int32 A) override { return Compiler->TransformVector(SourceCoordType,DestCoordType,A); }
-	virtual int32 TransformPosition(uint8 SourceCoordType,uint8 DestCoordType,int32 A) override { return Compiler->TransformPosition(SourceCoordType,DestCoordType,A); }
+	virtual int32 TransformVector(EMaterialCommonBasis SourceCoordBasis, EMaterialCommonBasis DestCoordBasis, int32 A) override
+	{
+		return Compiler->TransformVector(SourceCoordBasis, DestCoordBasis, A);
+	}
+	virtual int32 TransformPosition(EMaterialCommonBasis SourceCoordBasis, EMaterialCommonBasis DestCoordBasis, int32 A) override
+	{
+		return Compiler->TransformPosition(SourceCoordBasis, DestCoordBasis, A);
+	}
 
 	virtual int32 DynamicParameter() override { return Compiler->DynamicParameter(); }
 	virtual int32 LightmapUVs() override { return Compiler->LightmapUVs(); }
diff --git a/Engine/Source/Runtime/Engine/Public/MaterialShared.h b/Engine/Source/Runtime/Engine/Public/MaterialShared.h
--- a/Engine/Source/Runtime/Engine/Public/MaterialShared.h
+++ b/Engine/Source/Runtime/Engine/Public/MaterialShared.h
@@ -83,6 +83,19 @@ enum EMaterialValueType
 	MCT_MaterialAttributes	= 256
 };
 
+/**
+ * The common bases of material
+ */
+enum EMaterialCommonBasis
+{
+	MCB_Tangent,
+	MCB_Local,
+	MCB_TranslatedWorld,
+	MCB_World,
+	MCB_View,
+	MCB_MAX,
+};
+
 /**
  * The context of a material being rendered.
  */
diff --git a/Engine/Source/Runtime/Engine/Public/SceneView.h b/Engine/Source/Runtime/Engine/Public/SceneView.h
--- a/Engine/Source/Runtime/Engine/Public/SceneView.h
+++ b/Engine/Source/Runtime/Engine/Public/SceneView.h
@@ -373,6 +373,8 @@ BEGIN_UNIFORM_BUFFER_STRUCT_WITH_CONSTRUCTOR(FViewUniformShaderParameters,ENGINE
 	DECLARE_UNIFORM_BUFFER_STRUCT_MEMBER(FMatrix,PrevViewProj)
 	DECLARE_UNIFORM_BUFFER_STRUCT_MEMBER(FMatrix,PrevViewRotationProj)
 	DECLARE_UNIFORM_BUFFER_STRUCT_MEMBER(FMatrix,PrevTranslatedWorldToClip)
+	DECLARE_UNIFORM_BUFFER_STRUCT_MEMBER(FMatrix,PrevTranslatedWorldToView)
+	DECLARE_UNIFORM_BUFFER_STRUCT_MEMBER(FMatrix,PrevViewToTranslatedWorld)
 	DECLARE_UNIFORM_BUFFER_STRUCT_MEMBER(FVector,PrevViewOrigin)
 	DECLARE_UNIFORM_BUFFER_STRUCT_MEMBER(FVector,PrevPreViewTranslation)
 	DECLARE_UNIFORM_BUFFER_STRUCT_MEMBER(FMatrix,PrevInvViewProj)
diff --git a/Engine/Source/Runtime/Renderer/Private/SceneRendering.cpp b/Engine/Source/Runtime/Renderer/Private/SceneRendering.cpp
--- a/Engine/Source/Runtime/Renderer/Private/SceneRendering.cpp
+++ b/Engine/Source/Runtime/Renderer/Private/SceneRendering.cpp
@@ -486,6 +486,8 @@ TUniformBufferRef<FViewUniformShaderParameters> FViewInfo::CreateUniformBuffer(
 	ViewUniformShaderParameters.PrevViewProj = PrevViewProjMatrix;
 	ViewUniformShaderParameters.PrevViewRotationProj = PrevViewRotationProjMatrix;
 	ViewUniformShaderParameters.PrevTranslatedWorldToClip = PrevViewMatrices.TranslatedViewProjectionMatrix;
+	ViewUniformShaderParameters.PrevTranslatedWorldToView = PrevViewMatrices.TranslatedViewMatrix;
+	ViewUniformShaderParameters.PrevViewToTranslatedWorld = ViewUniformShaderParameters.PrevTranslatedWorldToView.Inverse();
 	ViewUniformShaderParameters.PrevViewOrigin = PrevViewMatrices.ViewOrigin;
 	ViewUniformShaderParameters.PrevPreViewTranslation = PrevViewMatrices.PreViewTranslation;
 	// can be optimized
