added comment

diff --git a/Engine/Source/Runtime/Engine/Public/SceneView.h b/Engine/Source/Runtime/Engine/Public/SceneView.h
--- a/Engine/Source/Runtime/Engine/Public/SceneView.h
+++ b/Engine/Source/Runtime/Engine/Public/SceneView.h
@@ -699,7 +699,7 @@ public:
 	/** Configure post process settings for the buffer visualization system */
 	void ConfigureBufferVisualizationSettings();
 
-	/** Get the feature level for this view **/
+	/** Get the feature level for this view (cached from the scene so this is not different per view) **/
 	ERHIFeatureLevel::Type GetFeatureLevel() const { return FeatureLevel; }
 
 	/** Get the feature level for this view **/
