ARecastNavMesh::PostEditChangeProperty() now respects auto navigation rebuild setting

Merging 2594459

//depot/UE4-UT/...

to //depot/UE4/...

diff --git a/Engine/Source/Runtime/Engine/Private/AI/Navigation/RecastNavMesh.cpp b/Engine/Source/Runtime/Engine/Private/AI/Navigation/RecastNavMesh.cpp
--- a/Engine/Source/Runtime/Engine/Private/AI/Navigation/RecastNavMesh.cpp
+++ b/Engine/Source/Runtime/Engine/Private/AI/Navigation/RecastNavMesh.cpp
@@ -2005,7 +2005,7 @@ void ARecastNavMesh::PostEditChangeProperty(FPropertyChangedEvent& PropertyChang
 				UpdatePolyRefBitsPreview();
 			}
 
-			if (HasAnyFlags(RF_ClassDefaultObject) == false)
+			if (!HasAnyFlags(RF_ClassDefaultObject) && UNavigationSystem::GetIsNavigationAutoUpdateEnabled())
 			{
 				RebuildAll();
 			}
