Fixed: Android Galaxy S 3 (Mali) or Galaxy Note 2 (Mali) Mobile Devices Not Rendering BSP
UE-5428

diff --git a/Engine/Source/Runtime/Engine/Private/ModelRender.cpp b/Engine/Source/Runtime/Engine/Private/ModelRender.cpp
--- a/Engine/Source/Runtime/Engine/Private/ModelRender.cpp
+++ b/Engine/Source/Runtime/Engine/Private/ModelRender.cpp
@@ -138,6 +138,9 @@ void UModelComponent::BuildRenderData()
 		}
 
 		IndexBuffer->Indices.Shrink();
+#if !DISALLOW_32BIT_INDICES
+		IndexBuffer->ComputeIndexWidth();
+#endif
 	}
 }
 
diff --git a/Engine/Source/Runtime/Engine/Private/RawIndexBuffer.cpp b/Engine/Source/Runtime/Engine/Private/RawIndexBuffer.cpp
--- a/Engine/Source/Runtime/Engine/Private/RawIndexBuffer.cpp
+++ b/Engine/Source/Runtime/Engine/Private/RawIndexBuffer.cpp
@@ -72,18 +72,54 @@ void FRawIndexBuffer16or32::CacheOptimize()
 #endif
 }
 
+void FRawIndexBuffer16or32::ComputeIndexWidth()
+{
+	if (GetFeatureLevel() < ERHIFeatureLevel::SM4)
+	{
+		const int32 NumIndices = Indices.Num();
+		bool bShouldUse32Bit = false;
+		int32 i = 0;
+		while (!bShouldUse32Bit && i < NumIndices)
+		{
+			bShouldUse32Bit = Indices[i] > MAX_uint16;
+			i++;
+		}
+	
+		b32Bit = bShouldUse32Bit;
+	}
+	else
+	{
+		b32Bit = true;
+	}
+}
+
 void FRawIndexBuffer16or32::InitRHI()
 {
-	uint32 Size = Indices.Num() * sizeof(uint32);
-	if( Size > 0 )
+	const int32 IndexStride = b32Bit ? sizeof(uint32) : sizeof(uint16);
+	const int32 NumIndices = Indices.Num();
+	const uint32 Size = NumIndices * IndexStride;
+	
+	if (Size > 0)
 	{
 		// Create the index buffer.
 		FRHIResourceCreateInfo CreateInfo;
-		IndexBufferRHI = RHICreateIndexBuffer(sizeof(uint32),Size,BUF_Static,CreateInfo);
-
+		IndexBufferRHI = RHICreateIndexBuffer(IndexStride,Size,BUF_Static,CreateInfo);
+		
 		// Initialize the buffer.
 		void* Buffer = RHILockIndexBuffer(IndexBufferRHI, 0, Size, RLM_WriteOnly);
-		FMemory::Memcpy(Buffer,Indices.GetData(),Size);
+		if (b32Bit)
+		{
+			FMemory::Memcpy(Buffer, Indices.GetData(), Size);
+		}
+		else
+		{
+			uint16* DestIndices16Bit = (uint16*)Buffer;
+			for (int32 i = 0; i < NumIndices; ++i)
+			{
+				DestIndices16Bit[i] = Indices[i];
+			}
+		}
+		
 		RHIUnlockIndexBuffer(IndexBufferRHI);
 	}
 
diff --git a/Engine/Source/Runtime/Engine/Public/RawIndexBuffer.h b/Engine/Source/Runtime/Engine/Public/RawIndexBuffer.h
--- a/Engine/Source/Runtime/Engine/Public/RawIndexBuffer.h
+++ b/Engine/Source/Runtime/Engine/Public/RawIndexBuffer.h
@@ -36,6 +36,10 @@ typedef FRawIndexBuffer FRawIndexBuffer16or32;
 class FRawIndexBuffer16or32 : public FIndexBuffer
 {
 public:
+	FRawIndexBuffer16or32()
+		: b32Bit(true)
+	{
+	}
 
 	TArray<uint32> Indices;
 
@@ -44,11 +48,19 @@ public:
 	 */
 	void CacheOptimize();
 
+	/**
+	 * Computes whether index buffer should be 32 bit
+	 */
+	void ComputeIndexWidth();
+
 	// FRenderResource interface.
 	virtual void InitRHI() override;
 
 	// Serialization.
 	friend FArchive& operator<<(FArchive& Ar,FRawIndexBuffer16or32& I);
+
+private:
+	bool b32Bit;
 };
 #endif
 
