Fixed ShapeComponents with collision disabled not being able to influence navmesh generation via navigation obstacles mechanism #UE4

diff --git a/Engine/Source/Runtime/Engine/Classes/Components/ShapeComponent.h b/Engine/Source/Runtime/Engine/Classes/Components/ShapeComponent.h
--- a/Engine/Source/Runtime/Engine/Classes/Components/ShapeComponent.h
+++ b/Engine/Source/Runtime/Engine/Classes/Components/ShapeComponent.h
@@ -47,6 +47,10 @@ class ENGINE_API UShapeComponent : public UPrimitiveComponent
 	virtual void GetNavigationData(FNavigationRelevantData& Data) const override;
 	// End UPrimitiveComponent interface.
 
+	// Begin INavRelevantInterface Interface
+	virtual bool IsNavigationRelevant() const override;
+	// End INavRelevantInterface Interface
+
 	// Begin USceneComponent interface
 	virtual FBoxSphereBounds CalcBounds(const FTransform& LocalToWorld) const override;
 	virtual bool ShouldCollideWhenPlacing() const override
diff --git a/Engine/Source/Runtime/Engine/Private/Components/ShapeComponent.cpp b/Engine/Source/Runtime/Engine/Private/Components/ShapeComponent.cpp
--- a/Engine/Source/Runtime/Engine/Private/Components/ShapeComponent.cpp
+++ b/Engine/Source/Runtime/Engine/Private/Components/ShapeComponent.cpp
@@ -79,3 +79,8 @@ void UShapeComponent::GetNavigationData(FNavigationRelevantData& Data) const
 		Data.Modifiers.CreateAreaModifiers(this, AreaClass);
 	}
 }
+
+bool UShapeComponent::IsNavigationRelevant() const
+{
+	return bDynamicObstacle || Super::IsNavigationRelevant();
+}
\ No newline at end of file
