Turn on PVD using command line argument

diff --git a/Engine/Source/Runtime/Engine/Private/PhysicsEngine/PhysLevel.cpp b/Engine/Source/Runtime/Engine/Private/PhysicsEngine/PhysLevel.cpp
--- a/Engine/Source/Runtime/Engine/Private/PhysicsEngine/PhysLevel.cpp
+++ b/Engine/Source/Runtime/Engine/Private/PhysicsEngine/PhysLevel.cpp
@@ -305,6 +305,9 @@ FString FEndClothSimulationFunction::DiagnosticMessage()
 	return TEXT("FStartClothSimulationFunction");
 }
 
+
+void PvdConnect(FString Host);
+
 //////// GAME-LEVEL RIGID BODY PHYSICS STUFF ///////
 void InitGamePhys()
 {
@@ -364,6 +367,11 @@ void InitGamePhys()
 	//This approach facilitates mixing the use of height fields and meshes in the application with no tangible difference in collision behavior between the two approaches
 	PxRegisterUnifiedHeightFields(*GPhysXSDK);
 
+	if( FParse::Param( FCommandLine::Get(), TEXT( "PVD" ) ) )
+	{
+		PvdConnect(TEXT("localhost"));
+	}
+
 
 #if WITH_PHYSICS_COOKING || WITH_RUNTIME_PHYSICS_COOKING
 	// Create Cooking
diff --git a/Engine/Source/Runtime/Engine/Private/PhysicsEngine/PhysUtils.cpp b/Engine/Source/Runtime/Engine/Private/PhysicsEngine/PhysUtils.cpp
--- a/Engine/Source/Runtime/Engine/Private/PhysicsEngine/PhysUtils.cpp
+++ b/Engine/Source/Runtime/Engine/Private/PhysicsEngine/PhysUtils.cpp
@@ -445,6 +445,22 @@ static bool ExecApexVis(UWorld* InWorld, uint32 SceneType, const TCHAR* Cmd, FOu
 	return true;
 }
 
+#if WITH_PHYSX
+void PvdConnect(FString Host)
+{
+	int32		Port = 5425;         // TCP port to connect to, where PVD is listening
+	uint32	Timeout = 100;          // timeout in milliseconds to wait for PVD to respond, consoles and remote PCs need a higher timeout.
+	PxVisualDebuggerConnectionFlags ConnectionFlags = PxVisualDebuggerExt::getAllConnectionFlags();
+
+	// and now try to connect
+	PxVisualDebuggerExt::createConnection(GPhysXSDK->getPvdConnectionManager(), TCHAR_TO_ANSI(*Host), Port, Timeout, ConnectionFlags);
+
+	GPhysXSDK->getVisualDebugger()->setVisualDebuggerFlag(PxVisualDebuggerFlag::eTRANSMIT_CONTACTS, true);
+	GPhysXSDK->getVisualDebugger()->setVisualDebuggerFlag(PxVisualDebuggerFlag::eTRANSMIT_SCENEQUERIES, true);
+	GPhysXSDK->getVisualDebugger()->setVisualDebuggerFlag(PxVisualDebuggerFlag::eTRANSMIT_CONSTRAINTS, true);
+}
+#endif
+
 //// EXEC
 bool ExecPhysCommands(const TCHAR* Cmd, FOutputDevice* Ar, UWorld* InWorld)
 {
@@ -541,16 +557,7 @@ bool ExecPhysCommands(const TCHAR* Cmd, FOutputDevice* Ar, UWorld* InWorld)
 					Host = Cmd;
 				}
 
-				int32		Port	= 5425;         // TCP port to connect to, where PVD is listening
-				uint32	Timeout = 100;          // timeout in milliseconds to wait for PVD to respond, consoles and remote PCs need a higher timeout.
-				PxVisualDebuggerConnectionFlags ConnectionFlags = PxVisualDebuggerExt::getAllConnectionFlags();
-
-				// and now try to connect
-				PxVisualDebuggerExt::createConnection(GPhysXSDK->getPvdConnectionManager(), TCHAR_TO_ANSI(*Host), Port, Timeout, ConnectionFlags);
-
-				GPhysXSDK->getVisualDebugger()->setVisualDebuggerFlag(PxVisualDebuggerFlag::eTRANSMIT_CONTACTS, true);
-				GPhysXSDK->getVisualDebugger()->setVisualDebuggerFlag(PxVisualDebuggerFlag::eTRANSMIT_SCENEQUERIES, true);
-				GPhysXSDK->getVisualDebugger()->setVisualDebuggerFlag(PxVisualDebuggerFlag::eTRANSMIT_CONSTRAINTS, true);
+				PvdConnect(Host);
 
 			}
 			else if(FParse::Command(&Cmd, TEXT("DISCONNECT")))
