Fix for hierarchical instanced static meshes, making sure the cluster tree gets rebuilt when there are unbuilt instances.

diff --git a/Engine/Source/Runtime/Engine/Private/HierarchicalInstancedStaticMesh.cpp b/Engine/Source/Runtime/Engine/Private/HierarchicalInstancedStaticMesh.cpp
--- a/Engine/Source/Runtime/Engine/Private/HierarchicalInstancedStaticMesh.cpp
+++ b/Engine/Source/Runtime/Engine/Private/HierarchicalInstancedStaticMesh.cpp
@@ -2220,7 +2220,7 @@ void UHierarchicalInstancedStaticMeshComponent::PostLoad()
 		}
 	}
 	// For some reason we don't have a tree, or it is out of date. Build one now!
-	if (StaticMesh && PerInstanceSMData.Num() > 0 && (!ClusterTreePtr.IsValid() || ClusterTreePtr->Num() == 0 || GetLinkerUE4Version() < VER_UE4_REBUILD_HIERARCHICAL_INSTANCE_TREES))
+	if (StaticMesh && PerInstanceSMData.Num() > 0 && (!ClusterTreePtr.IsValid() || ClusterTreePtr->Num() == 0 || (NumBuiltInstances != PerInstanceSMData.Num()) || GetLinkerUE4Version() < VER_UE4_REBUILD_HIERARCHICAL_INSTANCE_TREES))
 	{
 		UE_LOG(LogStaticMesh, Warning, TEXT("Rebuilding foliage, please resave map %s."), *GetFullName());
 		check(!IsAsyncBuilding());
