UE4 - fix crash on zero scaling for HISMC

diff --git a/Engine/Source/Runtime/Engine/Private/HierarchicalInstancedStaticMesh.cpp b/Engine/Source/Runtime/Engine/Private/HierarchicalInstancedStaticMesh.cpp
--- a/Engine/Source/Runtime/Engine/Private/HierarchicalInstancedStaticMesh.cpp
+++ b/Engine/Source/Runtime/Engine/Private/HierarchicalInstancedStaticMesh.cpp
@@ -1257,18 +1257,26 @@ void FHierarchicalStaticMeshSceneProxy::GetDynamicMeshElements(const TArray<cons
 
 					if (View->ViewMatrices.IsPerspectiveProjection())
 					{
-						InstanceParams.ViewFrustumLocal.Planes.Pop(false); // we don't want the far plane either
-						FMatrix ThreePlanes;
-						ThreePlanes.SetIdentity();
-						ThreePlanes.SetAxes(&InstanceParams.ViewFrustumLocal.Planes[0], &InstanceParams.ViewFrustumLocal.Planes[1], &InstanceParams.ViewFrustumLocal.Planes[2]);
-						FVector ProjectionOrigin = ThreePlanes.Inverse().GetTransposed().TransformVector(FVector(InstanceParams.ViewFrustumLocal.Planes[0].W, InstanceParams.ViewFrustumLocal.Planes[1].W, InstanceParams.ViewFrustumLocal.Planes[2].W));
-
-						FVector Forward = LocalViewProjForCulling.GetColumn(3).GetSafeNormal();
-						for (int32 Index = 0; Index < InstanceParams.ViewFrustumLocal.Planes.Num(); Index++)
+						if (InstanceParams.ViewFrustumLocal.Planes.Num() == 5)
 						{
-							FPlane Src = InstanceParams.ViewFrustumLocal.Planes[Index];
-							FVector Normal = Src.GetSafeNormal();
-							InstanceParams.ViewFrustumLocal.Planes[Index] = FPlane(Normal, Normal | ProjectionOrigin);
+							InstanceParams.ViewFrustumLocal.Planes.Pop(false); // we don't want the far plane either
+							FMatrix ThreePlanes;
+							ThreePlanes.SetIdentity();
+							ThreePlanes.SetAxes(&InstanceParams.ViewFrustumLocal.Planes[0], &InstanceParams.ViewFrustumLocal.Planes[1], &InstanceParams.ViewFrustumLocal.Planes[2]);
+							FVector ProjectionOrigin = ThreePlanes.Inverse().GetTransposed().TransformVector(FVector(InstanceParams.ViewFrustumLocal.Planes[0].W, InstanceParams.ViewFrustumLocal.Planes[1].W, InstanceParams.ViewFrustumLocal.Planes[2].W));
+
+							FVector Forward = LocalViewProjForCulling.GetColumn(3).GetSafeNormal();
+							for (int32 Index = 0; Index < InstanceParams.ViewFrustumLocal.Planes.Num(); Index++)
+							{
+								FPlane Src = InstanceParams.ViewFrustumLocal.Planes[Index];
+								FVector Normal = Src.GetSafeNormal();
+								InstanceParams.ViewFrustumLocal.Planes[Index] = FPlane(Normal, Normal | ProjectionOrigin);
+							}
+						}
+						else
+						{
+							 // zero scaling or something, cull everything
+							continue;
 						}
 					}
 					else
