Expose GetCollisionObjectType and GetCollisionResponseToChannel to blueprints

#UE-8569

diff --git a/Engine/Source/Runtime/Engine/Classes/Components/PrimitiveComponent.h b/Engine/Source/Runtime/Engine/Classes/Components/PrimitiveComponent.h
--- a/Engine/Source/Runtime/Engine/Classes/Components/PrimitiveComponent.h
+++ b/Engine/Source/Runtime/Engine/Classes/Components/PrimitiveComponent.h
@@ -1388,9 +1388,18 @@ protected:
 	
 public:
 	virtual bool IsWorldGeometry() const override;
+
+
 	virtual ECollisionEnabled::Type GetCollisionEnabled() const override;
+
+	/** Gets the response type given a specific channel */
+	UFUNCTION(BlueprintCallable, Category="Physics")
 	virtual ECollisionResponse GetCollisionResponseToChannel(ECollisionChannel Channel) const override;
+
+	/** Gets the collision object type */
+	UFUNCTION(BlueprintCallable, Category="Physics")
 	virtual ECollisionChannel GetCollisionObjectType() const override;
+
 	virtual const FCollisionResponseContainer& GetCollisionResponseToChannels() const override;
 	virtual FVector GetComponentVelocity() const override;
 	//End USceneComponent Interface
