Expose the teleport flag to public so that movement functions can determine if to sweep or teleport orthogonally

#UE-17485

diff --git a/Engine/Plugins/2D/Paper2D/Source/Paper2D/Private/PaperGroupedSpriteComponent.cpp b/Engine/Plugins/2D/Paper2D/Source/Paper2D/Private/PaperGroupedSpriteComponent.cpp
--- a/Engine/Plugins/2D/Paper2D/Source/Paper2D/Private/PaperGroupedSpriteComponent.cpp
+++ b/Engine/Plugins/2D/Paper2D/Source/Paper2D/Private/PaperGroupedSpriteComponent.cpp
@@ -92,7 +92,7 @@ bool UPaperGroupedSpriteComponent::UpdateInstanceTransform(int32 InstanceIndex,
 		if (FBodyInstance* InstanceBodyInstance = InstanceBodies[InstanceIndex])
 		{
 			// Update transform.
-			InstanceBodyInstance->SetBodyTransform(WorldTransform, false);
+			InstanceBodyInstance->SetBodyTransform(WorldTransform, ETeleportType::None);
 			InstanceBodyInstance->UpdateBodyScale(WorldTransform.GetScale3D());
 		}
 	}
diff --git a/Engine/Source/Runtime/Engine/Classes/Components/ActorComponent.h b/Engine/Source/Runtime/Engine/Classes/Components/ActorComponent.h
--- a/Engine/Source/Runtime/Engine/Classes/Components/ActorComponent.h
+++ b/Engine/Source/Runtime/Engine/Classes/Components/ActorComponent.h
@@ -22,6 +22,18 @@ enum class EComponentCreationMethod : uint8
 	Instance,
 };
 
+
+/** Whether to teleport physics body or not */
+enum class ETeleportType
+{
+	/** Do not teleport physics body. This means velocity will reflect the movement between initial and final position, and collisions along the way will occur */
+	None,
+	/** Teleport physics body so that velocity remains the same and no collision occurs */
+	TeleportPhysics
+};
+
+FORCEINLINE ETeleportType TeleportFlagToEnum(bool bTeleport) { return bTeleport ? ETeleportType::TeleportPhysics : ETeleportType::None; }
+
 /**
  * ActorComponent is the base class for components that define reusable behavior that can be added to different types of Actors.
  * ActorComponents that have a transform are known as SceneComponents and those that can be rendered are PrimitiveComponents.
@@ -489,7 +501,7 @@ public:
 	void DoDeferredRenderUpdates_Concurrent();
 
 	/** Recalculate the value of our component to world transform */
-	virtual void UpdateComponentToWorld(bool bSkipPhysicsMove = false, bool bTeleport = false) {}
+	virtual void UpdateComponentToWorld(bool bSkipPhysicsMove = false, ETeleportType Teleport = ETeleportType::None){}
 
 	/** Mark the render state as dirty - will be sent to the render thread at the end of the frame. */
 	void MarkRenderStateDirty();
diff --git a/Engine/Source/Runtime/Engine/Classes/Components/AudioComponent.h b/Engine/Source/Runtime/Engine/Classes/Components/AudioComponent.h
--- a/Engine/Source/Runtime/Engine/Classes/Components/AudioComponent.h
+++ b/Engine/Source/Runtime/Engine/Classes/Components/AudioComponent.h
@@ -270,7 +270,7 @@ public:
 	// Begin USceneComponent Interface
 	virtual void Activate(bool bReset=false) override;
 	virtual void Deactivate() override;
-	virtual void OnUpdateTransform(bool bSkipPhysicsMove, bool bTeleport = false) override;
+	virtual void OnUpdateTransform(bool bSkipPhysicsMove, ETeleportType Teleport = ETeleportType::None) override;
 	// do nothing! on purpose
 	virtual void UpdateNavigationData() override {}
 	// End USceneComponent Interface
diff --git a/Engine/Source/Runtime/Engine/Classes/Components/DestructibleComponent.h b/Engine/Source/Runtime/Engine/Classes/Components/DestructibleComponent.h
--- a/Engine/Source/Runtime/Engine/Classes/Components/DestructibleComponent.h
+++ b/Engine/Source/Runtime/Engine/Classes/Components/DestructibleComponent.h
@@ -107,7 +107,7 @@ public:
 
 	// Begin USceneComponent interface.
 	virtual FBoxSphereBounds CalcBounds(const FTransform& LocalToWorld) const override;
-	virtual void OnUpdateTransform(bool bSkipPhysicsMove, bool bTeleport = false) override;
+	virtual void OnUpdateTransform(bool bSkipPhysicsMove, ETeleportType Teleport = ETeleportType::None) override;
 	virtual void Activate(bool bReset=false) override;
 	virtual void Deactivate() override;
 	// End USceneComponent interface.
diff --git a/Engine/Source/Runtime/Engine/Classes/Components/PrimitiveComponent.h b/Engine/Source/Runtime/Engine/Classes/Components/PrimitiveComponent.h
--- a/Engine/Source/Runtime/Engine/Classes/Components/PrimitiveComponent.h
+++ b/Engine/Source/Runtime/Engine/Classes/Components/PrimitiveComponent.h
@@ -1286,7 +1286,7 @@ protected:
 	friend class FStaticMeshComponentRecreateRenderStateContext;
 
 	// Begin USceneComponent Interface
-	virtual void OnUpdateTransform(bool bSkipPhysicsMove, bool bTeleport = false) override;
+	virtual void OnUpdateTransform(bool bSkipPhysicsMove, ETeleportType Teleport = ETeleportType::None) override;
 
 	/** Event called when AttachParent changes, to allow the scene to update its attachment state. */
 	virtual void OnAttachmentChanged() override;
@@ -1341,7 +1341,7 @@ protected:
 	virtual void UpdatePhysicsToRBChannels();
 
 	/** Called to send a transform update for this component to the physics engine */
-	void SendPhysicsTransform(bool bTeleport);
+	void SendPhysicsTransform(ETeleportType Teleport);
 
 	/** Ensure physics state created **/
 	void EnsurePhysicsStateCreated();
@@ -1384,7 +1384,7 @@ public:
 	//Begin USceneComponent Interface
 
 protected:
-	virtual bool MoveComponentImpl(const FVector& Delta, const FQuat& NewRotation, bool bSweep, FHitResult* OutHit = NULL, EMoveComponentFlags MoveFlags = MOVECOMP_NoFlags) override;
+	virtual bool MoveComponentImpl(const FVector& Delta, const FQuat& NewRotation, bool bSweep, FHitResult* OutHit = NULL, EMoveComponentFlags MoveFlags = MOVECOMP_NoFlags, ETeleportType Teleport = ETeleportType::None) override;
 	
 public:
 	virtual bool IsWorldGeometry() const override;
diff --git a/Engine/Source/Runtime/Engine/Classes/Components/SceneComponent.h b/Engine/Source/Runtime/Engine/Classes/Components/SceneComponent.h
--- a/Engine/Source/Runtime/Engine/Classes/Components/SceneComponent.h
+++ b/Engine/Source/Runtime/Engine/Classes/Components/SceneComponent.h
@@ -247,20 +247,20 @@ public:
 public:
 	/** Set the location of this component relative to its parent */
 	UFUNCTION(BlueprintCallable, Category="Utilities|Transformation", meta=(DisplayName="SetRelativeLocation"))
-	void K2_SetRelativeLocation(FVector NewLocation, bool bSweep, FHitResult& SweepHitResult);
-	void SetRelativeLocation(FVector NewLocation, bool bSweep=false, FHitResult* OutSweepHitResult=nullptr);
+	void K2_SetRelativeLocation(FVector NewLocation, bool bSweep, FHitResult& SweepHitResult, bool bTeleport);
+	void SetRelativeLocation(FVector NewLocation, bool bSweep=false, FHitResult* OutSweepHitResult=nullptr, ETeleportType Teleport = ETeleportType::None);
 
 
 	/** Set the rotation of this component relative to its parent */
 	UFUNCTION(BlueprintCallable, Category="Utilities|Transformation", meta=(DisplayName="SetRelativeRotation", AdvancedDisplay="bSweep,SweepHitResult"))
-	void K2_SetRelativeRotation(FRotator NewRotation, bool bSweep, FHitResult& SweepHitResult);
-	void SetRelativeRotation(FRotator NewRotation, bool bSweep=false, FHitResult* OutSweepHitResult=nullptr);
-	void SetRelativeRotation(const FQuat& NewRotation, bool bSweep=false, FHitResult* OutSweepHitResult=nullptr);
+	void K2_SetRelativeRotation(FRotator NewRotation, bool bSweep, FHitResult& SweepHitResult, bool bTeleport);
+	void SetRelativeRotation(FRotator NewRotation, bool bSweep=false, FHitResult* OutSweepHitResult=nullptr, ETeleportType Teleport = ETeleportType::None);
+	void SetRelativeRotation(const FQuat& NewRotation, bool bSweep=false, FHitResult* OutSweepHitResult=nullptr, ETeleportType Teleport = ETeleportType::None);
 
 	/** Set the transform of this component relative to its parent */
 	UFUNCTION(BlueprintCallable, Category="Utilities|Transformation", meta=(DisplayName="SetRelativeTransform"))
-	void K2_SetRelativeTransform(const FTransform& NewTransform, bool bSweep, FHitResult& SweepHitResult);
-	void SetRelativeTransform(const FTransform& NewTransform, bool bSweep=false, FHitResult* OutSweepHitResult=nullptr);
+	void K2_SetRelativeTransform(const FTransform& NewTransform, bool bSweep, FHitResult& SweepHitResult, bool bTeleport);
+	void SetRelativeTransform(const FTransform& NewTransform, bool bSweep=false, FHitResult* OutSweepHitResult=nullptr, ETeleportType Teleport = ETeleportType::None);
 
 	/** Returns the transform of this component relative to its parent */
 	UFUNCTION(BlueprintCallable, Category="Utilities|Transformation")
@@ -276,45 +276,45 @@ public:
 
 	/** Adds a delta to the translation of this component relative to its parent */
 	UFUNCTION(BlueprintCallable, Category="Utilities|Transformation", meta=(DisplayName="AddRelativeLocation"))
-	void K2_AddRelativeLocation(FVector DeltaLocation, bool bSweep, FHitResult& SweepHitResult);
-	void AddRelativeLocation(FVector DeltaLocation, bool bSweep=false, FHitResult* OutSweepHitResult=nullptr);
+	void K2_AddRelativeLocation(FVector DeltaLocation, bool bSweep, FHitResult& SweepHitResult, bool bTeleport);
+	void AddRelativeLocation(FVector DeltaLocation, bool bSweep=false, FHitResult* OutSweepHitResult=nullptr, ETeleportType Teleport = ETeleportType::None);
 
 	/** Adds a delta the rotation of this component relative to its parent */
 	UFUNCTION(BlueprintCallable, Category="Utilities|Transformation", meta=(DisplayName="AddRelativeRotation", AdvancedDisplay="bSweep,SweepHitResult"))
-	void K2_AddRelativeRotation(FRotator DeltaRotation, bool bSweep, FHitResult& SweepHitResult);
-	void AddRelativeRotation(FRotator DeltaRotation, bool bSweep=false, FHitResult* OutSweepHitResult=nullptr);
-	void AddRelativeRotation(const FQuat& DeltaRotation, bool bSweep=false, FHitResult* OutSweepHitResult=nullptr);
+	void K2_AddRelativeRotation(FRotator DeltaRotation, bool bSweep, FHitResult& SweepHitResult, bool bTeleport);
+	void AddRelativeRotation(FRotator DeltaRotation, bool bSweep=false, FHitResult* OutSweepHitResult=nullptr, ETeleportType Teleport = ETeleportType::None);
+	void AddRelativeRotation(const FQuat& DeltaRotation, bool bSweep=false, FHitResult* OutSweepHitResult=nullptr, ETeleportType Teleport = ETeleportType::None);
 
 	/** Adds a delta to the location of this component in its local reference frame */
 	UFUNCTION(BlueprintCallable, Category="Utilities|Transformation", meta=(DisplayName="AddLocalOffset", Keywords="location position"))
-	void K2_AddLocalOffset(FVector DeltaLocation, bool bSweep, FHitResult& SweepHitResult);
-	void AddLocalOffset(FVector DeltaLocation, bool bSweep=false, FHitResult* OutSweepHitResult=nullptr);
+	void K2_AddLocalOffset(FVector DeltaLocation, bool bSweep, FHitResult& SweepHitResult, bool bTeleport);
+	void AddLocalOffset(FVector DeltaLocation, bool bSweep=false, FHitResult* OutSweepHitResult=nullptr, ETeleportType Teleport = ETeleportType::None);
 
 	
 	/** Adds a delta to the rotation of this component in its local reference frame */
 	UFUNCTION(BlueprintCallable, Category="Utilities|Transformation", meta=(DisplayName="AddLocalRotation", AdvancedDisplay="bSweep,SweepHitResult"))
-	void K2_AddLocalRotation(FRotator DeltaRotation, bool bSweep, FHitResult& SweepHitResult);
-	void AddLocalRotation(FRotator DeltaRotation, bool bSweep=false, FHitResult* OutSweepHitResult=nullptr);
-	void AddLocalRotation(const FQuat& DeltaRotation, bool bSweep=false, FHitResult* OutSweepHitResult=nullptr);
+	void K2_AddLocalRotation(FRotator DeltaRotation, bool bSweep, FHitResult& SweepHitResult, bool bTeleport);
+	void AddLocalRotation(FRotator DeltaRotation, bool bSweep=false, FHitResult* OutSweepHitResult=nullptr, ETeleportType Teleport = ETeleportType::None);
+	void AddLocalRotation(const FQuat& DeltaRotation, bool bSweep=false, FHitResult* OutSweepHitResult=nullptr, ETeleportType Teleport = ETeleportType::None);
 
 
 	/** Adds a delta to the transform of this component in its local reference frame. Scale is unchanged. */
 	UFUNCTION(BlueprintCallable, Category="Utilities|Transformation", meta=(DisplayName="AddLocalTransform"))
-	void K2_AddLocalTransform(const FTransform& DeltaTransform, bool bSweep, FHitResult& SweepHitResult);
-	void AddLocalTransform(const FTransform& DeltaTransform, bool bSweep=false, FHitResult* OutSweepHitResult=nullptr);
+	void K2_AddLocalTransform(const FTransform& DeltaTransform, bool bSweep, FHitResult& SweepHitResult, bool bTeleport);
+	void AddLocalTransform(const FTransform& DeltaTransform, bool bSweep=false, FHitResult* OutSweepHitResult=nullptr, ETeleportType Teleport = ETeleportType::None);
 
 
 	/** Set the relative location of this component to put it at the supplied location in world space. */
 	UFUNCTION(BlueprintCallable, Category="Utilities|Transformation", meta=(DisplayName="SetWorldLocation"))
-	void K2_SetWorldLocation(FVector NewLocation, bool bSweep, FHitResult& SweepHitResult);
-	void SetWorldLocation(FVector NewLocation, bool bSweep=false, FHitResult* OutSweepHitResult=nullptr);
+	void K2_SetWorldLocation(FVector NewLocation, bool bSweep, FHitResult& SweepHitResult, bool bTeleport);
+	void SetWorldLocation(FVector NewLocation, bool bSweep=false, FHitResult* OutSweepHitResult=nullptr, ETeleportType Teleport = ETeleportType::None);
 
 
 	/** Set the relative rotation of this component to put it at the supplied orientation in world space. */
 	UFUNCTION(BlueprintCallable, Category="Utilities|Transformation", meta=(DisplayName="SetWorldRotation", AdvancedDisplay="bSweep,SweepHitResult"))
-	void K2_SetWorldRotation(FRotator NewRotation, bool bSweep, FHitResult& SweepHitResult);
-	void SetWorldRotation(FRotator NewRotation, bool bSweep=false, FHitResult* OutSweepHitResult=nullptr);
-	void SetWorldRotation(const FQuat& NewRotation, bool bSweep=false, FHitResult* OutSweepHitResult=nullptr);
+	void K2_SetWorldRotation(FRotator NewRotation, bool bSweep, FHitResult& SweepHitResult, bool bTeleport);
+	void SetWorldRotation(FRotator NewRotation, bool bSweep=false, FHitResult* OutSweepHitResult=nullptr, ETeleportType Teleport = ETeleportType::None);
+	void SetWorldRotation(const FQuat& NewRotation, bool bSweep=false, FHitResult* OutSweepHitResult=nullptr, ETeleportType Teleport = ETeleportType::None);
 
 
 	/** Set the relative scale of this component to put it at the supplied scale in world space. */
@@ -323,26 +323,26 @@ public:
 
 	/** Set the transform of this component in world space. */
 	UFUNCTION(BlueprintCallable, Category="Utilities|Transformation", meta=(DisplayName="SetWorldTransform"))
-	void K2_SetWorldTransform(const FTransform& NewTransform, bool bSweep, FHitResult& SweepHitResult);
-	void SetWorldTransform(const FTransform& NewTransform, bool bSweep=false, FHitResult* OutSweepHitResult=nullptr);
+	void K2_SetWorldTransform(const FTransform& NewTransform, bool bSweep, FHitResult& SweepHitResult, bool bTeleport);
+	void SetWorldTransform(const FTransform& NewTransform, bool bSweep=false, FHitResult* OutSweepHitResult=nullptr, ETeleportType Teleport = ETeleportType::None);
 
 
 	/** Adds a delta to the location of this component in world space. */
 	UFUNCTION(BlueprintCallable, Category="Utilities|Transformation", meta=(DisplayName="AddWorldOffset", Keywords="location position"))
-	void K2_AddWorldOffset(FVector DeltaLocation, bool bSweep, FHitResult& SweepHitResult);
-	void AddWorldOffset(FVector DeltaLocation, bool bSweep=false, FHitResult* OutSweepHitResult=nullptr);
+	void K2_AddWorldOffset(FVector DeltaLocation, bool bSweep, FHitResult& SweepHitResult, bool bTeleport);
+	void AddWorldOffset(FVector DeltaLocation, bool bSweep=false, FHitResult* OutSweepHitResult=nullptr, ETeleportType Teleport = ETeleportType::None);
 
 
 	/** Adds a delta to the rotation of this component in world space. */
 	UFUNCTION(BlueprintCallable, Category="Utilities|Transformation", meta=(DisplayName="AddWorldRotation", AdvancedDisplay="bSweep,SweepHitResult"))
-	void K2_AddWorldRotation(FRotator DeltaRotation, bool bSweep, FHitResult& SweepHitResult);
-	void AddWorldRotation(FRotator DeltaRotation, bool bSweep=false, FHitResult* OutSweepHitResult=nullptr);
-	void AddWorldRotation(const FQuat& DeltaRotation, bool bSweep=false, FHitResult* OutSweepHitResult=nullptr);
+	void K2_AddWorldRotation(FRotator DeltaRotation, bool bSweep, FHitResult& SweepHitResult, bool bTeleport);
+	void AddWorldRotation(FRotator DeltaRotation, bool bSweep=false, FHitResult* OutSweepHitResult=nullptr, ETeleportType Teleport = ETeleportType::None);
+	void AddWorldRotation(const FQuat& DeltaRotation, bool bSweep=false, FHitResult* OutSweepHitResult=nullptr, ETeleportType Teleport = ETeleportType::None);
 
 	/** Adds a delta to the transform of this component in world space. Scale is unchanged. */
 	UFUNCTION(BlueprintCallable, Category="Utilities|Transformation", meta=(DisplayName="AddWorldTransform"))
-	void K2_AddWorldTransform(const FTransform& DeltaTransform, bool bSweep, FHitResult& SweepHitResult);
-	void AddWorldTransform(const FTransform& DeltaTransform, bool bSweep=false, FHitResult* OutSweepHitResult=nullptr);
+	void K2_AddWorldTransform(const FTransform& DeltaTransform, bool bSweep, FHitResult& SweepHitResult, bool bTeleport);
+	void AddWorldTransform(const FTransform& DeltaTransform, bool bSweep=false, FHitResult* OutSweepHitResult=nullptr, ETeleportType Teleport = ETeleportType::None);
 
 	/** Return location of the component, in world space */
 	UFUNCTION(BlueprintCallable, meta=(DisplayName = "GetWorldLocation"), Category="Utilities|Transformation")
@@ -531,7 +531,7 @@ public:
 	{
 		return true;
 	}
-	virtual void UpdateComponentToWorld(bool bSkipPhysicsMove = false, bool bTeleport = false) override final;
+	virtual void UpdateComponentToWorld(bool bSkipPhysicsMove = false, ETeleportType Teleport = ETeleportType::None) override final;
 	virtual void DestroyComponent(bool bPromoteChildren = false) override;
 	virtual void OnComponentDestroyed() override;
 	virtual void ApplyWorldOffset(const FVector& InOffset, bool bWorldShift) override;
@@ -558,17 +558,17 @@ protected:
 	 * Internal helper, for use from MoveComponent().  Special codepath since the normal setters call MoveComponent.
 	 * @return: true if location or rotation was changed.
 	 */
-	bool InternalSetWorldLocationAndRotation(FVector NewLocation, const FQuat& NewQuat, bool bNoPhysics=false, bool bTeleport = false);
+	bool InternalSetWorldLocationAndRotation(FVector NewLocation, const FQuat& NewQuat, bool bNoPhysics = false, ETeleportType Teleport = ETeleportType::None);
 
-	virtual void OnUpdateTransform(bool bSkipPhysicsMove, bool bTeleport = false);
+	virtual void OnUpdateTransform(bool bSkipPhysicsMove, ETeleportType Teleport = ETeleportType::None);
 
 	/** Check if mobility is set to non-static. If it's static we trigger a PIE warning and return true*/
 	bool CheckStaticMobilityAndWarn(const FText& ActionText) const;
 
 private:
 
-	void PropagateTransformUpdate(bool bTransformChanged, bool bSkipPhysicsMove = false, bool bTeleport = false);
-	void UpdateComponentToWorldWithParent(USceneComponent* Parent, bool bSkipPhysicsMove, const FQuat& RelativeRotationQuat, bool bTeleport = false);
+	void PropagateTransformUpdate(bool bTransformChanged, bool bSkipPhysicsMove = false, ETeleportType Teleport = ETeleportType::None);
+	void UpdateComponentToWorldWithParent(USceneComponent* Parent, bool bSkipPhysicsMove, const FQuat& RelativeRotationQuat, ETeleportType Teleport = ETeleportType::None);
 
 
 public:
@@ -585,18 +585,19 @@ public:
 	 * @note The overload taking rotation as an FQuat is slightly faster than the version using FRotator (which will be converted to an FQuat)..
 	 * @param Delta			The desired location change in world space.
 	 * @param NewRotation	The new desired rotation in world space.
-	 * @param bSweep		True to do a swept move, which will stop at a blocking collision. False to do a simple teleport, which will not stop for collisions.
+	 * @param bSweep		True to do a swept move, which will stop at a blocking collision.
 	 * @param Hit			Optional output describing the blocking hit that stopped the move, if any.
 	 * @param MoveFlags		Flags controlling behavior of the move. @see EMoveComponentFlags
+	 * @param Teleport      Determines whether to teleport the physics body or not. Teleporting will maintain constant velocity and avoid collisions along the path
 	 * @return				True if some movement occurred, false if no movement occurred.
 	 */
-	bool MoveComponent( const FVector& Delta, const FQuat& NewRotation,    bool bSweep, FHitResult* Hit=NULL, EMoveComponentFlags MoveFlags = MOVECOMP_NoFlags );
-	bool MoveComponent( const FVector& Delta, const FRotator& NewRotation, bool bSweep, FHitResult* Hit=NULL, EMoveComponentFlags MoveFlags = MOVECOMP_NoFlags );
+	bool MoveComponent( const FVector& Delta, const FQuat& NewRotation,    bool bSweep, FHitResult* Hit=NULL, EMoveComponentFlags MoveFlags = MOVECOMP_NoFlags, ETeleportType Teleport = ETeleportType::None);
+	bool MoveComponent( const FVector& Delta, const FRotator& NewRotation, bool bSweep, FHitResult* Hit=NULL, EMoveComponentFlags MoveFlags = MOVECOMP_NoFlags, ETeleportType Teleport = ETeleportType::None);
 
 protected:
 
 	// Override this method for custom behavior.
-	virtual bool MoveComponentImpl( const FVector& Delta, const FQuat& NewRotation, bool bSweep, FHitResult* Hit=NULL, EMoveComponentFlags MoveFlags = MOVECOMP_NoFlags );
+	virtual bool MoveComponentImpl(const FVector& Delta, const FQuat& NewRotation, bool bSweep, FHitResult* Hit = NULL, EMoveComponentFlags MoveFlags = MOVECOMP_NoFlags, ETeleportType Teleport = ETeleportType::None);
 
 public:
 
@@ -683,7 +684,7 @@ public:
 	}
 
 	/** Update transforms of any components attached to this one. */
-	void UpdateChildTransforms(bool bSkipPhysicsMove = false, bool bTeleport = false);
+	void UpdateChildTransforms(bool bSkipPhysicsMove = false, ETeleportType Teleport = ETeleportType::None);
 
 	/** Calculate the bounds of this component. Default behavior is a bounding box/sphere of zero size. */
 	virtual FBoxSphereBounds CalcBounds(const FTransform& LocalToWorld) const;
@@ -765,9 +766,9 @@ protected:
 public:
 	/** Set the location and rotation of this component relative to its parent */
 	UFUNCTION(BlueprintCallable, Category="Utilities|Transformation", meta=(DisplayName="SetRelativeLocationAndRotation"))
-	void K2_SetRelativeLocationAndRotation(FVector NewLocation, FRotator NewRotation, bool bSweep, FHitResult& SweepHitResult);
-	void SetRelativeLocationAndRotation(FVector NewLocation, FRotator NewRotation, bool bSweep=false, FHitResult* OutSweepHitResult=nullptr);
-	void SetRelativeLocationAndRotation(FVector NewLocation, const FQuat& NewRotation, bool bSweep=false, FHitResult* OutSweepHitResult=nullptr);
+	void K2_SetRelativeLocationAndRotation(FVector NewLocation, FRotator NewRotation, bool bSweep, FHitResult& SweepHitResult, bool bTeleport);
+	void SetRelativeLocationAndRotation(FVector NewLocation, FRotator NewRotation, bool bSweep=false, FHitResult* OutSweepHitResult=nullptr, ETeleportType Teleport = ETeleportType::None);
+	void SetRelativeLocationAndRotation(FVector NewLocation, const FQuat& NewRotation, bool bSweep=false, FHitResult* OutSweepHitResult=nullptr, ETeleportType Teleport = ETeleportType::None);
 
 	/** Set which parts of the relative transform should be relative to parent, and which should be relative to world */
 	UFUNCTION(BlueprintCallable, Category="Utilities|Transformation")
@@ -775,11 +776,11 @@ public:
 
 	/** Set the relative location and FRotator rotation of this component to put it at the supplied pose in world space. */
 	UFUNCTION(BlueprintCallable, Category="Utilities|Transformation", meta=(DisplayName="SetWorldLocationAndRotation"))
-	void K2_SetWorldLocationAndRotation(FVector NewLocation, FRotator NewRotation, bool bSweep, FHitResult& SweepHitResult);
-	void SetWorldLocationAndRotation(FVector NewLocation, FRotator NewRotation, bool bSweep=false, FHitResult* OutSweepHitResult=nullptr);
+	void K2_SetWorldLocationAndRotation(FVector NewLocation, FRotator NewRotation, bool bSweep, FHitResult& SweepHitResult, bool bTeleport);
+	void SetWorldLocationAndRotation(FVector NewLocation, FRotator NewRotation, bool bSweep=false, FHitResult* OutSweepHitResult=nullptr, ETeleportType Teleport = ETeleportType::None);
 
 	/** Set the relative location and FQuat rotation of this component to put it at the supplied pose in world space. */
-	void SetWorldLocationAndRotation(FVector NewLocation, const FQuat& NewRotation, bool bSweep=false, FHitResult* OutSweepHitResult=nullptr);
+	void SetWorldLocationAndRotation(FVector NewLocation, const FQuat& NewRotation, bool bSweep=false, FHitResult* OutSweepHitResult=nullptr, ETeleportType Teleport = ETeleportType::None);
 
 	/** Special version of SetWorldLocationAndRotation that does not affect physics. */
 	void SetWorldLocationAndRotationNoPhysics(const FVector& NewLocation, const FRotator& NewRotation);
@@ -877,49 +878,49 @@ FORCEINLINE_DEBUGGABLE void USceneComponent::ConditionalUpdateComponentToWorld()
 	}
 }
 
-FORCEINLINE_DEBUGGABLE bool USceneComponent::MoveComponent(const FVector& Delta, const FQuat& NewRotation, bool bSweep, FHitResult* Hit, EMoveComponentFlags MoveFlags)
+FORCEINLINE_DEBUGGABLE bool USceneComponent::MoveComponent(const FVector& Delta, const FQuat& NewRotation, bool bSweep, FHitResult* Hit, EMoveComponentFlags MoveFlags, ETeleportType Teleport)
 {
-	return MoveComponentImpl(Delta, NewRotation, bSweep, Hit, MoveFlags);
+	return MoveComponentImpl(Delta, NewRotation, bSweep, Hit, MoveFlags, Teleport);
 }
 
-FORCEINLINE_DEBUGGABLE bool USceneComponent::MoveComponent(const FVector& Delta, const FRotator& NewRotation, bool bSweep, FHitResult* Hit, EMoveComponentFlags MoveFlags)
+FORCEINLINE_DEBUGGABLE bool USceneComponent::MoveComponent(const FVector& Delta, const FRotator& NewRotation, bool bSweep, FHitResult* Hit, EMoveComponentFlags MoveFlags, ETeleportType Teleport)
 {
-	return MoveComponentImpl(Delta, NewRotation.Quaternion(), bSweep, Hit, MoveFlags);
+	return MoveComponentImpl(Delta, NewRotation.Quaternion(), bSweep, Hit, MoveFlags, Teleport);
 }
 
-FORCEINLINE_DEBUGGABLE void USceneComponent::SetRelativeLocation(FVector NewLocation, bool bSweep, FHitResult* OutSweepHitResult)
+FORCEINLINE_DEBUGGABLE void USceneComponent::SetRelativeLocation(FVector NewLocation, bool bSweep, FHitResult* OutSweepHitResult, ETeleportType Teleport)
 {
-	SetRelativeLocationAndRotation(NewLocation, RelativeRotationCache.RotatorToQuat(RelativeRotation), bSweep, OutSweepHitResult);
+	SetRelativeLocationAndRotation(NewLocation, RelativeRotationCache.RotatorToQuat(RelativeRotation), bSweep, OutSweepHitResult, Teleport);
 }
 
-FORCEINLINE_DEBUGGABLE void USceneComponent::SetRelativeRotation(FRotator NewRotation, bool bSweep, FHitResult* OutSweepHitResult)
+FORCEINLINE_DEBUGGABLE void USceneComponent::SetRelativeRotation(FRotator NewRotation, bool bSweep, FHitResult* OutSweepHitResult, ETeleportType Teleport)
 {
-	SetRelativeLocationAndRotation(RelativeLocation, NewRotation, bSweep, OutSweepHitResult);
+	SetRelativeLocationAndRotation(RelativeLocation, NewRotation, bSweep, OutSweepHitResult, Teleport);
 }
 
-FORCEINLINE_DEBUGGABLE void USceneComponent::SetRelativeRotation(const FQuat& NewRotation, bool bSweep, FHitResult* OutSweepHitResult)
+FORCEINLINE_DEBUGGABLE void USceneComponent::SetRelativeRotation(const FQuat& NewRotation, bool bSweep, FHitResult* OutSweepHitResult, ETeleportType Teleport)
 {
-	SetRelativeLocationAndRotation(RelativeLocation, NewRotation, bSweep, OutSweepHitResult);
+	SetRelativeLocationAndRotation(RelativeLocation, NewRotation, bSweep, OutSweepHitResult, Teleport);
 }
 
-FORCEINLINE_DEBUGGABLE void USceneComponent::SetRelativeLocationAndRotation(FVector NewLocation, FRotator NewRotation, bool bSweep, FHitResult* OutSweepHitResult)
+FORCEINLINE_DEBUGGABLE void USceneComponent::SetRelativeLocationAndRotation(FVector NewLocation, FRotator NewRotation, bool bSweep, FHitResult* OutSweepHitResult, ETeleportType Teleport)
 {
-	SetRelativeLocationAndRotation(NewLocation, NewRotation.Quaternion(), bSweep, OutSweepHitResult);
+	SetRelativeLocationAndRotation(NewLocation, NewRotation.Quaternion(), bSweep, OutSweepHitResult, Teleport);
 }
 
-FORCEINLINE_DEBUGGABLE void USceneComponent::AddRelativeLocation(FVector DeltaLocation, bool bSweep, FHitResult* OutSweepHitResult)
+FORCEINLINE_DEBUGGABLE void USceneComponent::AddRelativeLocation(FVector DeltaLocation, bool bSweep, FHitResult* OutSweepHitResult, ETeleportType Teleport)
 {
-	SetRelativeLocationAndRotation(RelativeLocation + DeltaLocation, RelativeRotationCache.RotatorToQuat(RelativeRotation), bSweep, OutSweepHitResult);
+	SetRelativeLocationAndRotation(RelativeLocation + DeltaLocation, RelativeRotationCache.RotatorToQuat(RelativeRotation), bSweep, OutSweepHitResult, Teleport);
 }
 
-FORCEINLINE_DEBUGGABLE void USceneComponent::AddRelativeRotation(FRotator DeltaRotation, bool bSweep, FHitResult* OutSweepHitResult)
+FORCEINLINE_DEBUGGABLE void USceneComponent::AddRelativeRotation(FRotator DeltaRotation, bool bSweep, FHitResult* OutSweepHitResult, ETeleportType Teleport)
 {
-	SetRelativeLocationAndRotation(RelativeLocation, RelativeRotation + DeltaRotation, bSweep, OutSweepHitResult);
+	SetRelativeLocationAndRotation(RelativeLocation, RelativeRotation + DeltaRotation, bSweep, OutSweepHitResult, Teleport);
 }
 
-FORCEINLINE_DEBUGGABLE void USceneComponent::SetWorldRotation(const FRotator NewRotation, bool bSweep, FHitResult* OutSweepHitResult)
+FORCEINLINE_DEBUGGABLE void USceneComponent::SetWorldRotation(const FRotator NewRotation, bool bSweep, FHitResult* OutSweepHitResult, ETeleportType Teleport)
 {
-	SetWorldRotation(NewRotation.Quaternion(), bSweep, OutSweepHitResult);
+	SetWorldRotation(NewRotation.Quaternion(), bSweep, OutSweepHitResult, Teleport);
 }
 
 //////////////////////////////////////////////////////////////////////////
diff --git a/Engine/Source/Runtime/Engine/Classes/Components/SkeletalMeshComponent.h b/Engine/Source/Runtime/Engine/Classes/Components/SkeletalMeshComponent.h
--- a/Engine/Source/Runtime/Engine/Classes/Components/SkeletalMeshComponent.h
+++ b/Engine/Source/Runtime/Engine/Classes/Components/SkeletalMeshComponent.h
@@ -832,7 +832,7 @@ public:
 	virtual void UpdateBounds() override;
 	virtual FBoxSphereBounds CalcBounds(const FTransform& LocalToWorld) const override;
 	virtual bool IsAnySimulatingPhysics() const override;
-	virtual void OnUpdateTransform(bool bSkipPhysicsMove, bool bTeleport = false) override;
+	virtual void OnUpdateTransform(bool bSkipPhysicsMove, ETeleportType Teleport = ETeleportType::None) override;
 	virtual void UpdateOverlaps(TArray<FOverlapInfo> const* PendingOverlaps=NULL, bool bDoNotifies=true, const TArray<FOverlapInfo>* OverlapsAtEndLocation=NULL) override;
 	// End USceneComponent interface.
 
@@ -1061,12 +1061,12 @@ public:
 	 *	Iterate over each physics body in the physics for this mesh, and for each 'kinematic' (ie fixed or default if owner isn't simulating) one, update its
 	 *	transform based on the animated transform.
 	 */
-	void UpdateKinematicBonesToAnim(const TArray<FTransform>& InSpaceBases, bool bTeleport, bool bNeedsSkinning);
+	void UpdateKinematicBonesToAnim(const TArray<FTransform>& InSpaceBases, ETeleportType Teleport, bool bNeedsSkinning);
 
-	DEPRECATED(4.9, "bForceUpdate is no longer used. Please use bTeleport")
+	DEPRECATED(4.9, "bForceUpdate is no longer used. Please use Teleport")
 	void UpdateKinematicBonesToAnim(const TArray<FTransform>& InSpaceBases, bool bTeleport, bool bNeedsSkinning, bool bForceUpdate)
 	{
-		UpdateKinematicBonesToAnim(InSpaceBases, bTeleport || bForceUpdate, bNeedsSkinning);
+		UpdateKinematicBonesToAnim(InSpaceBases, TeleportFlagToEnum(bTeleport || bForceUpdate), bNeedsSkinning);
 	}
 	
 	/**
diff --git a/Engine/Source/Runtime/Engine/Classes/GameFramework/Actor.h b/Engine/Source/Runtime/Engine/Classes/GameFramework/Actor.h
--- a/Engine/Source/Runtime/Engine/Classes/GameFramework/Actor.h
+++ b/Engine/Source/Runtime/Engine/Classes/GameFramework/Actor.h
@@ -638,7 +638,7 @@ public:
 	 *	@return	Whether the location was successfully set (if not swept), or whether movement occurred at all (if swept).
 	 */
 	UFUNCTION(BlueprintCallable, meta=(DisplayName = "SetActorLocation", Keywords="position"), Category="Utilities|Transformation")
-	bool K2_SetActorLocation(FVector NewLocation, bool bSweep, FHitResult& SweepHitResult);
+	bool K2_SetActorLocation(FVector NewLocation, bool bSweep, FHitResult& SweepHitResult, bool bTeleport);
 
 	/** Returns rotation of the RootComponent of this Actor. */
 	UFUNCTION(BlueprintCallable, meta=(DisplayName = "GetActorRotation"), Category="Utilities|Transformation")
@@ -679,7 +679,7 @@ public:
 	 * @param OutSweepHitResult The hit result from the move if swept.
 	 * @return	Whether the location was successfully set if not swept, or whether movement occurred if swept.
 	 */
-	bool SetActorLocation(const FVector& NewLocation, bool bSweep=false, FHitResult* OutSweepHitResult=nullptr);
+	bool SetActorLocation(const FVector& NewLocation, bool bSweep=false, FHitResult* OutSweepHitResult=nullptr, ETeleportType Teleport = ETeleportType::None);
 
 	/** 
 	 * Set the Actor's rotation instantly to the specified rotation.
@@ -701,7 +701,7 @@ public:
 	 * @return	Whether the rotation was successfully set.
 	 */
 	UFUNCTION(BlueprintCallable, Category="Utilities|Transformation", meta=(DisplayName="SetActorLocationAndRotation"))
-	bool K2_SetActorLocationAndRotation(FVector NewLocation, FRotator NewRotation, bool bSweep, FHitResult& SweepHitResult);
+	bool K2_SetActorLocationAndRotation(FVector NewLocation, FRotator NewRotation, bool bSweep, FHitResult& SweepHitResult, bool bTeleport);
 	
 	/** 
 	 * Move the actor instantly to the specified location and rotation.
@@ -712,8 +712,8 @@ public:
 	 * @param OutSweepHitResult	The hit result from the move if swept.
 	 * @return	Whether the rotation was successfully set.
 	 */
-	bool SetActorLocationAndRotation(FVector NewLocation, FRotator NewRotation, bool bSweep=false, FHitResult* OutSweepHitResult=nullptr);
-	bool SetActorLocationAndRotation(FVector NewLocation, const FQuat& NewRotation, bool bSweep=false, FHitResult* OutSweepHitResult=nullptr);
+	bool SetActorLocationAndRotation(FVector NewLocation, FRotator NewRotation, bool bSweep=false, FHitResult* OutSweepHitResult=nullptr, ETeleportType Teleport = ETeleportType::None);
+	bool SetActorLocationAndRotation(FVector NewLocation, const FQuat& NewRotation, bool bSweep=false, FHitResult* OutSweepHitResult=nullptr, ETeleportType Teleport = ETeleportType::None);
 
 	/** Set the Actor's world-space scale. */
 	UFUNCTION(BlueprintCallable, Category="Utilities|Transformation")
@@ -751,7 +751,7 @@ public:
 	 * @param  SweepHitResult		The hit result from the move if swept.
 	 */
 	UFUNCTION(BlueprintCallable, Category="Utilities|Transformation", meta=(DisplayName="AddActorWorldOffset", Keywords="location position"))
-	void K2_AddActorWorldOffset(FVector DeltaLocation, bool bSweep, FHitResult& SweepHitResult);
+	void K2_AddActorWorldOffset(FVector DeltaLocation, bool bSweep, FHitResult& SweepHitResult, bool bTeleport);
 
 	/**
 	 * Adds a delta to the location of this actor in world space.
@@ -760,7 +760,7 @@ public:
 	 * @param  bSweep				Whether to sweep to the destination location, triggering overlaps along the way and stopping at the first blocking hit.
 	 * @param  SweepHitResult		The hit result from the move if swept.
 	 */
-	void AddActorWorldOffset(FVector DeltaLocation, bool bSweep=false, FHitResult* OutSweepHitResult=nullptr);
+	void AddActorWorldOffset(FVector DeltaLocation, bool bSweep=false, FHitResult* OutSweepHitResult=nullptr, ETeleportType Teleport = ETeleportType::None);
 
 
 	/**
@@ -771,15 +771,15 @@ public:
 	 * @param  SweepHitResult		The hit result from the move if swept.
 	 */
 	UFUNCTION(BlueprintCallable, Category="Utilities|Transformation", meta=(DisplayName="AddActorWorldRotation", AdvancedDisplay="bSweep,SweepHitResult"))
-	void K2_AddActorWorldRotation(FRotator DeltaRotation, bool bSweep, FHitResult& SweepHitResult);
-	void AddActorWorldRotation(FRotator DeltaRotation, bool bSweep=false, FHitResult* OutSweepHitResult=nullptr);
-	void AddActorWorldRotation(const FQuat& DeltaRotation, bool bSweep=false, FHitResult* OutSweepHitResult=nullptr);
+	void K2_AddActorWorldRotation(FRotator DeltaRotation, bool bSweep, FHitResult& SweepHitResult, bool bTeleport);
+	void AddActorWorldRotation(FRotator DeltaRotation, bool bSweep=false, FHitResult* OutSweepHitResult=nullptr, ETeleportType Teleport = ETeleportType::None);
+	void AddActorWorldRotation(const FQuat& DeltaRotation, bool bSweep=false, FHitResult* OutSweepHitResult=nullptr, ETeleportType Teleport = ETeleportType::None);
 
 
 	/** Adds a delta to the transform of this actor in world space. Scale is unchanged. */
 	UFUNCTION(BlueprintCallable, Category="Utilities|Transformation", meta=(DisplayName="AddActorWorldTransform"))
-	void K2_AddActorWorldTransform(const FTransform& DeltaTransform, bool bSweep, FHitResult& SweepHitResult);
-	void AddActorWorldTransform(const FTransform& DeltaTransform, bool bSweep=false, FHitResult* OutSweepHitResult=nullptr);
+	void K2_AddActorWorldTransform(const FTransform& DeltaTransform, bool bSweep, FHitResult& SweepHitResult, bool bTeleport);
+	void AddActorWorldTransform(const FTransform& DeltaTransform, bool bSweep=false, FHitResult* OutSweepHitResult=nullptr, ETeleportType Teleport = ETeleportType::None);
 
 
 	/** 
@@ -787,27 +787,27 @@ public:
 	 * @param bSweep		Whether to sweep to the destination location, stopping short of the target if blocked by something.
 	 */
 	UFUNCTION(BlueprintCallable, Category="Utilities|Transformation", meta=(DisplayName="SetActorTransform"))
-	bool K2_SetActorTransform(const FTransform& NewTransform, bool bSweep, FHitResult& SweepHitResult);
-	bool SetActorTransform(const FTransform& NewTransform, bool bSweep=false, FHitResult* OutSweepHitResult=nullptr);
+	bool K2_SetActorTransform(const FTransform& NewTransform, bool bSweep, FHitResult& SweepHitResult, bool bTeleport);
+	bool SetActorTransform(const FTransform& NewTransform, bool bSweep=false, FHitResult* OutSweepHitResult=nullptr, ETeleportType Teleport = ETeleportType::None);
 
 
 	/** Adds a delta to the location of this component in its local reference frame */
 	UFUNCTION(BlueprintCallable, Category="Utilities|Transformation", meta=(DisplayName="AddActorLocalOffset", Keywords="location position"))
-	void K2_AddActorLocalOffset(FVector DeltaLocation, bool bSweep, FHitResult& SweepHitResult);
-	void AddActorLocalOffset(FVector DeltaLocation, bool bSweep=false, FHitResult* OutSweepHitResult=nullptr);
+	void K2_AddActorLocalOffset(FVector DeltaLocation, bool bSweep, FHitResult& SweepHitResult, bool bTeleport);
+	void AddActorLocalOffset(FVector DeltaLocation, bool bSweep=false, FHitResult* OutSweepHitResult=nullptr, ETeleportType Teleport = ETeleportType::None);
 
 
 	/** Adds a delta to the rotation of this component in its local reference frame */
 	UFUNCTION(BlueprintCallable, Category="Utilities|Transformation", meta=(DisplayName="AddActorLocalRotation", AdvancedDisplay="bSweep,SweepHitResult"))
-	void K2_AddActorLocalRotation(FRotator DeltaRotation, bool bSweep, FHitResult& SweepHitResult);
-	void AddActorLocalRotation(FRotator DeltaRotation, bool bSweep=false, FHitResult* OutSweepHitResult=nullptr);
-	void AddActorLocalRotation(const FQuat& DeltaRotation, bool bSweep=false, FHitResult* OutSweepHitResult=nullptr);
+	void K2_AddActorLocalRotation(FRotator DeltaRotation, bool bSweep, FHitResult& SweepHitResult, bool bTeleport);
+	void AddActorLocalRotation(FRotator DeltaRotation, bool bSweep=false, FHitResult* OutSweepHitResult=nullptr, ETeleportType Teleport = ETeleportType::None);
+	void AddActorLocalRotation(const FQuat& DeltaRotation, bool bSweep=false, FHitResult* OutSweepHitResult=nullptr, ETeleportType Teleport = ETeleportType::None);
 
 
 	/** Adds a delta to the transform of this component in its local reference frame */
 	UFUNCTION(BlueprintCallable, Category="Utilities|Transformation", meta=(DisplayName="AddActorLocalTransform"))
-	void K2_AddActorLocalTransform(const FTransform& NewTransform, bool bSweep, FHitResult& SweepHitResult);
-	void AddActorLocalTransform(const FTransform& NewTransform, bool bSweep=false, FHitResult* OutSweepHitResult=nullptr);
+	void K2_AddActorLocalTransform(const FTransform& NewTransform, bool bSweep, FHitResult& SweepHitResult, bool bTeleport);
+	void AddActorLocalTransform(const FTransform& NewTransform, bool bSweep=false, FHitResult* OutSweepHitResult=nullptr, ETeleportType Teleport = ETeleportType::None);
 
 
 	/**
@@ -816,8 +816,8 @@ public:
 	 * @param bSweep				Should we sweep to the destination location. If true, will stop short of the target if blocked by something
 	 */
 	UFUNCTION(BlueprintCallable, Category="Utilities|Transformation", meta=(DisplayName="SetActorRelativeLocation"))
-	void K2_SetActorRelativeLocation(FVector NewRelativeLocation, bool bSweep, FHitResult& SweepHitResult);
-	void SetActorRelativeLocation(FVector NewRelativeLocation, bool bSweep=false, FHitResult* OutSweepHitResult=nullptr);
+	void K2_SetActorRelativeLocation(FVector NewRelativeLocation, bool bSweep, FHitResult& SweepHitResult, bool bTeleport);
+	void SetActorRelativeLocation(FVector NewRelativeLocation, bool bSweep=false, FHitResult* OutSweepHitResult=nullptr, ETeleportType Teleport = ETeleportType::None);
 
 	/**
 	 * Set the actor's RootComponent to the specified relative rotation
@@ -825,9 +825,9 @@ public:
 	 * @param bSweep					Should we sweep to the destination rotation. If true, will stop short of the target if blocked by something
 	 */
 	UFUNCTION(BlueprintCallable, Category="Utilities|Transformation", meta=(DisplayName="SetActorRelativeRotation", AdvancedDisplay="bSweep,SweepHitResult"))
-	void K2_SetActorRelativeRotation(FRotator NewRelativeRotation, bool bSweep, FHitResult& SweepHitResult);
-	void SetActorRelativeRotation(FRotator NewRelativeRotation, bool bSweep=false, FHitResult* OutSweepHitResult=nullptr);
-	void SetActorRelativeRotation(const FQuat& NewRelativeRotation, bool bSweep=false, FHitResult* OutSweepHitResult=nullptr);
+	void K2_SetActorRelativeRotation(FRotator NewRelativeRotation, bool bSweep, FHitResult& SweepHitResult, bool bTeleport);
+	void SetActorRelativeRotation(FRotator NewRelativeRotation, bool bSweep=false, FHitResult* OutSweepHitResult=nullptr, ETeleportType Teleport = ETeleportType::None);
+	void SetActorRelativeRotation(const FQuat& NewRelativeRotation, bool bSweep=false, FHitResult* OutSweepHitResult=nullptr, ETeleportType Teleport = ETeleportType::None);
 
 	/**
 	 * Set the actor's RootComponent to the specified relative transform
@@ -835,8 +835,8 @@ public:
 	 * @param bSweep					Should we sweep to the destination transform. If true, will stop short of the target if blocked by something
 	 */
 	UFUNCTION(BlueprintCallable, Category="Utilities|Transformation", meta=(DisplayName="SetActorRelativeTransform"))
-	void K2_SetActorRelativeTransform(const FTransform& NewRelativeTransform, bool bSweep, FHitResult& SweepHitResult);
-	void SetActorRelativeTransform(const FTransform& NewRelativeTransform, bool bSweep=false, FHitResult* OutSweepHitResult=nullptr);
+	void K2_SetActorRelativeTransform(const FTransform& NewRelativeTransform, bool bSweep, FHitResult& SweepHitResult, bool bTeleport);
+	void SetActorRelativeTransform(const FTransform& NewRelativeTransform, bool bSweep=false, FHitResult* OutSweepHitResult=nullptr, ETeleportType Teleport = ETeleportType::None);
 
 	/**
 	 * Set the actor's RootComponent to the specified relative scale 3d
diff --git a/Engine/Source/Runtime/Engine/Classes/PhysicsEngine/BodyInstance.h b/Engine/Source/Runtime/Engine/Classes/PhysicsEngine/BodyInstance.h
--- a/Engine/Source/Runtime/Engine/Classes/PhysicsEngine/BodyInstance.h
+++ b/Engine/Source/Runtime/Engine/Classes/PhysicsEngine/BodyInstance.h
@@ -11,6 +11,7 @@
 struct FCollisionShape;
 struct FConstraintInstance;
 class UPhysicsConstraintComponent;
+enum class ETeleportType;
 
 /** Delegate for applying custom physics forces upon the body. Can be passed to "AddCustomPhysics" so 
   * custom forces and torques can be calculated induvidually for every physics substep.
@@ -734,6 +735,9 @@ public:
 	 *	Move the physics body to a new pose.
 	 *	@param	bTeleport	If true, no velocity is inferred on the kinematic body from this movement, but it moves right away.
 	 */
+	void SetBodyTransform(const FTransform& NewTransform, ETeleportType Teleport);
+
+	DEPRECATED(4.9, "Please pass the teleport flag using the new ETeleportType enum ")
 	void SetBodyTransform(const FTransform& NewTransform, bool bTeleport);
 
 	/** Get current velocity in world space from physics body. */
diff --git a/Engine/Source/Runtime/Engine/Private/Actor.cpp b/Engine/Source/Runtime/Engine/Private/Actor.cpp
--- a/Engine/Source/Runtime/Engine/Private/Actor.cpp
+++ b/Engine/Source/Runtime/Engine/Private/Actor.cpp
@@ -320,7 +320,7 @@ bool AActor::TeleportTo( const FVector& DestLocation, const FRotator& DestRotati
 			else
 			{
 				FVector const Delta = NewLocation - PrevLocation;
-				bTeleportSucceeded = ActorPrimComp->MoveComponent(Delta, DestRotation, false);
+				bTeleportSucceeded = ActorPrimComp->MoveComponent(Delta, DestRotation, false, nullptr, MOVECOMP_NoFlags, ETeleportType::TeleportPhysics);
 			}
 			if( bTeleportSucceeded )
 			{
@@ -331,7 +331,7 @@ bool AActor::TeleportTo( const FVector& DestLocation, const FRotator& DestRotati
 	else if (RootComponent)
 	{
 		// not a primitivecomponent, just set directly
-		GetRootComponent()->SetWorldLocationAndRotation(NewLocation, DestRotation);
+		GetRootComponent()->SetWorldLocationAndRotation(NewLocation, DestRotation, false, nullptr, ETeleportType::TeleportPhysics);
 	}
 
 	return bTeleportSucceeded; 
@@ -2686,12 +2686,12 @@ FVector AActor::GetInputVectorAxisValue(const FKey InputAxisKey) const
 	return Value;
 }
 
-bool AActor::SetActorLocation(const FVector& NewLocation, bool bSweep, FHitResult* OutSweepHitResult)
+bool AActor::SetActorLocation(const FVector& NewLocation, bool bSweep, FHitResult* OutSweepHitResult, ETeleportType Teleport)
 {
 	if (RootComponent)
 	{
 		const FVector Delta = NewLocation - GetActorLocation();
-		return RootComponent->MoveComponent(Delta, GetActorQuat(), bSweep, OutSweepHitResult);
+		return RootComponent->MoveComponent(Delta, GetActorQuat(), bSweep, OutSweepHitResult, MOVECOMP_NoFlags, Teleport);
 	}
 	else if (OutSweepHitResult)
 	{
@@ -2721,12 +2721,12 @@ bool AActor::SetActorRotation(const FQuat& NewRotation)
 	return false;
 }
 
-bool AActor::SetActorLocationAndRotation(FVector NewLocation, FRotator NewRotation, bool bSweep, FHitResult* OutSweepHitResult)
+bool AActor::SetActorLocationAndRotation(FVector NewLocation, FRotator NewRotation, bool bSweep, FHitResult* OutSweepHitResult, ETeleportType Teleport)
 {
 	if (RootComponent)
 	{
 		const FVector Delta = NewLocation - GetActorLocation();
-		return RootComponent->MoveComponent(Delta, NewRotation, bSweep, OutSweepHitResult);
+		return RootComponent->MoveComponent(Delta, NewRotation, bSweep, OutSweepHitResult, MOVECOMP_NoFlags, Teleport);
 	}
 	else if (OutSweepHitResult)
 	{
@@ -2736,12 +2736,12 @@ bool AActor::SetActorLocationAndRotation(FVector NewLocation, FRotator NewRotati
 	return false;
 }
 
-bool AActor::SetActorLocationAndRotation(FVector NewLocation, const FQuat& NewRotation, bool bSweep, FHitResult* OutSweepHitResult)
+bool AActor::SetActorLocationAndRotation(FVector NewLocation, const FQuat& NewRotation, bool bSweep, FHitResult* OutSweepHitResult, ETeleportType Teleport)
 {
 	if (RootComponent)
 	{
 		const FVector Delta = NewLocation - GetActorLocation();
-		return RootComponent->MoveComponent(Delta, NewRotation, bSweep, OutSweepHitResult);
+		return RootComponent->MoveComponent(Delta, NewRotation, bSweep, OutSweepHitResult, MOVECOMP_NoFlags, Teleport);
 	}
 	else if (OutSweepHitResult)
 	{
@@ -2769,11 +2769,11 @@ FVector AActor::GetActorScale3D() const
 	return FVector(1,1,1);
 }
 
-void AActor::AddActorWorldOffset(FVector DeltaLocation, bool bSweep, FHitResult* OutSweepHitResult)
+void AActor::AddActorWorldOffset(FVector DeltaLocation, bool bSweep, FHitResult* OutSweepHitResult, ETeleportType Teleport)
 {
 	if (RootComponent)
 	{
-		RootComponent->AddWorldOffset(DeltaLocation, bSweep, OutSweepHitResult);
+		RootComponent->AddWorldOffset(DeltaLocation, bSweep, OutSweepHitResult, Teleport);
 	}
 	else if (OutSweepHitResult)
 	{
@@ -2781,11 +2781,11 @@ void AActor::AddActorWorldOffset(FVector DeltaLocation, bool bSweep, FHitResult*
 	}
 }
 
-void AActor::AddActorWorldRotation(FRotator DeltaRotation, bool bSweep, FHitResult* OutSweepHitResult)
+void AActor::AddActorWorldRotation(FRotator DeltaRotation, bool bSweep, FHitResult* OutSweepHitResult, ETeleportType Teleport)
 {
 	if (RootComponent)
 	{
-		RootComponent->AddWorldRotation(DeltaRotation, bSweep, OutSweepHitResult);
+		RootComponent->AddWorldRotation(DeltaRotation, bSweep, OutSweepHitResult, Teleport);
 	}
 	else if (OutSweepHitResult)
 	{
@@ -2793,11 +2793,11 @@ void AActor::AddActorWorldRotation(FRotator DeltaRotation, bool bSweep, FHitResu
 	}
 }
 
-void AActor::AddActorWorldRotation(const FQuat& DeltaRotation, bool bSweep, FHitResult* OutSweepHitResult)
+void AActor::AddActorWorldRotation(const FQuat& DeltaRotation, bool bSweep, FHitResult* OutSweepHitResult, ETeleportType Teleport)
 {
 	if (RootComponent)
 	{
-		RootComponent->AddWorldRotation(DeltaRotation, bSweep, OutSweepHitResult);
+		RootComponent->AddWorldRotation(DeltaRotation, bSweep, OutSweepHitResult, Teleport);
 	}
 	else if (OutSweepHitResult)
 	{
@@ -2805,11 +2805,11 @@ void AActor::AddActorWorldRotation(const FQuat& DeltaRotation, bool bSweep, FHit
 	}
 }
 
-void AActor::AddActorWorldTransform(const FTransform& DeltaTransform, bool bSweep, FHitResult* OutSweepHitResult)
+void AActor::AddActorWorldTransform(const FTransform& DeltaTransform, bool bSweep, FHitResult* OutSweepHitResult, ETeleportType Teleport)
 {
 	if (RootComponent)
 	{
-		RootComponent->AddWorldTransform(DeltaTransform, bSweep, OutSweepHitResult);
+		RootComponent->AddWorldTransform(DeltaTransform, bSweep, OutSweepHitResult, Teleport);
 	}
 	else if (OutSweepHitResult)
 	{
@@ -2818,7 +2818,7 @@ void AActor::AddActorWorldTransform(const FTransform& DeltaTransform, bool bSwee
 }
 
 
-bool AActor::SetActorTransform(const FTransform& NewTransform, bool bSweep, FHitResult* OutSweepHitResult)
+bool AActor::SetActorTransform(const FTransform& NewTransform, bool bSweep, FHitResult* OutSweepHitResult, ETeleportType Teleport)
 {
 	// we have seen this gets NAN from kismet, and would like to see if this
 	// happens, and if so, something else is giving NAN as output
@@ -2826,7 +2826,7 @@ bool AActor::SetActorTransform(const FTransform& NewTransform, bool bSweep, FHit
 	{
 		if (ensure(!NewTransform.ContainsNaN()))
 		{
-			RootComponent->SetWorldTransform(NewTransform, bSweep, OutSweepHitResult);
+			RootComponent->SetWorldTransform(NewTransform, bSweep, OutSweepHitResult, Teleport);
 		}
 		else
 		{
@@ -2846,11 +2846,11 @@ bool AActor::SetActorTransform(const FTransform& NewTransform, bool bSweep, FHit
 	return false;
 }
 
-void AActor::AddActorLocalOffset(FVector DeltaLocation, bool bSweep, FHitResult* OutSweepHitResult)
+void AActor::AddActorLocalOffset(FVector DeltaLocation, bool bSweep, FHitResult* OutSweepHitResult, ETeleportType Teleport)
 {
 	if(RootComponent)
 	{
-		RootComponent->AddLocalOffset(DeltaLocation, bSweep, OutSweepHitResult);
+		RootComponent->AddLocalOffset(DeltaLocation, bSweep, OutSweepHitResult, Teleport);
 	}
 	else if (OutSweepHitResult)
 	{
@@ -2858,11 +2858,11 @@ void AActor::AddActorLocalOffset(FVector DeltaLocation, bool bSweep, FHitResult*
 	}
 }
 
-void AActor::AddActorLocalRotation(FRotator DeltaRotation, bool bSweep, FHitResult* OutSweepHitResult)
+void AActor::AddActorLocalRotation(FRotator DeltaRotation, bool bSweep, FHitResult* OutSweepHitResult, ETeleportType Teleport)
 {
 	if(RootComponent)
 	{
-		RootComponent->AddLocalRotation(DeltaRotation, bSweep, OutSweepHitResult);
+		RootComponent->AddLocalRotation(DeltaRotation, bSweep, OutSweepHitResult, Teleport);
 	}
 	else if (OutSweepHitResult)
 	{
@@ -2870,11 +2870,11 @@ void AActor::AddActorLocalRotation(FRotator DeltaRotation, bool bSweep, FHitResu
 	}
 }
 
-void AActor::AddActorLocalRotation(const FQuat& DeltaRotation, bool bSweep, FHitResult* OutSweepHitResult)
+void AActor::AddActorLocalRotation(const FQuat& DeltaRotation, bool bSweep, FHitResult* OutSweepHitResult, ETeleportType Teleport)
 {
 	if (RootComponent)
 	{
-		RootComponent->AddLocalRotation(DeltaRotation, bSweep, OutSweepHitResult);
+		RootComponent->AddLocalRotation(DeltaRotation, bSweep, OutSweepHitResult, Teleport);
 	}
 	else if (OutSweepHitResult)
 	{
@@ -2882,11 +2882,11 @@ void AActor::AddActorLocalRotation(const FQuat& DeltaRotation, bool bSweep, FHit
 	}
 }
 
-void AActor::AddActorLocalTransform(const FTransform& NewTransform, bool bSweep, FHitResult* OutSweepHitResult)
+void AActor::AddActorLocalTransform(const FTransform& NewTransform, bool bSweep, FHitResult* OutSweepHitResult, ETeleportType Teleport)
 {
 	if(RootComponent)
 	{
-		RootComponent->AddLocalTransform(NewTransform, bSweep, OutSweepHitResult);
+		RootComponent->AddLocalTransform(NewTransform, bSweep, OutSweepHitResult, Teleport);
 	}
 	else if (OutSweepHitResult)
 	{
@@ -2894,11 +2894,11 @@ void AActor::AddActorLocalTransform(const FTransform& NewTransform, bool bSweep,
 	}
 }
 
-void AActor::SetActorRelativeLocation(FVector NewRelativeLocation, bool bSweep, FHitResult* OutSweepHitResult)
+void AActor::SetActorRelativeLocation(FVector NewRelativeLocation, bool bSweep, FHitResult* OutSweepHitResult, ETeleportType Teleport)
 {
 	if (RootComponent)
 	{
-		RootComponent->SetRelativeLocation(NewRelativeLocation, bSweep, OutSweepHitResult);
+		RootComponent->SetRelativeLocation(NewRelativeLocation, bSweep, OutSweepHitResult, Teleport);
 	}
 	else if (OutSweepHitResult)
 	{
@@ -2906,11 +2906,11 @@ void AActor::SetActorRelativeLocation(FVector NewRelativeLocation, bool bSweep,
 	}
 }
 
-void AActor::SetActorRelativeRotation(FRotator NewRelativeRotation, bool bSweep, FHitResult* OutSweepHitResult)
+void AActor::SetActorRelativeRotation(FRotator NewRelativeRotation, bool bSweep, FHitResult* OutSweepHitResult, ETeleportType Teleport)
 {
 	if (RootComponent)
 	{
-		RootComponent->SetRelativeRotation(NewRelativeRotation, bSweep, OutSweepHitResult);
+		RootComponent->SetRelativeRotation(NewRelativeRotation, bSweep, OutSweepHitResult, Teleport);
 	}
 	else if (OutSweepHitResult)
 	{
@@ -2918,11 +2918,11 @@ void AActor::SetActorRelativeRotation(FRotator NewRelativeRotation, bool bSweep,
 	}
 }
 
-void AActor::SetActorRelativeRotation(const FQuat& NewRelativeRotation, bool bSweep, FHitResult* OutSweepHitResult)
+void AActor::SetActorRelativeRotation(const FQuat& NewRelativeRotation, bool bSweep, FHitResult* OutSweepHitResult, ETeleportType Teleport)
 {
 	if (RootComponent)
 	{
-		RootComponent->SetRelativeRotation(NewRelativeRotation, bSweep, OutSweepHitResult);
+		RootComponent->SetRelativeRotation(NewRelativeRotation, bSweep, OutSweepHitResult, Teleport);
 	}
 	else if (OutSweepHitResult)
 	{
@@ -2930,11 +2930,11 @@ void AActor::SetActorRelativeRotation(const FQuat& NewRelativeRotation, bool bSw
 	}
 }
 
-void AActor::SetActorRelativeTransform(const FTransform& NewRelativeTransform, bool bSweep, FHitResult* OutSweepHitResult)
+void AActor::SetActorRelativeTransform(const FTransform& NewRelativeTransform, bool bSweep, FHitResult* OutSweepHitResult, ETeleportType Teleport)
 {
 	if (RootComponent)
 	{
-		RootComponent->SetRelativeTransform(NewRelativeTransform, bSweep, OutSweepHitResult);
+		RootComponent->SetRelativeTransform(NewRelativeTransform, bSweep, OutSweepHitResult, Teleport);
 	}
 	else if (OutSweepHitResult)
 	{
@@ -3859,64 +3859,64 @@ const int32 AActor::GetNumUncachedStaticLightingInteractions() const
 // Note: we pass null for the hit result if not sweeping, for better perf.
 // This assumes this K2 function is only used by blueprints, which initializes the param for each function call.
 
-bool AActor::K2_SetActorLocation(FVector NewLocation, bool bSweep, FHitResult& SweepHitResult)
+bool AActor::K2_SetActorLocation(FVector NewLocation, bool bSweep, FHitResult& SweepHitResult, bool bTeleport)
 {
-	return SetActorLocation(NewLocation, bSweep, (bSweep ? &SweepHitResult : nullptr));
+	return SetActorLocation(NewLocation, bSweep, (bSweep ? &SweepHitResult : nullptr), TeleportFlagToEnum(bTeleport));
 }
 
-bool AActor::K2_SetActorLocationAndRotation(FVector NewLocation, FRotator NewRotation, bool bSweep, FHitResult& SweepHitResult)
+bool AActor::K2_SetActorLocationAndRotation(FVector NewLocation, FRotator NewRotation, bool bSweep, FHitResult& SweepHitResult, bool bTeleport)
 {
-	return SetActorLocationAndRotation(NewLocation, NewRotation, bSweep, (bSweep ? &SweepHitResult : nullptr));
+	return SetActorLocationAndRotation(NewLocation, NewRotation, bSweep, (bSweep ? &SweepHitResult : nullptr), TeleportFlagToEnum(bTeleport));
 }
 
-void AActor::K2_AddActorWorldOffset(FVector DeltaLocation, bool bSweep, FHitResult& SweepHitResult)
+void AActor::K2_AddActorWorldOffset(FVector DeltaLocation, bool bSweep, FHitResult& SweepHitResult, bool bTeleport)
 {
-	AddActorWorldOffset(DeltaLocation, bSweep, (bSweep ? &SweepHitResult : nullptr));
+	AddActorWorldOffset(DeltaLocation, bSweep, (bSweep ? &SweepHitResult : nullptr), TeleportFlagToEnum(bTeleport));
 }
 
-void AActor::K2_AddActorWorldRotation(FRotator DeltaRotation, bool bSweep, FHitResult& SweepHitResult)
+void AActor::K2_AddActorWorldRotation(FRotator DeltaRotation, bool bSweep, FHitResult& SweepHitResult, bool bTeleport)
 {
-	AddActorWorldRotation(DeltaRotation, bSweep, (bSweep ? &SweepHitResult : nullptr));
+	AddActorWorldRotation(DeltaRotation, bSweep, (bSweep ? &SweepHitResult : nullptr), TeleportFlagToEnum(bTeleport));
 }
 
-void AActor::K2_AddActorWorldTransform(const FTransform& DeltaTransform, bool bSweep, FHitResult& SweepHitResult)
+void AActor::K2_AddActorWorldTransform(const FTransform& DeltaTransform, bool bSweep, FHitResult& SweepHitResult, bool bTeleport)
 {
-	AddActorWorldTransform(DeltaTransform, bSweep, (bSweep ? &SweepHitResult : nullptr));
+	AddActorWorldTransform(DeltaTransform, bSweep, (bSweep ? &SweepHitResult : nullptr), TeleportFlagToEnum(bTeleport));
 }
 
-bool AActor::K2_SetActorTransform(const FTransform& NewTransform, bool bSweep, FHitResult& SweepHitResult)
+bool AActor::K2_SetActorTransform(const FTransform& NewTransform, bool bSweep, FHitResult& SweepHitResult, bool bTeleport)
 {
-	return SetActorTransform(NewTransform, bSweep, (bSweep ? &SweepHitResult : nullptr));
+	return SetActorTransform(NewTransform, bSweep, (bSweep ? &SweepHitResult : nullptr), TeleportFlagToEnum(bTeleport));
 }
 
-void AActor::K2_AddActorLocalOffset(FVector DeltaLocation, bool bSweep, FHitResult& SweepHitResult)
+void AActor::K2_AddActorLocalOffset(FVector DeltaLocation, bool bSweep, FHitResult& SweepHitResult, bool bTeleport)
 {
-	AddActorLocalOffset(DeltaLocation, bSweep, (bSweep ? &SweepHitResult : nullptr));
+	AddActorLocalOffset(DeltaLocation, bSweep, (bSweep ? &SweepHitResult : nullptr), TeleportFlagToEnum(bTeleport));
 }
 
-void AActor::K2_AddActorLocalRotation(FRotator DeltaRotation, bool bSweep, FHitResult& SweepHitResult)
+void AActor::K2_AddActorLocalRotation(FRotator DeltaRotation, bool bSweep, FHitResult& SweepHitResult, bool bTeleport)
 {
-	AddActorLocalRotation(DeltaRotation, bSweep, (bSweep ? &SweepHitResult : nullptr));
+	AddActorLocalRotation(DeltaRotation, bSweep, (bSweep ? &SweepHitResult : nullptr), TeleportFlagToEnum(bTeleport));
 }
 
-void AActor::K2_AddActorLocalTransform(const FTransform& NewTransform, bool bSweep, FHitResult& SweepHitResult)
+void AActor::K2_AddActorLocalTransform(const FTransform& NewTransform, bool bSweep, FHitResult& SweepHitResult, bool bTeleport)
 {
-	AddActorLocalTransform(NewTransform, bSweep, (bSweep ? &SweepHitResult : nullptr));
+	AddActorLocalTransform(NewTransform, bSweep, (bSweep ? &SweepHitResult : nullptr), TeleportFlagToEnum(bTeleport));
 }
 
-void AActor::K2_SetActorRelativeLocation(FVector NewRelativeLocation, bool bSweep, FHitResult& SweepHitResult)
+void AActor::K2_SetActorRelativeLocation(FVector NewRelativeLocation, bool bSweep, FHitResult& SweepHitResult, bool bTeleport)
 {
-	SetActorRelativeLocation(NewRelativeLocation, bSweep, (bSweep ? &SweepHitResult : nullptr));
+	SetActorRelativeLocation(NewRelativeLocation, bSweep, (bSweep ? &SweepHitResult : nullptr), TeleportFlagToEnum(bTeleport));
 }
 
-void AActor::K2_SetActorRelativeRotation(FRotator NewRelativeRotation, bool bSweep, FHitResult& SweepHitResult)
+void AActor::K2_SetActorRelativeRotation(FRotator NewRelativeRotation, bool bSweep, FHitResult& SweepHitResult, bool bTeleport)
 {
-	SetActorRelativeRotation(NewRelativeRotation, bSweep, (bSweep ? &SweepHitResult : nullptr));
+	SetActorRelativeRotation(NewRelativeRotation, bSweep, (bSweep ? &SweepHitResult : nullptr), TeleportFlagToEnum(bTeleport));
 }
 
-void AActor::K2_SetActorRelativeTransform(const FTransform& NewRelativeTransform, bool bSweep, FHitResult& SweepHitResult)
+void AActor::K2_SetActorRelativeTransform(const FTransform& NewRelativeTransform, bool bSweep, FHitResult& SweepHitResult, bool bTeleport)
 {
-	SetActorRelativeTransform(NewRelativeTransform, bSweep, (bSweep ? &SweepHitResult : nullptr));
+	SetActorRelativeTransform(NewRelativeTransform, bSweep, (bSweep ? &SweepHitResult : nullptr), TeleportFlagToEnum(bTeleport));
 }
 
 
diff --git a/Engine/Source/Runtime/Engine/Private/Components/AudioComponent.cpp b/Engine/Source/Runtime/Engine/Private/Components/AudioComponent.cpp
--- a/Engine/Source/Runtime/Engine/Private/Components/AudioComponent.cpp
+++ b/Engine/Source/Runtime/Engine/Private/Components/AudioComponent.cpp
@@ -92,9 +92,9 @@ void UAudioComponent::SetSound( USoundBase* NewSound )
 	}
 }
 
-void UAudioComponent::OnUpdateTransform(bool bSkipPhysicsMove, bool bTeleport)
+void UAudioComponent::OnUpdateTransform(bool bSkipPhysicsMove, ETeleportType Teleport)
 {
-	Super::OnUpdateTransform(bSkipPhysicsMove, bTeleport);
+	Super::OnUpdateTransform(bSkipPhysicsMove, Teleport);
 
 	if (bIsActive && !bPreviewComponent)
 	{
diff --git a/Engine/Source/Runtime/Engine/Private/Components/DestructibleComponent.cpp b/Engine/Source/Runtime/Engine/Private/Components/DestructibleComponent.cpp
--- a/Engine/Source/Runtime/Engine/Private/Components/DestructibleComponent.cpp
+++ b/Engine/Source/Runtime/Engine/Private/Components/DestructibleComponent.cpp
@@ -128,10 +128,10 @@ bool IsImpactDamageEnabled(const UDestructibleMesh* TheDestructibleMesh, int32 L
 	}
 }
 
-void UDestructibleComponent::OnUpdateTransform(bool bSkipPhysicsMove, bool bTeleport)
+void UDestructibleComponent::OnUpdateTransform(bool bSkipPhysicsMove, ETeleportType Teleport)
 {
 	// We are handling the physics move below, so don't handle it at higher levels
-	Super::OnUpdateTransform(true, bTeleport);
+	Super::OnUpdateTransform(true, Teleport);
 
 	if (SkeletalMesh == NULL)
 	{
diff --git a/Engine/Source/Runtime/Engine/Private/Components/PrimitiveComponent.cpp b/Engine/Source/Runtime/Engine/Private/Components/PrimitiveComponent.cpp
--- a/Engine/Source/Runtime/Engine/Private/Components/PrimitiveComponent.cpp
+++ b/Engine/Source/Runtime/Engine/Private/Components/PrimitiveComponent.cpp
@@ -507,20 +507,20 @@ void UPrimitiveComponent::EnsurePhysicsStateCreated()
 	}
 }
 
-void UPrimitiveComponent::OnUpdateTransform(bool bSkipPhysicsMove, bool bTeleport)
+void UPrimitiveComponent::OnUpdateTransform(bool bSkipPhysicsMove, ETeleportType Teleport)
 {
-	Super::OnUpdateTransform(bSkipPhysicsMove);
+	Super::OnUpdateTransform(bSkipPhysicsMove, Teleport);
 
 	// Always send new transform to physics
 	if(bPhysicsStateCreated && !bSkipPhysicsMove)
 	{
-		SendPhysicsTransform(bTeleport);
+		SendPhysicsTransform(Teleport);
 	}
 }
 
-void UPrimitiveComponent::SendPhysicsTransform(bool bTeleport)
+void UPrimitiveComponent::SendPhysicsTransform(ETeleportType Teleport)
 {
-	BodyInstance.SetBodyTransform(ComponentToWorld, bTeleport);
+	BodyInstance.SetBodyTransform(ComponentToWorld, Teleport);
 	BodyInstance.UpdateBodyScale(ComponentToWorld.GetScale3D());
 }
 
@@ -1347,7 +1347,7 @@ FCollisionShape UPrimitiveComponent::GetCollisionShape(float Inflation) const
 }
 
 
-bool UPrimitiveComponent::MoveComponentImpl( const FVector& Delta, const FQuat& NewRotationQuat, bool bSweep, FHitResult* OutHit, EMoveComponentFlags MoveFlags)
+bool UPrimitiveComponent::MoveComponentImpl( const FVector& Delta, const FQuat& NewRotationQuat, bool bSweep, FHitResult* OutHit, EMoveComponentFlags MoveFlags, ETeleportType Teleport)
 {
 	SCOPE_CYCLE_COUNTER(STAT_MoveComponentTime);
 
@@ -1411,7 +1411,7 @@ bool UPrimitiveComponent::MoveComponentImpl( const FVector& Delta, const FQuat&
 	if ( !bSweep )
 	{
 		// not sweeping, just go directly to the new transform
-		bMoved = InternalSetWorldLocationAndRotation(TraceEnd, NewRotationQuat, bSkipPhysicsMove, !bSweep);
+		bMoved = InternalSetWorldLocationAndRotation(TraceEnd, NewRotationQuat, bSkipPhysicsMove, Teleport);
 		bRotationOnly = (DeltaSizeSq == 0);
 		bIncludesOverlapsAtEnd = bRotationOnly && (AreSymmetricRotations(InitialRotationQuat, NewRotationQuat, GetComponentScale())) && IsCollisionEnabled();
 	}
@@ -1583,7 +1583,7 @@ bool UPrimitiveComponent::MoveComponentImpl( const FVector& Delta, const FQuat&
 		}
 
 		// Update the location.  This will teleport any child components as well (not sweep).
-		bMoved = InternalSetWorldLocationAndRotation(NewLocation, NewRotationQuat, bSkipPhysicsMove);
+		bMoved = InternalSetWorldLocationAndRotation(NewLocation, NewRotationQuat, bSkipPhysicsMove, Teleport);
 	}
 
 	// Handle overlap notifications.
diff --git a/Engine/Source/Runtime/Engine/Private/Components/SceneComponent.cpp b/Engine/Source/Runtime/Engine/Private/Components/SceneComponent.cpp
--- a/Engine/Source/Runtime/Engine/Private/Components/SceneComponent.cpp
+++ b/Engine/Source/Runtime/Engine/Private/Components/SceneComponent.cpp
@@ -319,11 +319,11 @@ FTransform USceneComponent::CalcNewComponentToWorld(const FTransform& NewRelativ
 	}
 }
 
-void USceneComponent::OnUpdateTransform(bool bSkipPhysicsMove, bool bTeleport)
+void USceneComponent::OnUpdateTransform(bool bSkipPhysicsMove, ETeleportType Teleport)
 {
 }
 
-void USceneComponent::UpdateComponentToWorldWithParent(USceneComponent* Parent, bool bSkipPhysicsMove, const FQuat& RelativeRotationQuat, bool bTeleport)
+void USceneComponent::UpdateComponentToWorldWithParent(USceneComponent* Parent, bool bSkipPhysicsMove, const FQuat& RelativeRotationQuat, ETeleportType Teleport)
 {
 	// If our parent hasn't been updated before, we'll need walk up our parent attach hierarchy
 	if (Parent && !Parent->bWorldToComponentUpdated)
@@ -359,7 +359,7 @@ void USceneComponent::UpdateComponentToWorldWithParent(USceneComponent* Parent,
 	{
 		// Update transform
 		ComponentToWorld = NewTransform;
-		PropagateTransformUpdate(true, bSkipPhysicsMove, bTeleport);
+		PropagateTransformUpdate(true, bSkipPhysicsMove, Teleport);
 	}
 	else
 	{
@@ -399,13 +399,13 @@ void USceneComponent::OnRegister()
 #endif
 }
 
-void USceneComponent::UpdateComponentToWorld(bool bSkipPhysicsMove, bool bTeleport)
+void USceneComponent::UpdateComponentToWorld(bool bSkipPhysicsMove, ETeleportType Teleport)
 {
-	UpdateComponentToWorldWithParent(AttachParent, bSkipPhysicsMove, RelativeRotationCache.RotatorToQuat(RelativeRotation), bTeleport);
+	UpdateComponentToWorldWithParent(AttachParent, bSkipPhysicsMove, RelativeRotationCache.RotatorToQuat(RelativeRotation), Teleport);
 }
 
 
-void USceneComponent::PropagateTransformUpdate(bool bTransformChanged, bool bSkipPhysicsMove, bool bTeleport)
+void USceneComponent::PropagateTransformUpdate(bool bTransformChanged, bool bSkipPhysicsMove, ETeleportType Teleport)
 {
 	if (IsDeferringMovementUpdates())
 	{
@@ -419,13 +419,13 @@ void USceneComponent::PropagateTransformUpdate(bool bTransformChanged, bool bSki
 		UpdateBounds();
 
 		// Always send new transform to physics
-		OnUpdateTransform(bSkipPhysicsMove, bTeleport);
+		OnUpdateTransform(bSkipPhysicsMove, Teleport);
 
 		// Flag render transform as dirty
 		MarkRenderTransformDirty();
 		
 		// Now go and update children
-		UpdateChildTransforms(bSkipPhysicsMove, bTeleport);
+		UpdateChildTransforms(bSkipPhysicsMove, Teleport);
 
 		// Refresh navigation
 		UpdateNavigationData();
@@ -728,7 +728,7 @@ void USceneComponent::UpdateBounds()
 }
 
 
-void USceneComponent::SetRelativeLocationAndRotation(FVector NewLocation, const FQuat& NewRotation, bool bSweep, FHitResult* OutSweepHitResult)
+void USceneComponent::SetRelativeLocationAndRotation(FVector NewLocation, const FQuat& NewRotation, bool bSweep, FHitResult* OutSweepHitResult, ETeleportType Teleport)
 {
 	ConditionalUpdateComponentToWorld();
 
@@ -736,68 +736,68 @@ void USceneComponent::SetRelativeLocationAndRotation(FVector NewLocation, const
 	const FTransform DesiredWorldTransform = CalcNewComponentToWorld(DesiredRelTransform);
 	const FVector DesiredDelta = FTransform::SubtractTranslations(DesiredWorldTransform, ComponentToWorld);
 
-	MoveComponent(DesiredDelta, DesiredWorldTransform.GetRotation(), bSweep, OutSweepHitResult);
+	MoveComponent(DesiredDelta, DesiredWorldTransform.GetRotation(), bSweep, OutSweepHitResult, MOVECOMP_NoFlags, Teleport);
 }
 
 
-void USceneComponent::AddRelativeRotation(const FQuat& DeltaRotation, bool bSweep, FHitResult* OutSweepHitResult)
+void USceneComponent::AddRelativeRotation(const FQuat& DeltaRotation, bool bSweep, FHitResult* OutSweepHitResult, ETeleportType Teleport)
 {
 	const FQuat CurRelRotQuat = RelativeRotationCache.RotatorToQuat(RelativeRotation);
 	const FQuat NewRelRotQuat = DeltaRotation * CurRelRotQuat;
-	SetRelativeLocationAndRotation(RelativeLocation, NewRelRotQuat, bSweep, OutSweepHitResult);
+	SetRelativeLocationAndRotation(RelativeLocation, NewRelRotQuat, bSweep, OutSweepHitResult, Teleport);
 }
 
-void USceneComponent::AddLocalOffset(FVector DeltaLocation, bool bSweep, FHitResult* OutSweepHitResult)
+void USceneComponent::AddLocalOffset(FVector DeltaLocation, bool bSweep, FHitResult* OutSweepHitResult, ETeleportType Teleport)
 {
 	const FQuat CurRelRotQuat = RelativeRotationCache.RotatorToQuat(RelativeRotation);
 	const FVector LocalOffset = CurRelRotQuat.RotateVector(DeltaLocation);
-	SetRelativeLocationAndRotation(RelativeLocation + LocalOffset, CurRelRotQuat, bSweep, OutSweepHitResult);
+	SetRelativeLocationAndRotation(RelativeLocation + LocalOffset, CurRelRotQuat, bSweep, OutSweepHitResult, Teleport);
 }
 
-void USceneComponent::AddLocalRotation(FRotator DeltaRotation, bool bSweep, FHitResult* OutSweepHitResult)
+void USceneComponent::AddLocalRotation(FRotator DeltaRotation, bool bSweep, FHitResult* OutSweepHitResult, ETeleportType Teleport)
 {
 	const FQuat CurRelRotQuat = RelativeRotationCache.RotatorToQuat(RelativeRotation);
 	const FQuat NewRelRotQuat = CurRelRotQuat * DeltaRotation.Quaternion();
-	SetRelativeLocationAndRotation(RelativeLocation, NewRelRotQuat, bSweep, OutSweepHitResult);
+	SetRelativeLocationAndRotation(RelativeLocation, NewRelRotQuat, bSweep, OutSweepHitResult, Teleport);
 }
 
-void USceneComponent::AddLocalRotation(const FQuat& DeltaRotation, bool bSweep, FHitResult* OutSweepHitResult)
+void USceneComponent::AddLocalRotation(const FQuat& DeltaRotation, bool bSweep, FHitResult* OutSweepHitResult, ETeleportType Teleport)
 {
 	const FQuat CurRelRotQuat = RelativeRotationCache.RotatorToQuat(RelativeRotation);
 	const FQuat NewRelRotQuat = CurRelRotQuat * DeltaRotation;
-	SetRelativeLocationAndRotation(RelativeLocation, NewRelRotQuat, bSweep, OutSweepHitResult);
+	SetRelativeLocationAndRotation(RelativeLocation, NewRelRotQuat, bSweep, OutSweepHitResult, Teleport);
 }
 
-void USceneComponent::AddLocalTransform(const FTransform& DeltaTransform, bool bSweep, FHitResult* OutSweepHitResult)
+void USceneComponent::AddLocalTransform(const FTransform& DeltaTransform, bool bSweep, FHitResult* OutSweepHitResult, ETeleportType Teleport)
 {
 	const FTransform RelativeTransform( RelativeRotationCache.RotatorToQuat(RelativeRotation), RelativeLocation, FVector(1,1,1) ); // don't use scaling, so it matches how AddLocalRotation/Offset work
 	const FTransform NewRelTransform = DeltaTransform * RelativeTransform;
-	SetRelativeTransform(NewRelTransform, bSweep, OutSweepHitResult);
+	SetRelativeTransform(NewRelTransform, bSweep, OutSweepHitResult, Teleport);
 }
 
-void USceneComponent::AddWorldOffset(FVector DeltaLocation, bool bSweep, FHitResult* OutSweepHitResult)
+void USceneComponent::AddWorldOffset(FVector DeltaLocation, bool bSweep, FHitResult* OutSweepHitResult, ETeleportType Teleport)
 {
 	const FVector NewWorldLocation = DeltaLocation + ComponentToWorld.GetTranslation();
-	SetWorldLocation(NewWorldLocation, bSweep, OutSweepHitResult);
+	SetWorldLocation(NewWorldLocation, bSweep, OutSweepHitResult, Teleport);
 }
 
-void USceneComponent::AddWorldRotation(FRotator DeltaRotation, bool bSweep, FHitResult* OutSweepHitResult)
+void USceneComponent::AddWorldRotation(FRotator DeltaRotation, bool bSweep, FHitResult* OutSweepHitResult, ETeleportType Teleport)
 {
 	const FQuat NewWorldRotation = DeltaRotation.Quaternion() * ComponentToWorld.GetRotation();
-	SetWorldRotation(NewWorldRotation, bSweep, OutSweepHitResult);
+	SetWorldRotation(NewWorldRotation, bSweep, OutSweepHitResult, Teleport);
 }
 
-void USceneComponent::AddWorldRotation(const FQuat& DeltaRotation, bool bSweep, FHitResult* OutSweepHitResult)
+void USceneComponent::AddWorldRotation(const FQuat& DeltaRotation, bool bSweep, FHitResult* OutSweepHitResult, ETeleportType Teleport)
 {
 	const FQuat NewWorldRotation = DeltaRotation * ComponentToWorld.GetRotation();
-	SetWorldRotation(NewWorldRotation, bSweep, OutSweepHitResult);
+	SetWorldRotation(NewWorldRotation, bSweep, OutSweepHitResult, Teleport);
 }
 
-void USceneComponent::AddWorldTransform(const FTransform& DeltaTransform, bool bSweep, FHitResult* OutSweepHitResult)
+void USceneComponent::AddWorldTransform(const FTransform& DeltaTransform, bool bSweep, FHitResult* OutSweepHitResult, ETeleportType Teleport)
 {
 	const FQuat NewWorldRotation = DeltaTransform.GetRotation() * ComponentToWorld.GetRotation();
 	const FVector NewWorldLocation = FTransform::AddTranslations(DeltaTransform, ComponentToWorld);
-	SetWorldTransform(FTransform(NewWorldRotation, NewWorldLocation, FVector(1,1,1)));
+	SetWorldTransform(FTransform(NewWorldRotation, NewWorldLocation, FVector(1,1,1)),false, nullptr, Teleport);
 }
 
 void USceneComponent::SetRelativeScale3D(FVector NewScale3D)
@@ -834,9 +834,9 @@ void USceneComponent::ResetRelativeTransform()
 	SetRelativeScale3D(FVector(1.f));
 }
 
-void USceneComponent::SetRelativeTransform(const FTransform& NewTransform, bool bSweep, FHitResult* OutSweepHitResult)
+void USceneComponent::SetRelativeTransform(const FTransform& NewTransform, bool bSweep, FHitResult* OutSweepHitResult, ETeleportType Teleport)
 {
-	SetRelativeLocationAndRotation(NewTransform.GetTranslation(), NewTransform.GetRotation(), bSweep, OutSweepHitResult);
+	SetRelativeLocationAndRotation(NewTransform.GetTranslation(), NewTransform.GetRotation(), bSweep, OutSweepHitResult, Teleport);
 	SetRelativeScale3D(NewTransform.GetScale3D());
 }
 
@@ -846,7 +846,7 @@ FTransform USceneComponent::GetRelativeTransform() const
 	return RelativeTransform;
 }
 
-void USceneComponent::SetWorldLocation(FVector NewLocation, bool bSweep, FHitResult* OutSweepHitResult)
+void USceneComponent::SetWorldLocation(FVector NewLocation, bool bSweep, FHitResult* OutSweepHitResult, ETeleportType Teleport)
 {
 	FVector NewRelLocation = NewLocation;
 
@@ -857,10 +857,10 @@ void USceneComponent::SetWorldLocation(FVector NewLocation, bool bSweep, FHitRes
 		NewRelLocation = ParentToWorld.InverseTransformPosition(NewLocation);
 	}
 
-	SetRelativeLocation(NewRelLocation, bSweep, OutSweepHitResult);
+	SetRelativeLocation(NewRelLocation, bSweep, OutSweepHitResult, Teleport);
 }
 
-void USceneComponent::SetWorldRotation(const FQuat& NewRotation, bool bSweep, FHitResult* OutSweepHitResult)
+void USceneComponent::SetWorldRotation(const FQuat& NewRotation, bool bSweep, FHitResult* OutSweepHitResult, ETeleportType Teleport)
 {
 	FQuat NewRelRotation = NewRotation;
 
@@ -873,7 +873,7 @@ void USceneComponent::SetWorldRotation(const FQuat& NewRotation, bool bSweep, FH
 		NewRelRotation = NewRelQuat;
 	}
 
-	SetRelativeRotation(NewRelRotation, bSweep, OutSweepHitResult);
+	SetRelativeRotation(NewRelRotation, bSweep, OutSweepHitResult, Teleport);
 }
 
 
@@ -891,7 +891,7 @@ void USceneComponent::SetWorldScale3D(FVector NewScale)
 	SetRelativeScale3D(NewRelScale);
 }
 
-void USceneComponent::SetWorldTransform(const FTransform& NewTransform, bool bSweep, FHitResult* OutSweepHitResult)
+void USceneComponent::SetWorldTransform(const FTransform& NewTransform, bool bSweep, FHitResult* OutSweepHitResult, ETeleportType Teleport)
 {
 	// If attached to something, transform into local space
 	FQuat NewRotation = NewTransform.GetRotation();
@@ -919,15 +919,15 @@ void USceneComponent::SetWorldTransform(const FTransform& NewTransform, bool bSw
 		}
 	}
 
-	SetRelativeTransform(FTransform(NewRotation, NewLocation, NewScale), bSweep, OutSweepHitResult);
+	SetRelativeTransform(FTransform(NewRotation, NewLocation, NewScale), bSweep, OutSweepHitResult, Teleport);
 }
 
-void USceneComponent::SetWorldLocationAndRotation(FVector NewLocation, FRotator NewRotation, bool bSweep, FHitResult* OutSweepHitResult)
+void USceneComponent::SetWorldLocationAndRotation(FVector NewLocation, FRotator NewRotation, bool bSweep, FHitResult* OutSweepHitResult, ETeleportType Teleport)
 {
-	SetWorldLocationAndRotation(NewLocation, NewRotation.Quaternion(), bSweep, OutSweepHitResult);
+	SetWorldLocationAndRotation(NewLocation, NewRotation.Quaternion(), bSweep, OutSweepHitResult, Teleport);
 }
 
-void USceneComponent::SetWorldLocationAndRotation(FVector NewLocation, const FQuat& NewRotation, bool bSweep, FHitResult* OutSweepHitResult)
+void USceneComponent::SetWorldLocationAndRotation(FVector NewLocation, const FQuat& NewRotation, bool bSweep, FHitResult* OutSweepHitResult, ETeleportType Teleport)
 {
 	// If attached to something, transform into local space
 	FQuat NewFinalRotation = NewRotation;
@@ -948,7 +948,7 @@ void USceneComponent::SetWorldLocationAndRotation(FVector NewLocation, const FQu
 		}
 	}
 
-	SetRelativeLocationAndRotation(NewLocation, NewFinalRotation, bSweep, OutSweepHitResult);
+	SetRelativeLocationAndRotation(NewLocation, NewFinalRotation, bSweep, OutSweepHitResult, Teleport);
 }
 
 void USceneComponent::SetWorldLocationAndRotationNoPhysics(const FVector& NewLocation, const FRotator& NewRotation)
@@ -1489,14 +1489,14 @@ FActorComponentInstanceData* USceneComponent::GetComponentInstanceData() const
 	return InstanceData;
 }
 
-void USceneComponent::UpdateChildTransforms(bool bSkipPhysicsMove, bool bTeleport)
+void USceneComponent::UpdateChildTransforms(bool bSkipPhysicsMove, ETeleportType Teleport)
 {
 	for(int32 i=0; i<AttachChildren.Num(); i++)
 	{
 		USceneComponent* ChildComp = AttachChildren[i];
 		if(ChildComp != NULL)
 		{
-			ChildComp->UpdateComponentToWorld(bSkipPhysicsMove, bTeleport);
+			ChildComp->UpdateComponentToWorld(bSkipPhysicsMove, Teleport);
 		}
 	}
 }
@@ -1827,7 +1827,7 @@ void USceneComponent::BeginDestroy()
 	Super::BeginDestroy();
 }
 
-bool USceneComponent::InternalSetWorldLocationAndRotation(FVector NewLocation, const FQuat& RotationQuat, bool bNoPhysics, bool bTeleport)
+bool USceneComponent::InternalSetWorldLocationAndRotation(FVector NewLocation, const FQuat& RotationQuat, bool bNoPhysics, ETeleportType Teleport)
 {
 	checkSlow(bWorldToComponentUpdated);
 	FQuat NewRotationQuat(RotationQuat);
@@ -1852,7 +1852,7 @@ bool USceneComponent::InternalSetWorldLocationAndRotation(FVector NewLocation, c
 	if (!NewLocation.Equals(RelativeLocation) || !NewRotationQuat.Equals(RelativeRotationCache.RotatorToQuat_ReadOnly(RelativeRotation), 1e-6f))
 	{
 		RelativeLocation = NewLocation;
-		UpdateComponentToWorldWithParent(AttachParent, bNoPhysics, NewRotationQuat, bTeleport);
+		UpdateComponentToWorldWithParent(AttachParent, bNoPhysics, NewRotationQuat, Teleport);
 		return true;
 	}
 
@@ -1910,7 +1910,7 @@ bool USceneComponent::CheckStaticMobilityAndWarn(const FText& ActionText) const
 	return false;
 }
 
-bool USceneComponent::MoveComponentImpl( const FVector& Delta, const FQuat& NewRotation, bool bSweep, FHitResult* OutHit, EMoveComponentFlags MoveFlags )
+bool USceneComponent::MoveComponentImpl(const FVector& Delta, const FQuat& NewRotation, bool bSweep, FHitResult* OutHit, EMoveComponentFlags MoveFlags, ETeleportType Teleport)
 {
 	SCOPE_CYCLE_COUNTER(STAT_MoveComponentSceneComponentTime);
 
@@ -1943,7 +1943,7 @@ bool USceneComponent::MoveComponentImpl( const FVector& Delta, const FQuat& NewR
 	}
 
 	// just teleport, sweep is supported for PrimitiveComponents. This will update child components as well.
-	const bool bMoved = InternalSetWorldLocationAndRotation(GetComponentLocation() + Delta, NewRotation, false, !bSweep);
+	const bool bMoved = InternalSetWorldLocationAndRotation(GetComponentLocation() + Delta, NewRotation, false, Teleport);
 
 	// Only update overlaps if not deferring updates within a scope
 	if (bMoved && !IsDeferringMovementUpdates())
@@ -2497,89 +2497,88 @@ void USceneComponent::UpdateNavigationData()
 	}
 }
 
-
 // K2 versions of various transform changing operations.
 // Note: we pass null for the hit result if not sweeping, for better perf.
 // This assumes this K2 function is only used by blueprints, which initializes the param for each function call.
 
-void USceneComponent::K2_SetRelativeLocationAndRotation(FVector NewLocation, FRotator NewRotation, bool bSweep, FHitResult& SweepHitResult)
+void USceneComponent::K2_SetRelativeLocationAndRotation(FVector NewLocation, FRotator NewRotation, bool bSweep, FHitResult& SweepHitResult, bool bTeleport)
 {
-	SetRelativeLocationAndRotation(NewLocation, NewRotation, bSweep, (bSweep ? &SweepHitResult : nullptr));
+	SetRelativeLocationAndRotation(NewLocation, NewRotation, bSweep, (bSweep ? &SweepHitResult : nullptr), TeleportFlagToEnum(bTeleport));
 }
 
-void USceneComponent::K2_SetWorldLocationAndRotation(FVector NewLocation, FRotator NewRotation, bool bSweep, FHitResult& SweepHitResult)
+void USceneComponent::K2_SetWorldLocationAndRotation(FVector NewLocation, FRotator NewRotation, bool bSweep, FHitResult& SweepHitResult, bool bTeleport)
 {
-	SetWorldLocationAndRotation(NewLocation, NewRotation, bSweep, (bSweep ? &SweepHitResult : nullptr));
+	SetWorldLocationAndRotation(NewLocation, NewRotation, bSweep, (bSweep ? &SweepHitResult : nullptr), TeleportFlagToEnum(bTeleport));
 }
 
-void USceneComponent::K2_SetRelativeLocation(FVector NewLocation, bool bSweep, FHitResult& SweepHitResult)
+void USceneComponent::K2_SetRelativeLocation(FVector NewLocation, bool bSweep, FHitResult& SweepHitResult, bool bTeleport)
 {
-	SetRelativeLocation(NewLocation, bSweep, (bSweep ? &SweepHitResult : nullptr));
+	SetRelativeLocation(NewLocation, bSweep, (bSweep ? &SweepHitResult : nullptr), TeleportFlagToEnum(bTeleport));
 }
 
-void USceneComponent::K2_SetRelativeRotation(FRotator NewRotation, bool bSweep, FHitResult& SweepHitResult)
+void USceneComponent::K2_SetRelativeRotation(FRotator NewRotation, bool bSweep, FHitResult& SweepHitResult, bool bTeleport)
 {
-	SetRelativeRotation(NewRotation, bSweep, (bSweep ? &SweepHitResult : nullptr));
+	SetRelativeRotation(NewRotation, bSweep, (bSweep ? &SweepHitResult : nullptr), TeleportFlagToEnum(bTeleport));
 }
 
-void USceneComponent::K2_SetRelativeTransform(const FTransform& NewTransform, bool bSweep, FHitResult& SweepHitResult)
+void USceneComponent::K2_SetRelativeTransform(const FTransform& NewTransform, bool bSweep, FHitResult& SweepHitResult, bool bTeleport)
 {
-	SetRelativeTransform(NewTransform, bSweep, (bSweep ? &SweepHitResult : nullptr));
+	SetRelativeTransform(NewTransform, bSweep, (bSweep ? &SweepHitResult : nullptr), TeleportFlagToEnum(bTeleport));
 }
 
-void USceneComponent::K2_AddRelativeLocation(FVector DeltaLocation, bool bSweep, FHitResult& SweepHitResult)
+void USceneComponent::K2_AddRelativeLocation(FVector DeltaLocation, bool bSweep, FHitResult& SweepHitResult, bool bTeleport)
 {
-	AddRelativeLocation(DeltaLocation, bSweep, (bSweep ? &SweepHitResult : nullptr));
+	AddRelativeLocation(DeltaLocation, bSweep, (bSweep ? &SweepHitResult : nullptr), TeleportFlagToEnum(bTeleport));
 }
 
-void USceneComponent::K2_AddRelativeRotation(FRotator DeltaRotation, bool bSweep, FHitResult& SweepHitResult)
+void USceneComponent::K2_AddRelativeRotation(FRotator DeltaRotation, bool bSweep, FHitResult& SweepHitResult, bool bTeleport)
 {
-	AddRelativeRotation(DeltaRotation, bSweep, (bSweep ? &SweepHitResult : nullptr));
+	AddRelativeRotation(DeltaRotation, bSweep, (bSweep ? &SweepHitResult : nullptr), TeleportFlagToEnum(bTeleport));
 }
 
-void USceneComponent::K2_AddLocalOffset(FVector DeltaLocation, bool bSweep, FHitResult& SweepHitResult)
+void USceneComponent::K2_AddLocalOffset(FVector DeltaLocation, bool bSweep, FHitResult& SweepHitResult, bool bTeleport)
 {
-	AddLocalOffset(DeltaLocation, bSweep, (bSweep ? &SweepHitResult : nullptr));
+	AddLocalOffset(DeltaLocation, bSweep, (bSweep ? &SweepHitResult : nullptr), TeleportFlagToEnum(bTeleport));
 }
 
-void USceneComponent::K2_AddLocalRotation(FRotator DeltaRotation, bool bSweep, FHitResult& SweepHitResult)
+void USceneComponent::K2_AddLocalRotation(FRotator DeltaRotation, bool bSweep, FHitResult& SweepHitResult, bool bTeleport)
 {
-	AddLocalRotation(DeltaRotation, bSweep, (bSweep ? &SweepHitResult : nullptr));
+	AddLocalRotation(DeltaRotation, bSweep, (bSweep ? &SweepHitResult : nullptr), TeleportFlagToEnum(bTeleport));
 }
 
-void USceneComponent::K2_AddLocalTransform(const FTransform& DeltaTransform, bool bSweep, FHitResult& SweepHitResult)
+void USceneComponent::K2_AddLocalTransform(const FTransform& DeltaTransform, bool bSweep, FHitResult& SweepHitResult, bool bTeleport)
 {
-	AddLocalTransform(DeltaTransform, bSweep, (bSweep ? &SweepHitResult : nullptr));
+	AddLocalTransform(DeltaTransform, bSweep, (bSweep ? &SweepHitResult : nullptr), TeleportFlagToEnum(bTeleport));
 }
 
-void USceneComponent::K2_SetWorldLocation(FVector NewLocation, bool bSweep, FHitResult& SweepHitResult)
+void USceneComponent::K2_SetWorldLocation(FVector NewLocation, bool bSweep, FHitResult& SweepHitResult, bool bTeleport)
 {
-	SetWorldLocation(NewLocation, bSweep, (bSweep ? &SweepHitResult : nullptr));
+	SetWorldLocation(NewLocation, bSweep, (bSweep ? &SweepHitResult : nullptr), TeleportFlagToEnum(bTeleport));
 }
 
-void USceneComponent::K2_SetWorldRotation(FRotator NewRotation, bool bSweep, FHitResult& SweepHitResult)
+void USceneComponent::K2_SetWorldRotation(FRotator NewRotation, bool bSweep, FHitResult& SweepHitResult, bool bTeleport)
 {
-	SetWorldRotation(NewRotation, bSweep, (bSweep ? &SweepHitResult : nullptr));
+	SetWorldRotation(NewRotation, bSweep, (bSweep ? &SweepHitResult : nullptr), TeleportFlagToEnum(bTeleport));
 }
 
-void USceneComponent::K2_SetWorldTransform(const FTransform& NewTransform, bool bSweep, FHitResult& SweepHitResult)
+void USceneComponent::K2_SetWorldTransform(const FTransform& NewTransform, bool bSweep, FHitResult& SweepHitResult, bool bTeleport)
 {
-	SetWorldTransform(NewTransform, bSweep, (bSweep ? &SweepHitResult : nullptr));
+	SetWorldTransform(NewTransform, bSweep, (bSweep ? &SweepHitResult : nullptr), TeleportFlagToEnum(bTeleport));
 }
 
-void USceneComponent::K2_AddWorldOffset(FVector DeltaLocation, bool bSweep, FHitResult& SweepHitResult)
+void USceneComponent::K2_AddWorldOffset(FVector DeltaLocation, bool bSweep, FHitResult& SweepHitResult, bool bTeleport)
 {
-	AddWorldOffset(DeltaLocation, bSweep, (bSweep ? &SweepHitResult : nullptr));
+	AddWorldOffset(DeltaLocation, bSweep, (bSweep ? &SweepHitResult : nullptr), TeleportFlagToEnum(bTeleport));
 }
 
-void USceneComponent::K2_AddWorldRotation(FRotator DeltaRotation, bool bSweep, FHitResult& SweepHitResult)
+void USceneComponent::K2_AddWorldRotation(FRotator DeltaRotation, bool bSweep, FHitResult& SweepHitResult, bool bTeleport)
 {
-	AddWorldRotation(DeltaRotation, bSweep, (bSweep ? &SweepHitResult : nullptr));
+	AddWorldRotation(DeltaRotation, bSweep, (bSweep ? &SweepHitResult : nullptr), TeleportFlagToEnum(bTeleport));
 }
 
-void USceneComponent::K2_AddWorldTransform(const FTransform& DeltaTransform, bool bSweep, FHitResult& SweepHitResult)
+void USceneComponent::K2_AddWorldTransform(const FTransform& DeltaTransform, bool bSweep, FHitResult& SweepHitResult, bool bTeleport)
 {
-	AddWorldTransform(DeltaTransform, bSweep, (bSweep ? &SweepHitResult : nullptr));
+	AddWorldTransform(DeltaTransform, bSweep, (bSweep ? &SweepHitResult : nullptr), TeleportFlagToEnum(bTeleport));
 }
 
 #undef LOCTEXT_NAMESPACE
diff --git a/Engine/Source/Runtime/Engine/Private/Components/SkeletalMeshComponent.cpp b/Engine/Source/Runtime/Engine/Private/Components/SkeletalMeshComponent.cpp
--- a/Engine/Source/Runtime/Engine/Private/Components/SkeletalMeshComponent.cpp
+++ b/Engine/Source/Runtime/Engine/Private/Components/SkeletalMeshComponent.cpp
@@ -1185,7 +1185,7 @@ void USkeletalMeshComponent::PostAnimEvaluation(FAnimationEvaluationContext& Eva
 	InvalidateCachedBounds();
 
 	// update physics data from animated data
-	UpdateKinematicBonesToAnim(GetEditableSpaceBases(), false, true);
+	UpdateKinematicBonesToAnim(GetEditableSpaceBases(), ETeleportType::None, true);
 	UpdateRBJointMotors();
 
 	// @todo anim : hack TTP 224385	ANIM: Skeletalmesh double buffer
diff --git a/Engine/Source/Runtime/Engine/Private/InstancedStaticMesh.cpp b/Engine/Source/Runtime/Engine/Private/InstancedStaticMesh.cpp
--- a/Engine/Source/Runtime/Engine/Private/InstancedStaticMesh.cpp
+++ b/Engine/Source/Runtime/Engine/Private/InstancedStaticMesh.cpp
@@ -1464,7 +1464,7 @@ bool UInstancedStaticMeshComponent::UpdateInstanceTransform(int32 InstanceIndex,
 		FBodyInstance* InstanceBodyInstance = InstanceBodies[InstanceIndex];
 #if WITH_PHYSX
 		// Update transform.
-		InstanceBodyInstance->SetBodyTransform(WorldTransform, false);
+		InstanceBodyInstance->SetBodyTransform(WorldTransform, ETeleportType::None);
 		InstanceBodyInstance->UpdateBodyScale(WorldTransform.GetScale3D());
 #endif //WITH_PHYSX
 	}
diff --git a/Engine/Source/Runtime/Engine/Private/PhysicsEngine/BodyInstance.cpp b/Engine/Source/Runtime/Engine/Private/PhysicsEngine/BodyInstance.cpp
--- a/Engine/Source/Runtime/Engine/Private/PhysicsEngine/BodyInstance.cpp
+++ b/Engine/Source/Runtime/Engine/Private/PhysicsEngine/BodyInstance.cpp
@@ -1503,7 +1503,7 @@ struct FInitBodiesHelper
 					else
 					{
 						// Position the body
-						Instance->SetBodyTransform(Transform, /*bTeleport=*/ true);
+						Instance->SetBodyTransform(Transform, ETeleportType::TeleportPhysics);
 					}
 
 					// Apply correct physical materials to shape we created.
@@ -2510,6 +2510,11 @@ FTransform FBodyInstance::GetUnrealWorldTransform_AssumesLocked(bool bWithProjec
 }
 
 void FBodyInstance::SetBodyTransform(const FTransform& NewTransform, bool bTeleport)
+{
+	SetBodyTransform(NewTransform, TeleportFlagToEnum(bTeleport));
+}
+
+void FBodyInstance::SetBodyTransform(const FTransform& NewTransform, ETeleportType Teleport)
 {
 	SCOPE_CYCLE_COUNTER(STAT_SetBodyTransform);
 
@@ -2545,7 +2550,7 @@ void FBodyInstance::SetBodyTransform(const FTransform& NewTransform, bool bTelep
 			if (PxRigidDynamic* PRigidDynamic = GetPxRigidDynamic_AssumesLocked())
 			{
 				// If kinematic and not teleporting, set kinematic target
-				if (!IsRigidBodyNonKinematic_AssumesLocked(PRigidDynamic) && !bTeleport)
+				if (!IsRigidBodyNonKinematic_AssumesLocked(PRigidDynamic) && Teleport == ETeleportType::None)
 				{
 					if(FPhysScene* PhysScene = GetPhysicsScene(this))
 					{
diff --git a/Engine/Source/Runtime/Engine/Private/PhysicsEngine/PhysAnim.cpp b/Engine/Source/Runtime/Engine/Private/PhysicsEngine/PhysAnim.cpp
--- a/Engine/Source/Runtime/Engine/Private/PhysicsEngine/PhysAnim.cpp
+++ b/Engine/Source/Runtime/Engine/Private/PhysicsEngine/PhysAnim.cpp
@@ -215,7 +215,7 @@ void USkeletalMeshComponent::BlendPhysicsBones( TArray<FBoneIndexType>& InRequir
 
 		if (bUpdatePhysics && BodyInstance)
 		{
-			BodyInstance->SetBodyTransform(EditableSpaceBases[BoneIndex] * ComponentToWorld, true);
+			BodyInstance->SetBodyTransform(EditableSpaceBases[BoneIndex] * ComponentToWorld, ETeleportType::TeleportPhysics);
 		}
 	}
 
@@ -285,7 +285,7 @@ void USkeletalMeshComponent::BlendInPhysics()
 }
 
 
-void USkeletalMeshComponent::UpdateKinematicBonesToAnim(const TArray<FTransform>& InSpaceBases, bool bTeleport, bool bNeedsSkinning)
+void USkeletalMeshComponent::UpdateKinematicBonesToAnim(const TArray<FTransform>& InSpaceBases, ETeleportType Teleport, bool bNeedsSkinning)
 {
 	SCOPE_CYCLE_COUNTER(STAT_UpdateRBBones);
 
@@ -295,7 +295,7 @@ void USkeletalMeshComponent::UpdateKinematicBonesToAnim(const TArray<FTransform>
 	// If desired, pass the animation data to the physics joints so they can be used by motors.
 	// See if we are going to need to update kinematics
 	const bool bUpdateKinematics = (KinematicBonesUpdateType != EKinematicBonesUpdateToPhysics::SkipAllBones);
-
+	const bool bTeleport = Teleport == ETeleportType::TeleportPhysics;
 	// If desired, update physics bodies associated with skeletal mesh component to match.
 	if(!bUpdateKinematics && !(bTeleport && IsAnySimulatingPhysics()))
 	{
@@ -492,7 +492,7 @@ void USkeletalMeshComponent::UpdateKinematicBonesToAnim(const TArray<FTransform>
 				BodyInstance.UpdateTriMeshVertices(NewPositions);
 			}
 			
-			BodyInstance.SetBodyTransform(CurrentLocalToWorld, bTeleport);
+			BodyInstance.SetBodyTransform(CurrentLocalToWorld, Teleport);
 		}
 	}
 
diff --git a/Engine/Source/Runtime/Engine/Private/PrimitiveComponentPhysics.cpp b/Engine/Source/Runtime/Engine/Private/PrimitiveComponentPhysics.cpp
--- a/Engine/Source/Runtime/Engine/Private/PrimitiveComponentPhysics.cpp
+++ b/Engine/Source/Runtime/Engine/Private/PrimitiveComponentPhysics.cpp
@@ -86,7 +86,7 @@ bool UPrimitiveComponent::ApplyRigidBodyState(const FRigidBodyState& NewState, c
 		}
 
 		/////// BODY UPDATE ///////
-		BI->SetBodyTransform(FTransform(UpdatedQuat, UpdatedPos), true);
+		BI->SetBodyTransform(FTransform(UpdatedQuat, UpdatedPos), ETeleportType::TeleportPhysics);
 		BI->SetLinearVelocity(NewState.LinVel + FixLinVel, false);
 		BI->SetAngularVelocity(NewState.AngVel + FixAngVel, false);
 
diff --git a/Engine/Source/Runtime/Engine/Private/SkeletalMeshComponentPhysics.cpp b/Engine/Source/Runtime/Engine/Private/SkeletalMeshComponentPhysics.cpp
--- a/Engine/Source/Runtime/Engine/Private/SkeletalMeshComponentPhysics.cpp
+++ b/Engine/Source/Runtime/Engine/Private/SkeletalMeshComponentPhysics.cpp
@@ -589,7 +589,7 @@ void USkeletalMeshComponent::SetAllPhysicsPosition(FVector NewPos)
 			FTransform RootBodyTM = RootBI->GetUnrealWorldTransform();
 			FVector DeltaLoc = NewPos - RootBodyTM.GetLocation();
 			RootBodyTM.SetTranslation(NewPos);
-			RootBI->SetBodyTransform(RootBodyTM, true);
+			RootBI->SetBodyTransform(RootBodyTM, ETeleportType::TeleportPhysics);
 
 #if DO_CHECK
 			FVector RelativeVector = (RootBI->GetUnrealWorldTransform().GetLocation() - NewPos);
@@ -606,7 +606,7 @@ void USkeletalMeshComponent::SetAllPhysicsPosition(FVector NewPos)
 
 					FTransform BodyTM = BI->GetUnrealWorldTransform();
 					BodyTM.SetTranslation(BodyTM.GetTranslation() + DeltaLoc);
-					BI->SetBodyTransform(BodyTM, true);
+					BI->SetBodyTransform(BodyTM, ETeleportType::TeleportPhysics);
 				}
 			}
 
@@ -630,7 +630,7 @@ void USkeletalMeshComponent::SetAllPhysicsRotation(FRotator NewRot)
 			FTransform RootBodyTM = RootBI->GetUnrealWorldTransform();
 			FQuat DeltaQuat = RootBodyTM.GetRotation().Inverse() * NewRotQuat;
 			RootBodyTM.SetRotation(NewRotQuat);
-			RootBI->SetBodyTransform(RootBodyTM, true);
+			RootBI->SetBodyTransform(RootBodyTM, ETeleportType::TeleportPhysics);
 
 			// apply the delta to all the other bodies
 			for (int32 i = 0; i < Bodies.Num(); i++)
@@ -642,7 +642,7 @@ void USkeletalMeshComponent::SetAllPhysicsRotation(FRotator NewRot)
 
 					FTransform BodyTM = BI->GetUnrealWorldTransform();
 					BodyTM.SetRotation(BodyTM.GetRotation() * DeltaQuat);
-					BI->SetBodyTransform( BodyTM, true );
+					BI->SetBodyTransform( BodyTM, ETeleportType::TeleportPhysics );
 				}
 			}
 
@@ -665,7 +665,7 @@ void USkeletalMeshComponent::ApplyDeltaToAllPhysicsTransforms(const FVector& Del
 			FTransform RootBodyTM = RootBI->GetUnrealWorldTransform();
 			RootBodyTM.SetRotation(RootBodyTM.GetRotation() * DeltaRotation);
 			RootBodyTM.SetTranslation(RootBodyTM.GetTranslation() + DeltaLocation);
-			RootBI->SetBodyTransform(RootBodyTM, true);
+			RootBI->SetBodyTransform(RootBodyTM, ETeleportType::TeleportPhysics);
 
 			// apply the delta to all the other bodies
 			for (int32 i = 0; i < Bodies.Num(); i++)
@@ -678,7 +678,7 @@ void USkeletalMeshComponent::ApplyDeltaToAllPhysicsTransforms(const FVector& Del
 					FTransform BodyTM = BI->GetUnrealWorldTransform();
 					BodyTM.SetRotation(BodyTM.GetRotation() * DeltaRotation);
 					BodyTM.SetTranslation(BodyTM.GetTranslation() + DeltaLocation);
-					BI->SetBodyTransform( BodyTM, true );
+					BI->SetBodyTransform( BodyTM, ETeleportType::TeleportPhysics );
 				}
 			}
 
@@ -1331,18 +1331,18 @@ FConstraintInstance* USkeletalMeshComponent::FindConstraintInstance(FName ConNam
 #define OLD_FORCE_UPDATE_BEHAVIOR 0
 #endif
 
-void USkeletalMeshComponent::OnUpdateTransform(bool bSkipPhysicsMove, bool bTeleport)
+void USkeletalMeshComponent::OnUpdateTransform(bool bSkipPhysicsMove, ETeleportType Teleport)
 {
 	// We are handling the physics move below, so don't handle it at higher levels
-	Super::OnUpdateTransform(true, bTeleport);
+	Super::OnUpdateTransform(true, Teleport);
 
 	// Always send new transform to physics
 	if(bPhysicsStateCreated && !bSkipPhysicsMove)
 	{
 #if !OLD_FORCE_UPDATE_BEHAVIOR
-		UpdateKinematicBonesToAnim(GetSpaceBases(), bTeleport, false);
+		UpdateKinematicBonesToAnim(GetSpaceBases(), Teleport, false);
 #else
-		UpdateKinematicBonesToAnim(GetSpaceBases(), true, false);
+		UpdateKinematicBonesToAnim(GetSpaceBases(), ETeleportType::TeleportPhysics, false);
 #endif
 	}
 
