Fix skeletal mesh not getting blend physics called when component is not ticking, but simulation has changed.

diff --git a/Engine/Source/Runtime/Engine/Private/SkeletalMeshComponentPhysics.cpp b/Engine/Source/Runtime/Engine/Private/SkeletalMeshComponentPhysics.cpp
--- a/Engine/Source/Runtime/Engine/Private/SkeletalMeshComponentPhysics.cpp
+++ b/Engine/Source/Runtime/Engine/Private/SkeletalMeshComponentPhysics.cpp
@@ -454,6 +454,8 @@ void USkeletalMeshComponent::SetSimulatePhysics(bool bSimulate)
 			}
 		}
 	}
+
+	UpdatePreClothTickRegisteredState();
 }
 
 void USkeletalMeshComponent::OnComponentCollisionSettingsChanged()
@@ -1026,6 +1028,8 @@ void USkeletalMeshComponent::SetAllBodiesSimulatePhysics(bool bNewSimulate)
 	{
 		Bodies[i]->SetInstanceSimulatePhysics(bNewSimulate);
 	}
+
+	UpdatePreClothTickRegisteredState();
 }
 
 
@@ -1069,6 +1073,8 @@ void USkeletalMeshComponent::SetAllBodiesBelowSimulatePhysics( const FName& InBo
 		//UE_LOG(LogSkeletalMesh, Warning, TEXT( "ForceAllBodiesBelowUnfixed %s" ), *InAsset->BodySetup(BodyIndices(i))->BoneName.ToString() );
 		Bodies[BodyIndices[i]]->SetInstanceSimulatePhysics(bNewSimulate);
 	}
+
+	UpdatePreClothTickRegisteredState();
 }
 
 
