hidden abstract navigation data from scene outliner
#ue4 UE-16725

diff --git a/Engine/Source/Runtime/Engine/Private/AI/Navigation/AbstractNavData.cpp b/Engine/Source/Runtime/Engine/Private/AI/Navigation/AbstractNavData.cpp
--- a/Engine/Source/Runtime/Engine/Private/AI/Navigation/AbstractNavData.cpp
+++ b/Engine/Source/Runtime/Engine/Private/AI/Navigation/AbstractNavData.cpp
@@ -17,6 +17,11 @@ INavigationQueryFilterInterface* FAbstractQueryFilter::CreateCopy() const
 
 AAbstractNavData::AAbstractNavData(const FObjectInitializer& ObjectInitializer) : Super(ObjectInitializer)
 {
+#if WITH_EDITORONLY_DATA
+	bEditable = false;
+	bListedInSceneOutliner = false;
+#endif
+
 	if (HasAnyFlags(RF_ClassDefaultObject) == false)
 	{
 		FindPathImplementation = FindPathAbstract;
