Adding prefixes to GC and Streaming console commands to make them more discoverable. Also added a few placeholder commands to suppress warnings on startup.

Console commands affected:

s.AsyncIOBandwidthLimit
s.MinBulkDataSizeForAsyncLoading
s.AsyncLoadingThreadEnabled
s.WarnIfTimeLimitExceeded
s.TimeLimitExceededMultiplier
s.TimeLimitExceededMinTime
s.UseBackgroundLevelStreaming
s.PriorityAsyncLoadingExtraTime
s.LevelStreamingActorsUpdateTimeLimit
s.LevelStreamingComponentsRegistrationGranularity

gc.MaxObjectsNotConsideredByGC
gc.SizeOfPermanentObjectPool
gc.FlushStreamingOnGC
gc.NumRetriesBeforeForcingGC
gc.AllowParallelGC
gc.TimeBetweenPurgingPendingKillObjects
gc.CollectGarbageEveryFrame
gc.StressTestGC

New console variables:
s.UseBackgroundLevelStreaming
s.AsyncLoadingTimeLimit
s.AsyncLoadingUseFullTimeLimit
s.PriorityAsyncLoadingExtraTime
s.LevelStreamingActorsUpdateTimeLimit
s.LevelStreamingComponentsRegistrationGranularity

diff --git a/Engine/Config/BaseEngine.ini b/Engine/Config/BaseEngine.ini
--- a/Engine/Config/BaseEngine.ini
+++ b/Engine/Config/BaseEngine.ini
@@ -1337,24 +1337,24 @@ UseStrictEngineVersioning=True
 UseVTableConstructors=True
 
 [/Script/Engine.StreamingSettings]
-AsyncIOBandwidthLimit=0
-MinBulkDataSizeForAsyncLoading=131072
-AsyncLoadingThreadEnabled=False
-WarnIfTimeLimitExceeded=False
-TimeLimitExceededMultiplier=1.5
-TimeLimitExceededMinTime=0.005
-bUseBackgroundLevelStreaming=True
-PriorityAsyncLoadingExtraTime=20.0
-LevelStreamingActorsUpdateTimeLimit = 5.0
-LevelStreamingComponentsRegistrationGranularity = 10
+s.AsyncIOBandwidthLimit=0
+s.MinBulkDataSizeForAsyncLoading=131072
+s.AsyncLoadingThreadEnabled=False
+s.WarnIfTimeLimitExceeded=False
+s.TimeLimitExceededMultiplier=1.5
+s.TimeLimitExceededMinTime=0.005
+s.UseBackgroundLevelStreaming=True
+s.PriorityAsyncLoadingExtraTime=20.0
+s.LevelStreamingActorsUpdateTimeLimit = 5.0
+s.LevelStreamingComponentsRegistrationGranularity = 10
 
 [/Script/Engine.GarbageCollectionSettings]
-MaxObjectsNotConsideredByGC=0
-SizeOfPermanentObjectPool=0
-FlushStreamingOnGC=0
-NumRetriesBeforeForcingGC=0
-AllowParallelGC=True
-TimeBetweenPurgingPendingKillObjects=60
+gc.MaxObjectsNotConsideredByGC=0
+gc.SizeOfPermanentObjectPool=0
+gc.FlushStreamingOnGC=0
+gc.NumRetriesBeforeForcingGC=0
+gc.AllowParallelGC=True
+gc.TimeBetweenPurgingPendingKillObjects=60
 
 [Internationalization]
 +LocalizationPaths=../../../Engine/Content/Localization/Engine
diff --git a/Engine/Source/Editor/UnrealEd/Private/EditorEngine.cpp b/Engine/Source/Editor/UnrealEd/Private/EditorEngine.cpp
--- a/Engine/Source/Editor/UnrealEd/Private/EditorEngine.cpp
+++ b/Engine/Source/Editor/UnrealEd/Private/EditorEngine.cpp
@@ -982,7 +982,7 @@ void UEditorEngine::Tick( float DeltaSeconds, bool bIdleMode )
 	// Update subsystems.
 	{
 		// This assumes that UObject::StaticTick only calls ProcessAsyncLoading.	
-		StaticTick(DeltaSeconds, GetDefault<UStreamingSettings>()->bAsyncLoadingUseFullTimeLimit, GetDefault<UStreamingSettings>()->AsyncLoadingTimeLimit / 1000.f);
+		StaticTick(DeltaSeconds, !!GAsyncLoadingUseFullTimeLimit, GAsyncLoadingTimeLimit / 1000.f);
 	}
 
 	// Look for realtime flags.
diff --git a/Engine/Source/Runtime/Core/Private/Serialization/AsyncIOSystemBase.cpp b/Engine/Source/Runtime/Core/Private/Serialization/AsyncIOSystemBase.cpp
--- a/Engine/Source/Runtime/Core/Private/Serialization/AsyncIOSystemBase.cpp
+++ b/Engine/Source/Runtime/Core/Private/Serialization/AsyncIOSystemBase.cpp
@@ -32,7 +32,7 @@ DECLARE_FLOAT_COUNTER_STAT( TEXT( "Bandwidth (MByte/ sec)" ), STAT_AsyncIO_Bandw
 // Constrain bandwidth if wanted. Value is in MByte/ sec.
 float GAsyncIOBandwidthLimit = 0.0f;
 static FAutoConsoleVariableRef CVarAsyncIOBandwidthLimit(
-	TEXT("AsyncIOBandwidthLimit"),
+	TEXT("s.AsyncIOBandwidthLimit"),
 	GAsyncIOBandwidthLimit,
 	TEXT("Constrain bandwidth if wanted. Value is in MByte/ sec."),
 	ECVF_Default
diff --git a/Engine/Source/Runtime/CoreUObject/Private/Serialization/AsyncLoading.cpp b/Engine/Source/Runtime/CoreUObject/Private/Serialization/AsyncLoading.cpp
--- a/Engine/Source/Runtime/CoreUObject/Private/Serialization/AsyncLoading.cpp
+++ b/Engine/Source/Runtime/CoreUObject/Private/Serialization/AsyncLoading.cpp
@@ -210,9 +210,18 @@ public:
 static TAutoPtr<FAsyncLoadingExec> GAsyncLoadingExec;
 #endif
 
+static int32 GAsyncLoadingThreadEnabled;
+static FAutoConsoleVariableRef CVarAsyncLoadingThreadEnabledg(
+	TEXT("s.AsyncLoadingThreadEnabled"),
+	GAsyncLoadingThreadEnabled,
+	TEXT("Placeholder console variable, currently not used in runtime."),
+	ECVF_Default
+	);
+
+
 static int32 GWarnIfTimeLimitExceeded = 0;
 static FAutoConsoleVariableRef CVarWarnIfTimeLimitExceeded(
-	TEXT("WarnIfTimeLimitExceeded"),
+	TEXT("s.WarnIfTimeLimitExceeded"),
 	GWarnIfTimeLimitExceeded,
 	TEXT("Enables log warning if time limit for time-sliced package streaming has been exceeded."),
 	ECVF_Default
@@ -220,7 +229,7 @@ static FAutoConsoleVariableRef CVarWarnIfTimeLimitExceeded(
 
 static float GTimeLimitExceededMultiplier = 1.5f;
 static FAutoConsoleVariableRef CVarTimeLimitExceededMultiplier(
-	TEXT("TimeLimitExceededMultiplier"),
+	TEXT("s.TimeLimitExceededMultiplier"),
 	GTimeLimitExceededMultiplier,
 	TEXT("Multiplier for time limit exceeded warning time threshold."),
 	ECVF_Default
@@ -228,7 +237,7 @@ static FAutoConsoleVariableRef CVarTimeLimitExceededMultiplier(
 
 static float GTimeLimitExceededMinTime = 0.005f;
 static FAutoConsoleVariableRef CVarTimeLimitExceededMinTime(
-	TEXT("TimeLimitExceededMinTime"),
+	TEXT("s.TimeLimitExceededMinTime"),
 	GTimeLimitExceededMinTime,
 	TEXT("Minimum time the time limit exceeded warning will be triggered by."),
 	ECVF_Default
diff --git a/Engine/Source/Runtime/CoreUObject/Private/Serialization/AsyncLoadingThread.h b/Engine/Source/Runtime/CoreUObject/Private/Serialization/AsyncLoadingThread.h
--- a/Engine/Source/Runtime/CoreUObject/Private/Serialization/AsyncLoadingThread.h
+++ b/Engine/Source/Runtime/CoreUObject/Private/Serialization/AsyncLoadingThread.h
@@ -115,7 +115,7 @@ public:
 				{
 					check(GConfig);
 					bool bConfigValue = true;
-					GConfig->GetBool(TEXT("/Script/Engine.StreamingSettings"), TEXT("AsyncLoadingThreadEnabled"), bConfigValue, GEngineIni);
+					GConfig->GetBool(TEXT("/Script/Engine.StreamingSettings"), TEXT("s.AsyncLoadingThreadEnabled"), bConfigValue, GEngineIni);
 					bool bCommandLineNoAsyncThread = false;
 					Value = bConfigValue && FApp::ShouldUseThreadingForPerformance() && !FParse::Param(FCommandLine::Get(), TEXT("NoAsyncLoadingThread"));
 				}
diff --git a/Engine/Source/Runtime/CoreUObject/Private/Serialization/BulkData.cpp b/Engine/Source/Runtime/CoreUObject/Private/Serialization/BulkData.cpp
--- a/Engine/Source/Runtime/CoreUObject/Private/Serialization/BulkData.cpp
+++ b/Engine/Source/Runtime/CoreUObject/Private/Serialization/BulkData.cpp
@@ -692,7 +692,7 @@ void FUntypedBulkData::StartSerializingBulkData(FArchive& Ar, UObject* Owner, in
 
 int32 GMinimumBulkDataSizeForAsyncLoading = 131072;
 static FAutoConsoleVariableRef CVarMinimumBulkDataSizeForAsyncLoading(
-	TEXT("MinBulkDataSizeForAsyncLoading"),
+	TEXT("s.MinBulkDataSizeForAsyncLoading"),
 	GMinimumBulkDataSizeForAsyncLoading,
 	TEXT("Minimum time the time limit exceeded warning will be triggered by."),
 	ECVF_Default
diff --git a/Engine/Source/Runtime/CoreUObject/Private/UObject/GarbageCollection.cpp b/Engine/Source/Runtime/CoreUObject/Private/UObject/GarbageCollection.cpp
--- a/Engine/Source/Runtime/CoreUObject/Private/UObject/GarbageCollection.cpp
+++ b/Engine/Source/Runtime/CoreUObject/Private/UObject/GarbageCollection.cpp
@@ -1197,7 +1197,7 @@ COREUOBJECT_API EditorPostReachabilityAnalysisCallbackType EditorPostReachabilit
 // Allow parralel GC to be overriden to single threaded via console command.
 static int32 GAllowParallelGC = (!PLATFORM_MAC || !WITH_EDITORONLY_DATA) ? 1 : 0;
 static FAutoConsoleVariableRef CVarAllowParallelGC(
-	TEXT("AllowParallelGC"),
+	TEXT("gc.AllowParallelGC"),
 	GAllowParallelGC,
 	TEXT("sed to control parallel GC."),
 	ECVF_Default
@@ -1209,7 +1209,7 @@ static int32 GNumAttemptsSinceLastGC = 0;
 // Number of times GC can be skipped.
 static int32 GNumRetriesBeforeForcingGC = 0;
 static FAutoConsoleVariableRef CVarNumRetriesBeforeForcingGC(
-	TEXT("NumRetriesBeforeForcingGC"),
+	TEXT("gc.NumRetriesBeforeForcingGC"),
 	GNumRetriesBeforeForcingGC,
 	TEXT("Maximum number of times GC can be skipped if worker threads are currently modifying UObject state."),
 	ECVF_Default
@@ -1218,7 +1218,7 @@ static FAutoConsoleVariableRef CVarNumRetriesBeforeForcingGC(
 // Force flush streaming on GC console variable
 static int32 GFlushStreamingOnGC = 0;
 static FAutoConsoleVariableRef CVarFlushStreamingOnGC(
-	TEXT("FlushStreamingOnGC"),
+	TEXT("gc.FlushStreamingOnGC"),
 	GFlushStreamingOnGC,
 	TEXT("If enabled, streaming will be flushed each time garbage collection is triggered."),
 	ECVF_Default
diff --git a/Engine/Source/Runtime/CoreUObject/Private/UObject/UObjectBase.cpp b/Engine/Source/Runtime/CoreUObject/Private/UObject/UObjectBase.cpp
--- a/Engine/Source/Runtime/CoreUObject/Private/UObject/UObjectBase.cpp
+++ b/Engine/Source/Runtime/CoreUObject/Private/UObject/UObjectBase.cpp
@@ -823,6 +823,21 @@ void ProcessNewlyLoadedUObjects()
 #endif
 }
 
+static int32 GVarWarnIfTimeLimitExceeded;
+static FAutoConsoleVariableRef CVarWarnIfTimeLimitExceeded(
+	TEXT("gc.MaxObjectsNotConsideredByGC"),
+	GVarWarnIfTimeLimitExceeded,
+	TEXT("Placeholder console variable, currently not used in runtime."),
+	ECVF_Default
+	);
+
+static int32 GSizeOfPermanentObjectPool;
+static FAutoConsoleVariableRef CSizeOfPermanentObjectPool(
+	TEXT("gc.SizeOfPermanentObjectPool"),
+	GSizeOfPermanentObjectPool,
+	TEXT("Placeholder console variable, currently not used in runtime."),
+	ECVF_Default
+	);
 
 /**
  * Final phase of UObject initialization. all auto register objects are added to the main data structures.
@@ -840,10 +855,10 @@ void UObjectBaseInit()
 	// FPlatformProperties::RequiresCookedData() will be false. Please note that GIsEditor and FApp::IsGame() are not valid at this point.
 	if( FPlatformProperties::RequiresCookedData() )
 	{
-		GConfig->GetInt( TEXT("Core.System"), TEXT("MaxObjectsNotConsideredByGC"), MaxObjectsNotConsideredByGC, GEngineIni );
+		GConfig->GetInt( TEXT("/Script/Engine.GarbageCollectionSettings"), TEXT("gc.MaxObjectsNotConsideredByGC"), MaxObjectsNotConsideredByGC, GEngineIni );
 
 		// Not used on PC as in-place creation inside bigger pool interacts with the exit purge and deleting UObject directly.
-		GConfig->GetInt( TEXT("Core.System"), TEXT("SizeOfPermanentObjectPool"), SizeOfPermanentObjectPool, GEngineIni );
+		GConfig->GetInt( TEXT("/Script/Engine.GarbageCollectionSettings"), TEXT("gc.SizeOfPermanentObjectPool"), SizeOfPermanentObjectPool, GEngineIni );
 	}
 
 	// Log what we're doing to track down what really happens as log in LaunchEngineLoop doesn't report those settings in pristine form.
diff --git a/Engine/Source/Runtime/Engine/Classes/Engine/CoreSettings.h b/Engine/Source/Runtime/Engine/Classes/Engine/CoreSettings.h
--- a/Engine/Source/Runtime/Engine/Classes/Engine/CoreSettings.h
+++ b/Engine/Source/Runtime/Engine/Classes/Engine/CoreSettings.h
@@ -20,66 +20,68 @@ public:
 
 	UStreamingSettings();
 
+protected:
+
 	UPROPERTY(config, EditAnywhere, Category = PackageStreaming, meta = (
-		DisplayName = "Async Loading Thread Enabled",
+		ConsoleVariable = "s.AsyncLoadingThreadEnabled", DisplayName = "Async Loading Thread Enabled",
 		ToolTip = "Enables separate thread for package streaming. Requires restart to take effect."))
 	uint32 AsyncLoadingThreadEnabled : 1;
 
 	UPROPERTY(config, EditAnywhere, Category = PackageStreaming, AdvancedDisplay, meta = (
-		ConsoleVariable = "WarnIfTimeLimitExceeded", DisplayName = "Warn If Time Limit Has Been Exceeded",
+		ConsoleVariable = "s.WarnIfTimeLimitExceeded", DisplayName = "Warn If Time Limit Has Been Exceeded",
 		ToolTip = "Enables log warning if time limit for time-sliced package streaming has been exceeded."))
 	uint32 WarnIfTimeLimitExceeded : 1;
 
 	UPROPERTY(config, EditAnywhere, Category = PackageStreaming, AdvancedDisplay, meta = (
-		ConsoleVariable = "TimeLimitExceededMultiplier", DisplayName = "Time Limit Exceeded Warning Multiplier",
+		ConsoleVariable = "s.TimeLimitExceededMultiplier", DisplayName = "Time Limit Exceeded Warning Multiplier",
 		ToolTip = "Multiplier for time limit exceeded warning time threshold."))
 	float TimeLimitExceededMultiplier;
 
 	UPROPERTY(config, EditAnywhere, Category = PackageStreaming, AdvancedDisplay, meta = (
-		ConsoleVariable = "TimeLimitExceededMinTime", DisplayName = "Minimum Time Limit For Time Limit Exceeded Warning",
+		ConsoleVariable = "s.TimeLimitExceededMinTime", DisplayName = "Minimum Time Limit For Time Limit Exceeded Warning",
 		ToolTip = "Minimum time the time limit exceeded warning will be triggered by."))
 	float TimeLimitExceededMinTime;
 
 	UPROPERTY(config, EditAnywhere, Category = PackageStreaming, AdvancedDisplay, meta = (
-		ConsoleVariable = "MinBulkDataSizeForAsyncLoading", DisplayName = "Minimum Bulk Data Size For Async Loading",
+		ConsoleVariable = "s.MinBulkDataSizeForAsyncLoading", DisplayName = "Minimum Bulk Data Size For Async Loading",
 		ToolTip = "Minimum time the time limit exceeded warning will be triggered by."))
 	int32 MinBulkDataSizeForAsyncLoading;
 
 	UPROPERTY(config, EditAnywhere, Category = IO, meta = (
-		ConsoleVariable = "AsyncIOBandwidthLimit", DisplayName = "Asynchronous IO Bandwidth Limit",
+		ConsoleVariable = "s.AsyncIOBandwidthLimit", DisplayName = "Asynchronous IO Bandwidth Limit",
 		ToolTip = "Constrain bandwidth if wanted. Value is in MByte/ sec."))
 	float AsyncIOBandwidthLimit;
 
 	UPROPERTY(EditAnywhere, config, Category = LevelStreaming, meta = (
-		DisplayName = "Use Background Level Streaming",
+		ConsoleVariable = "s.UseBackgroundLevelStreaming", DisplayName = "Use Background Level Streaming",
 		ToolTip = "Whether to allow background level streaming."))
-	bool bUseBackgroundLevelStreaming;
-
-	UPROPERTY(EditAnywhere, config, Category = LevelStreaming, AdvancedDisplay, meta = (
-		DisplayName = "Async Loading Time Limit",
-		ToolTip = "Maximum amount of time to spend doing asynchronous loading (ms per frame)."))
-	float AsyncLoadingTimeLimit;
+	uint32 UseBackgroundLevelStreaming:1;
 
 	/** Whether to use the entire time limit even if blocked on I/O */
 	UPROPERTY(EditAnywhere, config, Category = LevelStreaming, AdvancedDisplay, meta = (
-		DisplayName = "Async Loading Use Full Time Limit",
+		ConsoleVariable = "s.AsyncLoadingUseFullTimeLimit", DisplayName = "Async Loading Use Full Time Limit",
 		ToolTip = "Whether to use the entire time limit even if blocked on I/O."))
-	bool bAsyncLoadingUseFullTimeLimit;
+	uint32 AsyncLoadingUseFullTimeLimit:1;
 
 	UPROPERTY(EditAnywhere, config, Category = LevelStreaming, AdvancedDisplay, meta = (
-		DisplayName = "Priority Async Loading Extra Time",
+		ConsoleVariable = "s.AsyncLoadingTimeLimit", DisplayName = "Async Loading Time Limit",
+		ToolTip = "Maximum amount of time to spend doing asynchronous loading (ms per frame)."))
+	float AsyncLoadingTimeLimit;
+
+	UPROPERTY(EditAnywhere, config, Category = LevelStreaming, AdvancedDisplay, meta = (
+		ConsoleVariable = "s.PriorityAsyncLoadingExtraTime", DisplayName = "Priority Async Loading Extra Time",
 		ToolTip = "Additional time to spend asynchronous loading during a high priority load."))
 	float PriorityAsyncLoadingExtraTime;
 
 	/** Maximum allowed time to spend for actor registration steps during level streaming (ms per frame)*/
 	UPROPERTY(EditAnywhere, config, Category = LevelStreaming, AdvancedDisplay, meta = (
-		DisplayName = "Level Streaming Actors Update Time Limit",
+		ConsoleVariable = "s.LevelStreamingActorsUpdateTimeLimit", DisplayName = "Level Streaming Actors Update Time Limit",
 		ToolTip = "Maximum allowed time to spend for actor registration steps during level streaming (ms per frame)."))
 	float LevelStreamingActorsUpdateTimeLimit;
 
 	/** Batching granularity used to register actor components during level streaming */
 	UPROPERTY(EditAnywhere, config, Category = LevelStreaming, AdvancedDisplay, meta = (
-		DisplayName = "Level Streaming Components Registration Granularity",
+		ConsoleVariable = "s.LevelStreamingComponentsRegistrationGranularity", DisplayName = "Level Streaming Components Registration Granularity",
 		ToolTip = "Batching granularity used to register actor components during level streaming."))
 	int32 LevelStreamingComponentsRegistrationGranularity;
 
@@ -92,6 +94,18 @@ public:
 	// End UObject interface
 };
 
+/** Whether to allow background level streaming. */
+extern ENGINE_API int32 GUseBackgroundLevelStreaming;
+/** Maximum amount of time to spend doing asynchronous loading (ms per frame). */
+extern ENGINE_API float GAsyncLoadingTimeLimit;
+/** Whether to use the entire time limit even if blocked on I/O. */
+extern ENGINE_API int32 GAsyncLoadingUseFullTimeLimit;
+/** Additional time to spend asynchronous loading during a high priority load. */
+extern ENGINE_API float GPriorityAsyncLoadingExtraTime;
+/** Maximum allowed time to spend for actor registration steps during level streaming (ms per frame). */
+extern ENGINE_API float GLevelStreamingActorsUpdateTimeLimit;
+/** Batching granularity used to register actor components during level streaming. */
+extern ENGINE_API int32 GLevelStreamingComponentsRegistrationGranularity;
 
 /**
 * Implements the settings for garbage collection.
@@ -105,33 +119,35 @@ public:
 
 	UGarbageCollectionSettings();
 
+protected:
+
 	UPROPERTY(EditAnywhere, config, Category = General, meta = (
-		ConsoleVariable = "TimeBetweenPurgingPendingKillObjects", DisplayName = "Time Between Purging Pending Kill Objects",
+		ConsoleVariable = "gc.TimeBetweenPurgingPendingKillObjects", DisplayName = "Time Between Purging Pending Kill Objects",
 		ToolTip = "Time in seconds (game time) we should wait between purging object references to objects that are pending kill."))
 	float TimeBetweenPurgingPendingKillObjects;
 
 	UPROPERTY(EditAnywhere, config, Category = General, meta = (
-		ConsoleVariable = "FlushStreamingOnGC", DisplayName = "Flush Streaming On GC",
+		ConsoleVariable = "gc.FlushStreamingOnGC", DisplayName = "Flush Streaming On GC",
 		ToolTip = "If enabled, streaming will be flushed each time garbage collection is triggered."))
 	uint32 FlushStreamingOnGC : 1;
 
 	UPROPERTY(EditAnywhere, config, Category = Optimization, meta = (
-		ConsoleVariable = "AllowParallelGC", DisplayName = "Allow Parallel GC",
+		ConsoleVariable = "gc.AllowParallelGC", DisplayName = "Allow Parallel GC",
 		ToolTip = "If enabled, garbage collection will use multiple threads."))
 	uint32 AllowParallelGC : 1;
 
 	UPROPERTY(EditAnywhere, config, Category = General, meta = (
-		ConsoleVariable = "NumRetriesBeforeForcingGC", DisplayName = "Number Of Retries Before Forcing GC",
+		ConsoleVariable = "gc.NumRetriesBeforeForcingGC", DisplayName = "Number Of Retries Before Forcing GC",
 		ToolTip = "Maximum number of times GC can be skipped if worker threads are currently modifying UObject state. 0 = never force GC"))
 	int32 NumRetriesBeforeForcingGC;
 
 	UPROPERTY(EditAnywhere, config, Category = Optimization, meta = (
-		DisplayName = "Maximum Object Count Not Considered By GC",
+		ConsoleVariable = "gc.MaxObjectsNotConsideredByGC", DisplayName = "Maximum Object Count Not Considered By GC",
 		ToolTip = "Maximum Object Count Not Considered By GC. Works only in cooked builds."))
 	int32 MaxObjectsNotConsideredByGC;
 
 	UPROPERTY(EditAnywhere, config, Category = Optimization, meta = (
-		DisplayName = "Size Of Permanent Object Pool",
+		ConsoleVariable = "gc.SizeOfPermanentObjectPool", DisplayName = "Size Of Permanent Object Pool",
 		ToolTip = "Size Of Permanent Object Pool (bytes). Works only in cooked builds."))
 	int32 SizeOfPermanentObjectPool;
 
@@ -142,4 +158,4 @@ public:
 	virtual void PostEditChangeProperty(FPropertyChangedEvent& PropertyChangedEvent) override;
 #endif
 	// End UObject interface
-};
\ No newline at end of file
+};
diff --git a/Engine/Source/Runtime/Engine/Private/CoreSettings.cpp b/Engine/Source/Runtime/Engine/Private/CoreSettings.cpp
--- a/Engine/Source/Runtime/Engine/Private/CoreSettings.cpp
+++ b/Engine/Source/Runtime/Engine/Private/CoreSettings.cpp
@@ -3,6 +3,57 @@
 #include "EnginePrivate.h"
 #include "Engine/CoreSettings.h"
 
+
+int32 GUseBackgroundLevelStreaming = 1;
+float GAsyncLoadingTimeLimit = 5.0f;
+int32 GAsyncLoadingUseFullTimeLimit = 1;
+float GPriorityAsyncLoadingExtraTime = 20.0f;
+float GLevelStreamingActorsUpdateTimeLimit = 5.0f;
+int32 GLevelStreamingComponentsRegistrationGranularity = 10;
+
+
+static FAutoConsoleVariableRef CVarUseBackgroundLevelStreaming(
+	TEXT("s.UseBackgroundLevelStreaming"),
+	GUseBackgroundLevelStreaming,
+	TEXT("Whether to allow background level streaming."),
+	ECVF_Default
+	);
+
+static FAutoConsoleVariableRef CVarAsyncLoadingTimeLimit(
+	TEXT("s.AsyncLoadingTimeLimit"),
+	GAsyncLoadingTimeLimit,
+	TEXT("Maximum amount of time to spend doing asynchronous loading (ms per frame)."),
+	ECVF_Default
+	);
+
+static FAutoConsoleVariableRef CVarAsyncLoadingUseFullTimeLimit(
+	TEXT("s.AsyncLoadingUseFullTimeLimit"),
+	GAsyncLoadingUseFullTimeLimit,
+	TEXT("Whether to use the entire time limit even if blocked on I/O."),
+	ECVF_Default
+	);
+
+static FAutoConsoleVariableRef CVarPriorityAsyncLoadingExtraTime(
+	TEXT("s.PriorityAsyncLoadingExtraTime"),
+	GPriorityAsyncLoadingExtraTime,
+	TEXT("Additional time to spend asynchronous loading during a high priority load."),
+	ECVF_Default
+	);
+
+static FAutoConsoleVariableRef CVarLevelStreamingActorsUpdateTimeLimit(
+	TEXT("s.LevelStreamingActorsUpdateTimeLimit"),
+	GLevelStreamingActorsUpdateTimeLimit,
+	TEXT("Maximum allowed time to spend for actor registration steps during level streaming (ms per frame)."),
+	ECVF_Default
+	);
+
+static FAutoConsoleVariableRef CVarLevelStreamingComponentsRegistrationGranularity(
+	TEXT("s.LevelStreamingComponentsRegistrationGranularity"),
+	GLevelStreamingComponentsRegistrationGranularity,
+	TEXT("Batching granularity used to register actor components during level streaming."),
+	ECVF_Default
+	);
+
 UStreamingSettings::UStreamingSettings()
 	: Super()
 {
@@ -14,9 +65,9 @@ UStreamingSettings::UStreamingSettings()
 	TimeLimitExceededMinTime = 0.005f;
 	MinBulkDataSizeForAsyncLoading = 131072;
 	AsyncIOBandwidthLimit = 0;
-	bUseBackgroundLevelStreaming = true;
+	UseBackgroundLevelStreaming = true;
 	AsyncLoadingTimeLimit = 5.0f;
-	bAsyncLoadingUseFullTimeLimit = true;
+	AsyncLoadingUseFullTimeLimit = true;
 	PriorityAsyncLoadingExtraTime = 20.0f;
 	LevelStreamingActorsUpdateTimeLimit = 5.0f;
 	LevelStreamingComponentsRegistrationGranularity = 10;
diff --git a/Engine/Source/Runtime/Engine/Private/GameEngine.cpp b/Engine/Source/Runtime/Engine/Private/GameEngine.cpp
--- a/Engine/Source/Runtime/Engine/Private/GameEngine.cpp
+++ b/Engine/Source/Runtime/Engine/Private/GameEngine.cpp
@@ -846,7 +846,7 @@ void UGameEngine::Tick( float DeltaSeconds, bool bIdleMode )
 	// Update subsystems.
 	{
 		// This assumes that UObject::StaticTick only calls ProcessAsyncLoading.
-		StaticTick(DeltaSeconds, GetDefault<UStreamingSettings>()->bAsyncLoadingUseFullTimeLimit, GetDefault<UStreamingSettings>()->AsyncLoadingTimeLimit / 1000.f);
+		StaticTick(DeltaSeconds, !!GAsyncLoadingUseFullTimeLimit, GAsyncLoadingTimeLimit / 1000.f);
 	}
 
 	// -----------------------------------------------------
diff --git a/Engine/Source/Runtime/Engine/Private/LevelStreaming.cpp b/Engine/Source/Runtime/Engine/Private/LevelStreaming.cpp
--- a/Engine/Source/Runtime/Engine/Private/LevelStreaming.cpp
+++ b/Engine/Source/Runtime/Engine/Private/LevelStreaming.cpp
@@ -172,7 +172,7 @@ bool FStreamLevelAction::UpdateLevel( ULevelStreaming* LevelStreamingObject )
 		return true;
 	}
 	// Level shouldn't be loaded but is as background level streaming is enabled so we need to fire finished event regardless.
-	else if (LevelStreamingObject->GetLoadedLevel() && !LevelStreamingObject->bShouldBeLoaded && !GetDefault<UStreamingSettings>()->bUseBackgroundLevelStreaming)
+	else if (LevelStreamingObject->GetLoadedLevel() && !LevelStreamingObject->bShouldBeLoaded && !GUseBackgroundLevelStreaming)
 	{
 		return true;
 	}
diff --git a/Engine/Source/Runtime/Engine/Private/LevelTick.cpp b/Engine/Source/Runtime/Engine/Private/LevelTick.cpp
--- a/Engine/Source/Runtime/Engine/Private/LevelTick.cpp
+++ b/Engine/Source/Runtime/Engine/Private/LevelTick.cpp
@@ -712,7 +712,7 @@ static TAutoConsoleVariable<int32> CVarAllowAsyncRenderThreadUpdatesEditor(
 	TEXT("Used to control async renderthread updates in the editor."));
 
 static TAutoConsoleVariable<int32> CVarCollectGarbageEveryFrame(
-	TEXT("CollectGarbageEveryFrame"),
+	TEXT("gc.CollectGarbageEveryFrame"),
 	0,
 	TEXT("Used to debug garbage collection...Collects garbage every frame if the value is > 0."));
 
@@ -992,14 +992,14 @@ extern bool GCollisionAnalyzerIsRecording;
 
 #if !(UE_BUILD_SHIPPING || UE_BUILD_TEST)
 static TAutoConsoleVariable<int32> CVarStressTestGCWhileStreaming(
-	TEXT("t.StressTestGC"),
+	TEXT("gc.StressTestGC"),
 	0,
 	TEXT("If set to 1, the engine will attempt to trigger GC each frame while async loading."));
 #endif
 
 static float GTimeBetweenPurgingPendingKillObjects = 60.0f;
 static FAutoConsoleVariableRef CVarTimeBetweenPurgingPendingKillObjects(
-	TEXT("TimeBetweenPurgingPendingKillObjects"),
+	TEXT("gc.TimeBetweenPurgingPendingKillObjects"),
 	GTimeBetweenPurgingPendingKillObjects,
 	TEXT("Time in seconds (game time) we should wait between purging object references to objects that are pending kill."),
 	ECVF_Default
@@ -1093,7 +1093,7 @@ void UWorld::Tick( ELevelTick TickType, float DeltaSeconds )
 	if (Info->bHighPriorityLoading || Info->bHighPriorityLoadingLocal || IsInSeamlessTravel())
 	{
 		// Force it to use the entire time slice, even if blocked on I/O
-		ProcessAsyncLoading(true, true, GetDefault<UStreamingSettings>()->PriorityAsyncLoadingExtraTime / 1000.0f);
+		ProcessAsyncLoading(true, true, GPriorityAsyncLoadingExtraTime / 1000.0f);
 	}
 
 	// Translate world origin if requested
diff --git a/Engine/Source/Runtime/Engine/Private/World.cpp b/Engine/Source/Runtime/Engine/Private/World.cpp
--- a/Engine/Source/Runtime/Engine/Private/World.cpp
+++ b/Engine/Source/Runtime/Engine/Private/World.cpp
@@ -1589,7 +1589,7 @@ static bool IsTimeLimitExceeded( const TCHAR* CurrentTask, double StartTime, ULe
 	// We don't spread work across several frames in the Editor to avoid potential side effects.
 	if( Level->OwningWorld->IsGameWorld() == true )
 	{
-		double TimeLimit = GetDefault<UStreamingSettings>()->LevelStreamingActorsUpdateTimeLimit;
+		double TimeLimit = GLevelStreamingActorsUpdateTimeLimit;
 		double CurrentTime	= FPlatformTime::Seconds();
 		// Delta time in ms.
 		double DeltaTime	= (CurrentTime - StartTime) * 1000;
@@ -1791,7 +1791,7 @@ void UWorld::AddToWorld( ULevel* Level, const FTransform& LevelTransform )
 		const bool bRerunConstructionScript = !(FPlatformProperties::RequiresCookedData() || (IsGameWorld() && (Level->bWasDuplicatedForPIE || !bRerunConstructionDuringEditorStreaming)));
 		
 		// Incrementally update components.
-		int32 NumComponentsToUpdate = GetDefault<UStreamingSettings>()->LevelStreamingComponentsRegistrationGranularity;
+		int32 NumComponentsToUpdate = GLevelStreamingComponentsRegistrationGranularity;
 		do
 		{
 			Level->IncrementalUpdateComponents( (!IsGameWorld() || IsRunningCommandlet()) ? 0 : NumComponentsToUpdate, bRerunConstructionScript );
@@ -2349,7 +2349,7 @@ void UWorld::UpdateLevelStreamingInner(ULevelStreaming* StreamingLevel)
 	bool bHasVisibilityRequestPending	= StreamingLevel->GetLoadedLevel() && StreamingLevel->GetLoadedLevel() == CurrentLevelPendingVisibility;
 		
 	// Figure out whether level should be loaded, visible and block on load if it should be loaded but currently isn't.
-	bool bShouldBeLoaded = bHasVisibilityRequestPending || (!GetDefault<UStreamingSettings>()->bUseBackgroundLevelStreaming && !bShouldForceUnloadStreamingLevels && !StreamingLevel->bIsRequestingUnloadAndRemoval);
+	bool bShouldBeLoaded = bHasVisibilityRequestPending || (!GUseBackgroundLevelStreaming && !bShouldForceUnloadStreamingLevels && !StreamingLevel->bIsRequestingUnloadAndRemoval);
 	bool bShouldBeVisible	= bHasVisibilityRequestPending || bShouldForceVisibleStreamingLevels;
 	bool bShouldBlockOnLoad	= StreamingLevel->bShouldBlockOnLoad || StreamingLevel->ShouldBeAlwaysLoaded();
 
@@ -2382,7 +2382,7 @@ void UWorld::UpdateLevelStreamingInner(ULevelStreaming* StreamingLevel)
 	// See whether level is already loaded
 	if (bShouldBeLoaded)
 	{
-		const bool bBlockOnLoad = (!IsGameWorld() || !GetDefault<UStreamingSettings>()->bUseBackgroundLevelStreaming || bShouldBlockOnLoad);
+		const bool bBlockOnLoad = (!IsGameWorld() || !GUseBackgroundLevelStreaming || bShouldBlockOnLoad);
 		// Request to load or duplicate existing level
 		StreamingLevel->RequestLevel(this, bAllowLevelLoadRequests, bBlockOnLoad);
 	}
