Put guards again null worlds in UGameViewportClient::ToggleShowCollision to fix order of destruction issues when force exiting the game.

diff --git a/Engine/Source/Runtime/Engine/Private/GameViewportClient.cpp b/Engine/Source/Runtime/Engine/Private/GameViewportClient.cpp
--- a/Engine/Source/Runtime/Engine/Private/GameViewportClient.cpp
+++ b/Engine/Source/Runtime/Engine/Private/GameViewportClient.cpp
@@ -2316,6 +2316,8 @@ void UGameViewportClient::ToggleShowCollision()
 	// special case: for the Engine.Collision flag, we need to un-hide any primitive components that collide so their collision geometry gets rendered
 	const bool bIsShowingCollision = EngineShowFlags.Collision;
 
+	UWorld* World = GetWorld();
+
 	if (bIsShowingCollision)
 	{
 		// Don't allow 'show collision' and 'show volumes' at the same time, so turn collision off
@@ -2326,13 +2328,19 @@ void UGameViewportClient::ToggleShowCollision()
 		}
 
 		NumViewportsShowingCollision++;
-		ShowCollisionOnSpawnedActorsDelegateHandle = GetWorld()->AddOnActorSpawnedHandler(FOnActorSpawned::FDelegate::CreateUObject(this, &UGameViewportClient::ShowCollisionOnSpawnedActors));
+		if (World != nullptr)
+		{
+			ShowCollisionOnSpawnedActorsDelegateHandle = World->AddOnActorSpawnedHandler(FOnActorSpawned::FDelegate::CreateUObject(this, &UGameViewportClient::ShowCollisionOnSpawnedActors));
+		}
 	}
 	else
 	{
 		NumViewportsShowingCollision--;
 		check(NumViewportsShowingCollision >= 0);
-		GetWorld()->RemoveOnActorSpawnedHandler(ShowCollisionOnSpawnedActorsDelegateHandle);
+		if (World != nullptr)
+		{
+			World->RemoveOnActorSpawnedHandler(ShowCollisionOnSpawnedActorsDelegateHandle);
+		}
 	}
 
 	CollisionComponentVisibilityMap& Mapping = GetCollisionComponentVisibilityMap();
@@ -2350,12 +2358,17 @@ void UGameViewportClient::ToggleShowCollision()
 	}
 	Mapping.Empty();
 
+	if (World == nullptr)
+	{
+		return;
+	}
+
 	if (NumViewportsShowingCollision > 0)
 	{
 		for (TObjectIterator<UPrimitiveComponent> It; It; ++It)
 		{
 			UPrimitiveComponent* PrimitiveComponent = *It;
-			if (!PrimitiveComponent->IsVisible() && PrimitiveComponent->IsCollisionEnabled() && PrimitiveComponent->GetScene() == GetWorld()->Scene)
+			if (!PrimitiveComponent->IsVisible() && PrimitiveComponent->IsCollisionEnabled() && PrimitiveComponent->GetScene() == World->Scene)
 			{
 				if (PrimitiveComponent->GetOwner() && PrimitiveComponent->GetOwner()->GetWorld() && PrimitiveComponent->GetOwner()->GetWorld()->IsGameWorld())
 				{
@@ -2371,7 +2384,7 @@ void UGameViewportClient::ToggleShowCollision()
 #if !(UE_BUILD_SHIPPING || UE_BUILD_TEST)
 	if (EngineShowFlags.Collision)
 	{
-		for (FLocalPlayerIterator It((UEngine*)GetOuter(), GetWorld()); It; ++It)
+		for (FLocalPlayerIterator It((UEngine*)GetOuter(), World); It; ++It)
 		{
 			APlayerController* PC = It->PlayerController;
 			if (PC != NULL && PC->GetPawn() != NULL)
