Build Fix - Changing the camera component to export all symbols instead of selective exporting to prevent errors in the future from missing export macros on individual functions.

diff --git a/Engine/Source/Runtime/Engine/Classes/Camera/CameraComponent.h b/Engine/Source/Runtime/Engine/Classes/Camera/CameraComponent.h
--- a/Engine/Source/Runtime/Engine/Classes/Camera/CameraComponent.h
+++ b/Engine/Source/Runtime/Engine/Classes/Camera/CameraComponent.h
@@ -9,8 +9,8 @@
   * Represents a camera viewpoint and settings, such as projection type, field of view, and post-process overrides.
   * The default behavior for an actor used as the camera view target is to look for an attached camera component and use its location, rotation, and settings.
   */
-UCLASS(HideCategories=(Mobility, Rendering, LOD), Blueprintable, ClassGroup=Camera, meta=(BlueprintSpawnableComponent), MinimalAPI)
-class UCameraComponent : public USceneComponent
+UCLASS(HideCategories=(Mobility, Rendering, LOD), Blueprintable, ClassGroup=Camera, meta=(BlueprintSpawnableComponent))
+class ENGINE_API UCameraComponent : public USceneComponent
 {
 	GENERATED_UCLASS_BODY()
 
@@ -62,16 +62,16 @@ class UCameraComponent : public USceneComponent
 	struct FPostProcessSettings PostProcessSettings;
 
 	// UActorComponent interface
-	ENGINE_API virtual void OnRegister() override;
-	ENGINE_API virtual void OnUnregister() override;
-	ENGINE_API virtual void PostLoad() override;
+	virtual void OnRegister() override;
+	virtual void OnUnregister() override;
+	virtual void PostLoad() override;
 #if WITH_EDITOR
-	ENGINE_API virtual void OnComponentDestroyed() override;
-	ENGINE_API virtual void CheckForErrors() override;
+	virtual void OnComponentDestroyed() override;
+	virtual void CheckForErrors() override;
 	// End of UActorComponent interface
 
 	// UObject interface
-	ENGINE_API virtual void PostEditChangeProperty(FPropertyChangedEvent& PropertyChangedEvent) override;
+	virtual void PostEditChangeProperty(FPropertyChangedEvent& PropertyChangedEvent) override;
 #endif
 	virtual void Serialize(FArchive& Ar) override;
 	// End of UObject interface
@@ -81,11 +81,11 @@ class UCameraComponent : public USceneComponent
 	 * Called by Camera class. Subclass and postprocess to add any effects.
 	 */
 	UFUNCTION(BlueprintCallable, Category=Camera)
-	ENGINE_API virtual void GetCameraView(float DeltaTime, FMinimalViewInfo& DesiredView);
+	virtual void GetCameraView(float DeltaTime, FMinimalViewInfo& DesiredView);
 
 	/** Adds an Blendable (implements IBlendableInterface) to the array of Blendables (if it doesn't exist) and update the weight */
 	UFUNCTION(BlueprintCallable, Category="Rendering")
-	ENGINE_API void AddOrUpdateBlendable(TScriptInterface<IBlendableInterface> InBlendableObject, float InWeight = 1.0f) { PostProcessSettings.AddBlendable(InBlendableObject, InWeight); }
+	void AddOrUpdateBlendable(TScriptInterface<IBlendableInterface> InBlendableObject, float InWeight = 1.0f) { PostProcessSettings.AddBlendable(InBlendableObject, InWeight); }
 
 protected:
 #if WITH_EDITORONLY_DATA
@@ -102,17 +102,17 @@ protected:
 
 public:
 
-	ENGINE_API virtual void SetCameraMesh(UStaticMesh* Mesh) { CameraMesh = Mesh; }
+	virtual void SetCameraMesh(UStaticMesh* Mesh) { CameraMesh = Mesh; }
 
 #endif
 public:
 #if WITH_EDITORONLY_DATA
 	// Refreshes the visual components to match the component state
-	ENGINE_API virtual void RefreshVisualRepresentation();
+	virtual void RefreshVisualRepresentation();
 
 
-	ENGINE_API void OverrideFrustumColor(FColor OverrideColor);
-	ENGINE_API void RestoreFrustumColor();
+	void OverrideFrustumColor(FColor OverrideColor);
+	void RestoreFrustumColor();
 #endif
 
 public:
