Fixed shadowed member warning in clang.

diff --git a/Engine/Source/Runtime/Engine/Private/GameViewportClient.cpp b/Engine/Source/Runtime/Engine/Private/GameViewportClient.cpp
--- a/Engine/Source/Runtime/Engine/Private/GameViewportClient.cpp
+++ b/Engine/Source/Runtime/Engine/Private/GameViewportClient.cpp
@@ -2316,8 +2316,6 @@ void UGameViewportClient::ToggleShowCollision()
 	// special case: for the Engine.Collision flag, we need to un-hide any primitive components that collide so their collision geometry gets rendered
 	const bool bIsShowingCollision = EngineShowFlags.Collision;
 
-	UWorld* World = GetWorld();
-
 	if (bIsShowingCollision)
 	{
 		// Don't allow 'show collision' and 'show volumes' at the same time, so turn collision off
