UE4 - Compact cloth buffers. Option to defer uniform buffer creation till it is visible. Option to defer skeletal stuff to visible and do it on the RHI thread. High priority task graph tasks. Hi priority ticks.

diff --git a/Engine/Shaders/GpuSkinVertexFactory.usf b/Engine/Shaders/GpuSkinVertexFactory.usf
--- a/Engine/Shaders/GpuSkinVertexFactory.usf
+++ b/Engine/Shaders/GpuSkinVertexFactory.usf
@@ -52,9 +52,7 @@ bool PerBoneMotionBlur;
 
 #if GPUSKIN_APEX_CLOTH
 /** Vertex buffer from which to read simulated positions of clothing. */
-Buffer<float4> ClothSimulVertsPositions;
-/** Vertex buffer from which to read simulated normals of clothing. */
-Buffer<float4> ClothSimulVertsNormals;
+Buffer<float2> ClothSimulVertsPositionsNormals;
 /** blend weight between simulated positions and original key-framed animation */
 float ClothBlendWeight;
 #endif// #if GPUSKIN_APEX_CLOTH
@@ -230,7 +228,7 @@ bool IsSimulatedVertex( FVertexFactoryInput Input )
 float3 GetClothSimulPosition(int Index)
 {
 #if FEATURE_LEVEL >= FEATURE_LEVEL_SM4
-		return ClothSimulVertsPositions[Index].xyz;
+		return float3(ClothSimulVertsPositionsNormals[Index * 3], ClothSimulVertsPositionsNormals[Index * 3 + 1].x);
 #else
 	return APEXClothParam.Positions[Index];
 #endif //FEATURE_LEVEL >= FEATURE_LEVEL_SM4
@@ -239,7 +237,7 @@ float3 GetClothSimulPosition(int Index)
 float3 GetClothSimulNormal(int Index)
 {
 #if FEATURE_LEVEL >= FEATURE_LEVEL_SM4
-		return ClothSimulVertsNormals[Index].xyz;
+		return float3(ClothSimulVertsPositionsNormals[Index * 3 + 1].y, ClothSimulVertsPositionsNormals[Index * 3 + 2]);
 #else
 	return APEXClothParam.Normals[Index];
 #endif //FEATURE_LEVEL >= FEATURE_LEVEL_SM4
diff --git a/Engine/Source/Runtime/Core/Private/Async/TaskGraph.cpp b/Engine/Source/Runtime/Core/Private/Async/TaskGraph.cpp
--- a/Engine/Source/Runtime/Core/Private/Async/TaskGraph.cpp
+++ b/Engine/Source/Runtime/Core/Private/Async/TaskGraph.cpp
@@ -393,7 +393,36 @@ public:
 		while (1)
 		{
 			FBaseGraphTask* Task = NULL;
-			Task = Queue(QueueIndex).PrivateQueue.Dequeue();
+			Task = Queue(QueueIndex).PrivateQueueHiPri.Dequeue();
+			if (!Task)
+			{
+				if (!bAllowsStealsFromMe && Queue(QueueIndex).OutstandingHiPriTasks.GetValue())
+				{
+					Queue(QueueIndex).IncomingQueue.PopAll(NewTasks);
+					if (NewTasks.Num())
+					{
+						for (int32 Index = NewTasks.Num() - 1; Index >= 0 ; Index--) // reverse the order since PopAll is implicitly backwards
+						{
+							FBaseGraphTask* NewTask = NewTasks[Index];
+							if (ENamedThreads::GetPriority(NewTask->ThreadToExecuteOn))
+							{
+								Queue(QueueIndex).PrivateQueueHiPri.Enqueue(NewTask);
+								Queue(QueueIndex).OutstandingHiPriTasks.Decrement();
+							}
+							else
+							{
+								Queue(QueueIndex).PrivateQueue.Enqueue(NewTask);
+							}
+						}
+						NewTasks.Reset();
+						Task = Queue(QueueIndex).PrivateQueueHiPri.Dequeue();
+					}
+				}
+				if (!Task)
+				{
+					Task = Queue(QueueIndex).PrivateQueue.Dequeue();
+				}
+			}
 			if (!Task)
 			{
 				if (!bAllowsStealsFromMe)
@@ -432,10 +461,23 @@ public:
 					{
 						for (int32 Index = NewTasks.Num() - 1; Index >= 0 ; Index--) // reverse the order since PopAll is implicitly backwards
 						{
-							Queue(QueueIndex).PrivateQueue.Enqueue(NewTasks[Index]);
+							FBaseGraphTask* NewTask = NewTasks[Index];
+							if (ENamedThreads::GetPriority(NewTask->ThreadToExecuteOn))
+							{
+								Queue(QueueIndex).PrivateQueueHiPri.Enqueue(NewTask);
+								Queue(QueueIndex).OutstandingHiPriTasks.Decrement();
+							}
+							else
+							{
+								Queue(QueueIndex).PrivateQueue.Enqueue(NewTask);
+							}
 						}
-						Task = Queue(QueueIndex).PrivateQueue.Dequeue();
 						NewTasks.Reset();
+						Task = Queue(QueueIndex).PrivateQueueHiPri.Dequeue();
+						if (!Task)
+						{
+							Task = Queue(QueueIndex).PrivateQueue.Dequeue();
+						}
 					}
 				}
 				else
@@ -529,7 +571,14 @@ public:
 		else
 		{
 			checkThreadGraph((FTaskThread*)FPlatformTLS::GetTlsValue(PerThreadIDTLSSlot) == this); // verify that we are the thread they say we are
-			Queue(QueueIndex).PrivateQueue.Enqueue(Task);
+			if (Task != WakeUpBaseGraphTask && ENamedThreads::GetPriority(Task->ThreadToExecuteOn))
+			{
+				Queue(QueueIndex).PrivateQueueHiPri.Enqueue(Task);
+			}
+			else
+			{
+				Queue(QueueIndex).PrivateQueue.Enqueue(Task);
+			}
 		}
 	}
 
@@ -578,10 +627,15 @@ public:
 		TestRandomizedThreads();
 		checkThreadGraph(Queue(QueueIndex).StallRestartEvent); // make sure we are started up
 		bool bWasReopenedByMe;
+		bool bHiPri = Task != WakeUpBaseGraphTask && ENamedThreads::GetPriority(Task->ThreadToExecuteOn);
 		{
 			TASKGRAPH_SCOPE_CYCLE_COUNTER(0, STAT_TaskGraph_EnqueueFromOtherThread_ReopenIfClosedAndPush);
 			bWasReopenedByMe = Queue(QueueIndex).IncomingQueue.ReopenIfClosedAndPush(Task);
 		}
+		if (bHiPri)
+		{
+			Queue(QueueIndex).OutstandingHiPriTasks.Increment();
+		}
 		if (bWasReopenedByMe)
 		{
 			TASKGRAPH_SCOPE_CYCLE_COUNTER(1, STAT_TaskGraph_EnqueueFromOtherThread_Trigger);
@@ -674,6 +728,10 @@ private:
 	{
 		/** A non-threas safe queue for the thread locked tasks of a named thread **/
 		FTaskQueue											PrivateQueue;
+		/** A non-threas safe queue for the thread locked tasks of a named thread. These are high priority.**/
+		FTaskQueue											PrivateQueueHiPri;
+		/** Used to signal pending hi pri tasks. **/
+		FThreadSafeCounter									OutstandingHiPriTasks;
 		/** 
 		 *	For named threads, this is a queue of thread locked tasks coming from other threads. They are not stealable.
 		 *	For unnamed thread this is the public queue, subject to stealing.
@@ -917,7 +975,7 @@ public:
 
 		TestRandomizedThreads();
 		checkThreadGraph(NextUnnamedThreadMod);
-		if (CurrentThreadIfKnown == ENamedThreads::AnyThread)
+		if (ENamedThreads::GetThreadIndex(CurrentThreadIfKnown) == ENamedThreads::AnyThread)
 		{
 			 CurrentThreadIfKnown = GetCurrentThread();
 		}
@@ -926,14 +984,21 @@ public:
 			CurrentThreadIfKnown = ENamedThreads::GetThreadIndex(CurrentThreadIfKnown);
 			checkThreadGraph(CurrentThreadIfKnown == GetCurrentThread());
 		}
-		if (ThreadToExecuteOn == ENamedThreads::AnyThread)
+		if (ENamedThreads::GetThreadIndex(ThreadToExecuteOn) == ENamedThreads::AnyThread)
 		{
 			TASKGRAPH_SCOPE_CYCLE_COUNTER(3, STAT_TaskGraph_QueueTask_AnyThread);
 			if (FPlatformProcess::SupportsMultithreading())
 			{
 				{
 					TASKGRAPH_SCOPE_CYCLE_COUNTER(4, STAT_TaskGraph_QueueTask_IncomingAnyThreadTasks_Push);
-					IncomingAnyThreadTasks.Push(Task);
+					if (ENamedThreads::GetPriority(Task->ThreadToExecuteOn))
+					{
+						IncomingAnyThreadTasksHiPri.Push(Task);
+					}
+					else
+					{
+						IncomingAnyThreadTasks.Push(Task);
+					}
 				}
 				FTaskThread* TempTarget;
 				{
@@ -1039,10 +1104,12 @@ public:
 	virtual void WaitUntilTasksComplete(const FGraphEventArray& Tasks, ENamedThreads::Type CurrentThreadIfKnown = ENamedThreads::AnyThread) override
 	{
 		ENamedThreads::Type CurrentThread = CurrentThreadIfKnown;
-		if (CurrentThreadIfKnown == ENamedThreads::AnyThread)
+		if (ENamedThreads::GetThreadIndex(CurrentThreadIfKnown) == ENamedThreads::AnyThread)
 		{
+			bool bIsHiPri = !!ENamedThreads::GetPriority(CurrentThreadIfKnown);
+			check(!ENamedThreads::GetQueueIndex(CurrentThreadIfKnown));
 			CurrentThreadIfKnown = GetCurrentThread();
-			CurrentThread = CurrentThreadIfKnown;
+			CurrentThread = bIsHiPri ? ENamedThreads::HiPri(CurrentThreadIfKnown) : CurrentThreadIfKnown;
 		}
 		else
 		{
@@ -1109,6 +1176,45 @@ public:
 	FBaseGraphTask* FindWork(ENamedThreads::Type ThreadInNeed)
 	{
 		TestRandomizedThreads();
+		{
+			FBaseGraphTask* Task = SortedAnyThreadTasksHiPri.Pop();
+			if (Task)
+			{
+				return Task;
+			}
+		}
+		if (!IncomingAnyThreadTasksHiPri.IsEmpty())
+		{
+			do
+			{
+				FScopeLock ScopeLock(&CriticalSectionForSortingIncomingAnyThreadTasksHiPri);
+				if (!IncomingAnyThreadTasksHiPri.IsEmpty() && SortedAnyThreadTasksHiPri.IsEmpty())
+				{
+					static TArray<FBaseGraphTask*> NewTasks;
+					NewTasks.Reset();
+					IncomingAnyThreadTasksHiPri.PopAll(NewTasks);
+					check(NewTasks.Num());
+
+					if (NewTasks.Num() > 1)
+					{
+						TLockFreePointerList<FBaseGraphTask> TempSortedAnyThreadTasks;
+						for (int32 Index = 0 ; Index < NewTasks.Num() - 1; Index++) // we are going to take the last one for ourselves
+						{
+							TempSortedAnyThreadTasks.Push(NewTasks[Index]);
+						}
+						verify(SortedAnyThreadTasksHiPri.ReplaceListIfEmpty(TempSortedAnyThreadTasks));
+					}
+					return NewTasks[NewTasks.Num() - 1];
+				}
+				{
+					FBaseGraphTask* Task = SortedAnyThreadTasksHiPri.Pop();
+					if (Task)
+					{
+						return Task;
+					}
+				}
+			} while (!IncomingAnyThreadTasksHiPri.IsEmpty() || !SortedAnyThreadTasksHiPri.IsEmpty());
+		}
 		{
 			FBaseGraphTask* Task = SortedAnyThreadTasks.Pop();
 			if (Task)
@@ -1145,6 +1251,7 @@ public:
 				}
 			}
 		} while (!IncomingAnyThreadTasks.IsEmpty() || !SortedAnyThreadTasks.IsEmpty());
+
 		// this can be called before my constructor is finished
 		for (int32 Pass = 0; Pass < 2; Pass++)
 		{
@@ -1251,8 +1358,11 @@ private:
 	TLockFreePointerList<FTaskThread>		StalledUnnamedThreads; 
 
 	TLockFreePointerList<FBaseGraphTask>		IncomingAnyThreadTasks;
+	TLockFreePointerList<FBaseGraphTask>		IncomingAnyThreadTasksHiPri;
 	TLockFreePointerList<FBaseGraphTask>		SortedAnyThreadTasks;
+	TLockFreePointerList<FBaseGraphTask>		SortedAnyThreadTasksHiPri;
 	FCriticalSection CriticalSectionForSortingIncomingAnyThreadTasks;
+	FCriticalSection CriticalSectionForSortingIncomingAnyThreadTasksHiPri;
 };
 
 
diff --git a/Engine/Source/Runtime/Core/Public/Async/ParallelFor.h b/Engine/Source/Runtime/Core/Public/Async/ParallelFor.h
--- a/Engine/Source/Runtime/Core/Public/Async/ParallelFor.h
+++ b/Engine/Source/Runtime/Core/Public/Async/ParallelFor.h
@@ -23,33 +23,17 @@ struct FParallelForData
 		, DoneEvent(InDoneEvent)
 	{
 	}
-	void Process()
-	{
-		while (true)
-		{
-			int32 MyIndex = IndexToDo.Increment() - 1;
-			if (MyIndex < Num)
-			{
-				Body(MyIndex);
-				if (NumCompleted.Increment() == Num)
-				{
-					DoneEvent.Trigger(); // I was the last one; let parallelfor exit
-				}
-			}
-			else
-			{
-				break;
-			}
-		}
-	}
+	void Process(int32 TasksToSpawn, TSharedRef<FParallelForData, ESPMode::ThreadSafe>& Data);
 };
 
 class FParallelForTask
 {
 	TSharedRef<FParallelForData, ESPMode::ThreadSafe> Data;
+	int32 TasksToSpawn;
 public:
-	FParallelForTask(TSharedRef<FParallelForData, ESPMode::ThreadSafe>& InData)
+	FParallelForTask(TSharedRef<FParallelForData, ESPMode::ThreadSafe>& InData, int32 InTasksToSpawn = 0)
 		: Data(InData) 
+		, TasksToSpawn(InTasksToSpawn)
 	{
 	}
 	static FORCEINLINE TStatId GetStatId()
@@ -58,7 +42,7 @@ public:
 	}
 	static FORCEINLINE ENamedThreads::Type GetDesiredThread()
 	{
-		return ENamedThreads::AnyThread;
+		return ENamedThreads::HiPri(ENamedThreads::AnyThread);
 	}
 	static FORCEINLINE ESubsequentsMode::Type GetSubsequentsMode() 
 	{ 
@@ -66,10 +50,36 @@ public:
 	}
 	void DoTask(ENamedThreads::Type CurrentThread, const FGraphEventRef& MyCompletionGraphEvent)
 	{
-		Data->Process();
+		Data->Process(TasksToSpawn, Data);
 	}
 };
 
+inline void FParallelForData::Process(int32 TasksToSpawn, TSharedRef<FParallelForData, ESPMode::ThreadSafe>& Data)
+{
+	int32 MaybeTasksLeft = Num - IndexToDo.GetValue();
+	if (TasksToSpawn && MaybeTasksLeft > 0)
+	{
+		TasksToSpawn = FMath::Min<int32>(TasksToSpawn, MaybeTasksLeft);
+		TGraphTask<FParallelForTask>::CreateTask().ConstructAndDispatchWhenReady(Data, TasksToSpawn - 1);		
+	}
+	while (true)
+	{
+		int32 MyIndex = IndexToDo.Increment() - 1;
+		if (MyIndex < Num)
+		{
+			Body(MyIndex);
+			if (NumCompleted.Increment() == Num)
+			{
+				DoneEvent.Trigger(); // I was the last one; let parallelfor exit
+			}
+		}
+		else
+		{
+			break;
+		}
+	}
+}
+
 
 /** 
 	*	General purpose parallel for that uses the taskgraph
@@ -100,12 +110,9 @@ inline void ParallelFor(int32 Num, TFunctionRef<void(int32)> Body, bool bForceSi
 
 	FScopedEvent DoneEvent;
 	TSharedRef<FParallelForData, ESPMode::ThreadSafe> Data = MakeShareable(new FParallelForData(Num, Body, DoneEvent));
-	for (int32 Task = 0; Task < AnyThreadTasks; Task++)
-	{
-		TGraphTask<FParallelForTask>::CreateTask().ConstructAndDispatchWhenReady(Data);		
-	}
+	TGraphTask<FParallelForTask>::CreateTask().ConstructAndDispatchWhenReady(Data, AnyThreadTasks - 1);		
 	// this thread can help too and this is important to prevent deadlock on recursion 
-	Data->Process();
+	Data->Process(0, Data);
 	// DoneEvent waits here if some other thread finishes the last item
 	// Data must live on until all of the tasks are cleared which might be long after this function exits
 }
@@ -142,14 +149,11 @@ inline void ParallelForWithPreWork(int32 Num, TFunctionRef<void(int32)> Body, TF
 
 	FScopedEvent DoneEvent;
 	TSharedRef<FParallelForData, ESPMode::ThreadSafe> Data = MakeShareable(new FParallelForData(Num, Body, DoneEvent));
-	for (int32 Task = 0; Task < AnyThreadTasks; Task++)
-	{
-		TGraphTask<FParallelForTask>::CreateTask().ConstructAndDispatchWhenReady(Data);		
-	}
+	TGraphTask<FParallelForTask>::CreateTask().ConstructAndDispatchWhenReady(Data, AnyThreadTasks - 1);		
 	// do the prework
 	CurrentThreadWorkToDoBeforeHelping();
 	// this thread can help too and this is important to prevent deadlock on recursion 
-	Data->Process();
+	Data->Process(0, Data);
 	// DoneEvent waits here if some other thread finishes the last item
 	// Data must live on until all of the tasks are cleared which might be long after this function exits
 }
diff --git a/Engine/Source/Runtime/Core/Public/Async/TaskGraphInterfaces.h b/Engine/Source/Runtime/Core/Public/Async/TaskGraphInterfaces.h
--- a/Engine/Source/Runtime/Core/Public/Async/TaskGraphInterfaces.h
+++ b/Engine/Source/Runtime/Core/Public/Async/TaskGraphInterfaces.h
@@ -29,9 +29,7 @@ namespace ENamedThreads
 {
 	enum Type
 	{
-		/** not actually a thread index. Means "Unknown Thread" or "Any Unnamed Thread" **/
-		AnyThread = -1, 
-
+		UnusedAnchor = -1,
 		/** The always-present, named threads are listed next **/
 #if STATS
 		StatsThread, 
@@ -42,8 +40,10 @@ namespace ENamedThreads
 		ActualRenderingThread = GameThread + 1,
 		// CAUTION ThreadedRenderingThread must be the last named thread, insert new named threads before it
 
+		/** not actually a thread index. Means "Unknown Thread" or "Any Unnamed Thread" **/
+		AnyThread = 0xff, 
 
-		/** High bits are used for a queue index **/
+		/** High bits are used for a queue index and priority**/
 
 		MainQueue =			0x000,
 		LocalQueue =		0x100,
@@ -53,6 +53,16 @@ namespace ENamedThreads
 		QueueIndexMask =	0x100,
 		QueueIndexShift =	8,
 
+		/** High bits are used for a queue index and priority**/
+
+		NormalPriority =	0x000,
+		HighPriority =		0x200,
+
+		NumPriorities =		2,
+		PriorityMask =		0x200,
+		PriorityShift =		9,
+
+
 		/** Combinations **/
 #if STATS
 		StatsThread_Local = StatsThread | LocalQueue,
@@ -65,15 +75,26 @@ namespace ENamedThreads
 
 	FORCEINLINE Type GetThreadIndex(Type ThreadAndIndex)
 	{
-		return (ThreadAndIndex == AnyThread) ? AnyThread : Type(ThreadAndIndex & ThreadIndexMask);
+		return ((ThreadAndIndex & ThreadIndexMask) == AnyThread) ? AnyThread : Type(ThreadAndIndex & ThreadIndexMask);
 	}
 
 	FORCEINLINE int32 GetQueueIndex(Type ThreadAndIndex)
 	{
-		return (ThreadAndIndex & ~ThreadIndexMask) >> QueueIndexShift;
+		return (ThreadAndIndex & QueueIndexMask) >> QueueIndexShift;
+	}
+
+	FORCEINLINE int32 GetPriority(Type ThreadAndIndex)
+	{
+		return (ThreadAndIndex & PriorityMask) >> PriorityShift;
+	}
+
+	FORCEINLINE Type HiPri(Type ThreadAndIndex)
+	{
+		return Type(ThreadAndIndex | HighPriority);
 	}
 }
 
+
 namespace ESubsequentsMode
 {
 	enum Type
@@ -952,7 +973,7 @@ public:
 	 **/
 	void DoTask(ENamedThreads::Type CurrentThread, const FGraphEventRef& MyCompletionGraphEvent)
 	{
-		checkThreadGraph(ThreadToReturnFrom == CurrentThread); // we somehow are executing on the wrong thread.
+		checkThreadGraph(ENamedThreads::GetThreadIndex(ThreadToReturnFrom) == ENamedThreads::GetThreadIndex(CurrentThread)); // we somehow are executing on the wrong thread.
 		FTaskGraphInterface::Get().RequestReturn(ThreadToReturnFrom);
 	}
 
diff --git a/Engine/Source/Runtime/Engine/Classes/Engine/EngineBaseTypes.h b/Engine/Source/Runtime/Engine/Classes/Engine/EngineBaseTypes.h
--- a/Engine/Source/Runtime/Engine/Classes/Engine/EngineBaseTypes.h
+++ b/Engine/Source/Runtime/Engine/Classes/Engine/EngineBaseTypes.h
@@ -121,7 +121,6 @@ struct FTickPrerequisite
 			}
 			return nullptr;
 		}
-	
 };
 
 /** 
@@ -166,6 +165,9 @@ public:
 	UPROPERTY(EditDefaultsOnly, Category="Tick", AdvancedDisplay)
 	uint32 bAllowTickOnDedicatedServer:1;
 
+	/** Run this tick first within the tick group, presumably to start async tasks that must be completed with this tick group, hiding the latency. */
+	uint32 bHighPriority:1;
+
 	/** If false, this tick will run on the game thread, otherwise it will run on any thread in parallel with the game thread and in parallel with other "async ticks" **/
 	uint32 bRunOnAnyThread:1;
 
@@ -246,6 +248,11 @@ public:
 	 * @param TargetTickFunction - Actual tick function to use as a prerequisite
 	 **/
 	void RemovePrerequisite(UObject* TargetObject, struct FTickFunction& TargetTickFunction);
+	/** 
+	 * Sets this function to hipri and all prerequisites recursively
+	 * @param bInHighPriority - priority to set
+	 **/
+	void SetPriorityIncludingPrerequisites(bool bInHighPriority);
 
 	/**
 	 * @return a reference to prerequisites for this tick function.
@@ -270,6 +277,11 @@ private:
 	 */
 	void QueueTickFunction(class FTickTaskSequencer& TTS, const struct FTickContext& TickContext);
 
+	/** 
+	 * Logs the prerequisites
+	 */
+	void ShowPrerequistes(int32 Indent = 1);
+
 	/** 
 	 * Abstract function actually execute the tick. 
 	 * @param DeltaTime - frame time to advance, in seconds
diff --git a/Engine/Source/Runtime/Engine/Private/Components/SkeletalMeshComponent.cpp b/Engine/Source/Runtime/Engine/Private/Components/SkeletalMeshComponent.cpp
--- a/Engine/Source/Runtime/Engine/Private/Components/SkeletalMeshComponent.cpp
+++ b/Engine/Source/Runtime/Engine/Private/Components/SkeletalMeshComponent.cpp
@@ -168,20 +168,20 @@ void USkeletalMeshComponent::RegisterPreClothTick(bool bRegister)
 {
 	if (bRegister != PreClothTickFunction.IsTickFunctionRegistered())
 	{
-		if (bRegister)
+	if (bRegister)
+	{
+		if (SetupActorComponentTickFunction(&PreClothTickFunction))
 		{
-			if (SetupActorComponentTickFunction(&PreClothTickFunction))
+			PreClothTickFunction.Target = this;
+			// Set a prereq for the pre cloth tick to happen after physics is finished
+			if (World != NULL)
 			{
-				PreClothTickFunction.Target = this;
-				// Set a prereq for the pre cloth tick to happen after physics is finished
-				if (World != NULL)
-				{
-					PreClothTickFunction.AddPrerequisite(World, World->EndPhysicsTickFunction);
-				}
+				PreClothTickFunction.AddPrerequisite(World, World->EndPhysicsTickFunction);
 			}
 		}
-		else
-		{
+	}
+	else
+	{
 			PreClothTickFunction.UnRegisterTickFunction();
 		}
 	}
@@ -1031,6 +1031,10 @@ void USkeletalMeshComponent::PerformAnimationEvaluation(const USkeletalMesh* InS
 	FillSpaceBases(InSkeletalMesh, OutLocalAtoms, OutSpaceBases);
 }
 
+static TAutoConsoleVariable<int32> CVarHiPriSkinnedMeshesTicks(
+	TEXT("tick.HiPriSkinnedMeshes"),1,
+	TEXT("If > 0, then schedule the skinned component ticks in a tick group before other ticks."));
+
 void USkeletalMeshComponent::RefreshBoneTransforms(FActorComponentTickFunction* TickFunction)
 {
 	SCOPE_CYCLE_COUNTER(STAT_RefreshBoneTransforms);
@@ -1059,7 +1063,7 @@ void USkeletalMeshComponent::RefreshBoneTransforms(FActorComponentTickFunction*
 	//Dont mark cache as invalid if we aren't performing optimization anyway
 	const bool bInvalidCachedBones = bDoEvaluationRateOptimization &&
 									 ((LocalAtoms.Num() != SkeletalMesh->RefSkeleton.GetNum())
-									 || (LocalAtoms.Num() != CachedLocalAtoms.Num())
+									  || (LocalAtoms.Num() != CachedLocalAtoms.Num())
 									 || (GetNumSpaceBases() != CachedSpaceBases.Num()));
 
 	const bool bShouldDoEvaluation = !bDoEvaluationRateOptimization || bInvalidCachedBones || !AnimUpdateRateParams->ShouldSkipEvaluation();
@@ -1097,6 +1101,14 @@ void USkeletalMeshComponent::RefreshBoneTransforms(FActorComponentTickFunction*
 
 	if (bDoParallelEvaluation)
 	{
+		// this makes no sense to do on the fly. If we want to use hipri ticks, this needs to be done at register time
+		// this does allow us to toggle it easily for performance testing
+		bool bDoHiPri = CVarHiPriSkinnedMeshesTicks.GetValueOnGameThread() > 0;
+		if (TickFunction->bHighPriority != bDoHiPri)
+		{
+			TickFunction->SetPriorityIncludingPrerequisites(bDoHiPri);
+		}
+
 		if (SkeletalMesh->RefSkeleton.GetNum() != AnimEvaluationContext.LocalAtoms.Num())
 		{
 			// Initialize Parallel Task arrays
@@ -1116,8 +1128,8 @@ void USkeletalMeshComponent::RefreshBoneTransforms(FActorComponentTickFunction*
 
 		if ( TickFunction )
 		{
-			TickFunction->GetCompletionHandle()->DontCompleteUntil(TickCompletionEvent);
-		}
+		TickFunction->GetCompletionHandle()->DontCompleteUntil(TickCompletionEvent);
+	}
 	}
 	else
 	{
@@ -1270,8 +1282,8 @@ FBoxSphereBounds USkeletalMeshComponent::CalcBounds(const FTransform& LocalToWor
 			MasterPoseComponentInst->IsA((USkeletalMeshComponent::StaticClass())))
 		{
 			const USkeletalMeshComponent* BaseComponent = CastChecked<USkeletalMeshComponent>(MasterPoseComponentInst);
-			RootBoneOffset = BaseComponent->RootBoneTranslation; // Adjust bounds by root bone translation
-		}
+		RootBoneOffset = BaseComponent->RootBoneTranslation; // Adjust bounds by root bone translation
+	}
 	}
 
 	FBoxSphereBounds NewBounds = CalcMeshBound( RootBoneOffset, bHasValidBodies, LocalToWorld );
@@ -2082,7 +2094,7 @@ void USkeletalMeshComponent::ValidateAnimation()
 		{
 			if (SkeletalMesh->Skeleton)
 			{
-				UE_LOG(LogAnimation, Warning, TEXT("Animation %s is incompatible with skeleton %s, removing animation from actor."), *AnimationData.AnimToPlay->GetName(), *SkeletalMesh->Skeleton->GetName());
+			UE_LOG(LogAnimation, Warning, TEXT("Animation %s is incompatible with skeleton %s, removing animation from actor."), *AnimationData.AnimToPlay->GetName(), *SkeletalMesh->Skeleton->GetName());
 			}
 			else
 			{
@@ -2098,7 +2110,7 @@ void USkeletalMeshComponent::ValidateAnimation()
 		{
 			if(SkeletalMesh->Skeleton)
 			{
-				UE_LOG(LogAnimation, Warning, TEXT("AnimBP %s is incompatible with skeleton %s, removing AnimBP from actor."), *AnimBlueprintGeneratedClass->GetName(), *SkeletalMesh->Skeleton->GetName());
+			UE_LOG(LogAnimation, Warning, TEXT("AnimBP %s is incompatible with skeleton %s, removing AnimBP from actor."), *AnimBlueprintGeneratedClass->GetName(), *SkeletalMesh->Skeleton->GetName());
 			}
 			else
 			{
diff --git a/Engine/Source/Runtime/Engine/Private/GPUSkinVertexFactory.cpp b/Engine/Source/Runtime/Engine/Private/GPUSkinVertexFactory.cpp
--- a/Engine/Source/Runtime/Engine/Private/GPUSkinVertexFactory.cpp
+++ b/Engine/Source/Runtime/Engine/Private/GPUSkinVertexFactory.cpp
@@ -107,6 +107,16 @@ FSharedPoolPolicyData::CreationArguments FBoneBufferPoolPolicy::GetCreationArgum
 	return (Resource.VertexBufferRHI->GetSize());
 }
 
+FBoneBufferTypeRef FClothBufferPoolPolicy::CreateResource(CreationArguments Args)
+{
+	uint32 BufferSize = GetPoolBucketSize(GetPoolBucketIndex(Args));
+	FBoneBuffer Buffer;
+	FRHIResourceCreateInfo CreateInfo;
+	Buffer.VertexBufferRHI = RHICreateVertexBuffer( BufferSize, (BUF_Dynamic | BUF_ShaderResource), CreateInfo );
+	Buffer.VertexBufferSRV = RHICreateShaderResourceView( Buffer.VertexBufferRHI, sizeof(FVector2D), PF_G32R32F );
+	return Buffer;
+}
+
 /*-----------------------------------------------------------------------------
  FBoneBufferPool
  -----------------------------------------------------------------------------*/
@@ -119,13 +129,77 @@ TStatId FBoneBufferPool::GetStatId() const
 	RETURN_QUICK_DECLARE_CYCLE_STAT(FBoneBufferPool, STATGROUP_Tickables);
 }
 
+FClothBufferPool::~FClothBufferPool()
+{
+}
+
+TStatId FClothBufferPool::GetStatId() const
+{
+	RETURN_QUICK_DECLARE_CYCLE_STAT(FClothBufferPool, STATGROUP_Tickables);
+}
+
 TConsoleVariableData<int32>* FGPUBaseSkinVertexFactory::ShaderDataType::MaxBonesVar = NULL;
 uint32 FGPUBaseSkinVertexFactory::ShaderDataType::MaxGPUSkinBones = 0;
 
-void FGPUBaseSkinVertexFactory::ShaderDataType::UpdateBoneData(ERHIFeatureLevel::Type FeatureLevel)
+static TAutoConsoleVariable<int32> CVarRHICmdDeferSkeletalLockAndFillToRHIThread(
+	TEXT("r.RHICmdDeferSkeletalLockAndFillToRHIThread"),
+	0,
+	TEXT("If > 0, then do the bone and cloth copies on the RHI thread. Experimental option."));
+
+static bool DeferSkeletalLockAndFillToRHIThread()
 {
-	uint32 NumBones = BoneMatrices.Num();
+	return GRHIThread && CVarRHICmdDeferSkeletalLockAndFillToRHIThread.GetValueOnRenderThread() > 0;
+}
+
+struct FRHICommandUpdateBoneBuffer : public FRHICommand<FRHICommandUpdateBoneBuffer>
+{
+	FVertexBufferRHIParamRef VertexBuffer;
+	uint32 BufferSize;
+	const TArray<FMatrix>& ReferenceToLocalMatrices;
+	const TArray<FBoneIndexType>& BoneMap;
+
+
+	FORCEINLINE_DEBUGGABLE FRHICommandUpdateBoneBuffer(FVertexBufferRHIParamRef InVertexBuffer, uint32 InBufferSize, const TArray<FMatrix>& InReferenceToLocalMatrices, const TArray<FBoneIndexType>& InBoneMap)
+		: VertexBuffer(InVertexBuffer)
+		, BufferSize(InBufferSize)
+		, ReferenceToLocalMatrices(InReferenceToLocalMatrices)
+		, BoneMap(InBoneMap)
+	{
+	}
+	void Execute(FRHICommandListBase& CmdList)
+	{
+		QUICK_SCOPE_CYCLE_COUNTER(STAT_FRHICommandUpdateBoneBuffer_Execute);
+		FBoneSkinning* ChunkMatrices = (FBoneSkinning*)GDynamicRHI->RHILockVertexBuffer(VertexBuffer, 0, BufferSize, RLM_WriteOnly);
+		//FSkinMatrix3x4 is sizeof() == 48
+		// CACHE_LINE_SIZE (128) / 48 = 2.6
+		//  sizeof(FMatrix) == 64
+		// CACHE_LINE_SIZE (128) / 64 = 2
+		const uint32 NumBones = BoneMap.Num();
+		check(NumBones > 0 && NumBones < 256); // otherwise maybe some bad threading on BoneMap, maybe we need to copy that
+		const int32 PreFetchStride = 2; // FPlatformMisc::Prefetch stride
+		for (uint32 BoneIdx = 0; BoneIdx < NumBones; BoneIdx++)
+		{
+			const FBoneIndexType RefToLocalIdx = BoneMap[BoneIdx];
+			check(ReferenceToLocalMatrices.IsValidIndex(RefToLocalIdx)); // otherwise maybe some bad threading on BoneMap, maybe we need to copy that
+			FPlatformMisc::Prefetch( ReferenceToLocalMatrices.GetData() + RefToLocalIdx + PreFetchStride );
+			FPlatformMisc::Prefetch( ReferenceToLocalMatrices.GetData() + RefToLocalIdx + PreFetchStride, CACHE_LINE_SIZE );
+
+			FBoneSkinning& BoneMat = ChunkMatrices[BoneIdx];
+			const FMatrix& RefToLocal = ReferenceToLocalMatrices[RefToLocalIdx];
+			RefToLocal.To3x4MatrixTranspose( (float*)BoneMat.M );
+		}
+		GDynamicRHI->RHIUnlockVertexBuffer(VertexBuffer);
+	}
+};
+
+bool FGPUBaseSkinVertexFactory::ShaderDataType::UpdateBoneData(FRHICommandListImmediate& RHICmdList, const TArray<FMatrix>& ReferenceToLocalMatrices, const TArray<FBoneIndexType>& BoneMap, ERHIFeatureLevel::Type FeatureLevel, bool bUseSkinCache)
+{
+	HackReferenceToLocalMatrices = &ReferenceToLocalMatrices;
+	HackBoneMap = &BoneMap;
+
+	const uint32 NumBones = BoneMap.Num();
 	check(NumBones <= MaxGPUSkinBones);
+	FBoneSkinning* ChunkMatrices = nullptr;
 	if (FeatureLevel >= ERHIFeatureLevel::ES3_1)
 	{
 		static FSharedPoolPolicyData PoolPolicy;
@@ -144,18 +218,53 @@ void FGPUBaseSkinVertexFactory::ShaderDataType::UpdateBoneData(ERHIFeatureLevel:
 		}
 		if(NumBones)
 		{
-				float* Data = (float*)RHILockVertexBuffer(BoneBuffer.VertexBufferRHI, 0, VectorArraySize, RLM_WriteOnly);
-				checkSlow(Data);
-				FMemory::Memcpy(Data, BoneMatrices.GetData(), NumBones * sizeof(BoneMatrices[0]));
-				RHIUnlockVertexBuffer(BoneBuffer.VertexBufferRHI);
+			if (!bUseSkinCache && DeferSkeletalLockAndFillToRHIThread())
+			{
+				new (RHICmdList.AllocCommand<FRHICommandUpdateBoneBuffer>()) FRHICommandUpdateBoneBuffer(BoneBuffer.VertexBufferRHI, VectorArraySize, ReferenceToLocalMatrices, BoneMap);
+				return true;
 			}
+			ChunkMatrices = (FBoneSkinning*)RHILockVertexBuffer(BoneBuffer.VertexBufferRHI, 0, VectorArraySize, RLM_WriteOnly);
 		}
+	}
+	else
+	{
+		if(NumBones)
+		{
+			check(NumBones * sizeof(FBoneSkinning) <= sizeof(GBoneUniformStruct));
+			ChunkMatrices = (FBoneSkinning*)&GBoneUniformStruct;
+		}
+	}
+
+	{
+		QUICK_SCOPE_CYCLE_COUNTER(STAT_FGPUBaseSkinVertexFactory_ShaderDataType_UpdateBoneData_CopyBones);
+		//FSkinMatrix3x4 is sizeof() == 48
+		// CACHE_LINE_SIZE (128) / 48 = 2.6
+		//  sizeof(FMatrix) == 64
+		// CACHE_LINE_SIZE (128) / 64 = 2
+		const int32 PreFetchStride = 2; // FPlatformMisc::Prefetch stride
+		for (uint32 BoneIdx = 0; BoneIdx < NumBones; BoneIdx++)
+		{
+			const FBoneIndexType RefToLocalIdx = BoneMap[BoneIdx];
+			FPlatformMisc::Prefetch( ReferenceToLocalMatrices.GetData() + RefToLocalIdx + PreFetchStride );
+			FPlatformMisc::Prefetch( ReferenceToLocalMatrices.GetData() + RefToLocalIdx + PreFetchStride, CACHE_LINE_SIZE );
+
+			FBoneSkinning& BoneMat = ChunkMatrices[BoneIdx];
+			const FMatrix& RefToLocal = ReferenceToLocalMatrices[RefToLocalIdx];
+			RefToLocal.To3x4MatrixTranspose( (float*)BoneMat.M );
+		}
+	}
+	if (FeatureLevel >= ERHIFeatureLevel::ES3_1)
+	{
+		if (NumBones)
+		{
+			RHIUnlockVertexBuffer(BoneBuffer.VertexBufferRHI);
+		}
+	}
 	else
 	{
-		check(NumBones * BoneMatrices.GetTypeSize() <= sizeof(GBoneUniformStruct));
-		FMemory::Memcpy(&GBoneUniformStruct, BoneMatrices.GetData(), NumBones * sizeof(BoneMatrices[0]));
 		UniformBuffer = RHICreateUniformBuffer(&GBoneUniformStruct, FBoneMatricesUniformShaderParameters::StaticStruct.GetLayout(), UniformBuffer_MultiFrame);
 	}
+	return false;
 }
 
 /*-----------------------------------------------------------------------------
@@ -307,7 +416,7 @@ template <bool bExtraBoneInfluencesT>
 void TGPUSkinVertexFactory<bExtraBoneInfluencesT>::InitDynamicRHI()
 {
 	FVertexFactory::InitDynamicRHI();
-	ShaderData.UpdateBoneData(GetFeatureLevel());
+	//ShaderData.UpdateBoneData(GetFeatureLevel());
 }
 
 template <bool bExtraBoneInfluencesT>
@@ -409,7 +518,8 @@ public:
 			bool bLocalPerBoneMotionBlur = false;
 
 
-			if (FeatureLevel >= ERHIFeatureLevel::SM4 && GPrevPerBoneMotionBlur.IsVelocityPass(RHICmdList) && PreviousBoneMatrices.IsBound())
+			if (FeatureLevel >= ERHIFeatureLevel::SM4 && GPrevPerBoneMotionBlur.IsVelocityPass(RHICmdList) && PreviousBoneMatrices.IsBound() &&
+				ShaderData.HackReferenceToLocalMatrices && ShaderData.HackBoneMap)
 			{
 				const bool bWorldIsPaused = View.Family->bWorldIsPaused;
 
@@ -474,18 +584,46 @@ public:
 #if !(UE_BUILD_SHIPPING || UE_BUILD_TEST)
 						if (CVarVelocityTest.GetValueOnRenderThread())
 						{
-							FBoneSkinning Tab[10];
-
-								FMemory::Memset(Tab, 0, sizeof(FBoneSkinning) * 10);
-								// We add some black gap between the elements, to ensure we pick the right data for this frame
-								// wwe create different gap size on alternating frames.
-								GPrevPerBoneMotionBlur.AppendData(Tab, (FrameNumber % 2) + 4);
+							// We add some black gap between the elements, to ensure we pick the right data for this frame
+							// wwe create different gap size on alternating frames.
+							FBoneSkinning* ChunkMatrices = nullptr;
+							GPrevPerBoneMotionBlur.AppendData(ChunkMatrices, (FrameNumber % 2) + 4);
+							if (ChunkMatrices)
+							{
+								FMemory::Memset(ChunkMatrices, 0, sizeof(FBoneSkinning) * 10);
+							}
 						}
 #endif // if !(UE_BUILD_SHIPPING || UE_BUILD_TEST)
 
+						const TArray<FMatrix>& ReferenceToLocalMatrices = *ShaderData.HackReferenceToLocalMatrices;
+						const TArray<FBoneIndexType>& BoneMap = *ShaderData.HackBoneMap;
 						// copy the bone data and tell the instance where it can pick it up next frame
 						// append data to a buffer we bind next frame to read old matrix data for motion blur
-						uint32 OldBoneDataStartIndex = GPrevPerBoneMotionBlur.AppendData(ShaderData.BoneMatrices.GetData(), ShaderData.BoneMatrices.Num());
+						const int32 NumBones = BoneMap.Num();
+						FBoneSkinning* ChunkMatrices = nullptr;
+						uint32 OldBoneDataStartIndex = GPrevPerBoneMotionBlur.AppendData(ChunkMatrices, NumBones);
+						if (ChunkMatrices)
+						{
+							//FSkinMatrix3x4 is sizeof() == 48
+							// CACHE_LINE_SIZE (128) / 48 = 2.6
+							//  sizeof(FMatrix) == 64
+							// CACHE_LINE_SIZE (128) / 64 = 2
+							const int32 PreFetchStride = 2; // FPlatformMisc::Prefetch stride
+
+							{
+								QUICK_SCOPE_CYCLE_COUNTER(STAT_FGPUBaseSkinVertexFactory_SFGPUSkinVertexFactoryShaderParameters_CopyMotionBlurBones);
+								for( int32 BoneIdx=0; BoneIdx < NumBones; BoneIdx++ )
+								{
+									const FBoneIndexType RefToLocalIdx = BoneMap[BoneIdx];
+									FPlatformMisc::Prefetch( ReferenceToLocalMatrices.GetData() + RefToLocalIdx + PreFetchStride );
+									FPlatformMisc::Prefetch( ReferenceToLocalMatrices.GetData() + RefToLocalIdx + PreFetchStride, CACHE_LINE_SIZE );
+
+									FBoneSkinning& BoneMat = ChunkMatrices[BoneIdx];
+									const FMatrix& RefToLocal = ReferenceToLocalMatrices[RefToLocalIdx];
+									RefToLocal.To3x4MatrixTranspose( (float*)BoneMat.M );
+								}
+							}
+						}
 						GPUVertexFactory->SetOldBoneDataStartIndex(FrameNumber, OldBoneDataStartIndex);
 					}
 				}
@@ -666,8 +804,7 @@ public:
 	virtual void Bind(const FShaderParameterMap& ParameterMap) override
 	{
 		FGPUSkinVertexFactoryShaderParameters::Bind(ParameterMap);
-		ClothSimulPositionsParameter.Bind(ParameterMap,TEXT("ClothSimulVertsPositions"));
-		ClothSimulNormalsParameter.Bind(ParameterMap,TEXT("ClothSimulVertsNormals"));
+		ClothSimulVertsPositionsNormalsParameter.Bind(ParameterMap,TEXT("ClothSimulVertsPositionsNormals"));
 		ClothBlendWeightParameter.Bind(ParameterMap, TEXT("ClothBlendWeight"));
 	}
 	/**
@@ -677,8 +814,7 @@ public:
 	virtual void Serialize(FArchive& Ar) override
 	{ 
 		FGPUSkinVertexFactoryShaderParameters::Serialize(Ar);
-		Ar << ClothSimulPositionsParameter;
-		Ar << ClothSimulNormalsParameter;
+		Ar << ClothSimulVertsPositionsNormalsParameter;
 		Ar << ClothBlendWeightParameter;
 	}
 
@@ -698,16 +834,11 @@ public:
 			SetUniformBufferParameter(RHICmdList, Shader->GetVertexShader(),Shader->GetUniformBufferParameter<FAPEXClothUniformShaderParameters>(),ClothShaderData.GetClothUniformBuffer());
 
 			// we tell the shader where to pickup the data
-			if(ClothSimulPositionsParameter.IsBound())
+			if(ClothSimulVertsPositionsNormalsParameter.IsBound())
 			{
-				RHICmdList.SetShaderResourceViewParameter(Shader->GetVertexShader(), ClothSimulPositionsParameter.GetBaseIndex(), ClothShaderData.GetClothSimulPositionBuffer().VertexBufferSRV);
+				RHICmdList.SetShaderResourceViewParameter(Shader->GetVertexShader(), ClothSimulVertsPositionsNormalsParameter.GetBaseIndex(), ClothShaderData.GetClothSimulPositionNormalBuffer().VertexBufferSRV);
 			}
 			
-			if(ClothSimulNormalsParameter.IsBound())
-			{
-				RHICmdList.SetShaderResourceViewParameter(Shader->GetVertexShader(), ClothSimulNormalsParameter.GetBaseIndex(), ClothShaderData.GetClothSimulNormalBuffer().VertexBufferSRV);
-			}
-
 			SetShaderValue(
 				RHICmdList,
 				Shader->GetVertexShader(),
@@ -718,61 +849,107 @@ public:
 	}
 
 protected:
-	FShaderResourceParameter ClothSimulPositionsParameter;
-	FShaderResourceParameter ClothSimulNormalsParameter;
+	FShaderResourceParameter ClothSimulVertsPositionsNormalsParameter;
 	FShaderParameter ClothBlendWeightParameter;
 };
 
 /*-----------------------------------------------------------------------------
 	TGPUSkinAPEXClothVertexFactory::ClothShaderType
 -----------------------------------------------------------------------------*/
-void FGPUBaseSkinAPEXClothVertexFactory::ClothShaderType::UpdateClothSimulData(const TArray<FVector4>& InSimulPositions, const TArray<FVector4>& InSimulNormals, ERHIFeatureLevel::Type FeatureLevel)
+
+struct FRHICommandUpdateClothBuffer : public FRHICommand<FRHICommandUpdateClothBuffer>
+{
+	FVertexBufferRHIParamRef VertexBuffer;
+	uint32 BufferSize;
+	const TArray<FVector4>& SimulPositions;
+	const TArray<FVector4>& SimulNormals;
+
+
+	FORCEINLINE_DEBUGGABLE FRHICommandUpdateClothBuffer(FVertexBufferRHIParamRef InVertexBuffer, uint32 InBufferSize, const TArray<FVector4>& InSimulPositions, const TArray<FVector4>& InSimulNormals)
+		: VertexBuffer(InVertexBuffer)
+		, BufferSize(InBufferSize)
+		, SimulPositions(InSimulPositions)
+		, SimulNormals(InSimulNormals)
+	{
+	}
+	void Execute(FRHICommandListBase& CmdList)
+	{
+		QUICK_SCOPE_CYCLE_COUNTER(STAT_FRHICommandUpdateBoneBuffer_Execute);
+		float* RESTRICT Data = (float* RESTRICT)GDynamicRHI->RHILockVertexBuffer(VertexBuffer, 0, BufferSize, RLM_WriteOnly);
+		uint32 NumSimulVerts = SimulPositions.Num();
+		check(NumSimulVerts > 0 && NumSimulVerts <= MAX_APEXCLOTH_VERTICES_FOR_VB);
+		float* RESTRICT Pos = (float* RESTRICT) &SimulPositions[0].X;
+		float* RESTRICT Normal = (float* RESTRICT) &SimulNormals[0].X;
+		for (uint32 Index = 0; Index < NumSimulVerts; Index++)
+		{
+			Data[0] = Pos[0];
+			Data[1] = Pos[1];
+			Data[2] = Pos[2];
+			Data[3] = Normal[0];
+			Data[4] = Normal[1];
+			Data[5] = Normal[2];
+			Data += 6;
+			Pos += 4;
+			Normal += 4;
+		}
+		GDynamicRHI->RHIUnlockVertexBuffer(VertexBuffer);
+	}
+};
+
+bool FGPUBaseSkinAPEXClothVertexFactory::ClothShaderType::UpdateClothSimulData(FRHICommandListImmediate& RHICmdList, const TArray<FVector4>& InSimulPositions, const TArray<FVector4>& InSimulNormals, ERHIFeatureLevel::Type FeatureLevel)
 {
+	QUICK_SCOPE_CYCLE_COUNTER(STAT_FGPUBaseSkinAPEXClothVertexFactory_UpdateClothSimulData);
 	uint32 NumSimulVerts = InSimulPositions.Num();
 
 	if (FeatureLevel >= ERHIFeatureLevel::SM4)
 	{
 		FMath::Min(NumSimulVerts, (uint32)MAX_APEXCLOTH_VERTICES_FOR_VB);
 
-		uint32 VectorArraySize = NumSimulVerts * sizeof(FVector4);
+		uint32 VectorArraySize = NumSimulVerts * sizeof(float) * 6;
 		uint32 PooledArraySize = ClothSimulDataBufferPool.PooledSizeForCreationArguments(VectorArraySize);
-		if(!IsValidRef(ClothSimulPositionBuffer) || PooledArraySize != ClothSimulPositionBuffer.VertexBufferRHI->GetSize())
+		if(!IsValidRef(ClothSimulPositionNormalBuffer) || PooledArraySize != ClothSimulPositionNormalBuffer.VertexBufferRHI->GetSize())
 		{
-			if(IsValidRef(ClothSimulPositionBuffer))
+			if(IsValidRef(ClothSimulPositionNormalBuffer))
 			{
-				ClothSimulDataBufferPool.ReleasePooledResource(ClothSimulPositionBuffer);
+				ClothSimulDataBufferPool.ReleasePooledResource(ClothSimulPositionNormalBuffer);
 			}
-			ClothSimulPositionBuffer = ClothSimulDataBufferPool.CreatePooledResource(VectorArraySize);
-			check(IsValidRef(ClothSimulPositionBuffer));
+			ClothSimulPositionNormalBuffer = ClothSimulDataBufferPool.CreatePooledResource(VectorArraySize);
+			check(IsValidRef(ClothSimulPositionNormalBuffer));
 		}
 
-		if(!IsValidRef(ClothSimulNormalBuffer) || PooledArraySize != ClothSimulNormalBuffer.VertexBufferRHI->GetSize())
+		if(NumSimulVerts)
 		{
-			if(IsValidRef(ClothSimulNormalBuffer))
+			if (DeferSkeletalLockAndFillToRHIThread())
 			{
-				ClothSimulDataBufferPool.ReleasePooledResource(ClothSimulNormalBuffer);
+				new (RHICmdList.AllocCommand<FRHICommandUpdateClothBuffer>()) FRHICommandUpdateClothBuffer(ClothSimulPositionNormalBuffer.VertexBufferRHI, VectorArraySize, InSimulPositions, InSimulNormals);
+				return true;
 			}
-			ClothSimulNormalBuffer = ClothSimulDataBufferPool.CreatePooledResource(VectorArraySize);
-			check(IsValidRef(ClothSimulNormalBuffer));
-		}
-
-		if(NumSimulVerts)
-		{
-				float* Data = (float*)RHILockVertexBuffer(ClothSimulPositionBuffer.VertexBufferRHI, 0, VectorArraySize, RLM_WriteOnly);
-				checkSlow(Data);
-				FMemory::Memcpy(Data, InSimulPositions.GetData(), NumSimulVerts * sizeof(FVector4));
-				RHIUnlockVertexBuffer(ClothSimulPositionBuffer.VertexBufferRHI);
-
-				Data = (float*)RHILockVertexBuffer(ClothSimulNormalBuffer.VertexBufferRHI, 0, VectorArraySize, RLM_WriteOnly);
-				checkSlow(Data);
-				FMemory::Memcpy(Data, InSimulNormals.GetData(), NumSimulVerts * sizeof(FVector4));
-				RHIUnlockVertexBuffer(ClothSimulNormalBuffer.VertexBufferRHI);
+			float* RESTRICT Data = (float* RESTRICT)RHILockVertexBuffer(ClothSimulPositionNormalBuffer.VertexBufferRHI, 0, VectorArraySize, RLM_WriteOnly);
+			{
+				QUICK_SCOPE_CYCLE_COUNTER(STAT_FGPUBaseSkinAPEXClothVertexFactory_UpdateClothSimulData_CopyData);
+				float* RESTRICT Pos = (float* RESTRICT) &InSimulPositions[0].X;
+				float* RESTRICT Normal = (float* RESTRICT) &InSimulNormals[0].X;
+				for (uint32 Index = 0; Index < NumSimulVerts; Index++)
+				{
+					Data[0] = Pos[0];
+					Data[1] = Pos[1];
+					Data[2] = Pos[2];
+					Data[3] = Normal[0];
+					Data[4] = Normal[1];
+					Data[5] = Normal[2];
+					Data += 6;
+					Pos += 4;
+					Normal += 4;
+				}
 			}
+			RHIUnlockVertexBuffer(ClothSimulPositionNormalBuffer.VertexBufferRHI);
 		}
+	}
 	else
 	{
 		UpdateClothUniformBuffer(InSimulPositions, InSimulNormals);
 	}
+	return false;
 }
 
 /*-----------------------------------------------------------------------------
diff --git a/Engine/Source/Runtime/Engine/Private/LevelTick.cpp b/Engine/Source/Runtime/Engine/Private/LevelTick.cpp
--- a/Engine/Source/Runtime/Engine/Private/LevelTick.cpp
+++ b/Engine/Source/Runtime/Engine/Private/LevelTick.cpp
@@ -983,8 +983,6 @@ public:
 } FileProfileWrapperExec;
 #endif // !UE_BUILD_SHIPPING
 
-
-
 #if ENABLE_COLLISION_ANALYZER
 #include "CollisionAnalyzerModule.h"
 extern bool GCollisionAnalyzerIsRecording;
@@ -997,6 +995,16 @@ static TAutoConsoleVariable<int32> CVarStressTestGCWhileStreaming(
 	TEXT("If set to 1, the engine will attempt to trigger GC each frame while async loading."));
 #endif
 
+DECLARE_CYCLE_STAT(TEXT("TG_PrePhysics"), STAT_TG_PrePhysics, STATGROUP_TickGroups);
+DECLARE_CYCLE_STAT(TEXT("TG_StartPhysics"), STAT_TG_StartPhysics, STATGROUP_TickGroups);
+DECLARE_CYCLE_STAT(TEXT("Start TG_DuringPhysics"), STAT_TG_DuringPhysics, STATGROUP_TickGroups);
+DECLARE_CYCLE_STAT(TEXT("TG_EndPhysics"), STAT_TG_EndPhysics, STATGROUP_TickGroups);
+DECLARE_CYCLE_STAT(TEXT("TG_PreCloth"), STAT_TG_PreCloth, STATGROUP_TickGroups);
+DECLARE_CYCLE_STAT(TEXT("TG_StartCloth"), STAT_TG_StartCloth, STATGROUP_TickGroups);
+DECLARE_CYCLE_STAT(TEXT("TG_EndCloth"), STAT_TG_EndCloth, STATGROUP_TickGroups);
+DECLARE_CYCLE_STAT(TEXT("TG_PostPhysics"), STAT_TG_PostPhysics, STATGROUP_TickGroups);
+DECLARE_CYCLE_STAT(TEXT("TG_PostUpdateWork"), STAT_TG_PostUpdateWork, STATGROUP_TickGroups);
+
 static float GTimeBetweenPurgingPendingKillObjects = 60.0f;
 static FAutoConsoleVariableRef CVarTimeBetweenPurgingPendingKillObjects(
 	TEXT("gc.TimeBetweenPurgingPendingKillObjects"),
@@ -1136,28 +1144,49 @@ void UWorld::Tick( ELevelTick TickType, float DeltaSeconds )
 		TickGroup = TG_PrePhysics; // reset this to the start tick group
 		FTickTaskManagerInterface::Get().StartFrame(this, DeltaSeconds, TickType);
 
-		///////////////////
-		// There is an opportunity for the game thread to do a ms or two of work here "for free" while the tick scheduling proceeds on other cores.
-		// We can run any code that doesn't affect the registration, enabled state lifetime or basically anything else relating to any FTickFunction
-
 		SCOPE_CYCLE_COUNTER(STAT_TickTime);
+		{
+			SCOPE_CYCLE_COUNTER(STAT_TG_PrePhysics);
 		RunTickGroup(TG_PrePhysics);
+		}
         bInTick = false;
         EnsureCollisionTreeIsBuilt();
         bInTick = true;
+		{
+			SCOPE_CYCLE_COUNTER(STAT_TG_StartPhysics);
 		RunTickGroup(TG_StartPhysics); 
+		}
+		{
+			SCOPE_CYCLE_COUNTER(STAT_TG_DuringPhysics);
+			QUICK_SCOPE_CYCLE_COUNTER(FStat_Tick_PostPhysics);
 		RunTickGroup(TG_DuringPhysics, false); // No wait here, we should run until idle though. We don't care if all of the async ticks are done before we start running post-phys stuff
+		}
 		TickGroup = TG_EndPhysics; // set this here so the current tick group is correct during collision notifies, though I am not sure it matters. 'cause of the false up there^^^
+		{
+			SCOPE_CYCLE_COUNTER(STAT_TG_EndPhysics);
 		RunTickGroup(TG_EndPhysics);
+		}
 		if ( PhysicsScene != NULL )
 		{
 			PhysicsScene->DeferredCommandHandler.Flush();
 		}
+		{
+			SCOPE_CYCLE_COUNTER(STAT_TG_PreCloth);
 		RunTickGroup(TG_PreCloth);
+		}
+		{
+			SCOPE_CYCLE_COUNTER(STAT_TG_StartCloth);
 		RunTickGroup(TG_StartCloth);
+		}
+		{
+			SCOPE_CYCLE_COUNTER(STAT_TG_EndCloth);
 		RunTickGroup(TG_EndCloth);
+		}
+		{
+			SCOPE_CYCLE_COUNTER(STAT_TG_PostPhysics);
 		RunTickGroup(TG_PostPhysics);
 	}
+	}
 	else if( bIsPaused )
 	{
 		FTickTaskManagerInterface::Get().RunPauseFrame(this, DeltaSeconds, LEVELTICK_PauseTick);
@@ -1249,7 +1278,10 @@ void UWorld::Tick( ELevelTick TickType, float DeltaSeconds )
 	if (bDoingActorTicks)
 	{
 		SCOPE_CYCLE_COUNTER(STAT_TickTime);
+		{
+			SCOPE_CYCLE_COUNTER(STAT_TG_PostUpdateWork);
 		RunTickGroup(TG_PostUpdateWork);
+		}
 		FTickTaskManagerInterface::Get().EndFrame(); 
 
 		// All tick is done, execute async trace
diff --git a/Engine/Source/Runtime/Engine/Private/PrimitiveSceneProxy.cpp b/Engine/Source/Runtime/Engine/Private/PrimitiveSceneProxy.cpp
--- a/Engine/Source/Runtime/Engine/Private/PrimitiveSceneProxy.cpp
+++ b/Engine/Source/Runtime/Engine/Private/PrimitiveSceneProxy.cpp
@@ -7,6 +7,7 @@
 #include "EnginePrivate.h"
 #include "PrimitiveSceneProxy.h"
 #include "Components/BrushComponent.h"
+#include "PrimitiveSceneInfo.h"
 
 FPrimitiveSceneProxy::FPrimitiveSceneProxy(const UPrimitiveComponent* InComponent, FName InResourceName)
 :	WireframeColor(FLinearColor::White)
@@ -186,25 +187,54 @@ void FPrimitiveSceneProxy::UpdateActorPosition(FVector InActorPosition)
 		if (PrimitiveSceneProxy->ActorPosition != InActorPosition)
 		{
 			PrimitiveSceneProxy->ActorPosition = InActorPosition;
-			// Update the uniform shader parameters.
-			const FPrimitiveUniformShaderParameters PrimitiveUniformShaderParameters =
-			GetPrimitiveUniformShaderParameters(
-				PrimitiveSceneProxy->LocalToWorld, 
-				InActorPosition, 
-				PrimitiveSceneProxy->Bounds,
-				PrimitiveSceneProxy->LocalBounds, 
-				PrimitiveSceneProxy->bReceivesDecals, 
-				PrimitiveSceneProxy->HasDistanceFieldRepresentation(), 
-				PrimitiveSceneProxy->SupportsHeightfieldRepresentation(),
-				PrimitiveSceneProxy->UseEditorDepthTest(),
-				PrimitiveSceneProxy->GetLpvBiasMultiplier() );
-
-			PrimitiveSceneProxy->UniformBuffer.SetContents(PrimitiveUniformShaderParameters);
+			PrimitiveSceneProxy->UpdateUniformBufferMaybeLazy();
 			PrimitiveSceneProxy->OnActorPositionChanged();
 		}
 	});
 }
 
+static TAutoConsoleVariable<int32> CVarDeferUniformBufferUpdatesUntilVisible(
+	TEXT("r.DeferUniformBufferUpdatesUntilVisible"),
+	0,
+	TEXT("If > 0, then don't update the primitive uniform buffer until it is visible. Experimental option."));
+
+void FPrimitiveSceneProxy::UpdateUniformBufferMaybeLazy()
+{
+	if (CVarDeferUniformBufferUpdatesUntilVisible.GetValueOnAnyThread() > 0)
+	{
+		if (PrimitiveSceneInfo)
+		{
+			PrimitiveSceneInfo->SetNeedsUniformBufferUpdate(true);
+		}
+	}
+	else
+	{
+		UpdateUniformBuffer();
+		if (PrimitiveSceneInfo)
+		{
+			PrimitiveSceneInfo->SetNeedsUniformBufferUpdate(false);
+		}
+	}
+}
+
+void FPrimitiveSceneProxy::UpdateUniformBuffer()
+{
+	QUICK_SCOPE_CYCLE_COUNTER(STAT_FPrimitiveSceneProxy_UpdateUniformBuffer);
+	// Update the uniform shader parameters.
+	const FPrimitiveUniformShaderParameters PrimitiveUniformShaderParameters = 
+		GetPrimitiveUniformShaderParameters(
+			LocalToWorld, 
+			ActorPosition, 
+			Bounds, 
+			LocalBounds, 
+			bReceivesDecals, 
+			HasDistanceFieldRepresentation(), 
+			SupportsHeightfieldRepresentation(), 
+			UseEditorDepthTest(), 
+			LpvBiasMultiplier );
+	UniformBuffer.SetContents(PrimitiveUniformShaderParameters);
+}
+
 void FPrimitiveSceneProxy::SetTransform(const FMatrix& InLocalToWorld, const FBoxSphereBounds& InBounds, const FBoxSphereBounds& InLocalBounds, FVector InActorPosition)
 {
 	check(IsInRenderingThread());
@@ -218,9 +248,7 @@ void FPrimitiveSceneProxy::SetTransform(const FMatrix& InLocalToWorld, const FBo
 	LocalBounds = InLocalBounds;
 	ActorPosition = InActorPosition;
 	
-	// Update the uniform shader parameters.
-	const FPrimitiveUniformShaderParameters PrimitiveUniformShaderParameters = GetPrimitiveUniformShaderParameters(LocalToWorld, ActorPosition, Bounds, LocalBounds, bReceivesDecals, HasDistanceFieldRepresentation(), SupportsHeightfieldRepresentation(), UseEditorDepthTest(), LpvBiasMultiplier );
-	UniformBuffer.SetContents(PrimitiveUniformShaderParameters);
+	UpdateUniformBufferMaybeLazy();
 	
 	// Notify the proxy's implementation of the change.
 	OnTransformChanged();
diff --git a/Engine/Source/Runtime/Engine/Private/SkeletalMesh.cpp b/Engine/Source/Runtime/Engine/Private/SkeletalMesh.cpp
--- a/Engine/Source/Runtime/Engine/Private/SkeletalMesh.cpp
+++ b/Engine/Source/Runtime/Engine/Private/SkeletalMesh.cpp
@@ -4353,6 +4353,7 @@ void FSkeletalMeshSceneProxy::GetMeshElementsConditionallySelectable(const TArra
 	{
 		return;
 	}	
+	MeshObject->PreGDMECallback();
 
 	for (int32 ViewIndex = 0; ViewIndex < Views.Num(); ViewIndex++)
 	{
diff --git a/Engine/Source/Runtime/Engine/Private/SkeletalRenderGPUSkin.cpp b/Engine/Source/Runtime/Engine/Private/SkeletalRenderGPUSkin.cpp
--- a/Engine/Source/Runtime/Engine/Private/SkeletalRenderGPUSkin.cpp
+++ b/Engine/Source/Runtime/Engine/Private/SkeletalRenderGPUSkin.cpp
@@ -96,6 +96,8 @@ FSkeletalMeshObjectGPUSkin
 FSkeletalMeshObjectGPUSkin::FSkeletalMeshObjectGPUSkin(USkinnedMeshComponent* InMeshComponent, FSkeletalMeshResource* InSkeletalMeshResource, ERHIFeatureLevel::Type InFeatureLevel)
 	: FSkeletalMeshObject(InMeshComponent, InSkeletalMeshResource, InFeatureLevel)
 ,	DynamicData(NULL)
+,	bNeedsUpdateDeferred(false)
+,	bMorphNeedsUpdateDeferred(false)
 ,	bMorphResourcesInitialized(false)
 {
 	// create LODs to match the base mesh
@@ -111,6 +113,7 @@ FSkeletalMeshObjectGPUSkin::FSkeletalMeshObjectGPUSkin(USkinnedMeshComponent* In
 
 FSkeletalMeshObjectGPUSkin::~FSkeletalMeshObjectGPUSkin()
 {
+	check(!RHIThreadFenceForDynamicData.GetReference());
 	delete DynamicData;
 }
 
@@ -134,6 +137,14 @@ void FSkeletalMeshObjectGPUSkin::ReleaseResources()
 	}
 	// also release morph resources
 	ReleaseMorphResources();
+	ENQUEUE_UNIQUE_RENDER_COMMAND_ONEPARAMETER(
+		WaitRHIThreadFenceForDynamicData,
+		FSkeletalMeshObjectGPUSkin*, MeshObject, this,
+	{
+		FScopeCycleCounter Context(MeshObject->GetStatId());
+		MeshObject->WaitForRHIThreadFenceForDynamicData();
+	}
+	);
 }
 
 void FSkeletalMeshObjectGPUSkin::InitMorphResources(bool bInUsePerBoneMotionBlur)
@@ -199,20 +210,62 @@ void FSkeletalMeshObjectGPUSkin::Update(int32 LODIndex,USkinnedMeshComponent* In
 	}
 }
 
+static TAutoConsoleVariable<int32> CVarDeferSkeletalDynamicDataUpdateUntilGDME(
+	TEXT("r.DeferSkeletalDynamicDataUpdateUntilGDME"),
+	0,
+	TEXT("If > 0, then do skeletal mesh dynamic data updates will be deferred until GDME. Experimental option."));
+
 void FSkeletalMeshObjectGPUSkin::UpdateDynamicData_RenderThread(FRHICommandListImmediate& RHICmdList, FDynamicSkelMeshObjectData* InDynamicData)
 {
 	SCOPE_CYCLE_COUNTER(STAT_GPUSkinUpdateRTTime);
 	bool bMorphNeedsUpdate=false;
 	// figure out if the morphing vertex buffer needs to be updated. compare old vs new active morphs
-	bMorphNeedsUpdate = DynamicData ? (DynamicData->LODIndex != ((FDynamicSkelMeshObjectDataGPUSkin*)InDynamicData)->LODIndex ||
+	bMorphNeedsUpdate = 
+		(bMorphNeedsUpdateDeferred && bNeedsUpdateDeferred) || // the need for an update sticks
+		(DynamicData ? (DynamicData->LODIndex != ((FDynamicSkelMeshObjectDataGPUSkin*)InDynamicData)->LODIndex ||
 		!DynamicData->ActiveVertexAnimsEqual(((FDynamicSkelMeshObjectDataGPUSkin*)InDynamicData)->ActiveVertexAnims))
-		: true;
+		: true);
 
-	// we should be done with the old data at this point
+	WaitForRHIThreadFenceForDynamicData();
 	delete DynamicData;
 	// update with new data
 	DynamicData = (FDynamicSkelMeshObjectDataGPUSkin*)InDynamicData;
-	checkSlow(DynamicData);
+	check(DynamicData);
+
+	if (CVarDeferSkeletalDynamicDataUpdateUntilGDME.GetValueOnRenderThread())
+	{
+		bMorphNeedsUpdateDeferred = bMorphNeedsUpdate;
+		bNeedsUpdateDeferred = true;
+	}
+	else
+	{
+		ProcessUpdatedDynamicData(RHICmdList, bMorphNeedsUpdate);
+	}
+}
+
+void FSkeletalMeshObjectGPUSkin::PreGDMECallback()
+{
+	if (bNeedsUpdateDeferred)
+	{
+		ProcessUpdatedDynamicData(FRHICommandListExecutor::GetImmediateCommandList(), bMorphNeedsUpdateDeferred);
+	}
+}
+
+void FSkeletalMeshObjectGPUSkin::WaitForRHIThreadFenceForDynamicData()
+{
+	// we should be done with the old data at this point
+	if (RHIThreadFenceForDynamicData.GetReference())
+	{
+		FRHICommandListExecutor::WaitOnRHIThreadFence(RHIThreadFenceForDynamicData);
+		RHIThreadFenceForDynamicData = nullptr;
+	}
+}
+
+void FSkeletalMeshObjectGPUSkin::ProcessUpdatedDynamicData(FRHICommandListImmediate& RHICmdList, bool bMorphNeedsUpdate)
+{
+	QUICK_SCOPE_CYCLE_COUNTER(STAT_FSkeletalMeshObjectGPUSkin_ProcessUpdatedDynamicData);
+	bNeedsUpdateDeferred = false;
+	bMorphNeedsUpdateDeferred = false;
 
 	FSkeletalMeshObjectLOD& LOD = LODs[DynamicData->LODIndex];
 	const FSkelMeshObjectLODInfo& MeshLODInfo = LODInfo[DynamicData->LODIndex];
@@ -239,6 +292,23 @@ void FSkeletalMeshObjectGPUSkin::UpdateDynamicData_RenderThread(FRHICommandListI
 		DataPresent = VertexFactoryData.VertexFactories.Num() > 0;
 	}
 
+	// clear all old hack pointers to the old dynamic data and bone maps
+	for( int32 ChunkIdx=0; ChunkIdx < Chunks.Num(); ChunkIdx++ )
+	{
+		if (VertexFactoryData.ClothVertexFactories.IsValidIndex(ChunkIdx) && VertexFactoryData.ClothVertexFactories[ChunkIdx])
+		{
+			VertexFactoryData.ClothVertexFactories[ChunkIdx]->GetVertexFactory()->GetShaderData().ClearBoneData();
+		}
+		if (VertexFactoryData.MorphVertexFactories.IsValidIndex(ChunkIdx))
+		{
+			VertexFactoryData.MorphVertexFactories[ChunkIdx].GetShaderData().ClearBoneData();
+		}
+		if (VertexFactoryData.VertexFactories.IsValidIndex(ChunkIdx))
+		{
+			VertexFactoryData.VertexFactories[ChunkIdx].GetShaderData().ClearBoneData();
+		}
+	}
+
 	if(DataPresent)
 	{
 		bool bGPUSkinCacheEnabled = GEnableGPUSkinCache && (FeatureLevel >= ERHIFeatureLevel::SM5);
@@ -257,43 +327,14 @@ void FSkeletalMeshObjectGPUSkin::UpdateDynamicData_RenderThread(FRHICommandListI
 				bClothFactory ? VertexFactoryData.ClothVertexFactories[ChunkIdx]->GetVertexFactory()->GetShaderData() :
 				(DynamicData->NumWeightedActiveVertexAnims > 0 ? VertexFactoryData.MorphVertexFactories[ChunkIdx].GetShaderData() : VertexFactoryData.VertexFactories[ChunkIdx].GetShaderData());
 
-			TArray<FBoneSkinning>& ChunkMatrices = ShaderData.BoneMatrices;
-
-			// update bone matrix shader data for the vertex factory of each chunk
-			ChunkMatrices.Reset(); // remove all elts but leave allocated
-
-			const int32 NumBones = Chunk.BoneMap.Num();
-			ChunkMatrices.Reserve( NumBones ); // we are going to keep adding data to this for each bone
-
-			const int32 NumToAdd = NumBones - ChunkMatrices.Num();
-			ChunkMatrices.AddUninitialized( NumToAdd );
-
-			//FSkinMatrix3x4 is sizeof() == 48
-			// CACHE_LINE_SIZE (128) / 48 = 2.6
-			//  sizeof(FMatrix) == 64
-			// CACHE_LINE_SIZE (128) / 64 = 2
-			const int32 PreFetchStride = 2; // FPlatformMisc::Prefetch stride
-	
-			TArray<FMatrix>& ReferenceToLocalMatrices = DynamicData->ReferenceToLocal;
-			const int32 NumReferenceToLocal = ReferenceToLocalMatrices.Num();
-			for( int32 BoneIdx=0; BoneIdx < NumBones; BoneIdx++ )
-			{
-				FPlatformMisc::Prefetch( ChunkMatrices.GetData() + BoneIdx + PreFetchStride ); 
-				FPlatformMisc::Prefetch( ChunkMatrices.GetData() + BoneIdx + PreFetchStride, CACHE_LINE_SIZE ); 
-				FPlatformMisc::Prefetch( ReferenceToLocalMatrices.GetData() + BoneIdx + PreFetchStride );
-				FPlatformMisc::Prefetch( ReferenceToLocalMatrices.GetData() + BoneIdx + PreFetchStride, CACHE_LINE_SIZE );
-	
-				FBoneSkinning& BoneMat = ChunkMatrices[BoneIdx];
-				const FBoneIndexType RefToLocalIdx = Chunk.BoneMap[BoneIdx];
-				const FMatrix& RefToLocal = ReferenceToLocalMatrices[RefToLocalIdx];
-				RefToLocal.To3x4MatrixTranspose( (float*)BoneMat.M );
-			}
+			bool bUseSkinCache = bGPUSkinCacheEnabled && ChunkIdx < MAX_GPUSKINCACHE_CHUNKS_PER_LOD && !bClothFactory && Chunk.MaxBoneInfluences > 0 && DynamicData->NumWeightedActiveVertexAnims <= 0;
 
 			// Create a uniform buffer from the bone transforms.
-			ShaderData.UpdateBoneData(FeatureLevel);
+			TArray<FMatrix>& ReferenceToLocalMatrices = DynamicData->ReferenceToLocal;
+			bool bNeedFence = ShaderData.UpdateBoneData(RHICmdList, ReferenceToLocalMatrices, Chunk.BoneMap, FeatureLevel, bUseSkinCache);
 
 			// Try to use the GPU skinning cache if possible
-			if (bGPUSkinCacheEnabled && ChunkIdx < MAX_GPUSKINCACHE_CHUNKS_PER_LOD && !bClothFactory && Chunk.MaxBoneInfluences > 0 && DynamicData->NumWeightedActiveVertexAnims <= 0)
+			if (bUseSkinCache)
 			{
 				int32 Key = GGPUSkinCache.StartCacheMesh(RHICmdList, GPUSkinCacheKeys[ChunkIdx], &VertexFactoryData.VertexFactories[ChunkIdx], &VertexFactoryData.PassthroughVertexFactories[ChunkIdx], Chunk, this, Chunk.HasExtraBoneInfluences());
 				if(Key >= 0)
@@ -311,10 +352,14 @@ void FSkeletalMeshObjectGPUSkin::UpdateDynamicData_RenderThread(FRHICommandListI
 				int16 ActorIdx = Chunk.CorrespondClothAssetIndex;
 				if( DynamicData->ClothSimulUpdateData.IsValidIndex(ActorIdx) )
 				{
-					ClothShaderData.UpdateClothSimulData( DynamicData->ClothSimulUpdateData[ActorIdx].ClothSimulPositions, DynamicData->ClothSimulUpdateData[ActorIdx].ClothSimulNormals, FeatureLevel);
+					bNeedFence = ClothShaderData.UpdateClothSimulData(RHICmdList, DynamicData->ClothSimulUpdateData[ActorIdx].ClothSimulPositions, DynamicData->ClothSimulUpdateData[ActorIdx].ClothSimulNormals, FeatureLevel) || bNeedFence;
 				}
 			}
 #endif // WITH_APEX_CLOTHING
+			if (bNeedFence)
+			{
+				RHIThreadFenceForDynamicData = RHICmdList.RHIThreadFence(true);
+			}
 		}
 	}
 
@@ -623,11 +668,10 @@ TDynamicUpdateVertexFactoryData<VertexFactoryType>,VertexUpdateData,VertexUpdate
 template <class VertexFactoryTypeBase, class VertexFactoryType>
 static void CreateVertexFactory(TIndirectArray<VertexFactoryTypeBase>& VertexFactories,
 						 const FSkeletalMeshObjectGPUSkin::FVertexFactoryBuffers& InVertexBuffers,
-						 TArray<FBoneSkinning>& InBoneMatrices,
 						 ERHIFeatureLevel::Type FeatureLevel
 						 )
 {
-	auto* VertexFactory = new VertexFactoryType(InBoneMatrices, FeatureLevel);
+	auto* VertexFactory = new VertexFactoryType(FeatureLevel);
 	VertexFactories.Add(VertexFactory);
 
 	// Setup the update data for enqueue
@@ -660,12 +704,11 @@ TDynamicUpdateVertexFactoryData<VertexFactoryType>,VertexUpdateData,VertexUpdate
 template <class VertexFactoryTypeBase, class VertexFactoryType>
 static void CreateVertexFactoryMorph(TIndirectArray<VertexFactoryTypeBase>& VertexFactories,
 						 const FSkeletalMeshObjectGPUSkin::FVertexFactoryBuffers& InVertexBuffers,
-						 TArray<FBoneSkinning>& InBoneMatrices,
 						 ERHIFeatureLevel::Type FeatureLevel
 						 )
 
 {
-	auto* VertexFactory = new VertexFactoryType(InBoneMatrices, FeatureLevel);
+	auto* VertexFactory = new VertexFactoryType(FeatureLevel);
 	VertexFactories.Add(VertexFactory);
 						
 	// Setup the update data for enqueue
@@ -699,12 +742,11 @@ TDynamicUpdateVertexFactoryData<VertexFactoryType>,VertexUpdateData,VertexUpdate
 template <class VertexFactoryTypeBase, class VertexFactoryType>
 static void CreateVertexFactoryCloth(TArray<VertexFactoryTypeBase*>& VertexFactories,
 						 const FSkeletalMeshObjectGPUSkin::FVertexFactoryBuffers& InVertexBuffers,
-						 TArray<FBoneSkinning>& InBoneMatrices,
 						 ERHIFeatureLevel::Type FeatureLevel
 						 )
 
 {
-	VertexFactoryType* VertexFactory = new VertexFactoryType(InBoneMatrices, FeatureLevel);
+	VertexFactoryType* VertexFactory = new VertexFactoryType(FeatureLevel);
 	VertexFactories.Add(VertexFactory);
 						
 	// Setup the update data for enqueue
@@ -735,23 +777,6 @@ void FSkeletalMeshObjectGPUSkin::FSkeletalMeshObjectLOD::GetVertexBuffers(
 	OutVertexBuffers.APEXClothVertexBuffer = &LODModel.APEXClothVertexBuffer;
 }
 
-/**
- * Init one array of matrices for each chunk (shared across vertex factory types)
- *
- * @param Chunks - relevant chunk information (either original or from swapped influence)
- */
-void FSkeletalMeshObjectGPUSkin::FVertexFactoryData::InitPerChunkBoneMatrices(const TArray<FSkelMeshChunk>& Chunks)
-{
-	checkSlow(!IsInActualRenderingThread());
-
-	// one array of matrices for each chunk (shared across vertex factory types)
-	if (PerChunkBoneMatricesArray.Num() != Chunks.Num())
-	{
-		PerChunkBoneMatricesArray.Empty(Chunks.Num());
-		PerChunkBoneMatricesArray.AddZeroed(Chunks.Num());
-	}
-}
-
 /** 
  * Init vertex factory resources for this LOD 
  *
@@ -763,9 +788,6 @@ void FSkeletalMeshObjectGPUSkin::FVertexFactoryData::InitVertexFactories(
 	const TArray<FSkelMeshChunk>& Chunks, 
 	ERHIFeatureLevel::Type FeatureLevel)
 {
-	// one array of matrices for each chunk (shared across vertex factory types)
-	InitPerChunkBoneMatrices(Chunks);
-
 	// first clear existing factories (resources assumed to have been released already)
 	// then [re]create the factories
 
@@ -775,13 +797,13 @@ void FSkeletalMeshObjectGPUSkin::FVertexFactoryData::InitVertexFactories(
 		{
 			if (VertexBuffers.VertexBufferGPUSkin->HasExtraBoneInfluences())
 			{
-				CreateVertexFactory< FGPUBaseSkinVertexFactory, TGPUSkinVertexFactory<true> >(VertexFactories, VertexBuffers, PerChunkBoneMatricesArray[FactoryIdx], FeatureLevel);
-				CreateVertexFactory< FGPUBaseSkinVertexFactory, FGPUSkinPassthroughVertexFactory >(PassthroughVertexFactories, VertexBuffers, PerChunkBoneMatricesArray[FactoryIdx], FeatureLevel);
+				CreateVertexFactory< FGPUBaseSkinVertexFactory, TGPUSkinVertexFactory<true> >(VertexFactories, VertexBuffers, FeatureLevel);
+				CreateVertexFactory< FGPUBaseSkinVertexFactory, FGPUSkinPassthroughVertexFactory >(PassthroughVertexFactories, VertexBuffers, FeatureLevel);
 			}
 			else
 			{
-				CreateVertexFactory< FGPUBaseSkinVertexFactory, TGPUSkinVertexFactory<false> >(VertexFactories, VertexBuffers, PerChunkBoneMatricesArray[FactoryIdx], FeatureLevel);
-				CreateVertexFactory< FGPUBaseSkinVertexFactory, FGPUSkinPassthroughVertexFactory >(PassthroughVertexFactories, VertexBuffers, PerChunkBoneMatricesArray[FactoryIdx], FeatureLevel);
+				CreateVertexFactory< FGPUBaseSkinVertexFactory, TGPUSkinVertexFactory<false> >(VertexFactories, VertexBuffers, FeatureLevel);
+				CreateVertexFactory< FGPUBaseSkinVertexFactory, FGPUSkinPassthroughVertexFactory >(PassthroughVertexFactories, VertexBuffers, FeatureLevel);
 			}
 		}
 	}
@@ -810,19 +832,17 @@ void FSkeletalMeshObjectGPUSkin::FVertexFactoryData::InitMorphVertexFactories(
 	bool bInUsePerBoneMotionBlur,
 	ERHIFeatureLevel::Type InFeatureLevel)
 {
-	// one array of matrices for each chunk (shared across vertex factory types)
-	InitPerChunkBoneMatrices(Chunks);
 	// clear existing factories (resources assumed to have been released already)
 	MorphVertexFactories.Empty(Chunks.Num());
 	for( int32 FactoryIdx=0; FactoryIdx < Chunks.Num(); FactoryIdx++ )
 	{
 		if (VertexBuffers.VertexBufferGPUSkin->HasExtraBoneInfluences())
 		{
-			CreateVertexFactoryMorph<FGPUBaseSkinVertexFactory, TGPUSkinMorphVertexFactory<true> >(MorphVertexFactories,VertexBuffers,PerChunkBoneMatricesArray[FactoryIdx], InFeatureLevel);
+			CreateVertexFactoryMorph<FGPUBaseSkinVertexFactory, TGPUSkinMorphVertexFactory<true> >(MorphVertexFactories,VertexBuffers,InFeatureLevel);
 		}
 		else
 		{
-			CreateVertexFactoryMorph<FGPUBaseSkinVertexFactory, TGPUSkinMorphVertexFactory<false> >(MorphVertexFactories, VertexBuffers, PerChunkBoneMatricesArray[FactoryIdx], InFeatureLevel);
+			CreateVertexFactoryMorph<FGPUBaseSkinVertexFactory, TGPUSkinMorphVertexFactory<false> >(MorphVertexFactories, VertexBuffers,InFeatureLevel);
 		}
 	}
 }
@@ -845,8 +865,6 @@ void FSkeletalMeshObjectGPUSkin::FVertexFactoryData::InitAPEXClothVertexFactorie
 	const TArray<FSkelMeshChunk>& Chunks,
 	ERHIFeatureLevel::Type InFeatureLevel)
 {
-	// one array of matrices for each chunk (shared across vertex factory types)
-	InitPerChunkBoneMatrices(Chunks);
 
 	// clear existing factories (resources assumed to have been released already)
 	ClothVertexFactories.Empty(Chunks.Num());
@@ -856,11 +874,11 @@ void FSkeletalMeshObjectGPUSkin::FVertexFactoryData::InitAPEXClothVertexFactorie
 		{
 			if (VertexBuffers.VertexBufferGPUSkin->HasExtraBoneInfluences())
 			{
-				CreateVertexFactoryCloth<FGPUBaseSkinAPEXClothVertexFactory, TGPUSkinAPEXClothVertexFactory<true> >(ClothVertexFactories, VertexBuffers, PerChunkBoneMatricesArray[FactoryIdx], InFeatureLevel);
+				CreateVertexFactoryCloth<FGPUBaseSkinAPEXClothVertexFactory, TGPUSkinAPEXClothVertexFactory<true> >(ClothVertexFactories, VertexBuffers, InFeatureLevel);
 			}
 			else
 			{
-				CreateVertexFactoryCloth<FGPUBaseSkinAPEXClothVertexFactory, TGPUSkinAPEXClothVertexFactory<false> >(ClothVertexFactories, VertexBuffers, PerChunkBoneMatricesArray[FactoryIdx], InFeatureLevel);
+				CreateVertexFactoryCloth<FGPUBaseSkinAPEXClothVertexFactory, TGPUSkinAPEXClothVertexFactory<false> >(ClothVertexFactories, VertexBuffers, InFeatureLevel);
 			}
 		}
 		else
@@ -1168,7 +1186,7 @@ void FPreviousPerBoneMotionBlur::StartAppend(FRHICommandListImmediate& RHICmdLis
 	}
 }
 
-uint32 FPreviousPerBoneMotionBlur::AppendData(FBoneSkinning *DataStart, uint32 BoneCount)
+uint32 FPreviousPerBoneMotionBlur::AppendData(FBoneSkinning*& OutChunkMatrices, uint32 BoneCount)
 {
 	check(LockedData);
 	checkSlow(DataStart);
@@ -1180,12 +1198,13 @@ uint32 FPreviousPerBoneMotionBlur::AppendData(FBoneSkinning *DataStart, uint32 B
 	
 	if(OldLockedTexelPosition + TexelCount <= LockedTexelCount)
 	{
-		FMemory::Memcpy(&LockedData[OldLockedTexelPosition * 4], DataStart, BoneCount * sizeof(FBoneSkinning));
+		OutChunkMatrices = (FBoneSkinning*)&LockedData[OldLockedTexelPosition * 4];
 
 		return OldLockedTexelPosition;
 	}
 	else
 	{
+		OutChunkMatrices = nullptr;
 		// Not enough space in the texture, we should increase the texture size. The new bigger size 
 		// can be found in LockedTexelPosition. This is currently not done - so we might not see motion blur
 		// skinning on all objects.
diff --git a/Engine/Source/Runtime/Engine/Private/SkeletalRenderGPUSkin.h b/Engine/Source/Runtime/Engine/Private/SkeletalRenderGPUSkin.h
--- a/Engine/Source/Runtime/Engine/Private/SkeletalRenderGPUSkin.h
+++ b/Engine/Source/Runtime/Engine/Private/SkeletalRenderGPUSkin.h
@@ -194,11 +194,11 @@ public:
 
 	/**
 	 * use between LockData() and UnlockData()
-	 * @param	DataStart, must not be 0
+	 * @param	OutChunkMatrices, filled in with a pointer to bones that should be filled in
 	 * @param	BoneCount number of FBoneSkinning elements, must not be 0
 	 * @return	StartIndex where the data can be referenced in the texture, 0xffffffff if this failed (not enough space in the buffer, will be fixed next frame)
 	 */
-	uint32 AppendData(FBoneSkinning *DataStart, uint32 BoneCount);
+	uint32 AppendData(FBoneSkinning*& OutChunkMatrices, uint32 BoneCount);
 
 	/** only call if StartAppend(), if the append wasn't started it silently ignores the call */
 	ENGINE_API void EndAppendFence(FRHICommandListImmediate& RHICmdList);
@@ -280,6 +280,7 @@ public:
 	virtual void ReleaseResources() override;
 	virtual void Update(int32 LODIndex,USkinnedMeshComponent* InMeshComponent,const TArray<FActiveVertexAnim>& ActiveVertexAnims) override;
 	virtual void UpdateDynamicData_RenderThread(FRHICommandListImmediate& RHICmdList, FDynamicSkelMeshObjectData* InDynamicData) override;
+	virtual void PreGDMECallback() override;
 	virtual const FVertexFactory* GetVertexFactory(int32 LODIndex,int32 ChunkIdx) const override;
 	virtual void CacheVertices(int32 LODIndex, bool bForce) const override {}
 	virtual bool IsCPUSkinned() const override { return false; }
@@ -365,10 +366,6 @@ private:
 		/** Vertex factory defining both the base mesh as well as the APEX cloth vertex data */
 		TArray<FGPUBaseSkinAPEXClothVertexFactory*> ClothVertexFactories;
 
-		/** shared ref pose to local space matrices */
-		TArray< TArray<FBoneSkinning>, TInlineAllocator<1> > PerChunkBoneMatricesArray;
-
-
 		/** 
 		 * Init default vertex factory resources for this LOD 
 		 *
@@ -403,12 +400,6 @@ private:
 		 */
 		void ReleaseAPEXClothVertexFactories();
 
-		/**
-		 * Init one array of matrices for each chunk (shared across vertex factory types)
-		 *
-		 * @param Chunks - relevant chunk information (either original or from swapped influence)
-		 */
-		void InitPerChunkBoneMatrices(const TArray<FSkelMeshChunk>& Chunks);
 		/**
 		 * Clear factory arrays
 		 */
@@ -431,8 +422,6 @@ private:
 
 			Size += ClothVertexFactories.GetAllocatedSize();
 
-			Size += PerChunkBoneMatricesArray.GetAllocatedSize();
-
 			return Size;
 		}	
 	};
@@ -525,12 +514,25 @@ private:
 	*/
 	void ReleaseMorphResources();
 
+	void ProcessUpdatedDynamicData(FRHICommandListImmediate& RHICmdList, bool bMorphNeedsUpdate);
+
+	void WaitForRHIThreadFenceForDynamicData();
+
 	/** Render data for each LOD */
 	TArray<struct FSkeletalMeshObjectLOD> LODs;
 
 	/** Data that is updated dynamically and is needed for rendering */
 	FDynamicSkelMeshObjectDataGPUSkin* DynamicData;
 
+	/** Fence for dynamic Data */
+	FGraphEventRef RHIThreadFenceForDynamicData;
+
+	/** True if we are doing a deferred update later in GDME. */
+	bool bNeedsUpdateDeferred;
+
+	/** If true and we are doing a deferred update, then also update the morphs */
+	bool bMorphNeedsUpdateDeferred;
+
 	/** true if the morph resources have been initialized */
 	bool bMorphResourcesInitialized;
 
diff --git a/Engine/Source/Runtime/Engine/Private/TickTaskManager.cpp b/Engine/Source/Runtime/Engine/Private/TickTaskManager.cpp
--- a/Engine/Source/Runtime/Engine/Private/TickTaskManager.cpp
+++ b/Engine/Source/Runtime/Engine/Private/TickTaskManager.cpp
@@ -14,10 +14,15 @@ DECLARE_CYCLE_STAT(TEXT("Queue Ticks Wait"),STAT_QueueTicksWait,STATGROUP_Game);
 DECLARE_CYCLE_STAT(TEXT("Queue Tick Task"),STAT_QueueTickTask,STATGROUP_Game);
 DECLARE_CYCLE_STAT(TEXT("Post Queue Tick Task"),STAT_PostTickTask,STATGROUP_Game);
 DECLARE_DWORD_COUNTER_STAT(TEXT("Ticks Queued"),STAT_TicksQueued,STATGROUP_Game);
+DECLARE_CYCLE_STAT(TEXT("TG_NewlySpawned"), STAT_TG_NewlySpawned, STATGROUP_TickGroups);
 
 static TAutoConsoleVariable<int32> CVarLogTicks(
 	TEXT("tick.LogTicks"),0,
-	TEXT("Used to logging of ticks."));
+	TEXT("Spew ticks for debugging."));
+
+static TAutoConsoleVariable<int32> CVarLogTicksShowPrerequistes(
+	TEXT("tick.ShowPrerequistes"),1,
+	TEXT("When logging ticks, show the prerequistes; debugging."));
 
 static TAutoConsoleVariable<int32> CVarAllowAsyncComponentTicks(
 	TEXT("tick.AllowAsyncComponentTicks"),
@@ -59,6 +64,8 @@ struct FTickContext
 	}
 };
 
+#include"scopedtimers.h"
+#include"ParallelFor.h"
 /**
  * Class that handles the actual tick tasks and starting and completing tick groups
  */
@@ -73,15 +80,18 @@ class FTickTaskSequencer
 		FTickContext			Context;
 		/** If true, log each tick **/
 		bool					bLogTick; 
+		/** If true, log prereqs **/
+		bool					bLogTicksShowPrerequistes; 
 	public:
 		/** Constructor
 		 * @param InTarget - Function to tick
 		 * @param InContext - context to tick in, here thread is desired execution thread
 		**/
-		FTickFunctionTask(FTickFunction* InTarget, const FTickContext* InContext, bool InbLogTick)
+		FORCEINLINE FTickFunctionTask(FTickFunction* InTarget, const FTickContext* InContext, bool InbLogTick, bool bInLogTicksShowPrerequistes)
 			: Target(InTarget)
 			, Context(*InContext)
 			, bLogTick(InbLogTick)
+			, bLogTicksShowPrerequistes(bInLogTicksShowPrerequistes)
 		{
 		}
 		FORCEINLINE TStatId GetStatId() const
@@ -89,11 +99,15 @@ class FTickTaskSequencer
 			RETURN_QUICK_DECLARE_CYCLE_STAT(FTickFunctionTask, STATGROUP_TaskGraphTasks);
 		}
 		/** return the thread for this task **/
-		ENamedThreads::Type GetDesiredThread()
+		FORCEINLINE ENamedThreads::Type GetDesiredThread()
 		{
+			if (Target->bHighPriority)
+			{
+				return ENamedThreads::HiPri(Context.Thread);
+			}
 			return Context.Thread;
 		}
-		static ESubsequentsMode::Type GetSubsequentsMode() 
+		FORCEINLINE static ESubsequentsMode::Type GetSubsequentsMode() 
 		{ 
 			return ESubsequentsMode::TrackSubsequents; 
 		}
@@ -107,7 +121,11 @@ class FTickTaskSequencer
 		{
 			if (bLogTick)
 			{
-				UE_LOG(LogTick, Log, TEXT("tick %6d %2d %s"),GFrameCounter, (int32)CurrentThread, *Target->DiagnosticMessage());
+				UE_LOG(LogTick, Log, TEXT("tick %s %6d %2d %s"), Target->bHighPriority ? TEXT("*") : TEXT(" "),GFrameCounter, (int32)CurrentThread, *Target->DiagnosticMessage());
+				if (bLogTicksShowPrerequistes)
+				{
+					Target->ShowPrerequistes();
+				}
 			}
 			Target->ExecuteTick(Context.DeltaSeconds, Context.TickType, CurrentThread, MyCompletionGraphEvent);
 			Target->CompletionHandle = NULL; // Allow the old completion handle to be recycled
@@ -118,7 +136,10 @@ class FTickTaskSequencer
 	/** Completion handles for each phase of ticks */
 	FGraphEventArray	TickCompletionEvents[TG_MAX];
 
-	/** Held tasks for each tick group. */
+	/** HiPri Held tasks for each tick group. */
+	TArray<TGraphTask<FTickFunctionTask>*> HiPriTickTasks[TG_MAX];
+
+	/** LowPri Held tasks for each tick group. */
 	TArray<TGraphTask<FTickFunctionTask>*> TickTasks[TG_MAX];
 
 	/** These are waited for at the end of the frame; they are not on the critical path, but they have to be done before we leave the frame. */
@@ -132,6 +153,8 @@ class FTickTaskSequencer
 
 	/** If true, log each tick **/
 	bool				bLogTicks; 
+	/** If true, log each tick **/
+	bool				bLogTicksShowPrerequistes; 
 
 public:
 
@@ -177,8 +200,15 @@ public:
 		}
 		QUICK_SCOPE_CYCLE_COUNTER(STAT_StartTickTask_ConstructAndDispatchWhenReady);
 
-		TGraphTask<FTickFunctionTask>* Task = TGraphTask<FTickFunctionTask>::CreateTask(Prerequisites, TickContext.Thread).ConstructAndHold(TickFunction, &UseContext, bLogTicks);
-		TickTasks[TickFunction->ActualTickGroup].Add(Task);
+		TGraphTask<FTickFunctionTask>* Task = TGraphTask<FTickFunctionTask>::CreateTask(Prerequisites, TickContext.Thread).ConstructAndHold(TickFunction, &UseContext, bLogTicks, bLogTicksShowPrerequistes);
+		if (TickFunction->bHighPriority)
+		{
+			HiPriTickTasks[TickFunction->ActualTickGroup].Add(Task);
+		}
+		else
+		{
+			TickTasks[TickFunction->ActualTickGroup].Add(Task);
+		}
 		TickFunction->CompletionHandle = Task->GetCompletionEvent();
 	}
 
@@ -228,8 +258,11 @@ public:
 			FTaskGraphInterface::Get().WaitUntilTaskCompletes(
 				FDelegateGraphTask::CreateAndDispatchWhenReady(
 					FDelegateGraphTask::FDelegate::CreateRaw(this, &FTickTaskSequencer::DispatchTickGroup, WorldTickGroup),
-					GET_STATID(STAT_FDelegateGraphTask_DispatchTickGroup)),
-				ENamedThreads::GameThread);
+					GET_STATID(STAT_FDelegateGraphTask_DispatchTickGroup), 
+					nullptr,
+					ENamedThreads::GameThread, 
+					ENamedThreads::HiPri(ENamedThreads::AnyThread)),
+				ENamedThreads::HiPri(ENamedThreads::GameThread));
 		}
 
 		if (bBlockTillComplete || SingleThreadedMode())
@@ -269,7 +302,129 @@ public:
 	 */
 	void StartFrame()
 	{
+#if 0
+		class FNoop
+		{
+		public:
+			FORCEINLINE TStatId GetStatId() const
+			{
+				RETURN_QUICK_DECLARE_CYCLE_STAT(FNoop, STATGROUP_TaskGraphTasks);
+			}
+			FORCEINLINE ENamedThreads::Type GetDesiredThread()
+			{
+				return ENamedThreads::AnyThread;
+			}
+			FORCEINLINE static ESubsequentsMode::Type GetSubsequentsMode() 
+			{ 
+				return ESubsequentsMode::TrackSubsequents; 
+			}
+			void DoTask(ENamedThreads::Type CurrentThread, const FGraphEventRef& MyCompletionGraphEvent)
+			{
+			}
+		};
+		class FNoopFireAndForget
+		{
+			FThreadSafeCounter& Work;
+		public:
+
+			FNoopFireAndForget(FThreadSafeCounter& InWork)
+				: Work(InWork)
+			{
+			}
+			FORCEINLINE TStatId GetStatId() const
+			{
+				RETURN_QUICK_DECLARE_CYCLE_STAT(FNoop, STATGROUP_TaskGraphTasks);
+			}
+			FORCEINLINE ENamedThreads::Type GetDesiredThread()
+			{
+				return ENamedThreads::AnyThread;
+			}
+			FORCEINLINE static ESubsequentsMode::Type GetSubsequentsMode() 
+			{ 
+				return ESubsequentsMode::FireAndForget; 
+			}
+			void DoTask(ENamedThreads::Type CurrentThread, const FGraphEventRef& MyCompletionGraphEvent)
+			{
+				Work.Increment();
+			}
+		};
+
+		{
+			double QueueNoop = 0.0;
+			double WaitNoop = 0.0;
+			double QueueFireAndForget = 0.0;
+			double WaitFireAndForget = 0.0;
+			FGraphEventArray Tasks;
+			Tasks.Reserve(1000);
+			{
+				FScopedDurationTimer Scope(QueueNoop);
+				for (int32 Index = 0; Index < 1000; Index++)
+				{
+					Tasks.Add(TGraphTask<FNoop>::CreateTask(nullptr,  ENamedThreads::GameThread).ConstructAndDispatchWhenReady());
+				}
+			}
+			{
+				FScopedDurationTimer Scope(WaitNoop);
+				FTaskGraphInterface::Get().WaitUntilTasksComplete(Tasks, ENamedThreads::GameThread_Local);
+			}
+			FThreadSafeCounter Work;		
+			{
+				FScopedDurationTimer Scope(QueueFireAndForget);
+				for (int32 Index = 0; Index < 1000; Index++)
+				{
+					TGraphTask<FNoopFireAndForget>::CreateTask(nullptr,  ENamedThreads::GameThread).ConstructAndDispatchWhenReady(Work);
+				}
+			}
+			{
+				FScopedDurationTimer Scope(WaitFireAndForget);
+				while (Work.GetValue() != 1000)
+					;
+			}
+
+			FPlatformMisc::LowLevelOutputDebugStringf(TEXT("GT Dispatch    Normal: %7.5fms / queue    %7.5fms / execute    FireAndForget: %7.5fms / queue    %7.5fms / execute\n"), float(QueueNoop), float(WaitNoop), float(QueueFireAndForget), float(WaitFireAndForget));
+		}
+
+		{
+			double QueueNoop = 0.0;
+			double WaitNoop = 0.0;
+			double QueueFireAndForget = 0.0;
+			double WaitFireAndForget = 0.0;
+			FGraphEventArray Tasks;
+			Tasks.AddZeroed(1000);
+			{
+				FScopedDurationTimer Scope(QueueNoop);
+				ParallelFor(1000,
+					[&Tasks](int32 Index)
+					{
+						Tasks[Index] = TGraphTask<FNoop>::CreateTask(nullptr).ConstructAndDispatchWhenReady();
+					}
+				);
+			}
+			{
+				FScopedDurationTimer Scope(WaitNoop);
+				FTaskGraphInterface::Get().WaitUntilTasksComplete(Tasks, ENamedThreads::GameThread_Local);
+			}
+			FThreadSafeCounter Work;		
+			{
+				FScopedDurationTimer Scope(QueueFireAndForget);
+				ParallelFor(1000,
+					[&Work](int32 Index)
+					{
+						TGraphTask<FNoopFireAndForget>::CreateTask(nullptr).ConstructAndDispatchWhenReady(Work);
+					}
+				);
+			}
+			{
+				FScopedDurationTimer Scope(WaitFireAndForget);
+				while (Work.GetValue() != 1000)
+					;
+			}
+
+			FPlatformMisc::LowLevelOutputDebugStringf(TEXT("Parallel For   Normal: %7.5fms / queue    %7.5fms / execute    FireAndForget: %7.5fms / queue    %7.5fms / execute\n"), float(QueueNoop), float(WaitNoop), float(QueueFireAndForget), float(WaitFireAndForget));
+		}
+#endif
 		bLogTicks = !!CVarLogTicks.GetValueOnGameThread();
+		bLogTicksShowPrerequistes = !!CVarLogTicksShowPrerequistes.GetValueOnGameThread();
 
 		if (bLogTicks)
 		{
@@ -288,8 +443,9 @@ public:
 		{
 			check(!TickCompletionEvents[Index].Num());  // we should not be adding to these outside of a ticking proper and they were already cleared after they were ticked
 			TickCompletionEvents[Index].Reset();
-			check(!TickTasks[Index].Num());  // we should not be adding to these outside of a ticking proper and they were already cleared after they were ticked
+			check(!TickTasks[Index].Num() && !HiPriTickTasks[Index].Num());  // we should not be adding to these outside of a ticking proper and they were already cleared after they were ticked
 			TickTasks[Index].Reset();
+			HiPriTickTasks[Index].Reset();
 		}
 		WaitForTickGroup = (ETickingGroup)0;
 	}
@@ -307,7 +463,7 @@ public:
 		for (int32 Index = 0; Index < TG_MAX; Index++)
 		{
 			check(!TickCompletionEvents[Index].Num());  // we should not be adding to these outside of a ticking proper and they were already cleared after they were ticked
-			check(!TickTasks[Index].Num()); 
+			check(!TickTasks[Index].Num() && !HiPriTickTasks[Index].Num());  // we should not be adding to these outside of a ticking proper and they were already cleared after they were ticked
 		}
 	}
 private:
@@ -315,6 +471,7 @@ private:
 	FTickTaskSequencer()
 		: bAllowConcurrentTicks(false)
 		, bLogTicks(false)
+		, bLogTicksShowPrerequistes(false)
 	{
 	}
 
@@ -329,12 +486,22 @@ private:
 
 	void DispatchTickGroupInner(ENamedThreads::Type CurrentThread, ETickingGroup WorldTickGroup)
 	{
-		TArray<TGraphTask<FTickFunctionTask>*>& TickArray = TickTasks[WorldTickGroup];
-		for (int32 Index = 0; Index < TickArray.Num(); Index++)
 		{
-			TickArray[Index]->Unlock(CurrentThread);
+			TArray<TGraphTask<FTickFunctionTask>*>& TickArray = HiPriTickTasks[WorldTickGroup];
+			for (int32 Index = 0; Index < TickArray.Num(); Index++)
+			{
+				TickArray[Index]->Unlock(CurrentThread);
+			}
+			TickArray.Reset();
+		}
+		{
+			TArray<TGraphTask<FTickFunctionTask>*>& TickArray = TickTasks[WorldTickGroup];
+			for (int32 Index = 0; Index < TickArray.Num(); Index++)
+			{
+				TickArray[Index]->Unlock(CurrentThread);
+			}
+			TickArray.Reset();
 		}
-		TickArray.Reset();
 	}
 
 	void DispatchTickGroup(ENamedThreads::Type CurrentThread, const FGraphEventRef& MyCompletionGraphEvent, ETickingGroup WorldTickGroup)
@@ -655,6 +822,7 @@ public:
 				}
 				if (Num && Context.TickGroup == TG_NewlySpawned)
 				{
+					SCOPE_CYCLE_COUNTER(STAT_TG_NewlySpawned);
 					TickTaskSequencer.ReleaseTickGroup(TG_NewlySpawned, true);
 				}
 				else
@@ -784,6 +952,7 @@ FTickFunction::FTickFunction()
 	, bTickEvenWhenPaused(false)
 	, bCanEverTick(false)
 	, bAllowTickOnDedicatedServer(true)
+	, bHighPriority(false)
 	, bRunOnAnyThread(false)
 	, bRegistered(false)
 	, bTickEnabled(true)
@@ -847,11 +1016,6 @@ void FTickFunction::SetTickFunctionEnable(bool bInEnabled)
 	}
 }
 
-/** 
- * Adds a tick function to the list of prerequisites...in other words, adds the requirement that TargetTickFunction is called before this tick function is 
- * @param TargetObject - UObject containing this tick function. Only used to verify that the other pointer is still usable
- * @param TargetTickFunction - Actual tick function to use as a prerequisite
-**/
 void FTickFunction::AddPrerequisite(UObject* TargetObject, struct FTickFunction& TargetTickFunction)
 {
 	const bool bThisCanTick = (bCanEverTick || IsTickFunctionRegistered());
@@ -863,20 +1027,38 @@ void FTickFunction::AddPrerequisite(UObject* TargetObject, struct FTickFunction&
 	}
 }
 
-/** 
- * Removes a prerequisite that was previously added.
- * @param TargetObject - UObject containing this tick function. Only used to verify that the other pointer is still usable
- * @param TargetTickFunction - Actual tick function to use as a prerequisite
-**/
 void FTickFunction::RemovePrerequisite(UObject* TargetObject, struct FTickFunction& TargetTickFunction)
 {
 	Prerequisites.RemoveSwap(FTickPrerequisite(TargetObject, TargetTickFunction));
 }
 
-/**
-	* Queues a tick function for execution from the game thread
-	* @param TickContext - context to tick in
-*/
+void FTickFunction::SetPriorityIncludingPrerequisites(bool bInHighPriority)
+{
+	if (bHighPriority != bInHighPriority)
+	{
+		bHighPriority = bInHighPriority;
+		for (auto& Prereq : Prerequisites)
+		{
+			if (Prereq.PrerequisiteObject.Get() && Prereq.PrerequisiteTickFunction && Prereq.PrerequisiteTickFunction->bHighPriority != bInHighPriority)
+			{
+				Prereq.PrerequisiteTickFunction->SetPriorityIncludingPrerequisites(bInHighPriority);
+			}
+		}
+	}
+}
+
+void FTickFunction::ShowPrerequistes(int32 Indent)
+{
+	for (auto& Prereq : Prerequisites)
+	{
+		if (Prereq.PrerequisiteTickFunction)
+		{
+			UE_LOG(LogTick, Log, TEXT("%s prereq %s"), FCString::Spc(Indent * 2), *Prereq.PrerequisiteTickFunction->DiagnosticMessage());
+			Prereq.PrerequisiteTickFunction->ShowPrerequistes(Indent + 1);
+		}
+	}
+}
+
 void FTickFunction::QueueTickFunction(FTickTaskSequencer& TTS, const struct FTickContext& TickContext)
 {
 	checkSlow(TickContext.Thread == ENamedThreads::GameThread); // we assume same thread here
diff --git a/Engine/Source/Runtime/Engine/Public/GPUSkinVertexFactory.h b/Engine/Source/Runtime/Engine/Public/GPUSkinVertexFactory.h
--- a/Engine/Source/Runtime/Engine/Public/GPUSkinVertexFactory.h
+++ b/Engine/Source/Runtime/Engine/Public/GPUSkinVertexFactory.h
@@ -194,13 +194,33 @@ class FBoneBufferPool : public TRenderResourcePool<FBoneBufferTypeRef, FBoneBuff
 public:
 	/** Destructor */
 	virtual ~FBoneBufferPool();
+
+public: // From FTickableObjectRenderThread
+	virtual TStatId GetStatId() const override;
+};
+
+/** The policy for pooling bone vertex buffers */
+class FClothBufferPoolPolicy : public FBoneBufferPoolPolicy
+{
+public:
+	/** Creates the resource 
+	 * @param Args The buffer size in bytes.
+	 * @returns A suitably sized buffer or NULL on failure.
+	 */
+	FBoneBufferTypeRef CreateResource(FSharedPoolPolicyData::CreationArguments Args);
+};
+
+/** A pool for vertex buffers with consistent usage, bucketed for efficiency. */
+class FClothBufferPool : public TRenderResourcePool<FBoneBufferTypeRef, FClothBufferPoolPolicy, FSharedPoolPolicyData::CreationArguments>
+{
+public:
+	/** Destructor */
+	virtual ~FClothBufferPool();
 	
 public: // From FTickableObjectRenderThread
 	virtual TStatId GetStatId() const override;
 };
 
-/** The type for the buffer pool */
-typedef FBoneBufferPool FBoneBufferPool;
 
 /** for motion blur skinning */
 class FBoneDataVertexBuffer : public FRenderResource
@@ -266,15 +286,7 @@ class FGPUBaseSkinVertexFactory : public FVertexFactory
 public:
 	struct ShaderDataType
 	{
-		/**
-		 * Constructor, presizing bones array.
-		 *
-		 * @param	InBoneMatrices	Reference to shared bone matrices array.
-		 */
-		ShaderDataType(
-			TArray<FBoneSkinning>& InBoneMatrices
-			)
-		:	BoneMatrices( InBoneMatrices )
+		ShaderDataType()
 		{
 			OldBoneDataStartIndex[0] = 0xffffffff;
 			OldBoneDataStartIndex[1] = 0xffffffff;
@@ -286,10 +298,10 @@ public:
 				MaxGPUSkinBones = MaxBonesVar->GetValueOnGameThread();
 				check(MaxGPUSkinBones <= 256);
 			}
-		}
 
-		/** Reference to shared ref pose to local space matrices */
-		TArray<FBoneSkinning>& BoneMatrices;
+			HackReferenceToLocalMatrices = nullptr;
+			HackBoneMap = nullptr;
+		}
 
 		/** Mesh origin and Mesh Extension for Mesh compressions **/
 		/** This value will be (0, 0, 0), (1, 1, 1) relatively for non compressed meshes **/
@@ -356,7 +368,7 @@ public:
 			return true;
 		}
 
-		void UpdateBoneData(ERHIFeatureLevel::Type FeatureLevel);
+		bool UpdateBoneData(FRHICommandListImmediate& RHICmdList, const TArray<FMatrix>& ReferenceToLocalMatrices, const TArray<FBoneIndexType>& BoneMap, ERHIFeatureLevel::Type FeatureLevel, bool bUseSkinCache);
 
 		void ReleaseBoneData()
 		{
@@ -390,6 +402,16 @@ public:
 #endif // !(UE_BUILD_SHIPPING || UE_BUILD_TEST)
 
 		mutable FCriticalSection OldBoneDataLock;
+
+		// awful code to recall where the bone info comes from for motion blur
+		const TArray<FMatrix>* HackReferenceToLocalMatrices;
+		const TArray<FBoneIndexType>* HackBoneMap;
+		void ClearBoneData()
+		{
+			HackReferenceToLocalMatrices = nullptr;
+			HackBoneMap = nullptr;
+		}
+
 private:
 
 		// the following members can be stored in less bytes (after some adjustments)
@@ -411,8 +433,8 @@ private:
 	 *
 	 * @param	InBoneMatrices	Reference to shared bone matrices array.
 	 */
-	FGPUBaseSkinVertexFactory(TArray<FBoneSkinning>& InBoneMatrices, ERHIFeatureLevel::Type InFeatureLevel)
-	:	FVertexFactory(InFeatureLevel), ShaderData( InBoneMatrices )
+	FGPUBaseSkinVertexFactory(ERHIFeatureLevel::Type InFeatureLevel)
+	:	FVertexFactory(InFeatureLevel)
 	{}
 
 	/** accessor */
@@ -507,8 +529,8 @@ public:
 	 *
 	 * @param	InBoneMatrices	Reference to shared bone matrices array.
 	 */
-	TGPUSkinVertexFactory(TArray<FBoneSkinning>& InBoneMatrices, ERHIFeatureLevel::Type InFeatureLevel)
-		: FGPUBaseSkinVertexFactory(InBoneMatrices, InFeatureLevel)
+	TGPUSkinVertexFactory(ERHIFeatureLevel::Type InFeatureLevel)
+		: FGPUBaseSkinVertexFactory(InFeatureLevel)
 	{}
 
 	virtual bool UsesExtraBoneInfluences() const override
@@ -565,8 +587,8 @@ public:
 	 *
 	 * @param	InBoneMatrices	Reference to shared bone matrices array.
 	 */
-	FGPUSkinPassthroughVertexFactory(TArray<FBoneSkinning>& InBoneMatrices, ERHIFeatureLevel::Type InFeatureLevel)
-		: TGPUSkinVertexFactory<false>(InBoneMatrices, InFeatureLevel)
+	FGPUSkinPassthroughVertexFactory(ERHIFeatureLevel::Type InFeatureLevel)
+		: TGPUSkinVertexFactory<false>(InFeatureLevel)
 	{}
 
 	static void ModifyCompilationEnvironment(EShaderPlatform Platform, const class FMaterial* Material, FShaderCompilerEnvironment& OutEnvironment);
@@ -598,8 +620,8 @@ public:
 	 *
 	 * @param	InBoneMatrices	Reference to shared bone matrices array.
 	 */
-	TGPUSkinMorphVertexFactory(TArray<FBoneSkinning>& InBoneMatrices, ERHIFeatureLevel::Type InFeatureLevel)
-	: TGPUSkinVertexFactory<bExtraBoneInfluencesT>(InBoneMatrices, InFeatureLevel)
+	TGPUSkinMorphVertexFactory(ERHIFeatureLevel::Type InFeatureLevel)
+	: TGPUSkinVertexFactory<bExtraBoneInfluencesT>(InFeatureLevel)
 	{}
 
 	static void ModifyCompilationEnvironment(EShaderPlatform Platform, const class FMaterial* Material, FShaderCompilerEnvironment& OutEnvironment);
@@ -642,7 +664,7 @@ private:
 
 // to reuse BoneBuffer structure for clothing simulation data
 typedef FBoneBufferTypeRef FClothSimulDataBufferTypeRef;
-typedef FBoneBufferPool FClothSimulDataBufferPool;
+typedef FClothBufferPool FClothSimulDataBufferPool;
 
 /** Vertex factory with vertex stream components for GPU-skinned and morph target streams */
 class FGPUBaseSkinAPEXClothVertexFactory
@@ -658,23 +680,18 @@ public:
 
 		void UpdateClothUniformBuffer(const TArray<FVector4>& InSimulPositions, const TArray<FVector4>& InSimulNormals);
 
-		void UpdateClothSimulData(const TArray<FVector4>& InSimulPositions, const TArray<FVector4>& InSimulNormals, ERHIFeatureLevel::Type FeatureLevel);
+		bool UpdateClothSimulData(FRHICommandListImmediate& RHICmdList, const TArray<FVector4>& InSimulPositions, const TArray<FVector4>& InSimulNormals, ERHIFeatureLevel::Type FeatureLevel);
 
 		void ReleaseClothSimulData()
 		{
 			APEXClothUniformBuffer.SafeRelease();
 
-			if(IsValidRef(ClothSimulPositionBuffer))
+			if(IsValidRef(ClothSimulPositionNormalBuffer))
 			{
-				ClothSimulDataBufferPool.ReleasePooledResource(ClothSimulPositionBuffer);
+				ClothSimulDataBufferPool.ReleasePooledResource(ClothSimulPositionNormalBuffer);
 			}
-			ClothSimulPositionBuffer.SafeRelease();
+			ClothSimulPositionNormalBuffer.SafeRelease();
 
-			if(IsValidRef(ClothSimulNormalBuffer))
-			{
-				ClothSimulDataBufferPool.ReleasePooledResource(ClothSimulNormalBuffer);
-			}
-			ClothSimulNormalBuffer.SafeRelease();
 		}
 
 		TUniformBufferRef<FAPEXClothUniformShaderParameters> GetClothUniformBuffer() const
@@ -682,20 +699,14 @@ public:
 			return APEXClothUniformBuffer;
 		}
 
-		FClothSimulDataBufferTypeRef GetClothSimulPositionBuffer() const
-		{
-			return ClothSimulPositionBuffer;
-		}
-
-		FClothSimulDataBufferTypeRef GetClothSimulNormalBuffer() const
+		FClothSimulDataBufferTypeRef GetClothSimulPositionNormalBuffer() const
 		{
-			return ClothSimulNormalBuffer;
+			return ClothSimulPositionNormalBuffer;
 		}
 
 	private:
 		TUniformBufferRef<FAPEXClothUniformShaderParameters> APEXClothUniformBuffer;
-		FClothSimulDataBufferTypeRef ClothSimulPositionBuffer;
-		FClothSimulDataBufferTypeRef ClothSimulNormalBuffer;
+		FClothSimulDataBufferTypeRef ClothSimulPositionNormalBuffer;
 
 	};
 
@@ -748,8 +759,8 @@ public:
 	 *
 	 * @param	InBoneMatrices	Reference to shared bone matrices array.
 	 */
-	TGPUSkinAPEXClothVertexFactory(TArray<FBoneSkinning>& InBoneMatrices, ERHIFeatureLevel::Type InFeatureLevel)
-		: TGPUSkinVertexFactory<bExtraBoneInfluencesT>(InBoneMatrices, InFeatureLevel)
+	TGPUSkinAPEXClothVertexFactory(ERHIFeatureLevel::Type InFeatureLevel)
+		: TGPUSkinVertexFactory<bExtraBoneInfluencesT>(InFeatureLevel)
 	{}
 
 	static void ModifyCompilationEnvironment(EShaderPlatform Platform, const class FMaterial* Material, FShaderCompilerEnvironment& OutEnvironment);
diff --git a/Engine/Source/Runtime/Engine/Public/PrimitiveSceneProxy.h b/Engine/Source/Runtime/Engine/Public/PrimitiveSceneProxy.h
--- a/Engine/Source/Runtime/Engine/Public/PrimitiveSceneProxy.h
+++ b/Engine/Source/Runtime/Engine/Public/PrimitiveSceneProxy.h
@@ -484,6 +484,10 @@ public:
 	 */
 	ENGINE_API virtual void ApplyWorldOffset(FVector InOffset);
 
+	/**
+	 * Updates the primitive proxy's uniform buffer.
+	 */
+	ENGINE_API void UpdateUniformBuffer();
 
 protected:
 
@@ -752,6 +756,11 @@ private:
 	 */
 	ENGINE_API void SetTransform(const FMatrix& InLocalToWorld, const FBoxSphereBounds& InBounds, const FBoxSphereBounds& InLocalBounds, FVector ActorPosition);
 
+	/**
+	 * Either updates the uniform buffer or defers it until it becomes visible depending on a cvar
+	 */
+	ENGINE_API void UpdateUniformBufferMaybeLazy();
+
 	/** Updates the hidden editor view visibility map on the render thread */
 	void SetHiddenEdViews_RenderThread( uint64 InHiddenEditorViews );
 
diff --git a/Engine/Source/Runtime/Engine/Public/SkeletalRenderPublic.h b/Engine/Source/Runtime/Engine/Public/SkeletalRenderPublic.h
--- a/Engine/Source/Runtime/Engine/Public/SkeletalRenderPublic.h
+++ b/Engine/Source/Runtime/Engine/Public/SkeletalRenderPublic.h
@@ -46,6 +46,13 @@ public:
 	 */
 	virtual void UpdateDynamicData_RenderThread(FRHICommandListImmediate& RHICmdList, class FDynamicSkelMeshObjectData* InDynamicData) = 0;
 
+	/**
+	 * Called by FSkeletalMeshObject prior to GDME. This allows the GPU skin version to update bones etc now that we know we are going to render
+	 */
+	virtual void PreGDMECallback()
+	{
+	}
+
 	/**
 	 * @param	LODIndex - each LOD has its own vertex data
 	 * @param	ChunkIdx - not used
diff --git a/Engine/Source/Runtime/Engine/Public/TickTaskManagerInterface.h b/Engine/Source/Runtime/Engine/Public/TickTaskManagerInterface.h
--- a/Engine/Source/Runtime/Engine/Public/TickTaskManagerInterface.h
+++ b/Engine/Source/Runtime/Engine/Public/TickTaskManagerInterface.h
@@ -1,8 +1,9 @@
 // Copyright 1998-2015 Epic Games, Inc. All Rights Reserved.
 
 
-#ifndef __TickTaskManagerInterface_h__
-#define __TickTaskManagerInterface_h__
+#pragma once
+
+DECLARE_STATS_GROUP(TEXT("TickGroups"), STATGROUP_TickGroups, STATCAT_Advanced);
 
 /** 
  * Interface for the tick task manager
@@ -61,4 +62,3 @@ public:
 };
 
 
-#endif		// __DerivedDataBackendInterface_h__
\ No newline at end of file
diff --git a/Engine/Source/Runtime/RHI/Private/RHICommandList.cpp b/Engine/Source/Runtime/RHI/Private/RHICommandList.cpp
--- a/Engine/Source/Runtime/RHI/Private/RHICommandList.cpp
+++ b/Engine/Source/Runtime/RHI/Private/RHICommandList.cpp
@@ -553,18 +553,26 @@ struct FRHICommandRHIThreadFence : public FRHICommand<FRHICommandRHIThreadFence>
 };
 
 
-FGraphEventRef FRHICommandListImmediate::RHIThreadFence()
-	{
+FGraphEventRef FRHICommandListImmediate::RHIThreadFence(bool bSetLockFence)
+{
 	check(IsInRenderingThread() && GRHIThread);
 	FRHICommandRHIThreadFence* Cmd = new (AllocCommand<FRHICommandRHIThreadFence>()) FRHICommandRHIThreadFence();
+	if (bSetLockFence)
+	{
+		RHIThreadBufferLockFence = Cmd->Fence;
+	}
 	return Cmd->Fence;
-	}		
+}		
 	
 void FRHICommandListExecutor::WaitOnRHIThreadFence(FGraphEventRef& Fence)
 {
 	check(IsInRenderingThread());
 	if (Fence.GetReference() && !Fence->IsComplete())
 	{
+		{
+			QUICK_SCOPE_CYCLE_COUNTER(STAT_WaitOnRHIThreadFence_Dispatch);
+			GetImmediateCommandList().ImmediateFlush(EImmediateFlushType::DispatchToRHIThread); // necessary to prevent deadlock
+		}
 		check(GRHIThread);
 		QUICK_SCOPE_CYCLE_COUNTER(STAT_WaitOnRHIThreadFence_Wait);
 		if (FTaskGraphInterface::Get().IsThreadProcessingTasks(ENamedThreads::RenderThread_Local))
@@ -1490,20 +1498,22 @@ static struct FLockTracker
 {
 	struct FLockParams
 	{
+		void* RHIBuffer;
 		void* Buffer;
 		uint32 BufferSize;
 		uint32 Offset;
 		EResourceLockMode LockMode;
 
-		FORCEINLINE_DEBUGGABLE FLockParams(void* InBuffer, uint32 InOffset, uint32 InBufferSize, EResourceLockMode InLockMode)
-			: Buffer(InBuffer)
+		FORCEINLINE_DEBUGGABLE FLockParams(void* InRHIBuffer, void* InBuffer, uint32 InOffset, uint32 InBufferSize, EResourceLockMode InLockMode)
+			: RHIBuffer(InRHIBuffer)
+			, Buffer(InBuffer)
 			, BufferSize(InBufferSize)
 			, Offset(InOffset)
 			, LockMode(InLockMode)
 		{
 		}
 	};
-	TMap<void *, FLockParams> OutstandingLocks;
+	TArray<FLockParams, TInlineAllocator<16> > OutstandingLocks;
 	uint32 TotalMemoryOutstanding;
 
 	FLockTracker()
@@ -1511,15 +1521,30 @@ static struct FLockTracker
 		TotalMemoryOutstanding = 0;
 	}
 
-	void Lock(void* RHIBuffer, void* Buffer, uint32 Offset, uint32 SizeRHI, EResourceLockMode LockMode)
+	FORCEINLINE_DEBUGGABLE void Lock(void* RHIBuffer, void* Buffer, uint32 Offset, uint32 SizeRHI, EResourceLockMode LockMode)
 	{
-		check(!OutstandingLocks.Contains(RHIBuffer));
-		OutstandingLocks.Add(RHIBuffer, FLockParams(Buffer, Offset, SizeRHI, LockMode));
+#if DO_CHECK
+		for (auto& Parms : OutstandingLocks)
+		{
+			check(Parms.RHIBuffer != RHIBuffer);
+		}
+#endif
+		OutstandingLocks.Add(FLockParams(RHIBuffer, Buffer, Offset, SizeRHI, LockMode));
 		TotalMemoryOutstanding += SizeRHI;
 	}
-	FLockParams Unlock(void* RHIBuffer)
+	FORCEINLINE_DEBUGGABLE FLockParams Unlock(void* RHIBuffer)
 	{
-		return OutstandingLocks.FindAndRemoveChecked(RHIBuffer);
+		for (int32 Index = 0; Index < OutstandingLocks.Num(); Index++)
+		{
+			if (OutstandingLocks[Index].RHIBuffer == RHIBuffer)
+			{
+				FLockParams Result = OutstandingLocks[Index];
+				OutstandingLocks.RemoveAtSwap(Index, 1, false);
+				return Result;
+			}
+		}
+		check(!"Mismatched RHI buffer locks.");
+		return FLockParams(nullptr, nullptr, 0, 0, RLM_WriteOnly);
 	}
 } GLockTracker;
 
@@ -1532,11 +1557,13 @@ void* FDynamicRHI::LockVertexBuffer_RenderThread(class FRHICommandListImmediate&
 	void* Result;
 	if (!bBuffer || LockMode != RLM_WriteOnly || RHICmdList.Bypass() || !GRHIThread)
 	{
+		QUICK_SCOPE_CYCLE_COUNTER(STAT_RHIMETHOD_LockVertexBuffer_Flush);
 		RHICmdList.ImmediateFlush(EImmediateFlushType::FlushRHIThread); 
 		Result = GDynamicRHI->RHILockVertexBuffer(VertexBuffer, Offset, SizeRHI, LockMode);
 	}
 	else
 	{
+		QUICK_SCOPE_CYCLE_COUNTER(STAT_RHIMETHOD_LockVertexBuffer_Malloc);
 		Result = FMemory::Malloc(SizeRHI, 16);
 	}
 	check(Result);
@@ -1560,9 +1587,10 @@ void FDynamicRHI::UnlockVertexBuffer_RenderThread(class FRHICommandListImmediate
 	else
 	{
 		new (RHICmdList.AllocCommand<FRHICommandUpdateVertexBuffer>()) FRHICommandUpdateVertexBuffer(VertexBuffer, Params.Buffer, Params.Offset, Params.BufferSize);
-		RHIThreadBufferLockFence = RHICmdList.RHIThreadFence();
+		RHICmdList.RHIThreadFence(true);
 		if (GLockTracker.TotalMemoryOutstanding > 256 * 1024)
 		{
+			QUICK_SCOPE_CYCLE_COUNTER(STAT_RHIMETHOD_UnlockVertexBuffer_FlushForMem);
 			// we could be loading a level or something, lets get this stuff going
 			RHICmdList.ImmediateFlush(EImmediateFlushType::DispatchToRHIThread); 
 			GLockTracker.TotalMemoryOutstanding = 0;
@@ -1578,11 +1606,13 @@ void* FDynamicRHI::LockIndexBuffer_RenderThread(class FRHICommandListImmediate&
 	void* Result;
 	if (!bBuffer || LockMode != RLM_WriteOnly || RHICmdList.Bypass() || !GRHIThread)
 	{
+		QUICK_SCOPE_CYCLE_COUNTER(STAT_RHIMETHOD_LockIndexBuffer_Flush);
 		RHICmdList.ImmediateFlush(EImmediateFlushType::FlushRHIThread); 
 		Result = GDynamicRHI->RHILockIndexBuffer(IndexBuffer, Offset, SizeRHI, LockMode);
 	}
 	else
 	{
+		QUICK_SCOPE_CYCLE_COUNTER(STAT_RHIMETHOD_LockIndexBuffer_Malloc);
 		Result = FMemory::Malloc(SizeRHI, 16);
 	}
 	check(Result);
@@ -1606,9 +1636,10 @@ void FDynamicRHI::UnlockIndexBuffer_RenderThread(class FRHICommandListImmediate&
 	else
 	{
 		new (RHICmdList.AllocCommand<FRHICommandUpdateIndexBuffer>()) FRHICommandUpdateIndexBuffer(IndexBuffer, Params.Buffer, Params.Offset, Params.BufferSize);
-		RHIThreadBufferLockFence = RHICmdList.RHIThreadFence();
+		RHICmdList.RHIThreadFence(true);
 		if (GLockTracker.TotalMemoryOutstanding > 256 * 1024)
 		{
+			QUICK_SCOPE_CYCLE_COUNTER(STAT_RHIMETHOD_UnlockIndexBuffer_FlushForMem);
 			// we could be loading a level or something, lets get this stuff going
 			RHICmdList.ImmediateFlush(EImmediateFlushType::DispatchToRHIThread); 
 			GLockTracker.TotalMemoryOutstanding = 0;
@@ -1666,7 +1697,7 @@ void FRHICommandListImmediate::UpdateTextureReference(FTextureReferenceRHIParamR
 		return;
 	}
 	new (AllocCommand<FRHICommandUpdateTextureReference>()) FRHICommandUpdateTextureReference(TextureRef, NewTexture);
-	RHIThreadBufferLockFence = RHIThreadFence();
+	RHIThreadFence(true);
 	if (GetUsedMemory() > 256 * 1024)
 	{
 		// we could be loading a level or something, lets get this stuff going
diff --git a/Engine/Source/Runtime/RHI/Public/RHICommandList.h b/Engine/Source/Runtime/RHI/Public/RHICommandList.h
--- a/Engine/Source/Runtime/RHI/Public/RHICommandList.h
+++ b/Engine/Source/Runtime/RHI/Public/RHICommandList.h
@@ -1922,7 +1922,7 @@ public:
 	static FGraphEventRef RenderThreadTaskFence();
 	static void WaitOnRenderThreadTaskFence(FGraphEventRef& Fence);
 	static bool AnyRenderThreadTasksOutstanding();
-	FGraphEventRef RHIThreadFence();
+	FGraphEventRef RHIThreadFence(bool bSetLockFence = false);
 
 	FORCEINLINE void GpuTimeBegin(uint32 Hash,bool bCompute)
 	{
diff --git a/Engine/Source/Runtime/Renderer/Private/PrimitiveSceneInfo.cpp b/Engine/Source/Runtime/Renderer/Private/PrimitiveSceneInfo.cpp
--- a/Engine/Source/Runtime/Renderer/Private/PrimitiveSceneInfo.cpp
+++ b/Engine/Source/Runtime/Renderer/Private/PrimitiveSceneInfo.cpp
@@ -99,7 +99,8 @@ FPrimitiveSceneInfo::FPrimitiveSceneInfo(UPrimitiveComponent* InComponent,FScene
 	NumES2DynamicPointLights(0),
 	PackedIndex(INDEX_NONE),
 	ComponentForDebuggingOnly(InComponent),
-	bNeedsStaticMeshUpdate(false)
+	bNeedsStaticMeshUpdate(false),
+	bNeedsUniformBufferUpdate(false)
 {
 	check(ComponentForDebuggingOnly);
 	check(PrimitiveComponentId.IsValid());
@@ -322,6 +323,13 @@ void FPrimitiveSceneInfo::UpdateStaticMeshes(FRHICommandListImmediate& RHICmdLis
 	}
 }
 
+void FPrimitiveSceneInfo::UpdateUniformBuffer(FRHICommandListImmediate& RHICmdList)
+{
+	checkSlow(bNeedsUniformBufferUpdate);
+	bNeedsUniformBufferUpdate = false;
+	Proxy->UpdateUniformBuffer();
+}
+
 void FPrimitiveSceneInfo::BeginDeferredUpdateStaticMeshes()
 {
 	// Set a flag which causes InitViews to update the static meshes the next time the primitive is visible.
diff --git a/Engine/Source/Runtime/Renderer/Private/SceneVisibility.cpp b/Engine/Source/Runtime/Renderer/Private/SceneVisibility.cpp
--- a/Engine/Source/Runtime/Renderer/Private/SceneVisibility.cpp
+++ b/Engine/Source/Runtime/Renderer/Private/SceneVisibility.cpp
@@ -802,7 +802,7 @@ struct FRelevancePacket
 	FRelevancePrimSet<FTranslucentPrimSet::FSortedPrim> SortedTranslucencyPrims;
 	FRelevancePrimSet<FPrimitiveSceneProxy*> DistortionPrimSet;
 	FRelevancePrimSet<FPrimitiveSceneProxy*> CustomDepthSet;
-	FRelevancePrimSet<FPrimitiveSceneInfo*> UpdateStaticMeshes;
+	FRelevancePrimSet<FPrimitiveSceneInfo*> LazyUpdates;
 	FRelevancePrimSet<FPrimitiveSceneInfo*> VisibleEditorPrimitives;
 	uint16 CombinedShadingModelMask;
 	bool bUsesGlobalDistanceField;
@@ -941,9 +941,9 @@ struct FRelevancePacket
 				PrimitiveSceneInfo->CachedReflectionCaptureProxy = Scene->FindClosestReflectionCapture(Scene->PrimitiveBounds[BitIndex].Origin);
 				PrimitiveSceneInfo->bNeedsCachedReflectionCaptureUpdate = false;
 			}
-			if (PrimitiveSceneInfo->NeedsUpdateStaticMeshes())
+			if (PrimitiveSceneInfo->NeedsLazyUpdateForRendering())
 			{
-				UpdateStaticMeshes.AddPrim(PrimitiveSceneInfo);
+				LazyUpdates.AddPrim(PrimitiveSceneInfo);
 			}
 		}
 	}
@@ -1047,9 +1047,9 @@ struct FRelevancePacket
 		DistortionPrimSet.AppendTo(WriteView.DistortionPrimSet);
 		CustomDepthSet.AppendTo(WriteView.CustomDepthSet);
 
-		for (int32 Index = 0; Index < UpdateStaticMeshes.NumPrims; Index++)
+		for (int32 Index = 0; Index < LazyUpdates.NumPrims; Index++)
 		{
-			UpdateStaticMeshes.Prims[Index]->UpdateStaticMeshes(RHICmdList);
+			LazyUpdates.Prims[Index]->ConditionalLazyUpdateForRendering(RHICmdList);
 		}
 	}
 };
diff --git a/Engine/Source/Runtime/Renderer/Private/ShadowSetup.cpp b/Engine/Source/Runtime/Renderer/Private/ShadowSetup.cpp
--- a/Engine/Source/Runtime/Renderer/Private/ShadowSetup.cpp
+++ b/Engine/Source/Runtime/Renderer/Private/ShadowSetup.cpp
@@ -727,6 +727,10 @@ void FProjectedShadowInfo::AddSubjectPrimitive(FPrimitiveSceneInfo* PrimitiveSce
 			// Update the primitive component's last render time. Allows the component to update when using bCastWhenHidden.
 			const float CurrentWorldTime = Views[0]->Family->CurrentWorldTime;
 			*(PrimitiveSceneInfo->ComponentLastRenderTime) = CurrentWorldTime;
+			if (PrimitiveSceneInfo->NeedsLazyUpdateForRendering())
+			{
+				PrimitiveSceneInfo->ConditionalLazyUpdateForRendering(FRHICommandListExecutor::GetImmediateCommandList());
+			}
 		}
 
 		if (bOpaqueRelevance && bShadowRelevance)
diff --git a/Engine/Source/Runtime/Renderer/Public/PrimitiveSceneInfo.h b/Engine/Source/Runtime/Renderer/Public/PrimitiveSceneInfo.h
--- a/Engine/Source/Runtime/Renderer/Public/PrimitiveSceneInfo.h
+++ b/Engine/Source/Runtime/Renderer/Public/PrimitiveSceneInfo.h
@@ -255,6 +255,18 @@ public:
 		return bNeedsStaticMeshUpdate;
 	}
 
+	/** return true if we need to call LazyUpdateForRendering */
+	FORCEINLINE bool NeedsUniformBufferUpdate()
+	{
+		return bNeedsUniformBufferUpdate;
+	}
+
+	/** return true if we need to call ConditionalLazyUpdateForRendering */
+	FORCEINLINE bool NeedsLazyUpdateForRendering()
+	{
+		return NeedsUniformBufferUpdate() || NeedsUpdateStaticMeshes();
+	}
+
 	/** Updates the primitive's static meshes in the scene. */
 	void UpdateStaticMeshes(FRHICommandListImmediate& RHICmdList);
 
@@ -267,6 +279,25 @@ public:
 		}
 	}
 
+	/** Updates the primitive's uniform buffer. */
+	void UpdateUniformBuffer(FRHICommandListImmediate& RHICmdList);
+
+	/** Updates the primitive's uniform buffer. */
+	FORCEINLINE void ConditionalUpdateUniformBuffer(FRHICommandListImmediate& RHICmdList)
+	{
+		if (NeedsUniformBufferUpdate())
+		{
+			UpdateUniformBuffer(RHICmdList);
+		}
+	}
+
+	/** Updates all lazy data for the rendering. */
+	FORCEINLINE void ConditionalLazyUpdateForRendering(FRHICommandListImmediate& RHICmdList)
+	{
+		ConditionalUpdateUniformBuffer(RHICmdList);
+		ConditionalUpdateStaticMeshes(RHICmdList);
+	}
+
 	/** Sets a flag to update the primitive's static meshes before it is next rendered. */
 	void BeginDeferredUpdateStaticMeshes();
 
@@ -323,6 +354,11 @@ public:
 	 */
 	void ApplyWorldOffset(FVector InOffset);
 
+	FORCEINLINE void SetNeedsUniformBufferUpdate(bool bInNeedsUniformBufferUpdate)
+	{
+		bNeedsUniformBufferUpdate = bInNeedsUniformBufferUpdate;
+	}
+
 private:
 
 	/** Let FScene have direct access to the Id. */
@@ -341,8 +377,11 @@ private:
 	 */
 	const UPrimitiveComponent* ComponentForDebuggingOnly;
 
-	/** If this is TRUE, this primitive's static meshes need to be updated before it can be rendered. */
+	/** If this is TRUE, this primitive's static meshes needs to be updated before it can be rendered. */
 	bool bNeedsStaticMeshUpdate;
+
+	/** If this is TRUE, this primitive's uniform buffer needs to be updated before it can be rendered. */
+	bool bNeedsUniformBufferUpdate;
 };
 
 /** Defines how the primitive is stored in the scene's primitive octree. */
