UE4 - remove task overhead test code

diff --git a/Engine/Source/Runtime/Engine/Private/TickTaskManager.cpp b/Engine/Source/Runtime/Engine/Private/TickTaskManager.cpp
--- a/Engine/Source/Runtime/Engine/Private/TickTaskManager.cpp
+++ b/Engine/Source/Runtime/Engine/Private/TickTaskManager.cpp
@@ -64,8 +64,6 @@ struct FTickContext
 	}
 };
 
-#include"scopedtimers.h"
-#include"ParallelFor.h"
 /**
  * Class that handles the actual tick tasks and starting and completing tick groups
  */
@@ -302,127 +300,6 @@ public:
 	 */
 	void StartFrame()
 	{
-#if 0
-		class FNoop
-		{
-		public:
-			FORCEINLINE TStatId GetStatId() const
-			{
-				RETURN_QUICK_DECLARE_CYCLE_STAT(FNoop, STATGROUP_TaskGraphTasks);
-			}
-			FORCEINLINE ENamedThreads::Type GetDesiredThread()
-			{
-				return ENamedThreads::AnyThread;
-			}
-			FORCEINLINE static ESubsequentsMode::Type GetSubsequentsMode() 
-			{ 
-				return ESubsequentsMode::TrackSubsequents; 
-			}
-			void DoTask(ENamedThreads::Type CurrentThread, const FGraphEventRef& MyCompletionGraphEvent)
-			{
-			}
-		};
-		class FNoopFireAndForget
-		{
-			FThreadSafeCounter& Work;
-		public:
-
-			FNoopFireAndForget(FThreadSafeCounter& InWork)
-				: Work(InWork)
-			{
-			}
-			FORCEINLINE TStatId GetStatId() const
-			{
-				RETURN_QUICK_DECLARE_CYCLE_STAT(FNoop, STATGROUP_TaskGraphTasks);
-			}
-			FORCEINLINE ENamedThreads::Type GetDesiredThread()
-			{
-				return ENamedThreads::AnyThread;
-			}
-			FORCEINLINE static ESubsequentsMode::Type GetSubsequentsMode() 
-			{ 
-				return ESubsequentsMode::FireAndForget; 
-			}
-			void DoTask(ENamedThreads::Type CurrentThread, const FGraphEventRef& MyCompletionGraphEvent)
-			{
-				Work.Increment();
-			}
-		};
-
-		{
-			double QueueNoop = 0.0;
-			double WaitNoop = 0.0;
-			double QueueFireAndForget = 0.0;
-			double WaitFireAndForget = 0.0;
-			FGraphEventArray Tasks;
-			Tasks.Reserve(1000);
-			{
-				FScopedDurationTimer Scope(QueueNoop);
-				for (int32 Index = 0; Index < 1000; Index++)
-				{
-					Tasks.Add(TGraphTask<FNoop>::CreateTask(nullptr,  ENamedThreads::GameThread).ConstructAndDispatchWhenReady());
-				}
-			}
-			{
-				FScopedDurationTimer Scope(WaitNoop);
-				FTaskGraphInterface::Get().WaitUntilTasksComplete(Tasks, ENamedThreads::GameThread_Local);
-			}
-			FThreadSafeCounter Work;		
-			{
-				FScopedDurationTimer Scope(QueueFireAndForget);
-				for (int32 Index = 0; Index < 1000; Index++)
-				{
-					TGraphTask<FNoopFireAndForget>::CreateTask(nullptr,  ENamedThreads::GameThread).ConstructAndDispatchWhenReady(Work);
-				}
-			}
-			{
-				FScopedDurationTimer Scope(WaitFireAndForget);
-				while (Work.GetValue() != 1000)
-					;
-			}
-
-			FPlatformMisc::LowLevelOutputDebugStringf(TEXT("GT Dispatch    Normal: %7.5fms / queue    %7.5fms / execute    FireAndForget: %7.5fms / queue    %7.5fms / execute\n"), float(QueueNoop), float(WaitNoop), float(QueueFireAndForget), float(WaitFireAndForget));
-		}
-
-		{
-			double QueueNoop = 0.0;
-			double WaitNoop = 0.0;
-			double QueueFireAndForget = 0.0;
-			double WaitFireAndForget = 0.0;
-			FGraphEventArray Tasks;
-			Tasks.AddZeroed(1000);
-			{
-				FScopedDurationTimer Scope(QueueNoop);
-				ParallelFor(1000,
-					[&Tasks](int32 Index)
-					{
-						Tasks[Index] = TGraphTask<FNoop>::CreateTask(nullptr).ConstructAndDispatchWhenReady();
-					}
-				);
-			}
-			{
-				FScopedDurationTimer Scope(WaitNoop);
-				FTaskGraphInterface::Get().WaitUntilTasksComplete(Tasks, ENamedThreads::GameThread_Local);
-			}
-			FThreadSafeCounter Work;		
-			{
-				FScopedDurationTimer Scope(QueueFireAndForget);
-				ParallelFor(1000,
-					[&Work](int32 Index)
-					{
-						TGraphTask<FNoopFireAndForget>::CreateTask(nullptr).ConstructAndDispatchWhenReady(Work);
-					}
-				);
-			}
-			{
-				FScopedDurationTimer Scope(WaitFireAndForget);
-				while (Work.GetValue() != 1000)
-					;
-			}
-
-			FPlatformMisc::LowLevelOutputDebugStringf(TEXT("Parallel For   Normal: %7.5fms / queue    %7.5fms / execute    FireAndForget: %7.5fms / queue    %7.5fms / execute\n"), float(QueueNoop), float(WaitNoop), float(QueueFireAndForget), float(WaitFireAndForget));
-		}
-#endif
 		bLogTicks = !!CVarLogTicks.GetValueOnGameThread();
 		bLogTicksShowPrerequistes = !!CVarLogTicksShowPrerequistes.GetValueOnGameThread();
 
