Quick fix for UE-17327 MAIN: Client crash when sitting in homebase

References on UberGraph Frame are not hendlead as weak.

diff --git a/Engine/Source/Runtime/Engine/Private/BlueprintGeneratedClass.cpp b/Engine/Source/Runtime/Engine/Private/BlueprintGeneratedClass.cpp
--- a/Engine/Source/Runtime/Engine/Private/BlueprintGeneratedClass.cpp
+++ b/Engine/Source/Runtime/Engine/Private/BlueprintGeneratedClass.cpp
@@ -688,7 +688,7 @@ public:
 protected:
 	virtual FArchive& operator<<(UObject*& Object) override
 	{
-		const bool bWeakRef = Object ? !Object->HasAnyFlags(RF_StrongRefOnFrame) : false;
+		const bool bWeakRef = false; //Object ? !Object->HasAnyFlags(RF_StrongRefOnFrame) : false;
 		Collector.SetShouldHandleAsWeakRef(bWeakRef); 
 		return FSimpleObjectReferenceCollectorArchive::operator<<(Object);
 	}
