Build - Fixing the build.

diff --git a/Engine/Source/Runtime/Engine/Private/SkeletalRenderGPUSkin.cpp b/Engine/Source/Runtime/Engine/Private/SkeletalRenderGPUSkin.cpp
--- a/Engine/Source/Runtime/Engine/Private/SkeletalRenderGPUSkin.cpp
+++ b/Engine/Source/Runtime/Engine/Private/SkeletalRenderGPUSkin.cpp
@@ -1189,7 +1189,6 @@ void FPreviousPerBoneMotionBlur::StartAppend(FRHICommandListImmediate& RHICmdLis
 uint32 FPreviousPerBoneMotionBlur::AppendData(FBoneSkinning*& OutChunkMatrices, uint32 BoneCount)
 {
 	check(LockedData);
-	checkSlow(DataStart);
 	checkSlow(BoneCount);
 
 	uint32 TexelCount = BoneCount * sizeof(FBoneSkinning) / sizeof(float) / 4;
