UMG - Adding the concept of shared layers to the Game Layer Manager.  Rather than a widget added directly for a player, the layer manager allows you to add a layer by name at a specific ZOrder.  The IGameLayer interface is what you implement, that allows itself to be turned into a Widget.  You can later request a layer by name, cast it to your needed type and then interact with whatever functions you exposed on your layer class.

This is preliminary work to allow widget components to share a single canvas when displaying in screenspace, which is nessesary for proper sorting.

diff --git a/Engine/Source/Runtime/Engine/Classes/Engine/GameViewportClient.h b/Engine/Source/Runtime/Engine/Classes/Engine/GameViewportClient.h
--- a/Engine/Source/Runtime/Engine/Classes/Engine/GameViewportClient.h
+++ b/Engine/Source/Runtime/Engine/Classes/Engine/GameViewportClient.h
@@ -279,12 +279,20 @@ public:
 		GameLayerManagerPtr = LayerManager;
 	}
 
+	/**
+	 * Gets the layer manager for the UI.
+	 */
+	TSharedPtr< IGameLayerManager > GetGameLayerManager() const
+	{
+		return GameLayerManagerPtr.Pin();
+	}
+
 	/** Returns access to this viewport's Slate window */
 	TSharedPtr< SWindow > GetWindow()
 	{
 		 return Window.Pin();
 	}
-	 
+	
 	/** 
 	 * Sets bDropDetail and other per-frame detail level flags on the current WorldSettings
 	 *
diff --git a/Engine/Source/Runtime/Engine/Private/Slate/SGameLayerManager.cpp b/Engine/Source/Runtime/Engine/Private/Slate/SGameLayerManager.cpp
--- a/Engine/Source/Runtime/Engine/Private/Slate/SGameLayerManager.cpp
+++ b/Engine/Source/Runtime/Engine/Private/Slate/SGameLayerManager.cpp
@@ -108,6 +108,41 @@ void SGameLayerManager::ClearWidgetsForPlayer(ULocalPlayer* Player)
 	}
 }
 
+TSharedPtr<IGameLayer> SGameLayerManager::FindLayerForPlayer(ULocalPlayer* Player, const FName& LayerName)
+{
+	TSharedPtr<FPlayerLayer>* PlayerLayerPtr = PlayerLayers.Find(Player);
+	if ( PlayerLayerPtr )
+	{
+		return (*PlayerLayerPtr)->Layers.FindRef(LayerName);
+	}
+
+	return TSharedPtr<IGameLayer>();
+}
+
+bool SGameLayerManager::AddLayerForPlayer(ULocalPlayer* Player, const FName& LayerName, TSharedRef<IGameLayer> Layer, int32 ZOrder)
+{
+	TSharedPtr<FPlayerLayer> PlayerLayer = FindOrCreatePlayerLayer(Player);
+	if ( PlayerLayer.IsValid() )
+	{
+		TSharedPtr<IGameLayer> ExistingLayer = PlayerLayer->Layers.FindRef(LayerName);
+		if ( ExistingLayer.IsValid() )
+		{
+			return false;
+		}
+
+		PlayerLayer->Layers.Add(LayerName, Layer);
+
+		PlayerLayer->Widget->AddSlot(ZOrder)
+		[
+			Layer->AsWidget()
+		];
+
+		return true;
+	}
+
+	return false;
+}
+
 void SGameLayerManager::ClearWidgets()
 {
 	PlayerCanvas->ClearChildren();
diff --git a/Engine/Source/Runtime/Engine/Public/Slate/SGameLayerManager.h b/Engine/Source/Runtime/Engine/Public/Slate/SGameLayerManager.h
--- a/Engine/Source/Runtime/Engine/Public/Slate/SGameLayerManager.h
+++ b/Engine/Source/Runtime/Engine/Public/Slate/SGameLayerManager.h
@@ -5,6 +5,18 @@
 class FSceneViewport;
 class ULocalPlayer;
 
+
+/**
+ * Allows you to provide a custom layer that multiple sources can contribute to.  Unlike
+ * adding widgets directly to the layer manager.  First registering a layer with a name
+ * allows multiple widgets to be added.
+ */
+class IGameLayer : public TSharedFromThis<IGameLayer>
+{
+public:
+	virtual TSharedRef<SWidget> AsWidget() = 0;
+};
+
 /**
  * Allows widgets to be managed for different users.
  */
@@ -14,10 +26,13 @@ public:
 	virtual void NotifyPlayerAdded(int32 PlayerIndex, ULocalPlayer* AddedPlayer) = 0;
 	virtual void NotifyPlayerRemoved(int32 PlayerIndex, ULocalPlayer* RemovedPlayer) = 0;
 
-	virtual void AddWidgetForPlayer(ULocalPlayer* Player, TSharedRef<SWidget> ViewportContent, const int32 ZOrder) = 0;
+	virtual void AddWidgetForPlayer(ULocalPlayer* Player, TSharedRef<SWidget> ViewportContent, int32 ZOrder) = 0;
 	virtual void RemoveWidgetForPlayer(ULocalPlayer* Player, TSharedRef<SWidget> ViewportContent) = 0;
 	virtual void ClearWidgetsForPlayer(ULocalPlayer* Player) = 0;
 
+	virtual TSharedPtr<IGameLayer> FindLayerForPlayer(ULocalPlayer* Player, const FName& LayerName) = 0;
+	virtual bool AddLayerForPlayer(ULocalPlayer* Player, const FName& LayerName, TSharedRef<IGameLayer> Layer, int32 ZOrder) = 0;
+
 	virtual void ClearWidgets() = 0;
 };
 
@@ -54,12 +69,15 @@ public:
 	virtual void NotifyPlayerAdded(int32 PlayerIndex, ULocalPlayer* AddedPlayer) override;
 	virtual void NotifyPlayerRemoved(int32 PlayerIndex, ULocalPlayer* RemovedPlayer) override;
 
-	virtual void AddWidgetForPlayer(ULocalPlayer* Player, TSharedRef<SWidget> ViewportContent, const int32 ZOrder) override;
+	virtual void AddWidgetForPlayer(ULocalPlayer* Player, TSharedRef<SWidget> ViewportContent, int32 ZOrder) override;
 	virtual void RemoveWidgetForPlayer(ULocalPlayer* Player, TSharedRef<SWidget> ViewportContent) override;
 	virtual void ClearWidgetsForPlayer(ULocalPlayer* Player) override;
 
+	virtual TSharedPtr<IGameLayer> FindLayerForPlayer(ULocalPlayer* Player, const FName& LayerName) override;
+	virtual bool AddLayerForPlayer(ULocalPlayer* Player, const FName& LayerName, TSharedRef<IGameLayer> Layer, int32 ZOrder) override;
+
 	virtual void ClearWidgets() override;
-	// Ened IGameLayerManager
+	// End IGameLayerManager
 
 public:
 
@@ -78,6 +96,8 @@ private:
 		TSharedPtr<SOverlay> Widget;
 		SCanvas::FSlot* Slot;
 
+		TMap< FName, TSharedPtr<IGameLayer> > Layers;
+
 		FPlayerLayer()
 			: Slot(nullptr)
 		{
