[UECORE-228] Deterministic names for blueprint variables. Retested the oringal bug (UE-9929) to make sure that it's still fixed.

diff --git a/Engine/Source/Runtime/Engine/Private/SimpleConstructionScript.cpp b/Engine/Source/Runtime/Engine/Private/SimpleConstructionScript.cpp
--- a/Engine/Source/Runtime/Engine/Private/SimpleConstructionScript.cpp
+++ b/Engine/Source/Runtime/Engine/Private/SimpleConstructionScript.cpp
@@ -1031,7 +1031,7 @@ USCS_Node* USimpleConstructionScript::CreateNode(UClass* NewComponentClass, FNam
 	// note that naming logic is duplicated in CreateNodeAndRenameComponent:
 	NewComponentVariableName = GenerateNewComponentName(NewComponentClass, NewComponentVariableName);
 
-	UActorComponent* NewComponentTemplate = NewObject<UActorComponent>(Blueprint->GeneratedClass, NewComponentClass, *(NewComponentVariableName.GetPlainNameString() + FGuid::NewGuid().ToString() ), RF_ArchetypeObject|RF_Transactional|RF_Public);
+	UActorComponent* NewComponentTemplate = NewObject<UActorComponent>(Blueprint->GeneratedClass, NewComponentClass, *(NewComponentVariableName.ToString() + TEXT("_GEN_VARIABLE")), RF_ArchetypeObject | RF_Transactional | RF_Public);
 
 	return CreateNodeImpl(NewComponentTemplate, NewComponentVariableName);
 }
@@ -1044,7 +1044,7 @@ USCS_Node* USimpleConstructionScript::CreateNodeAndRenameComponent(UActorCompone
 	FName NewComponentVariableName = GenerateNewComponentName(NewComponentTemplate->GetClass());
 
 	// Relocate the instance from the transient package to the BPGC and assign it a unique object name
-	NewComponentTemplate->Rename(*(NewComponentVariableName.GetPlainNameString() + FGuid::NewGuid().ToString()), GetBlueprint()->GeneratedClass, REN_DontCreateRedirectors | REN_DoNotDirty);
+	NewComponentTemplate->Rename(*(NewComponentVariableName.ToString() + TEXT("_GEN_VARIABLE")), GetBlueprint()->GeneratedClass, REN_DontCreateRedirectors | REN_DoNotDirty);
 
 	return CreateNodeImpl(NewComponentTemplate, NewComponentVariableName);
 }
