Build - Fixing the build.

diff --git a/Engine/Source/Runtime/Engine/Public/Slate/SGameLayerManager.h b/Engine/Source/Runtime/Engine/Public/Slate/SGameLayerManager.h
--- a/Engine/Source/Runtime/Engine/Public/Slate/SGameLayerManager.h
+++ b/Engine/Source/Runtime/Engine/Public/Slate/SGameLayerManager.h
@@ -14,7 +14,13 @@ class ULocalPlayer;
 class IGameLayer : public TSharedFromThis<IGameLayer>
 {
 public:
+	/** Get the layer as a widget. */
 	virtual TSharedRef<SWidget> AsWidget() = 0;
+
+public:
+
+	/** Virtual destructor. */
+	virtual ~IToolTip() { }
 };
 
 /**
