Build - Fixing it.

diff --git a/Engine/Source/Runtime/Engine/Public/Slate/SGameLayerManager.h b/Engine/Source/Runtime/Engine/Public/Slate/SGameLayerManager.h
--- a/Engine/Source/Runtime/Engine/Public/Slate/SGameLayerManager.h
+++ b/Engine/Source/Runtime/Engine/Public/Slate/SGameLayerManager.h
@@ -20,7 +20,7 @@ public:
 public:
 
 	/** Virtual destructor. */
-	virtual ~IToolTip() { }
+	virtual ~IGameLayer() { }
 };
 
 /**
