[UE-17500] Fix an issue where deleting the scene root component in the Blueprint editor could lead to incorrect state of instances of the Blueprint class post-reconstruction.

change summary:
- modified FActorComponentInstanceData::MatchesComponent() to bypass the serialization index match check for SCS-created components. since SCS component names should always be unique, this will limit the component data cache apply phase to only match SCS-created components by name. this way we avoid the case caused by adding/deleting SCS component nodes of the same type at both the root and in subtrees of the SCS scene component hierarchy, which shifts the reconstructed components into a different serialization slot prior to applying saved component instance data.

diff --git a/Engine/Source/Runtime/Engine/Private/ComponentInstanceDataCache.cpp b/Engine/Source/Runtime/Engine/Private/ComponentInstanceDataCache.cpp
--- a/Engine/Source/Runtime/Engine/Private/ComponentInstanceDataCache.cpp
+++ b/Engine/Source/Runtime/Engine/Private/ComponentInstanceDataCache.cpp
@@ -85,7 +85,8 @@ bool FActorComponentInstanceData::MatchesComponent(const UActorComponent* Compon
 		{
 			bMatches = true;
 		}
-		else if (SourceComponentTypeSerializedIndex >= 0)
+		else if (SourceComponentTypeSerializedIndex >= 0
+			&& SourceComponentCreationMethod != EComponentCreationMethod::SimpleConstructionScript)
 		{
 			int32 FoundSerializedComponentsOfType = -1;
 			AActor* ComponentOwner = Component->GetOwner();
