Added a safety check for detour navmesh existance to ARecastNavMesh::TestHierarchicalPath #UE4

diff --git a/Engine/Source/Runtime/Engine/Private/AI/Navigation/RecastNavMesh.cpp b/Engine/Source/Runtime/Engine/Private/AI/Navigation/RecastNavMesh.cpp
--- a/Engine/Source/Runtime/Engine/Private/AI/Navigation/RecastNavMesh.cpp
+++ b/Engine/Source/Runtime/Engine/Private/AI/Navigation/RecastNavMesh.cpp
@@ -1795,7 +1795,7 @@ bool ARecastNavMesh::TestHierarchicalPath(const FNavAgentProperties& AgentProper
 	check(Cast<const ARecastNavMesh>(Self));
 
 	const ARecastNavMesh* RecastNavMesh = (const ARecastNavMesh*)Self;
-	if (Self == NULL || RecastNavMesh->RecastNavMeshImpl == NULL)
+	if (Self == nullptr || RecastNavMesh->RecastNavMeshImpl == nullptr || RecastNavMesh->RecastNavMeshImpl->DetourNavMesh == nullptr)
 	{
 		return false;
 	}
