Fixed navigation system's issues with deep actor hierarchies #UE4

UE-17020

diff --git a/Engine/Source/Runtime/Engine/Private/AI/Navigation/NavigationSystem.cpp b/Engine/Source/Runtime/Engine/Private/AI/Navigation/NavigationSystem.cpp
--- a/Engine/Source/Runtime/Engine/Private/AI/Navigation/NavigationSystem.cpp
+++ b/Engine/Source/Runtime/Engine/Private/AI/Navigation/NavigationSystem.cpp
@@ -2540,13 +2540,24 @@ void UNavigationSystem::UpdateNavOctreeAll(AActor* Actor)
 	if (Actor)
 	{
 		UpdateNavOctree(Actor);
-
+		
 		TInlineComponentArray<UActorComponent*> Components;
 		Actor->GetComponents(Components);
 
-		for (int32 Idx = 0; Idx < Components.Num(); Idx++)
+		for (int32 ComponentIndex = 0; ComponentIndex < Components.Num(); ComponentIndex++)
+		{
+			UpdateNavOctree(Components[ComponentIndex]);
+		}
+
+		if (Actor->GetRootComponent())
 		{
-			UpdateNavOctree(Components[Idx]);
+			for (int32 RootChildIndex = 0; RootChildIndex < Actor->GetRootComponent()->AttachChildren.Num(); RootChildIndex++)
+			{
+				if (Actor->GetRootComponent()->AttachChildren[RootChildIndex] && Actor->GetRootComponent()->AttachChildren[RootChildIndex]->GetOuter() != Actor)
+				{
+					UpdateNavOctreeAll(Cast<AActor>(Actor->GetRootComponent()->AttachChildren[RootChildIndex]->GetOuter()));
+				}
+			}
 		}
 	}
 }
