Stop ensure when scaling a skeletal mesh to 0

diff --git a/Engine/Source/Runtime/Engine/Private/SkeletalMesh.cpp b/Engine/Source/Runtime/Engine/Private/SkeletalMesh.cpp
--- a/Engine/Source/Runtime/Engine/Private/SkeletalMesh.cpp
+++ b/Engine/Source/Runtime/Engine/Private/SkeletalMesh.cpp
@@ -4560,10 +4560,17 @@ void FSkeletalMeshSceneProxy::GetDynamicElementsSection(const TArray<const FScen
  * @param OutLocalToWorld - Will contain the local-to-world transform when the function returns.
  * @param OutWorldToLocal - Will contain the world-to-local transform when the function returns.
  */
-void FSkeletalMeshSceneProxy::GetWorldMatrices( FMatrix& OutLocalToWorld, FMatrix& OutWorldToLocal ) const
+bool FSkeletalMeshSceneProxy::GetWorldMatrices( FMatrix& OutLocalToWorld, FMatrix& OutWorldToLocal ) const
 {
 	OutLocalToWorld = GetLocalToWorld();
+	if (OutLocalToWorld.GetScaledAxis(EAxis::X).IsNearlyZero(SMALL_NUMBER) &&
+		OutLocalToWorld.GetScaledAxis(EAxis::Y).IsNearlyZero(SMALL_NUMBER) &&
+		OutLocalToWorld.GetScaledAxis(EAxis::Z).IsNearlyZero(SMALL_NUMBER))
+	{
+		return false;
+	}
 	OutWorldToLocal = GetLocalToWorld().InverseFast();
+	return true;
 }
 
 /**
@@ -4606,7 +4613,10 @@ int32 FSkeletalMeshSceneProxy::GetCurrentLODIndex()
 void FSkeletalMeshSceneProxy::DebugDrawPhysicsAsset(int32 ViewIndex, FMeshElementCollector& Collector, const FEngineShowFlags& EngineShowFlags) const
 {
 	FMatrix LocalToWorld, WorldToLocal;
-	GetWorldMatrices( LocalToWorld, WorldToLocal );
+	if (!GetWorldMatrices(LocalToWorld, WorldToLocal))
+	{
+		return; // Cannot draw this, world matrix not valid
+	}
 
 	FMatrix ScalingMatrix = LocalToWorld;
 	FVector TotalScale = ScalingMatrix.ExtractScaling();
diff --git a/Engine/Source/Runtime/Engine/Public/SkeletalMeshTypes.h b/Engine/Source/Runtime/Engine/Public/SkeletalMeshTypes.h
--- a/Engine/Source/Runtime/Engine/Public/SkeletalMeshTypes.h
+++ b/Engine/Source/Runtime/Engine/Public/SkeletalMeshTypes.h
@@ -1741,11 +1741,12 @@ public:
 	
 	/**
 	 * Returns the world transform to use for drawing.
-	 * @param View - Current view
 	 * @param OutLocalToWorld - Will contain the local-to-world transform when the function returns.
 	 * @param OutWorldToLocal - Will contain the world-to-local transform when the function returns.
+	 * 
+	 * @return true if out matrices are valid 
 	 */
-	void GetWorldMatrices( FMatrix& OutLocalToWorld, FMatrix& OutWorldToLocal ) const;
+	bool GetWorldMatrices( FMatrix& OutLocalToWorld, FMatrix& OutWorldToLocal ) const;
 
 	/** Util for getting LOD index currently used by this SceneProxy. */
 	int32 GetCurrentLODIndex();
