Making FStreamableManager's AsyncLoadCallback null checks consistent.

diff --git a/Engine/Source/Runtime/Engine/Private/StreamableManager.cpp b/Engine/Source/Runtime/Engine/Private/StreamableManager.cpp
--- a/Engine/Source/Runtime/Engine/Private/StreamableManager.cpp
+++ b/Engine/Source/Runtime/Engine/Private/StreamableManager.cpp
@@ -344,25 +344,32 @@ void FStreamableManager::AsyncLoadCallback(FStringAssetReference Request)
 		TargetName = ResolveRedirects(TargetName);
 		Existing = StreamableItems.FindRef(TargetName);
 	}
-	if (Existing && Existing->bAsyncLoadRequestOutstanding)
+	if (Existing)
 	{
-		Existing->bAsyncLoadRequestOutstanding = false;
-		if (!Existing->Target)
+		if (Existing->bAsyncLoadRequestOutstanding)
 		{
-			FindInMemory(TargetName, Existing);
-		}
+			Existing->bAsyncLoadRequestOutstanding = false;
+			if (!Existing->Target)
+			{
+				FindInMemory(TargetName, Existing);
+			}
 
-		CheckCompletedRequests(Request, Existing);
-	}
-	if (Existing->Target)
-	{
-		UE_LOG(LogStreamableManager, Verbose, TEXT("    Found target %s"), *Existing->Target->GetFullName());
+			CheckCompletedRequests(Request, Existing);
+		}
+		if (Existing->Target)
+		{
+			UE_LOG(LogStreamableManager, Verbose, TEXT("    Found target %s"), *Existing->Target->GetFullName());
+		}
+		else
+		{
+			// Async load failed to find the object
+			Existing->bLoadFailed = true;
+			UE_LOG(LogStreamableManager, Verbose, TEXT("    Failed async load."), *TargetName.ToString());
+		}
 	}
 	else
 	{
-		// Async load failed to find the object
-		Existing->bLoadFailed = true;
-		UE_LOG(LogStreamableManager, Verbose, TEXT("    Failed async load."), *TargetName.ToString());
+		UE_LOG(LogStreamableManager, Verbose, TEXT("    Can't find streamable for %s"), *Request.ToString());
 	}
 }
 
