Stats - Added markers for FSeamlessTravelHandler, PrepareMapChange, ULevelStreaming

diff --git a/Engine/Source/Runtime/Engine/Private/LevelStreaming.cpp b/Engine/Source/Runtime/Engine/Private/LevelStreaming.cpp
--- a/Engine/Source/Runtime/Engine/Private/LevelStreaming.cpp
+++ b/Engine/Source/Runtime/Engine/Private/LevelStreaming.cpp
@@ -482,6 +482,7 @@ bool ULevelStreaming::RequestLevel(UWorld* PersistentWorld, bool bAllowLevelLoad
 			UWorld::WorldTypePreLoadMap.FindOrAdd(DesiredPackageName) = PersistentWorld->WorldType;
 
 			// Kick off async load request.
+			STAT_ADD_CUSTOMMESSAGE_NAME( STAT_NamedMarker, *(FString( TEXT( "RequestLevel - " ) + DesiredPackageName.ToString() )) );
 			LoadPackageAsync(DesiredPackageName.ToString(), nullptr, *PackageNameToLoadFrom, FLoadPackageAsyncDelegate::CreateUObject(this, &ULevelStreaming::AsyncLevelLoadComplete), PackageFlags, PIEInstanceID);
 
 			// streamingServer: server loads everything?
@@ -643,6 +644,8 @@ void ULevelStreaming::AsyncLevelLoadComplete(const FName& InPackageName, UPackag
 	// Clean up the world type list and owning world list now that PostLoad has occurred
 	UWorld::WorldTypePreLoadMap.Remove(InPackageName);
 	ULevel::StreamedLevelsOwningWorld.Remove(InPackageName);
+
+	STAT_ADD_CUSTOMMESSAGE_NAME( STAT_NamedMarker, *(FString( TEXT( "RequestLevelComplete - " ) + InPackageName.ToString() )) );
 }
 
 bool ULevelStreaming::IsLevelVisible() const
diff --git a/Engine/Source/Runtime/Engine/Private/UnrealEngine.cpp b/Engine/Source/Runtime/Engine/Private/UnrealEngine.cpp
--- a/Engine/Source/Runtime/Engine/Private/UnrealEngine.cpp
+++ b/Engine/Source/Runtime/Engine/Private/UnrealEngine.cpp
@@ -8794,7 +8794,7 @@ void LoadGametypeContent(FWorldContext &Context, const FURL& URL)
 
 bool UEngine::LoadMap( FWorldContext& WorldContext, FURL URL, class UPendingNetGame* Pending, FString& Error )
 {
-	STAT_ADD_CUSTOMMESSAGE_NAME( STAT_NamedMarker, *(URL.Map + TEXT( " - LoadMapStart" )) );
+	STAT_ADD_CUSTOMMESSAGE_NAME( STAT_NamedMarker, *(FString( TEXT( "LoadMap - " ) + URL.Map )) );
 
 	DECLARE_SCOPE_CYCLE_COUNTER(TEXT("UEngine::LoadMap"), STAT_LoadMap, STATGROUP_LoadTime);
 
@@ -8993,7 +8993,6 @@ bool UEngine::LoadMap( FWorldContext& WorldContext, FURL URL, class UPendingNetG
 	
 #endif
 
-	STAT_ADD_CUSTOMMESSAGE_NAME( STAT_NamedMarker, *(URL.Map + TEXT( " - LoadMapMid" )) );
 	MALLOC_PROFILER( FMallocProfiler::SnapshotMemoryLoadMapMid( URL.Map ); )
 
 	if( GUseSeekFreeLoading )
@@ -9320,7 +9319,7 @@ bool UEngine::LoadMap( FWorldContext& WorldContext, FURL URL, class UPendingNetG
 		}
 	}
 
-	STAT_ADD_CUSTOMMESSAGE_NAME( STAT_NamedMarker, *(URL.Map + TEXT( " - LoadMapEnd" )) );
+	STAT_ADD_CUSTOMMESSAGE_NAME( STAT_NamedMarker, *(FString( TEXT( "LoadMapComplete - " ) + URL.Map )) );
 	MALLOC_PROFILER( FMallocProfiler::SnapshotMemoryLoadMapEnd( URL.Map ); )
 
 	// Successfully started local level.
@@ -9943,6 +9942,8 @@ static void AsyncMapChangeLevelLoadCompletionCallback(const FName& PackageName,
 		Context.LoadedLevelsForPendingMapChange.Add( NULL );
 		UE_LOG(LogEngine, Warning, TEXT("NULL LevelPackage as argument to AsyncMapChangeLevelCompletionCallback") );
 	}
+
+	STAT_ADD_CUSTOMMESSAGE_NAME( STAT_NamedMarker, *(FString( TEXT( "PrepareMapChangeComplete - " ) + PackageName.ToString() )) );
 }
 
 
@@ -10002,6 +10003,7 @@ bool UEngine::PrepareMapChange(FWorldContext &Context, const TArray<FName>& Leve
 				}
 			}
 			
+			STAT_ADD_CUSTOMMESSAGE_NAME( STAT_NamedMarker, *(FString( TEXT( "PrepareMapChange - " ) + LevelName.ToString() )) );
 			LoadPackageAsync(LevelName.ToString(),
 				FLoadPackageAsyncDelegate::CreateStatic(&AsyncMapChangeLevelLoadCompletionCallback, Context.ContextHandle)
 				);
diff --git a/Engine/Source/Runtime/Engine/Private/World.cpp b/Engine/Source/Runtime/Engine/Private/World.cpp
--- a/Engine/Source/Runtime/Engine/Private/World.cpp
+++ b/Engine/Source/Runtime/Engine/Private/World.cpp
@@ -3026,8 +3026,6 @@ void UWorld::InitializeActorsForPlay(const FURL& InURL, bool bResetTime)
 	{
 		UE_LOG(LogWorld, Log, TEXT("Bringing up level for play took: %f"), FPlatformTime::Seconds() - StartTime );
 	}
-
-	STAT_ADD_CUSTOMMESSAGE_NAME( STAT_NamedMarker, *(InURL.Map + TEXT( " - Bringing Up" )) );
 }
 
 void UWorld::BeginPlay()
@@ -4302,6 +4300,8 @@ void FSeamlessTravelHandler::SeamlessTravelLoadCallback(const FName& PackageName
 
 		SetHandlerLoadedData(LevelPackage, World);
 	}
+
+	STAT_ADD_CUSTOMMESSAGE_NAME( STAT_NamedMarker, *(FString( TEXT( "StartTravelComplete - " ) + PackageName.ToString() )) );
 }
 
 bool FSeamlessTravelHandler::StartTravel(UWorld* InCurrentWorld, const FURL& InURL, const FGuid& InGuid)
@@ -4403,6 +4403,7 @@ bool FSeamlessTravelHandler::StartTravel(UWorld* InCurrentWorld, const FURL& InU
 				UWorld::WorldTypePreLoadMap.FindOrAdd(*TransitionMap) = CurrentWorld->WorldType;
 
 				// first, load the entry level package
+				STAT_ADD_CUSTOMMESSAGE_NAME( STAT_NamedMarker, *(FString( TEXT( "StartTravel - " ) + TransitionMap )) );
 				LoadPackageAsync(TransitionMap, 
 					FLoadPackageAsyncDelegate::CreateRaw(this, &FSeamlessTravelHandler::SeamlessTravelLoadCallback),
 					0, 
