Shader Complexity for PS4

diff --git a/Engine/Source/Runtime/Engine/Private/GameViewportClient.cpp b/Engine/Source/Runtime/Engine/Private/GameViewportClient.cpp
--- a/Engine/Source/Runtime/Engine/Private/GameViewportClient.cpp
+++ b/Engine/Source/Runtime/Engine/Private/GameViewportClient.cpp
@@ -2545,8 +2545,7 @@ bool UGameViewportClient::HandleViewModeCommand( const TCHAR* Cmd, FOutputDevice
 
 	if (FPlatformProperties::SupportsWindowedMode() == false)
 	{
-		if(ViewModeIndex == VMI_Unlit
-			|| ViewModeIndex == VMI_ShaderComplexity
+		if(ViewModeIndex == VMI_Unlit			
 			|| ViewModeIndex == VMI_StationaryLightOverlap
 			|| ViewModeIndex == VMI_Lit_DetailLighting
 			|| ViewModeIndex == VMI_ReflectionOverride)
@@ -2555,7 +2554,7 @@ bool UGameViewportClient::HandleViewModeCommand( const TCHAR* Cmd, FOutputDevice
 			ViewModeIndex = VMI_Lit;
 		}
 	}
-	if (ViewModeIndex != VMI_Lit && !AllowDebugViewmodes())
+	if ((ViewModeIndex != VMI_Lit && ViewModeIndex != VMI_ShaderComplexity) && !AllowDebugViewmodes())
 	{
 		Ar.Logf(TEXT("Debug viewmodes not allowed on consoles by default.  See AllowDebugViewmodes()."));
 		ViewModeIndex = VMI_Lit;
diff --git a/Engine/Source/Runtime/Engine/Private/SceneView.cpp b/Engine/Source/Runtime/Engine/Private/SceneView.cpp
--- a/Engine/Source/Runtime/Engine/Private/SceneView.cpp
+++ b/Engine/Source/Runtime/Engine/Private/SceneView.cpp
@@ -1525,11 +1525,7 @@ FSceneViewFamily::FSceneViewFamily( const ConstructionValues& CVS )
 #endif
 
 #if !WITH_EDITOR
-	// Console shader compilers don't set instruction count, 
-	// Also various console-specific rendering paths haven't been tested with shader complexity
-	check(!EngineShowFlags.ShaderComplexity);
 	check(!EngineShowFlags.StationaryLightOverlap);
-
 #else
 
 	// instead of checking IsGameWorld on rendering thread to see if we allow this flag to be disabled 
