Fix skipPhysics being passed to children by accident.

diff --git a/Engine/Source/Runtime/Engine/Private/Components/SceneComponent.cpp b/Engine/Source/Runtime/Engine/Private/Components/SceneComponent.cpp
--- a/Engine/Source/Runtime/Engine/Private/Components/SceneComponent.cpp
+++ b/Engine/Source/Runtime/Engine/Private/Components/SceneComponent.cpp
@@ -418,7 +418,8 @@ void USceneComponent::PropagateTransformUpdate(bool bTransformChanged, bool bSki
 		MarkRenderTransformDirty();
 		
 		// Now go and update children
-		UpdateChildTransforms(bSkipPhysicsMove, Teleport);
+		//Do not pass skip physics to children. This is only used when physics updates us, but in that case we really do need to update the attached children since they are kinematic
+		UpdateChildTransforms(false, Teleport);
 
 		// Refresh navigation
 		UpdateNavigationData();
