fixed UE-17545 Rendering: Volumetric Decals: crash changing surface type back to "Surface" from volumetric decal

diff --git a/Engine/Source/Runtime/Engine/Private/Materials/HLSLMaterialTranslator.h b/Engine/Source/Runtime/Engine/Private/Materials/HLSLMaterialTranslator.h
--- a/Engine/Source/Runtime/Engine/Private/Materials/HLSLMaterialTranslator.h
+++ b/Engine/Source/Runtime/Engine/Private/Materials/HLSLMaterialTranslator.h
@@ -294,8 +294,9 @@ public:
 			Chunk[MP_TessellationMultiplier]		= Material->CompilePropertyAndSetMaterialProperty(MP_TessellationMultiplier,this);
 
 			EMaterialShadingModel MaterialShadingModel = Material->GetShadingModel();
+			const EMaterialDomain Domain = (const EMaterialDomain)Material->GetMaterialDomain();
 
-			if (Material->GetMaterialDomain() == MD_Surface
+			if (Domain == MD_Surface
 				&& IsSubsurfaceShadingModel(MaterialShadingModel))
 			{
 				// Note we don't test for the blend mode as you can have a translucent material using the subsurface shading model
@@ -351,7 +352,7 @@ public:
 			}
 
 			bUsesPixelDepthOffset = IsMaterialPropertyUsed(MP_PixelDepthOffset, Chunk[MP_PixelDepthOffset], FLinearColor(0, 0, 0, 0), 1)
-				|| Material->GetDecalBlendMode() == DBM_Volumetric_DistanceFunction;
+				|| (Domain == MD_DeferredDecal && Material->GetDecalBlendMode() == DBM_Volumetric_DistanceFunction);
 
 			const bool bUsesWorldPositionOffset = IsMaterialPropertyUsed(MP_WorldPositionOffset, Chunk[MP_WorldPositionOffset], FLinearColor(0, 0, 0, 0), 3);
 			MaterialCompilationOutput.bModifiesMeshPosition = bUsesPixelDepthOffset || bUsesWorldPositionOffset;
@@ -361,7 +362,7 @@ public:
 				Errorf(TEXT("Dynamically lit translucency is not supported for BLEND_Modulate materials."));
 			}
 
-			if (Material->GetMaterialDomain() == MD_Surface)
+			if (Domain == MD_Surface)
 			{
 				if (Material->GetBlendMode() == BLEND_Modulate && Material->IsSeparateTranslucencyEnabled())
 				{
@@ -370,14 +371,14 @@ public:
 			}
 
 			// Don't allow opaque and masked materials to scene depth as the results are undefined
-			if (bUsesSceneDepth && Material->GetMaterialDomain() != MD_PostProcess && !IsTranslucentBlendMode(Material->GetBlendMode()))
+			if (bUsesSceneDepth && Domain != MD_PostProcess && !IsTranslucentBlendMode(Material->GetBlendMode()))
 			{
 				Errorf(TEXT("Only transparent or postprocess materials can read from scene depth."));
 			}
 
 			if (MaterialCompilationOutput.bRequiresSceneColorCopy)
 			{
-				if (Material->GetMaterialDomain() != MD_Surface)
+				if (Domain != MD_Surface)
 				{
 					Errorf(TEXT("Only 'surface' material domain can use the scene color node."));
 				}
@@ -397,14 +398,14 @@ public:
 				Errorf(TEXT("Light function materials must use unlit."));
 			}
 
-			if (Material->GetMaterialDomain() == MD_PostProcess && MaterialShadingModel != MSM_Unlit)
+			if (Domain == MD_PostProcess && MaterialShadingModel != MSM_Unlit)
 			{
 				Errorf(TEXT("Post process materials must use unlit."));
 			}
 
 			if (MaterialCompilationOutput.bNeedsSceneTextures)
 			{
-				if (Material->GetMaterialDomain() != MD_PostProcess)
+				if (Domain != MD_PostProcess)
 				{
 					if (Material->GetBlendMode() == BLEND_Opaque || Material->GetBlendMode() == BLEND_Masked)
 					{
