#ue4 - Add BrakingFriction for CharacterMovementComponent, optionally enabled by "bUseSeparateBrakingFriction" setting.

diff --git a/Engine/Source/Runtime/Core/Public/UObject/ObjectVersion.h b/Engine/Source/Runtime/Core/Public/UObject/ObjectVersion.h
--- a/Engine/Source/Runtime/Core/Public/UObject/ObjectVersion.h
+++ b/Engine/Source/Runtime/Core/Public/UObject/ObjectVersion.h
@@ -538,6 +538,8 @@ enum EUnrealEngineObjectUE4Version
 	VER_UE4_FIXED_DEFAULT_ORIENTATION_OF_WIDGET_COMPONENT,
 	// Removed bUsedWithUI flag from UMaterial and replaced it with a new material domain for UI
 	VER_UE4_REMOVED_MATERIAL_USED_WITH_UI_FLAG,
+	// Added braking friction separate from turning friction.
+	VER_UE4_CHARACTER_MOVEMENT_ADD_BRAKING_FRICTION,
 	// -----<new versions can be added before this line>-------------------------------------------------
 	// - this needs to be the last line (see note below)
 	VER_UE4_AUTOMATIC_VERSION_PLUS_ONE,
diff --git a/Engine/Source/Runtime/Engine/Classes/GameFramework/CharacterMovementComponent.h b/Engine/Source/Runtime/Engine/Classes/GameFramework/CharacterMovementComponent.h
--- a/Engine/Source/Runtime/Engine/Classes/GameFramework/CharacterMovementComponent.h
+++ b/Engine/Source/Runtime/Engine/Classes/GameFramework/CharacterMovementComponent.h
@@ -197,10 +197,11 @@ public:
 	FQuat OldBaseQuat;
 
 	/**
-	 * Coefficient of friction.
+	 * Setting that affects movement control. Higher values allow faster changes in direction.
+	 * If bUseSeparateBrakingFriction is false, also affects the ability to stop more quickly when braking (whenever Acceleration is zero).
 	 * This property allows you to control how much friction is applied when moving across the ground, applying an opposing force that scales with current velocity.
 	 * This can be used to simulate slippery surfaces such as ice or oil by changing the value (possibly based on the material pawn is standing on).
-	 * @see BrakingDecelerationWalking
+	 * @see BrakingDecelerationWalking, BrakingFriction, bUseSeparateBrakingFriction
 	 */
 	UPROPERTY(Category="Character Movement: Walking", EditAnywhere, BlueprintReadWrite, meta=(ClampMin="0", UIMin="0"))
 	float GroundFriction;
@@ -229,6 +230,14 @@ public:
 	UPROPERTY(Category="Character Movement (General Settings)", EditAnywhere, BlueprintReadWrite, meta=(ClampMin="0", UIMin="0"))
 	float MaxAcceleration;
 
+	/**
+	 * Friction (drag) coefficient applied when braking (whenever Acceleration = 0, or if character is exceeding max speed).
+	 * Only used if bUseSeparateBrakingFriction setting is true, otherwise current friction such as GroundFriction is used.
+	 * @see bUseSeparateBrakingFriction
+	 */
+	UPROPERTY(Category="Character Movement (General Settings)", EditAnywhere, BlueprintReadWrite, meta=(ClampMin="0", UIMin="0", EditCondition="bUseSeparateBrakingFriction"))
+	float BrakingFriction;
+
 	/**
 	 * Deceleration when walking and not applying acceleration. This is a constant opposing force that directly lowers velocity by a constant value.
 	 * @see GroundFriction, MaxAcceleration
@@ -278,7 +287,11 @@ public:
 	UPROPERTY(Category="Character Movement: Jumping / Falling", EditAnywhere, BlueprintReadWrite, meta=(ClampMin="0", UIMin="0"))
 	float AirControlBoostVelocityThreshold;
 
-	/** Friction to apply to lateral air movement when falling. */
+	/**
+	 * Friction to apply to lateral air movement when falling.
+	 * If bUseSeparateBrakingFriction is false, also affects the ability to stop more quickly when braking (whenever Acceleration is zero).
+	 * @see BrakingFriction, bUseSeparateBrakingFriction
+	 */
 	UPROPERTY(Category="Character Movement: Jumping / Falling", EditAnywhere, BlueprintReadWrite, meta=(ClampMin="0", UIMin="0"))
 	float FallingLateralFriction;
 
@@ -320,6 +333,14 @@ public:
 	UPROPERTY(Category="Character Movement (General Settings)", EditAnywhere, BlueprintReadWrite)
 	uint32 bOrientRotationToMovement:1;
 
+	/**
+	 * If true, BrakingFriction will be used to slow the character to a stop (when there is no Acceleration).
+	 * If false, braking uses the same friction passed to CalcVelocity() (ie GroundFriction when walking), multiplied by 2.
+	 * @see BrakingFriction
+	 */
+	UPROPERTY(Category="Character Movement (General Settings)", EditDefaultsOnly, BlueprintReadWrite)
+	uint32 bUseSeparateBrakingFriction : 1;
+
 protected:
 
 	/**
diff --git a/Engine/Source/Runtime/Engine/Private/Components/CharacterMovementComponent.cpp b/Engine/Source/Runtime/Engine/Private/Components/CharacterMovementComponent.cpp
--- a/Engine/Source/Runtime/Engine/Private/Components/CharacterMovementComponent.cpp
+++ b/Engine/Source/Runtime/Engine/Private/Components/CharacterMovementComponent.cpp
@@ -241,6 +241,8 @@ UCharacterMovementComponent::UCharacterMovementComponent(const FObjectInitialize
 
 	bUseControllerDesiredRotation = false;
 
+	bUseSeparateBrakingFriction = false; // Old default behavior.
+
 	bMaintainHorizontalGroundVelocity = true;
 	bImpartBaseVelocityX = true;
 	bImpartBaseVelocityY = true;
@@ -2335,7 +2337,9 @@ void UCharacterMovementComponent::CalcVelocity(float DeltaTime, float Friction,
 	if ((bZeroAcceleration && bZeroRequestedAcceleration) || bVelocityOverMax)
 	{
 		const FVector OldVelocity = Velocity;
-		ApplyVelocityBraking(DeltaTime, Friction, BrakingDeceleration);
+
+		const float ActualBrakingFriction = (bUseSeparateBrakingFriction ? BrakingFriction : Friction);
+		ApplyVelocityBraking(DeltaTime, ActualBrakingFriction, BrakingDeceleration);
 	
 		// Don't allow braking to lower us below max speed if we started above it.
 		if (bVelocityOverMax && Velocity.SizeSquared() < FMath::Square(MaxSpeed) && FVector::DotProduct(Acceleration, OldVelocity) > 0.0f)
@@ -2734,6 +2738,9 @@ void UCharacterMovementComponent::ApplyVelocityBraking(float DeltaTime, float Fr
 	float RemainingTime = DeltaTime;
 	const float MaxTimeStep = (1.0f / 33.0f);
 
+	// Old (legacy) friction doubled the affect of "Friction" here.
+	const float FrictionFactor = (bUseSeparateBrakingFriction ? 1.0f : 2.0f);
+
 	// Decelerate to brake to a stop
 	const FVector RevAccel = (bZeroBraking ? FVector::ZeroVector : (-BrakingDeceleration * Velocity.GetSafeNormal()));
 	while( RemainingTime >= MIN_TICK_TIME )
@@ -2743,7 +2750,7 @@ void UCharacterMovementComponent::ApplyVelocityBraking(float DeltaTime, float Fr
 		RemainingTime -= dt;
 
 		// apply friction and braking
-		Velocity = Velocity + ((-2.f * Friction) * Velocity + RevAccel) * dt ; 
+		Velocity = Velocity + ((-FrictionFactor * Friction) * Velocity + RevAccel) * dt ; 
 		
 		// Don't reverse direction
 		if ((Velocity | OldVel) <= 0.f)
