Fixed: DynamicModifiersOnly navmesh generation mode doesn't work
UE-17443

diff --git a/Engine/Source/Runtime/Engine/Private/AI/Navigation/PImplRecastNavMesh.cpp b/Engine/Source/Runtime/Engine/Private/AI/Navigation/PImplRecastNavMesh.cpp
--- a/Engine/Source/Runtime/Engine/Private/AI/Navigation/PImplRecastNavMesh.cpp
+++ b/Engine/Source/Runtime/Engine/Private/AI/Navigation/PImplRecastNavMesh.cpp
@@ -444,11 +444,12 @@ void FPImplRecastNavMesh::Serialize( FArchive& Ar, int32 NavMeshVersion )
 
 			// Serialize compressed tile cache layer only if navmesh requires it
 			{
+				FNavMeshTileData TileCacheLayer;
 				uint8* CompressedData = nullptr;
 				int32 CompressedDataSize = 0;
 				if (bSupportsRuntimeGeneration)
 				{
-					FNavMeshTileData TileCacheLayer = GetTileCacheLayer(Tile->header->x, Tile->header->y, Tile->header->layer);
+					TileCacheLayer = GetTileCacheLayer(Tile->header->x, Tile->header->y, Tile->header->layer);
 					CompressedData = TileCacheLayer.GetDataSafe();
 					CompressedDataSize = TileCacheLayer.DataSize;
 				}
diff --git a/Engine/Source/Runtime/Engine/Private/AI/Navigation/RecastNavMeshGenerator.cpp b/Engine/Source/Runtime/Engine/Private/AI/Navigation/RecastNavMeshGenerator.cpp
--- a/Engine/Source/Runtime/Engine/Private/AI/Navigation/RecastNavMeshGenerator.cpp
+++ b/Engine/Source/Runtime/Engine/Private/AI/Navigation/RecastNavMeshGenerator.cpp
@@ -3025,27 +3025,37 @@ FRecastNavMeshGenerator::FRecastNavMeshGenerator(ARecastNavMesh& InDestNavMesh)
 	bool bRecreateNavmesh = true;
 	if (DestNavMesh->HasValidNavmesh())
 	{
+		const bool bGameStaticNavMesh = IsGameStaticNavMesh(DestNavMesh);
 		const dtNavMeshParams* SavedNavParams = DestNavMesh->GetRecastNavMeshImpl()->DetourNavMesh->getParams();
 		if (SavedNavParams)
 		{
-			const float TileDim = Config.tileSize * Config.cs;
-			if (SavedNavParams->tileHeight == TileDim && SavedNavParams->tileWidth == TileDim)
+			if (bGameStaticNavMesh)
 			{
-				const FVector Orig = Recast2UnrealPoint(SavedNavParams->orig);
-				const FVector OrigError(FMath::Fmod(Orig.X, TileDim), FMath::Fmod(Orig.X, TileDim), FMath::Fmod(Orig.X, TileDim));
-				if (OrigError.IsNearlyZero())
+				bRecreateNavmesh = false;
+				MaxTiles = SavedNavParams->maxTiles;
+				MaxPolysPerTile = SavedNavParams->maxPolys;
+			}
+			else
+			{
+				const float TileDim = Config.tileSize * Config.cs;
+				if (SavedNavParams->tileHeight == TileDim && SavedNavParams->tileWidth == TileDim)
 				{
-					bRecreateNavmesh = false;
+					const FVector Orig = Recast2UnrealPoint(SavedNavParams->orig);
+					const FVector OrigError(FMath::Fmod(Orig.X, TileDim), FMath::Fmod(Orig.X, TileDim), FMath::Fmod(Orig.X, TileDim));
+					if (OrigError.IsNearlyZero())
+					{
+						bRecreateNavmesh = false;
+					}
 				}
-			}
 
-			// if new navmesh needs more tiles, force recreating
-			if (!bRecreateNavmesh)
-			{
-				CalcNavMeshProperties(MaxTiles, MaxPolysPerTile);
-				if (FMath::Log2(MaxTiles) != FMath::Log2(SavedNavParams->maxTiles))
+				// if new navmesh needs more tiles, force recreating
+				if (!bRecreateNavmesh)
 				{
-					bRecreateNavmesh = true;
+					CalcNavMeshProperties(MaxTiles, MaxPolysPerTile);
+					if (FMath::Log2(MaxTiles) != FMath::Log2(SavedNavParams->maxTiles))
+					{
+						bRecreateNavmesh = true;
+					}
 				}
 			}
 		};
diff --git a/Engine/Source/Runtime/Engine/Public/AI/Navigation/RecastNavMeshGenerator.h b/Engine/Source/Runtime/Engine/Public/AI/Navigation/RecastNavMeshGenerator.h
--- a/Engine/Source/Runtime/Engine/Public/AI/Navigation/RecastNavMeshGenerator.h
+++ b/Engine/Source/Runtime/Engine/Public/AI/Navigation/RecastNavMeshGenerator.h
@@ -479,7 +479,6 @@ private:
 	friend ARecastNavMesh;
 
 	/** Parameters defining navmesh tiles */
-	struct dtNavMeshParams TiledMeshParams;
 	FRecastBuildConfig Config;
 	
 	int32 NumActiveTiles;
