Fixed one place in visual logger which can give us a crash.

diff --git a/Engine/Source/Runtime/Engine/Private/VisualLogger/VisualLogger.cpp b/Engine/Source/Runtime/Engine/Private/VisualLogger/VisualLogger.cpp
--- a/Engine/Source/Runtime/Engine/Private/VisualLogger/VisualLogger.cpp
+++ b/Engine/Source/Runtime/Engine/Private/VisualLogger/VisualLogger.cpp
@@ -59,7 +59,7 @@ FVisualLogEntry* FVisualLogger::GetEntryToWrite(const class UObject* Object, flo
 {
 	FVisualLogEntry* CurrentEntry = nullptr;
 	UObject * LogOwner = FVisualLogger::FindRedirection(Object);
-	if (LogOwner == nullptr)
+	if (LogOwner == nullptr || (LogOwner != Object && CurrentEntryPerObject.Contains(LogOwner) == false))
 	{
 		return nullptr;
 	}
@@ -92,6 +92,7 @@ FVisualLogEntry* FVisualLogger::GetEntryToWrite(const class UObject* Object, flo
 
 	if (!CurrentEntry)
 	{
+		// It's first and only one usage of LogOwner as regular object to get names. We assume once that LogOwner is correct here and only here.
 		CurrentEntry = &CurrentEntryPerObject.Add(LogOwner);
 		ObjectToNameMap.Add(LogOwner, LogOwner->GetFName());
 		ObjectToClassNameMap.Add(LogOwner, *(LogOwner->GetClass()->GetName()));
