Major optimizations to the BSP building process, and BSP Undo.

Summary of changes:
- Added 'fuzzy map' for quick look up of nearly equal FVectors, improving the BSP rebuild algorithm from O(n^4) to O(n*m) where m is much smaller than n.
- Added 'fuzzy map' for quick look up of nearly equal FPlanes, improving plane merging from O(n^2) to O(n).
- Changed TTransArrays in UModel to regular arrays, to remove a lot of transaction modifies during BSP rebuild.
- UModel array allocations are now allocated upfront where possible, instead of constant dynamic resizing.
- Optimized algorithm used by UModelComponent::GenerateElements to reduce allocations and simplify map lookups.
- Removed call to ULevel::CommitModelSurfaces in UModelComponent::PostEditUndo, which was causing an enormous amount of unnecessary work on each undo operation.  It is now done, as ever, by UWorld::Tick.

diff --git a/Engine/Source/Editor/GeometryMode/Private/GeometryModifiers.cpp b/Engine/Source/Editor/GeometryMode/Private/GeometryModifiers.cpp
--- a/Engine/Source/Editor/GeometryMode/Private/GeometryModifiers.cpp
+++ b/Engine/Source/Editor/GeometryMode/Private/GeometryModifiers.cpp
@@ -1851,6 +1851,8 @@ static ABrush* ClipBrushAgainstPlane( const FPlane& InPlane, ABrush* InBrush)
 	// Intersect the giant brush with the source brush's BSP.  This will give us the finished, clipping brush
 	// contained inside of the giant brush.
 
+	ClippedBrush->Modify();
+	InBrush->Brush->Modify();
 	GEditor->bspBrushCSG( ClippedBrush, InBrush->Brush, 0, Brush_MAX, CSG_Intersect, false, false, true );
 	FBSPOps::bspUnlinkPolys( ClippedBrush->Brush );
 
@@ -2689,7 +2691,7 @@ bool UGeomModifier_Split::OnApply()
 			}
 
 			// Replace the old polygon list with the new one
-			Brush->Brush->Polys->Element.AssignButKeepOwner(NewPolygons);
+			Brush->Brush->Polys->Element = NewPolygons;
 		}
 	}
 	else if( NumPolygonsSelected == 0 && NumEdgesSelected > 0 && NumVerticesSelected == 0 )
@@ -2746,7 +2748,7 @@ bool UGeomModifier_Split::OnApply()
 		}
 
 		// Replace the old polygon list with the new one
-		Brush->Brush->Polys->Element.AssignButKeepOwner(NewPolygons);
+		Brush->Brush->Polys->Element = NewPolygons;
 	}
 	else if( NumPolygonsSelected == 1 && NumEdgesSelected == 0 && NumVerticesSelected == 2 )
 	{
@@ -2831,7 +2833,7 @@ bool UGeomModifier_Split::OnApply()
 		}
 
 		// Replace the old polygon list with the new one
-		Brush->Brush->Polys->Element.AssignButKeepOwner(NewPolygons);
+		Brush->Brush->Polys->Element = NewPolygons;
 	}
 	else if( NumPolygonsSelected == 0 && NumEdgesSelected == 0 && NumVerticesSelected == 2 )
 	{
@@ -2888,7 +2890,7 @@ bool UGeomModifier_Split::OnApply()
 		}
 
 		// Replace the old polygon list with the new one
-		Brush->Brush->Polys->Element.AssignButKeepOwner(NewPolygons);
+		Brush->Brush->Polys->Element = NewPolygons;
 	}
 
 	mode->FinalizeSourceData();
diff --git a/Engine/Source/Editor/LevelEditor/Private/LevelEditorActions.cpp b/Engine/Source/Editor/LevelEditor/Private/LevelEditorActions.cpp
--- a/Engine/Source/Editor/LevelEditor/Private/LevelEditorActions.cpp
+++ b/Engine/Source/Editor/LevelEditor/Private/LevelEditorActions.cpp
@@ -2792,7 +2792,7 @@ void FLevelEditorActionCallbacks::OnSaveBrushAsCollision()
 		// We keep no reference to this, so it will be GC'd at some point.
 		UModel* TempModel = NewObject<UModel>();
 		TempModel->Initialize(nullptr, 1);
-		TempModel->Polys->Element.AssignButKeepOwner(BuilderModel->Polys->Element);
+		TempModel->Polys->Element = BuilderModel->Polys->Element;
 
 		// Now transform each poly into local space for the selected static mesh.
 		for (int32 i = 0; i < TempModel->Polys->Element.Num(); i++)
diff --git a/Engine/Source/Editor/UnrealEd/Private/BSPOps.cpp b/Engine/Source/Editor/UnrealEd/Private/BSPOps.cpp
--- a/Engine/Source/Editor/UnrealEd/Private/BSPOps.cpp
+++ b/Engine/Source/Editor/UnrealEd/Private/BSPOps.cpp
@@ -12,6 +12,8 @@ DEFINE_LOG_CATEGORY_STATIC(LogBSPOps, Log, All);
 /** Errors encountered in Csg operation. */
 int32 FBSPOps::GErrors = 0;
 bool FBSPOps::GFastRebuild = false;
+FBspPointsGrid* FBspPointsGrid::GBspPoints = nullptr;
+FBspPointsGrid* FBspPointsGrid::GBspVectors = nullptr;
 
 static void TagReferencedNodes( UModel *Model, int32 *NodeRef, int32 *PolyRef, int32 iNode )
 {
@@ -550,9 +552,7 @@ void FBSPOps::csgCopyBrush( ABrush* Dest, ABrush* Src, uint32 PolyFlags, EObject
 	Dest->Brush = NewObject<UModel>(Dest, NAME_None, ResFlags);
 	Dest->Brush->Initialize(nullptr, Src->Brush->RootOutside);
 	Dest->Brush->Polys	= NewObject<UPolys>(Dest->Brush, NAME_None, ResFlags);
-	check(Dest->Brush->Polys->Element.GetOwner()==Dest->Brush->Polys);
-	Dest->Brush->Polys->Element.AssignButKeepOwner(Src->Brush->Polys->Element);
-	check(Dest->Brush->Polys->Element.GetOwner()==Dest->Brush->Polys);
+	Dest->Brush->Polys->Element = Src->Brush->Polys->Element;
 	Dest->GetBrushComponent()->Brush = Dest->Brush;
 	if(Src->BrushBuilder != nullptr)
 	{
@@ -665,6 +665,21 @@ static int32 AddThing( TArray<FVector>& Vectors, FVector& V, float Thresh, int32
 /** Add a new vector to the model, merging near-duplicates, and return its index. */
 int32 FBSPOps::bspAddVector( UModel* Model, FVector* V, bool Exact )
 {
+	const float Thresh = Exact ? THRESH_NORMALS_ARE_SAME : THRESH_VECTORS_ARE_NEAR;
+
+	if (FBspPointsGrid::GBspVectors)
+	{
+		// If a points grid has been built for quick vector lookup, use that instead of doing a linear search
+		const int32 NextIndex = Model->Vectors.Num();
+		const int32 ReturnedIndex = FBspPointsGrid::GBspVectors->FindOrAddPoint(*V, NextIndex, Thresh);
+		if (ReturnedIndex == NextIndex)
+		{
+			Model->Vectors.Add(*V);
+		}
+
+		return ReturnedIndex;
+	}
+
 	return AddThing
 	(
 		Model->Vectors,
@@ -677,7 +692,21 @@ int32 FBSPOps::bspAddVector( UModel* Model, FVector* V, bool Exact )
 /** Add a new point to the model, merging near-duplicates, and return its index. */
 int32 FBSPOps::bspAddPoint( UModel* Model, FVector* V, bool Exact )
 {
-	float Thresh = Exact ? THRESH_POINTS_ARE_SAME : THRESH_POINTS_ARE_NEAR;
+	const float Thresh = Exact ? THRESH_POINTS_ARE_SAME : THRESH_POINTS_ARE_NEAR;
+
+	if (FBspPointsGrid::GBspPoints)
+	{
+		// If a points grid has been built for quick point lookup, use that instead of doing a linear search
+		const int32 NextIndex = Model->Points.Num();
+		// Always look for points with a low threshold; a generous threshold can result in 'leaks' in the BSP and unwanted polys being generated
+		const int32 ReturnedIndex = FBspPointsGrid::GBspPoints->FindOrAddPoint(*V, NextIndex, THRESH_POINTS_ARE_SAME);
+		if (ReturnedIndex == NextIndex)
+		{
+			Model->Points.Add(*V);
+		}
+
+		return ReturnedIndex;
+	}
 
 	// Try to find a match quickly from the Bsp. This finds all potential matches
 	// except for any dissociated from nodes/surfaces during a rebuild.
@@ -692,13 +721,11 @@ int32 FBSPOps::bspAddPoint( UModel* Model, FVector* V, bool Exact )
 	else
 	{
 		// No match found; add it slowly to find duplicates.
-		return AddThing( Model->Points, *V, Thresh, !GFastRebuild );
+		return AddThing(Model->Points, *V, Thresh, !GFastRebuild);
 	}
 }
 
 
-//int32 FBSPOps::bspNodeToFPoly( UModel* Model, int32 iNode, FPoly* EdPoly );
-
 /**
  * Builds Bsp from the editor polygon set (EdPolys) of a model.
  *
@@ -1112,10 +1139,6 @@ int32	FBSPOps::bspAddNode( UModel* Model, int32 iParent, ENodePlace NodePlace, u
 	else
 	{
 		// Add node.
-		if( NodePlace!=NODE_Root )
-		{
-			Model->Nodes.ModifyItem( iParent );
-		}
 		int32 iNode			 = Model->Nodes.AddZeroed();
 		FBspNode& Node       = Model->Nodes[iNode];
 
@@ -1158,9 +1181,18 @@ int32	FBSPOps::bspAddNode( UModel* Model, int32 iParent, ENodePlace NodePlace, u
 		}
 
 		// Link parent to this node.
-		if     ( NodePlace==NODE_Front ) Parent->iFront = iNode;
-		else if( NodePlace==NODE_Back  ) Parent->iBack  = iNode;
-		else if( NodePlace==NODE_Plane ) Parent->iPlane = iNode;
+		if (NodePlace == NODE_Front)
+		{
+			Parent->iFront = iNode;
+		}
+		else if (NodePlace == NODE_Back)
+		{
+			Parent->iBack = iNode;
+		}
+		else if (NodePlace == NODE_Plane)
+		{
+			Parent->iPlane = iNode;
+		}
 
 		// Add all points to point table, merging nearly-overlapping polygon points
 		// with other points in the poly to prevent criscrossing vertices from
@@ -1263,3 +1295,113 @@ void FBSPOps::HandleVolumeShapeChanged(AVolume& Volume)
 		FBSPOps::csgPrepMovingBrush( &Volume );
 	}
 }
+
+
+void FBspPointsGrid::Clear(int32 InitialSize)
+{
+	GridMap.Empty(InitialSize);
+}
+
+
+// Given a grid index in one axis, a real position on the grid and a threshold radius,
+// return either:
+// - the additional grid index it can overlap in that axis, or
+// - the original grid index if there is no overlap.
+static FORCEINLINE int32 GetAdjacentIndexIfOverlapping(int32 GridIndex, float GridPos, float GridThreshold)
+{
+	if (GridPos - GridIndex < GridThreshold)
+	{
+		return GridIndex - 1;
+	}
+	else if (1.0f - (GridPos - GridIndex) < GridThreshold)
+	{
+		return GridIndex + 1;
+	}
+	else
+	{
+		return GridIndex;
+	}
+}
+
+int32 FBspPointsGrid::FindOrAddPoint(const FVector& Point, int32 Index, float PointThreshold)
+{
+	// Offset applied to the grid coordinates so aligned vertices (the normal case) don't overlap several grid items (taking into account the threshold)
+	const float GridOffset = 0.12345f;
+
+	const float AdjustedPointX = Point.X - GridOffset;
+	const float AdjustedPointY = Point.Y - GridOffset;
+	const float AdjustedPointZ = Point.Z - GridOffset;
+
+	const float GridX = AdjustedPointX * OneOverGranularity;
+	const float GridY = AdjustedPointY * OneOverGranularity;
+	const float GridZ = AdjustedPointZ * OneOverGranularity;
+
+	// Get the grid indices corresponding to the point coordinates
+	const int32 GridIndexX = FMath::FloorToInt(GridX);
+	const int32 GridIndexY = FMath::FloorToInt(GridY);
+	const int32 GridIndexZ = FMath::FloorToInt(GridZ);
+
+	// Find grid item in map
+	FBspPointsGridItem& GridItem = GridMap.FindOrAdd(FBspPointsKey(GridIndexX, GridIndexY, GridIndexZ));
+
+	// Iterate through grid item points and return a point if it's close to the threshold
+	const float PointThresholdSquared = PointThreshold * PointThreshold;
+	for (const FBspIndexedPoint& IndexedPoint : GridItem.IndexedPoints)
+	{
+		if (FVector::DistSquared(IndexedPoint.Point, Point) <= PointThresholdSquared)
+		{
+			return IndexedPoint.Index;
+		}
+	}
+
+	// Otherwise, the point is new: add it to the grid item.
+	GridItem.IndexedPoints.Emplace(Point, Index);
+
+	// The grid has a maximum threshold of a certain radius. If the point is near the edge of a grid cube, it may overlap into other items.
+	// Add it to all grid items it can be seen from.
+	const float GridThreshold = Threshold * OneOverGranularity;
+	const int32 NeighbourX = GetAdjacentIndexIfOverlapping(GridIndexX, GridX, GridThreshold);
+	const int32 NeighbourY = GetAdjacentIndexIfOverlapping(GridIndexY, GridY, GridThreshold);
+	const int32 NeighbourZ = GetAdjacentIndexIfOverlapping(GridIndexZ, GridZ, GridThreshold);
+
+	const bool bOverlapsInX = (NeighbourX != GridIndexX);
+	const bool bOverlapsInY = (NeighbourY != GridIndexY);
+	const bool bOverlapsInZ = (NeighbourZ != GridIndexZ);
+
+	if (bOverlapsInX)
+	{
+		GridMap.FindOrAdd(FBspPointsKey(NeighbourX, GridIndexY, GridIndexZ)).IndexedPoints.Emplace(Point, Index);
+
+		if (bOverlapsInY)
+		{
+			GridMap.FindOrAdd(FBspPointsKey(NeighbourX, NeighbourY, GridIndexZ)).IndexedPoints.Emplace(Point, Index);
+
+			if (bOverlapsInZ)
+			{
+				GridMap.FindOrAdd(FBspPointsKey(NeighbourX, NeighbourY, NeighbourZ)).IndexedPoints.Emplace(Point, Index);
+			}
+		}
+		else if (bOverlapsInZ)
+		{
+			GridMap.FindOrAdd(FBspPointsKey(NeighbourX, GridIndexY, NeighbourZ)).IndexedPoints.Emplace(Point, Index);
+		}
+	}
+	else
+	{
+		if (bOverlapsInY)
+		{
+			GridMap.FindOrAdd(FBspPointsKey(GridIndexX, NeighbourY, GridIndexZ)).IndexedPoints.Emplace(Point, Index);
+
+			if (bOverlapsInZ)
+			{
+				GridMap.FindOrAdd(FBspPointsKey(GridIndexX, NeighbourY, NeighbourZ)).IndexedPoints.Emplace(Point, Index);
+			}
+		}
+		else if (bOverlapsInZ)
+		{
+			GridMap.FindOrAdd(FBspPointsKey(GridIndexX, GridIndexY, NeighbourZ)).IndexedPoints.Emplace(Point, Index);
+		}
+	}
+
+	return Index;
+}
diff --git a/Engine/Source/Editor/UnrealEd/Private/EditorBsp.cpp b/Engine/Source/Editor/UnrealEd/Private/EditorBsp.cpp
--- a/Engine/Source/Editor/UnrealEd/Private/EditorBsp.cpp
+++ b/Engine/Source/Editor/UnrealEd/Private/EditorBsp.cpp
@@ -345,7 +345,6 @@ void CleanupNodes( UModel *Model, int32 iNode, int32 iParent )
 	// If this is an empty node with a coplanar, replace it with the coplanar.
 	if( Node->NumVertices==0 && Node->iPlane!=INDEX_NONE )
 	{
-		Model->Nodes.ModifyItem( Node->iPlane );
 		FBspNode* PlaneNode = &Model->Nodes[ Node->iPlane ];
 
 		// Stick our front, back, and parent nodes on the coplanar.
@@ -363,14 +362,12 @@ void CleanupNodes( UModel *Model, int32 iNode, int32 iParent )
 		if( iParent == INDEX_NONE )
 		{
 			// This node is the root.
-			Model->Nodes.ModifyItem( iNode );
 			*Node                  = *PlaneNode;   // Replace root.
 			PlaneNode->NumVertices = 0;            // Mark as unused.
 		}
 		else
 		{
 			// This is a child node.
-			Model->Nodes.ModifyItem( iParent );
 			FBspNode *ParentNode = &Model->Nodes[iParent];
 
 			if      ( ParentNode->iFront == iNode ) ParentNode->iFront = Node->iPlane;
@@ -398,7 +395,6 @@ void CleanupNodes( UModel *Model, int32 iNode, int32 iParent )
 			}
 			else
 			{
-				Model->Nodes.ModifyItem( iNode );
          		*Node = Model->Nodes[iReplacementNode];
 			}
 		}
@@ -406,7 +402,6 @@ void CleanupNodes( UModel *Model, int32 iNode, int32 iParent )
 		{
 			// Regular node.
 			FBspNode *ParentNode = &Model->Nodes[iParent];
-			Model->Nodes.ModifyItem( iParent );
 
 			if     ( ParentNode->iFront == iNode ) ParentNode->iFront = iReplacementNode;
 			else if( ParentNode->iBack  == iNode ) ParentNode->iBack  = iReplacementNode;
@@ -466,7 +461,6 @@ void AddWorldToBrushFunc( UModel* Model, int32 iNode, FPoly* EdPoly,
 			GDiscarded++;
 			if( GModel->Nodes[GNode].NumVertices )
 			{
-				GModel->Nodes.ModifyItem( GNode );
 				GModel->Nodes[GNode].NumVertices = 0;
 			}
 			break;
@@ -511,7 +505,6 @@ void SubtractWorldToBrushFunc( UModel* Model, int32 iNode, FPoly* EdPoly,
 			GDiscarded++;
 			if( GModel->Nodes[GNode].NumVertices )
 			{
-				GModel->Nodes.ModifyItem( GNode );
 				GModel->Nodes[GNode].NumVertices = 0;
 			}
 			break;
@@ -1028,14 +1021,14 @@ void FilterWorldThroughBrush
 				{
 					// Get rid of all the fragments we added.
 					Model->Nodes[GLastCoplanar].iPlane = INDEX_NONE;
-					Model->Nodes.RemoveAt( GNumNodes, Model->Nodes.Num()-GNumNodes );
+					const bool bAllowShrinking = false;
+					Model->Nodes.RemoveAt( GNumNodes, Model->Nodes.Num()-GNumNodes, bAllowShrinking );
 				}
 				else
 				{
 					// Tag original world poly for deletion; has been deleted or replaced by partial fragments.
 					if( GModel->Nodes[GNode].NumVertices )
 					{
-						GModel->Nodes.ModifyItem( GNode );
 						GModel->Nodes[GNode].NumVertices = 0;
 					}
 				}
@@ -1106,14 +1099,6 @@ int UEditorEngine::bspBrushCSG
 		NotPolyFlags |= (PF_Semisolid | PF_NotSolid);
 	}
 
-	// Prep the model.
-	Model->Modify();
-
-	// Prep the brush.
-	Actor->Modify();
-
-	// Prep the temporary model.
-	TempModel->Modify();
 	TempModel->EmptyModel(1,1);
 
 	// Update status.
@@ -1255,10 +1240,23 @@ int UEditorEngine::bspBrushCSG
 		// the tree.  We only need the cutting planes, though the entire Bsp struct (polys and
 		// all) is built.
 
+		FBspPointsGrid* LevelModelPointsGrid = FBspPointsGrid::GBspPoints;
+		FBspPointsGrid* LevelModelVectorsGrid = FBspPointsGrid::GBspVectors;
+
+		// For the bspBuild call, temporarily create a new pair of BspPointsGrids for the TempModel.
+		TUniquePtr<FBspPointsGrid> BspPoints = MakeUnique<FBspPointsGrid>(50.0f, THRESH_POINTS_ARE_SAME);
+ 		TUniquePtr<FBspPointsGrid> BspVectors = MakeUnique<FBspPointsGrid>(1 / 16.0f, FMath::Max(THRESH_NORMALS_ARE_SAME, THRESH_VECTORS_ARE_NEAR));
+ 		FBspPointsGrid::GBspPoints = BspPoints.Get();
+ 		FBspPointsGrid::GBspVectors = BspVectors.Get();
+
 		if( ReallyBig ) GWarn->StatusUpdate( 0, 0, NSLOCTEXT("UnrealEd", "BuildingBSP", "Building BSP") );
 		
 		FBSPOps::bspBuild( TempModel, FBSPOps::BSP_Lame, 0, 70, 1, 0 );
-		
+
+		// Reinstate the original BspPointsGrids used for building the level Model.
+		FBspPointsGrid::GBspPoints = LevelModelPointsGrid;
+		FBspPointsGrid::GBspVectors = LevelModelVectorsGrid;
+
 		if( ReallyBig ) GWarn->StatusUpdate( 0, 0, NSLOCTEXT("UnrealEd", "FilteringWorld", "Filtering world") );
 		GModel = Brush;
 		TempModel->BuildBound();
diff --git a/Engine/Source/Editor/UnrealEd/Private/EditorCsg.cpp b/Engine/Source/Editor/UnrealEd/Private/EditorCsg.cpp
--- a/Engine/Source/Editor/UnrealEd/Private/EditorCsg.cpp
+++ b/Engine/Source/Editor/UnrealEd/Private/EditorCsg.cpp
@@ -137,7 +137,8 @@ void UEditorEngine::csgRebuild( UWorld* InWorld )
 	FinishAllSnaps();
 
 	// Empty the model out.
-	InWorld->GetModel()->EmptyModel( 1, 1 );
+	InWorld->GetModel()->Modify();
+	InWorld->GetModel()->EmptyModel(1, 1);
 
 	// Count brushes.
 	int32 BrushTotal=0, BrushCount=0;
@@ -193,6 +194,7 @@ void UEditorEngine::csgRebuild( UWorld* InWorld )
 				Args.Add( TEXT("BrushTotal"), BrushTotal );
 				GWarn->StatusUpdate( BrushCount, BrushTotal, FText::Format( NSLOCTEXT("UnrealEd", "ApplyingStructuralBrushF", "Applying structural brush {BrushCount} of {BrushTotal}"), Args ) );
 
+				Brush->Modify();
 				bspBrushCSG( Brush, InWorld->GetModel(), Brush->PolyFlags, (EBrushType)Brush->BrushType, CSG_None, false, true, false );
 			}
 		}
@@ -236,6 +238,7 @@ void UEditorEngine::csgRebuild( UWorld* InWorld )
 			Args.Add( TEXT("BrushTotal"), BrushTotal );
 			GWarn->StatusUpdate( BrushCount, BrushTotal, FText::Format( NSLOCTEXT("UnrealEd", "ApplyingDetailBrushF", "Applying detail brush {BrushCount} of {BrushTotal}"), Args ) );
 
+			Brush->Modify();
 			bspBrushCSG( Brush, InWorld->GetModel(), Brush->PolyFlags, (EBrushType)Brush->BrushType, CSG_None, false, true, false );
 		}
 	}
@@ -388,8 +391,6 @@ void UEditorEngine::polyUpdateMaster
 		FPoly& MasterEdPoly = Brush->Polys->Element[iEdPoly];
 		if( iEdPoly==Surf.iBrushPoly || MasterEdPoly.iLink==Surf.iBrushPoly )
 		{
-			Brush->Polys->Element.ModifyItem( iEdPoly );
-
 			MasterEdPoly.Material  = Surf.Material;
 			MasterEdPoly.PolyFlags = Surf.PolyFlags & ~(PF_NoEdit);
 
@@ -1286,7 +1287,7 @@ void UEditorEngine::MapBrushGet(UWorld* InWorld)
 			ABrush* WorldBrush = BrushActor->GetWorld()->GetDefaultBrush();
 			check( WorldBrush );
 			WorldBrush->Modify();
-			WorldBrush->Brush->Polys->Element.AssignButKeepOwner(BrushActor->Brush->Polys->Element);
+			WorldBrush->Brush->Polys->Element = BrushActor->Brush->Polys->Element;
 			WorldBrush->CopyPosRotScaleFrom( BrushActor );
 
 			WorldBrush->ReregisterAllComponents();
@@ -1317,7 +1318,7 @@ void UEditorEngine::mapBrushPut()
 			check( WorldBrush );
 
 			BrushActor->Modify();
-			BrushActor->Brush->Polys->Element.AssignButKeepOwner(WorldBrush->Brush->Polys->Element);
+			BrushActor->Brush->Polys->Element = WorldBrush->Brush->Polys->Element;
 			BrushActor->CopyPosRotScaleFrom( WorldBrush );
 			BrushActor->SetNeedRebuild(BrushActor->GetLevel());
 
diff --git a/Engine/Source/Editor/UnrealEd/Private/EditorEngine.cpp b/Engine/Source/Editor/UnrealEd/Private/EditorEngine.cpp
--- a/Engine/Source/Editor/UnrealEd/Private/EditorEngine.cpp
+++ b/Engine/Source/Editor/UnrealEd/Private/EditorEngine.cpp
@@ -2137,15 +2137,6 @@ void UEditorEngine::ApplyDeltaToActor(AActor* InActor,
 		InActor->Modify();
 	}
 	
-	ABrush* Brush = Cast< ABrush >( InActor );
-	if( Brush )
-	{
-		if( Brush->GetBrushComponent() && Brush->GetBrushComponent()->Brush )
-		{
-			Brush->GetBrushComponent()->Brush->Polys->Element.ModifyAllItems();
-		}
-	}
-
 	FNavigationLockContext LockNavigationUpdates(InActor->GetWorld(), ENavigationLockReason::ContinuousEditorMove);
 
 	bool bTranslationOnly = true;
diff --git a/Engine/Source/Editor/UnrealEd/Private/EditorSelectUtils.cpp b/Engine/Source/Editor/UnrealEd/Private/EditorSelectUtils.cpp
--- a/Engine/Source/Editor/UnrealEd/Private/EditorSelectUtils.cpp
+++ b/Engine/Source/Editor/UnrealEd/Private/EditorSelectUtils.cpp
@@ -88,15 +88,6 @@ void UUnrealEdEngine::NoteActorMovement()
 					GroupActors.Add(ActorLockedRootGroup);
 				}
 			}
-
-			ABrush* Brush = Cast< ABrush >( Actor );
-			if ( Brush )
-			{
-				if( Brush->Brush )
-				{
-					Brush->Brush->Polys->Element.ModifyAllItems();
-				}
-			}
 		}
 
 		// Modify unique group actors
diff --git a/Engine/Source/Editor/UnrealEd/Private/EditorServer.cpp b/Engine/Source/Editor/UnrealEd/Private/EditorServer.cpp
--- a/Engine/Source/Editor/UnrealEd/Private/EditorServer.cpp
+++ b/Engine/Source/Editor/UnrealEd/Private/EditorServer.cpp
@@ -179,7 +179,6 @@ static int32 CleanBSPMaterials(UWorld* InWorld, bool bPreviewOnly, bool bLogBrus
 					}
 					if ( !bPreviewOnly )
 					{
-						Actor->Brush->Polys->Element.ModifyItem(PolyIndex);
 						ReferencedMaterial = UMaterial::GetDefaultMaterial(MD_Surface);
 					}
 				}
@@ -697,8 +696,9 @@ bool UEditorEngine::Exec_Brush( UWorld* InWorld, const TCHAR* Str, FOutputDevice
 					LoadAndSelectAssets( SelectedAssets, UMaterial::StaticClass() );
 				}
 
+				InWorld->GetModel()->Modify();
+				NewBrush->Modify();
 				bspBrushCSG( NewBrush, InWorld->GetModel(), DWord1, Brush_Add, CSG_None, true, true, true );
-				NewBrush->MarkPackageDirty();
 			}
 			InWorld->InvalidateModelGeometry( InWorld->GetCurrentLevel() );
 		}
@@ -780,8 +780,9 @@ bool UEditorEngine::Exec_Brush( UWorld* InWorld, const TCHAR* Str, FOutputDevice
 			NewBrush = FBSPOps::csgAddOperation(WorldBrush,0,Brush_Subtract); // Layer
 			if( NewBrush )
 			{
+				NewBrush->Modify();
+				InWorld->GetModel()->Modify();
 				bspBrushCSG( NewBrush, InWorld->GetModel(), 0, Brush_Subtract, CSG_None, true, true, true );
-				NewBrush->MarkPackageDirty();
 			}
 			InWorld->InvalidateModelGeometry( InWorld->GetCurrentLevel() );
 		}
@@ -1575,8 +1576,27 @@ void UEditorEngine::RebuildLevel(ULevel& Level)
 
 void UEditorEngine::RebuildModelFromBrushes(UModel* Model, bool bSelectedBrushesOnly)
 {
+	TUniquePtr<FBspPointsGrid> BspPoints = MakeUnique<FBspPointsGrid>(50.0f, THRESH_POINTS_ARE_SAME);
+	TUniquePtr<FBspPointsGrid> BspVectors = MakeUnique<FBspPointsGrid>(1/16.0f, FMath::Max(THRESH_NORMALS_ARE_SAME, THRESH_VECTORS_ARE_NEAR));
+	FBspPointsGrid::GBspPoints = BspPoints.Get();
+	FBspPointsGrid::GBspVectors = BspVectors.Get();
+
 	// Empty the model out.
+	const int32 NumPoints = Model->Points.Num();
+	const int32 NumNodes = Model->Nodes.Num();
+	const int32 NumVerts = Model->Verts.Num();
+	const int32 NumVectors = Model->Vectors.Num();
+	const int32 NumSurfs = Model->Surfs.Num();
+
+	Model->Modify();
 	Model->EmptyModel(1, 1);
+
+	// Reserve arrays an eighth bigger than the previous allocation
+	Model->Points.Empty(NumPoints + NumPoints / 8);
+	Model->Nodes.Empty(NumNodes + NumNodes / 8);
+	Model->Verts.Empty(NumVerts + NumVerts / 8);
+	Model->Vectors.Empty(NumVectors + NumVectors / 8);
+	Model->Surfs.Empty(NumSurfs + NumSurfs / 8);
 		
 	// Limit the brushes used to the level the model is for
 	ULevel* Level = Model->GetTypedOuter<ULevel>();
@@ -1602,6 +1622,7 @@ void UEditorEngine::RebuildModelFromBrushes(UModel* Model, bool bSelectedBrushes
 			{
 				Brush->PolyFlags = (Brush->PolyFlags & ~PF_Semisolid) | PF_NotSolid;
 			}
+			Brush->Modify();
 			bspBrushCSG(Brush, Model, Brush->PolyFlags, (EBrushType)Brush->BrushType, CSG_None, false, true, false, false);
 		}
 	}
@@ -1614,6 +1635,7 @@ void UEditorEngine::RebuildModelFromBrushes(UModel* Model, bool bSelectedBrushes
 			(!bSelectedBrushesOnly || Brush->IsSelected()) &&
 			(Brush->PolyFlags & PF_Semisolid) && !(Brush->PolyFlags & PF_Portal) && (Brush->BrushType == Brush_Add))
 		{
+			Brush->Modify();
 			bspBrushCSG(Brush, Model, Brush->PolyFlags, (EBrushType)Brush->BrushType, CSG_None, false, true, false, false);
 		}
 	}
@@ -1625,9 +1647,17 @@ void UEditorEngine::RebuildModelFromBrushes(UModel* Model, bool bSelectedBrushes
 		if (DynamicBrush && DynamicBrush->Brush && !DynamicBrush->IsStaticBrush() &&
 			(!bSelectedBrushesOnly || DynamicBrush->IsSelected()))
 		{
+			BspPoints = MakeUnique<FBspPointsGrid>(50.0f, THRESH_POINTS_ARE_SAME);
+			BspVectors = MakeUnique<FBspPointsGrid>(1 / 16.0f, FMath::Max(THRESH_NORMALS_ARE_SAME, THRESH_VECTORS_ARE_NEAR));
+			FBspPointsGrid::GBspPoints = BspPoints.Get();
+			FBspPointsGrid::GBspVectors = BspVectors.Get();
+
 			FBSPOps::csgPrepMovingBrush(DynamicBrush);
 		}
 	}
+
+	FBspPointsGrid::GBspPoints = nullptr;
+	FBspPointsGrid::GBspVectors = nullptr;
 }
 
 
@@ -1718,7 +1748,8 @@ void UEditorEngine::BSPIntersectionHelper(UWorld* InWorld, ECsgOper Operation)
 	ABrush* DefaultBrush = InWorld->GetDefaultBrush();
 	if (DefaultBrush != NULL)
 	{
-		DefaultBrush->Brush->Modify();
+		DefaultBrush->Modify();
+		InWorld->GetModel()->Modify();
 		FinishAllSnaps();
 		bspBrushCSG(DefaultBrush, InWorld->GetModel(), 0, Brush_MAX, Operation, false, true, true);
 	}
@@ -3758,7 +3789,6 @@ bool UEditorEngine::Map_Scale( UWorld* InWorld, const TCHAR* Str, FOutputDevice&
 			ABrush* Brush = Cast< ABrush >( Actor );
 			if( Brush )
 			{
-				Brush->Brush->Polys->Element.ModifyAllItems();
 				for( int32 poly = 0 ; poly < Brush->Brush->Polys->Element.Num() ; poly++ )
 				{
 					FPoly* Poly = &(Brush->Brush->Polys->Element[poly]);
@@ -3851,6 +3881,7 @@ namespace {
 		{
 			FBspSurf* Surf = *It;
 			UModel* Model = It.GetModel();
+			Model->Modify();
 			const FVector TextureU( Model->Vectors[Surf->vTextureU] );
 			const FVector TextureV( Model->Vectors[Surf->vTextureV] );
 			Surf->vTextureU = Model->Vectors.Add(TextureU);
@@ -3865,7 +3896,8 @@ namespace {
 
 		for( FConstLevelIterator Iterator = InWorld->GetLevelIterator(); Iterator; ++Iterator )
 		{
-			UModel* Model = (*Iterator)->Model;;
+			UModel* Model = (*Iterator)->Model;
+			Model->Modify();
 			GEditor->polyTexScale( Model, UU, UV, VU, VV, !!Word2 );
 		}
 	}
@@ -4157,6 +4189,7 @@ bool UEditorEngine::Exec_Poly( UWorld* InWorld, const TCHAR* Str, FOutputDevice&
 			{
 				FBspSurf* Surf = *It;
 				UModel* Model = It.GetModel();
+				Model->Modify();
 				const FVector Base( Model->Points[Surf->pBase] );
 				Surf->pBase = Model->Points.Add(Base);
 			}
@@ -4166,6 +4199,7 @@ bool UEditorEngine::Exec_Poly( UWorld* InWorld, const TCHAR* Str, FOutputDevice&
 				for( FConstLevelIterator Iterator = InWorld->GetLevelIterator(); Iterator; ++Iterator )
 				{
 					UModel* Model = (*Iterator)->Model;
+					Model->Modify();
 					polyTexPan( Model, 0, 0, 1 );
 				}
 			}
@@ -4175,6 +4209,7 @@ bool UEditorEngine::Exec_Poly( UWorld* InWorld, const TCHAR* Str, FOutputDevice&
 			for( FConstLevelIterator Iterator = InWorld->GetLevelIterator(); Iterator; ++Iterator )
 			{
 				UModel* Model = (*Iterator)->Model;
+				Model->Modify();
 				polyTexPan( Model, PanU, PanV, 0 );
 			}
 		}
diff --git a/Engine/Source/Editor/UnrealEd/Public/BSPOps.h b/Engine/Source/Editor/UnrealEd/Public/BSPOps.h
--- a/Engine/Source/Editor/UnrealEd/Public/BSPOps.h
+++ b/Engine/Source/Editor/UnrealEd/Public/BSPOps.h
@@ -4,6 +4,7 @@
 #ifndef __BSPOPS_H__
 #define __BSPOPS_H__
 
+
 class FBSPOps
 {
 public:
@@ -76,4 +77,77 @@ protected:
 /** @todo: remove when uses of ENodePlane have been replaced with FBSPOps::ENodePlace. */
 typedef FBSPOps::ENodePlace ENodePlace;
 
+
+struct FBspPointsKey
+{
+	int32 X;
+	int32 Y;
+	int32 Z;
+
+	FBspPointsKey(int32 InX, int32 InY, int32 InZ)
+		: X(InX)
+		, Y(InY)
+		, Z(InZ)
+	{}
+
+	friend FORCEINLINE bool operator == (const FBspPointsKey& A, const FBspPointsKey& B)
+	{
+		return A.X == B.X && A.Y == B.Y && A.Z == B.Z;
+	}
+
+	friend FORCEINLINE uint32 GetTypeHash(const FBspPointsKey& Key)
+	{
+		return HashCombine(static_cast<uint32>(Key.X), HashCombine(static_cast<uint32>(Key.Y), static_cast<uint32>(Key.Z)));
+	}
+};
+
+struct FBspIndexedPoint
+{
+	FBspIndexedPoint(const FVector& InPoint, int32 InIndex)
+		: Point(InPoint)
+		, Index(InIndex)
+	{}
+
+	FVector Point;
+	int32 Index;
+};
+
+
+struct FBspPointsGridItem
+{
+	TArray<FBspIndexedPoint, TInlineAllocator<16>> IndexedPoints;
+};
+
+
+// Represents a sparse granular 3D grid into which points are added for quick (~O(1)) lookup.
+// The 3D space is divided into a grid with a given granularity.
+// Points are considered to have a given radius (threshold) and are added to the grid cube they fall in, and to up to seven neighbours if they overlap.
+class FBspPointsGrid
+{
+public:
+	FBspPointsGrid(float InGranularity, float InThreshold, int32 InitialSize = 0)
+		: OneOverGranularity(1.0f / InGranularity)
+		, Threshold(InThreshold)
+	{
+		check(InThreshold / InGranularity <= 0.5f);
+		Clear(InitialSize);
+	}
+
+	void Clear(int32 InitialSize = 0);
+
+	int32 FindOrAddPoint(const FVector& Point, int32 Index, float Threshold);
+
+	static FBspPointsGrid* GBspPoints;
+	static FBspPointsGrid* GBspVectors;
+
+private:
+	float OneOverGranularity;
+	float Threshold;
+
+	typedef TMap<FBspPointsKey, FBspPointsGridItem> FGridMap;
+	FGridMap GridMap;
+};
+
+
+
 #endif // __BSPOPS_H__
diff --git a/Engine/Source/Runtime/Core/Public/UObject/ObjectVersion.h b/Engine/Source/Runtime/Core/Public/UObject/ObjectVersion.h
--- a/Engine/Source/Runtime/Core/Public/UObject/ObjectVersion.h
+++ b/Engine/Source/Runtime/Core/Public/UObject/ObjectVersion.h
@@ -1,4 +1,4 @@
-﻿// Copyright 1998-2015 Epic Games, Inc. All Rights Reserved.
+// Copyright 1998-2015 Epic Games, Inc. All Rights Reserved.
 
 #pragma once
 
@@ -540,6 +540,8 @@ enum EUnrealEngineObjectUE4Version
 	VER_UE4_REMOVED_MATERIAL_USED_WITH_UI_FLAG,
 	// Added braking friction separate from turning friction.
 	VER_UE4_CHARACTER_MOVEMENT_ADD_BRAKING_FRICTION,
+	// Removed TTransArrays from UModel
+	VER_UE4_BSP_UNDO_FIX,
 	// -----<new versions can be added before this line>-------------------------------------------------
 	// - this needs to be the last line (see note below)
 	VER_UE4_AUTOMATIC_VERSION_PLUS_ONE,
diff --git a/Engine/Source/Runtime/Engine/Classes/Engine/Polys.h b/Engine/Source/Runtime/Engine/Classes/Engine/Polys.h
--- a/Engine/Source/Runtime/Engine/Classes/Engine/Polys.h
+++ b/Engine/Source/Runtime/Engine/Classes/Engine/Polys.h
@@ -583,19 +583,19 @@ class UPolys : public UObject
 {
 	GENERATED_UCLASS_BODY()
 	// Elements.
-	TTransArray<FPoly> Element;
+	TArray<FPoly> Element;
 
 	// Constructors.
 	UPolys(const FObjectInitializer& ObjectInitializer = FObjectInitializer::Get())
 	:  UObject(ObjectInitializer)
-	, Element( this )
+	, Element( )
 	{}
 
 #if WITH_HOT_RELOAD_CTORS
 	/** DO NOT USE. This constructor is for internal usage only for hot-reload purposes. */
 	UPolys(FVTableHelper& Helper)
 		: Super(Helper)
-		, Element(this)
+		, Element()
 	{}
 #endif // WITH_HOT_RELOAD_CTORS
 
diff --git a/Engine/Source/Runtime/Engine/Private/Components/ModelComponent.cpp b/Engine/Source/Runtime/Engine/Private/Components/ModelComponent.cpp
--- a/Engine/Source/Runtime/Engine/Private/Components/ModelComponent.cpp
+++ b/Engine/Source/Runtime/Engine/Private/Components/ModelComponent.cpp
@@ -72,8 +72,6 @@ void UModelComponent::InitializeModelComponent(UModel* InModel, uint16 InCompone
 	// Model components are transacted.
 	SetFlags(RF_Transactional);
 
-	GenerateElements(true);
-
 	CastShadow = true;
 	bUseAsOccluder = true;
 	Mobility = EComponentMobility::Static;
@@ -255,7 +253,7 @@ void UModelComponent::PostEditUndo()
 	if(ensure(Level))
 	{
 		Level->InvalidateModelSurface();
-		Level->CommitModelSurfaces();
+//		Level->CommitModelSurfaces();
 	}
 	Super::PostEditUndo();
 }
@@ -387,6 +385,29 @@ void UModelComponent::GetStreamingTextureInfo(TArray<FStreamingTexturePrimitiveI
 
 #if WITH_EDITOR
 
+struct FNodeGroupKey
+{
+	FNodeGroup* NodeGroup;
+	uint32 LightMapScale;
+	UMaterialInterface* MaterialInterface;
+
+	FNodeGroupKey(FNodeGroup* InNodeGroup, float InLightMapScale, UMaterialInterface* InMaterialInterface)
+		: NodeGroup(InNodeGroup)
+		, LightMapScale(FMath::TruncToInt(InLightMapScale))
+		, MaterialInterface(InMaterialInterface)
+	{}
+
+	friend FORCEINLINE bool operator == (const FNodeGroupKey& A, const FNodeGroupKey& B)
+	{
+		return A.NodeGroup == B.NodeGroup && A.LightMapScale == B.LightMapScale && A.MaterialInterface == B.MaterialInterface;
+	}
+
+	friend FORCEINLINE uint32 GetTypeHash(const FNodeGroupKey& Key)
+	{
+		return HashCombine(GetTypeHash(Key.NodeGroup), HashCombine(Key.LightMapScale, GetTypeHash(Key.MaterialInterface)));
+	}
+};
+
 bool UModelComponent::GenerateElements(bool bBuildRenderData)
 {
 	Elements.Empty();
@@ -414,62 +435,42 @@ bool UModelComponent::GenerateElements(bool bBuildRenderData)
 	}
 	else
 	{
-		typedef TMap<UMaterialInterface*,FModelElement*> TMaterialElementMap;
-		typedef TMap<uint32,TMaterialElementMap*> TResolutionToMaterialElementMap;
-		TMap<FNodeGroup*,TResolutionToMaterialElementMap*> NodeGroupToResolutionToMaterialElementMap;
+		// Prebuild an array relating the node index to the corresponding NodeGroup
+		TArray<FNodeGroup*> NodeGroupForNode;
+		NodeGroupForNode.AddZeroed(Model->Nodes.Num());
+		for (const auto& NodeGroupPair : Model->NodeGroups)
+		{
+			FNodeGroup* NodeGroup = NodeGroupPair.Value;
+
+			for (int32 Node : NodeGroup->Nodes)
+			{
+				check(NodeGroupForNode[Node] == nullptr);
+				NodeGroupForNode[Node] = NodeGroup;
+			}
+		}
+
+		TMap<FNodeGroupKey, FModelElement*> NodeGroupToElementMap;
+		NodeGroupToElementMap.Empty(Nodes.Num());
 
 		// Find the BSP nodes which are in this component's zone.
-		for (int32 NodeIndex = 0;NodeIndex < Nodes.Num();NodeIndex++)
+		for (int32 NodeIndex = 0; NodeIndex < Nodes.Num(); NodeIndex++)
 		{
 			FBspNode& Node = Model->Nodes[Nodes[NodeIndex]];
 			FBspSurf& Surf = Model->Surfs[Node.iSurf];
 
 			// Find the node group it belong to
-			FNodeGroup* FoundNodeGroup = NULL;
-			// find the NodeGroup that this node went into, and get all of its node
-			for (TMap<int32, FNodeGroup*>::TIterator It(Model->NodeGroups); It && (FoundNodeGroup == NULL); ++It)
-			{
-				FNodeGroup* NodeGroup = It.Value();
-				for (int32 InnerNodeIdx = 0; InnerNodeIdx < NodeGroup->Nodes.Num(); InnerNodeIdx++)
-				{
-					if (NodeGroup->Nodes[InnerNodeIdx] == Nodes[NodeIndex])
-					{
-						FoundNodeGroup = NodeGroup;
-						break;
-					}
-				}
-			}
+			FNodeGroup* FoundNodeGroup = NodeGroupForNode[Nodes[NodeIndex]];
 
-			TResolutionToMaterialElementMap* ResToMaterialElementMap = NodeGroupToResolutionToMaterialElementMap.FindRef(FoundNodeGroup);
-			if (!ResToMaterialElementMap)
+			const FNodeGroupKey NodeGroupKey(FoundNodeGroup, Surf.LightMapScale, Surf.Material);
+			FModelElement** Element = NodeGroupToElementMap.Find(NodeGroupKey);
+			if (Element == nullptr)
 			{
-				TResolutionToMaterialElementMap* TempResToMaterialElementMap = new TResolutionToMaterialElementMap();
-				NodeGroupToResolutionToMaterialElementMap.Add(FoundNodeGroup, TempResToMaterialElementMap);
-				ResToMaterialElementMap = NodeGroupToResolutionToMaterialElementMap.FindRef(FoundNodeGroup);
-				check(ResToMaterialElementMap);
-			}
-
-			// Find an element with the same material as this node.
-			uint32 ResValue = FMath::TruncToInt(Surf.LightMapScale);
-			TMap<UMaterialInterface*,FModelElement*>* MaterialElementMap = ResToMaterialElementMap->FindRef(ResValue);
-			if (!MaterialElementMap)
-			{
-				TMap<UMaterialInterface*,FModelElement*>* TempMap = new TMap<UMaterialInterface*,FModelElement*>();
-				ResToMaterialElementMap->Add(ResValue, TempMap);
-				MaterialElementMap = ResToMaterialElementMap->FindRef(ResValue);
-				check(MaterialElementMap);
-			}
-
-			FModelElement* Element = MaterialElementMap->FindRef(Surf.Material);
-			if (!Element)
-			{
-				// If there's no matching element, create a new element.
-				Element = MaterialElementMap->Add(Surf.Material,new(Elements) FModelElement(this,Surf.Material));
+				Element = &NodeGroupToElementMap.Add(NodeGroupKey, new(Elements)FModelElement(this, Surf.Material));
 			}
 			check(Element);
 
 			// Add the node to the element.
-			Element->Nodes.Add(Nodes[NodeIndex]);
+			(*Element)->Nodes.Add(Nodes[NodeIndex]);
 		}
 	}
 
diff --git a/Engine/Source/Runtime/Engine/Private/Level.cpp b/Engine/Source/Runtime/Engine/Private/Level.cpp
--- a/Engine/Source/Runtime/Engine/Private/Level.cpp
+++ b/Engine/Source/Runtime/Engine/Private/Level.cpp
@@ -785,46 +785,43 @@ void ULevel::CreateModelComponents()
 
 		if(Node.NumVertices > 0)
 		{
-			for(int32 BackFace = 0;BackFace < ((Surf.PolyFlags & PF_TwoSided) ? 2 : 1);BackFace++)
+			// Calculate the bounding box of this node.
+			FBox NodeBounds(0);
+			for(int32 VertexIndex = 0;VertexIndex < Node.NumVertices;VertexIndex++)
 			{
-				// Calculate the bounding box of this node.
-				FBox NodeBounds(0);
-				for(int32 VertexIndex = 0;VertexIndex < Node.NumVertices;VertexIndex++)
-				{
-					NodeBounds += Model->Points[Model->Verts[Node.iVertPool + VertexIndex].pVertex];
-				}
+				NodeBounds += Model->Points[Model->Verts[Node.iVertPool + VertexIndex].pVertex];
+			}
 
-				// Create a sort key for this node using the grid cell containing the center of the node's bounding box.
+			// Create a sort key for this node using the grid cell containing the center of the node's bounding box.
 #define MODEL_GRID_SIZE_XY	2048.0f
 #define MODEL_GRID_SIZE_Z	4096.0f
-				FModelComponentKey Key;
-				check( OwningWorld );
-				if (OwningWorld->GetWorldSettings()->bMinimizeBSPSections)
-				{
-					Key.X				= 0;
-					Key.Y				= 0;
-					Key.Z				= 0;
-				}
-				else
-				{
-					Key.X				= FMath::FloorToInt(NodeBounds.GetCenter().X / MODEL_GRID_SIZE_XY);
-					Key.Y				= FMath::FloorToInt(NodeBounds.GetCenter().Y / MODEL_GRID_SIZE_XY);
-					Key.Z				= FMath::FloorToInt(NodeBounds.GetCenter().Z / MODEL_GRID_SIZE_Z);
-				}
-
-				Key.MaskedPolyFlags = 0;
+			FModelComponentKey Key;
+			check( OwningWorld );
+			if (OwningWorld->GetWorldSettings()->bMinimizeBSPSections)
+			{
+				Key.X				= 0;
+				Key.Y				= 0;
+				Key.Z				= 0;
+			}
+			else
+			{
+				Key.X				= FMath::FloorToInt(NodeBounds.GetCenter().X / MODEL_GRID_SIZE_XY);
+				Key.Y				= FMath::FloorToInt(NodeBounds.GetCenter().Y / MODEL_GRID_SIZE_XY);
+				Key.Z				= FMath::FloorToInt(NodeBounds.GetCenter().Z / MODEL_GRID_SIZE_Z);
+			}
 
-				// Find an existing node list for the grid cell.
-				TArray<uint16>* ComponentNodes = ModelComponentMap.Find(Key);
-				if(!ComponentNodes)
-				{
-					// This is the first node we found in this grid cell, create a new node list for the grid cell.
-					ComponentNodes = &ModelComponentMap.Add(Key,TArray<uint16>());
-				}
+			Key.MaskedPolyFlags = 0;
 
-				// Add the node to the grid cell's node list.
-				ComponentNodes->AddUnique(NodeIndex);
+			// Find an existing node list for the grid cell.
+			TArray<uint16>* ComponentNodes = ModelComponentMap.Find(Key);
+			if(!ComponentNodes)
+			{
+				// This is the first node we found in this grid cell, create a new node list for the grid cell.
+				ComponentNodes = &ModelComponentMap.Add(Key,TArray<uint16>());
 			}
+
+			// Add the node to the grid cell's node list.
+			ComponentNodes->AddUnique(NodeIndex);
 		}
 		else
 		{
@@ -836,7 +833,7 @@ void ULevel::CreateModelComponents()
 	// Create a UModelComponent for each grid cell's node list.
 	for(TMap< FModelComponentKey, TArray<uint16> >::TConstIterator It(ModelComponentMap);It;++It)
 	{
-		const FModelComponentKey&	Key		= It.Key();
+		const FModelComponentKey&		Key		= It.Key();
 		const TArray<uint16>&			Nodes	= It.Value();	
 
 		for(int32 NodeIndex = 0;NodeIndex < Nodes.Num();NodeIndex++)
diff --git a/Engine/Source/Runtime/Engine/Private/Model.cpp b/Engine/Source/Runtime/Engine/Private/Model.cpp
--- a/Engine/Source/Runtime/Engine/Private/Model.cpp
+++ b/Engine/Source/Runtime/Engine/Private/Model.cpp
@@ -161,9 +161,25 @@ void UModel::Serialize( FArchive& Ar )
 
 	Ar << Bounds;
 
-	Vectors.BulkSerialize( Ar );
-	Points.BulkSerialize( Ar );
-	Nodes.BulkSerialize( Ar );
+	if (Ar.IsLoading() && Ar.UE4Ver() < VER_UE4_BSP_UNDO_FIX )
+	{
+		TTransArray<FVector> OldVectors(this);
+		TTransArray<FVector> OldPoints(this);
+		TTransArray<FBspNode> OldNodes(this);
+		OldVectors.BulkSerialize(Ar);
+		OldPoints.BulkSerialize(Ar);
+		OldNodes.BulkSerialize(Ar);
+
+		Vectors = OldVectors;
+		Points = OldPoints;
+		Nodes = OldNodes;
+	}
+	else
+	{
+		Vectors.BulkSerialize(Ar);
+		Points.BulkSerialize(Ar);
+		Nodes.BulkSerialize(Ar);
+	}
 	if( Ar.IsLoading() )
 	{
 		for( int32 NodeIndex=0; NodeIndex<Nodes.Num(); NodeIndex++ )
@@ -171,8 +187,23 @@ void UModel::Serialize( FArchive& Ar )
 			Nodes[NodeIndex].NodeFlags &= ~(NF_IsNew|NF_IsFront|NF_IsBack);
 		}
 	}
-	Ar << Surfs;
-	Verts.BulkSerialize( Ar );
+	
+	if (Ar.IsLoading() && Ar.UE4Ver() < VER_UE4_BSP_UNDO_FIX )
+	{
+		TTransArray<FBspSurf> OldSurfs(this);
+		TTransArray<FVert> OldVerts(this);
+
+		Ar << OldSurfs;
+		OldVerts.BulkSerialize(Ar);
+
+		Surfs = OldSurfs;
+		Verts = OldVerts;
+	}
+	else
+	{
+		Ar << Surfs;
+		Verts.BulkSerialize(Ar);
+	}
 
 	if( Ar.IsLoading() && Ar.UE4Ver() < VER_UE4_REMOVE_ZONES_FROM_MODEL)
 	{
@@ -264,16 +295,6 @@ void UModel::AddReferencedObjects(UObject* InThis, FReferenceCollector& Collecto
 	{
 		This->Surfs[ Index ].AddReferencedObjects( Collector );
 	}
-	UObject* NodesOwner = This->Nodes.GetOwner();
-	Collector.AddReferencedObject( NodesOwner, This );
-	UObject* VertsOwner = This->Verts.GetOwner();
-	Collector.AddReferencedObject( VertsOwner, This );
-	UObject* VectorsOwner = This->Vectors.GetOwner();
-	Collector.AddReferencedObject( VectorsOwner, This );
-	UObject* PointsOwner = This->Points.GetOwner();
-	Collector.AddReferencedObject( PointsOwner, This );
-	UObject* SurfsOwner = This->Surfs.GetOwner();
-	Collector.AddReferencedObject( SurfsOwner, This );
 
 	Super::AddReferencedObjects( This, Collector );
 }
@@ -356,11 +377,12 @@ void UModel::PostEditUndo()
 
 void UModel::ModifySurf( int32 InIndex, bool UpdateMaster )
 {
-	Surfs.ModifyItem( InIndex );
+	// TODO: Remove this method and replace calls to it with Modify()
+	//Surfs.ModifyItem( InIndex );
 	FBspSurf& Surf = Surfs[InIndex];
 	if( UpdateMaster && Surf.Actor )
 	{
-		Surf.Actor->Brush->Polys->Element.ModifyItem( Surf.iBrushPoly );
+		//Surf.Actor->Brush->Polys->Element.ModifyItem( Surf.iBrushPoly );
 	}
 }
 void UModel::ModifyAllSurfs( bool UpdateMaster )
@@ -533,11 +555,11 @@ void UModel::EmptyModel( int32 EmptySurfInfo, int32 EmptyPolys )
 //
 UModel::UModel(const FObjectInitializer& ObjectInitializer)
 	: UObject(ObjectInitializer)
-	, Nodes(this)
-	, Verts(this)
-	, Vectors(this)
-	, Points(this)
-	, Surfs(this)
+	, Nodes()
+	, Verts()
+	, Vectors()
+	, Points()
+	, Surfs()
 	, VertexBuffer(this)
 #if WITH_EDITOR
 	, bCachedOwnerTransformValid(false)
@@ -549,11 +571,11 @@ UModel::UModel(const FObjectInitializer& ObjectInitializer)
 #if WITH_HOT_RELOAD_CTORS
 UModel::UModel(FVTableHelper& Helper)
 	: Super(Helper)
-	, Nodes(this)
-	, Verts(this)
-	, Vectors(this)
-	, Points(this)
-	, Surfs(this)
+	, Nodes()
+	, Verts()
+	, Vectors()
+	, Points()
+	, Surfs()
 	, VertexBuffer(this)
 #if WITH_EDITOR
 	, bCachedOwnerTransformValid(false)
@@ -617,8 +639,6 @@ void UModel::Transform( ABrush* Owner )
 {
 	check(Owner);
 
-	Polys->Element.ModifyAllItems();
-
 	for( int32 i=0; i<Polys->Element.Num(); i++ )
 		Polys->Element[ i ].Transform( Owner->GetPrePivot(), Owner->GetActorLocation());
 
diff --git a/Engine/Source/Runtime/Engine/Private/ModelLight.cpp b/Engine/Source/Runtime/Engine/Private/ModelLight.cpp
--- a/Engine/Source/Runtime/Engine/Private/ModelLight.cpp
+++ b/Engine/Source/Runtime/Engine/Private/ModelLight.cpp
@@ -498,6 +498,201 @@ void UModelComponent::GetLightAndShadowMapMemoryUsage( int32& LightMapMemoryUsag
 }
 
 
+typedef	TArray<int32, TInlineAllocator<16>> FPlaneMapItem;
+
+struct FPlaneKey
+{
+	int32 X;
+	int32 Y;
+	int32 Z;
+	int32 W;
+
+	FPlaneKey(int32 InX, int32 InY, int32 InZ, int32 InW)
+		: X(InX)
+		, Y(InY)
+		, Z(InZ)
+		, W(InW)
+	{}
+
+	friend FORCEINLINE bool operator == (const FPlaneKey& A, const FPlaneKey& B)
+	{
+		return A.X == B.X && A.Y == B.Y && A.Z == B.Z && A.W == B.W;
+	}
+
+	friend FORCEINLINE uint32 GetTypeHash(const FPlaneKey& Key)
+	{
+		return HashCombine(static_cast<uint32>(Key.X), HashCombine(static_cast<uint32>(Key.Y), HashCombine(static_cast<uint32>(Key.Z), static_cast<uint32>(Key.W))));
+	}
+};
+
+class FPlaneMap
+{
+public:
+	FPlaneMap(float InGranularityXYZ, float InGranularityW, float InThreshold, int32 InitialSize = 0)
+		: OneOverGranularityXYZ(1.0f / InGranularityXYZ)
+		, OneOverGranularityW(1.0f / InGranularityW)
+		, Threshold(InThreshold)
+	{
+		Clear(InitialSize);
+	}
+
+	void Clear(int32 InitialSize = 0);
+
+	void AddPlane(const FPlane& Plane, int32 Index);
+
+	const TMap<FPlaneKey, FPlaneMapItem>& GetMap() const { return PlaneMap; }
+
+private:
+	float OneOverGranularityXYZ;
+	float OneOverGranularityW;
+	float Threshold;
+
+	TMap<FPlaneKey, FPlaneMapItem> PlaneMap;
+};
+
+
+void FPlaneMap::Clear(int32 InitialSize)
+{
+	PlaneMap.Empty(InitialSize);
+}
+
+
+// Given a grid index in one axis, a real position on the grid and a threshold radius,
+// return either:
+// - the additional grid index it can overlap in that axis, or
+// - the original grid index if there is no overlap.
+static FORCEINLINE int32 GetAdjacentIndexIfOverlapping(int32 GridIndex, float GridPos, float GridThreshold)
+{
+	if (GridPos - GridIndex < GridThreshold)
+	{
+		return GridIndex - 1;
+	}
+	else if (1.0f - (GridPos - GridIndex) < GridThreshold)
+	{
+		return GridIndex + 1;
+	}
+	else
+	{
+		return GridIndex;
+	}
+}
+
+void FPlaneMap::AddPlane(const FPlane& Plane, int32 Index)
+{
+	// Offset applied to the grid coordinates so aligned vertices (the normal case) don't overlap several grid items (taking into account the threshold)
+	const float GridOffset = 0.12345f;
+
+	const float AdjustedPlaneX = Plane.X - GridOffset;
+	const float AdjustedPlaneY = Plane.Y - GridOffset;
+	const float AdjustedPlaneZ = Plane.Z - GridOffset;
+	const float AdjustedPlaneW = Plane.W - GridOffset;
+
+	const float GridX = AdjustedPlaneX * OneOverGranularityXYZ;
+	const float GridY = AdjustedPlaneY * OneOverGranularityXYZ;
+	const float GridZ = AdjustedPlaneZ * OneOverGranularityXYZ;
+	const float GridW = AdjustedPlaneW * OneOverGranularityW;
+
+	// Get the grid indices corresponding to the plane components
+	const int32 GridIndexX = FMath::FloorToInt(GridX);
+	const int32 GridIndexY = FMath::FloorToInt(GridY);
+	const int32 GridIndexZ = FMath::FloorToInt(GridZ);
+	const int32 GridIndexW = FMath::FloorToInt(GridW);
+
+	PlaneMap.FindOrAdd(FPlaneKey(GridIndexX, GridIndexY, GridIndexZ, GridIndexW)).Add(Index);
+
+	// The grid has a maximum threshold of a certain radius. If the plane is near the edge of a grid item, it may overlap into other items.
+	// Add it to all grid items it can be seen from.
+	const float GridThresholdXYZ = Threshold * OneOverGranularityXYZ;
+	const float GridThresholdW = Threshold * OneOverGranularityW;
+	const int32 NeighbourX = GetAdjacentIndexIfOverlapping(GridIndexX, GridX, GridThresholdXYZ);
+	const int32 NeighbourY = GetAdjacentIndexIfOverlapping(GridIndexY, GridY, GridThresholdXYZ);
+	const int32 NeighbourZ = GetAdjacentIndexIfOverlapping(GridIndexZ, GridZ, GridThresholdXYZ);
+	const int32 NeighbourW = GetAdjacentIndexIfOverlapping(GridIndexW, GridW, GridThresholdW);
+
+	const bool bOverlapsInX = (NeighbourX != GridIndexX);
+	const bool bOverlapsInY = (NeighbourY != GridIndexY);
+	const bool bOverlapsInZ = (NeighbourZ != GridIndexZ);
+	const bool bOverlapsInW = (NeighbourW != GridIndexW);
+
+	if (bOverlapsInX)
+	{
+		PlaneMap.FindOrAdd(FPlaneKey(NeighbourX, GridIndexY, GridIndexZ, GridIndexW)).Add(Index);
+
+		if (bOverlapsInY)
+		{
+			PlaneMap.FindOrAdd(FPlaneKey(NeighbourX, NeighbourY, GridIndexZ, GridIndexW)).Add(Index);
+
+			if (bOverlapsInZ)
+			{
+				PlaneMap.FindOrAdd(FPlaneKey(NeighbourX, NeighbourY, NeighbourZ, GridIndexW)).Add(Index);
+
+				if (bOverlapsInW)
+				{
+					PlaneMap.FindOrAdd(FPlaneKey(NeighbourX, NeighbourY, NeighbourZ, NeighbourW)).Add(Index);
+				}
+			}
+			else if (bOverlapsInW)
+			{
+				PlaneMap.FindOrAdd(FPlaneKey(NeighbourX, NeighbourY, GridIndexZ, NeighbourW)).Add(Index);
+			}
+		}
+		else
+		{
+			if (bOverlapsInZ)
+			{
+				PlaneMap.FindOrAdd(FPlaneKey(NeighbourX, GridIndexY, NeighbourZ, GridIndexW)).Add(Index);
+
+				if (bOverlapsInW)
+				{
+					PlaneMap.FindOrAdd(FPlaneKey(NeighbourX, GridIndexY, NeighbourZ, NeighbourW)).Add(Index);
+				}
+			}
+			else if (bOverlapsInW)
+			{
+				PlaneMap.FindOrAdd(FPlaneKey(NeighbourX, GridIndexY, GridIndexZ, NeighbourW)).Add(Index);
+			}
+		}
+	}
+	else
+	{
+		if (bOverlapsInY)
+		{
+			PlaneMap.FindOrAdd(FPlaneKey(GridIndexX, NeighbourY, GridIndexZ, GridIndexW)).Add(Index);
+
+			if (bOverlapsInZ)
+			{
+				PlaneMap.FindOrAdd(FPlaneKey(GridIndexX, NeighbourY, NeighbourZ, GridIndexW)).Add(Index);
+
+				if (bOverlapsInW)
+				{
+					PlaneMap.FindOrAdd(FPlaneKey(GridIndexX, NeighbourY, NeighbourZ, NeighbourW)).Add(Index);
+				}
+			}
+			else if (bOverlapsInW)
+			{
+				PlaneMap.FindOrAdd(FPlaneKey(GridIndexX, NeighbourY, GridIndexZ, NeighbourW)).Add(Index);
+			}
+		}
+		else
+		{
+			if (bOverlapsInZ)
+			{
+				PlaneMap.FindOrAdd(FPlaneKey(GridIndexX, GridIndexY, NeighbourZ, GridIndexW)).Add(Index);
+
+				if (bOverlapsInW)
+				{
+					PlaneMap.FindOrAdd(FPlaneKey(GridIndexX, GridIndexY, NeighbourZ, NeighbourW)).Add(Index);
+				}
+			}
+			else if (bOverlapsInW)
+			{
+				PlaneMap.FindOrAdd(FPlaneKey(GridIndexX, GridIndexY, GridIndexZ, NeighbourW)).Add(Index);
+			}
+		}
+	}
+}
+
+
 
 /**
  * Groups all nodes in the model into NodeGroups (cached in the NodeGroups object)
@@ -554,172 +749,179 @@ void UModel::GroupAllNodes(const ULevel* Level, const TArray<ULightComponentBase
 		HasStaticLightingCache[ComponentIndex] = Level->ModelComponents[ComponentIndex]->HasStaticLighting();
 	}
 
-	// N^2 node loop to find nodes that are coplanar and same lightmap resolution, then 
-	// for each matching one, an N^2 vertex loop looking for shared vertices
-	for (int32 NodeIndex1 = 0; NodeIndex1 < Nodes.Num(); NodeIndex1++)
+	// Prebuild results of comparing two LightmassSettings
+	const int32 NumLightmassSettings = LightmassSettings.Num();
+	TArray<bool> LightmassSettingsEquality;
+	for (int Index1 = 0; Index1 < NumLightmassSettings; Index1++)
 	{
-		FBspNode& Node1 = Nodes[NodeIndex1];
-		if (Node1.NumVertices == 0)
+		for (int Index2 = 0; Index2 < NumLightmassSettings; Index2++)
 		{
-			continue;
+			LightmassSettingsEquality.Add(LightmassSettings[Index1] == LightmassSettings[Index2]);
 		}
-		FBspSurf& Surf1 = Surfs[Node1.iSurf];
-		UModelComponent* Comp1 = Level->ModelComponents[Node1.ComponentIndex];
+	}
 
-		// if this node's component has no shadowing, never use this node (as seen in UModelComponent::GetStaticLightingInfo)
-		if (!HasStaticLightingCache[Node1.ComponentIndex])
-		{
-			continue;
-		}
+	// We need to form groups of nodes which are nearly coplanar.
+	// First, identify nodes whose planes are similar in order to vastly reduce the search space.
+	// The FPlaneMap buckets together planes with components within a specified granular range.
+	FPlaneMap PlaneMap(1/16.0f, 50.0f, GLightmassDebugOptions.CoplanarTolerance);
 
-		FLightmassPrimitiveSettings Surf1LightmassSettings = LightmassSettings[Surf1.iLightmassIndex];
+	for (int32 NodeIndex = 0; NodeIndex < Nodes.Num(); NodeIndex++)
+	{
+		const FBspNode& Node = Nodes[NodeIndex];
 
-		// look through rest of nodes looking for "child" nodes
-		for (int32 NodeIndex2 = NodeIndex1 + 1; NodeIndex2 < Nodes.Num(); NodeIndex2++)
+		if (Node.NumVertices > 0 && HasStaticLightingCache[Node.ComponentIndex])
 		{
-			// if I've already been parented, I don't need to reparent
-			if (ParentNodes[NodeIndex1] != NULL && ParentNodes[NodeIndex2] != NULL && ParentNodes[NodeIndex1] == ParentNodes[NodeIndex2])
-			{
-				continue;
-			}
+			const FBspSurf& Surf = Surfs[Node.iSurf];
+			PlaneMap.AddPlane(Surf.Plane, NodeIndex);
+		}
+	}
 
-			FBspNode& Node2 = Nodes[NodeIndex2];
-			if (Node2.NumVertices == 0)
-			{
-				continue;
-			}
-			FBspSurf& Surf2 = Surfs[Node2.iSurf];
-			UModelComponent* Comp2 = Level->ModelComponents[Node2.ComponentIndex];
+	// Every item in the PlaneMap now contains a list of indices of nodes with similar planes.
+	// Now we can do a O(n^2) check to see if any pairs of nodes have planes within the allowed threshold, to be added to the same group.
+	for (const auto& Pair : PlaneMap.GetMap())
+	{
+		const auto& PlaneMapItem = Pair.Value;
 
-			// if this node's component has no shadowing, never use this node (as seen in UModelComponent::GetStaticLightingInfo)
-			if (!HasStaticLightingCache[Node2.ComponentIndex])
+		const int32 NumMapNodes = PlaneMapItem.Num();
+		if (NumMapNodes > 1)
+		{
+			for (int32 MapIndex1 = 0; MapIndex1 < NumMapNodes - 1; MapIndex1++)
 			{
-				continue;
-			}
-
-			FLightmassPrimitiveSettings Surf2LightmassSettings = LightmassSettings[Surf2.iLightmassIndex];
+				for (int32 MapIndex2 = MapIndex1 + 1; MapIndex2 < NumMapNodes; MapIndex2++)
+				{
+					const int32 NodeIndex1 = PlaneMapItem[MapIndex1];
+					const int32 NodeIndex2 = PlaneMapItem[MapIndex2];
+					const FBspNode& Node1 = Nodes[NodeIndex1];
+					const FBspNode& Node2 = Nodes[NodeIndex2];
+					const FBspSurf& Surf1 = Surfs[Node1.iSurf];
+					const FBspSurf& Surf2 = Surfs[Node2.iSurf];
+
+					// if I've already been parented, I don't need to reparent
+					if (ParentNodes[NodeIndex1] != nullptr && ParentNodes[NodeIndex2] != nullptr && ParentNodes[NodeIndex1] == ParentNodes[NodeIndex2])
+					{
+						continue;
+					}
 
-			// variable to see check if the 2 nodes are conodes
-			bool bNodesAreConodes = false;
+					// variable to see check if the 2 nodes are conodes
+					bool bNodesAreConodes = false;
 
-			// if we have a tolerance, then join based on coplanar adjacency
-			if (GLightmassDebugOptions.bGatherBSPSurfacesAcrossComponents)
-			{
-				// are these two nodes conodes?
-				if (Surf1.LightMapScale == Surf2.LightMapScale)
-				{
-					if (Surf1LightmassSettings == Surf2LightmassSettings)	// Do they have the same Lightmass settings?
+					// if we have a tolerance, then join based on coplanar adjacency
+					if (GLightmassDebugOptions.bGatherBSPSurfacesAcrossComponents)
 					{
-						if (/*(Node1.iSurf == Node2.iSurf) ||*/ 
-							Surf1.Plane.Equals(Surf2.Plane, GLightmassDebugOptions.CoplanarTolerance)
-							)
+						// are these two nodes conodes?
+						if (Surf1.LightMapScale == Surf2.LightMapScale)
 						{
-							// they are coplanar, have the same lightmap res and Lightmass settings, 
-							// now we need to check for adjacency which we check for by looking for a shared vertex
-
-							bool bFoundMatchingVert = false;
-							FVert* VertPool1 = &Verts[Node1.iVertPool];
-							for (int32 A=0; A < Node1.NumVertices && !bNodesAreConodes; A++)
+							if (LightmassSettingsEquality[Surf1.iLightmassIndex * NumLightmassSettings + Surf2.iLightmassIndex])
 							{
-								FVert* VertPool2 = &Verts[Node2.iVertPool];
-								for (int32 B=0; B < Node2.NumVertices && !bNodesAreConodes; B++)
+								if (Surf1.Plane.Equals(Surf2.Plane, GLightmassDebugOptions.CoplanarTolerance))
 								{
-									// if they share a vertex location, they are adjacent (this won't detect adjacency via T-joints)
-									if (VertPool1->pVertex == VertPool2->pVertex)
+									// they are coplanar, have the same lightmap res and Lightmass settings, 
+									// now we need to check for adjacency which we check for by looking for a shared vertex
+									// This is O(n^2) but since there are often only 3 or 4 verts in a poly, this will iterate on average only about 16 times.
+									// I doubt it would be any more efficient to use a TSet to check for duplicated indices in this case.
+									FVert* VertPool1 = &Verts[Node1.iVertPool];
+									for (int32 A = 0; A < Node1.NumVertices && !bNodesAreConodes; A++)
 									{
-										bNodesAreConodes = true;
+										FVert* VertPool2 = &Verts[Node2.iVertPool];
+										for (int32 B = 0; B < Node2.NumVertices && !bNodesAreConodes; B++)
+										{
+											// if they share a vertex location, they are adjacent (this won't detect adjacency via T-joints)
+											if (VertPool1->pVertex == VertPool2->pVertex)
+											{
+												bNodesAreConodes = true;
+											}
+											VertPool2++;
+										}
+										VertPool1++;
 									}
-									VertPool2++;
 								}
-								VertPool1++;
 							}
 						}
 					}
-				}
-			}
-			// if coplanar tolerance is < 0, then we join nodes together based on being in the same ModelComponent
-			// and from the same surface
-			else
-			{
-				if (Node1.iSurf == Node2.iSurf && Node1.ComponentIndex == Node2.ComponentIndex)
-				{
-					bNodesAreConodes = true;
-				}
-			}
-
-
-			// are Node1 and Node2 conodes - if so, join into a group
-			if (bNodesAreConodes)
-			{
-				// okay, these two nodes are conodes, so we need to stick them together into some pot of nodes
-				// look to see if either one are already in a group
-				FNodeGroup* NodeGroup = NULL;
-				// if both are already in different groups, we need to combine the groups
-				if (ParentNodes[NodeIndex1] != NULL && ParentNodes[NodeIndex2] != NULL)
-				{
-					NodeGroup = ParentNodes[NodeIndex1];
-
-					// merge 2 into 1
-					FNodeGroup* NodeGroup2 = ParentNodes[NodeIndex2];
-					for (int32 NodeIndex = 0; NodeIndex < NodeGroup2->Nodes.Num(); NodeIndex++)
-					{
-						NodeGroup->Nodes.Add(NodeGroup2->Nodes[NodeIndex]);
-					}
-					for (int32 LightIndex = 0; LightIndex < NodeGroup2->RelevantLights.Num(); LightIndex++)
-					{
-						NodeGroup->RelevantLights.AddUnique(NodeGroup2->RelevantLights[LightIndex]);
-					}
-
-					// replace all the users of NodeGroup2 with NodeGroup
-					for (int32 GroupIndex = 0; GroupIndex < ParentNodes.Num(); GroupIndex++)
+					// if coplanar tolerance is < 0, then we join nodes together based on being in the same ModelComponent
+					// and from the same surface
+					else
 					{
-						if (ParentNodes[GroupIndex] == NodeGroup2)
+						if (Node1.iSurf == Node2.iSurf && Node1.ComponentIndex == Node2.ComponentIndex)
 						{
-							ParentNodes[GroupIndex] = NodeGroup;
+							bNodesAreConodes = true;
 						}
 					}
 
-					// the key for the nodegroup is the 0th node (could just be a set now)
-					NodeGroups.Remove(NodeGroup2->Nodes[0]);
+					// are Node1 and Node2 conodes - if so, join into a group
+					if (bNodesAreConodes)
+					{
+						// okay, these two nodes are conodes, so we need to stick them together into some pot of nodes
+						// look to see if either one are already in a group
+						FNodeGroup* NodeGroup = NULL;
+						// if both are already in different groups, we need to combine the groups
+						if (ParentNodes[NodeIndex1] != NULL && ParentNodes[NodeIndex2] != NULL)
+						{
+							NodeGroup = ParentNodes[NodeIndex1];
 
-					// free the now useless nodegroup
-					delete NodeGroup2;
-				}
-				else if (ParentNodes[NodeIndex1] != NULL)
-				{
-					NodeGroup = ParentNodes[NodeIndex1];
-				}
-				else if (ParentNodes[NodeIndex2] != NULL)
-				{
-					NodeGroup = ParentNodes[NodeIndex2];
-				}
-				// otherwise, make a new group and put them both in it
-				else
-				{
-					NodeGroup = NodeGroups.Add(NodeIndex1, new FNodeGroup());
-				}
+							// merge 2 into 1
+							FNodeGroup* NodeGroup2 = ParentNodes[NodeIndex2];
+							for (int32 NodeIndex = 0; NodeIndex < NodeGroup2->Nodes.Num(); NodeIndex++)
+							{
+								NodeGroup->Nodes.Add(NodeGroup2->Nodes[NodeIndex]);
+							}
+							for (int32 LightIndex = 0; LightIndex < NodeGroup2->RelevantLights.Num(); LightIndex++)
+							{
+								NodeGroup->RelevantLights.AddUnique(NodeGroup2->RelevantLights[LightIndex]);
+							}
 
-				// apply both these nodes to the NodeGroup
-				for (int32 WhichNode = 0; WhichNode < 2; WhichNode++)
-				{
-					// operator on each node in this loop
-					int32 NodeIndex = WhichNode ? NodeIndex2 : NodeIndex1;
+							// replace all the users of NodeGroup2 with NodeGroup
+							for (int32 GroupIndex = 0; GroupIndex < ParentNodes.Num(); GroupIndex++)
+							{
+								if (ParentNodes[GroupIndex] == NodeGroup2)
+								{
+									ParentNodes[GroupIndex] = NodeGroup;
+								}
+							}
 
-					// track what group the node went into
-					ParentNodes[NodeIndex] = NodeGroup;
+							// the key for the nodegroup is the 0th node (could just be a set now)
+							NodeGroups.Remove(NodeGroup2->Nodes[0]);
 
-					// is this node already not yet in the group
-					if (NodeGroup->Nodes.Find(NodeIndex) == INDEX_NONE)
-					{
-						// add it to the group
-						NodeGroup->Nodes.Add(NodeIndex);
+							// free the now useless nodegroup
+							delete NodeGroup2;
+						}
+						else if (ParentNodes[NodeIndex1] != NULL)
+						{
+							NodeGroup = ParentNodes[NodeIndex1];
+						}
+						else if (ParentNodes[NodeIndex2] != NULL)
+						{
+							NodeGroup = ParentNodes[NodeIndex2];
+						}
+						// otherwise, make a new group and put them both in it
+						else
+						{
+							NodeGroup = NodeGroups.Add(NodeIndex1, new FNodeGroup());
+						}
 
-						// add the relevant lights to the nodegroup
-						TArray<ULightComponent*>* RelevantLights = ComponentRelevantLights.Find(Nodes[NodeIndex].ComponentIndex);
-						check(RelevantLights);
-						for (int32 LightIndex = 0; LightIndex < RelevantLights->Num(); LightIndex++)
+						// apply both these nodes to the NodeGroup
+						for (int32 WhichNode = 0; WhichNode < 2; WhichNode++)
 						{
-							NodeGroup->RelevantLights.AddUnique((*RelevantLights)[LightIndex]);
+							// operator on each node in this loop
+							int32 NodeIndex = WhichNode ? NodeIndex2 : NodeIndex1;
+
+							// track what group the node went into
+							ParentNodes[NodeIndex] = NodeGroup;
+
+							// is this node already not yet in the group
+							if (NodeGroup->Nodes.Find(NodeIndex) == INDEX_NONE)
+							{
+								// add it to the group
+								NodeGroup->Nodes.Add(NodeIndex);
+
+								// add the relevant lights to the nodegroup
+								TArray<ULightComponent*>* RelevantLights = ComponentRelevantLights.Find(Nodes[NodeIndex].ComponentIndex);
+								check(RelevantLights);
+								for (int32 LightIndex = 0; LightIndex < RelevantLights->Num(); LightIndex++)
+								{
+									NodeGroup->RelevantLights.AddUnique((*RelevantLights)[LightIndex]);
+								}
+							}
 						}
 					}
 				}
diff --git a/Engine/Source/Runtime/Engine/Private/Polygon.cpp b/Engine/Source/Runtime/Engine/Private/Polygon.cpp
--- a/Engine/Source/Runtime/Engine/Private/Polygon.cpp
+++ b/Engine/Source/Runtime/Engine/Private/Polygon.cpp
@@ -57,8 +57,18 @@ FArchive& operator<<(FArchive& Ar, FLightmassPrimitiveSettings& Settings)
  * Constructor, initializing all member variables.
  */
 FPoly::FPoly()
+	: Base(ForceInitToZero)
+	, Normal(ForceInitToZero)
+	, TextureU(ForceInitToZero)
+	, TextureV(ForceInitToZero)
+	, PolyFlags(PF_DefaultFlags)
+	, Actor(nullptr)
+	, Material(nullptr)
+	, iLink(INDEX_NONE)
+	, iBrushPoly(INDEX_NONE)
+	, SmoothingMask(0)
+	, LightMapScale(32.0f)
 {
-	Init();
 }
 
 
@@ -975,8 +985,6 @@ bool UPolys::Modify(bool bAlwaysMarkDirty)
 {
 	Super::Modify(bAlwaysMarkDirty);
 
-	Element.ModifyAllItems();
-
 	return !!GUndo; // we will make a broad assumption that if we have an undo buffer, we were saved in it
 }
 
@@ -984,14 +992,12 @@ void UPolys::AddReferencedObjects(UObject* InThis, FReferenceCollector& Collecto
 {
 	UPolys* This = CastChecked<UPolys>(InThis);
 	// Let GC know that we're referencing some Actor and Property objects
-	for( TTransArray<FPoly>::TConstIterator ElementIt(This->Element); ElementIt; ++ElementIt )
+	for( TArray<FPoly>::TConstIterator ElementIt(This->Element); ElementIt; ++ElementIt )
 	{
 		FPoly PolyElement = *ElementIt;
 		Collector.AddReferencedObject( PolyElement.Actor, This );
 		Collector.AddReferencedObject( PolyElement.Material, This );
 	}
-	UObject* ElementOwner = This->Element.GetOwner();
-	Collector.AddReferencedObject( ElementOwner, This );
 	Super::AddReferencedObjects(This, Collector);
 }
 
@@ -1004,23 +1010,29 @@ void UPolys::Serialize( FArchive& Ar )
 	}
 	else
 	{
-		Element.CountBytes( Ar );
-		int32 DbNum=Element.Num(), DbMax=DbNum;
-		Ar << DbNum << DbMax;
-
-		UObject* ElementOwner = Element.GetOwner();
-		Ar << ElementOwner;
+		if( Ar.IsLoading() && Ar.UE4Ver() < VER_UE4_BSP_UNDO_FIX )
+		{
+			Element.CountBytes(Ar);
+			int32 DbNum = Element.Num(), DbMax = DbNum;
+			Ar << DbNum << DbMax;
 
-		Element.SetOwner(ElementOwner);
+		
+			UObject* ElementOwner = NULL;
+			Ar << ElementOwner;
 
-		if( Ar.IsLoading() )
-		{
-			Element.Empty( DbNum );
-			Element.AddZeroed( DbNum );
+			if (Ar.IsLoading())
+			{
+				Element.Empty(DbNum);
+				Element.AddZeroed(DbNum);
+			}
+			for (int32 i = 0; i < Element.Num(); i++)
+			{
+				Ar << Element[i];
+			}
 		}
-		for( int32 i=0; i<Element.Num(); i++ )
-		{		
-			Ar << Element[i];
+		else
+		{
+			Ar << Element;
 		}
 	}
 }
diff --git a/Engine/Source/Runtime/Engine/Public/Model.h b/Engine/Source/Runtime/Engine/Public/Model.h
--- a/Engine/Source/Runtime/Engine/Public/Model.h
+++ b/Engine/Source/Runtime/Engine/Public/Model.h
@@ -357,11 +357,11 @@ class UModel : public UObject
 	UPolys*						Polys;
 #endif // WITH_EDITOR
 
-	TTransArray<FBspNode>		Nodes;
-	TTransArray<FVert>			Verts;
-	TTransArray<FVector>		Vectors;
-	TTransArray<FVector>		Points;
-	TTransArray<FBspSurf>		Surfs;
+	TArray<FBspNode>		Nodes;
+	TArray<FVert>			Verts;
+	TArray<FVector>			Vectors;
+	TArray<FVector>			Points;
+	TArray<FBspSurf>		Surfs;
 
 #if WITH_EDITOR
 	TArray<int32>				LeafHulls;
