fixed updating NavRelevantComponents attached to an actor without collision component
#ue4 UE-11716, UE-17437

diff --git a/Engine/Source/Runtime/Engine/Classes/AI/Navigation/NavModifierComponent.h b/Engine/Source/Runtime/Engine/Classes/AI/Navigation/NavModifierComponent.h
--- a/Engine/Source/Runtime/Engine/Classes/AI/Navigation/NavModifierComponent.h
+++ b/Engine/Source/Runtime/Engine/Classes/AI/Navigation/NavModifierComponent.h
@@ -16,6 +16,10 @@ class UNavModifierComponent : public UNavRelevantComponent
 	UPROPERTY(EditAnywhere, Category = Navigation)
 	TSubclassOf<UNavArea> AreaClass;
 
+	/** box extent used ONLY when owning actor doesn't have collision component */
+	UPROPERTY(EditAnywhere, Category = Navigation)
+	FVector FailsafeExtent;
+
 	virtual void CalcAndCacheBounds() const override;
 	virtual void GetNavigationData(FNavigationRelevantData& Data) const override;
 
diff --git a/Engine/Source/Runtime/Engine/Classes/AI/Navigation/NavRelevantComponent.h b/Engine/Source/Runtime/Engine/Classes/AI/Navigation/NavRelevantComponent.h
--- a/Engine/Source/Runtime/Engine/Classes/AI/Navigation/NavRelevantComponent.h
+++ b/Engine/Source/Runtime/Engine/Classes/AI/Navigation/NavRelevantComponent.h
@@ -47,4 +47,7 @@ protected:
 	/** attach navigation data to entry for owner's root component (depends on its relevancy) */
 	UPROPERTY()
 	uint32 bAttachToOwnersRoot : 1;
+
+	UPROPERTY(transient)
+	UObject* CachedNavParent;
 };
diff --git a/Engine/Source/Runtime/Engine/Classes/AI/Navigation/NavigationSystem.h b/Engine/Source/Runtime/Engine/Classes/AI/Navigation/NavigationSystem.h
--- a/Engine/Source/Runtime/Engine/Classes/AI/Navigation/NavigationSystem.h
+++ b/Engine/Source/Runtime/Engine/Classes/AI/Navigation/NavigationSystem.h
@@ -19,6 +19,7 @@ class UCrowdManager;
 struct FNavDataConfig;
 struct FPathFindingResult;
 class UActorComponent;
+class USceneComponent;
 class AActor;
 class UEnum;
 class AController;
@@ -595,6 +596,9 @@ public:
 	static void UpdateNavOctree(UActorComponent* Comp);
 	/** update all navoctree entries for actor and its components */
 	static void UpdateNavOctreeAll(AActor* Actor);
+	/** update all navoctree entries for actor and its non scene components after root movement */
+	static void UpdateNavOctreeAfterMove(USceneComponent* Comp);
+
 	/** removes all navoctree entries for actor and its components */
 	static void ClearNavOctreeAll(AActor* Actor);
 	/** updates bounds of all components implementing INavRelevantInterface */
diff --git a/Engine/Source/Runtime/Engine/Classes/Components/SceneComponent.h b/Engine/Source/Runtime/Engine/Classes/Components/SceneComponent.h
--- a/Engine/Source/Runtime/Engine/Classes/Components/SceneComponent.h
+++ b/Engine/Source/Runtime/Engine/Classes/Components/SceneComponent.h
@@ -851,9 +851,12 @@ protected:
 	 */
 	virtual void OnChildDetached(USceneComponent* ChildComponent) {}
 
-	/** Called after changing transform, tries to update navigation octree */
+	/** Called after changing transform, tries to update navigation octree for this component */
 	virtual void UpdateNavigationData();
 
+	/** Called after changing transform, tries to update navigation octree for owner */
+	void PostUpdateNavigationData();
+
 	/**
 	 * Determine if dynamic data is allowed to be changed.
 	 * 
diff --git a/Engine/Source/Runtime/Engine/Private/AI/Navigation/NavLinkProxy.cpp b/Engine/Source/Runtime/Engine/Private/AI/Navigation/NavLinkProxy.cpp
--- a/Engine/Source/Runtime/Engine/Private/AI/Navigation/NavLinkProxy.cpp
+++ b/Engine/Source/Runtime/Engine/Private/AI/Navigation/NavLinkProxy.cpp
@@ -60,11 +60,6 @@ void ANavLinkProxy::PostEditChangeProperty(FPropertyChangedEvent& PropertyChange
 	if (PropertyChangedEvent.Property && PropertyChangedEvent.Property->GetFName() == GET_MEMBER_NAME_CHECKED(ANavLinkProxy, bSmartLinkIsRelevant))
 	{
 		SmartLinkComp->SetNavigationRelevancy(bSmartLinkIsRelevant);
-		if (bSmartLinkIsRelevant)
-		{
-			SmartLinkComp->ForceNavigationRelevancy(true);
-		}
-
 		bUpdateInNavOctree = true;
 	}
 
diff --git a/Engine/Source/Runtime/Engine/Private/AI/Navigation/NavModifierComponent.cpp b/Engine/Source/Runtime/Engine/Private/AI/Navigation/NavModifierComponent.cpp
--- a/Engine/Source/Runtime/Engine/Private/AI/Navigation/NavModifierComponent.cpp
+++ b/Engine/Source/Runtime/Engine/Private/AI/Navigation/NavModifierComponent.cpp
@@ -8,6 +8,7 @@
 UNavModifierComponent::UNavModifierComponent(const FObjectInitializer& ObjectInitializer) : Super(ObjectInitializer)
 {
 	AreaClass = UNavArea_Null::StaticClass();
+	FailsafeExtent = FVector(100, 100, 100);
 }
 
 void UNavModifierComponent::CalcAndCacheBounds() const
@@ -17,7 +18,14 @@ void UNavModifierComponent::CalcAndCacheBounds() const
 	{
 		const float Radius = MyOwner->GetSimpleCollisionRadius();
 		const float HalfHeght = MyOwner->GetSimpleCollisionHalfHeight();
+		
 		ObstacleExtent = FVector(Radius, Radius, HalfHeght);
+		if (ObstacleExtent.IsNearlyZero())
+		{
+			const FVector ActorScale3D = MyOwner->GetActorScale3D();
+			ObstacleExtent = FailsafeExtent * ActorScale3D;
+		}
+
 		Bounds = FBox::BuildAABB(MyOwner->GetActorLocation(), ObstacleExtent);
 	}
 }
diff --git a/Engine/Source/Runtime/Engine/Private/AI/Navigation/NavRelevantComponent.cpp b/Engine/Source/Runtime/Engine/Private/AI/Navigation/NavRelevantComponent.cpp
--- a/Engine/Source/Runtime/Engine/Private/AI/Navigation/NavRelevantComponent.cpp
+++ b/Engine/Source/Runtime/Engine/Private/AI/Navigation/NavRelevantComponent.cpp
@@ -13,7 +13,40 @@ void UNavRelevantComponent::OnRegister()
 {
 	Super::OnRegister();
 
-	CalcAndCacheBounds();
+	if (bAttachToOwnersRoot)
+	{
+		bool bUpdateCachedParent = true;
+#if WITH_EDITOR
+		UNavigationSystem* NavSys = UNavigationSystem::GetCurrent(GetWorld());
+		if (NavSys && NavSys->IsNavigationRegisterLocked())
+		{
+			bUpdateCachedParent = false;
+		}
+#endif
+
+		AActor* OwnerActor = GetOwner();
+		if (OwnerActor && bUpdateCachedParent)
+		{
+			// attach to root component if it's relevant for navigation
+			UActorComponent* ActorComp = OwnerActor->GetRootComponent();			
+			INavRelevantInterface* NavInterface = ActorComp ? Cast<INavRelevantInterface>(ActorComp) : nullptr;
+			if (NavInterface && NavInterface->IsNavigationRelevant() &&
+				OwnerActor->IsComponentRelevantForNavigation(ActorComp))
+			{
+				CachedNavParent = ActorComp;
+			}
+
+			// otherwise try actor itself under the same condition
+			if (CachedNavParent == nullptr)
+			{
+				NavInterface = Cast<INavRelevantInterface>(OwnerActor);
+				if (NavInterface && NavInterface->IsNavigationRelevant())
+				{
+					CachedNavParent = OwnerActor;
+				}
+			}
+		}
+	}
 
 	UNavigationSystem::OnComponentRegistered(this);
 }
@@ -42,13 +75,7 @@ void UNavRelevantComponent::UpdateNavigationBounds()
 
 UObject* UNavRelevantComponent::GetNavigationParent() const
 {
-	if (bAttachToOwnersRoot)
-	{
-		AActor* OwnerActor = GetOwner();
-		return OwnerActor && OwnerActor->GetRootComponent() ? static_cast<UObject*>(OwnerActor->GetRootComponent()) : static_cast<UObject*>(OwnerActor);
-	}
-
-	return nullptr;
+	return CachedNavParent;
 }
 
 void UNavRelevantComponent::CalcBounds()
diff --git a/Engine/Source/Runtime/Engine/Private/AI/Navigation/NavigationSystem.cpp b/Engine/Source/Runtime/Engine/Private/AI/Navigation/NavigationSystem.cpp
--- a/Engine/Source/Runtime/Engine/Private/AI/Navigation/NavigationSystem.cpp
+++ b/Engine/Source/Runtime/Engine/Private/AI/Navigation/NavigationSystem.cpp
@@ -2562,6 +2562,34 @@ void UNavigationSystem::UpdateNavOctreeAll(AActor* Actor)
 	}
 }
 
+void UNavigationSystem::UpdateNavOctreeAfterMove(USceneComponent* Comp)
+{
+	AActor* OwnerActor = Comp->GetOwner();
+	if (OwnerActor && OwnerActor->GetRootComponent() == Comp)
+	{
+		UpdateNavOctree(OwnerActor);
+
+		TInlineComponentArray<UActorComponent*> Components;
+		OwnerActor->GetComponents(Components);
+
+		for (int32 ComponentIndex = 0; ComponentIndex < Components.Num(); ComponentIndex++)
+		{
+			if (Components[ComponentIndex] && !Components[ComponentIndex]->IsA(USceneComponent::StaticClass()))
+			{
+				UpdateNavOctree(Components[ComponentIndex]);
+			}
+		}
+
+		for (int32 RootChildIndex = 0; RootChildIndex < Comp->AttachChildren.Num(); RootChildIndex++)
+		{
+			if (Comp->AttachChildren[RootChildIndex] && Comp->AttachChildren[RootChildIndex]->GetOuter() != OwnerActor)
+			{
+				UpdateNavOctreeAll(Cast<AActor>(Comp->AttachChildren[RootChildIndex]->GetOuter()));
+			}
+		}
+	}
+}
+
 void UNavigationSystem::UpdateNavOctreeBounds(AActor* Actor)
 {
 	TInlineComponentArray<UActorComponent*> Components;
diff --git a/Engine/Source/Runtime/Engine/Private/Components/SceneComponent.cpp b/Engine/Source/Runtime/Engine/Private/Components/SceneComponent.cpp
--- a/Engine/Source/Runtime/Engine/Private/Components/SceneComponent.cpp
+++ b/Engine/Source/Runtime/Engine/Private/Components/SceneComponent.cpp
@@ -1848,6 +1848,8 @@ bool USceneComponent::InternalSetWorldLocationAndRotation(FVector NewLocation, c
 		RelativeLocation = NewLocation;
 		RelativeRotation = RelativeRotationCache.QuatToRotator(NewRotationQuat); // Normalizes rotator, if this is a new rotation. Then we'll use it below.
 		UpdateComponentToWorldWithParent(AttachParent, bNoPhysics, RelativeRotationCache.RotatorToQuat(RelativeRotation), Teleport);
+
+		PostUpdateNavigationData();
 		return true;
 	}
 
@@ -2492,6 +2494,17 @@ void USceneComponent::UpdateNavigationData()
 	}
 }
 
+void USceneComponent::PostUpdateNavigationData()
+{
+	if (UNavigationSystem::ShouldUpdateNavOctreeOnComponentChange() &&
+		IsRegistered() && World && World->IsGameWorld() &&
+		World->GetNetMode() < ENetMode::NM_Client)
+	{
+		UNavigationSystem::UpdateNavOctreeAfterMove(this);
+	}
+}
+
+
 // K2 versions of various transform changing operations.
 // Note: we pass null for the hit result if not sweeping, for better perf.
 // This assumes this K2 function is only used by blueprints, which initializes the param for each function call.
diff --git a/Engine/Source/Runtime/Engine/Private/Components/ShapeComponent.cpp b/Engine/Source/Runtime/Engine/Private/Components/ShapeComponent.cpp
--- a/Engine/Source/Runtime/Engine/Private/Components/ShapeComponent.cpp
+++ b/Engine/Source/Runtime/Engine/Private/Components/ShapeComponent.cpp
@@ -82,5 +82,8 @@ void UShapeComponent::GetNavigationData(FNavigationRelevantData& Data) const
 
 bool UShapeComponent::IsNavigationRelevant() const
 {
-	return bDynamicObstacle || Super::IsNavigationRelevant();
+	// failed CanEverAffectNavigation() always takes priority
+	// dynamic obstacle overrides collision check
+
+	return (bDynamicObstacle && CanEverAffectNavigation()) || Super::IsNavigationRelevant();
 }
\ No newline at end of file
