Fixed possible crash in visual logger.

diff --git a/Engine/Source/Runtime/Engine/Private/VisualLogger/VisualLogger.cpp b/Engine/Source/Runtime/Engine/Private/VisualLogger/VisualLogger.cpp
--- a/Engine/Source/Runtime/Engine/Private/VisualLogger/VisualLogger.cpp
+++ b/Engine/Source/Runtime/Engine/Private/VisualLogger/VisualLogger.cpp
@@ -41,7 +41,7 @@ bool FVisualLogger::CheckVisualLogInputInternal(const class UObject* Object, con
 	}
 
 	*CurrentEntry = VisualLogger.GetEntryToWrite(Object, (*World)->TimeSeconds);
-	if (ensure(CurrentEntry != nullptr) == false)
+	if (*CurrentEntry == nullptr)
 	{
 		return false;
 	}
