Improved shader compile error message when compiling for multiple platforms

diff --git a/Engine/Source/Runtime/Engine/Private/ShaderCompiler/ShaderCompiler.cpp b/Engine/Source/Runtime/Engine/Private/ShaderCompiler/ShaderCompiler.cpp
--- a/Engine/Source/Runtime/Engine/Private/ShaderCompiler/ShaderCompiler.cpp
+++ b/Engine/Source/Runtime/Engine/Private/ShaderCompiler/ShaderCompiler.cpp
@@ -1651,6 +1651,7 @@ bool FShaderCompilingManager::HandlePotentialRetryOnError(TMap<int32, FShaderMap
 				const TArray<FShaderCompileJob*>& CompleteJobs = Results.FinishedJobs;
 				TArray<const FShaderCompileJob*> ErrorJobs;
 				TArray<FString> UniqueErrors;
+				TArray<EShaderPlatform> ErrorPlatforms;
 
 				// Gather unique errors
 				for (int32 JobIndex = 0; JobIndex < CompleteJobs.Num(); JobIndex++)
@@ -1659,6 +1660,8 @@ bool FShaderCompilingManager::HandlePotentialRetryOnError(TMap<int32, FShaderMap
 
 					if (!CurrentJob.bSucceeded)
 					{
+						ErrorPlatforms.AddUnique((EShaderPlatform)CurrentJob.Input.Target.Platform);
+						
 						for (int32 ErrorIndex = 0; ErrorIndex < CurrentJob.Output.Errors.Num(); ErrorIndex++)
 						{
 							const FShaderCompilerError& CurrentError = CurrentJob.Output.Errors[ErrorIndex];
@@ -1673,10 +1676,22 @@ bool FShaderCompilingManager::HandlePotentialRetryOnError(TMap<int32, FShaderMap
 					}
 				}
 
-				// Assuming all the jobs are for the same platform
-				const EShaderPlatform TargetShaderPlatform = (EShaderPlatform)CompleteJobs[0]->Input.Target.Platform;
+				FString TargetShaderPlatformString;
+
+				for (int32 PlatformIndex = 0; PlatformIndex < ErrorPlatforms.Num(); PlatformIndex++)
+				{
+					if (TargetShaderPlatformString.IsEmpty())
+					{
+						TargetShaderPlatformString = LegacyShaderPlatformToShaderFormat(ErrorPlatforms[PlatformIndex]).ToString();
+					}
+					else
+					{
+						TargetShaderPlatformString += FString(TEXT(", ")) + LegacyShaderPlatformToShaderFormat(ErrorPlatforms[PlatformIndex]).ToString();
+					}
+				}
+
 				const TCHAR* MaterialName = ShaderMap ? *ShaderMap->GetFriendlyName() : TEXT("global shaders");
-				FString ErrorString = FString::Printf(TEXT("%i Shader compiler errors compiling %s for platform %s:"), UniqueErrors.Num(), MaterialName, *LegacyShaderPlatformToShaderFormat(TargetShaderPlatform).ToString());
+				FString ErrorString = FString::Printf(TEXT("%i Shader compiler errors compiling %s for platform %s:"), UniqueErrors.Num(), MaterialName, *TargetShaderPlatformString);
 				UE_LOG(LogShaderCompilers, Warning, TEXT("%s"), *ErrorString);
 				ErrorString += TEXT("\n");
 
