Cleaned up material properties a bit

diff --git a/Engine/Source/Runtime/Engine/Classes/Materials/Material.h b/Engine/Source/Runtime/Engine/Classes/Materials/Material.h
--- a/Engine/Source/Runtime/Engine/Classes/Materials/Material.h
+++ b/Engine/Source/Runtime/Engine/Classes/Materials/Material.h
@@ -331,7 +331,7 @@ class UMaterial : public UMaterialInterface
 	 * Defines how the material reacts on DBuffer decals (Affects look, performance and texture/sample usage).
 	 * Non DBuffer Decals can be disabled on the primitive (e.g. static mesh)
 	 */
-	UPROPERTY(EditAnywhere, BlueprintReadOnly, Category=Material, meta=(DisplayName = "Decal Response (DBuffer)"), AssetRegistrySearchable)
+	UPROPERTY(EditAnywhere, BlueprintReadOnly, Category=Material, AdvancedDisplay, meta=(DisplayName = "Decal Response (DBuffer)"), AssetRegistrySearchable)
 	TEnumAsByte<enum EMaterialDecalResponse> MaterialDecalResponse;
 
 private:
@@ -411,8 +411,8 @@ public:
 	UPROPERTY(EditAnywhere, Category=Material)
 	uint32 TwoSided:1;
 
-	/** Indicates that the material should be rendered without backface culling and the normal should be flipped for backfaces. */
-	UPROPERTY(EditAnywhere, Category=Material)
+	/** Whether the material should support a dithered LOD transition when used with the foliage system. */
+	UPROPERTY(EditAnywhere, Category=Material, AdvancedDisplay)
 	uint32 DitheredLODTransition:1;
 
 	/** Number of customized UV inputs to display.  Unconnected customized UV inputs will just pass through the vertex UVs. */
diff --git a/Engine/Source/Runtime/Engine/Private/Materials/Material.cpp b/Engine/Source/Runtime/Engine/Private/Materials/Material.cpp
--- a/Engine/Source/Runtime/Engine/Private/Materials/Material.cpp
+++ b/Engine/Source/Runtime/Engine/Private/Materials/Material.cpp
@@ -2450,6 +2450,11 @@ bool UMaterial::CanEditChange(const UProperty* InProperty) const
 			return BlendMode == BLEND_Masked;
 		}
 
+		if (PropertyName == GET_MEMBER_NAME_STRING_CHECKED(UMaterial, DecalBlendMode))
+		{
+			return MaterialDomain == MD_DeferredDecal;
+		}
+
 		if (PropertyName == GET_MEMBER_NAME_STRING_CHECKED(UMaterial, MaterialDecalResponse))
 		{
 			static auto* CVar = IConsoleManager::Get().FindTConsoleVariableDataInt(TEXT("r.DBuffer"));
