Add support for parallel blend physics.

#UE-17721

diff --git a/Engine/Source/Runtime/Engine/Classes/Components/SkeletalMeshComponent.h b/Engine/Source/Runtime/Engine/Classes/Components/SkeletalMeshComponent.h
--- a/Engine/Source/Runtime/Engine/Classes/Components/SkeletalMeshComponent.h
+++ b/Engine/Source/Runtime/Engine/Classes/Components/SkeletalMeshComponent.h
@@ -932,7 +932,11 @@ public:
 	 *
 	 * @param	RequiredBones	List of bones to be blend
 	 */
-	void BlendPhysicsBones( TArray<FBoneIndexType>& RequiredBones );
+	void BlendPhysicsBones( TArray<FBoneIndexType>& RequiredBones )
+	{
+		PerformBlendPhysicsBones(RequiredBones, AnimEvaluationContext.LocalAtoms);
+	}
+
 
 	/** Take the results of the physics and blend them with the animation state (based on the PhysicsWeight parameter), and update the SpaceBases array. */
 	void BlendInPhysics();	
@@ -1217,6 +1221,9 @@ private:
 	// Reference to our current parallel animation evaluation task (if there is one)
 	FGraphEventRef				ParallelAnimationEvaluationTask;
 
+	// Reference to our current blend physics task (if there is one)
+	FGraphEventRef				ParallelBlendPhysicsCompletionTask;
+
 public:
 	// Parallel evaluation wrappers
 	void ParallelAnimationEvaluation() { PerformAnimationEvaluation(AnimEvaluationContext.SkeletalMesh, AnimEvaluationContext.AnimInstance, AnimEvaluationContext.SpaceBases, AnimEvaluationContext.LocalAtoms, AnimEvaluationContext.VertexAnims, AnimEvaluationContext.RootBoneTranslation); }
@@ -1246,6 +1253,18 @@ private:
 
 	// Handles registering/unregistering the pre cloth tick as it is needed
 	void UpdatePreClothTickRegisteredState();
+	friend class FParallelBlendPhysicsTask;
+	
+	//wrapper for parallel blend physics
+	void ParallelBlendPhysics() { PerformBlendPhysicsBones(RequiredBones, AnimEvaluationContext.LocalAtoms); }
+
+	void PerformBlendPhysicsBones(const TArray<FBoneIndexType>& InRequiredBones, TArray<FTransform>& InLocalAtoms);
+
+	friend class FParallelBlendPhysicsCompletionTask;
+	void CompleteParallelBlendPhysics();
+
+	friend class FTickClothingTask;
+	void PerformTickClothing(float DeltaTime);
 
 	// these are deprecated variables from removing SingleAnimSkeletalComponent
 	// remove if this version goes away : VER_UE4_REMOVE_SINGLENODEINSTANCE
diff --git a/Engine/Source/Runtime/Engine/Private/PhysicsEngine/PhysAnim.cpp b/Engine/Source/Runtime/Engine/Private/PhysicsEngine/PhysAnim.cpp
--- a/Engine/Source/Runtime/Engine/Private/PhysicsEngine/PhysAnim.cpp
+++ b/Engine/Source/Runtime/Engine/Private/PhysicsEngine/PhysAnim.cpp
@@ -27,6 +27,73 @@ struct FAssetWorldBoneTM
 };
 
 
+
+class FParallelBlendPhysicsTask
+{
+	TWeakObjectPtr<USkeletalMeshComponent> SkeletalMeshComponent;
+
+public:
+	FParallelBlendPhysicsTask(TWeakObjectPtr<USkeletalMeshComponent> InSkeletalMeshComponent)
+		: SkeletalMeshComponent(InSkeletalMeshComponent)
+	{
+	}
+
+	FORCEINLINE TStatId GetStatId() const
+	{
+		RETURN_QUICK_DECLARE_CYCLE_STAT(FParallelBlendPhysicsTask, STATGROUP_TaskGraphTasks);
+	}
+	static ENamedThreads::Type GetDesiredThread()
+	{
+		return ENamedThreads::AnyThread;
+	}
+	static ESubsequentsMode::Type GetSubsequentsMode()
+	{
+		return ESubsequentsMode::TrackSubsequents;
+	}
+
+	void DoTask(ENamedThreads::Type CurrentThread, const FGraphEventRef& MyCompletionGraphEvent)
+	{
+		if (USkeletalMeshComponent* Comp = SkeletalMeshComponent.Get())
+		{
+			Comp->ParallelBlendPhysics();
+		}
+	}
+};
+
+class FParallelBlendPhysicsCompletionTask
+{
+	TWeakObjectPtr<USkeletalMeshComponent> SkeletalMeshComponent;
+
+public:
+	FParallelBlendPhysicsCompletionTask(TWeakObjectPtr<USkeletalMeshComponent> InSkeletalMeshComponent)
+		: SkeletalMeshComponent(InSkeletalMeshComponent)
+	{
+	}
+
+	FORCEINLINE TStatId GetStatId() const
+	{
+		RETURN_QUICK_DECLARE_CYCLE_STAT(FParallelBlendPhysicsCompletionTask, STATGROUP_TaskGraphTasks);
+	}
+	static ENamedThreads::Type GetDesiredThread()
+	{
+		return ENamedThreads::GameThread;
+	}
+	static ESubsequentsMode::Type GetSubsequentsMode()
+	{
+		return ESubsequentsMode::TrackSubsequents;
+	}
+
+	void DoTask(ENamedThreads::Type CurrentThread, const FGraphEventRef& MyCompletionGraphEvent)
+	{
+		SCOPE_CYCLE_COUNTER(STAT_AnimGameThreadTime);
+		if (USkeletalMeshComponent* Comp = SkeletalMeshComponent.Get())
+		{
+			Comp->CompleteParallelBlendPhysics();
+		}
+	}
+};
+
+
 // Use current pose to calculate world-space position of this bone without physics now.
 void UpdateWorldBoneTM(TArray<FAssetWorldBoneTM> & WorldBoneTMs, int32 BoneIndex, USkeletalMeshComponent* SkelComp, const FTransform &LocalToWorldTM, const FVector& Scale3D)
 {
@@ -58,7 +125,7 @@ void UpdateWorldBoneTM(TArray<FAssetWorldBoneTM> & WorldBoneTMs, int32 BoneIndex
 }
 
 
-void USkeletalMeshComponent::BlendPhysicsBones( TArray<FBoneIndexType>& InRequiredBones)
+void USkeletalMeshComponent::PerformBlendPhysicsBones(const TArray<FBoneIndexType>& InRequiredBones, TArray<FTransform>& InLocalAtoms)
 {
 	// Get drawscale from Owner (if there is one)
 	FVector TotalScale3D = ComponentToWorld.GetScale3D();
@@ -94,17 +161,24 @@ void USkeletalMeshComponent::BlendPhysicsBones( TArray<FBoneIndexType>& InRequir
 
 	TArray<FTransform>& EditableSpaceBases = GetEditableSpaceBases();
 
+	struct FBodyTMPair
+	{
+		FBodyInstance* BI;
+		FTransform TM;
+	};
+
+	TArray<FBodyTMPair> PendingBodyTMs;
+
 #if WITH_PHYSX
-	//TODO: we need to use a write lock because later down we might update body instance. This should be going away, at which point we can switch to a read lock (physx supports multiple readers)
 	// Lock the scenes we need (flags set in InitArticulated)
 	if (bHasBodiesInSyncScene)
 	{
-		SCENE_LOCK_WRITE(PhysScene->GetPhysXScene(PST_Sync))
+		SCENE_LOCK_READ(PhysScene->GetPhysXScene(PST_Sync))
 	}
 
 	if (bHasBodiesInAsyncScene)
 	{
-		SCENE_LOCK_WRITE(PhysScene->GetPhysXScene(PST_Async))
+		SCENE_LOCK_READ(PhysScene->GetPhysXScene(PST_Async))
 	}
 #endif
 
@@ -173,15 +247,15 @@ void USkeletalMeshComponent::BlendPhysicsBones( TArray<FBoneIndexType>& InRequir
 					RelTM.RemoveScaling();
 					FQuat RelRot(RelTM.GetRotation());
 					FVector RelPos =  RecipScale3D * RelTM.GetLocation();
-					FTransform PhysAtom = FTransform(RelRot, RelPos, LocalAtoms[BoneIndex].GetScale3D());
+					FTransform PhysAtom = FTransform(RelRot, RelPos, InLocalAtoms[BoneIndex].GetScale3D());
 
 					// Now blend in this atom. See if we are forcing this bone to always be blended in
-					LocalAtoms[BoneIndex].Blend( LocalAtoms[BoneIndex], PhysAtom, UsePhysWeight );
+					InLocalAtoms[BoneIndex].Blend( InLocalAtoms[BoneIndex], PhysAtom, UsePhysWeight );
 
 					if(BoneIndex == 0)
 					{
 						//We must update RecipScale3D based on the atom scale of the root
-						TotalScale3D *= LocalAtoms[0].GetScale3D();
+						TotalScale3D *= InLocalAtoms[0].GetScale3D();
 						RecipScale3D = TotalScale3D.Reciprocal();
 					}
 
@@ -196,12 +270,12 @@ void USkeletalMeshComponent::BlendPhysicsBones( TArray<FBoneIndexType>& InRequir
 		// Update SpaceBases entry for this bone now
 		if( BoneIndex == 0 )
 		{
-			EditableSpaceBases[0] = LocalAtoms[0];
+			EditableSpaceBases[0] = InLocalAtoms[0];
 		}
 		else
 		{
 			const int32 ParentIndex	= SkeletalMesh->RefSkeleton.GetParentIndex(BoneIndex);
-			EditableSpaceBases[BoneIndex] = LocalAtoms[BoneIndex] * EditableSpaceBases[ParentIndex];
+			EditableSpaceBases[BoneIndex] = InLocalAtoms[BoneIndex] * EditableSpaceBases[ParentIndex];
 
 			/**
 			* Normalize rotations.
@@ -215,7 +289,11 @@ void USkeletalMeshComponent::BlendPhysicsBones( TArray<FBoneIndexType>& InRequir
 
 		if (bUpdatePhysics && BodyInstance)
 		{
-			BodyInstance->SetBodyTransform(EditableSpaceBases[BoneIndex] * ComponentToWorld, ETeleportType::TeleportPhysics);
+			//This is extremely inefficient. We need to obtain a write lock which will block other threads from blending
+			//For now I'm juts deferring it to the end of this loop, but in general we need to move it all out of here and do it when the blend task is done
+			FBodyTMPair* BodyTMPair = new (PendingBodyTMs) FBodyTMPair;
+			BodyTMPair->BI = BodyInstance;
+			BodyTMPair->TM = EditableSpaceBases[BoneIndex] * ComponentToWorld;
 		}
 	}
 
@@ -224,14 +302,46 @@ void USkeletalMeshComponent::BlendPhysicsBones( TArray<FBoneIndexType>& InRequir
 	// Unlock the scenes 
 	if (bHasBodiesInSyncScene)
 	{
-		SCENE_UNLOCK_WRITE(PhysScene->GetPhysXScene(PST_Sync))
+		SCENE_UNLOCK_READ(PhysScene->GetPhysXScene(PST_Sync))
 	}
 
 	if (bHasBodiesInAsyncScene)
 	{
-		SCENE_UNLOCK_WRITE(PhysScene->GetPhysXScene(PST_Async))
+		SCENE_UNLOCK_READ(PhysScene->GetPhysXScene(PST_Async))
 	}
+
+	if(PendingBodyTMs.Num())
+	{
+		//This is extremely inefficient. We need to obtain a write lock which will block other threads from blending
+		//For now I'm juts deferring it to the end of this loop, but in general we need to move it all out of here and do it when the blend task is done
+
+		if (bHasBodiesInSyncScene)
+		{
+			SCENE_LOCK_WRITE(PhysScene->GetPhysXScene(PST_Sync))
+		}
+
+		if (bHasBodiesInAsyncScene)
+		{
+			SCENE_LOCK_WRITE(PhysScene->GetPhysXScene(PST_Async))
+		}
+
+		for (const FBodyTMPair& BodyTMPair : PendingBodyTMs)
+		{
+			BodyTMPair.BI->SetBodyTransform(BodyTMPair.TM, ETeleportType::TeleportPhysics);
+		}
+
+		if (bHasBodiesInSyncScene)
+		{
+			SCENE_UNLOCK_WRITE(PhysScene->GetPhysXScene(PST_Sync))
+		}
+
+		if (bHasBodiesInAsyncScene)
+		{
+			SCENE_UNLOCK_WRITE(PhysScene->GetPhysXScene(PST_Async))
+		}
+    }
 #endif
+	
 
 	// Transforms updated, cached local bounds are now out of date.
 	InvalidateCachedBounds();
@@ -257,9 +367,12 @@ bool USkeletalMeshComponent::DoAnyPhysicsBodiesHaveWeight() const
 	return false;
 }
 
+TAutoConsoleVariable<int32> CVarUseParallelBlendPhysics(TEXT("a.ParallelBlendPhysics"), 1, TEXT("If 1, physics blending will be run across the task graph system. If 0, blending will run purely on the game thread"));
+
 void USkeletalMeshComponent::BlendInPhysics()
 {
 	SCOPE_CYCLE_COUNTER(STAT_BlendInPhysics);
+	check(IsInGameThread());
 
 	// Can't do anything without a SkeletalMesh
 	if( !SkeletalMesh )
@@ -271,7 +384,41 @@ void USkeletalMeshComponent::BlendInPhysics()
 	// If we don't have or want any physics, we do nothing.
 	if( Bodies.Num() > 0 )
 	{
-		BlendPhysicsBones( RequiredBones );
+		IsRunningParallelEvaluation(/*bBlockOnTask = */ true, /*bPerformPostAnimEvaluation =*/ true);
+		// start parallel work
+		check(!IsValidRef(ParallelAnimationEvaluationTask));
+
+		const bool bParallelBlend = !!CVarUseParallelBlendPhysics.GetValueOnGameThread() && FApp::ShouldUseThreadingForPerformance();
+		if(bParallelBlend)
+		{
+			if (SkeletalMesh->RefSkeleton.GetNum() != AnimEvaluationContext.LocalAtoms.Num())
+			{
+				// Initialize Parallel Task arrays
+				AnimEvaluationContext.SpaceBases = GetSpaceBases();
+				AnimEvaluationContext.VertexAnims = ActiveVertexAnims;
+			}
+
+			AnimEvaluationContext.LocalAtoms = LocalAtoms;
+
+			ParallelAnimationEvaluationTask = TGraphTask<FParallelBlendPhysicsTask>::CreateTask().ConstructAndDispatchWhenReady(this);
+
+			// set up a task to run on the game thread to accept the results
+			FGraphEventArray Prerequistes;
+			Prerequistes.Add(ParallelAnimationEvaluationTask);
+
+			check(!IsValidRef(ParallelBlendPhysicsCompletionTask));
+			ParallelBlendPhysicsCompletionTask = TGraphTask<FParallelBlendPhysicsCompletionTask>::CreateTask(&Prerequistes).ConstructAndDispatchWhenReady(this);
+		}else
+		{
+			PerformBlendPhysicsBones(RequiredBones, LocalAtoms);
+			CompleteParallelBlendPhysics();
+		}
+	}
+}
+
+void USkeletalMeshComponent::CompleteParallelBlendPhysics()
+{
+	Exchange(AnimEvaluationContext.LocalAtoms, AnimEvaluationContext.bDoInterpolation ? CachedLocalAtoms : LocalAtoms);
 		
 		// Update Child Transform - The above function changes bone transform, so will need to update child transform
 		UpdateChildTransforms();
@@ -281,7 +428,11 @@ void USkeletalMeshComponent::BlendInPhysics()
 
 		// New bone positions need to be sent to render thread
 		MarkRenderDynamicDataDirty();
-	}
+
+	FlipEditableSpaceBases();
+
+	ParallelAnimationEvaluationTask.SafeRelease();
+	ParallelBlendPhysicsCompletionTask.SafeRelease();
 }
 
 
diff --git a/Engine/Source/Runtime/Engine/Private/SkeletalMeshComponentPhysics.cpp b/Engine/Source/Runtime/Engine/Private/SkeletalMeshComponentPhysics.cpp
--- a/Engine/Source/Runtime/Engine/Private/SkeletalMeshComponentPhysics.cpp
+++ b/Engine/Source/Runtime/Engine/Private/SkeletalMeshComponentPhysics.cpp
@@ -3145,11 +3145,9 @@ void USkeletalMeshComponent::PreClothTick(float DeltaTime)
 		{
 			BlendInPhysics();
 		}
-		// If we aren't blending we will have already flipped this
-		FlipEditableSpaceBases();
 	}
 
-	//TODO: move this into pre physics tick
+
 
 #if WITH_APEX_CLOTHING
 	// if skeletal mesh has clothing assets, call TickClothing
@@ -3662,7 +3660,42 @@ bool USkeletalMeshComponent::GetClothSimulatedPosition(int32 AssetIndex, int32 V
 
 #endif// #if WITH_APEX_CLOTHING
 
-void USkeletalMeshComponent::TickClothing(float DeltaTime)
+class FTickClothingTask
+{
+	TWeakObjectPtr<USkeletalMeshComponent> SkeletalMeshComponent;
+	float DeltaTime;
+
+public:
+	FTickClothingTask(TWeakObjectPtr<USkeletalMeshComponent> InSkeletalMeshComponent, float InDeltaTime)
+		: SkeletalMeshComponent(InSkeletalMeshComponent)
+		, DeltaTime(InDeltaTime)
+	{
+	}
+
+	FORCEINLINE TStatId GetStatId() const
+	{
+		RETURN_QUICK_DECLARE_CYCLE_STAT(FTickClothingTask, STATGROUP_TaskGraphTasks);
+	}
+	static ENamedThreads::Type GetDesiredThread()
+	{
+		return ENamedThreads::GameThread;
+	}
+	static ESubsequentsMode::Type GetSubsequentsMode()
+	{
+		return ESubsequentsMode::TrackSubsequents;
+	}
+
+	void DoTask(ENamedThreads::Type CurrentThread, const FGraphEventRef& MyCompletionGraphEvent)
+	{
+		//SCOPE_CYCLE_COUNTER(STAT_AnimGameThreadTime);
+		if (USkeletalMeshComponent* Comp = SkeletalMeshComponent.Get())
+		{
+			Comp->PerformTickClothing(DeltaTime);
+		}
+	}
+};
+
+void USkeletalMeshComponent::PerformTickClothing(float DeltaTime)
 {
 #if WITH_APEX_CLOTHING
 	// animated but bone transforms were not updated because it was not rendered
@@ -3681,6 +3714,19 @@ void USkeletalMeshComponent::TickClothing(float DeltaTime)
 #endif// #if WITH_APEX_CLOTHING
 }
 
+void USkeletalMeshComponent::TickClothing(float DeltaTime)
+{
+	if(IsValidRef(ParallelBlendPhysicsCompletionTask))
+	{
+		FGraphEventArray Prerequistes;
+		Prerequistes.Add(ParallelBlendPhysicsCompletionTask);
+		TGraphTask<FTickClothingTask>::CreateTask(&Prerequistes).ConstructAndDispatchWhenReady(this, DeltaTime);
+	}else
+	{
+		PerformTickClothing(DeltaTime);
+	}
+}
+
 void USkeletalMeshComponent::GetUpdateClothSimulationData(TArray<FClothSimulData>& OutClothSimData, USkeletalMeshComponent* OverrideLocalRootComponent)
 {
 #if WITH_APEX_CLOTHING
