Mesh particles now properly populate dynamic parameter names in Cascade

https://jira.ol.epicgames.net/browse/UE-17104

diff --git a/Engine/Source/Runtime/Engine/Classes/Particles/Parameter/ParticleModuleParameterDynamic.h b/Engine/Source/Runtime/Engine/Classes/Particles/Parameter/ParticleModuleParameterDynamic.h
--- a/Engine/Source/Runtime/Engine/Classes/Particles/Parameter/ParticleModuleParameterDynamic.h
+++ b/Engine/Source/Runtime/Engine/Classes/Particles/Parameter/ParticleModuleParameterDynamic.h
@@ -144,7 +144,7 @@ class UParticleModuleParameterDynamic : public UParticleModuleParameterBase
 	 *	@param	bIsMeshEmitter		true if the emitter is a mesh emitter...
 	 *
 	 */
-	virtual void UpdateParameterNames(UMaterialInterface* InMaterialInterface, bool bIsMeshEmitter);
+	virtual void UpdateParameterNames(UMaterialInterface* InMaterialInterface);
 
 	
 	/**
diff --git a/Engine/Source/Runtime/Engine/Private/Particles/ParticleModules_Parameter.cpp b/Engine/Source/Runtime/Engine/Private/Particles/ParticleModules_Parameter.cpp
--- a/Engine/Source/Runtime/Engine/Private/Particles/ParticleModules_Parameter.cpp
+++ b/Engine/Source/Runtime/Engine/Private/Particles/ParticleModules_Parameter.cpp
@@ -393,15 +393,12 @@ UMaterial* UParticleModuleParameterDynamic_RetrieveMaterial(UMaterialInterface*
 /**
  *	Helper function to find the DynamicParameter expression in a material
  */
-UMaterialExpressionDynamicParameter* UParticleModuleParameterDynamic_GetDynamicParameterExpression(UMaterial* InMaterial, bool bIsMeshEmitter)
+UMaterialExpressionDynamicParameter* UParticleModuleParameterDynamic_GetDynamicParameterExpression(UMaterial* InMaterial)
 {
 	UMaterialExpressionDynamicParameter* DynParamExp = NULL;
 	for (int32 ExpIndex = 0; ExpIndex < InMaterial->Expressions.Num(); ExpIndex++)
 	{
-		if (bIsMeshEmitter == false)
-		{
-			DynParamExp = Cast<UMaterialExpressionDynamicParameter>(InMaterial->Expressions[ExpIndex]);
-		}
+		DynParamExp = Cast<UMaterialExpressionDynamicParameter>(InMaterial->Expressions[ExpIndex]);
 
 		if (DynParamExp != NULL)
 		{
@@ -413,7 +410,7 @@ UMaterialExpressionDynamicParameter* UParticleModuleParameterDynamic_GetDynamicP
 }
 
 
-void UParticleModuleParameterDynamic::UpdateParameterNames(UMaterialInterface* InMaterialInterface, bool bIsMeshEmitter)
+void UParticleModuleParameterDynamic::UpdateParameterNames(UMaterialInterface* InMaterialInterface)
 {
 	UMaterial* Material = UParticleModuleParameterDynamic_RetrieveMaterial(InMaterialInterface);
 	if (Material == NULL)
@@ -422,7 +419,7 @@ void UParticleModuleParameterDynamic::UpdateParameterNames(UMaterialInterface* I
 	}
 
 	// Check the expressions...
-	UMaterialExpressionDynamicParameter* DynParamExp = UParticleModuleParameterDynamic_GetDynamicParameterExpression(Material, bIsMeshEmitter);
+	UMaterialExpressionDynamicParameter* DynParamExp = UParticleModuleParameterDynamic_GetDynamicParameterExpression(Material);
 	if (DynParamExp == NULL)
 	{
 		return;
@@ -450,10 +447,20 @@ void UParticleModuleParameterDynamic::RefreshModule(UInterpCurveEdSetup* EdSetup
 			}
 		}
 
-		UMaterialInterface* MaterialInterface = LODLevel->RequiredModule ? LODLevel->RequiredModule->Material : NULL;
+		UMaterialInterface* MaterialInterface = LODLevel->RequiredModule ? LODLevel->RequiredModule->Material : nullptr;
+		if( bIsMeshEmitter )
+		{
+			UParticleModuleMeshMaterial* MeshMaterialModule = nullptr;
+			LODLevel->Modules.FindItemByClass( &MeshMaterialModule );
+			if( MeshMaterialModule && MeshMaterialModule->MeshMaterials.Num() > 0 )
+			{
+				MaterialInterface = MeshMaterialModule->MeshMaterials[0];
+			}
+		}
+
 		if (MaterialInterface)
 		{
-			UpdateParameterNames(MaterialInterface, bIsMeshEmitter);
+			UpdateParameterNames(MaterialInterface);
 			for (int32 ParamIndex = 0; ParamIndex < 4; ParamIndex++)
 			{
 				FString TempName = FString::Printf(TEXT("%s (DP%d)"), 
