Added missing comment

diff --git a/Engine/Source/Runtime/Engine/Private/Particles/ParticleModules_Parameter.cpp b/Engine/Source/Runtime/Engine/Private/Particles/ParticleModules_Parameter.cpp
--- a/Engine/Source/Runtime/Engine/Private/Particles/ParticleModules_Parameter.cpp
+++ b/Engine/Source/Runtime/Engine/Private/Particles/ParticleModules_Parameter.cpp
@@ -454,6 +454,7 @@ void UParticleModuleParameterDynamic::RefreshModule(UInterpCurveEdSetup* EdSetup
 			LODLevel->Modules.FindItemByClass( &MeshMaterialModule );
 			if( MeshMaterialModule && MeshMaterialModule->MeshMaterials.Num() > 0 )
 			{
+				// Note: there is no way to know which material to gather parameter names from if there is more than one.  Assume the first material
 				MaterialInterface = MeshMaterialModule->MeshMaterials[0];
 			}
 		}
