Skylight cubemap blending
* New Blueprint function for creating sky lighting from a blend of two cubemaps - SetCubemapBlend
* Diffuse sky lighting is blended on the CPU, specular is blended per-pixel with an extra cubemap lookup until the blend fraction becomes 0 or 1
* Skylight capture still blocks on the GPU for readback but not the rendering thread.  This happens when a new source or destination cubemap is requested with SetCubemapBlend.

diff --git a/Engine/Shaders/ReflectionEnvironmentComputeShaders.usf b/Engine/Shaders/ReflectionEnvironmentComputeShaders.usf
--- a/Engine/Shaders/ReflectionEnvironmentComputeShaders.usf
+++ b/Engine/Shaders/ReflectionEnvironmentComputeShaders.usf
@@ -554,7 +554,7 @@ void GatherRadiance(inout float4 Color, float3 WorldPosition, float3 RayDirectio
 	{
 		// Normalize for static skylight types which mix with lightmaps
 		bool bNormalize = SkyLightParameters.z < 1 && USE_LIGHTMAPS;
-		float3 SkyLighting = GetSkyLightReflection( RayDirection, Roughness, bNormalize );
+		float3 SkyLighting = GetSkyLightReflectionSupportingBlend( RayDirection, Roughness, bNormalize );
 
 		FLATTEN
 		if (bNormalize)
diff --git a/Engine/Shaders/ReflectionEnvironmentShaders.usf b/Engine/Shaders/ReflectionEnvironmentShaders.usf
--- a/Engine/Shaders/ReflectionEnvironmentShaders.usf
+++ b/Engine/Shaders/ReflectionEnvironmentShaders.usf
@@ -523,7 +523,7 @@ void ReflectionApplyPS(
 			float3 ReflectionVector = reflect(-V, GBuffer.WorldNormal);
 			// Normalize for static skylight types which mix with lightmaps
 			bool bNormalize = SkyLightParameters.z < 1 && View.UseLightmaps > 0 && !ALLOW_STATIC_LIGHTING;
-			float3 SkyLighting = GetSkyLightReflection(ReflectionVector, GBuffer.Roughness, bNormalize);
+			float3 SkyLighting = GetSkyLightReflectionSupportingBlend(ReflectionVector, GBuffer.Roughness, bNormalize);
 
 		    float DirectionalOcclusion = ScreenSpaceData.DirectionalOcclusion.g;
 		    SkyLighting *= DirectionalOcclusion;
diff --git a/Engine/Shaders/ReflectionEnvironmentShared.usf b/Engine/Shaders/ReflectionEnvironmentShared.usf
--- a/Engine/Shaders/ReflectionEnvironmentShared.usf
+++ b/Engine/Shaders/ReflectionEnvironmentShared.usf
@@ -32,16 +32,16 @@ float ComputeReflectionCaptureRoughnessFromMip( float Mip )
 TextureCube SkyLightCubemap;
 SamplerState SkyLightCubemapSampler;
 
-/** X = max mip, Y = 1 if sky light should be rendered, 0 otherwise, Z = 1 if sky light is dynamic, 0 otherwise. */
-float3 SkyLightParameters;
+/** X = max mip, Y = 1 if sky light should be rendered, 0 otherwise, Z = 1 if sky light is dynamic, 0 otherwise, W = blend fraction. */
+float4 SkyLightParameters;
 
-float3 GetSkyLightReflection(float3 ReflectionVector, float Roughness, bool Normalized)
+float3 GetSkyLightReflection(float3 ReflectionVector, float Roughness, bool bNormalized)
 {
 	float AbsoluteSpecularMip = ComputeReflectionCaptureMipFromRoughness(Roughness);
 	float3 Reflection = TextureCubeSampleLevel(SkyLightCubemap, SkyLightCubemapSampler, ReflectionVector, AbsoluteSpecularMip).rgb;
 
 	BRANCH
-	if (Normalized)
+	if (bNormalized)
 	{
 		// Sample the lowest resolution mip to get the average color
 		//@todo - can't normalize sky lighting and reflection capture lighting separately
@@ -52,6 +52,36 @@ float3 GetSkyLightReflection(float3 ReflectionVector, float Roughness, bool Norm
 
 	return Reflection * View.SkyLightColor.rgb;
 }
+
+TextureCube SkyLightBlendDestinationCubemap;
+SamplerState SkyLightBlendDestinationCubemapSampler;
+
+float3 GetSkyLightReflectionSupportingBlend(float3 ReflectionVector, float Roughness, bool bNormalized)
+{
+	float3 Reflection = GetSkyLightReflection(ReflectionVector, Roughness, bNormalized);
+
+	BRANCH
+	if (SkyLightParameters.w > 0)
+	{
+		float AbsoluteSpecularMip = ComputeReflectionCaptureMipFromRoughness(Roughness);
+		float3 BlendDestinationReflection = TextureCubeSampleLevel(SkyLightBlendDestinationCubemap, SkyLightBlendDestinationCubemapSampler, ReflectionVector, AbsoluteSpecularMip).rgb;
+
+		BRANCH
+		if (bNormalized)
+		{
+			// Sample the lowest resolution mip to get the average color
+			//@todo - can't normalize sky lighting and reflection capture lighting separately
+			float3 LowFrequencyReflection = TextureCubeSampleLevel(SkyLightBlendDestinationCubemap, SkyLightBlendDestinationCubemapSampler, ReflectionVector, SkyLightParameters.x).rgb;
+			float LowFrequencyBrightness = Luminance(LowFrequencyReflection);
+			BlendDestinationReflection /= max(LowFrequencyBrightness, .00001f);
+		}
+
+		Reflection = lerp(Reflection, BlendDestinationReflection * View.SkyLightColor.rgb, SkyLightParameters.w);
+	}
+
+	return Reflection;
+}
+
 /** 
  * Computes sky diffuse lighting from the SH irradiance map.  
  * This has the SH basis evaluation and diffuse convolusion weights combined for minimal ALU's - see "Stupid Spherical Harmonics (SH) Tricks" 
diff --git a/Engine/Source/Runtime/Engine/Classes/Components/SkyLightComponent.h b/Engine/Source/Runtime/Engine/Classes/Components/SkyLightComponent.h
--- a/Engine/Source/Runtime/Engine/Classes/Components/SkyLightComponent.h
+++ b/Engine/Source/Runtime/Engine/Classes/Components/SkyLightComponent.h
@@ -155,6 +155,7 @@ class ENGINE_API USkyLightComponent : public ULightComponentBase
 
 	/** Called each tick to recapture and queued sky captures. */
 	static void UpdateSkyCaptureContents(UWorld* WorldToUpdate);
+	static void UpdateSkyCaptureContentsArray(UWorld* WorldToUpdate, TArray<USkyLightComponent*>& ComponentArray, bool bBlendSources);
 
 	/** Computes an irradiance environment map using only emissive contribution from the sky light. */
 	void CaptureEmissiveIrradianceEnvironmentMap(FSHVectorRGB3& OutIrradianceMap) const;
@@ -173,6 +174,15 @@ class ENGINE_API USkyLightComponent : public ULightComponentBase
 	UFUNCTION(BlueprintCallable, Category="SkyLight")
 	void SetCubemap(UTextureCube* NewCubemap);
 
+	/** 
+	 * Creates sky lighting from a blend between two cubemaps, which is only valid when SourceType is set to SpecifiedCubemap. 
+	 * This can be used to seamlessly transition sky lighting between different times of day.
+	 * The caller should continue to update the blend until BlendFraction is 0 or 1 to reduce rendering cost.
+	 * The caller is responsible for avoiding pops due to changing the source or destination.
+	 */
+	UFUNCTION(BlueprintCallable, Category="SkyLight")
+	void SetCubemapBlend(UTextureCube* SourceCubemap, UTextureCube* DestinationCubemap, float InBlendFraction);
+
 	UFUNCTION(BlueprintCallable, Category="Rendering|Components|SkyLight")
 	void SetOcclusionTint(const FColor& InTint);
 
@@ -183,6 +193,7 @@ class ENGINE_API USkyLightComponent : public ULightComponentBase
 
 	/** Indicates that the capture needs to recapture the scene, adds it to the recapture queue. */
 	void SetCaptureIsDirty();
+	void SetBlendDestinationCaptureIsDirty();
 
 	/** 
 	 * Recaptures the scene for the skylight. 
@@ -199,18 +210,25 @@ class ENGINE_API USkyLightComponent : public ULightComponentBase
 
 protected:
 
-	/** Whether the reflection capture needs to re-capture the scene. */
-	bool bCaptureDirty;
-
 	/** Indicates whether the cached data stored in GetComponentInstanceData is valid to be applied in ApplyComponentInstanceData. */
 	bool bSavedConstructionScriptValuesValid;
 
 	bool bHasEverCaptured;
 
 	TRefCountPtr<FSkyTextureCubeResource> ProcessedSkyTexture;
-
 	FSHVectorRGB3 IrradianceEnvironmentMap;
 
+	/** If 0, no blend is present.  If > 0, BlendDestinationProcessedSkyTexture and BlendDestinationIrradianceEnvironmentMap must be generated and used for rendering. */
+	float BlendFraction;
+
+	UPROPERTY(transient)
+	class UTextureCube* BlendDestinationCubemap;
+	TRefCountPtr<FSkyTextureCubeResource> BlendDestinationProcessedSkyTexture;
+	FSHVectorRGB3 BlendDestinationIrradianceEnvironmentMap;
+
+	/** Tracks when the rendering thread has completed its writes to IrradianceEnvironmentMap. */
+	FRenderCommandFence IrradianceMapFence;
+
 	/** Fence used to track progress of releasing resources on the rendering thread. */
 	FRenderCommandFence ReleaseResourcesFence;
 
@@ -222,6 +240,7 @@ protected:
 	 * This queue should be in the UWorld or the FSceneInterface, but those are not available yet in PostLoad.
 	 */
 	static TArray<USkyLightComponent*> SkyCapturesToUpdate;
+	static TArray<USkyLightComponent*> SkyCapturesToUpdateBlendDestinations;
 
 	// Begin UActorComponent Interface
 	virtual void CreateRenderState_Concurrent() override;
diff --git a/Engine/Source/Runtime/Engine/Private/Components/SkyLightComponent.cpp b/Engine/Source/Runtime/Engine/Private/Components/SkyLightComponent.cpp
--- a/Engine/Source/Runtime/Engine/Private/Components/SkyLightComponent.cpp
+++ b/Engine/Source/Runtime/Engine/Private/Components/SkyLightComponent.cpp
@@ -70,22 +70,57 @@ void UWorld::UpdateAllSkyCaptures()
 	USkyLightComponent::UpdateSkyCaptureContents(this);
 }
 
+void FSkyLightSceneProxy::Initialize(float InBlendFraction, const FSHVectorRGB3* InIrradianceEnvironmentMap, const FSHVectorRGB3* BlendDestinationIrradianceEnvironmentMap)
+{
+	BlendFraction = FMath::Clamp(InBlendFraction, 0.0f, 1.0f);
+
+	if (BlendFraction > 0 && BlendDestinationProcessedTexture != NULL)
+	{
+		if (BlendFraction < 1)
+		{
+			IrradianceEnvironmentMap = (*InIrradianceEnvironmentMap) * (1 - BlendFraction) + (*BlendDestinationIrradianceEnvironmentMap) * BlendFraction;
+		}
+		else
+		{
+			// Blend is full destination, treat as source to avoid blend overhead in shaders
+			IrradianceEnvironmentMap = *BlendDestinationIrradianceEnvironmentMap;
+		}
+	}
+	else
+	{
+		// Blend is full source
+		IrradianceEnvironmentMap = *InIrradianceEnvironmentMap;
+		BlendFraction = 0;
+	}
+}
+
 FSkyLightSceneProxy::FSkyLightSceneProxy(const USkyLightComponent* InLightComponent)
 	: LightComponent(InLightComponent)
 	, ProcessedTexture(InLightComponent->ProcessedSkyTexture)
+	, BlendDestinationProcessedTexture(InLightComponent->BlendDestinationProcessedSkyTexture)
 	, SkyDistanceThreshold(InLightComponent->SkyDistanceThreshold)
 	, bCastShadows(InLightComponent->CastShadows)
 	, bWantsStaticShadowing(InLightComponent->Mobility == EComponentMobility::Stationary)
 	, bPrecomputedLightingIsValid(InLightComponent->bPrecomputedLightingIsValid)
 	, bHasStaticLighting(InLightComponent->HasStaticLighting())
 	, LightColor(FLinearColor(InLightComponent->LightColor) * InLightComponent->Intensity)
-	, IrradianceEnvironmentMap(InLightComponent->IrradianceEnvironmentMap)
 	, IndirectLightingIntensity(InLightComponent->IndirectLightingIntensity)
 	, OcclusionMaxDistance(InLightComponent->OcclusionMaxDistance)
 	, Contrast(InLightComponent->Contrast)
 	, MinOcclusion(InLightComponent->MinOcclusion)
 	, OcclusionTint(InLightComponent->OcclusionTint)
 {
+	ENQUEUE_UNIQUE_RENDER_COMMAND_FOURPARAMETER(
+		FInitSkyProxy,
+		const FSHVectorRGB3*,InIrradianceEnvironmentMap,&InLightComponent->IrradianceEnvironmentMap,
+		const FSHVectorRGB3*,BlendDestinationIrradianceEnvironmentMap,&InLightComponent->BlendDestinationIrradianceEnvironmentMap,
+		float,InBlendFraction,InLightComponent->BlendFraction,
+		FSkyLightSceneProxy*,LightSceneProxy,this,
+	{
+		// Only access the irradiance maps on the RT, even though they belong to the USkyLightComponent, 
+		// Because FScene::UpdateSkyCaptureContents does not block the RT so the writes could still be in flight
+		LightSceneProxy->Initialize(InBlendFraction, InIrradianceEnvironmentMap, BlendDestinationIrradianceEnvironmentMap);
+	});
 }
 
 USkyLightComponent::USkyLightComponent(const FObjectInitializer& ObjectInitializer)
@@ -107,7 +142,6 @@ USkyLightComponent::USkyLightComponent(const FObjectInitializer& ObjectInitializ
 	IndirectLightingIntensity = 1.0f;
 	SkyDistanceThreshold = 150000;
 	Mobility = EComponentMobility::Stationary;
-	bCaptureDirty = false;
 	bLowerHemisphereIsBlack = true;
 	bSavedConstructionScriptValuesValid = true;
 	bHasEverCaptured = false;
@@ -131,7 +165,17 @@ void USkyLightComponent::SetCaptureIsDirty()
 	if (bVisible && bAffectsWorld)
 	{
 		SkyCapturesToUpdate.AddUnique(this);
-		bCaptureDirty = true; 
+
+		// Mark saved values as invalid, in case a sky recapture is requested in a construction script between a save / restore of sky capture state
+		bSavedConstructionScriptValuesValid = false;
+	}
+}
+
+void USkyLightComponent::SetBlendDestinationCaptureIsDirty()
+{ 
+	if (bVisible && bAffectsWorld && BlendDestinationCubemap)
+	{
+		SkyCapturesToUpdateBlendDestinations.AddUnique(this); 
 
 		// Mark saved values as invalid, in case a sky recapture is requested in a construction script between a save / restore of sky capture state
 		bSavedConstructionScriptValuesValid = false;
@@ -139,6 +183,7 @@ void USkyLightComponent::SetCaptureIsDirty()
 }
 
 TArray<USkyLightComponent*> USkyLightComponent::SkyCapturesToUpdate;
+TArray<USkyLightComponent*> USkyLightComponent::SkyCapturesToUpdateBlendDestinations;
 
 void USkyLightComponent::CreateRenderState_Concurrent()
 {
@@ -177,7 +222,6 @@ void USkyLightComponent::PostInitProperties()
 	{
 		// Enqueue an update by default, so that newly placed components will get an update
 		// PostLoad will undo this for components loaded from disk
-		bCaptureDirty = true;
 		SkyCapturesToUpdate.AddUnique(this);
 	}
 
@@ -192,7 +236,6 @@ void USkyLightComponent::PostLoad()
 	if (!bVisible || HasAnyFlags(RF_ClassDefaultObject))
 	{
 		SkyCapturesToUpdate.Remove(this);
-		bCaptureDirty = false;
 	}
 }
 
@@ -335,10 +378,8 @@ void USkyLightComponent::CheckForErrors()
 void USkyLightComponent::BeginDestroy()
 {
 	// Deregister the component from the update queue
-	if (bCaptureDirty)
-	{
-		SkyCapturesToUpdate.Remove(this);
-	}
+	SkyCapturesToUpdate.Remove(this);
+	SkyCapturesToUpdateBlendDestinations.Remove(this);
 
 	// Release reference
 	ProcessedSkyTexture = NULL;
@@ -382,6 +423,9 @@ FActorComponentInstanceData* USkyLightComponent::GetComponentInstanceData() cons
 	InstanceData->LightGuid = LightGuid;
 	InstanceData->bPrecomputedLightingIsValid = bPrecomputedLightingIsValid;
 	InstanceData->ProcessedSkyTexture = ProcessedSkyTexture;
+
+	// Block until the rendering thread has completed its writes from a previous capture
+	IrradianceMapFence.Wait();
 	InstanceData->IrradianceEnvironmentMap = IrradianceEnvironmentMap;
 
 	return InstanceData;
@@ -399,30 +443,29 @@ void USkyLightComponent::ApplyComponentInstanceData(FPrecomputedSkyLightInstance
 	if (ProcessedSkyTexture && bSavedConstructionScriptValuesValid)
 	{
 		// We have valid capture state, remove the queued update
-		bCaptureDirty = false;
 		SkyCapturesToUpdate.Remove(this);
 	}
 
 	MarkRenderStateDirty();
 }
 
-void USkyLightComponent::UpdateSkyCaptureContents(UWorld* WorldToUpdate)
+void USkyLightComponent::UpdateSkyCaptureContentsArray(UWorld* WorldToUpdate, TArray<USkyLightComponent*>& ComponentArray, bool bOperateOnBlendSource)
 {
-	if (WorldToUpdate->Scene)
+	const bool bIsCompilingShaders = GShaderCompilingManager != NULL && GShaderCompilingManager->IsCompiling();
+
+	// Iterate backwards so we can remove elements without changing the index
+	for (int32 CaptureIndex = ComponentArray.Num() - 1; CaptureIndex >= 0; CaptureIndex--)
 	{
-		const bool bIsCompilingShaders = GShaderCompilingManager != NULL && GShaderCompilingManager->IsCompiling();
+		USkyLightComponent* CaptureComponent = ComponentArray[CaptureIndex];
 
-		// Iterate backwards so we can remove elements without changing the index
-		for (int32 CaptureIndex = SkyCapturesToUpdate.Num() - 1; CaptureIndex >= 0; CaptureIndex--)
+		if ((!CaptureComponent->GetOwner() || WorldToUpdate->ContainsActor(CaptureComponent->GetOwner()))
+			// Only process sky capture requests once async shader compiling completes, otherwise we will capture the scene with temporary shaders
+			&& (!bIsCompilingShaders || CaptureComponent->SourceType == SLS_SpecifiedCubemap))
 		{
-			USkyLightComponent* CaptureComponent = SkyCapturesToUpdate[CaptureIndex];
-
-			if ((!CaptureComponent->GetOwner() || WorldToUpdate->ContainsActor(CaptureComponent->GetOwner()))
-				// Only process sky capture requests once async shader compiling completes, otherwise we will capture the scene with temporary shaders
-				&& (!bIsCompilingShaders || CaptureComponent->SourceType == SLS_SpecifiedCubemap))
+			// Only capture valid sky light components
+			if (CaptureComponent->SourceType != SLS_SpecifiedCubemap || CaptureComponent->Cubemap)
 			{
-				// Only capture valid sky light components
-				if (CaptureComponent->SourceType != SLS_SpecifiedCubemap || CaptureComponent->Cubemap)
+				if (bOperateOnBlendSource)
 				{
 					// Allocate the needed texture on first capture
 					if (!CaptureComponent->ProcessedSkyTexture)
@@ -433,19 +476,43 @@ void USkyLightComponent::UpdateSkyCaptureContents(UWorld* WorldToUpdate)
 						CaptureComponent->MarkRenderStateDirty();
 					}
 
-					WorldToUpdate->Scene->UpdateSkyCaptureContents(CaptureComponent, false, CaptureComponent->ProcessedSkyTexture, CaptureComponent->IrradianceEnvironmentMap);
+					WorldToUpdate->Scene->UpdateSkyCaptureContents(CaptureComponent, false, CaptureComponent->Cubemap, CaptureComponent->ProcessedSkyTexture, CaptureComponent->IrradianceEnvironmentMap);
+				}
+				else
+				{
+					// Allocate the needed texture on first capture
+					if (!CaptureComponent->BlendDestinationProcessedSkyTexture)
+					{
+						CaptureComponent->BlendDestinationProcessedSkyTexture = new FSkyTextureCubeResource();
+						CaptureComponent->BlendDestinationProcessedSkyTexture->SetupParameters(GReflectionCaptureSize, FMath::CeilLogTwo(GReflectionCaptureSize) + 1, PF_FloatRGBA);
+						BeginInitResource(CaptureComponent->BlendDestinationProcessedSkyTexture);
+						CaptureComponent->MarkRenderStateDirty(); 
+					}
 
-					CaptureComponent->bHasEverCaptured = true;
-					CaptureComponent->MarkRenderStateDirty();
+					WorldToUpdate->Scene->UpdateSkyCaptureContents(CaptureComponent, false, CaptureComponent->BlendDestinationCubemap, CaptureComponent->BlendDestinationProcessedSkyTexture, CaptureComponent->BlendDestinationIrradianceEnvironmentMap);
 				}
 
-				// Only remove queued update requests if we processed it for the right world
-				SkyCapturesToUpdate.RemoveAt(CaptureIndex);
+				CaptureComponent->IrradianceMapFence.BeginFence();
+				CaptureComponent->bHasEverCaptured = true;
+				CaptureComponent->MarkRenderStateDirty();
 			}
+
+			// Only remove queued update requests if we processed it for the right world
+			ComponentArray.RemoveAt(CaptureIndex);
 		}
 	}
 }
 
+void USkyLightComponent::UpdateSkyCaptureContents(UWorld* WorldToUpdate)
+{
+	if (WorldToUpdate->Scene)
+	{
+		QUICK_SCOPE_CYCLE_COUNTER(STAT_SkylightCaptures);
+		UpdateSkyCaptureContentsArray(WorldToUpdate, SkyCapturesToUpdate, true);
+		UpdateSkyCaptureContentsArray(WorldToUpdate, SkyCapturesToUpdateBlendDestinations, false);
+	}
+}
+
 void USkyLightComponent::CaptureEmissiveIrradianceEnvironmentMap(FSHVectorRGB3& OutIrradianceMap) const
 {
 	OutIrradianceMap = FSHVectorRGB3();
@@ -454,7 +521,9 @@ void USkyLightComponent::CaptureEmissiveIrradianceEnvironmentMap(FSHVectorRGB3&
 	{
 		// Capture emissive scene lighting only for the lighting build
 		// This is necessary to avoid a feedback loop with the last lighting build results
-		GetScene()->UpdateSkyCaptureContents(this, true, NULL, OutIrradianceMap);
+		GetScene()->UpdateSkyCaptureContents(this, true, Cubemap, NULL, OutIrradianceMap);
+		// Wait until writes to OutIrradianceMap have completed
+		FlushRenderingCommands();
 	}
 }
 
@@ -497,16 +566,58 @@ void USkyLightComponent::SetLightColor(FLinearColor NewLightColor)
 
 void USkyLightComponent::SetCubemap(UTextureCube* NewCubemap)
 {
-	// Can't set color on a static light
+	// Can't set on a static light
 	if (AreDynamicDataChangesAllowed()
 		&& Cubemap != NewCubemap)
 	{
 		Cubemap = NewCubemap;
+		// Note: this will cause all static draw lists to be recreated, which is often 10ms in a large level
 		MarkRenderStateDirty();
+		// Note: this will cause the cubemap to be reprocessed including readback from the GPU
 		SetCaptureIsDirty();
 	}
 }
 
+void USkyLightComponent::SetCubemapBlend(UTextureCube* SourceCubemap, UTextureCube* DestinationCubemap, float InBlendFraction)
+{
+	if (AreDynamicDataChangesAllowed()
+		&& (Cubemap != SourceCubemap || BlendDestinationCubemap != DestinationCubemap || BlendFraction != InBlendFraction)
+		&& SourceType == SLS_SpecifiedCubemap)
+	{
+		if (Cubemap != SourceCubemap)
+		{
+			Cubemap = SourceCubemap;
+			SetCaptureIsDirty();
+		}
+
+		if (BlendDestinationCubemap != DestinationCubemap)
+		{
+			BlendDestinationCubemap = DestinationCubemap;
+			SetBlendDestinationCaptureIsDirty();
+		}
+
+		if (BlendFraction != InBlendFraction)
+		{
+			BlendFraction = InBlendFraction;
+
+			if (SceneProxy)
+			{
+				ENQUEUE_UNIQUE_RENDER_COMMAND_FOURPARAMETER(
+					FUpdateSkyProxy,
+					const FSHVectorRGB3*,InIrradianceEnvironmentMap,&IrradianceEnvironmentMap,
+					const FSHVectorRGB3*,BlendDestinationIrradianceEnvironmentMap,&BlendDestinationIrradianceEnvironmentMap,
+					float,InBlendFraction,BlendFraction,
+					FSkyLightSceneProxy*,LightSceneProxy,SceneProxy,
+				{
+					// Only access the irradiance maps on the RT, even though they belong to the USkyLightComponent, 
+					// Because FScene::UpdateSkyCaptureContents does not block the RT so the writes could still be in flight
+					LightSceneProxy->Initialize(InBlendFraction, InIrradianceEnvironmentMap, BlendDestinationIrradianceEnvironmentMap);
+				});
+			}
+		}
+	}
+}
+
 void USkyLightComponent::SetOcclusionTint(const FColor& InTint)
 {
 	// Can't set on a static light
@@ -514,6 +625,7 @@ void USkyLightComponent::SetOcclusionTint(const FColor& InTint)
 		&& OcclusionTint != InTint)
 	{
 		OcclusionTint = InTint;
+		// Note: this will cause all static draw lists to be recreated, which is often 10ms in a large level
 		MarkRenderStateDirty();
 	}
 }
@@ -525,6 +637,7 @@ void USkyLightComponent::SetMinOcclusion(float InMinOcclusion)
 		&& MinOcclusion != InMinOcclusion)
 	{
 		MinOcclusion = InMinOcclusion;
+		// Note: this will cause all static draw lists to be recreated, which is often 10ms in a large level
 		MarkRenderStateDirty();
 	}
 }
@@ -539,6 +652,7 @@ void USkyLightComponent::SetVisibility(bool bNewVisibility, bool bPropagateToChi
 	{
 		// Capture if we are being enabled for the first time
 		SetCaptureIsDirty();
+		SetBlendDestinationCaptureIsDirty();
 	}
 }
 
diff --git a/Engine/Source/Runtime/Engine/Public/SceneInterface.h b/Engine/Source/Runtime/Engine/Public/SceneInterface.h
--- a/Engine/Source/Runtime/Engine/Public/SceneInterface.h
+++ b/Engine/Source/Runtime/Engine/Public/SceneInterface.h
@@ -104,7 +104,7 @@ public:
 	 * Updates the contents of the given sky capture by rendering the scene. 
 	 * This must be called on the game thread.
 	 */
-	virtual void UpdateSkyCaptureContents(const USkyLightComponent* CaptureComponent, bool bCaptureEmissiveOnly, FTexture* OutProcessedTexture, FSHVectorRGB3& OutIrradianceEnvironmentMap) {}
+	virtual void UpdateSkyCaptureContents(const USkyLightComponent* CaptureComponent, bool bCaptureEmissiveOnly, UTextureCube* SourceCubemap, FTexture* OutProcessedTexture, FSHVectorRGB3& OutIrradianceEnvironmentMap) {}
 
 	/** 
 	* Updates the contents of the given scene capture by rendering the scene. 
diff --git a/Engine/Source/Runtime/Engine/Public/SceneManagement.h b/Engine/Source/Runtime/Engine/Public/SceneManagement.h
--- a/Engine/Source/Runtime/Engine/Public/SceneManagement.h
+++ b/Engine/Source/Runtime/Engine/Public/SceneManagement.h
@@ -618,8 +618,12 @@ public:
 	/** Initialization constructor. */
 	FSkyLightSceneProxy(const class USkyLightComponent* InLightComponent);
 
+	void Initialize(float InBlendFraction, const FSHVectorRGB3* InIrradianceEnvironmentMap, const FSHVectorRGB3* BlendDestinationIrradianceEnvironmentMap);
+
 	const USkyLightComponent* LightComponent;
 	FTexture* ProcessedTexture;
+	float BlendFraction;
+	FTexture* BlendDestinationProcessedTexture;
 	float SkyDistanceThreshold;
 	bool bCastShadows;
 	bool bWantsStaticShadowing;
diff --git a/Engine/Source/Runtime/Renderer/Private/BasePassRendering.cpp b/Engine/Source/Runtime/Renderer/Private/BasePassRendering.cpp
--- a/Engine/Source/Runtime/Renderer/Private/BasePassRendering.cpp
+++ b/Engine/Source/Runtime/Renderer/Private/BasePassRendering.cpp
@@ -47,20 +47,46 @@ IMPLEMENT_BASEPASS_LIGHTMAPPED_SHADER_TYPE( FCachedVolumeIndirectLightingPolicy,
 IMPLEMENT_BASEPASS_LIGHTMAPPED_SHADER_TYPE( FCachedPointIndirectLightingPolicy, FCachedPointIndirectLightingPolicy );
 IMPLEMENT_BASEPASS_LIGHTMAPPED_SHADER_TYPE( FSimpleDynamicLightingPolicy, FSimpleDynamicLightingPolicy );
 
-void FSkyLightReflectionParameters::GetSkyParametersFromScene(const FScene* Scene, bool bApplySkyLight, FTexture*& OutSkyLightTextureResource, float& OutApplySkyLightMask, float& OutSkyMipCount, bool& bOutSkyLightIsDynamic)
+void FSkyLightReflectionParameters::GetSkyParametersFromScene(
+	const FScene* Scene, 
+	bool bApplySkyLight, 
+	FTexture*& OutSkyLightTextureResource, 
+	FTexture*& OutSkyLightBlendDestinationTextureResource, 
+	float& OutApplySkyLightMask, 
+	float& OutSkyMipCount, 
+	bool& bOutSkyLightIsDynamic, 
+	float& OutBlendFraction)
 {
 	OutSkyLightTextureResource = GBlackTextureCube;
+	OutSkyLightBlendDestinationTextureResource = GBlackTextureCube;
 	OutApplySkyLightMask = 0;
 	bOutSkyLightIsDynamic = false;
+	OutBlendFraction = 0;
 
 	if (Scene
 		&& Scene->SkyLight 
 		&& Scene->SkyLight->ProcessedTexture
 		&& bApplySkyLight)
 	{
-		OutSkyLightTextureResource = Scene->SkyLight->ProcessedTexture;
+		const FSkyLightSceneProxy& SkyLight = *Scene->SkyLight;
+		OutSkyLightTextureResource = SkyLight.ProcessedTexture;
+		OutBlendFraction = SkyLight.BlendFraction;
+
+		if (SkyLight.BlendFraction > 0.0f && SkyLight.BlendDestinationProcessedTexture)
+		{
+			if (SkyLight.BlendFraction < 1.0f)
+			{
+				OutSkyLightBlendDestinationTextureResource = SkyLight.BlendDestinationProcessedTexture;
+			}
+			else
+			{
+				OutSkyLightTextureResource = SkyLight.BlendDestinationProcessedTexture;
+				OutBlendFraction = 0;
+			}
+		}
+		
 		OutApplySkyLightMask = 1;
-		bOutSkyLightIsDynamic = !Scene->SkyLight->bHasStaticLighting && !Scene->SkyLight->bWantsStaticShadowing;
+		bOutSkyLightIsDynamic = !SkyLight.bHasStaticLighting && !SkyLight.bWantsStaticShadowing;
 	}
 
 	OutSkyMipCount = 1;
diff --git a/Engine/Source/Runtime/Renderer/Private/BasePassRendering.h b/Engine/Source/Runtime/Renderer/Private/BasePassRendering.h
--- a/Engine/Source/Runtime/Renderer/Private/BasePassRendering.h
+++ b/Engine/Source/Runtime/Renderer/Private/BasePassRendering.h
@@ -211,27 +211,35 @@ public:
 	{
 		SkyLightCubemap.Bind(ParameterMap, TEXT("SkyLightCubemap"));
 		SkyLightCubemapSampler.Bind(ParameterMap, TEXT("SkyLightCubemapSampler"));
+		SkyLightBlendDestinationCubemap.Bind(ParameterMap, TEXT("SkyLightBlendDestinationCubemap"));
+		SkyLightBlendDestinationCubemapSampler.Bind(ParameterMap, TEXT("SkyLightBlendDestinationCubemapSampler"));
 		SkyLightParameters.Bind(ParameterMap, TEXT("SkyLightParameters"));
 	}
 
 	template<typename TParamRef>
 	void SetParameters(FRHICommandList& RHICmdList, const TParamRef& ShaderRHI, const FScene* Scene, bool bApplySkyLight)
 	{
-		FTexture* SkyLightTextureResource = GBlackTextureCube;
-		float ApplySkyLightMask = 0;
-		float SkyMipCount = 1;
-		bool bSkyLightIsDynamic = false;
-
-		GetSkyParametersFromScene(Scene, bApplySkyLight, SkyLightTextureResource, ApplySkyLightMask, SkyMipCount, bSkyLightIsDynamic);
-
-		SetTextureParameter(RHICmdList, ShaderRHI, SkyLightCubemap, SkyLightCubemapSampler, SkyLightTextureResource);
-		const FVector SkyParametersValue(SkyMipCount - 1.0f, ApplySkyLightMask, bSkyLightIsDynamic ? 1.0f : 0.0f);
-		SetShaderValue(RHICmdList, ShaderRHI, SkyLightParameters, SkyParametersValue);
+		if (SkyLightCubemap.IsBound() || SkyLightBlendDestinationCubemap.IsBound() || SkyLightParameters.IsBound())
+		{
+			FTexture* SkyLightTextureResource = GBlackTextureCube;
+			FTexture* SkyLightBlendDestinationTextureResource = GBlackTextureCube;
+			float ApplySkyLightMask = 0;
+			float SkyMipCount = 1;
+			float BlendFraction = 0;
+			bool bSkyLightIsDynamic = false;
+
+			GetSkyParametersFromScene(Scene, bApplySkyLight, SkyLightTextureResource, SkyLightBlendDestinationTextureResource, ApplySkyLightMask, SkyMipCount, bSkyLightIsDynamic, BlendFraction);
+
+			SetTextureParameter(RHICmdList, ShaderRHI, SkyLightCubemap, SkyLightCubemapSampler, SkyLightTextureResource);
+			SetTextureParameter(RHICmdList, ShaderRHI, SkyLightBlendDestinationCubemap, SkyLightBlendDestinationCubemapSampler, SkyLightBlendDestinationTextureResource);
+			const FVector4 SkyParametersValue(SkyMipCount - 1.0f, ApplySkyLightMask, bSkyLightIsDynamic ? 1.0f : 0.0f, BlendFraction);
+			SetShaderValue(RHICmdList, ShaderRHI, SkyLightParameters, SkyParametersValue);
+		}
 	}
 
 	friend FArchive& operator<<(FArchive& Ar,FSkyLightReflectionParameters& P)
 	{
-		Ar << P.SkyLightCubemap << P.SkyLightCubemapSampler << P.SkyLightParameters;
+		Ar << P.SkyLightCubemap << P.SkyLightCubemapSampler << P.SkyLightParameters << P.SkyLightBlendDestinationCubemap << P.SkyLightBlendDestinationCubemapSampler;
 		return Ar;
 	}
 
@@ -239,9 +247,19 @@ private:
 
 	FShaderResourceParameter SkyLightCubemap;
 	FShaderResourceParameter SkyLightCubemapSampler;
+	FShaderResourceParameter SkyLightBlendDestinationCubemap;
+	FShaderResourceParameter SkyLightBlendDestinationCubemapSampler;
 	FShaderParameter SkyLightParameters;
 
-	void GetSkyParametersFromScene(const FScene* Scene, bool bApplySkyLight, FTexture*& OutSkyLightTextureResource, float& OutApplySkyLightMask, float& OutSkyMipCount, bool& bSkyLightIsDynamic);
+	void GetSkyParametersFromScene(
+		const FScene* Scene, 
+		bool bApplySkyLight, 
+		FTexture*& OutSkyLightTextureResource, 
+		FTexture*& OutSkyLightBlendDestinationTextureResource, 
+		float& OutApplySkyLightMask, 
+		float& OutSkyMipCount, 
+		bool& bSkyLightIsDynamic, 
+		float& OutBlendFraction);
 };
 
 /** Parameters needed for lighting translucency, shared by multiple shaders. */
diff --git a/Engine/Source/Runtime/Renderer/Private/ReflectionEnvironmentCapture.cpp b/Engine/Source/Runtime/Renderer/Private/ReflectionEnvironmentCapture.cpp
--- a/Engine/Source/Runtime/Renderer/Private/ReflectionEnvironmentCapture.cpp
+++ b/Engine/Source/Runtime/Renderer/Private/ReflectionEnvironmentCapture.cpp
@@ -1371,15 +1371,16 @@ void CopyToSkyTexture(FRHICommandList& RHICmdList, FScene* Scene, FTexture* Proc
 	}
 }
 
-void FScene::UpdateSkyCaptureContents(const USkyLightComponent* CaptureComponent, bool bCaptureEmissiveOnly, FTexture* OutProcessedTexture, FSHVectorRGB3& OutIrradianceEnvironmentMap)
+// Warning: returns before writes to OutIrradianceEnvironmentMap have completed, as they are queued on the rendering thread
+void FScene::UpdateSkyCaptureContents(const USkyLightComponent* CaptureComponent, bool bCaptureEmissiveOnly, UTextureCube* SourceCubemap, FTexture* OutProcessedTexture, FSHVectorRGB3& OutIrradianceEnvironmentMap)
 {	
-	//todo maybe use ShouldRenderSkylight here, but currently the Skylight hasn't been loaded in time.
 	if (GSupportsRenderTargetFormat_PF_FloatRGBA || GetFeatureLevel() >= ERHIFeatureLevel::SM4)
 	{
+		QUICK_SCOPE_CYCLE_COUNTER(STAT_UpdateSkyCaptureContents);
+
 		ENQUEUE_UNIQUE_RENDER_COMMAND( 
 			ClearCommand,
 		{
-			//@todo realtime skylight updates - skip this
 			ClearScratchCubemaps(RHICmdList);
 		});
 
@@ -1392,7 +1393,7 @@ void FScene::UpdateSkyCaptureContents(const USkyLightComponent* CaptureComponent
 		{
 			ENQUEUE_UNIQUE_RENDER_COMMAND_THREEPARAMETER( 
 				CopyCubemapCommand,
-				UTextureCube*, SourceTexture, CaptureComponent->Cubemap,
+				UTextureCube*, SourceTexture, SourceCubemap,
 				bool, bLowerHemisphereIsBlack, CaptureComponent->bLowerHemisphereIsBlack,
 				ERHIFeatureLevel::Type, FeatureLevel, GetFeatureLevel(),
 			{
@@ -1415,20 +1416,16 @@ void FScene::UpdateSkyCaptureContents(const USkyLightComponent* CaptureComponent
 			FilterReflectionEnvironment(RHICmdList, FeatureLevel, IrradianceEnvironmentMap, bNormalize);
 		});
 
-		// Wait until the SH coefficients have been written out by the RT before returning
-		//@todo realtime skylight updates - remove the need for this
-		FlushRenderingCommands();
-
 		// Optionally copy the filtered mip chain to the output texture
 		if (OutProcessedTexture)
 		{
-		ENQUEUE_UNIQUE_RENDER_COMMAND_TWOPARAMETER( 
-			CopyCommand,
-			FScene*, Scene, this,
+			ENQUEUE_UNIQUE_RENDER_COMMAND_TWOPARAMETER( 
+				CopyCommand,
+				FScene*, Scene, this,
 				FTexture*, ProcessedTexture, OutProcessedTexture,
-		{
-			CopyToSkyTexture(RHICmdList, Scene, ProcessedTexture);
-		});
+			{
+				CopyToSkyTexture(RHICmdList, Scene, ProcessedTexture);
+			});
 		}
 	}
 }
diff --git a/Engine/Source/Runtime/Renderer/Private/ScenePrivate.h b/Engine/Source/Runtime/Renderer/Private/ScenePrivate.h
--- a/Engine/Source/Runtime/Renderer/Private/ScenePrivate.h
+++ b/Engine/Source/Runtime/Renderer/Private/ScenePrivate.h
@@ -1591,7 +1591,7 @@ public:
 	virtual void UpdateSceneCaptureContents(class USceneCaptureComponent2D* CaptureComponent) override;
 	virtual void UpdateSceneCaptureContents(class USceneCaptureComponentCube* CaptureComponent) override;
 	virtual void AllocateReflectionCaptures(const TArray<UReflectionCaptureComponent*>& NewCaptures) override;
-	virtual void UpdateSkyCaptureContents(const USkyLightComponent* CaptureComponent, bool bCaptureEmissiveOnly, FTexture* OutProcessedTexture, FSHVectorRGB3& OutIrradianceEnvironmentMap) override;
+	virtual void UpdateSkyCaptureContents(const USkyLightComponent* CaptureComponent, bool bCaptureEmissiveOnly, UTextureCube* SourceCubemap, FTexture* OutProcessedTexture, FSHVectorRGB3& OutIrradianceEnvironmentMap) override; 
 	virtual void AddPrecomputedLightVolume(const class FPrecomputedLightVolume* Volume) override;
 	virtual void RemovePrecomputedLightVolume(const class FPrecomputedLightVolume* Volume) override;
 	virtual void UpdateLightTransform(ULightComponent* Light) override;
