Scalar material parameters can specify the slider min and max for the material instance editor

diff --git a/Engine/Source/Editor/MaterialEditor/Private/MaterialEditorInstanceDetailCustomization.cpp b/Engine/Source/Editor/MaterialEditor/Private/MaterialEditorInstanceDetailCustomization.cpp
--- a/Engine/Source/Editor/MaterialEditor/Private/MaterialEditorInstanceDetailCustomization.cpp
+++ b/Engine/Source/Editor/MaterialEditor/Private/MaterialEditorInstanceDetailCustomization.cpp
@@ -148,8 +148,15 @@ void FMaterialInstanceParameterDetails::CreateSingleGroupWidget(FEditorParameter
 		UDEditorTextureParameterValue* TextureParam = Cast<UDEditorTextureParameterValue>(Parameter);
 		UDEditorVectorParameterValue* VectorParam = Cast<UDEditorVectorParameterValue>(Parameter);
 
-		if (ScalarParam || SwitchParam || TextureParam || VectorParam || FontParam)
+		if( ScalarParam || SwitchParam || TextureParam || VectorParam || FontParam)
 		{
+			if (ScalarParam && ScalarParam->SliderMax > ScalarParam->SliderMin)
+			{
+				TSharedPtr<IPropertyHandle> ParameterValueProperty = ParameterProperty->GetChildHandle("ParameterValue");
+				ParameterValueProperty->SetInstanceMetaData("UIMin",FString::Printf(TEXT("%f"), ScalarParam->SliderMin));
+				ParameterValueProperty->SetInstanceMetaData("UIMax",FString::Printf(TEXT("%f"), ScalarParam->SliderMax));
+			}
+
 			CreateParameterValueWidget(Parameter, ParameterProperty, DetailGroup );
 		}
 		else if (CompMaskParam)
diff --git a/Engine/Source/Editor/UnrealEd/Classes/MaterialEditor/DEditorScalarParameterValue.h b/Engine/Source/Editor/UnrealEd/Classes/MaterialEditor/DEditorScalarParameterValue.h
--- a/Engine/Source/Editor/UnrealEd/Classes/MaterialEditor/DEditorScalarParameterValue.h
+++ b/Engine/Source/Editor/UnrealEd/Classes/MaterialEditor/DEditorScalarParameterValue.h
@@ -14,5 +14,7 @@ class UNREALED_API UDEditorScalarParameterValue : public UDEditorParameterValue
 	UPROPERTY(EditAnywhere, Category=DEditorScalarParameterValue)
 	float ParameterValue;
 
+	float SliderMin;
+	float SliderMax;
 };
 
diff --git a/Engine/Source/Editor/UnrealEd/Private/PreviewMaterial.cpp b/Engine/Source/Editor/UnrealEd/Private/PreviewMaterial.cpp
--- a/Engine/Source/Editor/UnrealEd/Private/PreviewMaterial.cpp
+++ b/Engine/Source/Editor/UnrealEd/Private/PreviewMaterial.cpp
@@ -323,9 +323,9 @@ void UMaterialEditorInstanceConstant::RegenerateArrays()
 		}
 		// Scalar Parameters.
 		ParentMaterial->GetAllScalarParameterNames(ParameterNames, Guids);
-		for(int32 ParameterIdx=0; ParameterIdx<ParameterNames.Num(); ParameterIdx++)
+		for (int32 ParameterIdx=0; ParameterIdx<ParameterNames.Num(); ParameterIdx++)
 		{			
-			UDEditorScalarParameterValue & ParameterValue = *(NewObject<UDEditorScalarParameterValue>());
+			UDEditorScalarParameterValue& ParameterValue = *(NewObject<UDEditorScalarParameterValue>());
 			FName ParameterName = ParameterNames[ParameterIdx];
 			float Value;
 
@@ -333,12 +333,12 @@ void UMaterialEditorInstanceConstant::RegenerateArrays()
 			ParameterValue.ParameterName = ParameterName;
 			ParameterValue.ExpressionId = Guids[ParameterIdx];
 
-			if(SourceInstance->GetScalarParameterValue(ParameterName, Value))
+			if (SourceInstance->GetScalarParameterValue(ParameterName, Value))
 			{
+				ParentMaterial->GetScalarParameterSliderMinMax(ParameterName, ParameterValue.SliderMin, ParameterValue.SliderMax);
 				ParameterValue.ParameterValue = Value;
 			}
 
-
 			// @todo: This is kind of slow, maybe store these in a map for lookup?
 			// See if this keyname exists in the source instance.
 			for(int32 ScalarParameterIdx=0; ScalarParameterIdx<SourceInstance->ScalarParameterValues.Num(); ScalarParameterIdx++)
diff --git a/Engine/Source/Runtime/Engine/Classes/Materials/Material.h b/Engine/Source/Runtime/Engine/Classes/Materials/Material.h
--- a/Engine/Source/Runtime/Engine/Classes/Materials/Material.h
+++ b/Engine/Source/Runtime/Engine/Classes/Materials/Material.h
@@ -806,6 +806,7 @@ public:
 	ENGINE_API virtual USubsurfaceProfile* GetSubsurfaceProfile_Internal() const override;
 
 	ENGINE_API void SetShadingModel(EMaterialShadingModel NewModel) {ShadingModel = NewModel;}
+	ENGINE_API bool GetScalarParameterSliderMinMax(FName ParameterName, float& OutMinSlider, float& OutMaxSlider) const;
 
 	/** Checks to see if an input property should be active, based on the state of the material */
 	ENGINE_API virtual bool IsPropertyActive(EMaterialProperty InProperty) const override;
diff --git a/Engine/Source/Runtime/Engine/Classes/Materials/MaterialExpressionScalarParameter.h b/Engine/Source/Runtime/Engine/Classes/Materials/MaterialExpressionScalarParameter.h
--- a/Engine/Source/Runtime/Engine/Classes/Materials/MaterialExpressionScalarParameter.h
+++ b/Engine/Source/Runtime/Engine/Classes/Materials/MaterialExpressionScalarParameter.h
@@ -13,6 +13,18 @@ class UMaterialExpressionScalarParameter : public UMaterialExpressionParameter
 	UPROPERTY(EditAnywhere, Category=MaterialExpressionScalarParameter)
 	float DefaultValue;
 
+	/** 
+	 * Sets the lower bound for the slider on this parameter in the material instance editor. 
+	 */
+	UPROPERTY(EditAnywhere, Category=MaterialExpressionScalarParameter)
+	float SliderMin;
+
+	/** 
+	 * Sets the upper bound for the slider on this parameter in the material instance editor. 
+	 * The slider will be disabled if SliderMax <= SliderMin.
+	 */
+	UPROPERTY(EditAnywhere, Category=MaterialExpressionScalarParameter)
+	float SliderMax;
 
 	// Begin UMaterialExpression Interface
 	virtual int32 Compile(class FMaterialCompiler* Compiler, int32 OutputIndex, int32 MultiplexIndex) override;
diff --git a/Engine/Source/Runtime/Engine/Private/Materials/Material.cpp b/Engine/Source/Runtime/Engine/Private/Materials/Material.cpp
--- a/Engine/Source/Runtime/Engine/Private/Materials/Material.cpp
+++ b/Engine/Source/Runtime/Engine/Private/Materials/Material.cpp
@@ -1433,6 +1433,55 @@ bool UMaterial::GetScalarParameterValue(FName ParameterName, float& OutValue) co
 	return false;
 }
 
+bool UMaterial::GetScalarParameterSliderMinMax(FName ParameterName, float& OutSliderMin, float& OutSliderMax) const
+{
+	float Value = 0;
+
+	for (const UMaterialExpression* Expression : Expressions)
+	{
+		if (Expression->IsA<UMaterialExpressionScalarParameter>())
+		{
+			const UMaterialExpressionScalarParameter* Parameter = CastChecked<const UMaterialExpressionScalarParameter>(Expression);
+			if (Parameter->IsNamedParameter(ParameterName, Value))
+			{
+				OutSliderMin = Parameter->SliderMin;
+				OutSliderMax = Parameter->SliderMax;
+				// Warning: in the case of duplicate parameters with different default values, this will find the first in the expression array, not necessarily the one that's used for rendering
+				return true;
+			}
+		}
+		else if (Expression->IsA<UMaterialExpressionMaterialFunctionCall>())
+		{
+			const UMaterialExpressionMaterialFunctionCall* FunctionCall = CastChecked<const UMaterialExpressionMaterialFunctionCall>(Expression);
+			if (FunctionCall->MaterialFunction)
+			{
+				TArray<UMaterialFunction*> Functions;
+				Functions.Add(FunctionCall->MaterialFunction);
+				FunctionCall->MaterialFunction->GetDependentFunctions(Functions);
+
+				for (UMaterialFunction* Function : Functions)
+				{
+					for (UMaterialExpression* FunctionExpression : Function->FunctionExpressions)
+					{
+						if (FunctionExpression->IsA<UMaterialExpressionScalarParameter>())
+						{
+							const UMaterialExpressionScalarParameter* Parameter = CastChecked<const UMaterialExpressionScalarParameter>(FunctionExpression);
+							if (Parameter->IsNamedParameter(ParameterName, Value))
+							{
+								OutSliderMin = Parameter->SliderMin;
+								OutSliderMax = Parameter->SliderMax;
+								return true;
+							}
+						}
+					}
+				}
+			}
+		}
+	}
+
+	return false;
+}
+
 bool UMaterial::GetTextureParameterValue(FName ParameterName, UTexture*& OutValue) const
 {
 	for (const UMaterialExpression* Expression : Expressions)
