UE-17356: World Position material expression hotfix on mixing translated world position and world position excluding shader offset.

diff --git a/Engine/Source/Runtime/Engine/Private/Materials/HLSLMaterialTranslator.h b/Engine/Source/Runtime/Engine/Private/Materials/HLSLMaterialTranslator.h
--- a/Engine/Source/Runtime/Engine/Private/Materials/HLSLMaterialTranslator.h
+++ b/Engine/Source/Runtime/Engine/Private/Materials/HLSLMaterialTranslator.h
@@ -2102,7 +2102,18 @@ protected:
 		{
 		case WPT_Default:
 			{
+				// TODO: in tessellation, no offset should actually return the world position without including the
+				// world position offset
 				WorldPositionSuffix = TEXT("");
+
+				if (ShaderFrequency != SF_Pixel)
+				{
+					// TODO: there is a miss understanding between no offset, translated world position and camera relative:
+					//		what the artist call camera relative, is the translated world position
+					//		what they call no offsets, is without apply material's world position offset
+					WorldPositionSuffix = TEXT("_NoOffsets");
+				}
+
 				break;
 			}
 
