Fixed shadowed variable declaration in SkeletalMeshComponent.h

diff --git a/Engine/Source/Runtime/Engine/Classes/Components/SkeletalMeshComponent.h b/Engine/Source/Runtime/Engine/Classes/Components/SkeletalMeshComponent.h
--- a/Engine/Source/Runtime/Engine/Classes/Components/SkeletalMeshComponent.h
+++ b/Engine/Source/Runtime/Engine/Classes/Components/SkeletalMeshComponent.h
@@ -932,9 +932,9 @@ public:
 	 *
 	 * @param	RequiredBones	List of bones to be blend
 	 */
-	void BlendPhysicsBones( TArray<FBoneIndexType>& RequiredBones )
+	void BlendPhysicsBones( TArray<FBoneIndexType>& Bones )
 	{
-		PerformBlendPhysicsBones(RequiredBones, AnimEvaluationContext.LocalAtoms);
+		PerformBlendPhysicsBones(Bones, AnimEvaluationContext.LocalAtoms);
 	}
 
 
