Fix constraint drives not being created when initialized as off and then turned on at runtime

#UE-17749

diff --git a/Engine/Source/Runtime/Engine/Classes/PhysicsEngine/ConstraintInstance.h b/Engine/Source/Runtime/Engine/Classes/PhysicsEngine/ConstraintInstance.h
--- a/Engine/Source/Runtime/Engine/Classes/PhysicsEngine/ConstraintInstance.h
+++ b/Engine/Source/Runtime/Engine/Classes/PhysicsEngine/ConstraintInstance.h
@@ -324,6 +324,17 @@ public:
 	/** Enables the angular drive towards a target orientation. */
 	UPROPERTY(EditAnywhere, BlueprintReadOnly, Category=AngularMotor)
 	uint32 bAngularOrientationDrive:1;
+
+private:
+	/** Enables the swing drive. Only relevant when in twist and swing mode */
+	UPROPERTY()
+	uint32 bEnableSwingDrive: 1;
+
+	/** Enables the swing drive. Only relevant when in twist and swing mode */
+	UPROPERTY()
+	uint32 bEnableTwistDrive : 1;
+
+public:
 	
 	/** Enables the angular drive towards a target velocity. */
 	UPROPERTY(EditAnywhere, BlueprintReadOnly, Category=AngularMotor)
diff --git a/Engine/Source/Runtime/Engine/Private/PhysicsEngine/ConstraintInstance.cpp b/Engine/Source/Runtime/Engine/Private/PhysicsEngine/ConstraintInstance.cpp
--- a/Engine/Source/Runtime/Engine/Private/PhysicsEngine/ConstraintInstance.cpp
+++ b/Engine/Source/Runtime/Engine/Private/PhysicsEngine/ConstraintInstance.cpp
@@ -244,6 +244,8 @@ FConstraintInstance::FConstraintInstance()
 	, LinearDriveForceLimit(0)
 	, bAngularOrientationDrive(false)
 	, bAngularVelocityDrive(false)
+    , bEnableSwingDrive(true)
+	, bEnableTwistDrive(true)
 	, AngularOrientationTarget(ForceInit)
 	, AngularVelocityTarget(ForceInit)
 	, AngularDriveSpring(50.0f)
@@ -766,109 +768,40 @@ void FConstraintInstance::GetConstraintForce(FVector& OutLinearForce, FVector& O
 /** Function for turning linear position drive on and off. */
 void FConstraintInstance::SetLinearPositionDrive(bool bEnableXDrive, bool bEnableYDrive, bool bEnableZDrive)
 {
-#if WITH_PHYSX
-	ExecuteOnUnbrokenJointReadWrite([&] (PxD6Joint* Joint)
-	{
-		// Get the current drives
-		PxD6JointDrive CurrentDriveX = Joint->getDrive(PxD6Drive::eX);
-		PxD6JointDrive CurrentDriveY = Joint->getDrive(PxD6Drive::eY);
-		PxD6JointDrive CurrentDriveZ = Joint->getDrive(PxD6Drive::eZ);
-		
-		CurrentDriveX.stiffness = bEnableXDrive ? LinearDriveSpring : 0.0f;
-		CurrentDriveY.stiffness = bEnableYDrive ? LinearDriveSpring : 0.0f;
-		CurrentDriveZ.stiffness = bEnableZDrive ? LinearDriveSpring : 0.0f;
-
-		Joint->setDrive(PxD6Drive::eX, CurrentDriveX);
-		Joint->setDrive(PxD6Drive::eY, CurrentDriveY);
-		Joint->setDrive(PxD6Drive::eZ, CurrentDriveZ);
-	});
-#endif
-
 	bLinearXPositionDrive = bEnableXDrive;
 	bLinearYPositionDrive = bEnableYDrive;
 	bLinearZPositionDrive = bEnableZDrive;
 	bLinearPositionDrive = bEnableXDrive || bEnableYDrive || bEnableZDrive;
+	SetLinearDriveParams(LinearDriveSpring, LinearDriveDamping, LinearDriveForceLimit);
 }
 
 
 /** Function for turning linear velocity drive on and off. */
 void FConstraintInstance::SetLinearVelocityDrive(bool bEnableXDrive, bool bEnableYDrive, bool bEnableZDrive)
 {
-#if WITH_PHYSX
-	PxD6Joint* Joint = (PxD6Joint*)ConstraintData;
-
-	if (Joint && !(Joint->getConstraintFlags()&PxConstraintFlag::eBROKEN))
-	{
-		// Get the current drives
-		PxD6JointDrive CurrentDriveX = Joint->getDrive(PxD6Drive::eX);
-		PxD6JointDrive CurrentDriveY = Joint->getDrive(PxD6Drive::eY);
-		PxD6JointDrive CurrentDriveZ = Joint->getDrive(PxD6Drive::eZ);
-
-		CurrentDriveX.damping = bEnableXDrive && FMath::Abs(LinearVelocityTarget.X) > 0.0f ? LinearDriveDamping : 0.0f;
-		CurrentDriveY.damping = bEnableYDrive && FMath::Abs(LinearVelocityTarget.Y) > 0.0f ? LinearDriveDamping : 0.0f;
-		CurrentDriveZ.damping = bEnableZDrive && FMath::Abs(LinearVelocityTarget.Z) > 0.0f ? LinearDriveDamping : 0.0f;
-
-		Joint->setDrive(PxD6Drive::eX, CurrentDriveX);
-		Joint->setDrive(PxD6Drive::eY, CurrentDriveY);
-		Joint->setDrive(PxD6Drive::eZ, CurrentDriveZ);
-	}
-#endif
-
 	bLinearXVelocityDrive = bEnableXDrive;
 	bLinearYVelocityDrive = bEnableYDrive;
 	bLinearZVelocityDrive = bEnableZDrive;
 	bLinearVelocityDrive = bEnableXDrive || bEnableYDrive || bEnableZDrive;
+	SetLinearDriveParams(LinearDriveSpring, LinearDriveDamping, LinearDriveForceLimit);
 }
 
 /** Function for turning angular position drive on and off. */
-void FConstraintInstance::SetAngularPositionDrive(bool bEnableSwingDrive, bool bEnableTwistDrive)
+void FConstraintInstance::SetAngularPositionDrive(bool InEnableSwingDrive, bool InEnableTwistDrive)
 {
-#if WITH_PHYSX
-	ExecuteOnUnbrokenJointReadWrite([&] (PxD6Joint* Joint)
-	{
-		// Get the current drives
-		PxD6JointDrive CurrentDriveSwing = Joint->getDrive(PxD6Drive::eSWING);
-		PxD6JointDrive CurrentDriveTwist = Joint->getDrive(PxD6Drive::eTWIST);
-		PxD6JointDrive CurrentDriveSlerp = Joint->getDrive(PxD6Drive::eSLERP);
-		const bool bSlerp = AngularDriveMode == EAngularDriveMode::SLERP;
-
-		CurrentDriveSwing.stiffness = !bSlerp && bEnableSwingDrive ? AngularDriveSpring : 0.0f;
-		CurrentDriveTwist.stiffness = !bSlerp && bEnableTwistDrive ? AngularDriveSpring : 0.0f;
-		CurrentDriveSlerp.stiffness = bSlerp  && (bEnableSwingDrive || bEnableTwistDrive) ? AngularDriveSpring : 0.0f;
-
-		Joint->setDrive(PxD6Drive::eSWING, CurrentDriveSwing);
-		Joint->setDrive(PxD6Drive::eTWIST, CurrentDriveTwist);
-		Joint->setDrive(PxD6Drive::eSLERP, CurrentDriveSlerp);
-		
-		bAngularOrientationDrive = bEnableSwingDrive || bEnableTwistDrive;
-	});
-#endif
+	bEnableSwingDrive = InEnableSwingDrive;
+	bEnableTwistDrive = InEnableTwistDrive;
+	bAngularOrientationDrive = bEnableSwingDrive || bEnableTwistDrive;
+	SetAngularDriveParams(AngularDriveSpring, AngularDriveDamping, AngularDriveForceLimit);
 }
 
 /** Function for turning angular velocity drive on and off. */
-void FConstraintInstance::SetAngularVelocityDrive(bool bEnableSwingDrive, bool bEnableTwistDrive)
+void FConstraintInstance::SetAngularVelocityDrive(bool InEnableSwingDrive, bool InEnableTwistDrive)
 {
-#if WITH_PHYSX
-	PxD6Joint* Joint = (PxD6Joint*)ConstraintData;
-	if (Joint &&  !(Joint->getConstraintFlags()&PxConstraintFlag::eBROKEN))
-	{
-		// Get the current drives
-		PxD6JointDrive CurrentDriveSwing = Joint->getDrive(PxD6Drive::eSWING);
-		PxD6JointDrive CurrentDriveTwist = Joint->getDrive(PxD6Drive::eTWIST);
-		PxD6JointDrive CurrentDriveSlerp = Joint->getDrive(PxD6Drive::eSLERP);
-		const bool bSlerp = AngularDriveMode == EAngularDriveMode::SLERP;
-
-		CurrentDriveSwing.damping = !bSlerp && bEnableSwingDrive ? AngularDriveDamping : 0.0f;
-		CurrentDriveTwist.damping = !bSlerp && bEnableTwistDrive ? AngularDriveDamping : 0.0f;
-		CurrentDriveSlerp.damping = bSlerp  && (bEnableSwingDrive && bEnableTwistDrive) ? AngularDriveDamping : 0.0f;
-
-		Joint->setDrive(PxD6Drive::eSWING, CurrentDriveSwing);
-		Joint->setDrive(PxD6Drive::eTWIST, CurrentDriveTwist);
-		Joint->setDrive(PxD6Drive::eSLERP, CurrentDriveSlerp);
-		
-		bAngularVelocityDrive = bEnableSwingDrive || bEnableTwistDrive;
-	}
-#endif
+	bEnableSwingDrive = InEnableSwingDrive;
+	bEnableTwistDrive = InEnableTwistDrive;
+	bAngularVelocityDrive = bEnableSwingDrive || bEnableTwistDrive;
+	SetAngularDriveParams(AngularDriveSpring, AngularDriveDamping, AngularDriveForceLimit);
 }
 
 /** Function for setting linear position target. */
@@ -1046,12 +979,12 @@ void FConstraintInstance::SetAngularDriveParams(float InSpring, float InDamping,
 			{
 				Joint->setDrive(PxD6Drive::eTWIST, PxD6JointDrive());
 				Joint->setDrive(PxD6Drive::eSWING, PxD6JointDrive());
-				Joint->setDrive(PxD6Drive::eSLERP, PxD6JointDrive(DriveSpring, DriveDamping, AngularForceLimit, bIsAccelerationDrive));
+				Joint->setDrive(PxD6Drive::eSLERP, (bEnableSwingDrive || bEnableTwistDrive) ? PxD6JointDrive(DriveSpring, DriveDamping, AngularForceLimit, bIsAccelerationDrive) : PxD6JointDrive());				
 			}
 			else
 			{
-				Joint->setDrive(PxD6Drive::eTWIST, PxD6JointDrive(DriveSpring, DriveDamping, AngularForceLimit, bIsAccelerationDrive));
-				Joint->setDrive(PxD6Drive::eSWING, PxD6JointDrive(DriveSpring, DriveDamping, AngularForceLimit, bIsAccelerationDrive));
+				Joint->setDrive(PxD6Drive::eTWIST, bEnableTwistDrive ? PxD6JointDrive(DriveSpring, DriveDamping, AngularForceLimit, bIsAccelerationDrive) : PxD6JointDrive());
+				Joint->setDrive(PxD6Drive::eSWING, bEnableSwingDrive ? PxD6JointDrive(DriveSpring, DriveDamping, AngularForceLimit, bIsAccelerationDrive) : PxD6JointDrive());
 				Joint->setDrive(PxD6Drive::eSLERP, PxD6JointDrive());
 			}
 		});
