TextRenderComponent no longer invalidates its render proxy whenever it is transformed.
We used to rebuild the text every time we moved it, but now we rely on transforms, so it is no longer necessary.

diff --git a/Engine/Source/Runtime/Engine/Private/Components/TextRenderComponent.cpp b/Engine/Source/Runtime/Engine/Private/Components/TextRenderComponent.cpp
--- a/Engine/Source/Runtime/Engine/Private/Components/TextRenderComponent.cpp
+++ b/Engine/Source/Runtime/Engine/Private/Components/TextRenderComponent.cpp
@@ -765,7 +765,9 @@ UMaterialInterface* UTextRenderComponent::GetMaterial(int32 ElementIndex) const
 
 bool UTextRenderComponent::ShouldRecreateProxyOnUpdateTransform() const
 {
-	return true;
+	// We used to rebuild the text every time we moved it, but
+	// now we rely on transforms, so it is no longer necessary.
+	return false;
 }
 
 FBoxSphereBounds UTextRenderComponent::CalcBounds(const FTransform& LocalToWorld) const
