FIXED UE-14081 Erroneous map check warning about non-uniform collision scaling

diff --git a/Engine/Source/Runtime/Engine/Private/Components/StaticMeshComponent.cpp b/Engine/Source/Runtime/Engine/Private/Components/StaticMeshComponent.cpp
--- a/Engine/Source/Runtime/Engine/Private/Components/StaticMeshComponent.cpp
+++ b/Engine/Source/Runtime/Engine/Private/Components/StaticMeshComponent.cpp
@@ -386,26 +386,6 @@ void UStaticMeshComponent::CheckForErrors()
 			->AddToken(FMapErrorToken::Create(FMapErrors::StaticMeshNull));
 	}
 
-	// Make sure any non uniform scaled meshes have appropriate collision
-	if ( IsCollisionEnabled() && StaticMesh != NULL && StaticMesh->BodySetup != NULL && Owner != NULL )
-	{
-		// Overall scale factor for this mesh.
-		const FVector& TotalScale3D = ComponentToWorld.GetScale3D();
-		if ( !TotalScale3D.IsUniform() &&
-			 (StaticMesh->BodySetup->AggGeom.BoxElems.Num() > 0   ||
-			  StaticMesh->BodySetup->AggGeom.SphylElems.Num() > 0 ||
-			  StaticMesh->BodySetup->AggGeom.SphereElems.Num() > 0) )
-
-		{
-			FFormatNamedArguments Arguments;
-			Arguments.Add(TEXT("MeshName"), FText::FromString(StaticMesh->GetName()));
-			FMessageLog("MapCheck").Warning()
-				->AddToken(FUObjectToken::Create(Owner))
-				->AddToken(FTextToken::Create(FText::Format(LOCTEXT( "MapCheck_Message_SimpleCollisionButNonUniformScale", "'{MeshName}' has simple collision but is being scaled non-uniformly - collision creation will fail" ), Arguments)))
-				->AddToken(FMapErrorToken::Create(FMapErrors::SimpleCollisionButNonUniformScale));
-		}
-	}
-
 	if ( BodyInstance.bSimulatePhysics && StaticMesh != NULL && StaticMesh->BodySetup != NULL && StaticMesh->BodySetup->AggGeom.GetElementCount() == 0) 
 	{
 		FMessageLog("MapCheck").Warning()
