FIX UE-14423 Show Player Collision / Show Visibility Collision don't work with non-uniformly scaled objects

diff --git a/Engine/Source/Editor/PhAT/Private/PhATRender.cpp b/Engine/Source/Editor/PhAT/Private/PhATRender.cpp
--- a/Engine/Source/Editor/PhAT/Private/PhATRender.cpp
+++ b/Engine/Source/Editor/PhAT/Private/PhATRender.cpp
@@ -81,6 +81,7 @@ void UPhATEdSkeletalMeshComponent::RenderAssetTools(const FSceneView* View, clas
 		{
 			FTransform BoneTM = GetBoneTransform(BoneIndex);
 			float Scale = BoneTM.GetScale3D().GetAbsMax();
+			FVector VectorScale(Scale);
 			BoneTM.RemoveScaling();
 
 			FKAggregateGeom* AggGeom = &PhysicsAsset->BodySetup[i]->AggGeom;
@@ -98,7 +99,7 @@ void UPhATEdSkeletalMeshComponent::RenderAssetTools(const FSceneView* View, clas
 				if( (CollisionViewMode == FPhATSharedData::PRM_Solid && !bHitTest) || bHitTestAndBodyMode)
 				{
 					UMaterialInterface*	PrimMaterial = GetPrimitiveMaterial(i, KPT_Sphere, j, bHitTestAndBodyMode);
-					AggGeom->SphereElems[j].DrawElemSolid(PDI, ElemTM, Scale, PrimMaterial->GetRenderProxy(0));
+					AggGeom->SphereElems[j].DrawElemSolid(PDI, ElemTM, VectorScale, PrimMaterial->GetRenderProxy(0));
 				}
 
 				//wires are never used during hit
@@ -106,7 +107,7 @@ void UPhATEdSkeletalMeshComponent::RenderAssetTools(const FSceneView* View, clas
 				{
 					if (CollisionViewMode == FPhATSharedData::PRM_Solid || CollisionViewMode == FPhATSharedData::PRM_Wireframe)
 					{
-						AggGeom->SphereElems[j].DrawElemWire(PDI, ElemTM, Scale, GetPrimitiveColor(i, KPT_Sphere, j));
+						AggGeom->SphereElems[j].DrawElemWire(PDI, ElemTM, VectorScale, GetPrimitiveColor(i, KPT_Sphere, j));
 					}
 				}
 
@@ -129,14 +130,14 @@ void UPhATEdSkeletalMeshComponent::RenderAssetTools(const FSceneView* View, clas
 				if ( (CollisionViewMode == FPhATSharedData::PRM_Solid && !bHitTest) || bHitTestAndBodyMode)
 				{
 					UMaterialInterface*	PrimMaterial = GetPrimitiveMaterial(i, KPT_Box, j, bHitTestAndBodyMode);
-					AggGeom->BoxElems[j].DrawElemSolid(PDI, ElemTM, Scale, PrimMaterial->GetRenderProxy(0));
+					AggGeom->BoxElems[j].DrawElemSolid(PDI, ElemTM, VectorScale, PrimMaterial->GetRenderProxy(0));
 				}
 
 				if(!bHitTest)
 				{
 					if (CollisionViewMode == FPhATSharedData::PRM_Solid || CollisionViewMode == FPhATSharedData::PRM_Wireframe)
 					{
-						AggGeom->BoxElems[j].DrawElemWire(PDI, ElemTM, Scale, GetPrimitiveColor(i, KPT_Box, j));
+						AggGeom->BoxElems[j].DrawElemWire(PDI, ElemTM, VectorScale, GetPrimitiveColor(i, KPT_Box, j));
 					}
 				}
 				
@@ -159,14 +160,14 @@ void UPhATEdSkeletalMeshComponent::RenderAssetTools(const FSceneView* View, clas
 				if ( (CollisionViewMode == FPhATSharedData::PRM_Solid && !bHitTest) || bHitTestAndBodyMode)
 				{
 					UMaterialInterface*	PrimMaterial = GetPrimitiveMaterial(i, KPT_Sphyl, j, bHitTestAndBodyMode);
-					AggGeom->SphylElems[j].DrawElemSolid(PDI, ElemTM, Scale, PrimMaterial->GetRenderProxy(0));
+					AggGeom->SphylElems[j].DrawElemSolid(PDI, ElemTM, VectorScale, PrimMaterial->GetRenderProxy(0));
 				}
 
 				if(!bHitTest)
 				{
 					if (CollisionViewMode == FPhATSharedData::PRM_Solid || CollisionViewMode == FPhATSharedData::PRM_Wireframe)
 					{
-						AggGeom->SphylElems[j].DrawElemWire(PDI, ElemTM, Scale, GetPrimitiveColor(i, KPT_Sphyl, j));
+						AggGeom->SphylElems[j].DrawElemWire(PDI, ElemTM, VectorScale, GetPrimitiveColor(i, KPT_Sphyl, j));
 					}
 				}
 
diff --git a/Engine/Source/Editor/StaticMeshEditor/Private/StaticMeshEditorViewportClient.cpp b/Engine/Source/Editor/StaticMeshEditor/Private/StaticMeshEditorViewportClient.cpp
--- a/Engine/Source/Editor/StaticMeshEditor/Private/StaticMeshEditorViewportClient.cpp
+++ b/Engine/Source/Editor/StaticMeshEditor/Private/StaticMeshEditorViewportClient.cpp
@@ -437,6 +437,8 @@ void FStaticMeshEditorViewportClient::Draw(const FSceneView* View,FPrimitiveDraw
 		const FColor SelectedColor(149, 223, 157);
 		const FColor UnselectedColor(157, 149, 223);
 
+		const FVector VectorScaleOne(1.0f);
+
 		// Draw bodies
 		FKAggregateGeom* AggGeom = &StaticMesh->BodySetup->AggGeom;
 
@@ -448,7 +450,7 @@ void FStaticMeshEditorViewportClient::Draw(const FSceneView* View,FPrimitiveDraw
 			const FColor CollisionColor = StaticMeshEditorPtr.Pin()->IsSelectedPrim(HitProxy->PrimData) ? SelectedColor : UnselectedColor;
 			const FKSphereElem& SphereElem = AggGeom->SphereElems[i];
 			const FTransform ElemTM = SphereElem.GetTransform();
-			SphereElem.DrawElemWire(PDI, ElemTM, 1.f, CollisionColor);
+			SphereElem.DrawElemWire(PDI, ElemTM, VectorScaleOne, CollisionColor);
 
 			PDI->SetHitProxy(NULL);
 		}
@@ -461,7 +463,7 @@ void FStaticMeshEditorViewportClient::Draw(const FSceneView* View,FPrimitiveDraw
 			const FColor CollisionColor = StaticMeshEditorPtr.Pin()->IsSelectedPrim(HitProxy->PrimData) ? SelectedColor : UnselectedColor;
 			const FKBoxElem& BoxElem = AggGeom->BoxElems[i];
 			const FTransform ElemTM = BoxElem.GetTransform();
-			BoxElem.DrawElemWire(PDI, ElemTM, 1.f, CollisionColor);
+			BoxElem.DrawElemWire(PDI, ElemTM, VectorScaleOne, CollisionColor);
 
 			PDI->SetHitProxy(NULL);
 		}
@@ -474,7 +476,7 @@ void FStaticMeshEditorViewportClient::Draw(const FSceneView* View,FPrimitiveDraw
 			const FColor CollisionColor = StaticMeshEditorPtr.Pin()->IsSelectedPrim(HitProxy->PrimData) ? SelectedColor : UnselectedColor;
 			const FKSphylElem& SphylElem = AggGeom->SphylElems[i];
 			const FTransform ElemTM = SphylElem.GetTransform();
-			SphylElem.DrawElemWire(PDI, ElemTM, 1.f, CollisionColor);
+			SphylElem.DrawElemWire(PDI, ElemTM, VectorScaleOne, CollisionColor);
 
 			PDI->SetHitProxy(NULL);
 		}
diff --git a/Engine/Source/Runtime/Engine/Classes/PhysicsEngine/BoxElem.h b/Engine/Source/Runtime/Engine/Classes/PhysicsEngine/BoxElem.h
--- a/Engine/Source/Runtime/Engine/Classes/PhysicsEngine/BoxElem.h
+++ b/Engine/Source/Runtime/Engine/Classes/PhysicsEngine/BoxElem.h
@@ -85,11 +85,13 @@ struct FKBoxElem : public FKShapeElem
 		Center = InTransform.GetLocation();
 	}
 
-	FORCEINLINE float GetVolume(const FVector& Scale) const { float MinScale = Scale.GetMin(); return (X * MinScale) * (Y * MinScale) * (Z * MinScale); }
+	FORCEINLINE float GetVolume(const FVector& Scale3D) const { float MinScale = Scale3D.GetMin(); return (X * MinScale) * (Y * MinScale) * (Z * MinScale); }
+
+	ENGINE_API void	DrawElemWire(class FPrimitiveDrawInterface* PDI, const FTransform& ElemTM, const FVector& Scale3D, const FColor Color) const;
+	ENGINE_API void	DrawElemSolid(class FPrimitiveDrawInterface* PDI, const FTransform& ElemTM, const FVector& Scale3D, const FMaterialRenderProxy* MaterialRenderProxy) const;
+	
+	ENGINE_API void GetElemSolid(const FTransform& ElemTM, const FVector& Scale3D, const FMaterialRenderProxy* MaterialRenderProxy, int32 ViewIndex, FMeshElementCollector& Collector) const;
 
-	ENGINE_API void	DrawElemWire(class FPrimitiveDrawInterface* PDI, const FTransform& ElemTM, float Scale, const FColor Color) const;
-	ENGINE_API void	DrawElemSolid(class FPrimitiveDrawInterface* PDI, const FTransform& ElemTM, float Scale, const FMaterialRenderProxy* MaterialRenderProxy) const;
-	ENGINE_API void GetElemSolid(const FTransform& ElemTM, float Scale, const FMaterialRenderProxy* MaterialRenderProxy, int32 ViewIndex, FMeshElementCollector& Collector) const;
 	ENGINE_API FBox CalcAABB(const FTransform& BoneTM, float Scale) const;
 
 	ENGINE_API void ScaleElem(FVector DeltaSize, float MinSize);
diff --git a/Engine/Source/Runtime/Engine/Classes/PhysicsEngine/PhysicsAsset.h b/Engine/Source/Runtime/Engine/Classes/PhysicsEngine/PhysicsAsset.h
--- a/Engine/Source/Runtime/Engine/Classes/PhysicsEngine/PhysicsAsset.h
+++ b/Engine/Source/Runtime/Engine/Classes/PhysicsEngine/PhysicsAsset.h
@@ -93,7 +93,7 @@ public:
 #endif
 
 	// @todo document
-	void GetCollisionMesh(int32 ViewIndex, FMeshElementCollector& Collector, const USkeletalMesh* SkelMesh, const TArray<FTransform>& SpaceBases, const FTransform& LocalToWorld, float Scale);
+	void GetCollisionMesh(int32 ViewIndex, FMeshElementCollector& Collector, const USkeletalMesh* SkelMesh, const TArray<FTransform>& SpaceBases, const FTransform& LocalToWorld, const FVector& Scale3D);
 
 	// @todo document
 	void DrawConstraints(class FPrimitiveDrawInterface* PDI, const USkeletalMesh* SkelMesh, const TArray<FTransform>& SpaceBases, const FTransform& LocalToWorld, float Scale);
diff --git a/Engine/Source/Runtime/Engine/Classes/PhysicsEngine/SphereElem.h b/Engine/Source/Runtime/Engine/Classes/PhysicsEngine/SphereElem.h
--- a/Engine/Source/Runtime/Engine/Classes/PhysicsEngine/SphereElem.h
+++ b/Engine/Source/Runtime/Engine/Classes/PhysicsEngine/SphereElem.h
@@ -59,9 +59,10 @@ struct FKSphereElem : public FKShapeElem
 
 	FORCEINLINE float GetVolume(const FVector& Scale) const { return 1.3333f * PI * FMath::Pow(Radius * Scale.GetMin(), 3); }
 	
-	ENGINE_API void	DrawElemWire(class FPrimitiveDrawInterface* PDI, const FTransform& ElemTM, float Scale, const FColor Color) const;
-	ENGINE_API void	DrawElemSolid(class FPrimitiveDrawInterface* PDI, const FTransform& ElemTM, float Scale, const FMaterialRenderProxy* MaterialRenderProxy) const;
-	ENGINE_API void GetElemSolid(const FTransform& ElemTM, float Scale, const FMaterialRenderProxy* MaterialRenderProxy, int32 ViewIndex, class FMeshElementCollector& Collector) const;
+	ENGINE_API void	DrawElemWire(class FPrimitiveDrawInterface* PDI, const FTransform& ElemTM, const FVector& Scale3D, const FColor Color) const;
+
+	ENGINE_API void	DrawElemSolid(class FPrimitiveDrawInterface* PDI, const FTransform& ElemTM, const FVector& Scale3D, const FMaterialRenderProxy* MaterialRenderProxy) const;
+	ENGINE_API void GetElemSolid(const FTransform& ElemTM, const FVector& Scale3D, const FMaterialRenderProxy* MaterialRenderProxy, int32 ViewIndex, class FMeshElementCollector& Collector) const;
 	ENGINE_API FBox CalcAABB(const FTransform& BoneTM, float Scale) const;
 
 	ENGINE_API void ScaleElem(FVector DeltaSize, float MinSize);
diff --git a/Engine/Source/Runtime/Engine/Classes/PhysicsEngine/SphylElem.h b/Engine/Source/Runtime/Engine/Classes/PhysicsEngine/SphylElem.h
--- a/Engine/Source/Runtime/Engine/Classes/PhysicsEngine/SphylElem.h
+++ b/Engine/Source/Runtime/Engine/Classes/PhysicsEngine/SphylElem.h
@@ -75,12 +75,12 @@ struct FKSphylElem : public FKShapeElem
 
 	FORCEINLINE float GetVolume(const FVector& Scale) const { float ScaledRadius = Radius * Scale.GetMin(); return PI * FMath::Square(ScaledRadius) * ( 1.3333f * ScaledRadius + (Length * Scale.GetMin())); }
 
-	ENGINE_API void	DrawElemWire(class FPrimitiveDrawInterface* PDI, const FTransform& ElemTM, float Scale, const FColor Color) const;
-	ENGINE_API void GetElemSolid(const FTransform& ElemTM, float Scale, const FMaterialRenderProxy* MaterialRenderProxy, int32 ViewIndex, FMeshElementCollector& Collector) const;
-	ENGINE_API void	DrawElemSolid(class FPrimitiveDrawInterface* PDI, const FTransform& ElemTM, float Scale, const FMaterialRenderProxy* MaterialRenderProxy) const;
+	ENGINE_API void	DrawElemWire(class FPrimitiveDrawInterface* PDI, const FTransform& ElemTM, const FVector& Scale3D, const FColor Color) const;
+	ENGINE_API void GetElemSolid(const FTransform& ElemTM, const FVector& Scale3D, const FMaterialRenderProxy* MaterialRenderProxy, int32 ViewIndex, FMeshElementCollector& Collector) const;
+	ENGINE_API void	DrawElemSolid(class FPrimitiveDrawInterface* PDI, const FTransform& ElemTM, const FVector& Scale3D, const FMaterialRenderProxy* MaterialRenderProxy) const;
 	ENGINE_API FBox CalcAABB(const FTransform& BoneTM, float Scale) const;
 
 	ENGINE_API void ScaleElem(FVector DeltaSize, float MinSize);
 
 	ENGINE_API static EAggCollisionShape::Type StaticShapeType;
-};
\ No newline at end of file
+};
diff --git a/Engine/Source/Runtime/Engine/Private/PhysicsEngine/PhysDrawing.cpp b/Engine/Source/Runtime/Engine/Private/PhysicsEngine/PhysDrawing.cpp
--- a/Engine/Source/Runtime/Engine/Private/PhysicsEngine/PhysDrawing.cpp
+++ b/Engine/Source/Runtime/Engine/Private/PhysicsEngine/PhysDrawing.cpp
@@ -30,77 +30,77 @@ static const FColor JointLockedColor(255,128,10);
 // FKSphereElem
 /////////////////////////////////////////////////////////////////////////////////////
 
-// NB: ElemTM is assumed to have no scaling in it!
-void FKSphereElem::DrawElemWire(FPrimitiveDrawInterface* PDI, const FTransform& ElemTM, float Scale, const FColor Color) const
+void FKSphereElem::DrawElemWire(class FPrimitiveDrawInterface* PDI, const FTransform& ElemTM, const FVector& Scale3D, const FColor Color) const
 {
 	FVector ElemCenter = ElemTM.GetLocation();
-	FVector X = ElemTM.GetScaledAxis( EAxis::X );
-	FVector Y = ElemTM.GetScaledAxis( EAxis::Y );
-	FVector Z = ElemTM.GetScaledAxis( EAxis::Z );
+	FVector X = ElemTM.GetScaledAxis(EAxis::X);
+	FVector Y = ElemTM.GetScaledAxis(EAxis::Y);
+	FVector Z = ElemTM.GetScaledAxis(EAxis::Z);
 
-	DrawCircle(PDI,ElemCenter, X, Y, Color, Scale*Radius, DrawCollisionSides, SDPG_World);
-	DrawCircle(PDI,ElemCenter, X, Z, Color, Scale*Radius, DrawCollisionSides, SDPG_World);
-	DrawCircle(PDI,ElemCenter, Y, Z, Color, Scale*Radius, DrawCollisionSides, SDPG_World);
+	const float ScaleRadius = Scale3D.GetAbsMin();
+
+	DrawCircle(PDI, ElemCenter, X, Y, Color, ScaleRadius*Radius, DrawCollisionSides, SDPG_World);
+	DrawCircle(PDI, ElemCenter, X, Z, Color, ScaleRadius*Radius, DrawCollisionSides, SDPG_World);
+	DrawCircle(PDI, ElemCenter, Y, Z, Color, ScaleRadius*Radius, DrawCollisionSides, SDPG_World);
 }
 
-void FKSphereElem::DrawElemSolid(class FPrimitiveDrawInterface* PDI, const FTransform& ElemTM, float Scale, const FMaterialRenderProxy* MaterialRenderProxy) const
+
+void FKSphereElem::DrawElemSolid(class FPrimitiveDrawInterface* PDI, const FTransform& ElemTM, const FVector& Scale3D, const FMaterialRenderProxy* MaterialRenderProxy) const
 {
-	DrawSphere(PDI, ElemTM.GetLocation(), FVector( this->Radius * Scale ), DrawCollisionSides, DrawCollisionSides/2, MaterialRenderProxy, SDPG_World );
+	DrawSphere(PDI, ElemTM.GetLocation(), FVector(this->Radius * Scale3D.GetAbsMin()), DrawCollisionSides, DrawCollisionSides / 2, MaterialRenderProxy, SDPG_World);
 }
 
-void FKSphereElem::GetElemSolid(const FTransform& ElemTM, float Scale, const FMaterialRenderProxy* MaterialRenderProxy, int32 ViewIndex, FMeshElementCollector& Collector) const
+void FKSphereElem::GetElemSolid(const FTransform& ElemTM, const FVector& Scale3D, const FMaterialRenderProxy* MaterialRenderProxy, int32 ViewIndex, FMeshElementCollector& Collector) const
 {
-	GetSphereMesh(ElemTM.GetLocation(), FVector( this->Radius * Scale ), DrawCollisionSides, DrawCollisionSides/2, MaterialRenderProxy, SDPG_World, false, ViewIndex, Collector );
+	GetSphereMesh(ElemTM.GetLocation(), FVector(this->Radius * Scale3D.GetAbsMin()), DrawCollisionSides, DrawCollisionSides / 2, MaterialRenderProxy, SDPG_World, false, ViewIndex, Collector);
 }
 
-
-/////////////////////////////////////////////////////////////////////////////////////
-// FKBoxElem
-/////////////////////////////////////////////////////////////////////////////////////
-
-// NB: ElemTM is assumed to have no scaling in it!
-void FKBoxElem::DrawElemWire(FPrimitiveDrawInterface* PDI, const FTransform& ElemTM, float Scale, const FColor Color) const
+void FKBoxElem::DrawElemWire(class FPrimitiveDrawInterface* PDI, const FTransform& ElemTM, const FVector& Scale3D, const FColor Color) const
 {
 	FVector	B[2], P, Q, Radii;
 
 	// X,Y,Z member variables are LENGTH not RADIUS
-	Radii.X = Scale*0.5f*X;
-	Radii.Y = Scale*0.5f*Y;
-	Radii.Z = Scale*0.5f*Z;
+	Radii.X = Scale3D.X*0.5f*X;
+	Radii.Y = Scale3D.Y*0.5f*Y;
+	Radii.Z = Scale3D.Z*0.5f*Z;
 
 	B[0] = Radii; // max
 	B[1] = -1.0f * Radii; // min
 
-	for( int32 i=0; i<2; i++ )
+	for (int32 i = 0; i < 2; i++)
 	{
-		for( int32 j=0; j<2; j++ )
+		for (int32 j = 0; j < 2; j++)
 		{
-			P.X=B[i].X; Q.X=B[i].X;
-			P.Y=B[j].Y; Q.Y=B[j].Y;
-			P.Z=B[0].Z; Q.Z=B[1].Z;
-			PDI->DrawLine( ElemTM.TransformPosition(P), ElemTM.TransformPosition(Q), Color, SDPG_World);
-
-			P.Y=B[i].Y; Q.Y=B[i].Y;
-			P.Z=B[j].Z; Q.Z=B[j].Z;
-			P.X=B[0].X; Q.X=B[1].X;
-			PDI->DrawLine( ElemTM.TransformPosition(P), ElemTM.TransformPosition(Q), Color, SDPG_World);
-
-			P.Z=B[i].Z; Q.Z=B[i].Z;
-			P.X=B[j].X; Q.X=B[j].X;
-			P.Y=B[0].Y; Q.Y=B[1].Y;
-			PDI->DrawLine( ElemTM.TransformPosition(P), ElemTM.TransformPosition(Q), Color, SDPG_World);
+			P.X = B[i].X; Q.X = B[i].X;
+			P.Y = B[j].Y; Q.Y = B[j].Y;
+			P.Z = B[0].Z; Q.Z = B[1].Z;
+			PDI->DrawLine(ElemTM.TransformPosition(P), ElemTM.TransformPosition(Q), Color, SDPG_World);
+
+			P.Y = B[i].Y; Q.Y = B[i].Y;
+			P.Z = B[j].Z; Q.Z = B[j].Z;
+			P.X = B[0].X; Q.X = B[1].X;
+			PDI->DrawLine(ElemTM.TransformPosition(P), ElemTM.TransformPosition(Q), Color, SDPG_World);
+
+			P.Z = B[i].Z; Q.Z = B[i].Z;
+			P.X = B[j].X; Q.X = B[j].X;
+			P.Y = B[0].Y; Q.Y = B[1].Y;
+			PDI->DrawLine(ElemTM.TransformPosition(P), ElemTM.TransformPosition(Q), Color, SDPG_World);
 		}
 	}
 }
 
-void FKBoxElem::DrawElemSolid(class FPrimitiveDrawInterface* PDI, const FTransform& ElemTM, float Scale, const FMaterialRenderProxy* MaterialRenderProxy) const
+/////////////////////////////////////////////////////////////////////////////////////
+// FKBoxElem
+/////////////////////////////////////////////////////////////////////////////////////
+
+void FKBoxElem::DrawElemSolid(class FPrimitiveDrawInterface* PDI, const FTransform& ElemTM, const FVector& Scale3D, const FMaterialRenderProxy* MaterialRenderProxy) const
 {
-	DrawBox(PDI, ElemTM.ToMatrixWithScale(), Scale * 0.5f * FVector(X, Y, Z), MaterialRenderProxy, SDPG_World );
+	DrawBox(PDI, ElemTM.ToMatrixWithScale(), Scale3D * 0.5f * FVector(X, Y, Z), MaterialRenderProxy, SDPG_World);
 }
 
-void FKBoxElem::GetElemSolid(const FTransform& ElemTM, float Scale, const FMaterialRenderProxy* MaterialRenderProxy, int32 ViewIndex, FMeshElementCollector& Collector) const
+void FKBoxElem::GetElemSolid(const FTransform& ElemTM, const FVector& Scale3D, const FMaterialRenderProxy* MaterialRenderProxy, int32 ViewIndex, FMeshElementCollector& Collector) const
 {
-	GetBoxMesh(ElemTM.ToMatrixWithScale(), Scale * 0.5f * FVector(X, Y, Z), MaterialRenderProxy, SDPG_World, ViewIndex, Collector);
+	GetBoxMesh(ElemTM.ToMatrixWithScale(), Scale3D * 0.5f * FVector(X, Y, Z), MaterialRenderProxy, SDPG_World, ViewIndex, Collector);
 }
 
 /////////////////////////////////////////////////////////////////////////////////////
@@ -120,102 +120,109 @@ static void DrawHalfCircle(FPrimitiveDrawInterface* PDI, const FVector& Base, co
 	}	
 }
 
-// NB: ElemTM is assumed to have no scaling in it!
-void FKSphylElem::DrawElemWire(FPrimitiveDrawInterface* PDI, const FTransform& ElemTM, float Scale, const FColor Color) const
+void FKSphylElem::DrawElemWire(class FPrimitiveDrawInterface* PDI, const FTransform& ElemTM, const FVector& Scale3D, const FColor Color) const
 {
+	const FVector Scale3DAbs = Scale3D.GetAbs();
+	const float ScaleRadius = FMath::Max(Scale3DAbs.X, Scale3DAbs.Y);
+	const float ScaleLength = Scale3DAbs.Z;
+
 	FVector Origin = ElemTM.GetLocation();
-	FVector XAxis = ElemTM.GetScaledAxis( EAxis::X );
-	FVector YAxis = ElemTM.GetScaledAxis( EAxis::Y );
-	FVector ZAxis = ElemTM.GetScaledAxis( EAxis::Z );
+	FVector XAxis = ElemTM.GetScaledAxis(EAxis::X);
+	FVector YAxis = ElemTM.GetScaledAxis(EAxis::Y);
+	FVector ZAxis = ElemTM.GetScaledAxis(EAxis::Z);
 
 	// Draw top and bottom circles
-	FVector TopEnd = Origin + Scale*0.5f*Length*ZAxis;
-	FVector BottomEnd = Origin - Scale*0.5f*Length*ZAxis;
+	FVector TopEnd = Origin + ScaleLength*0.5f*Length*ZAxis;
+	FVector BottomEnd = Origin - ScaleLength*0.5f*Length*ZAxis;
 
-	DrawCircle(PDI,TopEnd, XAxis, YAxis, Color, Scale*Radius, DrawCollisionSides, SDPG_World);
-	DrawCircle(PDI,BottomEnd, XAxis, YAxis, Color, Scale*Radius, DrawCollisionSides, SDPG_World);
+
+	DrawCircle(PDI, TopEnd, XAxis, YAxis, Color, ScaleRadius*Radius, DrawCollisionSides, SDPG_World);
+	DrawCircle(PDI, BottomEnd, XAxis, YAxis, Color, ScaleRadius*Radius, DrawCollisionSides, SDPG_World);
 
 	// Draw domed caps
-	DrawHalfCircle(PDI, TopEnd, YAxis, ZAxis, Color,Scale* Radius);
-	DrawHalfCircle(PDI, TopEnd, XAxis, ZAxis, Color, Scale*Radius);
+	DrawHalfCircle(PDI, TopEnd, YAxis, ZAxis, Color, ScaleRadius* Radius);
+	DrawHalfCircle(PDI, TopEnd, XAxis, ZAxis, Color, ScaleRadius*Radius);
 
 	FVector NegZAxis = -ZAxis;
 
-	DrawHalfCircle(PDI, BottomEnd, YAxis, NegZAxis, Color, Scale*Radius);
-	DrawHalfCircle(PDI, BottomEnd, XAxis, NegZAxis, Color, Scale*Radius);
+	DrawHalfCircle(PDI, BottomEnd, YAxis, NegZAxis, Color, ScaleRadius*Radius);
+	DrawHalfCircle(PDI, BottomEnd, XAxis, NegZAxis, Color, ScaleRadius*Radius);
 
 	// Draw connecty lines
-	PDI->DrawLine(TopEnd + Scale*Radius*XAxis, BottomEnd + Scale*Radius*XAxis, Color, SDPG_World);
-	PDI->DrawLine(TopEnd - Scale*Radius*XAxis, BottomEnd - Scale*Radius*XAxis, Color, SDPG_World);
-	PDI->DrawLine(TopEnd + Scale*Radius*YAxis, BottomEnd + Scale*Radius*YAxis, Color, SDPG_World);
-	PDI->DrawLine(TopEnd - Scale*Radius*YAxis, BottomEnd - Scale*Radius*YAxis, Color, SDPG_World);
+	PDI->DrawLine(TopEnd + ScaleRadius*Radius*XAxis, BottomEnd + ScaleRadius*Radius*XAxis, Color, SDPG_World);
+	PDI->DrawLine(TopEnd - ScaleRadius*Radius*XAxis, BottomEnd - ScaleRadius*Radius*XAxis, Color, SDPG_World);
+	PDI->DrawLine(TopEnd + ScaleRadius*Radius*YAxis, BottomEnd + ScaleRadius*Radius*YAxis, Color, SDPG_World);
+	PDI->DrawLine(TopEnd - ScaleRadius*Radius*YAxis, BottomEnd - ScaleRadius*Radius*YAxis, Color, SDPG_World);
 }
 
-
-void FKSphylElem::GetElemSolid(const FTransform& ElemTM, float Scale, const FMaterialRenderProxy* MaterialRenderProxy, int32 ViewIndex, FMeshElementCollector& Collector) const
+void FKSphylElem::GetElemSolid(const FTransform& ElemTM, const FVector& Scale3D, const FMaterialRenderProxy* MaterialRenderProxy, int32 ViewIndex, FMeshElementCollector& Collector) const
 {
+	const FVector Scale3DAbs = Scale3D.GetAbs();
+	const float ScaleRadius = FMath::Max(Scale3DAbs.X, Scale3DAbs.Y);
+	const float ScaleLength = Scale3DAbs.Z;
+
 	const int32 NumSides = DrawCollisionSides;
-	const int32 NumRings = (DrawCollisionSides/2) + 1;
+	const int32 NumRings = (DrawCollisionSides / 2) + 1;
 
 	// The first/last arc are on top of each other.
-	const int32 NumVerts = (NumSides+1) * (NumRings+1);
-	FDynamicMeshVertex* Verts = (FDynamicMeshVertex*)FMemory::Malloc( NumVerts * sizeof(FDynamicMeshVertex) );
+	const int32 NumVerts = (NumSides + 1) * (NumRings + 1);
+	FDynamicMeshVertex* Verts = (FDynamicMeshVertex*)FMemory::Malloc(NumVerts * sizeof(FDynamicMeshVertex));
 
 	// Calculate verts for one arc
-	FDynamicMeshVertex* ArcVerts = (FDynamicMeshVertex*)FMemory::Malloc( (NumRings+1) * sizeof(FDynamicMeshVertex) );
+	FDynamicMeshVertex* ArcVerts = (FDynamicMeshVertex*)FMemory::Malloc((NumRings + 1) * sizeof(FDynamicMeshVertex));
 
-	for(int32 RingIdx=0; RingIdx<NumRings+1; RingIdx++)
+	for (int32 RingIdx = 0; RingIdx < NumRings + 1; RingIdx++)
 	{
 		FDynamicMeshVertex* ArcVert = &ArcVerts[RingIdx];
 
 		float Angle;
 		float ZOffset;
-		if( RingIdx <= DrawCollisionSides/4 )
+		if (RingIdx <= DrawCollisionSides / 4)
 		{
-			Angle = ((float)RingIdx/(NumRings-1)) * PI;
-			ZOffset = 0.5 * Scale * Length;
+			Angle = ((float)RingIdx / (NumRings - 1)) * PI;
+			ZOffset = 0.5 * ScaleLength * Length;
 		}
 		else
 		{
-			Angle = ((float)(RingIdx-1)/(NumRings-1)) * PI;
-			ZOffset = -0.5 * Scale * Length;
+			Angle = ((float)(RingIdx - 1) / (NumRings - 1)) * PI;
+			ZOffset = -0.5 * ScaleLength * Length;
 		}
 
 		// Note- unit sphere, so position always has mag of one. We can just use it for normal!		
 		FVector SpherePos;
 		SpherePos.X = 0.0f;
-		SpherePos.Y = Scale * Radius * FMath::Sin(Angle);
-		SpherePos.Z = Scale * Radius * FMath::Cos(Angle);
+		SpherePos.Y = ScaleRadius * Radius * FMath::Sin(Angle);
+		SpherePos.Z = ScaleRadius * Radius * FMath::Cos(Angle);
 
-		ArcVert->Position = SpherePos + FVector(0,0,ZOffset);
+		ArcVert->Position = SpherePos + FVector(0, 0, ZOffset);
 
 		ArcVert->SetTangents(
-			FVector(1,0,0),
+			FVector(1, 0, 0),
 			FVector(0.0f, -SpherePos.Z, SpherePos.Y),
 			SpherePos
 			);
 
 		ArcVert->TextureCoordinate.X = 0.0f;
-		ArcVert->TextureCoordinate.Y = ((float)RingIdx/NumRings);
+		ArcVert->TextureCoordinate.Y = ((float)RingIdx / NumRings);
 	}
 
 	// Then rotate this arc NumSides+1 times.
-	for(int32 SideIdx=0; SideIdx<NumSides+1; SideIdx++)
+	for (int32 SideIdx = 0; SideIdx < NumSides + 1; SideIdx++)
 	{
-		const FRotator ArcRotator(0, 360.f * ((float)SideIdx/NumSides), 0);
-		const FRotationMatrix ArcRot( ArcRotator );
-		const float XTexCoord = ((float)SideIdx/NumSides);
+		const FRotator ArcRotator(0, 360.f * ((float)SideIdx / NumSides), 0);
+		const FRotationMatrix ArcRot(ArcRotator);
+		const float XTexCoord = ((float)SideIdx / NumSides);
 
-		for(int32 VertIdx=0; VertIdx<NumRings+1; VertIdx++)
+		for (int32 VertIdx = 0; VertIdx < NumRings + 1; VertIdx++)
 		{
-			int32 VIx = (NumRings+1)*SideIdx + VertIdx;
+			int32 VIx = (NumRings + 1)*SideIdx + VertIdx;
 
-			Verts[VIx].Position = ArcRot.TransformPosition( ArcVerts[VertIdx].Position );
+			Verts[VIx].Position = ArcRot.TransformPosition(ArcVerts[VertIdx].Position);
 
 			Verts[VIx].SetTangents(
-				ArcRot.TransformVector( ArcVerts[VertIdx].TangentX ),
-				ArcRot.TransformVector( ArcVerts[VertIdx].GetTangentY() ),
-				ArcRot.TransformVector( ArcVerts[VertIdx].TangentZ )
+				ArcRot.TransformVector(ArcVerts[VertIdx].TangentX),
+				ArcRot.TransformVector(ArcVerts[VertIdx].GetTangentY()),
+				ArcRot.TransformVector(ArcVerts[VertIdx].TangentZ)
 				);
 
 			Verts[VIx].TextureCoordinate.X = XTexCoord;
@@ -226,18 +233,18 @@ void FKSphylElem::GetElemSolid(const FTransform& ElemTM, float Scale, const FMat
 	FDynamicMeshBuilder MeshBuilder;
 	{
 		// Add all of the vertices to the mesh.
-		for(int32 VertIdx=0; VertIdx<NumVerts; VertIdx++)
+		for (int32 VertIdx = 0; VertIdx < NumVerts; VertIdx++)
 		{
 			MeshBuilder.AddVertex(Verts[VertIdx]);
 		}
 
 		// Add all of the triangles to the mesh.
-		for(int32 SideIdx=0; SideIdx<NumSides; SideIdx++)
+		for (int32 SideIdx = 0; SideIdx < NumSides; SideIdx++)
 		{
-			const int32 a0start = (SideIdx+0) * (NumRings+1);
-			const int32 a1start = (SideIdx+1) * (NumRings+1);
+			const int32 a0start = (SideIdx + 0) * (NumRings + 1);
+			const int32 a1start = (SideIdx + 1) * (NumRings + 1);
 
-			for(int32 RingIdx=0; RingIdx<NumRings; RingIdx++)
+			for (int32 RingIdx = 0; RingIdx < NumRings; RingIdx++)
 			{
 				MeshBuilder.AddTriangle(a0start + RingIdx + 0, a1start + RingIdx + 0, a0start + RingIdx + 1);
 				MeshBuilder.AddTriangle(a1start + RingIdx + 0, a1start + RingIdx + 1, a0start + RingIdx + 1);
@@ -251,70 +258,74 @@ void FKSphylElem::GetElemSolid(const FTransform& ElemTM, float Scale, const FMat
 	FMemory::Free(ArcVerts);
 }
 
-void FKSphylElem::DrawElemSolid(FPrimitiveDrawInterface* PDI, const FTransform& ElemTM, float Scale, const FMaterialRenderProxy* MaterialRenderProxy) const
+void FKSphylElem::DrawElemSolid(FPrimitiveDrawInterface* PDI, const FTransform& ElemTM, const FVector& Scale3D, const FMaterialRenderProxy* MaterialRenderProxy) const
 {
+	const FVector Scale3DAbs = Scale3D.GetAbs();
+	const float ScaleRadius = FMath::Max(Scale3DAbs.X, Scale3DAbs.Y);
+	const float ScaleLength = Scale3DAbs.Z;
+
 	const int32 NumSides = DrawCollisionSides;
-	const int32 NumRings = (DrawCollisionSides/2) + 1;
+	const int32 NumRings = (DrawCollisionSides / 2) + 1;
 
 	// The first/last arc are on top of each other.
-	const int32 NumVerts = (NumSides+1) * (NumRings+1);
-	FDynamicMeshVertex* Verts = (FDynamicMeshVertex*)FMemory::Malloc( NumVerts * sizeof(FDynamicMeshVertex) );
+	const int32 NumVerts = (NumSides + 1) * (NumRings + 1);
+	FDynamicMeshVertex* Verts = (FDynamicMeshVertex*)FMemory::Malloc(NumVerts * sizeof(FDynamicMeshVertex));
 
 	// Calculate verts for one arc
-	FDynamicMeshVertex* ArcVerts = (FDynamicMeshVertex*)FMemory::Malloc( (NumRings+1) * sizeof(FDynamicMeshVertex) );
+	FDynamicMeshVertex* ArcVerts = (FDynamicMeshVertex*)FMemory::Malloc((NumRings + 1) * sizeof(FDynamicMeshVertex));
 
-	for(int32 RingIdx=0; RingIdx<NumRings+1; RingIdx++)
+	for (int32 RingIdx = 0; RingIdx < NumRings + 1; RingIdx++)
 	{
 		FDynamicMeshVertex* ArcVert = &ArcVerts[RingIdx];
 
 		float Angle;
 		float ZOffset;
-		if( RingIdx <= DrawCollisionSides/4 )
+		if (RingIdx <= DrawCollisionSides / 4)
 		{
-			Angle = ((float)RingIdx/(NumRings-1)) * PI;
-			ZOffset = 0.5 * Scale * Length;
+			Angle = ((float)RingIdx / (NumRings - 1)) * PI;
+			ZOffset = 0.5 * ScaleLength * Length;
 		}
 		else
 		{
-			Angle = ((float)(RingIdx-1)/(NumRings-1)) * PI;
-			ZOffset = -0.5 * Scale * Length;
+			Angle = ((float)(RingIdx - 1) / (NumRings - 1)) * PI;
+			ZOffset = -0.5 * ScaleLength * Length;
 		}
 
 		// Note- unit sphere, so position always has mag of one. We can just use it for normal!		
 		FVector SpherePos;
 		SpherePos.X = 0.0f;
-		SpherePos.Y = Scale * Radius * FMath::Sin(Angle);
-		SpherePos.Z = Scale * Radius * FMath::Cos(Angle);
+		SpherePos.Y = ScaleRadius * Radius * FMath::Sin(Angle);
+		SpherePos.Z = ScaleRadius * Radius * FMath::Cos(Angle);
 
-		ArcVert->Position = SpherePos + FVector(0,0,ZOffset);
+		ArcVert->Position = SpherePos + FVector(0, 0, ZOffset);
 
 		ArcVert->SetTangents(
-			FVector(1,0,0),
+			FVector(1, 0, 0),
 			FVector(0.0f, -SpherePos.Z, SpherePos.Y),
 			SpherePos
 			);
 
 		ArcVert->TextureCoordinate.X = 0.0f;
-		ArcVert->TextureCoordinate.Y = ((float)RingIdx/NumRings);
+		ArcVert->TextureCoordinate.Y = ((float)RingIdx / NumRings);
 	}
 
 	// Then rotate this arc NumSides+1 times.
-	for(int32 SideIdx=0; SideIdx<NumSides+1; SideIdx++)
+	for (int32 SideIdx = 0; SideIdx < NumSides + 1; SideIdx++)
 	{
-		const FRotator ArcRotator(0, 360.f * ((float)SideIdx/NumSides), 0);
-		const FRotationMatrix ArcRot( ArcRotator );
-		const float XTexCoord = ((float)SideIdx/NumSides);
+		const FRotator ArcRotator(0, 360.f * ((float)SideIdx / NumSides), 0);
+		const FRotationMatrix ArcRot(ArcRotator);
+		const float XTexCoord = ((float)SideIdx / NumSides);
 
-		for(int32 VertIdx=0; VertIdx<NumRings+1; VertIdx++)
+		for (int32 VertIdx = 0; VertIdx < NumRings + 1; VertIdx++)
 		{
-			int32 VIx = (NumRings+1)*SideIdx + VertIdx;
+			int32 VIx = (NumRings + 1)*SideIdx + VertIdx;
 
-			Verts[VIx].Position = ArcRot.TransformPosition( ArcVerts[VertIdx].Position );
+			Verts[VIx].Position = ArcRot.TransformPosition(ArcVerts[VertIdx].Position);
 
 			Verts[VIx].SetTangents(
-				ArcRot.TransformVector( ArcVerts[VertIdx].TangentX ),
-				ArcRot.TransformVector( ArcVerts[VertIdx].GetTangentY() ),
-				ArcRot.TransformVector( ArcVerts[VertIdx].TangentZ )
+				ArcRot.TransformVector(ArcVerts[VertIdx].TangentX),
+				ArcRot.TransformVector(ArcVerts[VertIdx].GetTangentY()),
+				ArcRot.TransformVector(ArcVerts[VertIdx].TangentZ)
 				);
 
 			Verts[VIx].TextureCoordinate.X = XTexCoord;
@@ -325,18 +336,18 @@ void FKSphylElem::DrawElemSolid(FPrimitiveDrawInterface* PDI, const FTransform&
 	FDynamicMeshBuilder MeshBuilder;
 	{
 		// Add all of the vertices to the mesh.
-		for(int32 VertIdx=0; VertIdx<NumVerts; VertIdx++)
+		for (int32 VertIdx = 0; VertIdx < NumVerts; VertIdx++)
 		{
 			MeshBuilder.AddVertex(Verts[VertIdx]);
 		}
 
 		// Add all of the triangles to the mesh.
-		for(int32 SideIdx=0; SideIdx<NumSides; SideIdx++)
+		for (int32 SideIdx = 0; SideIdx < NumSides; SideIdx++)
 		{
-			const int32 a0start = (SideIdx+0) * (NumRings+1);
-			const int32 a1start = (SideIdx+1) * (NumRings+1);
+			const int32 a0start = (SideIdx + 0) * (NumRings + 1);
+			const int32 a1start = (SideIdx + 1) * (NumRings + 1);
 
-			for(int32 RingIdx=0; RingIdx<NumRings; RingIdx++)
+			for (int32 RingIdx = 0; RingIdx < NumRings; RingIdx++)
 			{
 				MeshBuilder.AddTriangle(a0start + RingIdx + 0, a1start + RingIdx + 0, a0start + RingIdx + 1);
 				MeshBuilder.AddTriangle(a1start + RingIdx + 0, a1start + RingIdx + 1, a0start + RingIdx + 1);
@@ -541,45 +552,42 @@ void FKAggregateGeom::GetAggGeom(const FTransform& Transform, const FColor Color
 	FTransform ParentTM = Transform;
 	ParentTM.RemoveScaling();
 
-	if( Scale3D.GetAbs().IsUniform() )
+	for (int32 i = 0; i < SphereElems.Num(); i++)
 	{
-		for(int32 i=0; i<SphereElems.Num(); i++)
-		{
-			FTransform ElemTM = SphereElems[i].GetTransform();
-			ElemTM.ScaleTranslation(Scale3D);
-			ElemTM *= ParentTM;
-
-			if(bDrawSolid)
-				SphereElems[i].GetElemSolid(ElemTM, Scale3D.X, MatInst, ViewIndex, Collector);
-			else
-				SphereElems[i].DrawElemWire(Collector.GetPDI(ViewIndex), ElemTM, Scale3D.X, Color);
-		}
+		FTransform ElemTM = SphereElems[i].GetTransform();
+		ElemTM.ScaleTranslation(Scale3D);
+		ElemTM *= ParentTM;
 
-		for(int32 i=0; i<BoxElems.Num(); i++)
-		{
-			FTransform ElemTM = BoxElems[i].GetTransform();
-			ElemTM.ScaleTranslation(Scale3D);
-			ElemTM *= ParentTM;
-
-			if(bDrawSolid)
-				BoxElems[i].GetElemSolid(ElemTM, Scale3D.X, MatInst, ViewIndex, Collector);
-			else
-				BoxElems[i].DrawElemWire(Collector.GetPDI(ViewIndex), ElemTM, Scale3D.X, Color);
-		}
+		if (bDrawSolid)
+			SphereElems[i].GetElemSolid(ElemTM, Scale3D, MatInst, ViewIndex, Collector);
+		else
+			SphereElems[i].DrawElemWire(Collector.GetPDI(ViewIndex), ElemTM, Scale3D, Color);
+	}
 
-		for(int32 i=0; i<SphylElems.Num(); i++)
-		{
-			FTransform ElemTM = SphylElems[i].GetTransform();
-			ElemTM.ScaleTranslation(Scale3D);
-			ElemTM *= ParentTM;
-
-			if(bDrawSolid)
-				SphylElems[i].GetElemSolid(ElemTM, Scale3D.X, MatInst, ViewIndex, Collector);
-			else
-				SphylElems[i].DrawElemWire(Collector.GetPDI(ViewIndex), ElemTM, Scale3D.X, Color);
-		}
+	for (int32 i = 0; i < BoxElems.Num(); i++)
+	{
+		FTransform ElemTM = BoxElems[i].GetTransform();
+		ElemTM.ScaleTranslation(Scale3D);
+		ElemTM *= ParentTM;
+
+		if (bDrawSolid)
+			BoxElems[i].GetElemSolid(ElemTM, Scale3D, MatInst, ViewIndex, Collector);
+		else
+			BoxElems[i].DrawElemWire(Collector.GetPDI(ViewIndex), ElemTM, Scale3D, Color);
 	}
 
+	for (int32 i = 0; i < SphylElems.Num(); i++)
+	{
+		FTransform ElemTM = SphylElems[i].GetTransform();
+		ElemTM.ScaleTranslation(Scale3D);
+		ElemTM *= ParentTM;
+
+		if (bDrawSolid)
+			SphylElems[i].GetElemSolid(ElemTM, Scale3D, MatInst, ViewIndex, Collector);
+		else
+			SphylElems[i].DrawElemWire(Collector.GetPDI(ViewIndex), ElemTM, Scale3D, Color);
+	}
+	
 	if(ConvexElems.Num() > 0)
 	{
 		if(bDrawSolid)
@@ -707,7 +715,7 @@ FTransform GetSkelBoneTransform(int32 BoneIndex, const TArray<FTransform>& Space
 	}
 }
 
-void UPhysicsAsset::GetCollisionMesh(int32 ViewIndex, FMeshElementCollector& Collector, const USkeletalMesh* SkelMesh, const TArray<FTransform>& SpaceBases, const FTransform& LocalToWorld, float Scale)
+void UPhysicsAsset::GetCollisionMesh(int32 ViewIndex, FMeshElementCollector& Collector, const USkeletalMesh* SkelMesh, const TArray<FTransform>& SpaceBases, const FTransform& LocalToWorld, const FVector& Scale3D)
 {
 	for( int32 i=0; i<BodySetup.Num(); i++)
 	{
@@ -716,7 +724,7 @@ void UPhysicsAsset::GetCollisionMesh(int32 ViewIndex, FMeshElementCollector& Col
 		FColor* BoneColor = (FColor*)( &BodySetup[i] );
 
 		FTransform BoneTransform = GetSkelBoneTransform(BoneIndex, SpaceBases, LocalToWorld);
-		BoneTransform.SetScale3D(FVector(Scale));
+		BoneTransform.SetScale3D(Scale3D);
 		BodySetup[i]->CreatePhysicsMeshes();
 		BodySetup[i]->AggGeom.GetAggGeom(BoneTransform, *BoneColor, NULL, false, false, false, ViewIndex, Collector);
 	}
diff --git a/Engine/Source/Runtime/Engine/Private/SkeletalMesh.cpp b/Engine/Source/Runtime/Engine/Private/SkeletalMesh.cpp
--- a/Engine/Source/Runtime/Engine/Private/SkeletalMesh.cpp
+++ b/Engine/Source/Runtime/Engine/Private/SkeletalMesh.cpp
@@ -4621,8 +4621,8 @@ void FSkeletalMeshSceneProxy::DebugDrawPhysicsAsset(int32 ViewIndex, FMeshElemen
 	FMatrix ScalingMatrix = LocalToWorld;
 	FVector TotalScale = ScalingMatrix.ExtractScaling();
 
-	// Only valid if scaling if uniform.
-	if( TotalScale.IsUniform() && !TotalScale.IsNearlyZero() )
+	// Only if valid
+	if( !TotalScale.IsNearlyZero() )
 	{
 		FTransform LocalToWorldTransform(LocalToWorld);
 
@@ -4632,7 +4632,7 @@ void FSkeletalMeshSceneProxy::DebugDrawPhysicsAsset(int32 ViewIndex, FMeshElemen
 			check(PhysicsAssetForDebug);
 			if( EngineShowFlags.Collision && IsCollisionEnabled() )
 			{
-				PhysicsAssetForDebug->GetCollisionMesh(ViewIndex, Collector, SkeletalMeshForDebug, *BoneSpaceBases, LocalToWorldTransform, TotalScale.X);
+				PhysicsAssetForDebug->GetCollisionMesh(ViewIndex, Collector, SkeletalMeshForDebug, *BoneSpaceBases, LocalToWorldTransform, TotalScale);
 			}
 			if( EngineShowFlags.Constraints )
 			{
