Default colour for modulated shadows.
Modulated shadows draw before translucency.

diff --git a/Engine/Source/Runtime/Engine/Private/Components/LightComponent.cpp b/Engine/Source/Runtime/Engine/Private/Components/LightComponent.cpp
--- a/Engine/Source/Runtime/Engine/Private/Components/LightComponent.cpp
+++ b/Engine/Source/Runtime/Engine/Private/Components/LightComponent.cpp
@@ -303,6 +303,7 @@ ULightComponentBase::ULightComponentBase(const FObjectInitializer& ObjectInitial
 	Brightness_DEPRECATED = 3.1415926535897932f;
 	Intensity = 3.1415926535897932f;
 	LightColor = FColor::White;
+	ModulatedShadowColor = FColor(128, 128, 128);
 	bAffectsWorld = true;
 	CastShadows = true;
 	CastStaticShadows = true;
