Added comment for ModulatedShadowColor

diff --git a/Engine/Source/Runtime/Engine/Classes/Components/LightComponentBase.h b/Engine/Source/Runtime/Engine/Classes/Components/LightComponentBase.h
--- a/Engine/Source/Runtime/Engine/Classes/Components/LightComponentBase.h
+++ b/Engine/Source/Runtime/Engine/Classes/Components/LightComponentBase.h
@@ -66,6 +66,9 @@ class ENGINE_API ULightComponentBase : public USceneComponent
 	UPROPERTY(EditAnywhere, BlueprintReadOnly, Category = Light, AdvancedDisplay)
 	uint32 bCastModulatedShadows : 1;
 
+	/**
+	* Color to modulate against the scene color when rendering modulated shadows. (mobile only)
+	**/
 	UPROPERTY(BlueprintReadOnly, interp, Category = Light, meta = (HideAlphaChannel), AdvancedDisplay)
 	FColor ModulatedShadowColor;
 
