Don't reset component properties on actor construction.

diff --git a/Engine/Source/Runtime/Engine/Private/ActorConstruction.cpp b/Engine/Source/Runtime/Engine/Private/ActorConstruction.cpp
--- a/Engine/Source/Runtime/Engine/Private/ActorConstruction.cpp
+++ b/Engine/Source/Runtime/Engine/Private/ActorConstruction.cpp
@@ -47,6 +47,24 @@ void AActor::SeedAllRandomStreams()
 }
 #endif //WITH_EDITOR
 
+/**  Helper function to check if the specified property is a component property */
+static bool IsComponentProperty(UProperty* Prop)
+{
+	bool bResult = false;
+	if (UObjectProperty* ObjProp = Cast<UObjectProperty>(Prop))
+	{
+		if (ObjProp->PropertyClass && ObjProp->PropertyClass->IsChildOf(UActorComponent::StaticClass()))
+		{
+			bResult = true;
+		}
+	}
+	else if (UArrayProperty* ArrayProp = Cast<UArrayProperty>(Prop))
+	{
+		bResult = IsComponentProperty(ArrayProp->Inner);
+	}
+	return bResult;
+}
+
 void AActor::ResetPropertiesForConstruction()
 {
 	// Get class CDO
@@ -80,7 +98,8 @@ void AActor::ResetPropertiesForConstruction()
 				&& !bCanEditInstanceValue
 				&& bCanBeSetInBlueprints
 				&& !Prop->IsA(UDelegateProperty::StaticClass()) 
-				&& !Prop->IsA(UMulticastDelegateProperty::StaticClass()) )
+				&& !Prop->IsA(UMulticastDelegateProperty::StaticClass())
+				&& !(Prop->ContainsInstancedObjectProperty() && IsComponentProperty(Prop))) // ContainsInstancedObjectProperty will short circuit fast if it's not a component
 		{
 			Prop->CopyCompleteValue_InContainer(this, Default);
 		}
