fix reordering warning

diff --git a/Engine/Source/Runtime/Engine/Private/PhysicsEngine/ConstraintInstance.cpp b/Engine/Source/Runtime/Engine/Private/PhysicsEngine/ConstraintInstance.cpp
--- a/Engine/Source/Runtime/Engine/Private/PhysicsEngine/ConstraintInstance.cpp
+++ b/Engine/Source/Runtime/Engine/Private/PhysicsEngine/ConstraintInstance.cpp
@@ -243,9 +243,9 @@ FConstraintInstance::FConstraintInstance()
 	, LinearDriveDamping(1.0f)
 	, LinearDriveForceLimit(0)
 	, bAngularOrientationDrive(false)
-	, bAngularVelocityDrive(false)
-    , bEnableSwingDrive(true)
+	, bEnableSwingDrive(true)
 	, bEnableTwistDrive(true)
+	, bAngularVelocityDrive(false)
 	, AngularOrientationTarget(ForceInit)
 	, AngularVelocityTarget(ForceInit)
 	, AngularDriveSpring(50.0f)
