Tick functions can now specify an interval so that they will only fire their tick once that interval has expired rather than every frame.

diff --git a/Engine/Source/Editor/DetailCustomizations/Private/ActorComponentDetails.cpp b/Engine/Source/Editor/DetailCustomizations/Private/ActorComponentDetails.cpp
--- a/Engine/Source/Editor/DetailCustomizations/Private/ActorComponentDetails.cpp
+++ b/Engine/Source/Editor/DetailCustomizations/Private/ActorComponentDetails.cpp
@@ -24,6 +24,7 @@ void FActorComponentDetails::CustomizeDetails( IDetailLayoutBuilder& DetailBuild
 		IDetailCategoryBuilder& TickCategory = DetailBuilder.EditCategory("ComponentTick");
 
 		TickCategory.AddProperty(PrimaryTickProperty->GetChildHandle(GET_MEMBER_NAME_CHECKED(FTickFunction, bStartWithTickEnabled)));
+		TickCategory.AddProperty(PrimaryTickProperty->GetChildHandle(GET_MEMBER_NAME_CHECKED(FTickFunction, TickInterval)));
 		TickCategory.AddProperty(PrimaryTickProperty->GetChildHandle(GET_MEMBER_NAME_CHECKED(FTickFunction, bTickEvenWhenPaused)), EPropertyLocation::Advanced);
 		TickCategory.AddProperty(PrimaryTickProperty->GetChildHandle(GET_MEMBER_NAME_CHECKED(FTickFunction, bAllowTickOnDedicatedServer)), EPropertyLocation::Advanced);
 		TickCategory.AddProperty(PrimaryTickProperty->GetChildHandle(GET_MEMBER_NAME_CHECKED(FTickFunction, TickGroup)), EPropertyLocation::Advanced);
diff --git a/Engine/Source/Editor/DetailCustomizations/Private/ActorDetails.cpp b/Engine/Source/Editor/DetailCustomizations/Private/ActorDetails.cpp
--- a/Engine/Source/Editor/DetailCustomizations/Private/ActorDetails.cpp
+++ b/Engine/Source/Editor/DetailCustomizations/Private/ActorDetails.cpp
@@ -136,6 +136,7 @@ void FActorDetails::CustomizeDetails( IDetailLayoutBuilder& DetailLayout )
 		IDetailCategoryBuilder& TickCategory = DetailLayout.EditCategory("Tick", LOCTEXT("TickCategoryName", "Actor Tick") );
 
 		TickCategory.AddProperty(PrimaryTickProperty->GetChildHandle(GET_MEMBER_NAME_CHECKED(FTickFunction, bStartWithTickEnabled)));
+		TickCategory.AddProperty(PrimaryTickProperty->GetChildHandle(GET_MEMBER_NAME_CHECKED(FTickFunction, TickInterval)));
 		TickCategory.AddProperty(PrimaryTickProperty->GetChildHandle(GET_MEMBER_NAME_CHECKED(FTickFunction, bTickEvenWhenPaused)), EPropertyLocation::Advanced);
 		TickCategory.AddProperty(PrimaryTickProperty->GetChildHandle(GET_MEMBER_NAME_CHECKED(FTickFunction, bAllowTickOnDedicatedServer)), EPropertyLocation::Advanced);
 		TickCategory.AddProperty(PrimaryTickProperty->GetChildHandle(GET_MEMBER_NAME_CHECKED(FTickFunction, TickGroup)), EPropertyLocation::Advanced);
diff --git a/Engine/Source/Runtime/Engine/Classes/Engine/EngineBaseTypes.h b/Engine/Source/Runtime/Engine/Classes/Engine/EngineBaseTypes.h
--- a/Engine/Source/Runtime/Engine/Classes/Engine/EngineBaseTypes.h
+++ b/Engine/Source/Runtime/Engine/Classes/Engine/EngineBaseTypes.h
@@ -85,42 +85,40 @@ struct FTickPrerequisite
 	/** Tick functions live inside of UObjects, so we need a separate weak pointer to the UObject solely for the purpose of determining if PrerequisiteTickFunction is still valid. */
 	TWeakObjectPtr<class UObject> PrerequisiteObject;
 
+	/** Pointer to the actual tick function and must be completed prior to our tick running. */
+	struct FTickFunction*		PrerequisiteTickFunction;
 
-
-		/** Pointer to the actual tick function and must be completed prior to our tick running. */
-		struct FTickFunction*		PrerequisiteTickFunction;
-
-		/** Noop constructor. */
-		FTickPrerequisite()
-		: PrerequisiteTickFunction(nullptr)
-		{
-		}
-		/** 
-		 * Constructor
-		 * @param TargetObject - UObject containing this tick function. Only used to verify that the other pointer is still usable
-		 * @param TargetTickFunction - Actual tick function to use as a prerequisite
-		**/
-		FTickPrerequisite(UObject* TargetObject, struct FTickFunction& TargetTickFunction)
-		: PrerequisiteObject(TargetObject)
-		, PrerequisiteTickFunction(&TargetTickFunction)
-		{
-			check(PrerequisiteTickFunction);
-		}
-		/** Equality operator, used to prevent duplicates and allow removal by value. */
-		bool operator==(const FTickPrerequisite& Other) const
-		{
-			return PrerequisiteObject == Other.PrerequisiteObject &&
-				PrerequisiteTickFunction == Other.PrerequisiteTickFunction;
-		}
-		/** Return the tick function, if it is still valid. Can be null if the tick function was null or the containing UObject has been garbage collected. */
-		struct FTickFunction* Get()
+	/** Noop constructor. */
+	FTickPrerequisite()
+	: PrerequisiteTickFunction(nullptr)
+	{
+	}
+	/** 
+		* Constructor
+		* @param TargetObject - UObject containing this tick function. Only used to verify that the other pointer is still usable
+		* @param TargetTickFunction - Actual tick function to use as a prerequisite
+	**/
+	FTickPrerequisite(UObject* TargetObject, struct FTickFunction& TargetTickFunction)
+	: PrerequisiteObject(TargetObject)
+	, PrerequisiteTickFunction(&TargetTickFunction)
+	{
+		check(PrerequisiteTickFunction);
+	}
+	/** Equality operator, used to prevent duplicates and allow removal by value. */
+	bool operator==(const FTickPrerequisite& Other) const
+	{
+		return PrerequisiteObject == Other.PrerequisiteObject &&
+			PrerequisiteTickFunction == Other.PrerequisiteTickFunction;
+	}
+	/** Return the tick function, if it is still valid. Can be null if the tick function was null or the containing UObject has been garbage collected. */
+	struct FTickFunction* Get()
+	{
+		if (PrerequisiteObject.IsValid(true))
 		{
-			if (PrerequisiteObject.IsValid(true))
-			{
-				return PrerequisiteTickFunction;
-			}
-			return nullptr;
+			return PrerequisiteTickFunction;
 		}
+		return nullptr;
+	}
 };
 
 /** 
@@ -175,11 +173,19 @@ private:
 	/** If true, means that this tick function is in the master array of tick functions **/
 	uint32 bRegistered:1;
 
+	enum class ETickState : uint8
+	{
+		Disabled,
+		Enabled,
+		CoolingDown
+	};
+
 	/** 
-	 * If false, this tick will not fire
+	 * If Disabled, this tick will not fire
+	 * If CoolingDown, this tick has an interval frequency that is being adhered to currently
 	 * CAUTION: Do not set this directly
 	 **/
-	uint32 bTickEnabled:1;
+	ETickState TickState;
 
 	/** Internal data to track if we have started visiting this tick function yet this frame **/
 	int32 TickVisitedGFrameCounter;
@@ -187,14 +193,26 @@ private:
 	/** Internal data to track if we have finshed visiting this tick function yet this frame **/
 	int32 TickQueuedGFrameCounter;
 
-private:
 	/** Completion handle for the task that will run this tick. Caution, this is no reset to nullptr until an unspecified future time **/
 	FGraphEventRef CompletionHandle;
 
 	/** Prerequisites for this tick function **/
 	TArray<struct FTickPrerequisite> Prerequisites;
 
+	/** The next function in the cooling down list for ticks with an interval*/
+	FTickFunction* Next;
+
+	/** 
+	  * If TickFrequency is greater than 0 and tick state is CoolingDown, this is the time, 
+	  * relative to the element ahead of it in the cooling down list, remaining until the next time this function will tick 
+	  */
+	float RelativeTickCooldown;
 public:
+
+	/** The frequency at which this tick function will be executed.  If less than or equal to 0 then it will tick every frame */
+	UPROPERTY(EditDefaultsOnly, Category="Tick")
+	float TickInterval;
+
 	/** Back pointer to the FTickTaskLevel containing this tick function if it is registered **/
 	class FTickTaskLevel*						TickTaskLevel;
 
@@ -216,7 +234,7 @@ public:
 	/** Enables or disables this tick function. **/
 	void SetTickFunctionEnable(bool bInEnabled);
 	/** Returns whether the tick function is currently enabled */
-	bool IsTickFunctionEnabled() const { return bTickEnabled; }
+	bool IsTickFunctionEnabled() const { return TickState != ETickState::Disabled; }
 
 	/**
 	* Gets the current completion handle of this tick function, so it can be delayed until a later point when some additional
diff --git a/Engine/Source/Runtime/Engine/Private/TickTaskManager.cpp b/Engine/Source/Runtime/Engine/Private/TickTaskManager.cpp
--- a/Engine/Source/Runtime/Engine/Private/TickTaskManager.cpp
+++ b/Engine/Source/Runtime/Engine/Private/TickTaskManager.cpp
@@ -13,6 +13,8 @@ DECLARE_CYCLE_STAT(TEXT("Queue Ticks"),STAT_QueueTicks,STATGROUP_Game);
 DECLARE_CYCLE_STAT(TEXT("Queue Ticks Wait"),STAT_QueueTicksWait,STATGROUP_Game);
 DECLARE_CYCLE_STAT(TEXT("Queue Tick Task"),STAT_QueueTickTask,STATGROUP_Game);
 DECLARE_CYCLE_STAT(TEXT("Post Queue Tick Task"),STAT_PostTickTask,STATGROUP_Game);
+DECLARE_CYCLE_STAT(TEXT("Cooldown Dequeuing"),STAT_DequeueCooldowns,STATGROUP_Game);
+DECLARE_CYCLE_STAT(TEXT("Schedule cooldowns"),STAT_ScheduleCooldowns,STATGROUP_Game);
 DECLARE_DWORD_COUNTER_STAT(TEXT("Ticks Queued"),STAT_TicksQueued,STATGROUP_Game);
 DECLARE_CYCLE_STAT(TEXT("TG_NewlySpawned"), STAT_TG_NewlySpawned, STATGROUP_TickGroups);
 
@@ -400,14 +402,20 @@ public:
 	}
 	~FTickTaskLevel()
 	{
-		for (TSet<FTickFunction *>::TIterator It(AllEnabledTickFunctions); It; ++It)
+		for (TSet<FTickFunction*>::TIterator It(AllEnabledTickFunctions); It; ++It)
 		{
 			(*It)->bRegistered = false;
 		}
-		for (TSet<FTickFunction *>::TIterator It(AllDisabledTickFunctions); It; ++It)
+		for (TSet<FTickFunction*>::TIterator It(AllDisabledTickFunctions); It; ++It)
 		{
 			(*It)->bRegistered = false;
 		}
+		FTickFunction* CoolingDownNode = AllCoolingDownTickFunctions.Head;
+		while (CoolingDownNode)
+		{
+			CoolingDownNode->bRegistered = false;
+			CoolingDownNode = CoolingDownNode->Next;
+		}
 	}
 
 	/**
@@ -423,21 +431,148 @@ public:
 		Context.TickType = InContext.TickType;
 		Context.Thread = ENamedThreads::GameThread;
 		bTickNewlySpawned = true;
-		return AllEnabledTickFunctions.Num();
+
+		int32 CooldownTicksEnabled = 0;
+		{
+			SCOPE_CYCLE_COUNTER(STAT_DequeueCooldowns);
+			// Determine which cooled down ticks will be enabled this frame
+			float CumulativeCooldown = 0.f;
+			FTickFunction* TickFunction = AllCoolingDownTickFunctions.Head;
+			while (TickFunction)
+			{
+				if (CumulativeCooldown + TickFunction->RelativeTickCooldown >= Context.DeltaSeconds)
+				{
+					TickFunction->RelativeTickCooldown -= (Context.DeltaSeconds - CumulativeCooldown);
+					break;
+				}
+				CumulativeCooldown += TickFunction->RelativeTickCooldown;
+
+				TickFunction->TickState = FTickFunction::ETickState::Enabled;
+				TickFunction = TickFunction->Next;
+				++CooldownTicksEnabled;
+			}
+		}
+
+		return AllEnabledTickFunctions.Num() + CooldownTicksEnabled;
 	}
+
+	struct FTickScheduleDetails
+	{
+		FTickScheduleDetails(FTickFunction* InTickFunction, const float InCooldown)
+			: TickFunction(InTickFunction)
+			, Cooldown(InCooldown)
+		{}
+
+		FTickFunction* TickFunction;
+		float Cooldown;
+	};
+
+	/* Puts a TickFunction in to the cooldown state*/
+	void ScheduleTickFunctionCooldowns(TArray<FTickScheduleDetails>& TickFunctionsToSchedule)
+	{
+		SCOPE_CYCLE_COUNTER(STAT_ScheduleCooldowns);
+
+		TickFunctionsToSchedule.Sort([](const FTickScheduleDetails& A, const FTickScheduleDetails& B)
+		{
+			return A.Cooldown < B.Cooldown;
+		});
+
+		int32 RescheduleIndex = 0;
+		float CumulativeCooldown = 0.f;
+		FTickFunction* PrevComparisonTickFunction = nullptr;
+		FTickFunction* ComparisonTickFunction = AllCoolingDownTickFunctions.Head;
+		while (ComparisonTickFunction && RescheduleIndex < TickFunctionsToSchedule.Num())
+		{
+			const float CooldownTime = TickFunctionsToSchedule[RescheduleIndex].Cooldown;
+			if ((CumulativeCooldown + ComparisonTickFunction->RelativeTickCooldown) > CooldownTime)
+			{
+				FTickFunction* TickFunction = TickFunctionsToSchedule[RescheduleIndex].TickFunction;
+				TickFunction->TickState = FTickFunction::ETickState::CoolingDown;
+				TickFunction->RelativeTickCooldown = CooldownTime - CumulativeCooldown;
+
+				if (PrevComparisonTickFunction)
+				{
+					PrevComparisonTickFunction->Next = TickFunction;
+				}
+				else
+				{
+					check(ComparisonTickFunction == AllCoolingDownTickFunctions.Head);
+					AllCoolingDownTickFunctions.Head = TickFunction;
+				}
+				TickFunction->Next = ComparisonTickFunction;
+				PrevComparisonTickFunction = TickFunction;
+				ComparisonTickFunction->RelativeTickCooldown -= TickFunction->RelativeTickCooldown;
+				CumulativeCooldown += TickFunction->RelativeTickCooldown;
+				++RescheduleIndex;
+			}
+			else
+			{
+				CumulativeCooldown += ComparisonTickFunction->RelativeTickCooldown;
+				PrevComparisonTickFunction = ComparisonTickFunction;
+				ComparisonTickFunction = ComparisonTickFunction->Next;
+			}
+		}
+		for ( ; RescheduleIndex < TickFunctionsToSchedule.Num(); ++RescheduleIndex)
+		{
+			FTickFunction* TickFunction = TickFunctionsToSchedule[RescheduleIndex].TickFunction;
+			const float CooldownTime = TickFunctionsToSchedule[RescheduleIndex].Cooldown;
+
+			TickFunction->TickState = FTickFunction::ETickState::CoolingDown;
+			TickFunction->RelativeTickCooldown = CooldownTime - CumulativeCooldown;
+
+			TickFunction->Next = nullptr;
+			if (PrevComparisonTickFunction)
+			{
+				PrevComparisonTickFunction->Next = TickFunction;
+			}
+			else
+			{
+				check(ComparisonTickFunction == AllCoolingDownTickFunctions.Head);
+				AllCoolingDownTickFunctions.Head = TickFunction;
+			}
+			PrevComparisonTickFunction = TickFunction;
+
+			CumulativeCooldown += TickFunction->RelativeTickCooldown;
+		}
+	}
+
 	/* Queue all tick functions for execution */
 	void QueueAllTicks()
 	{
 		FTickTaskSequencer& TTS = FTickTaskSequencer::Get();
-		for (TSet<FTickFunction *>::TIterator It(AllEnabledTickFunctions); It; ++It)
+		for (TSet<FTickFunction*>::TIterator It(AllEnabledTickFunctions); It; ++It)
 		{
-			(*It)->QueueTickFunction(TTS, Context);
+			FTickFunction* TickFunction = *It;
+			TickFunction->QueueTickFunction(TTS, Context);
+
+			if (TickFunction->TickInterval > 0.f)
+			{
+				It.RemoveCurrent();
+				TickFunctionsToReschedule.Add(FTickScheduleDetails(TickFunction, TickFunction->TickInterval));
+			}
+		}
+		int32 EnabledCooldownTicks = 0;
+		float CumulativeCooldown = 0.f;
+		while (FTickFunction* TickFunction = AllCoolingDownTickFunctions.Head)
+		{
+			if (TickFunction->TickState == FTickFunction::ETickState::Enabled)
+			{
+				CumulativeCooldown += TickFunction->RelativeTickCooldown;
+				TickFunction->QueueTickFunction(TTS, Context);
+				TickFunctionsToReschedule.Add(FTickScheduleDetails(TickFunction, TickFunction->TickInterval - (Context.DeltaSeconds - CumulativeCooldown))); // Give credit for any overrun
+				AllCoolingDownTickFunctions.Head = TickFunction->Next;
+			}
+			else
+			{
+				break;
+			}
+		}
+
+		if (TickFunctionsToReschedule.Num() > 0)
+		{
+			ScheduleTickFunctionCooldowns(TickFunctionsToReschedule);
+			TickFunctionsToReschedule.Reset();
 		}
-	}
-	/* Retrieve the tick set */
-	TSet<FTickFunction *>& GetTickSet()
-	{
-		return AllEnabledTickFunctions;
 	}
 	/**
 	 * Queues the newly spawned ticks for this level
@@ -448,10 +583,22 @@ public:
 		Context.TickGroup = CurrentTickGroup;
 		int32 Num = 0;
 		FTickTaskSequencer& TTS = FTickTaskSequencer::Get();
-		for (TSet<FTickFunction *>::TIterator It(NewlySpawnedTickFunctions); It; ++It)
+		for (TSet<FTickFunction*>::TIterator It(NewlySpawnedTickFunctions); It; ++It)
 		{
-			(*It)->QueueTickFunction(TTS, Context);
+			FTickFunction* TickFunction = *It;
+			TickFunction->QueueTickFunction(TTS, Context);
 			Num++;
+
+			if (TickFunction->TickInterval > 0.f)
+			{
+				AllEnabledTickFunctions.Remove(TickFunction);
+				TickFunctionsToReschedule.Add(FTickScheduleDetails(TickFunction, TickFunction->TickInterval));
+			}
+		}
+		if (TickFunctionsToReschedule.Num() > 0)
+		{
+			ScheduleTickFunctionCooldowns(TickFunctionsToReschedule);
+			TickFunctionsToReschedule.Reset();
 		}
 		NewlySpawnedTickFunctions.Empty();
 		return Num;
@@ -465,18 +612,72 @@ public:
 	void RunPauseFrame(const FTickContext& InContext)
 	{
 		check(!NewlySpawnedTickFunctions.Num()); // There shouldn't be any in here at this point in the frame
-		for (TSet<FTickFunction *>::TIterator It(AllEnabledTickFunctions); It; ++It)
+		
+		float CumulativeCooldown = 0.f;
+		FTickFunction* PrevTickFunction = nullptr;
+		FTickFunction* TickFunction = AllCoolingDownTickFunctions.Head;
+		while (TickFunction)
+		{
+			CumulativeCooldown += TickFunction->RelativeTickCooldown;
+			if (TickFunction->bTickEvenWhenPaused)
+			{
+				if (CumulativeCooldown < InContext.DeltaSeconds)
+				{
+					TickFunction->TickVisitedGFrameCounter = GFrameCounter;
+					TickFunction->TickQueuedGFrameCounter = GFrameCounter;
+					TickFunction->ExecuteTick(InContext.DeltaSeconds, InContext.TickType, ENamedThreads::GameThread, FGraphEventRef());
+					TickFunction->CompletionHandle = NULL; // Allow the old completion handle to be recycled
+
+					TickFunctionsToReschedule.Add(FTickScheduleDetails(TickFunction, TickFunction->TickInterval - (InContext.DeltaSeconds - CumulativeCooldown))); // Give credit for any overrun
+				}
+				else
+				{
+					TickFunctionsToReschedule.Add(FTickScheduleDetails(TickFunction, CumulativeCooldown - InContext.DeltaSeconds));
+				}
+				if (PrevTickFunction)
+				{
+					PrevTickFunction->Next = TickFunction->Next;
+				}
+				else
+				{
+					check(TickFunction == AllCoolingDownTickFunctions.Head);
+					AllCoolingDownTickFunctions.Head = TickFunction->Next;
+				}
+				if (TickFunction->Next)
+				{
+					TickFunction->Next->RelativeTickCooldown += TickFunction->RelativeTickCooldown;
+					CumulativeCooldown -= TickFunction->RelativeTickCooldown; // Since the next object in the list will have this cooldown included take it back out of the cumulative
+				}
+			}
+			PrevTickFunction = TickFunction;
+			TickFunction = TickFunction->Next;
+		}
+
+		for (TSet<FTickFunction*>::TIterator It(AllEnabledTickFunctions); It; ++It)
 		{
 			FTickFunction* TickFunction = *It;
 			TickFunction->CompletionHandle = NULL; // might as well NULL this out to allow these handles to be recycled
-			if (TickFunction->bTickEvenWhenPaused && TickFunction->bTickEnabled)
+			if (TickFunction->bTickEvenWhenPaused && TickFunction->TickState == FTickFunction::ETickState::Enabled)
 			{
 				TickFunction->TickVisitedGFrameCounter = GFrameCounter;
 				TickFunction->TickQueuedGFrameCounter = GFrameCounter;
 				TickFunction->ExecuteTick(InContext.DeltaSeconds, InContext.TickType, ENamedThreads::GameThread, FGraphEventRef());
 				TickFunction->CompletionHandle = NULL; // Allow the old completion handle to be recycled
+
+				if (TickFunction->TickInterval > 0.f)
+				{
+					It.RemoveCurrent();
+					TickFunctionsToReschedule.Add(FTickScheduleDetails(TickFunction, TickFunction->TickInterval));
+				}
 			}
 		}
+
+		if (TickFunctionsToReschedule.Num() > 0)
+		{
+			ScheduleTickFunctionCooldowns(TickFunctionsToReschedule);
+			TickFunctionsToReschedule.Reset();
+		}
+
 		check(!NewlySpawnedTickFunctions.Num()); // We don't support new spawns during pause ticks
 	}
 
@@ -491,13 +692,13 @@ public:
 	/** Return true if this tick function is in the master list **/
 	bool HasTickFunction(FTickFunction* TickFunction)
 	{
-		return AllEnabledTickFunctions.Contains(TickFunction) || AllDisabledTickFunctions.Contains(TickFunction);
+		return AllEnabledTickFunctions.Contains(TickFunction) || AllDisabledTickFunctions.Contains(TickFunction) || AllCoolingDownTickFunctions.Contains(TickFunction);
 	}
 	/** Add the tick function to the master list **/
 	void AddTickFunction(FTickFunction* TickFunction)
 	{
 		check(!HasTickFunction(TickFunction));
-		if (TickFunction->bTickEnabled)
+		if (TickFunction->TickState == FTickFunction::ETickState::Enabled)
 		{
 			AllEnabledTickFunctions.Add(TickFunction);
 			if (bTickNewlySpawned)
@@ -507,17 +708,18 @@ public:
 		}
 		else
 		{
+			check(TickFunction->TickState == FTickFunction::ETickState::Disabled);
 			AllDisabledTickFunctions.Add(TickFunction);
 		}
 	}
 
 	/** Dumps info about a tick function to output device. */
-	FORCEINLINE void DumpTickFunction(FOutputDevice& Ar, FTickFunction* Function, UEnum* TickGroupEnum)
+	FORCEINLINE void DumpTickFunction(FOutputDevice& Ar, FTickFunction* Function, UEnum* TickGroupEnum, const float RemainingCooldown = 0.f)
 	{
 		// Info about the function.
 		Ar.Logf(TEXT("%s, %s, ActualTickGroup: %s, Prerequesities: %d"),
 			*Function->DiagnosticMessage(),
-			Function->IsTickFunctionEnabled() ? TEXT("Enabled") : TEXT("Disabled"),
+			Function->IsTickFunctionEnabled() ? (RemainingCooldown > 0.f ? *FString::Printf(TEXT("Cooling Down for %.4g seconds"),RemainingCooldown) : TEXT("Enabled")) : TEXT("Disabled"),
 			*TickGroupEnum->GetEnumName(Function->ActualTickGroup),			
 			Function->Prerequisites.Num());
 
@@ -542,15 +744,33 @@ public:
 		UEnum* TickGroupEnum = CastChecked<UEnum>(StaticFindObject(UEnum::StaticClass(), ANY_PACKAGE, TEXT("ETickingGroup"), true));
 		if (bEnabled)
 		{
-			for (TSet<FTickFunction *>::TIterator It(AllEnabledTickFunctions); It; ++It)
+			for (TSet<FTickFunction*>::TIterator It(AllEnabledTickFunctions); It; ++It)
 			{
 				DumpTickFunction(Ar, *It, TickGroupEnum);
 			}
+			float CumulativeCooldown = 0.f;
+			FTickFunction* TickFunction = AllCoolingDownTickFunctions.Head;
+			while (TickFunction)
+			{
+				CumulativeCooldown += TickFunction->RelativeTickCooldown;
+				DumpTickFunction(Ar, TickFunction, TickGroupEnum, CumulativeCooldown);
+				TickFunction = TickFunction->Next;
+				++EnabledCount;
+			}
+		}
+		else
+		{
+			FTickFunction* TickFunction = AllCoolingDownTickFunctions.Head;
+			while (TickFunction)
+			{
+				TickFunction = TickFunction->Next;
+				++EnabledCount;
+			}
 		}
 		EnabledCount += AllEnabledTickFunctions.Num();
 		if (bDisabled)
 		{
-			for (TSet<FTickFunction *>::TIterator It(AllDisabledTickFunctions); It; ++It)
+			for (TSet<FTickFunction*>::TIterator It(AllDisabledTickFunctions); It; ++It)
 			{
 				DumpTickFunction(Ar, *It, TickGroupEnum);
 			}
@@ -561,30 +781,125 @@ public:
 	/** Remove the tick function from the master list **/
 	void RemoveTickFunction(FTickFunction* TickFunction)
 	{
-		if (TickFunction->bTickEnabled)
+		switch(TickFunction->TickState)
 		{
-			verify(AllEnabledTickFunctions.Remove(TickFunction) == 1); // otherwise you changed bEnabled while the tick function was registered. Call SetTickFunctionEnable instead.
+		case FTickFunction::ETickState::Enabled:
+			if (TickFunction->TickInterval > 0.f)
+			{
+				// An enabled function with a tick interval could be in either the enabled or cooling down list
+				if (AllEnabledTickFunctions.Remove(TickFunction) == 0)
+				{
+					FTickFunction* PrevComparisionFunction = nullptr;
+					FTickFunction* ComparisonFunction = AllCoolingDownTickFunctions.Head;
+					bool bFound = false;
+					while (ComparisonFunction)
+					{
+						if (ComparisonFunction == TickFunction)
+						{
+							bFound = true;
+							if (PrevComparisionFunction)
+							{
+								PrevComparisionFunction->Next = TickFunction->Next;
+							}
+							else
+							{
+								check(TickFunction == AllCoolingDownTickFunctions.Head);
+								AllCoolingDownTickFunctions.Head = TickFunction->Next;
+							}
+							TickFunction->Next = nullptr;
+						}
+						else
+						{
+							PrevComparisionFunction = ComparisonFunction;
+							ComparisonFunction = ComparisonFunction->Next;
+						}
+					}
+					check(bFound); // otherwise you changed TickState while the tick function was registered. Call SetTickFunctionEnable instead.
+				}
+			}
+			else
+			{
+				verify(AllEnabledTickFunctions.Remove(TickFunction) == 1); // otherwise you changed TickState while the tick function was registered. Call SetTickFunctionEnable instead.
+			}
+			break;
+
+		case FTickFunction::ETickState::Disabled:
+			verify(AllDisabledTickFunctions.Remove(TickFunction) == 1); // otherwise you changed TickState while the tick function was registered. Call SetTickFunctionEnable instead.
+			break;
+
+		case FTickFunction::ETickState::CoolingDown:
+			FTickFunction* PrevComparisionFunction = nullptr;
+			FTickFunction* ComparisonFunction = AllCoolingDownTickFunctions.Head;
+			bool bFound = false;
+			while (ComparisonFunction)
+			{
+				if (ComparisonFunction == TickFunction)
+				{
+					bFound = true;
+					if (PrevComparisionFunction)
+					{
+						PrevComparisionFunction->Next = TickFunction->Next;
+					}
+					else
+					{
+						check(TickFunction == AllCoolingDownTickFunctions.Head);
+						AllCoolingDownTickFunctions.Head = TickFunction->Next;
+					}
+					if (TickFunction->Next)
+					{
+						TickFunction->Next->RelativeTickCooldown += TickFunction->RelativeTickCooldown;
+						TickFunction->Next = nullptr;
+					}
+				}
+				else
+				{
+					PrevComparisionFunction = ComparisonFunction;
+					ComparisonFunction = ComparisonFunction->Next;
+				}
+			}
+			check(bFound); // otherwise you changed TickState while the tick function was registered. Call SetTickFunctionEnable instead.
+			break;
 		}
-		else
+	}
+
+private:
+
+	struct FCoolingDownTickFunctionList
+	{
+		FCoolingDownTickFunctionList()
+			: Head(nullptr)
 		{
-			verify(AllDisabledTickFunctions.Remove(TickFunction) == 1); // otherwise you changed bEnabled while the tick function was registered. Call SetTickFunctionEnable instead.
 		}
-		if (bTickNewlySpawned)
+
+		bool Contains(FTickFunction* TickFunction) const
 		{
-			NewlySpawnedTickFunctions.Remove(TickFunction);
+			FTickFunction* Node = Head;
+			while (Node)
+			{
+				if (Node == TickFunction)
+				{
+					return true;
+				}
+				Node = Node->Next;
+			}
+			return false;
 		}
-	}
 
-private:
+		FTickFunction* Head;
+	};
 
 	/** Global Sequencer														*/
 	FTickTaskSequencer&							TickTaskSequencer;
 	/** Master list of enabled tick functions **/
-	TSet<FTickFunction *>						AllEnabledTickFunctions;
+	TSet<FTickFunction*>						AllEnabledTickFunctions;
+	/** Master list of enabled tick functions **/
+	FCoolingDownTickFunctionList				AllCoolingDownTickFunctions;
 	/** Master list of disabled tick functions **/
-	TSet<FTickFunction *>						AllDisabledTickFunctions;
+	TSet<FTickFunction*>						AllDisabledTickFunctions;
+	/** Utility array to avoid memory reallocations when collecting functions to reschedule **/
+	TArray<FTickScheduleDetails>				TickFunctionsToReschedule;
 	/** List of tick functions added during a tick phase; these items are also duplicated in AllLiveTickFunctions for future frames **/
-	TSet<FTickFunction *>						NewlySpawnedTickFunctions;
+	TSet<FTickFunction*>						NewlySpawnedTickFunctions;
 	/** tick context **/
 	FTickContext								Context;
 	/** true during the tick phase, when true, tick function adds also go to the newly spawned list. **/
@@ -832,9 +1147,11 @@ FTickFunction::FTickFunction()
 	, bHighPriority(false)
 	, bRunOnAnyThread(false)
 	, bRegistered(false)
-	, bTickEnabled(true)
+	, TickState(ETickState::Enabled)
 	, TickVisitedGFrameCounter(0)
 	, TickQueuedGFrameCounter(0)
+	, RelativeTickCooldown(0.f)
+	, TickInterval(0.f)
 	, TickTaskLevel(NULL)
 {
 }
@@ -880,16 +1197,16 @@ void FTickFunction::UnRegisterTickFunction()
 /** Enables or disables this tick function. **/
 void FTickFunction::SetTickFunctionEnable(bool bInEnabled)
 {
-	if (bRegistered && bTickEnabled != bInEnabled)
+	if (bRegistered && (bInEnabled == (TickState == ETickState::Disabled)))
 	{
 		check(TickTaskLevel);
 		TickTaskLevel->RemoveTickFunction(this);
-		bTickEnabled = bInEnabled;
+		TickState = (bInEnabled ? ETickState::Enabled : ETickState::Disabled);
 		TickTaskLevel->AddTickFunction(this);
 	}
 	else
 	{
-		bTickEnabled = bInEnabled;
+		TickState = (bInEnabled ? ETickState::Enabled : ETickState::Disabled);
 	}
 }
 
@@ -944,9 +1261,9 @@ void FTickFunction::QueueTickFunction(FTickTaskSequencer& TTS, const struct FTic
 	if (TickVisitedGFrameCounter != GFrameCounter)
 	{
 		TickVisitedGFrameCounter = GFrameCounter;
-		if (bTickEnabled)
+		if (TickState != FTickFunction::ETickState::Disabled)
 		{
-			ETickingGroup MaxPrerequisiteTickGroup =  ETickingGroup(0);
+			ETickingGroup MaxPrerequisiteTickGroup = ETickingGroup(0);
 
 			FGraphEventArray TaskPrerequisites;
 			for (int32 PrereqIndex = 0; PrereqIndex < Prerequisites.Num(); PrereqIndex++)
@@ -972,7 +1289,7 @@ void FTickFunction::QueueTickFunction(FTickTaskSequencer& TTS, const struct FTic
 					}
 					else
 					{
-						MaxPrerequisiteTickGroup =  FMath::Max<ETickingGroup>(MaxPrerequisiteTickGroup, Prereq->ActualTickGroup);
+						MaxPrerequisiteTickGroup = FMath::Max<ETickingGroup>(MaxPrerequisiteTickGroup, Prereq->ActualTickGroup);
 						TaskPrerequisites.Add(Prereq->CompletionHandle);
 					}
 				}
@@ -986,7 +1303,10 @@ void FTickFunction::QueueTickFunction(FTickTaskSequencer& TTS, const struct FTic
 			}
 			ActualTickGroup = MyActualTickGroup;
 
-			TTS.QueueTickTask(&TaskPrerequisites, this, TickContext);
+			if (TickState == FTickFunction::ETickState::Enabled)
+			{
+				TTS.QueueTickTask(&TaskPrerequisites, this, TickContext);
+			}
 		}
 		TickQueuedGFrameCounter = GFrameCounter;
 	}
