#ANIM : Animation  Curve Evaluation to happen during Evaluation

- Now Curve is blending as bone transform blending
- Moved eval to each asset fom AnimInstance
- Removed unnecessary AnimInstance functions
- Fixed composite preview not working well with additive
- Split AnimTypes.h/AnimCurveTypes.h - I think this needs re-think
- This will be respected for all nodes blending including additive, cache pose, state transition
- This also will enhance in the future for us to create system to drive bone to drive morphtarget and any curve.

diff --git a/Engine/Source/Editor/UnrealEd/Private/Animation/DebugSkelMeshComponent.cpp b/Engine/Source/Editor/UnrealEd/Private/Animation/DebugSkelMeshComponent.cpp
--- a/Engine/Source/Editor/UnrealEd/Private/Animation/DebugSkelMeshComponent.cpp
+++ b/Engine/Source/Editor/UnrealEd/Private/Animation/DebugSkelMeshComponent.cpp
@@ -416,7 +416,8 @@ void UDebugSkelMeshComponent::GenSpaceBases(TArray<FTransform>& OutSpaceBases)
 	TempLocalAtoms.AddUninitialized(OutSpaceBases.Num());
 	TArray<FActiveVertexAnim> TempVertexAnims;
 	FVector TempRootBoneTranslation;
-	PerformAnimationEvaluation(SkeletalMesh, AnimScriptInstance, OutSpaceBases, TempLocalAtoms, TempVertexAnims, TempRootBoneTranslation);
+	FBlendedCurve TempCurve;
+	PerformAnimationEvaluation(SkeletalMesh, AnimScriptInstance, OutSpaceBases, TempLocalAtoms, TempVertexAnims, TempRootBoneTranslation, TempCurve);
 }
 
 void UDebugSkelMeshComponent::RefreshBoneTransforms(FActorComponentTickFunction* TickFunction)
@@ -489,9 +490,9 @@ void UDebugSkelMeshComponent::RefreshBoneTransforms(FActorComponentTickFunction*
 				FCSPose<FCompactPose> CSAdditiveBasePose;
 				{
 					FCompactPose AdditiveBasePose;
+					FBlendedCurve AdditiveCurve(PreviewInstance);
 					AdditiveBasePose.SetBoneContainer(&PreviewInstance->RequiredBones);
-
-					Sequence->GetAdditiveBasePose(AdditiveBasePose, FAnimExtractContext(PreviewInstance->CurrentTime));
+					Sequence->GetAdditiveBasePose(AdditiveBasePose, AdditiveCurve, FAnimExtractContext(PreviewInstance->CurrentTime));
 					CSAdditiveBasePose.InitPose(AdditiveBasePose);
 				}
 
diff --git a/Engine/Source/Runtime/Engine/Classes/Animation/AnimComposite.h b/Engine/Source/Runtime/Engine/Classes/Animation/AnimComposite.h
--- a/Engine/Source/Runtime/Engine/Classes/Animation/AnimComposite.h
+++ b/Engine/Source/Runtime/Engine/Classes/Animation/AnimComposite.h
@@ -19,10 +19,16 @@ public:
 	UPROPERTY()
 	struct FAnimTrack AnimationTrack;
 
+#if WITH_EDITORONLY_DATA
+	/** Preview Base pose for additive BlendSpace **/
+	UPROPERTY(EditAnywhere, Category=AdditiveSettings)
+	UAnimSequence* PreviewBasePose;
+#endif // WITH_EDITORONLY_DATA
+
 	// Begin UAnimSequenceBase interface
 	ENGINE_API virtual void OnAssetPlayerTickedInternal(FAnimAssetTickContext &Context, const float PreviousTime, const float MoveDelta, const FAnimTickRecord &Instance, class UAnimInstance* InstanceOwner) const override;
+	virtual void GetAnimationPose(FCompactPose& OutPose, FBlendedCurve& OutCurve, const FAnimExtractContext& ExtractionContext) const override;	
 	// End UAnimSequenceBase interface
-
 	// Begin UAnimSequence interface
 #if WITH_EDITOR
 	virtual bool GetAllAnimationSequencesReferred(TArray<UAnimSequence*>& AnimationSequences) override;
diff --git a/Engine/Source/Runtime/Engine/Classes/Animation/AnimCompositeBase.h b/Engine/Source/Runtime/Engine/Classes/Animation/AnimCompositeBase.h
--- a/Engine/Source/Runtime/Engine/Classes/Animation/AnimCompositeBase.h
+++ b/Engine/Source/Runtime/Engine/Classes/Animation/AnimCompositeBase.h
@@ -193,6 +193,9 @@ struct FAnimTrack
 	/** Get the segment at the given absolute montage time */
 	FAnimSegment* GetSegmentAtTime(float InTime);
 
+	/** Get animation pose function */
+	void GetAnimationPose(/*out*/ FCompactPose& OutPose,/*out*/ FBlendedCurve& OutCurve, const FAnimExtractContext& ExtractionContext) const;
+
 #if WITH_EDITOR
 	bool GetAllAnimationSequencesReferred(TArray<UAnimSequence*>& AnimationSequences) const;
 	void ReplaceReferredAnimations(const TMap<UAnimSequence*, UAnimSequence*>& ReplacementMap);
diff --git a/Engine/Source/Runtime/Engine/Classes/Animation/AnimInstance.h b/Engine/Source/Runtime/Engine/Classes/Animation/AnimInstance.h
--- a/Engine/Source/Runtime/Engine/Classes/Animation/AnimInstance.h
+++ b/Engine/Source/Runtime/Engine/Classes/Animation/AnimInstance.h
@@ -3,11 +3,11 @@
 #pragma once
 
 #include "AnimationAsset.h"
-#include "AnimSequence.h"
 #include "Components/SkeletalMeshComponent.h"
 #include "Components/SkinnedMeshComponent.h"
 #include "AnimStateMachineTypes.h"
 #include "BonePose.h"
+#include "Animation/AnimTypes.h"
 #include "AnimInstance.generated.h"
 
 struct FAnimMontageInstance;
@@ -182,7 +182,6 @@ struct FPerBoneBlendWeights
 	FPerBoneBlendWeights() {}
 
 };
-
 /** Helper struct for Slot node pose evaluation. */
 USTRUCT()
 struct FSlotEvaluationPose
@@ -199,7 +198,9 @@ struct FSlotEvaluationPose
 
 	/** Pose */
 	FCompactPose Pose;
-	
+	/** Curve */
+	FBlendedCurve Curve;
+
 	FSlotEvaluationPose()
 	{
 	}
@@ -365,17 +366,9 @@ public:
 	// Creates an uninitialized tick record in the list for the correct group or the ungrouped array.  If the group is valid, OutSyncGroupPtr will point to the group.
 	FAnimTickRecord& CreateUninitializedTickRecord(int32 GroupIndex, FAnimGroupInstance*& OutSyncGroupPtr);
 
-	void SequenceEvaluatePose(UAnimSequenceBase* Sequence, struct FCompactPose& Pose, const FAnimExtractContext& ExtractionContext);
-
-	void BlendSequences(const FCompactPose& Pose1, const FCompactPose& Pose2, float Alpha, FCompactPose& Result);
-
-	static void CopyPose(const struct FCompactPose& Source, struct FCompactPose& Destination);
-
-	void BlendSpaceEvaluatePose(UBlendSpaceBase* BlendSpace, TArray<FBlendSampleData>& BlendSampleDataCache, struct FCompactPose& OutPose);
-
 	// slotnode interfaces
 	void GetSlotWeight(FName const & SlotNodeName, float& out_SlotNodeWeight, float& out_SourceWeight) const;
-	void SlotEvaluatePose(FName SlotNodeName, const FCompactPose& SourcePose, FCompactPose& BlendedPose, float SlotNodeWeight);
+	void SlotEvaluatePose(FName SlotNodeName, const FCompactPose& SourcePose, const FBlendedCurve& SourceCurve, FCompactPose& BlendedPose, FBlendedCurve& BlendedCurve, float SlotNodeWeight);
 
 	// slot node run-time functions
 	void ReinitializeSlotNodes();
@@ -389,7 +382,8 @@ public:
 	void UpdateSlotRootMotionWeight(FName SlotNodeName, float Weight);
 	// Get the root motion weight for the montage slot
 	float GetSlotRootMotionWeight(FName SlotNodeName) const;
-
+	// Should Extract Root Motion or not. Return true if we do. 
+	bool ShouldExtractRootMotion() const { return RootMotionMode == ERootMotionMode::RootMotionFromEverything || RootMotionMode == ERootMotionMode::IgnoreRootMotion; }
 
 	// kismet event functions
 
@@ -836,16 +830,20 @@ public:
 	UPROPERTY(transient)
 	TArray<FAnimNotifyEvent> ActiveAnimNotifyState;
 
+private:
 	/** Curve Values that are added to trigger in event**/
 	TMap<FName, float>	EventCurves;
-	/** Morph Target Curves that will be used for SkeletalMeshComponent **/
-	TMap<FName, float>	MorphTargetCurves;
-	/** Material Curves that will be used for SkeletalMeshComponent **/
-	TMap<FName, float>	MaterialParameterCurves;
 	/** Material parameters that we had been changing and now need to clear */
 	TArray<FName> MaterialParamatersToClear;
 
+public: 
+	/** Update all internal curves from Blended Curve */
+	void UpdateCurves(const FBlendedCurve& InCurves);
 
+	/** Morph Target Curves that will be used for SkeletalMeshComponent **/
+	TMap<FName, float>	MorphTargetCurves;
+	/** Material Curves that will be used for SkeletalMeshComponent **/
+	TMap<FName, float>	MaterialParameterCurves;
 
 #if WITH_EDITORONLY_DATA
 	// Maximum playback position ever reached (only used when debugging in Persona)
diff --git a/Engine/Source/Runtime/Engine/Classes/Animation/AnimMontage.h b/Engine/Source/Runtime/Engine/Classes/Animation/AnimMontage.h
--- a/Engine/Source/Runtime/Engine/Classes/Animation/AnimMontage.h
+++ b/Engine/Source/Runtime/Engine/Classes/Animation/AnimMontage.h
@@ -535,10 +535,6 @@ private:
 
 
 public:
-	// UAnimSequenceBase Interface
-	virtual void EvaluateCurveData(class UAnimInstance* Instance, float CurrentTime, float BlendWeight) const override;
-	// End of UAnimSequenceBase Interface
-
 #if WITH_EDITOR
 	/**
 	 * Add Composite section with InSectionName
diff --git a/Engine/Source/Runtime/Engine/Classes/Animation/AnimNodeBase.h b/Engine/Source/Runtime/Engine/Classes/Animation/AnimNodeBase.h
--- a/Engine/Source/Runtime/Engine/Classes/Animation/AnimNodeBase.h
+++ b/Engine/Source/Runtime/Engine/Classes/Animation/AnimNodeBase.h
@@ -106,23 +106,29 @@ public:
 struct FPoseContext : public FAnimationBaseContext
 {
 public:
-	FCompactPose Pose;
+	FCompactPose	Pose;
+	FBlendedCurve	Curve;
 
 public:
 	// This constructor allocates a new uninitialized pose for the specified anim instance
 	FPoseContext(UAnimInstance* InAnimInstance)
 		: FAnimationBaseContext(InAnimInstance)
 	{
-		checkSlow( AnimInstance && AnimInstance->RequiredBones.IsValid() );
-		Pose.SetBoneContainer(&AnimInstance->RequiredBones);
+		Initialize(InAnimInstance);
 	}
 
 	// This constructor allocates a new uninitialized pose, copying non-pose state from the source context
 	FPoseContext(const FPoseContext& SourceContext)
 		: FAnimationBaseContext(SourceContext.AnimInstance)
 	{
-		checkSlow( AnimInstance && AnimInstance->RequiredBones.IsValid() );
+		Initialize(SourceContext.AnimInstance);
+	}
+
+	void Initialize(UAnimInstance* InAnimInstance)
+	{
+		checkSlow(AnimInstance && AnimInstance->RequiredBones.IsValid());
 		Pose.SetBoneContainer(&AnimInstance->RequiredBones);
+		Curve.InitFrom(AnimInstance->CurrentSkeleton);
 	}
 
 	void ResetToRefPose()
@@ -144,6 +150,18 @@ public:
 	{
 		return Pose.IsNormalized();
 	}
+
+	FPoseContext& operator=(const FPoseContext& Other)
+	{
+		if (AnimInstance != Other.AnimInstance)
+		{
+			Initialize(AnimInstance);
+		}
+
+		Pose = Other.Pose;
+		Curve = Other.Curve;
+		return *this;
+	}
 };
 
 
@@ -151,7 +169,8 @@ public:
 struct FComponentSpacePoseContext : public FAnimationBaseContext
 {
 public:
-	FCSPose<FCompactPose> Pose;
+	FCSPose<FCompactPose>	Pose;
+	FBlendedCurve			Curve;
 
 public:
 	// This constructor allocates a new uninitialized pose for the specified anim instance
@@ -172,6 +191,7 @@ public:
 	{
 		checkSlow( AnimInstance && AnimInstance->RequiredBones.IsValid() );
 		Pose.InitPose(&AnimInstance->RequiredBones);
+		Curve.InitFrom(AnimInstance->CurrentSkeleton);
 	}
 
 	bool ContainsNaN() const;
diff --git a/Engine/Source/Runtime/Engine/Classes/Animation/AnimNode_LayeredBoneBlend.h b/Engine/Source/Runtime/Engine/Classes/Animation/AnimNode_LayeredBoneBlend.h
--- a/Engine/Source/Runtime/Engine/Classes/Animation/AnimNode_LayeredBoneBlend.h
+++ b/Engine/Source/Runtime/Engine/Classes/Animation/AnimNode_LayeredBoneBlend.h
@@ -4,6 +4,7 @@
 
 #include "AnimNodeBase.h"
 #include "Animation/AnimData/BoneMaskFilter.h"
+#include "Animation/AnimTypes.h"
 #include "AnimNode_LayeredBoneBlend.generated.h"
 
 
@@ -29,7 +30,9 @@ public:
 
 	UPROPERTY(EditAnywhere, BlueprintReadWrite, Category=Config)
 	bool bMeshSpaceRotationBlend;
-
+	
+	UPROPERTY(EditAnywhere, BlueprintReadWrite, Category=Config)
+	TEnumAsByte<enum ECurveBlendOption::Type>	CurveBlendOption;
 
 protected:
 	TArray<FPerBoneBlendWeight> DesiredBoneBlendWeights;
diff --git a/Engine/Source/Runtime/Engine/Classes/Animation/AnimNode_SaveCachedPose.h b/Engine/Source/Runtime/Engine/Classes/Animation/AnimNode_SaveCachedPose.h
--- a/Engine/Source/Runtime/Engine/Classes/Animation/AnimNode_SaveCachedPose.h
+++ b/Engine/Source/Runtime/Engine/Classes/Animation/AnimNode_SaveCachedPose.h
@@ -26,6 +26,7 @@ struct ENGINE_API FAnimNode_SaveCachedPose : public FAnimNode_Base
 
 protected:
 	FCompactPose CachedPose;
+	FBlendedCurve CachedCurve;
 
 public:	
 	FAnimNode_SaveCachedPose();
diff --git a/Engine/Source/Runtime/Engine/Classes/Animation/AnimNode_TransitionPoseEvaluator.h b/Engine/Source/Runtime/Engine/Classes/Animation/AnimNode_TransitionPoseEvaluator.h
--- a/Engine/Source/Runtime/Engine/Classes/Animation/AnimNode_TransitionPoseEvaluator.h
+++ b/Engine/Source/Runtime/Engine/Classes/Animation/AnimNode_TransitionPoseEvaluator.h
@@ -52,6 +52,7 @@ struct ENGINE_API FAnimNode_TransitionPoseEvaluator : public FAnimNode_Base
 	int32 FramesToCachePose;
 
 	FCompactPose CachedPose;
+	FBlendedCurve CachedCurve;
 
 	UPROPERTY(transient)
 	int32 CacheFramesRemaining;
@@ -69,5 +70,5 @@ public:
 
 	bool InputNodeNeedsUpdate() const;
 	bool InputNodeNeedsEvaluate() const;
-	void CachePose(FPoseContext& Output, FCompactPose& PoseToCache);
+	void CachePose(FPoseContext& PoseToCache);
 };
diff --git a/Engine/Source/Runtime/Engine/Classes/Animation/AnimSequence.h b/Engine/Source/Runtime/Engine/Classes/Animation/AnimSequence.h
--- a/Engine/Source/Runtime/Engine/Classes/Animation/AnimSequence.h
+++ b/Engine/Source/Runtime/Engine/Classes/Animation/AnimSequence.h
@@ -44,21 +44,6 @@ enum AnimationKeyFormat
 	AKF_MAX,
 };
 
-/** 
- * Indicates whether an animation is additive, and what kind.
- */
-UENUM()
-enum EAdditiveAnimationType
-{
-	/** No additive. */
-	AAT_None  UMETA(DisplayName="No additive"),
-	/* Create Additive from LocalSpace Base. */
-	AAT_LocalSpaceBase UMETA(DisplayName="Local Space"),
-	/* Create Additive from MeshSpace Rotation Only, Translation still will be LocalSpace. */
-	AAT_RotationOffsetMeshSpace UMETA(DisplayName="Mesh Space"),
-	AAT_MAX,
-};
-
 /** 
  * For an additive animation, indicates what the animation is relative to.
  */
@@ -583,19 +568,21 @@ public:
 	* Get Bone Transform of the Time given, relative to Parent for all RequiredBones
 	* This returns different transform based on additive or not. Or what kind of additive.
 	*
-	* @param	OutPose				Pose object to fill
+	* @param	OutPose				[out] Pose object to fill
+	* @param	OutCurve			[out] Curves to fill
 	* @param	ExtractionContext	Extraction Context (position, looping, root motion, etc.)
 	*/
-	void GetAnimationPose(FCompactPose& OutPose, const FAnimExtractContext& ExtractionContext) const;
+	virtual void GetAnimationPose(FCompactPose& OutPose, FBlendedCurve& OutCurve, const FAnimExtractContext& ExtractionContext) const override;
 
 	/**
 	* Get Bone Transform of the animation for the Time given, relative to Parent for all RequiredBones
 	*
 	* @param	OutPose				[out] Array of output bone transforms
+	* @param	OutCurve			[out] Curves to fill	
 	* @param	RequiredBones		Array of Desired Tracks
 	* @param	ExtractionContext	Extraction Context (position, looping, root motion, etc.)
 	*/
-	ENGINE_API void GetBonePose(FCompactPose& OutPose, const FAnimExtractContext& ExtractionContext) const;
+	ENGINE_API void GetBonePose(FCompactPose& OutPose, FBlendedCurve& OutCurve, const FAnimExtractContext& ExtractionContext) const;
 
 private:
 	/** 
@@ -623,17 +610,19 @@ public:
 	* Get Bone Transform of the additive animation for the Time given, relative to Parent for all RequiredBones
 	*
 	* @param	OutPose				[out] Output bone transforms
+	* @param	OutCurve			[out] Curves to fill	
 	* @param	ExtractionContext	Extraction Context (position, looping, root motion, etc.)
 	*/
-	ENGINE_API void GetBonePose_Additive(FCompactPose& OutPose, const FAnimExtractContext& ExtractionContext) const;
+	ENGINE_API void GetBonePose_Additive(FCompactPose& OutPose, FBlendedCurve& OutCurve, const FAnimExtractContext& ExtractionContext) const;
 
 	/**
 	* Get Bone Transform of the base (reference) pose of the additive animation for the Time given, relative to Parent for all RequiredBones
 	*
-	* @param	OutPose			[out] Output bone transforms
+	* @param	OutPose				[out] Output bone transforms
+	* @param	OutCurve			[out] Curves to fill	
 	* @param	ExtractionContext	Extraction Context (position, looping, root motion, etc.)
 	*/
-	ENGINE_API void GetAdditiveBasePose(FCompactPose& OutPose, const FAnimExtractContext& ExtractionContext) const;
+	ENGINE_API void GetAdditiveBasePose(FCompactPose& OutPose, FBlendedCurve& OutCurve, const FAnimExtractContext& ExtractionContext) const;
 
 	/**
 	 * Get Bone Transform of the Time given, relative to Parent for the Track Given
@@ -808,9 +797,10 @@ private:
 	* This return mesh rotation only additive pose
 	*
 	* @param	OutPose				[out] Output bone transforms
+	* @param	OutCurve			[out] Curves to fill	
 	* @param	ExtractionContext	Extraction Context (position, looping, root motion, etc.)
 	*/
-	void GetBonePose_AdditiveMeshRotationOnly(FCompactPose& OutPose, const FAnimExtractContext& ExtractionContext) const;
+	void GetBonePose_AdditiveMeshRotationOnly(FCompactPose& OutPose, FBlendedCurve& OutCurve, const FAnimExtractContext& ExtractionContext) const;
 
 #if WITH_EDITOR
 	/**
diff --git a/Engine/Source/Runtime/Engine/Classes/Animation/AnimSequenceBase.h b/Engine/Source/Runtime/Engine/Classes/Animation/AnimSequenceBase.h
--- a/Engine/Source/Runtime/Engine/Classes/Animation/AnimSequenceBase.h
+++ b/Engine/Source/Runtime/Engine/Classes/Animation/AnimSequenceBase.h
@@ -10,571 +10,12 @@
 #include "AnimationAsset.h"
 #include "SmartName.h"
 #include "Skeleton.h"
-#include "Curves/CurveBase.h"
-#include "AnimLinkableElement.h"
+#include "Animation/AnimCurveTypes.h"
 #include "AnimSequenceBase.generated.h"
 
 #define DEFAULT_SAMPLERATE			30.f
 #define MINIMUM_ANIMATION_LENGTH	(1/DEFAULT_SAMPLERATE)
 
-namespace EAnimEventTriggerOffsets
-{
-	enum Type
-	{
-		OffsetBefore,
-		OffsetAfter,
-		NoOffset
-	};
-}
-
-ENGINE_API float GetTriggerTimeOffsetForType(EAnimEventTriggerOffsets::Type OffsetType);
-
-/** Ticking method for AnimNotifies in AnimMontages. */
-UENUM()
-namespace EMontageNotifyTickType
-{
-	enum Type
-	{
-		/** Queue notifies, and trigger them at the end of the evaluation phase (faster). Not suitable for changing sections or montage position. */
-		Queued,
-		/** Trigger notifies as they are encountered (Slower). Suitable for changing sections or montage position. */
-		BranchingPoint,
-	};
-}
-
-/** Filtering method for deciding whether to trigger a notify. */
-UENUM()
-namespace ENotifyFilterType
-{
-	enum Type
-	{
-		/** No filtering. */
-		NoFiltering,
-
-		/** Filter By Skeletal Mesh LOD. */
-		LOD,
-	};
-}
-
-/**
- * Triggers an animation notify.  Each AnimNotifyEvent contains an AnimNotify object
- * which has its Notify method called and passed to the animation.
- */
-USTRUCT()
-struct FAnimNotifyEvent : public FAnimLinkableElement
-{
-	GENERATED_USTRUCT_BODY()
-
-	/** The user requested time for this notify */
-	UPROPERTY()
-	float DisplayTime_DEPRECATED;
-
-	/** An offset from the DisplayTime to the actual time we will trigger the notify, as we cannot always trigger it exactly at the time the user wants */
-	UPROPERTY()
-	float TriggerTimeOffset;
-
-	/** An offset similar to TriggerTimeOffset but used for the end scrub handle of a notify state's duration */
-	UPROPERTY()
-	float EndTriggerTimeOffset;
-
-	UPROPERTY(EditAnywhere, BlueprintReadWrite, Category=AnimNotifyEvent)
-	float TriggerWeightThreshold;
-
-	UPROPERTY(EditAnywhere, BlueprintReadOnly, Category=AnimNotifyEvent)
-	FName NotifyName;
-
-	UPROPERTY(EditAnywhere, Instanced, BlueprintReadWrite, Category=AnimNotifyEvent)
-	class UAnimNotify * Notify;
-
-	UPROPERTY(EditAnywhere, Instanced, BlueprintReadWrite, Category=AnimNotifyEvent)
-	class UAnimNotifyState * NotifyStateClass;
-
-	UPROPERTY()
-	float Duration;
-
-	/** Linkable element to use for the end handle representing a notify state duration */
-	UPROPERTY()
-	FAnimLinkableElement EndLink;
-
-	/** If TRUE, this notify has been converted from an old BranchingPoint. */
-	UPROPERTY()
-	bool bConvertedFromBranchingPoint;
-
-	UPROPERTY(EditAnywhere, BlueprintReadWrite, Category=AnimNotifyEvent)
-	TEnumAsByte<EMontageNotifyTickType::Type> MontageTickType;
-
-	/** Defines the chance of of this notify triggering, 0 = No Chance, 1 = Always triggers */
-	UPROPERTY(EditAnywhere, BlueprintReadWrite, Category = AnimNotifyTriggerSettings, meta = (ClampMin = "0.0", ClampMax = "1.0"))
-	float NotifyTriggerChance;
-
-	/** Defines a method for filtering notifies (stopping them triggering) e.g. by looking at the meshes current LOD */
-	UPROPERTY(EditAnywhere, BlueprintReadWrite, Category = AnimNotifyTriggerSettings)
-	TEnumAsByte<ENotifyFilterType::Type> NotifyFilterType;
-
-	/** LOD to start filtering this notify from.*/
-	UPROPERTY(EditAnywhere, BlueprintReadWrite, Category = AnimNotifyTriggerSettings, meta = (ClampMin = "0"))
-	int32 NotifyFilterLOD;
-
-#if WITH_EDITORONLY_DATA
-	/** Color of Notify in editor */
-	UPROPERTY()
-	FColor NotifyColor;
-#endif // WITH_EDITORONLY_DATA
-
-	/** 'Track' that the notify exists on, used for visual placement in editor and sorting priority in runtime */
-	UPROPERTY()
-	int32 TrackIndex;
-
-	FAnimNotifyEvent()
-		: DisplayTime_DEPRECATED(0)
-		, TriggerTimeOffset(0)
-		, EndTriggerTimeOffset(0)
-		, TriggerWeightThreshold(ZERO_ANIMWEIGHT_THRESH)
-#if WITH_EDITORONLY_DATA
-		, Notify(NULL)
-#endif // WITH_EDITORONLY_DATA
-		, NotifyStateClass(NULL)
-		, Duration(0)
-		, bConvertedFromBranchingPoint(false)
-		, MontageTickType(EMontageNotifyTickType::Queued)
-		, NotifyTriggerChance(1.f)
-		, NotifyFilterType(ENotifyFilterType::NoFiltering)
-		, NotifyFilterLOD(0)
-#if WITH_EDITORONLY_DATA
-		, TrackIndex(0)
-#endif // WITH_EDITORONLY_DATA
-	{
-	}
-
-	/** Updates trigger offset based on a combination of predicted offset and current offset */
-	ENGINE_API void RefreshTriggerOffset(EAnimEventTriggerOffsets::Type PredictedOffsetType);
-
-	/** Updates end trigger offset based on a combination of predicted offset and current offset */
-	ENGINE_API void RefreshEndTriggerOffset(EAnimEventTriggerOffsets::Type PredictedOffsetType);
-
-	/** Returns the actual trigger time for this notify. In some cases this may be different to the DisplayTime that has been set */
-	ENGINE_API float GetTriggerTime() const;
-
-	/** Returns the actual end time for a state notify. In some cases this may be different from DisplayTime + Duration */
-	ENGINE_API float GetEndTriggerTime() const;
-
-	ENGINE_API float GetDuration() const;
-
-	ENGINE_API void SetDuration(float NewDuration);
-
-	/** Returns true is this AnimNotify is a BranchingPoint */
-	ENGINE_API bool IsBranchingPoint() const
-	{
-		return (MontageTickType == EMontageNotifyTickType::BranchingPoint) && GetLinkedMontage();
-	}
-
-	/** Returns true if this is blueprint derived notifies **/
-	bool IsBlueprintNotify() const
-	{
-		return Notify != NULL || NotifyStateClass != NULL;
-	}
-
-	bool operator ==(const FAnimNotifyEvent& Other) const
-	{
-		return(
-			(Notify && Notify == Other.Notify) || 
-			(NotifyStateClass && NotifyStateClass == Other.NotifyStateClass) ||
-			(!IsBlueprintNotify() && NotifyName == Other.NotifyName)
-			);
-	}
-
-	/** This can be used with the Sort() function on a TArray of FAnimNotifyEvents to sort the notifies array by time, earliest first. */
-	ENGINE_API bool operator <(const FAnimNotifyEvent& Other) const;
-
-	ENGINE_API virtual void SetTime(float NewTime, EAnimLinkMethod::Type ReferenceFrame = EAnimLinkMethod::Absolute) override;
-};
-
-// Used by UAnimSequenceBase::SortNotifies() to sort its Notifies array
-FORCEINLINE bool FAnimNotifyEvent::operator<(const FAnimNotifyEvent& Other) const
-{
-	float ATime = GetTriggerTime();
-	float BTime = Other.GetTriggerTime();
-
-	// When we've got the same time sort on the track index. Explicitly
-	// using SMALL_NUMBER here incase the underlying default changes as
-	// notifies can have an offset of KINDA_SMALL_NUMBER to be consider
-	// distinct
-	if (FMath::IsNearlyEqual(ATime, BTime, SMALL_NUMBER))
-	{
-		return TrackIndex < Other.TrackIndex;
-	}
-	else
-	{
-		return ATime < BTime;
-	}
-}
-
-/**
- * Keyframe position data for one track.  Pos(i) occurs at Time(i).  Pos.Num() always equals Time.Num().
- */
-USTRUCT()
-struct FAnimNotifyTrack
-{
-	GENERATED_USTRUCT_BODY()
-
-	UPROPERTY()
-	FName TrackName;
-
-	UPROPERTY()
-	FLinearColor TrackColor;
-
-	TArray<FAnimNotifyEvent*> Notifies;
-
-	FAnimNotifyTrack()
-		: TrackName(TEXT(""))
-		, TrackColor(FLinearColor::White)
-	{
-	}
-
-	FAnimNotifyTrack(FName InTrackName, FLinearColor InTrackColor)
-		: TrackName(InTrackName)
-		, TrackColor(InTrackColor)
-	{
-	}
-};
-
-/** native enum for curve types **/
-enum EAnimCurveFlags
-{
-	// Used as morph target curve
-	ACF_DrivesMorphTarget	= 0x00000001,
-	// Used as triggering event
-	ACF_TriggerEvent		= 0x00000002,
-	// Is editable in Sequence Editor
-	ACF_Editable			= 0x00000004,
-	// Used as a material curve
-	ACF_DrivesMaterial		= 0x00000008,
-	// Is a metadata 'curve'
-	ACF_Metadata			= 0x00000010,
-	// motifies bone track
-	ACF_DriveTrack			= 0x00000020,
-	// disabled, right now it's used by track
-	ACF_Disabled			= 0x00000040,
-
-	// default flag when created
-	ACF_DefaultCurve		= ACF_TriggerEvent | ACF_Editable,
-	// curves created from Morph Target
-	ACF_MorphTargetCurve	= ACF_DrivesMorphTarget
-};
-
-/**
- * Float curve data for one track
- */
-
-USTRUCT()
-struct FAnimCurveBase
-{
-	GENERATED_USTRUCT_BODY()
-
-	// Last observed name of the curve. We store this so we can recover from situations that
-	// mean the skeleton doesn't have a mapped name for our UID (such as a user saving the an
-	// animation but not the skeleton).
-	UPROPERTY()
-	FName		LastObservedName;
-
-	// For smart naming - management purpose - i.e. rename/delete
-	FSmartNameMapping::UID CurveUid;
-
-private:
-
-	/** Curve Type Flags */
-	UPROPERTY()
-	int32		CurveTypeFlags;
-
-public:
-	FAnimCurveBase(){}
-
-	FAnimCurveBase(FName InName, int32 InCurveTypeFlags)
-		: LastObservedName(InName)
-		, CurveTypeFlags(InCurveTypeFlags)
-	{	
-	}
-
-	FAnimCurveBase(USkeleton::AnimCurveUID Uid, int32 InCurveTypeFlags)
-		: CurveUid(Uid)
-		, CurveTypeFlags(InCurveTypeFlags)
-	{}
-
-	// To be able to use typedef'd types we need to serialize manually
-	virtual void Serialize(FArchive& Ar)
-	{
-		if(Ar.UE4Ver() >= VER_UE4_SKELETON_ADD_SMARTNAMES)
-		{
-			Ar << CurveUid;
-		}
-	}
-
-	/**
-	 * Set InFlag to bValue
-	 */
-	ENGINE_API void SetCurveTypeFlag(EAnimCurveFlags InFlag, bool bValue);
-
-	/**
-	 * Toggle the value of the specified flag
-	 */
-	ENGINE_API void ToggleCurveTypeFlag(EAnimCurveFlags InFlag);
-
-	/**
-	 * Return true if InFlag is set, false otherwise 
-	 */
-	ENGINE_API bool GetCurveTypeFlag(EAnimCurveFlags InFlag) const;
-
-	/**
-	 * Set CurveTypeFlags to NewCurveTypeFlags
-	 * This just overwrites CurveTypeFlags
-	 */
-	void SetCurveTypeFlags(int32 NewCurveTypeFlags);
-	/** 
-	 * returns CurveTypeFlags
-	 */
-	int32 GetCurveTypeFlags() const;
-};
-
-USTRUCT()
-struct FFloatCurve : public FAnimCurveBase
-{
-	GENERATED_USTRUCT_BODY()
-
-	/** Curve data for float. */
-	UPROPERTY()
-	FRichCurve	FloatCurve;
-
-	FFloatCurve(){}
-
-	FFloatCurve(FName InName, int32 InCurveTypeFlags)
-		: FAnimCurveBase(InName, InCurveTypeFlags)
-	{
-	}
-
-	FFloatCurve(USkeleton::AnimCurveUID Uid, int32 InCurveTypeFlags)
-		: FAnimCurveBase(Uid, InCurveTypeFlags)
-	{}
-
-	virtual void Serialize(FArchive& Ar) override
-	{
-		FAnimCurveBase::Serialize(Ar);
-	}
-
-	// we don't want to have = operator. This only copies curves, but leaving naming and everything else intact. 
-	void CopyCurve(FFloatCurve& SourceCurve);
-	float Evaluate(float CurrentTime, float BlendWeight) const;
-	void UpdateOrAddKey(float NewKey, float CurrentTime);
-	void Resize(float NewLength);
-};
-
-USTRUCT()
-struct FVectorCurve : public FAnimCurveBase
-{
-	GENERATED_USTRUCT_BODY()
-
-	enum EIndex
-	{
-		X = 0, 
-		Y, 
-		Z, 
-		Max
-	};
-
-	/** Curve data for float. */
-	UPROPERTY()
-	FRichCurve	FloatCurves[3];
-
-	FVectorCurve(){}
-
-	FVectorCurve(FName InName, int32 InCurveTypeFlags)
-		: FAnimCurveBase(InName, InCurveTypeFlags)
-	{
-	}
-
-	FVectorCurve(USkeleton::AnimCurveUID Uid, int32 InCurveTypeFlags)
-		: FAnimCurveBase(Uid, InCurveTypeFlags)
-	{}
-
-	virtual void Serialize(FArchive& Ar) override
-	{
-		FAnimCurveBase::Serialize(Ar);
-	}
-
-	// we don't want to have = operator. This only copies curves, but leaving naming and everything else intact. 
-	void CopyCurve(FVectorCurve& SourceCurve);
-	FVector Evaluate(float CurrentTime, float BlendWeight) const;
-	void UpdateOrAddKey(const FVector& NewKey, float CurrentTime);
-	bool DoesContainKey() const { return (FloatCurves[0].GetNumKeys() > 0 || FloatCurves[1].GetNumKeys() > 0 || FloatCurves[2].GetNumKeys() > 0);}
-	void Resize(float NewLength);
-};
-
-USTRUCT()
-struct FTransformCurve: public FAnimCurveBase
-{
-	GENERATED_USTRUCT_BODY()
-
-	/** Curve data for each transform. */
-	UPROPERTY()
-	FVectorCurve	TranslationCurve;
-
-	/** Rotation curve - right now we use euler because quat also doesn't provide linear interpolation - curve editor can't handle quat interpolation
-	 * If you hit gimbal lock, you should add extra key to fix it. This will cause gimbal lock. 
-	 * @TODO: Eventually we'll need FRotationCurve that would contain rotation curve - that will interpolate as slerp or as quaternion 
-	 */
-	UPROPERTY()
-	FVectorCurve	RotationCurve;
-
-	UPROPERTY()
-	FVectorCurve	ScaleCurve;
-
-	FTransformCurve(){}
-
-	FTransformCurve(FName InName, int32 InCurveTypeFlags)
-		: FAnimCurveBase(InName, InCurveTypeFlags)
-	{
-	}
-
-	FTransformCurve(USkeleton::AnimCurveUID Uid, int32 InCurveTypeFlags)
-		: FAnimCurveBase(Uid, InCurveTypeFlags)
-	{}
-
-	virtual void Serialize(FArchive& Ar) override
-	{
-		FAnimCurveBase::Serialize(Ar);
-	}
-
-	// we don't want to have = operator. This only copies curves, but leaving naming and everything else intact. 
-	void CopyCurve(FTransformCurve& SourceCurve);
-	FTransform Evaluate(float CurrentTime, float BlendWeight) const;
-	ENGINE_API void UpdateOrAddKey(const FTransform& NewKey, float CurrentTime);
-	void Resize(float NewLength);
-};
-/**
- * Raw Curve data for serialization
- */
-USTRUCT()
-struct FRawCurveTracks
-{
-	GENERATED_USTRUCT_BODY()
-
-	enum ESupportedCurveType
-	{
-		FloatType,
-		VectorType,
-		TransformType,
-		Max, 
-	};
-
-	UPROPERTY()
-	TArray<FFloatCurve>		FloatCurves;
-
-#if WITH_EDITORONLY_DATA
-	/**
-	 * @note : Currently VectorCurves are not evaluated or used for anything else but transient data for modifying bone track
-	 *			Note that it doesn't have UPROPERTY tag yet. In the future, we'd like this to be serialized, but not for now
-	 **/
-	UPROPERTY(transient)
-	TArray<FVectorCurve>	VectorCurves;
-
-	/**
-	 * @note : Currently VectorCurves are not evaluated or used for anything else but transient data for modifying bone track
-	 *			Note that it doesn't have UPROPERTY tag yet. In the future, we'd like this to be serialized, but not for now
-	 **/
-	UPROPERTY()
-	TArray<FTransformCurve>	TransformCurves;
-#endif // #if WITH_EDITOR_DATA
-	/**
-	 * Evaluate curve data at the time CurrentTime, and add to Instance. It only evaluates Float Curve for now
-	 */
-	void EvaluateCurveData(class UAnimInstance* Instance, float CurrentTime, float BlendWeight ) const;
-
-#if WITH_EDITOR
-	/**
-	 *	Evaluate transform curves 
-	 */
-	ENGINE_API void EvaluateTransformCurveData(USkeleton * Skeleton, TMap<FName, FTransform>&OutCurves, float CurrentTime, float BlendWeight) const;
-#endif // WITH_EDITOR
-	/**
-	 * Find curve data based on the curve UID
-	 */
-	ENGINE_API FAnimCurveBase * GetCurveData(USkeleton::AnimCurveUID Uid, ESupportedCurveType SupportedCurveType = FloatType);
-	/**
-	 * Add new curve from the provided UID and return true if success
-	 * bVectorInterpCurve == true, then it will create FVectorCuve, otherwise, FFloatCurve
-	 */
-	ENGINE_API bool AddCurveData(USkeleton::AnimCurveUID Uid, int32 CurveFlags = ACF_DefaultCurve, ESupportedCurveType SupportedCurveType = FloatType);
-
-	/**
-	 * Delete curve data 
-	 */
-	ENGINE_API bool DeleteCurveData(USkeleton::AnimCurveUID Uid, ESupportedCurveType SupportedCurveType = FloatType);
-	/**
-	 * Duplicate curve data
-	 * 
-	 */
-	ENGINE_API bool DuplicateCurveData(USkeleton::AnimCurveUID ToCopyUid, USkeleton::AnimCurveUID NewUid, ESupportedCurveType SupportedCurveType = FloatType);
-
-	/**
-	 * Updates the LastObservedName field of the curves from the provided name container
-	 */
-	ENGINE_API void UpdateLastObservedNames(FSmartNameMapping* NameMapping, ESupportedCurveType SupportedCurveType = FloatType);
-
-	/** 
-	 * Serialize
-	 */
-	void Serialize(FArchive& Ar);
-
-	/*
-	 * resize curve length. If longer, it doesn't do any. If shorter, remove previous keys and add new key to the end of the frame. 
-	 */
-	void Resize(float TotalLength);
-	/** 
-	 * Clear all keys
-	 */
-	void Empty()
-	{
-		FloatCurves.Empty();
-#if WITH_EDITORONLY_DATA
-		VectorCurves.Empty();
-		TransformCurves.Empty();
-#endif
-	}
-private:
-	/** 
-	 * Adding vector curve support - this is all transient data for now. This does not save and all these data will be baked into RawAnimationData
-	 */
-
-	/**
-	 * Find curve data based on the curve UID
-	 */
-	template <typename DataType>
-	DataType * GetCurveDataImpl(TArray<DataType>& Curves, USkeleton::AnimCurveUID Uid);
-
-	/**
-	 * Add new curve from the provided UID and return true if success
-	 * bVectorInterpCurve == true, then it will create FVectorCuve, otherwise, FFloatCurve
-	 */
-	template <typename DataType>
-	bool AddCurveDataImpl(TArray<DataType>& Curves, USkeleton::AnimCurveUID Uid, int32 CurveFlags);
-	/**
-	 * Delete curve data 
-	 */
-	template <typename DataType>
-	bool DeleteCurveDataImpl(TArray<DataType>& Curves, USkeleton::AnimCurveUID Uid);
-	/**
-	 * Duplicate curve data
-	 * 
-	 */
-	template <typename DataType>
-	bool DuplicateCurveDataImpl(TArray<DataType>& Curves, USkeleton::AnimCurveUID ToCopyUid, USkeleton::AnimCurveUID NewUid);
-
-	/**
-	 * Updates the LastObservedName field of the curves from the provided name container
-	 */
-	template <typename DataType>
-	void UpdateLastObservedNamesImpl(TArray<DataType>& Curves, FSmartNameMapping* NameMapping);
-};
-
 UENUM()
 enum ETypeAdvanceAnim
 {
@@ -604,7 +45,7 @@ class UAnimSequenceBase : public UAnimationAsset
 	 * Raw uncompressed float curve data 
 	 */
 	UPROPERTY()
-	FRawCurveTracks RawCurveData;
+	struct FRawCurveTracks RawCurveData;
 
 #if WITH_EDITORONLY_DATA
 	// if you change Notifies array, this will need to be rebuilt
@@ -640,7 +81,7 @@ class UAnimSequenceBase : public UAnimationAsset
 	ENGINE_API void GetAnimNotifiesFromDeltaPositions(const float& PreviousPosition, const float & CurrentPosition, TArray<const FAnimNotifyEvent *>& OutActiveNotifies) const;
 
 	/** Evaluate curve data to Instance at the time of CurrentTime **/
-	ENGINE_API virtual void EvaluateCurveData(class UAnimInstance* Instance, float CurrentTime, float BlendWeight ) const;
+	ENGINE_API virtual void EvaluateCurveData(FBlendedCurve& OutCurve, float CurrentTime) const;
 
 	/**
 	 * return true if this is valid additive animation
@@ -685,6 +126,16 @@ class UAnimSequenceBase : public UAnimationAsset
 	virtual float GetMaxCurrentTime() override { return SequenceLength; }
 	// End of UAnimationAsset interface
 
+	/**
+	* Get Bone Transform of the Time given, relative to Parent for all RequiredBones
+	* This returns different transform based on additive or not. Or what kind of additive.
+	*
+	* @param	OutPose				Pose object to fill
+	* @param	OutCurve			Curves to fill
+	* @param	ExtractionContext	Extraction Context (position, looping, root motion, etc.)
+	*/
+	ENGINE_API virtual void GetAnimationPose(struct FCompactPose& OutPose, FBlendedCurve& OutCurve, const FAnimExtractContext& ExtractionContext) const PURE_VIRTUAL(UAnimSequenceBase::GetAnimationPose, );
+	
 	virtual void OnAssetPlayerTickedInternal(FAnimAssetTickContext &Context, const float PreviousTime, const float MoveDelta, const FAnimTickRecord &Instance, class UAnimInstance* InstanceOwner) const;
 
 	virtual bool HasRootMotion() const { return false; }
diff --git a/Engine/Source/Runtime/Engine/Classes/Animation/AnimSingleNodeInstance.h b/Engine/Source/Runtime/Engine/Classes/Animation/AnimSingleNodeInstance.h
--- a/Engine/Source/Runtime/Engine/Classes/Animation/AnimSingleNodeInstance.h
+++ b/Engine/Source/Runtime/Engine/Classes/Animation/AnimSingleNodeInstance.h
@@ -63,7 +63,7 @@ class ENGINE_API UAnimSingleNodeInstance : public UAnimInstance
 
 protected:
 	virtual void Montage_Advance(float DeltaTime) override;
-	void InternalBlendSpaceEvaluatePose(class UBlendSpaceBase* BlendSpace, TArray<FBlendSampleData>& BlendSampleDataCache, FCompactPose& OutPose);
+	void InternalBlendSpaceEvaluatePose(class UBlendSpaceBase* BlendSpace, TArray<FBlendSampleData>& BlendSampleDataCache, FPoseContext& OutContext);
 	// End UAnimInstance interface
 public:
 
diff --git a/Engine/Source/Runtime/Engine/Classes/Animation/BlendSpaceBase.h b/Engine/Source/Runtime/Engine/Classes/Animation/BlendSpaceBase.h
--- a/Engine/Source/Runtime/Engine/Classes/Animation/BlendSpaceBase.h
+++ b/Engine/Source/Runtime/Engine/Classes/Animation/BlendSpaceBase.h
@@ -232,7 +232,7 @@ public:
 	 * When you use blend per bone, allows rotation to blend in mesh space. This only works if this does not contain additive animation samples
 	 * This is more performance intensive
 	 */
-	UPROPERTY(EditAnywhere, Category=SampleInterpolation)
+	UPROPERTY()
 	bool bRotationBlendInMeshSpace;
 
 	/** Number of dimensions for this blend space (1 or 2) **/
@@ -274,6 +274,11 @@ public:
 #endif
 	// End of UAnimationAsset interface
 
+	/**
+	 * BlendSpace Get Animation Pose function
+	 */
+	void GetAnimationPose(TArray<FBlendSampleData>& BlendSampleDataCache, /*out*/ FCompactPose& OutPose, /*out*/ FBlendedCurve& OutCurve);
+
 	/** Accessor for blend parameter **/
 	ENGINE_API const FBlendParameter& GetBlendParameter(int32 Index)
 	{
diff --git a/Engine/Source/Runtime/Engine/Classes/Components/SkeletalMeshComponent.h b/Engine/Source/Runtime/Engine/Classes/Components/SkeletalMeshComponent.h
--- a/Engine/Source/Runtime/Engine/Classes/Components/SkeletalMeshComponent.h
+++ b/Engine/Source/Runtime/Engine/Classes/Components/SkeletalMeshComponent.h
@@ -7,6 +7,7 @@
 #include "PhysicsEngine/PhysicsConstraintComponent.h"
 #include "SkeletalMeshTypes.h"
 #include "Animation/AnimationAsset.h"
+#include "AnimCurveTypes.h"
 #include "SkeletalMeshComponent.generated.h"
 
 class UAnimInstance;
@@ -27,6 +28,9 @@ struct FAnimationEvaluationContext
 	TArray<FActiveVertexAnim> VertexAnims;
 	FVector RootBoneTranslation;
 
+	// Double buffer curve data
+	FBlendedCurve	Curve;
+
 	// Are we performing interpolation this tick
 	bool bDoInterpolation;
 
@@ -36,6 +40,9 @@ struct FAnimationEvaluationContext
 	// Are we storing data in cache bones this tick
 	bool bDuplicateToCacheBones;
 
+	// duplicate the cache curves
+	bool bDuplicateToCacheCurve;
+
 	FAnimationEvaluationContext()
 	{
 		Clear();
@@ -318,6 +325,9 @@ public:
 	UPROPERTY(Transient)
 	TArray<FTransform> CachedSpaceBases;
 
+	/** Cached Curve result for Update Rate optimization */
+	FBlendedCurve CachedCurve;
+
 	/** Used to scale speed of all animations on this skeletal mesh. */
 	UPROPERTY(EditAnywhere, AdvancedDisplay, BlueprintReadWrite, Category=Animation)
 	float GlobalAnimRateScale;
@@ -924,7 +934,7 @@ public:
 	* @param	OutVertexAnims			Active vertex animations
 	* @param	OutRootBoneTranslation	Calculated root bone translation
 	*/
-	void PerformAnimationEvaluation(const USkeletalMesh* InSkeletalMesh, UAnimInstance* InAnimInstance, TArray<FTransform>& OutSpaceBases, TArray<FTransform>& OutLocalAtoms, TArray<FActiveVertexAnim>& OutVertexAnims, FVector& OutRootBoneTranslation) const;
+	void PerformAnimationEvaluation(const USkeletalMesh* InSkeletalMesh, UAnimInstance* InAnimInstance, TArray<FTransform>& OutSpaceBases, TArray<FTransform>& OutLocalAtoms, TArray<FActiveVertexAnim>& OutVertexAnims, FVector& OutRootBoneTranslation, FBlendedCurve& OutCurve) const;
 	void PostAnimEvaluation( FAnimationEvaluationContext& EvaluationContext );
 
 	/**
@@ -1198,7 +1208,7 @@ protected:
 	
 private:
 	/** Evaluate Anim System **/
-	void EvaluateAnimation(const USkeletalMesh* InSkeletalMesh, UAnimInstance* InAnimInstance, TArray<FTransform>& OutLocalAtoms, TArray<struct FActiveVertexAnim>& OutVertexAnims, FVector& OutRootBoneTranslation) const;
+	void EvaluateAnimation(const USkeletalMesh* InSkeletalMesh, UAnimInstance* InAnimInstance, TArray<FTransform>& OutLocalAtoms, TArray<struct FActiveVertexAnim>& OutVertexAnims, FVector& OutRootBoneTranslation, FBlendedCurve& OutCurve) const;
 
 	/**
 	* Take the SourceAtoms array (translation vector, rotation quaternion and scale vector) and update the array of component-space bone transformation matrices (DestSpaceBases).
@@ -1226,7 +1236,7 @@ private:
 
 public:
 	// Parallel evaluation wrappers
-	void ParallelAnimationEvaluation() { PerformAnimationEvaluation(AnimEvaluationContext.SkeletalMesh, AnimEvaluationContext.AnimInstance, AnimEvaluationContext.SpaceBases, AnimEvaluationContext.LocalAtoms, AnimEvaluationContext.VertexAnims, AnimEvaluationContext.RootBoneTranslation); }
+	void ParallelAnimationEvaluation() { PerformAnimationEvaluation(AnimEvaluationContext.SkeletalMesh, AnimEvaluationContext.AnimInstance, AnimEvaluationContext.SpaceBases, AnimEvaluationContext.LocalAtoms, AnimEvaluationContext.VertexAnims, AnimEvaluationContext.RootBoneTranslation, AnimEvaluationContext.Curve); }
 	void CompleteParallelAnimationEvaluation(bool bDoPostAnimEvaluation)
 	{
 		ParallelAnimationEvaluationTask.SafeRelease(); //We are done with this task now, clean up!
diff --git a/Engine/Source/Runtime/Engine/Private/Animation/AimOffsetBlendSpace.cpp b/Engine/Source/Runtime/Engine/Private/Animation/AimOffsetBlendSpace.cpp
--- a/Engine/Source/Runtime/Engine/Private/Animation/AimOffsetBlendSpace.cpp
+++ b/Engine/Source/Runtime/Engine/Private/Animation/AimOffsetBlendSpace.cpp
@@ -12,6 +12,7 @@
 UAimOffsetBlendSpace::UAimOffsetBlendSpace(const FObjectInitializer& ObjectInitializer)
 	: Super(ObjectInitializer)
 {
+	bRotationBlendInMeshSpace = true;
 }
 
 bool UAimOffsetBlendSpace::IsValidAdditive()  const
diff --git a/Engine/Source/Runtime/Engine/Private/Animation/AimOffsetBlendSpace1D.cpp b/Engine/Source/Runtime/Engine/Private/Animation/AimOffsetBlendSpace1D.cpp
--- a/Engine/Source/Runtime/Engine/Private/Animation/AimOffsetBlendSpace1D.cpp
+++ b/Engine/Source/Runtime/Engine/Private/Animation/AimOffsetBlendSpace1D.cpp
@@ -12,6 +12,7 @@
 UAimOffsetBlendSpace1D::UAimOffsetBlendSpace1D(const FObjectInitializer& ObjectInitializer)
 	: Super(ObjectInitializer)
 {
+	bRotationBlendInMeshSpace = true;
 }
 
 bool UAimOffsetBlendSpace1D::IsValidAdditive()  const
diff --git a/Engine/Source/Runtime/Engine/Private/Animation/AnimComposite.cpp b/Engine/Source/Runtime/Engine/Private/Animation/AnimComposite.cpp
--- a/Engine/Source/Runtime/Engine/Private/Animation/AnimComposite.cpp
+++ b/Engine/Source/Runtime/Engine/Private/Animation/AnimComposite.cpp
@@ -32,4 +32,16 @@ void UAnimComposite::OnAssetPlayerTickedInternal(FAnimAssetTickContext &Context,
 	Super::OnAssetPlayerTickedInternal(Context, PreviousTime, MoveDelta, Instance, InstanceOwner);
 
 	ExtractRootMotionFromTrack(AnimationTrack, PreviousTime, PreviousTime + MoveDelta, Context.RootMotionMovementParams);
+}
+
+void UAnimComposite::GetAnimationPose(FCompactPose& OutPose, FBlendedCurve& OutCurve, const FAnimExtractContext& ExtractionContext) const
+{
+	AnimationTrack.GetAnimationPose(OutPose, OutCurve, ExtractionContext);
+
+	FBlendedCurve CompositeCurve;
+	CompositeCurve.InitFrom(OutCurve);
+	EvaluateCurveData(CompositeCurve, ExtractionContext.CurrentTime);
+
+	// combine both curve
+	OutCurve.Combine(CompositeCurve);
 }
\ No newline at end of file
diff --git a/Engine/Source/Runtime/Engine/Private/Animation/AnimCompositeBase.cpp b/Engine/Source/Runtime/Engine/Private/Animation/AnimCompositeBase.cpp
--- a/Engine/Source/Runtime/Engine/Private/Animation/AnimCompositeBase.cpp
+++ b/Engine/Source/Runtime/Engine/Private/Animation/AnimCompositeBase.cpp
@@ -484,6 +484,66 @@ void FAnimTrack::SortAnimSegments()
 }
 #endif
 
+void FAnimTrack::GetAnimationPose(/*out*/ FCompactPose& OutPose,/*out*/ FBlendedCurve& OutCurve, const FAnimExtractContext& ExtractionContext) const
+{
+	TArray<FCompactPose> SourcePoses;
+	TArray<float> SourceWeights;
+	TArray<FBlendedCurve> SourceCurves;
+	float TotalWeight = 0.f;
+
+	float CurrentTime = FMath::Clamp(ExtractionContext.CurrentTime, 0.f, GetLength());
+
+	// first get all the montage instance weight this slot node has
+	for(int32 I = 0; I<AnimSegments.Num(); ++I)
+	{
+		const FAnimSegment& AnimSegment = AnimSegments[I];
+
+		float PositionInAnim = 0.f;
+		float Weight = 0.f;
+		UAnimSequenceBase* AnimRef = AnimSegment.GetAnimationData(CurrentTime, PositionInAnim, Weight);
+
+		// make this to be 1 function
+		if(AnimRef && (Weight > ZERO_ANIMWEIGHT_THRESH))
+		{
+			// todo anim: hack - until we fix animcomposite
+			const int32 NewIndex = SourceWeights.AddUninitialized(1);
+			SourcePoses.Add(FCompactPose());
+			SourceCurves.Add(FBlendedCurve());
+			SourcePoses[NewIndex].SetBoneContainer(&OutPose.GetBoneContainer());
+			SourceCurves[NewIndex].InitFrom(OutCurve);
+			SourceWeights[NewIndex] = Weight;
+			TotalWeight += Weight;
+
+			// Copy passed in Extraction Context, but override position and looping parameters.
+			FAnimExtractContext SequenceExtractionContext(ExtractionContext);
+			SequenceExtractionContext.CurrentTime = PositionInAnim;
+			SequenceExtractionContext.bExtractRootMotion &= AnimRef->HasRootMotion();
+				
+			AnimRef->GetAnimationPose(SourcePoses[NewIndex], SourceCurves[NewIndex], SequenceExtractionContext);
+		}
+	}
+
+	if(SourcePoses.Num() == 0)
+	{
+		OutPose.ResetToRefPose();
+	}
+	else if(SourcePoses.Num() == 1)
+	{
+		OutPose = SourcePoses[0];
+		OutCurve = SourceCurves[0];
+	}
+	else
+	{
+		// If we have SourcePoses.Num() > 0, then we will have a non zero weight.
+		check(TotalWeight >= ZERO_ANIMWEIGHT_THRESH);
+		for(int32 I = 0; I < SourceWeights.Num(); ++I)
+		{
+			// normalize I
+			SourceWeights[I] /= TotalWeight;
+		}
+		FAnimationRuntime::BlendPosesTogether(SourcePoses, SourceCurves, SourceWeights, OutPose, OutCurve);
+	}
+}
 ///////////////////////////////////////////////////////
 // UAnimCompositeBase
 ///////////////////////////////////////////////////////
diff --git a/Engine/Source/Runtime/Engine/Private/Animation/AnimCurveTypes.cpp b/Engine/Source/Runtime/Engine/Private/Animation/AnimCurveTypes.cpp
new file mode 100644
--- /dev/null
+++ b/Engine/Source/Runtime/Engine/Private/Animation/AnimCurveTypes.cpp
@@ -0,0 +1,577 @@
+// Copyright 1998-2015 Epic Games, Inc. All Rights Reserved.
+
+#include "EnginePrivate.h"
+#include "AnimCurveTypes.h"
+#include "Animation/AnimInstance.h"
+#include "AnimationRuntime.h"
+
+/////////////////////////////////////////////////////
+// FFloatCurve
+
+void FAnimCurveBase::SetCurveTypeFlag(EAnimCurveFlags InFlag, bool bValue)
+{
+	if (bValue)
+	{
+		CurveTypeFlags |= InFlag;
+	}
+	else
+	{
+		CurveTypeFlags &= ~InFlag;
+	}
+}
+
+void FAnimCurveBase::ToggleCurveTypeFlag(EAnimCurveFlags InFlag)
+{
+	bool Current = GetCurveTypeFlag(InFlag);
+	SetCurveTypeFlag(InFlag, !Current);
+}
+
+bool FAnimCurveBase::GetCurveTypeFlag(EAnimCurveFlags InFlag) const
+{
+	return (CurveTypeFlags & InFlag) != 0;
+}
+
+
+void FAnimCurveBase::SetCurveTypeFlags(int32 NewCurveTypeFlags)
+{
+	CurveTypeFlags = NewCurveTypeFlags;
+}
+
+int32 FAnimCurveBase::GetCurveTypeFlags() const
+{
+	return CurveTypeFlags;
+}
+
+////////////////////////////////////////////////////
+//  FFloatCurve
+
+// we don't want to have = operator. This only copies curves, but leaving naming and everything else intact. 
+void FFloatCurve::CopyCurve(FFloatCurve& SourceCurve)
+{
+	FloatCurve = SourceCurve.FloatCurve;
+}
+
+float FFloatCurve::Evaluate(float CurrentTime) const
+{
+	return FloatCurve.Eval(CurrentTime);
+}
+
+void FFloatCurve::UpdateOrAddKey(float NewKey, float CurrentTime)
+{
+	FloatCurve.UpdateOrAddKey(CurrentTime, NewKey);
+}
+
+void FFloatCurve::Resize(float NewLength)
+{
+	FloatCurve.ResizeTimeRange(0, NewLength);
+}
+////////////////////////////////////////////////////
+//  FVectorCurve
+
+// we don't want to have = operator. This only copies curves, but leaving naming and everything else intact. 
+void FVectorCurve::CopyCurve(FVectorCurve& SourceCurve)
+{
+	FMemory::Memcpy(FloatCurves, SourceCurve.FloatCurves);
+}
+
+FVector FVectorCurve::Evaluate(float CurrentTime, float BlendWeight) const
+{
+	FVector Value;
+
+	Value.X = FloatCurves[X].Eval(CurrentTime)*BlendWeight;
+	Value.Y = FloatCurves[Y].Eval(CurrentTime)*BlendWeight;
+	Value.Z = FloatCurves[Z].Eval(CurrentTime)*BlendWeight;
+
+	return Value;
+}
+
+void FVectorCurve::UpdateOrAddKey(const FVector& NewKey, float CurrentTime)
+{
+	FloatCurves[X].UpdateOrAddKey(CurrentTime, NewKey.X);
+	FloatCurves[Y].UpdateOrAddKey(CurrentTime, NewKey.Y);
+	FloatCurves[Z].UpdateOrAddKey(CurrentTime, NewKey.Z);
+}
+
+void FVectorCurve::Resize(float NewLength)
+{
+	FloatCurves[X].ResizeTimeRange(0, NewLength);
+	FloatCurves[Y].ResizeTimeRange(0, NewLength);
+	FloatCurves[Z].ResizeTimeRange(0, NewLength);
+}
+////////////////////////////////////////////////////
+//  FTransformCurve
+
+// we don't want to have = operator. This only copies curves, but leaving naming and everything else intact. 
+void FTransformCurve::CopyCurve(FTransformCurve& SourceCurve)
+{
+	TranslationCurve.CopyCurve(SourceCurve.TranslationCurve);
+	RotationCurve.CopyCurve(SourceCurve.RotationCurve);
+	ScaleCurve.CopyCurve(SourceCurve.ScaleCurve);
+}
+
+FTransform FTransformCurve::Evaluate(float CurrentTime, float BlendWeight) const
+{
+	FTransform Value;
+	Value.SetTranslation(TranslationCurve.Evaluate(CurrentTime, BlendWeight));
+	if (ScaleCurve.DoesContainKey())
+	{
+		Value.SetScale3D(ScaleCurve.Evaluate(CurrentTime, BlendWeight));
+	}
+	else
+	{
+		Value.SetScale3D(FVector(1.f));
+	}
+
+	// blend rotation float curve
+	FVector RotationAsVector = RotationCurve.Evaluate(CurrentTime, BlendWeight);
+	// pitch, yaw, roll order - please check AddKey function
+	FRotator Rotator(RotationAsVector.Y, RotationAsVector.Z, RotationAsVector.X);
+	Value.SetRotation(FQuat(Rotator));
+
+	return Value;
+}
+
+void FTransformCurve::UpdateOrAddKey(const FTransform& NewKey, float CurrentTime)
+{
+	TranslationCurve.UpdateOrAddKey(NewKey.GetTranslation(), CurrentTime);
+	// pitch, yaw, roll order - please check Evaluate function
+	FVector RotationAsVector;
+	FRotator Rotator = NewKey.GetRotation().Rotator();
+	RotationAsVector.X = Rotator.Roll;
+	RotationAsVector.Y = Rotator.Pitch;
+	RotationAsVector.Z = Rotator.Yaw;
+
+	RotationCurve.UpdateOrAddKey(RotationAsVector, CurrentTime);
+	ScaleCurve.UpdateOrAddKey(NewKey.GetScale3D(), CurrentTime);
+}
+
+void FTransformCurve::Resize(float NewLength)
+{
+	TranslationCurve.Resize(NewLength);
+	RotationCurve.Resize(NewLength);
+	ScaleCurve.Resize(NewLength);
+}
+
+/////////////////////////////////////////////////////
+// FRawCurveTracks
+
+void FRawCurveTracks::EvaluateCurveData( FBlendedCurve& Curves, float CurrentTime ) const
+{
+	// evaluate the curve data at the CurrentTime and add to Instance
+	for(auto CurveIter = FloatCurves.CreateConstIterator(); CurveIter; ++CurveIter)
+	{
+		const FFloatCurve& Curve = *CurveIter;
+		Curves.Set(Curve.CurveUid, Curve.Evaluate(CurrentTime), Curve.GetCurveTypeFlags());
+	}
+}
+
+#if WITH_EDITOR
+/**
+ * Since we don't care about blending, we just change this decoration to OutCurves
+ * @TODO : Fix this if we're saving vectorcurves and blending
+ */
+void FRawCurveTracks::EvaluateTransformCurveData(USkeleton * Skeleton, TMap<FName, FTransform>&OutCurves, float CurrentTime, float BlendWeight) const
+{
+	check (Skeleton);
+	// evaluate the curve data at the CurrentTime and add to Instance
+	for(auto CurveIter = TransformCurves.CreateConstIterator(); CurveIter; ++CurveIter)
+	{
+		const FTransformCurve& Curve = *CurveIter;
+
+		// if disabled, do not handle
+		if (Curve.GetCurveTypeFlag(ACF_Disabled))
+		{
+			continue;
+		}
+
+		FSmartNameMapping* NameMapping = Skeleton->SmartNames.GetContainer(USkeleton::AnimTrackCurveMappingName);
+
+		// Add or retrieve curve
+		FName CurveName;
+		
+		// make sure it was added
+		if (ensure (NameMapping->GetName(Curve.CurveUid, CurveName)))
+		{
+			// note we're not checking Curve.GetCurveTypeFlags() yet
+			FTransform & Value = OutCurves.FindOrAdd(CurveName);
+			Value = Curve.Evaluate(CurrentTime, BlendWeight);
+		}
+	}
+}
+#endif
+FAnimCurveBase * FRawCurveTracks::GetCurveData(USkeleton::AnimCurveUID Uid, ESupportedCurveType SupportedCurveType /*= FloatType*/)
+{
+	switch (SupportedCurveType)
+	{
+#if WITH_EDITOR
+	case VectorType:
+		return GetCurveDataImpl<FVectorCurve>(VectorCurves, Uid);
+	case TransformType:
+		return GetCurveDataImpl<FTransformCurve>(TransformCurves, Uid);
+#endif // WITH_EDITOR
+	case FloatType:
+	default:
+		return GetCurveDataImpl<FFloatCurve>(FloatCurves, Uid);
+	}
+}
+
+bool FRawCurveTracks::DeleteCurveData(USkeleton::AnimCurveUID Uid, ESupportedCurveType SupportedCurveType /*= FloatType*/)
+{
+	switch(SupportedCurveType)
+	{
+#if WITH_EDITOR
+	case VectorType:
+		return DeleteCurveDataImpl<FVectorCurve>(VectorCurves, Uid);
+	case TransformType:
+		return DeleteCurveDataImpl<FTransformCurve>(TransformCurves, Uid);
+#endif // WITH_EDITOR
+	case FloatType:
+	default:
+		return DeleteCurveDataImpl<FFloatCurve>(FloatCurves, Uid);
+	}
+}
+
+bool FRawCurveTracks::AddCurveData(USkeleton::AnimCurveUID Uid, int32 CurveFlags /*= ACF_DefaultCurve*/, ESupportedCurveType SupportedCurveType /*= FloatType*/)
+{
+	switch(SupportedCurveType)
+	{
+#if WITH_EDITOR
+	case VectorType:
+		return AddCurveDataImpl<FVectorCurve>(VectorCurves, Uid, CurveFlags);
+	case TransformType:
+		return AddCurveDataImpl<FTransformCurve>(TransformCurves, Uid, CurveFlags);
+#endif // WITH_EDITOR
+	case FloatType:
+	default:
+		return AddCurveDataImpl<FFloatCurve>(FloatCurves, Uid, CurveFlags);
+	}
+}
+
+void FRawCurveTracks::Resize(float TotalLength)
+{
+	for (auto& Curve: FloatCurves)
+	{
+		Curve.Resize(TotalLength);
+	}
+
+#if WITH_EDITORONLY_DATA
+	for(auto& Curve: VectorCurves)
+	{
+		Curve.Resize(TotalLength);
+	}
+
+	for(auto& Curve: TransformCurves)
+	{
+		Curve.Resize(TotalLength);
+	}
+#endif
+}
+
+void FRawCurveTracks::Serialize(FArchive& Ar)
+{
+	// @TODO: If we're about to serialize vector curve, add here
+	if(Ar.UE4Ver() >= VER_UE4_SKELETON_ADD_SMARTNAMES)
+	{
+		for(FFloatCurve& Curve : FloatCurves)
+		{
+			Curve.Serialize(Ar);
+		}
+	}
+#if WITH_EDITORONLY_DATA
+	if( !Ar.IsCooking() )
+	{
+		if( Ar.UE4Ver() >= VER_UE4_ANIMATION_ADD_TRACKCURVES )
+		{
+			for( FTransformCurve& Curve : TransformCurves )
+			{
+				Curve.Serialize( Ar );
+			}
+
+		}
+	}
+#endif // WITH_EDITORONLY_DATA
+
+	SortFloatCurvesByUID();
+}
+
+void FRawCurveTracks::UpdateLastObservedNames(FSmartNameMapping* NameMapping, ESupportedCurveType SupportedCurveType /*= FloatType*/)
+{
+	switch(SupportedCurveType)
+	{
+#if WITH_EDITOR
+	case VectorType:
+		UpdateLastObservedNamesImpl<FVectorCurve>(VectorCurves, NameMapping);
+		break;
+	case TransformType:
+		UpdateLastObservedNamesImpl<FTransformCurve>(TransformCurves, NameMapping);
+		break;
+#endif // WITH_EDITOR
+	case FloatType:
+	default:
+		UpdateLastObservedNamesImpl<FFloatCurve>(FloatCurves, NameMapping);
+	}
+}
+
+bool FRawCurveTracks::DuplicateCurveData(USkeleton::AnimCurveUID ToCopyUid, USkeleton::AnimCurveUID NewUid, ESupportedCurveType SupportedCurveType /*= FloatType*/)
+{
+	switch(SupportedCurveType)
+	{
+#if WITH_EDITOR
+	case VectorType:
+		return DuplicateCurveDataImpl<FVectorCurve>(VectorCurves, ToCopyUid, NewUid);
+	case TransformType:
+		return DuplicateCurveDataImpl<FTransformCurve>(TransformCurves, ToCopyUid, NewUid);
+#endif // WITH_EDITOR
+	case FloatType:
+	default:
+		return DuplicateCurveDataImpl<FFloatCurve>(FloatCurves, ToCopyUid, NewUid);
+	}
+}
+
+///////////////////////////////////
+// @TODO: REFACTOR THIS IF WE'RE SERIALIZING VECTOR CURVES
+//
+// implementation template functions to accomodate FloatCurve and VectorCurve
+// for now vector curve isn't used in run-time, so it's useless outside of editor
+// so just to reduce cost of run-time, functionality is split. 
+// this split worries me a bit because if the name conflict happens this will break down w.r.t. smart naming
+// currently vector curve is not saved and not evaluated, so it will be okay since the name doesn't matter much, 
+// but this has to be refactored once we'd like to move onto serialize
+///////////////////////////////////
+template <typename DataType>
+DataType * FRawCurveTracks::GetCurveDataImpl(TArray<DataType> & Curves, USkeleton::AnimCurveUID Uid)
+{
+	for(DataType& Curve : Curves)
+	{
+		if(Curve.CurveUid == Uid)
+		{
+			return &Curve;
+		}
+	}
+
+	return NULL;
+}
+
+template <typename DataType>
+bool FRawCurveTracks::DeleteCurveDataImpl(TArray<DataType> & Curves, USkeleton::AnimCurveUID Uid)
+{
+	for(int32 Idx = 0; Idx < Curves.Num(); ++Idx)
+	{
+		if(Curves[Idx].CurveUid == Uid)
+		{
+			Curves.RemoveAt(Idx);
+			return true;
+		}
+	}
+
+	return false;
+}
+
+template <typename DataType>
+bool FRawCurveTracks::AddCurveDataImpl(TArray<DataType> & Curves, USkeleton::AnimCurveUID Uid, int32 CurveFlags)
+{
+	if(GetCurveDataImpl<DataType>(Curves, Uid) == NULL)
+	{
+		Curves.Add(DataType(Uid, CurveFlags));
+		return true;
+	}
+	return false;
+}
+
+template <typename DataType>
+void FRawCurveTracks::UpdateLastObservedNamesImpl(TArray<DataType> & Curves, FSmartNameMapping* NameMapping)
+{
+	if(NameMapping)
+	{
+		for(DataType& Curve : Curves)
+		{
+			NameMapping->GetName(Curve.CurveUid, Curve.LastObservedName);
+		}
+	}
+}
+
+template <typename DataType>
+bool FRawCurveTracks::DuplicateCurveDataImpl(TArray<DataType> & Curves, USkeleton::AnimCurveUID ToCopyUid, USkeleton::AnimCurveUID NewUid)
+{
+	DataType* ExistingCurve = GetCurveDataImpl<DataType>(Curves, ToCopyUid);
+	if(ExistingCurve && GetCurveDataImpl<DataType>(Curves, NewUid) == NULL)
+	{
+		// Add the curve to the track and set its data to the existing curve
+		Curves.Add(DataType(NewUid, ExistingCurve->GetCurveTypeFlags()));
+		Curves.Last().CopyCurve(*ExistingCurve);
+
+		return true;
+	}
+	return false;
+}
+/////////////////////////////////////////////
+
+FORCEINLINE void ValidateCurve(FBlendedCurve* Curve)
+{
+#if 0 
+	for (auto& Element : Curve->Elements)
+	{
+		if (Element.Flags & ACF_DrivesMorphTarget)
+		{
+			ensure (Element.Value < 1.2f);
+		}
+	}
+#endif
+}
+
+void FBlendedCurve::Set(USkeleton::AnimCurveUID InUid, float InValue, int32 InFlags)
+{
+	int32 ArrayIndex;
+
+	check(bInitialized);
+
+	if (UIDList.Find(InUid, ArrayIndex))
+	{
+		Elements[ArrayIndex].Value = InValue;
+		Elements[ArrayIndex].Flags = InFlags;
+	}
+
+	ValidateCurve(this);
+}
+
+void FBlendedCurve::Reset(int32 Count)
+{
+	Elements.Empty(Count);
+	Elements.Reserve(Count);
+}
+
+//@Todo curve flags won't transfer over - it only overwritees
+void FBlendedCurve::Blend(const FBlendedCurve& A, const FBlendedCurve& B, float Alpha)
+{
+	check(A.Num()==B.Num());
+	if(FMath::Abs(Alpha) <= ZERO_ANIMWEIGHT_THRESH)
+	{
+		// if blend is all the way for child1, then just copy its bone atoms
+		Override(A);
+	}
+	else if(FMath::Abs(Alpha - 1.0f) <= ZERO_ANIMWEIGHT_THRESH)
+	{
+		// if blend is all the way for child2, then just copy its bone atoms
+		Override(B);
+	}
+	else
+	{
+		InitFrom(A);
+		for(int32 CurveId=0; CurveId<A.Elements.Num(); ++CurveId)
+		{
+			Elements[CurveId].Value = FMath::Lerp(A.Elements[CurveId].Value, B.Elements[CurveId].Value, Alpha); 
+			Elements[CurveId].Flags |= (A.Elements[CurveId].Flags | B.Elements[CurveId].Flags);
+		}
+	}
+
+	ValidateCurve(this);
+}
+
+void FBlendedCurve::ConvertToAdditive(const FBlendedCurve& BaseCurve)
+{
+	check(bInitialized);
+	check(Num()==BaseCurve.Num());
+
+	for(int32 CurveId=0; CurveId<Elements.Num(); ++CurveId)
+	{
+		Elements[CurveId].Value -= BaseCurve.Elements[CurveId].Value;
+		Elements[CurveId].Flags |= BaseCurve.Elements[CurveId].Flags;
+	}
+	ValidateCurve(this);
+}
+
+void FBlendedCurve::Accumulate(const FBlendedCurve& AdditiveCurve, float Weight)
+{
+	check(bInitialized);
+	check(Num()==AdditiveCurve.Num());
+
+	if (Weight > ZERO_ANIMWEIGHT_THRESH)
+	{
+		for(int32 CurveId=0; CurveId<Elements.Num(); ++CurveId)
+		{
+			Elements[CurveId].Value += AdditiveCurve.Elements[CurveId].Value * Weight;
+			Elements[CurveId].Flags |= AdditiveCurve.Elements[CurveId].Flags;
+		}
+	}
+
+	ValidateCurve(this);
+}
+
+void FBlendedCurve::Override(const FBlendedCurve& CurveToOverrideFrom, float Weight)
+{
+	InitFrom(CurveToOverrideFrom);
+
+	if ( FMath::IsNearlyEqual(Weight, 1.f) )
+	{
+		Override(CurveToOverrideFrom);
+	}
+	else
+	{
+		for(int32 CurveId=0; CurveId<CurveToOverrideFrom.Elements.Num(); ++CurveId)
+		{
+			Elements[CurveId].Value = CurveToOverrideFrom.Elements[CurveId].Value * Weight;
+			Elements[CurveId].Flags |= CurveToOverrideFrom.Elements[CurveId].Flags;
+		}
+	}
+
+	ValidateCurve(this);
+}
+
+void FBlendedCurve::Override(const FBlendedCurve& CurveToOverrideFrom)
+{
+	InitFrom(CurveToOverrideFrom);
+	Elements = CurveToOverrideFrom.Elements;
+
+	ValidateCurve(this);
+}
+
+void FBlendedCurve::InitFrom(const USkeleton* Skeleton)
+{
+	if(const FSmartNameMapping* Mapping = Skeleton->SmartNames.GetContainer(USkeleton::AnimCurveMappingName))
+	{
+		Mapping->FillUidArray(UIDList);
+		Elements.Empty(UIDList.Num());
+		Elements.AddZeroed(UIDList.Num());
+	}
+
+	// no name, means no curve
+	bInitialized = true;
+}
+
+FBlendedCurve::FBlendedCurve(const class UAnimInstance* AnimInstance)
+{
+	check(AnimInstance);
+	InitFrom(AnimInstance->CurrentSkeleton);
+}
+
+void FBlendedCurve::CopyFrom(const FBlendedCurve& CurveToCopyFrom)
+{
+	UIDList = CurveToCopyFrom.UIDList;
+	Elements = CurveToCopyFrom.Elements;
+	bInitialized = true;
+	ValidateCurve(this);
+}
+
+void FBlendedCurve::MoveFrom(FBlendedCurve& CurveToMoveFrom)
+{
+	UIDList = MoveTemp(CurveToMoveFrom.UIDList);
+	Elements = MoveTemp(CurveToMoveFrom.Elements);
+	bInitialized = true;
+	CurveToMoveFrom.bInitialized = false;
+
+	ValidateCurve(this);
+}
+
+void FBlendedCurve::Combine(const FBlendedCurve& CurveToCombine)
+{
+	check(bInitialized);
+	check(Num()==CurveToCombine.Num());
+
+	for(int32 CurveId=0; CurveId<CurveToCombine.Elements.Num(); ++CurveId)
+	{
+		Elements[CurveId].Value = FMath::Max(Elements[CurveId].Value, CurveToCombine.Elements[CurveId].Value);
+		Elements[CurveId].Flags |= CurveToCombine.Elements[CurveId].Flags;
+	}
+
+	ValidateCurve(this);
+}
\ No newline at end of file
diff --git a/Engine/Source/Runtime/Engine/Private/Animation/AnimInstance.cpp b/Engine/Source/Runtime/Engine/Private/Animation/AnimInstance.cpp
--- a/Engine/Source/Runtime/Engine/Private/Animation/AnimInstance.cpp
+++ b/Engine/Source/Runtime/Engine/Private/Animation/AnimInstance.cpp
@@ -7,7 +7,7 @@
 #include "EnginePrivate.h"
 #include "AnimationRuntime.h"
 #include "AnimationUtils.h"
-#include "AnimTree.h"
+#include "Animation/AnimStats.h"
 #include "Animation/AnimNode_StateMachine.h"
 #include "GameFramework/Character.h"
 #include "ParticleDefinitions.h"
@@ -39,7 +39,6 @@ DEFINE_STAT(STAT_GetAnimationPose);
 DEFINE_STAT(STAT_AnimNativeEvaluatePoses);
 DEFINE_STAT(STAT_AnimTriggerAnimNotifies);
 DEFINE_STAT(STAT_AnimNativeBlendPoses);
-DEFINE_STAT(STAT_AnimNativeCopyPoses);
 DEFINE_STAT(STAT_AnimGraphEvaluate);
 DEFINE_STAT(STAT_AnimBlendTime);
 DEFINE_STAT(STAT_RefreshBoneTransforms);
@@ -140,64 +139,6 @@ FAnimTickRecord& UAnimInstance::CreateUninitializedTickRecord(int32 GroupIndex,
 	return *TickRecord;
 }
 
-void UAnimInstance::SequenceEvaluatePose(UAnimSequenceBase* Sequence, struct FCompactPose& Pose, const FAnimExtractContext& ExtractionContext)
-{
-	SCOPE_CYCLE_COUNTER(STAT_AnimNativeEvaluatePoses);
-	checkSlow(RequiredBones.IsValid());
-
-	USkeletalMeshComponent* Component = GetSkelMeshComponent();
-
-	FAnimExtractContext NewExtractContext(ExtractionContext);
-	NewExtractContext.bExtractRootMotion = RootMotionMode == ERootMotionMode::RootMotionFromEverything || RootMotionMode == ERootMotionMode::IgnoreRootMotion;
-
-	if (const UAnimSequence* AnimSequence = Cast<const UAnimSequence>(Sequence))
-	{
-		FAnimationRuntime::GetPoseFromSequence(
-			AnimSequence,
-			/*out*/ Pose,
-			NewExtractContext);
-	}
-	else if (const UAnimComposite* Composite = Cast<const UAnimComposite>(Sequence))
-	{
-		FAnimationRuntime::GetPoseFromAnimTrack(
-			Composite->AnimationTrack,
-			/*out*/ Pose,
-			NewExtractContext);
-	}
-	else
-	{
-#if 0
-		UE_LOG(LogAnimation, Log, TEXT("FAnimationRuntime::GetPoseFromSequence - %s - No animation data!"), *GetFName());
-#endif
-		Pose.ResetToRefPose();
-	}
-}
-
-void UAnimInstance::BlendSequences(const FCompactPose& Pose1, const FCompactPose& Pose2, float Alpha, FCompactPose& Result)
-{
-	SCOPE_CYCLE_COUNTER(STAT_AnimNativeBlendPoses);
-
-	Alpha = FMath::Clamp<float>(Alpha, 0.0f, 1.0f);
-
-	if (Result.GetNumBones() < Pose1.GetNumBones())
-	{
-		ensureMsg(false, TEXT("Source Pose has more bones than Target Pose"));
-		//@hack
-		Result.SetBoneContainer(&Pose1.GetBoneContainer());
-	}
-
-	FAnimationRuntime::BlendTwoPosesTogether(Pose1, Pose2, (1.0f - Alpha), Result);
-}
-
-void UAnimInstance::CopyPose(const FCompactPose& Source, FCompactPose& Destination)
-{
-	if (&Destination != &Source)
-	{
-		SCOPE_CYCLE_COUNTER(STAT_AnimNativeCopyPoses);
-		Destination = Source;
-	}
-}
-
 AActor* UAnimInstance::GetOwningActor() const
 {
 	USkeletalMeshComponent* OwnerComponent = GetSkelMeshComponent();
@@ -392,22 +333,9 @@ void UAnimInstance::UpdateAnimation(float DeltaSeconds)
 #endif
 
 	AnimNotifies.Empty();
-	MorphTargetCurves.Empty();
 
 	ClearSlotNodeWeights();
 
-	//Track material params we set last time round so we can clear them if they aren't set again.
-	MaterialParamatersToClear.Empty();
-	for( auto Iter = MaterialParameterCurves.CreateConstIterator(); Iter; ++Iter )
-	{
-		if(Iter.Value() > 0.0f)
-		{
-			MaterialParamatersToClear.Add(Iter.Key());
-		}
-	}
-	MaterialParameterCurves.Empty();
-	VertexAnims.Empty();
-
 	// Reset the player tick list (but keep it presized)
 	UngroupedActivePlayers.Empty(UngroupedActivePlayers.Num());
 	for (int32 GroupIndex = 0; GroupIndex < SyncGroups.Num(); ++GroupIndex)
@@ -435,10 +363,6 @@ void UAnimInstance::UpdateAnimation(float DeltaSeconds)
 		RootNode->Update(UpdateContext);
 	}
 
-	// curve values can be used during update state, so we need to clear the array before ticking each elements
-	// where we collect new items
-	EventCurves.Empty();
-
 	// Handle all players inside sync groups
 	for (int32 GroupIndex = 0; GroupIndex < SyncGroups.Num(); ++GroupIndex)
 	{
@@ -507,15 +431,6 @@ void UAnimInstance::UpdateAnimation(float DeltaSeconds)
 	// Trigger Montage end events after notifies. In case Montage ending ends abilities or other states, we make sure notifies are processed before montage events.
 	TriggerQueuedMontageEvents();
 
-	// Add 0.0 curves to clear parameters that we have previously set but didn't set this tick.
-	//   - Make a copy of MaterialParametersToClear as it will be modified by AddCurveValue
-	TArray<FName> ParamsToClearCopy = MaterialParamatersToClear;
-	for (int i = 0; i < ParamsToClearCopy.Num(); ++i)
-	{
-		AddCurveValue(ParamsToClearCopy[i], 0.0f, ACF_DrivesMaterial);
-	}
-
-
 #if WITH_EDITOR && 0
 	{
 		// Take a snapshot if the scrub control is locked to the end, we are playing, and we are the one being debugged
@@ -538,6 +453,8 @@ void UAnimInstance::UpdateAnimation(float DeltaSeconds)
 
 void UAnimInstance::EvaluateAnimation(FPoseContext& Output)
 {
+	SCOPE_CYCLE_COUNTER(STAT_AnimNativeEvaluatePoses);
+
 	// If bone caches have been invalidated, have AnimNodes refresh those.
 	if( bBoneCachesInvalidated && RootNode )
 	{
@@ -1046,16 +963,6 @@ void UAnimInstance::DisplayDebug(class UCanvas* Canvas, const FDebugDisplayInfo&
 	}
 }
 
-void UAnimInstance::BlendSpaceEvaluatePose(class UBlendSpaceBase* BlendSpace, TArray<FBlendSampleData>& BlendSampleDataCache, struct FCompactPose& OutPose)
-{
-	SCOPE_CYCLE_COUNTER(STAT_AnimNativeEvaluatePoses);
-
-	FAnimationRuntime::GetPoseFromBlendSpace(
-		BlendSpace,
-		BlendSampleDataCache, 
-		/*out*/ OutPose);
-}
-
 void UAnimInstance::RecalcRequiredBones()
 {
 	USkeletalMeshComponent * SkelMeshComp = GetSkelMeshComponent();
@@ -1218,6 +1125,37 @@ void UAnimInstance::AddCurveValue(const USkeleton::AnimCurveUID Uid, float Value
 	AddCurveValue(CurrentCurveName, Value, CurveTypeFlags);
 }
 
+void UAnimInstance::UpdateCurves(const FBlendedCurve& InCurve)
+{
+	//Track material params we set last time round so we can clear them if they aren't set again.
+	MaterialParamatersToClear.Empty();
+	for(auto Iter = MaterialParameterCurves.CreateConstIterator(); Iter; ++Iter)
+	{
+		if(Iter.Value() > 0.0f)
+		{
+			MaterialParamatersToClear.Add(Iter.Key());
+		}
+	}
+
+	EventCurves.Empty();
+	MorphTargetCurves.Empty();
+	MaterialParameterCurves.Empty();
+
+	for(int32 CurveId=0; CurveId<InCurve.UIDList.Num(); ++CurveId)
+	{
+		// had to add to anotehr data type
+		AddCurveValue(InCurve.UIDList[CurveId], InCurve.Elements[CurveId].Value, InCurve.Elements[CurveId].Flags);
+	}
+
+	// Add 0.0 curves to clear parameters that we have previously set but didn't set this tick.
+	//   - Make a copy of MaterialParametersToClear as it will be modified by AddCurveValue
+	TArray<FName> ParamsToClearCopy = MaterialParamatersToClear;
+	for(int i = 0; i < ParamsToClearCopy.Num(); ++i)
+	{
+		AddCurveValue(ParamsToClearCopy[i], 0.0f, ACF_DrivesMaterial);
+	}
+}
+
 void UAnimInstance::TriggerAnimNotifies(float DeltaSeconds)
 {
 	SCOPE_CYCLE_COUNTER(STAT_AnimTriggerAnimNotifies);
@@ -1386,16 +1324,15 @@ void UAnimInstance::GetSlotWeight(FName const & SlotNodeName, float& out_SlotNod
 	out_SourceWeight = 1.f - NonAdditiveTotalWeight;
 }
 
-void UAnimInstance::SlotEvaluatePose(FName SlotNodeName, const FCompactPose& SourcePose, FCompactPose& BlendedPose, float SlotNodeWeight)
+void UAnimInstance::SlotEvaluatePose(FName SlotNodeName, const FCompactPose& SourcePose, const FBlendedCurve& SourceCurve, FCompactPose& BlendedPose, FBlendedCurve& BlendedCurve, float SlotNodeWeight)
 {
 	//Accessing MontageInstances from this function is not safe (as this can be called during Parallel Anim Evaluation!
 	//Any montage data you need to add should be part of MontageEvaluationData
 
-	//MDW TODO
-	SCOPE_CYCLE_COUNTER(STAT_AnimNativeEvaluatePoses);
 	if (SlotNodeWeight <= ZERO_ANIMWEIGHT_THRESH)
 	{
 		BlendedPose = SourcePose;
+		BlendedCurve = SourceCurve;
 		return;
 	}
 
@@ -1421,10 +1358,17 @@ void UAnimInstance::SlotEvaluatePose(FName SlotNodeName, const FCompactPose& Sou
 			
 			// Bone array has to be allocated prior to calling GetPoseFromAnimTrack
 			NewPose.Pose.SetBoneContainer(&RequiredBones);
+			NewPose.Curve.InitFrom(this);
 
 			// Extract pose from Track
 			FAnimExtractContext ExtractionContext(EvalState.MontagePosition, EvalState.Montage->HasRootMotion() && RootMotionMode != ERootMotionMode::NoRootMotionExtraction);
-			FAnimationRuntime::GetPoseFromAnimTrack(*AnimTrack, NewPose.Pose, ExtractionContext);
+			AnimTrack->GetAnimationPose(NewPose.Pose, NewPose.Curve, ExtractionContext);
+
+			// add montage curves 
+			FBlendedCurve MontageCurve;
+			MontageCurve.InitFrom(this);
+			EvalState.Montage->EvaluateCurveData(MontageCurve, EvalState.MontagePosition);
+			NewPose.Curve.Combine(MontageCurve);
 
 			TotalWeight += EvalState.MontageWeight;
 			if (AdditiveAnimType == AAT_None)
@@ -1436,6 +1380,7 @@ void UAnimInstance::SlotEvaluatePose(FName SlotNodeName, const FCompactPose& Sou
 			{
 				AdditivePoses.Add(NewPose);
 			}
+		
 		}
 	}
 
@@ -1443,6 +1388,7 @@ void UAnimInstance::SlotEvaluatePose(FName SlotNodeName, const FCompactPose& Sou
 	if (TotalWeight <= ZERO_ANIMWEIGHT_THRESH)
 	{
 		BlendedPose = SourcePose;
+		BlendedCurve = SourceCurve;
 		return;
 	}
 	// Make sure weights don't exceed 1.f, otherwise re-normalize.
@@ -1472,6 +1418,7 @@ void UAnimInstance::SlotEvaluatePose(FName SlotNodeName, const FCompactPose& Sou
 		if (NonAdditivePoses.Num() == 0)
 		{
 			BlendedPose = SourcePose;
+			BlendedCurve = SourceCurve;
 		}
 		// Otherwise we need to blend non additive poses together
 		else
@@ -1486,9 +1433,14 @@ void UAnimInstance::SlotEvaluatePose(FName SlotNodeName, const FCompactPose& Sou
 
 			TArray<float> BlendWeights;
 			BlendWeights.AddUninitialized(NumPoses);
+
+			TArray<const FBlendedCurve*> BlendingCurves;
+			BlendingCurves.AddUninitialized(NumPoses);
+
 			for (int32 Index = 0; Index < NonAdditivePoses.Num(); Index++)
 			{
 				BlendingPoses[Index] = &NonAdditivePoses[Index].Pose;
+				BlendingCurves[Index] = &NonAdditivePoses[Index].Curve;
 				BlendWeights[Index] = NonAdditivePoses[Index].Weight;
 			}
 
@@ -1496,11 +1448,12 @@ void UAnimInstance::SlotEvaluatePose(FName SlotNodeName, const FCompactPose& Sou
 			{
 				int32 const SourceIndex = BlendWeights.Num() - 1;
 				BlendingPoses[SourceIndex] = &SourcePose;
+				BlendingCurves[SourceIndex] = &SourceCurve;
 				BlendWeights[SourceIndex] = SourceWeight;
 			}
 
 			// now time to blend all montages
-			FAnimationRuntime::BlendPosesTogether(BlendingPoses, BlendWeights, BlendedPose);
+			FAnimationRuntime::BlendPosesTogether(BlendingPoses, BlendingCurves, BlendWeights, BlendedPose, BlendedCurve);
 		}
 	}
 
@@ -1510,19 +1463,7 @@ void UAnimInstance::SlotEvaluatePose(FName SlotNodeName, const FCompactPose& Sou
 		for (int32 Index = 0; Index < AdditivePoses.Num(); Index++)
 		{
 			FSlotEvaluationPose const & AdditivePose = AdditivePoses[Index];
-			// if additive, we should blend with source to make it full body
-			if (AdditivePose.AdditiveType == AAT_LocalSpaceBase)
-			{
-				FAnimationRuntime::AccumulateAdditivePose(BlendedPose, AdditivePose.Pose, AdditivePose.Weight);
-			}
-			else if (AdditivePose.AdditiveType == AAT_RotationOffsetMeshSpace)
-			{
-				FAnimationRuntime::AccumulateMeshSpaceRotationAdditiveToLocalPose(BlendedPose, AdditivePose.Pose, AdditivePose.Weight);
-			}
-			else
-			{
-				check(false);
-			}
+			FAnimationRuntime::AccumulateAdditivePose(BlendedPose, AdditivePose.Pose, BlendedCurve, AdditivePose.Curve, AdditivePose.Weight, AdditivePose.AdditiveType);
 		}
 	}
 
diff --git a/Engine/Source/Runtime/Engine/Private/Animation/AnimMontage.cpp b/Engine/Source/Runtime/Engine/Private/Animation/AnimMontage.cpp
--- a/Engine/Source/Runtime/Engine/Private/Animation/AnimMontage.cpp
+++ b/Engine/Source/Runtime/Engine/Private/Animation/AnimMontage.cpp
@@ -806,38 +806,6 @@ bool UAnimMontage::HasValidSlotSetup() const
 	return true;
 }
 
-void UAnimMontage::EvaluateCurveData(class UAnimInstance* Instance, float CurrentTime, float BlendWeight ) const
-{
-	Super::EvaluateCurveData(Instance, CurrentTime, BlendWeight);
-
-	// I also need to evaluate curve of the animation
-	// for now we only get the first slot
-	// in the future, we'll need to do based on highest weight?
-	// first get all the montage instance weight this slot node has
-	if ( SlotAnimTracks.Num() > 0 )
-	{
-		const FAnimTrack& Track = SlotAnimTracks[0].AnimTrack;
-		for (int32 I=0; I<Track.AnimSegments.Num(); ++I)
-		{
-			const FAnimSegment& AnimSegment = Track.AnimSegments[I];
-
-			float PositionInAnim = 0.f;
-			float Weight = 0.f;
-			UAnimSequenceBase* AnimRef = AnimSegment.GetAnimationData(CurrentTime, PositionInAnim, Weight);
-			// make this to be 1 function
-			if ( AnimRef && Weight > ZERO_ANIMWEIGHT_THRESH )
-			{
-				// todo anim: hack - until we fix animcomposite
-				UAnimSequence * Sequence = Cast<UAnimSequence>(AnimRef);
-				if ( Sequence )
-				{
-					Sequence->EvaluateCurveData(Instance, PositionInAnim, BlendWeight*Weight);
-				}
-			}
-		}	
-	}
-}
-
 //////////////////////////////////////////////////////////////////////////////////////////////
 // MontageInstance
 /////////////////////////////////////////////////////////////////////////////////////////////
@@ -1461,12 +1429,6 @@ void FAnimMontageInstance::Advance(float DeltaTime, struct FRootMotionMovementPa
 		return;
 	}
 
-	// Update curves based on final position.
-	if( (Weight > ZERO_ANIMWEIGHT_THRESH) && IsValid() )
-	{
-		Montage->EvaluateCurveData(AnimInstance.Get(), Position, Weight);
-	}
-
 	if (!bInterrupted)
 	{
 		// Tick all active state branching points
diff --git a/Engine/Source/Runtime/Engine/Private/Animation/AnimNodeSpaceConversions.cpp b/Engine/Source/Runtime/Engine/Private/Animation/AnimNodeSpaceConversions.cpp
--- a/Engine/Source/Runtime/Engine/Private/Animation/AnimNodeSpaceConversions.cpp
+++ b/Engine/Source/Runtime/Engine/Private/Animation/AnimNodeSpaceConversions.cpp
@@ -33,6 +33,7 @@ void FAnimNode_ConvertComponentToLocalSpace::Evaluate(FPoseContext & Output)
 
 	checkSlow( InputCSPose.Pose.GetPose().IsValid() );
 	InputCSPose.Pose.ConvertToLocalPoses(Output.Pose);
+	Output.Curve = InputCSPose.Curve;
 }
 
 void FAnimNode_ConvertComponentToLocalSpace::GatherDebugData(FNodeDebugData& DebugData)
@@ -78,4 +79,5 @@ void FAnimNode_ConvertLocalToComponentSpace::EvaluateComponentSpace(FComponentSp
 	LocalPose.Evaluate(InputPose);
 
 	OutputCSPose.Pose.InitPose(InputPose.Pose);
+	OutputCSPose.Curve = InputPose.Curve;
 }
diff --git a/Engine/Source/Runtime/Engine/Private/Animation/AnimNode_ApplyAdditive.cpp b/Engine/Source/Runtime/Engine/Private/Animation/AnimNode_ApplyAdditive.cpp
--- a/Engine/Source/Runtime/Engine/Private/Animation/AnimNode_ApplyAdditive.cpp
+++ b/Engine/Source/Runtime/Engine/Private/Animation/AnimNode_ApplyAdditive.cpp
@@ -41,7 +41,7 @@ void FAnimNode_ApplyAdditive::Evaluate(FPoseContext& Output)
 		Base.Evaluate(Output);
 		Additive.Evaluate(AdditiveEvalContext);
 
-		FAnimationRuntime::AccumulateAdditivePose(Output.Pose, AdditiveEvalContext.Pose, ActualAlpha);
+		FAnimationRuntime::AccumulateAdditivePose(Output.Pose, AdditiveEvalContext.Pose, Output.Curve, AdditiveEvalContext.Curve, ActualAlpha, AAT_LocalSpaceBase);
 		Output.Pose.NormalizeRotations();
 	}
 	else
diff --git a/Engine/Source/Runtime/Engine/Private/Animation/AnimNode_BlendListBase.cpp b/Engine/Source/Runtime/Engine/Private/Animation/AnimNode_BlendListBase.cpp
--- a/Engine/Source/Runtime/Engine/Private/Animation/AnimNode_BlendListBase.cpp
+++ b/Engine/Source/Runtime/Engine/Private/Animation/AnimNode_BlendListBase.cpp
@@ -3,7 +3,7 @@
 #include "EnginePrivate.h"
 #include "Animation/AnimNode_BlendListBase.h"
 #include "AnimationRuntime.h"
-#include "AnimTree.h"
+#include "Animation/AnimStats.h"
 
 /////////////////////////////////////////////////////
 // FAnimNode_BlendListBase
@@ -122,6 +122,9 @@ void FAnimNode_BlendListBase::Evaluate(FPoseContext& Output)
 		TArray<FCompactPose> FilteredPoses;
 		FilteredPoses.SetNum(NumPoses);
 
+		TArray<FBlendedCurve> FilteredCurve;
+		FilteredCurve.SetNum(NumPoses);
+
 		TArray<float> FilteredWeights;
 		FilteredWeights.AddUninitialized(NumPoses);
 
@@ -137,10 +140,11 @@ void FAnimNode_BlendListBase::Evaluate(FPoseContext& Output)
 			CurrentPose.Evaluate(EvaluateContext);
 
 			FilteredPoses[i].MoveBonesFrom(EvaluateContext.Pose);
+			FilteredCurve[i] = EvaluateContext.Curve;
 			FilteredWeights[i] = BlendWeight;
 		}
 
-		FAnimationRuntime::BlendPosesTogether(FilteredPoses, FilteredWeights, Output.Pose);
+		FAnimationRuntime::BlendPosesTogether(FilteredPoses, FilteredCurve, FilteredWeights, Output.Pose, Output.Curve);
 	}
 	else
 	{
diff --git a/Engine/Source/Runtime/Engine/Private/Animation/AnimNode_BlendSpacePlayer.cpp b/Engine/Source/Runtime/Engine/Private/Animation/AnimNode_BlendSpacePlayer.cpp
--- a/Engine/Source/Runtime/Engine/Private/Animation/AnimNode_BlendSpacePlayer.cpp
+++ b/Engine/Source/Runtime/Engine/Private/Animation/AnimNode_BlendSpacePlayer.cpp
@@ -74,7 +74,7 @@ void FAnimNode_BlendSpacePlayer::Evaluate(FPoseContext& Output)
 {
 	if ((BlendSpace != NULL) && (Output.AnimInstance->CurrentSkeleton->IsCompatible(BlendSpace->GetSkeleton())))
 	{
-		Output.AnimInstance->BlendSpaceEvaluatePose(BlendSpace, BlendSampleDataCache, Output.Pose);
+		BlendSpace->GetAnimationPose(BlendSampleDataCache, Output.Pose, Output.Curve);
 	}
 	else
 	{
diff --git a/Engine/Source/Runtime/Engine/Private/Animation/AnimNode_LayeredBoneBlend.cpp b/Engine/Source/Runtime/Engine/Private/Animation/AnimNode_LayeredBoneBlend.cpp
--- a/Engine/Source/Runtime/Engine/Private/Animation/AnimNode_LayeredBoneBlend.cpp
+++ b/Engine/Source/Runtime/Engine/Private/Animation/AnimNode_LayeredBoneBlend.cpp
@@ -3,7 +3,7 @@
 #include "EnginePrivate.h"
 #include "Animation/AnimNode_LayeredBoneBlend.h"
 #include "AnimationRuntime.h"
-#include "AnimTree.h"
+#include "Animation/AnimStats.h"
 
 /////////////////////////////////////////////////////
 // FAnimNode_LayeredBoneBlend
@@ -89,6 +89,9 @@ void FAnimNode_LayeredBoneBlend::Evaluate(FPoseContext& Output)
 		TArray<FCompactPose> TargetBlendPoses;
 		TargetBlendPoses.SetNum(NumPoses);
 
+		TArray<FBlendedCurve> TargetBlendCurves;
+		TargetBlendCurves.SetNum(NumPoses);
+
 		for (int32 ChildIndex = 0; ChildIndex < NumPoses; ++ChildIndex)
 		{
 			if (BlendWeights[ChildIndex] > ZERO_ANIMWEIGHT_THRESH)
@@ -97,14 +100,16 @@ void FAnimNode_LayeredBoneBlend::Evaluate(FPoseContext& Output)
 				BlendPoses[ChildIndex].Evaluate(CurrentPoseContext);
 
 				TargetBlendPoses[ChildIndex].MoveBonesFrom(CurrentPoseContext.Pose);
+				TargetBlendCurves[ChildIndex].MoveFrom(CurrentPoseContext.Curve);
 			}
 			else
 			{
 				TargetBlendPoses[ChildIndex].ResetToRefPose(BasePoseContext.Pose.GetBoneContainer());
+				TargetBlendCurves[ChildIndex].InitFrom(Output.Curve);
 			}
 		}
 
-		FAnimationRuntime::BlendPosesPerBoneFilter(BasePoseContext.Pose, TargetBlendPoses, Output.Pose, CurrentBoneBlendWeights, bMeshSpaceRotationBlend);
+		FAnimationRuntime::BlendPosesPerBoneFilter(BasePoseContext.Pose, TargetBlendPoses, BasePoseContext.Curve, TargetBlendCurves, Output.Pose, Output.Curve, CurrentBoneBlendWeights, bMeshSpaceRotationBlend, CurveBlendOption);
 	}
 }
 
diff --git a/Engine/Source/Runtime/Engine/Private/Animation/AnimNode_RotationOffsetBlendSpace.cpp b/Engine/Source/Runtime/Engine/Private/Animation/AnimNode_RotationOffsetBlendSpace.cpp
--- a/Engine/Source/Runtime/Engine/Private/Animation/AnimNode_RotationOffsetBlendSpace.cpp
+++ b/Engine/Source/Runtime/Engine/Private/Animation/AnimNode_RotationOffsetBlendSpace.cpp
@@ -34,7 +34,7 @@ void FAnimNode_RotationOffsetBlendSpace::Evaluate(FPoseContext& Context)
 	FAnimNode_BlendSpacePlayer::Evaluate(MeshSpaceRotationAdditivePoseContext);
 
 	// Accumulate poses together
-	FAnimationRuntime::AccumulateMeshSpaceRotationAdditiveToLocalPose(Context.Pose, MeshSpaceRotationAdditivePoseContext.Pose, 1.f);
+	FAnimationRuntime::AccumulateMeshSpaceRotationAdditiveToLocalPose(Context.Pose, MeshSpaceRotationAdditivePoseContext.Pose, Context.Curve, MeshSpaceRotationAdditivePoseContext.Curve, 1.f);
 
 	// Resulting rotations are not normalized, so normalize here.
 	Context.Pose.NormalizeRotations();
diff --git a/Engine/Source/Runtime/Engine/Private/Animation/AnimNode_SaveCachedPose.cpp b/Engine/Source/Runtime/Engine/Private/Animation/AnimNode_SaveCachedPose.cpp
--- a/Engine/Source/Runtime/Engine/Private/Animation/AnimNode_SaveCachedPose.cpp
+++ b/Engine/Source/Runtime/Engine/Private/Animation/AnimNode_SaveCachedPose.cpp
@@ -56,10 +56,12 @@ void FAnimNode_SaveCachedPose::Evaluate(FPoseContext& Output)
 		FPoseContext CachingContext(Output);
 		Pose.Evaluate(CachingContext);
 		CachedPose.MoveBonesFrom(CachingContext.Pose);
+		CachedCurve.MoveFrom(CachingContext.Curve);
 	}
 
 	// Return the cached result
-	Output.AnimInstance->CopyPose(CachedPose, Output.Pose);
+	Output.Pose.CopyBonesFrom(CachedPose);
+	Output.Curve.CopyFrom(CachedCurve);
 }
 
 
diff --git a/Engine/Source/Runtime/Engine/Private/Animation/AnimNode_SequenceEvaluator.cpp b/Engine/Source/Runtime/Engine/Private/Animation/AnimNode_SequenceEvaluator.cpp
--- a/Engine/Source/Runtime/Engine/Private/Animation/AnimNode_SequenceEvaluator.cpp
+++ b/Engine/Source/Runtime/Engine/Private/Animation/AnimNode_SequenceEvaluator.cpp
@@ -23,7 +23,7 @@ void FAnimNode_SequenceEvaluator::Evaluate(FPoseContext& Output)
 {
 	if ((Sequence != NULL) && (Output.AnimInstance->CurrentSkeleton->IsCompatible(Sequence->GetSkeleton())))
 	{
-		Output.AnimInstance->SequenceEvaluatePose(Sequence, Output.Pose, FAnimExtractContext(ExplicitTime));
+		Sequence->GetAnimationPose(Output.Pose, Output.Curve, FAnimExtractContext(ExplicitTime, Output.AnimInstance->ShouldExtractRootMotion()));
 	}
 	else
 	{
diff --git a/Engine/Source/Runtime/Engine/Private/Animation/AnimNode_SequencePlayer.cpp b/Engine/Source/Runtime/Engine/Private/Animation/AnimNode_SequencePlayer.cpp
--- a/Engine/Source/Runtime/Engine/Private/Animation/AnimNode_SequencePlayer.cpp
+++ b/Engine/Source/Runtime/Engine/Private/Animation/AnimNode_SequencePlayer.cpp
@@ -50,7 +50,7 @@ void FAnimNode_SequencePlayer::Evaluate(FPoseContext& Output)
 {
 	if ((Sequence != NULL) && (Output.AnimInstance->CurrentSkeleton->IsCompatible(Sequence->GetSkeleton())))
 	{
-		Output.AnimInstance->SequenceEvaluatePose(Sequence, Output.Pose, FAnimExtractContext(InternalTimeAccumulator));
+		Sequence->GetAnimationPose(Output.Pose, Output.Curve, FAnimExtractContext(InternalTimeAccumulator, Output.AnimInstance->ShouldExtractRootMotion()));
 	}
 	else
 	{
diff --git a/Engine/Source/Runtime/Engine/Private/Animation/AnimNode_Slot.cpp b/Engine/Source/Runtime/Engine/Private/Animation/AnimNode_Slot.cpp
--- a/Engine/Source/Runtime/Engine/Private/Animation/AnimNode_Slot.cpp
+++ b/Engine/Source/Runtime/Engine/Private/Animation/AnimNode_Slot.cpp
@@ -58,7 +58,7 @@ void FAnimNode_Slot::Evaluate(FPoseContext & Output)
 			Source.Evaluate(SourceContext);
 		}
 
-		Output.AnimInstance->SlotEvaluatePose(SlotName, SourceContext.Pose, Output.Pose, SlotNodeWeight);
+		Output.AnimInstance->SlotEvaluatePose(SlotName, SourceContext.Pose, SourceContext.Curve, Output.Pose, Output.Curve, SlotNodeWeight);
 
 		checkSlow(!Output.ContainsNaN());
 		checkSlow(Output.IsNormalized());
diff --git a/Engine/Source/Runtime/Engine/Private/Animation/AnimNode_StateMachine.cpp b/Engine/Source/Runtime/Engine/Private/Animation/AnimNode_StateMachine.cpp
--- a/Engine/Source/Runtime/Engine/Private/Animation/AnimNode_StateMachine.cpp
+++ b/Engine/Source/Runtime/Engine/Private/Animation/AnimNode_StateMachine.cpp
@@ -6,7 +6,7 @@
 #include "Animation/AnimNode_TransitionPoseEvaluator.h"
 #include "Animation/AnimNode_SequencePlayer.h"
 #include "AnimationRuntime.h"
-#include "AnimTree.h"
+#include "Animation/AnimStats.h"
 
 //////////////////////////////////////////////////////////////////////////
 // FAnimationActiveTransitionEntry
@@ -626,7 +626,7 @@ void FAnimNode_StateMachine::EvaluateTransitionStandardBlend(FPoseContext& Outpu
 
 	if (bIntermediatePoseIsValid)
 	{
-		Output.AnimInstance->CopyPose(Output.Pose, PreviouseStateResult.Pose);
+		PreviouseStateResult = Output;
 	}
 	else
 	{
@@ -644,6 +644,10 @@ void FAnimNode_StateMachine::EvaluateTransitionStandardBlend(FPoseContext& Outpu
 		Output.Pose[BoneIndex] = PreviouseStateResult.Pose[BoneIndex] * VPreviousWeight;
 		Output.Pose[BoneIndex].AccumulateWithShortestRotation(NextStateResult.Pose[BoneIndex], VWeight);
 	}
+
+	// blend curve in
+	Output.Curve.Override(PreviouseStateResult.Curve, 1.0-Transition.Alpha);
+	Output.Curve.Accumulate(NextStateResult.Curve, Transition.Alpha);
 }
 
 void FAnimNode_StateMachine::EvaluateTransitionCustomBlend(FPoseContext& Output, FAnimationActiveTransitionEntry& Transition, bool bIntermediatePoseIsValid)
@@ -668,12 +672,12 @@ void FAnimNode_StateMachine::EvaluateTransitionCustomBlend(FPoseContext& Output,
 					EvaluateState(EffectiveStateIndex, PoseEvalResult);
 
 					// push transform to node.
-					Evaluator->CachePose(Output, PoseEvalResult.Pose);
+					Evaluator->CachePose(PoseEvalResult);
 				}
 				else
 				{
 					// push transform to node.
-					Evaluator->CachePose(Output, Output.Pose);
+					Evaluator->CachePose(Output);
 				}
 			}
 		}
@@ -870,7 +874,7 @@ void FAnimNode_StateMachine::UpdateState(int32 StateIndex, const FAnimationUpdat
 
 void FAnimNode_StateMachine::EvaluateState(int32 StateIndex, FPoseContext& Output)
 {
-	Output.Pose.SetBoneContainer(&Output.AnimInstance->RequiredBones);
+	Output.Initialize(Output.AnimInstance);
 
 	if (!IsAConduitState(StateIndex))
 	{
diff --git a/Engine/Source/Runtime/Engine/Private/Animation/AnimNode_TransitionPoseEvaluator.cpp b/Engine/Source/Runtime/Engine/Private/Animation/AnimNode_TransitionPoseEvaluator.cpp
--- a/Engine/Source/Runtime/Engine/Private/Animation/AnimNode_TransitionPoseEvaluator.cpp
+++ b/Engine/Source/Runtime/Engine/Private/Animation/AnimNode_TransitionPoseEvaluator.cpp
@@ -31,6 +31,7 @@ void FAnimNode_TransitionPoseEvaluator::Initialize(const FAnimationInitializeCon
 void FAnimNode_TransitionPoseEvaluator::CacheBones(const FAnimationCacheBonesContext& Context) 
 {
 	CachedPose.SetBoneContainer(&Context.AnimInstance->RequiredBones);
+	CachedCurve.InitFrom(Context.AnimInstance);
 }
 
 void FAnimNode_TransitionPoseEvaluator::Update(const FAnimationUpdateContext& Context)
@@ -42,7 +43,8 @@ void FAnimNode_TransitionPoseEvaluator::Evaluate(FPoseContext& Output)
 {	
 	// the cached pose is evaluated in the state machine and set via CachePose(). 
 	// This is because we need information about the transition that is not available at this level
-	Output.AnimInstance->CopyPose(CachedPose, Output.Pose);
+	Output.Pose.CopyBonesFrom(CachedPose);
+	Output.Curve.CopyFrom(CachedCurve);
 
 	if ((EvaluatorMode != EEvaluatorMode::EM_Standard) && (CacheFramesRemaining > 0))
 	{
@@ -69,7 +71,8 @@ bool FAnimNode_TransitionPoseEvaluator::InputNodeNeedsEvaluate() const
 	return (EvaluatorMode == EEvaluatorMode::EM_Standard) || (CacheFramesRemaining > 0);
 }
 
-void FAnimNode_TransitionPoseEvaluator::CachePose(FPoseContext& Output, FCompactPose& PoseToCache)
+void FAnimNode_TransitionPoseEvaluator::CachePose(FPoseContext& PoseToCache)
 {
-	Output.AnimInstance->CopyPose(PoseToCache, CachedPose);
+	CachedPose.CopyBonesFrom(PoseToCache.Pose);
+	CachedCurve.CopyFrom(PoseToCache.Curve);
 }
\ No newline at end of file
diff --git a/Engine/Source/Runtime/Engine/Private/Animation/AnimNode_TwoWayBlend.cpp b/Engine/Source/Runtime/Engine/Private/Animation/AnimNode_TwoWayBlend.cpp
--- a/Engine/Source/Runtime/Engine/Private/Animation/AnimNode_TwoWayBlend.cpp
+++ b/Engine/Source/Runtime/Engine/Private/Animation/AnimNode_TwoWayBlend.cpp
@@ -57,7 +57,7 @@ void FAnimationNode_TwoWayBlend::Evaluate(FPoseContext& Output)
 			A.Evaluate(Pose1);
 			B.Evaluate(Pose2);
 
-			Output.AnimInstance->BlendSequences(Pose1.Pose, Pose2.Pose, ActualAlpha, Output.Pose);
+			FAnimationRuntime::BlendTwoPosesTogether(Pose1.Pose, Pose2.Pose, Pose1.Curve, Pose2.Curve, (1.0f - ActualAlpha), Output.Pose, Output.Curve);
 		}
 		else
 		{
diff --git a/Engine/Source/Runtime/Engine/Private/Animation/AnimSequence.cpp b/Engine/Source/Runtime/Engine/Private/Animation/AnimSequence.cpp
--- a/Engine/Source/Runtime/Engine/Private/Animation/AnimSequence.cpp
+++ b/Engine/Source/Runtime/Engine/Private/Animation/AnimSequence.cpp
@@ -851,24 +851,31 @@ FTransform UAnimSequence::ExtractRootMotionFromRange(float StartTrackPosition, f
 	return EndTransform.GetRelativeTransform(StartTransform);
 }
 
-void UAnimSequence::GetAnimationPose(FCompactPose& OutPose, const FAnimExtractContext& ExtractionContext) const
+void UAnimSequence::GetAnimationPose(FCompactPose& OutPose, FBlendedCurve& OutCurve, const FAnimExtractContext& ExtractionContext) const
 {
+	SCOPE_CYCLE_COUNTER(STAT_GetAnimationPose);
+
 	// @todo anim: if compressed and baked in the future, we don't have to do this 
 	if (IsValidAdditive())
 	{
 		if (AdditiveAnimType == AAT_LocalSpaceBase)
 		{
-			GetBonePose_Additive(OutPose, ExtractionContext);
+			GetBonePose_Additive(OutPose, OutCurve, ExtractionContext);
 		}
 		else if (AdditiveAnimType == AAT_RotationOffsetMeshSpace)
 		{
-			GetBonePose_AdditiveMeshRotationOnly(OutPose, ExtractionContext);
+			GetBonePose_AdditiveMeshRotationOnly(OutPose, OutCurve, ExtractionContext);
 		}
 	}
 	else
 	{
-		GetBonePose(OutPose, ExtractionContext);
+		GetBonePose(OutPose, OutCurve, ExtractionContext);
 	}
+
+	// Check that all bone atoms coming from animation are normalized
+#if DO_CHECK && WITH_EDITORONLY_DATA
+	check(OutPose.IsNormalized());
+#endif
 }
 
 void UAnimSequence::ResetRootBoneForRootMotion(FTransform& BoneTransform, const FBoneContainer& RequiredBones, ERootMotionRootLock::Type InRootMotionRootLock) const
@@ -882,7 +889,7 @@ void UAnimSequence::ResetRootBoneForRootMotion(FTransform& BoneTransform, const
 	}
 }
 
-void UAnimSequence::GetBonePose(FCompactPose& OutPose, const FAnimExtractContext& ExtractionContext) const
+void UAnimSequence::GetBonePose(FCompactPose& OutPose, FBlendedCurve& OutCurve, const FAnimExtractContext& ExtractionContext) const
 {
 	QUICK_SCOPE_CYCLE_COUNTER(STAT_GETBONEPOSE);
 
@@ -918,6 +925,9 @@ void UAnimSequence::GetBonePose(FCompactPose& OutPose, const FAnimExtractContext
 		OutPose.ResetToRefPose();
 	}
 
+	// extract curve data . Even if no track, it can contain curve data
+	EvaluateCurveData(OutCurve, ExtractionContext.CurrentTime);
+
 	int32 const NumTracks = GetNumberOfTracks();
 	if (NumTracks == 0)
 	{
@@ -1107,7 +1117,7 @@ void UAnimSequence::GetBonePose(FCompactPose& OutPose, const FAnimExtractContext
 	}
 }
 
-void UAnimSequence::GetBonePose_Additive(FCompactPose& OutPose, const FAnimExtractContext& ExtractionContext) const
+void UAnimSequence::GetBonePose_Additive(FCompactPose& OutPose, FBlendedCurve& OutCurve, const FAnimExtractContext& ExtractionContext) const
 {
 	if (!IsValidAdditive())
 	{
@@ -1116,19 +1126,23 @@ void UAnimSequence::GetBonePose_Additive(FCompactPose& OutPose, const FAnimExtra
 	}
 
 	// Extract target pose
-	GetBonePose(OutPose, ExtractionContext);
+	GetBonePose(OutPose, OutCurve, ExtractionContext);
 
 	// Extract base pose
 	FCompactPose BasePose;
+	FBlendedCurve BaseCurve;
+	
 	BasePose.SetBoneContainer(&OutPose.GetBoneContainer());
+	BaseCurve.InitFrom(OutCurve);	
 
-	GetAdditiveBasePose(BasePose, ExtractionContext);
+	GetAdditiveBasePose(BasePose, BaseCurve, ExtractionContext);
 
 	// Create Additive animation
 	FAnimationRuntime::ConvertPoseToAdditive(OutPose, BasePose);
+	OutCurve.ConvertToAdditive(BaseCurve);
 }
 
-void UAnimSequence::GetAdditiveBasePose(FCompactPose& OutPose, const FAnimExtractContext& ExtractionContext) const
+void UAnimSequence::GetAdditiveBasePose(FCompactPose& OutPose, FBlendedCurve& OutCurve, const FAnimExtractContext& ExtractionContext) const
 {
 	switch (RefPoseType)
 	{
@@ -1141,7 +1155,7 @@ void UAnimSequence::GetAdditiveBasePose(FCompactPose& OutPose, const FAnimExtrac
 
 		FAnimExtractContext BasePoseExtractionContext(ExtractionContext);
 		BasePoseExtractionContext.CurrentTime = BasePoseTime;
-		RefPoseSeq->GetBonePose(OutPose, BasePoseExtractionContext);
+		RefPoseSeq->GetBonePose(OutPose, OutCurve, BasePoseExtractionContext);
 		break;
 	}
 		// use animation as a base pose. Need BasePoseSeq and RefFrameIndex (will clamp if outside).
@@ -1152,7 +1166,7 @@ void UAnimSequence::GetAdditiveBasePose(FCompactPose& OutPose, const FAnimExtrac
 
 		FAnimExtractContext BasePoseExtractionContext(ExtractionContext);
 		BasePoseExtractionContext.CurrentTime = BasePoseTime;
-		RefPoseSeq->GetBonePose(OutPose, BasePoseExtractionContext);
+		RefPoseSeq->GetBonePose(OutPose, OutCurve, BasePoseExtractionContext);
 		break;
 	}
 		// use ref pose of Skeleton as base
@@ -1163,7 +1177,7 @@ void UAnimSequence::GetAdditiveBasePose(FCompactPose& OutPose, const FAnimExtrac
 	}
 }
 
-void UAnimSequence::GetBonePose_AdditiveMeshRotationOnly(FCompactPose& OutPose, const FAnimExtractContext& ExtractionContext) const
+void UAnimSequence::GetBonePose_AdditiveMeshRotationOnly(FCompactPose& OutPose, FBlendedCurve& OutCurve, const FAnimExtractContext& ExtractionContext) const
 {
 	if (!IsValidAdditive())
 	{
@@ -1173,12 +1187,14 @@ void UAnimSequence::GetBonePose_AdditiveMeshRotationOnly(FCompactPose& OutPose,
 	}
 
 	// Get target pose
-	GetBonePose(OutPose, ExtractionContext);
+	GetBonePose(OutPose, OutCurve, ExtractionContext);
 
 	// get base pose
 	FCompactPose BasePose;
+	FBlendedCurve BaseCurve;
 	BasePose.SetBoneContainer(&OutPose.GetBoneContainer());
-	GetAdditiveBasePose(BasePose, ExtractionContext);
+	BaseCurve.InitFrom(OutCurve);
+	GetAdditiveBasePose(BasePose, BaseCurve, ExtractionContext);
 
 	// Convert them to mesh rotation.
 	FAnimationRuntime::ConvertPoseToMeshRotation(OutPose);
@@ -1186,6 +1202,7 @@ void UAnimSequence::GetBonePose_AdditiveMeshRotationOnly(FCompactPose& OutPose,
 
 	// Turn into Additive
 	FAnimationRuntime::ConvertPoseToAdditive(OutPose, BasePose);
+	OutCurve.ConvertToAdditive(BaseCurve);
 }
 
 void UAnimSequence::RetargetBoneTransform(FTransform& BoneTransform, const int32& SkeletonBoneIndex, const FCompactPoseBoneIndex& BoneIndex, const FBoneContainer& RequiredBones) const
diff --git a/Engine/Source/Runtime/Engine/Private/Animation/AnimSequenceBase.cpp b/Engine/Source/Runtime/Engine/Private/Animation/AnimSequenceBase.cpp
--- a/Engine/Source/Runtime/Engine/Private/Animation/AnimSequenceBase.cpp
+++ b/Engine/Source/Runtime/Engine/Private/Animation/AnimSequenceBase.cpp
@@ -83,405 +83,6 @@ void FAnimNotifyEvent::SetTime(float NewTime, EAnimLinkMethod::Type ReferenceFra
 	SetDuration(Duration);
 }
 
-/////////////////////////////////////////////////////
-// FFloatCurve
-
-void FAnimCurveBase::SetCurveTypeFlag(EAnimCurveFlags InFlag, bool bValue)
-{
-	if (bValue)
-	{
-		CurveTypeFlags |= InFlag;
-	}
-	else
-	{
-		CurveTypeFlags &= ~InFlag;
-	}
-}
-
-void FAnimCurveBase::ToggleCurveTypeFlag(EAnimCurveFlags InFlag)
-{
-	bool Current = GetCurveTypeFlag(InFlag);
-	SetCurveTypeFlag(InFlag, !Current);
-}
-
-bool FAnimCurveBase::GetCurveTypeFlag(EAnimCurveFlags InFlag) const
-{
-	return (CurveTypeFlags & InFlag) != 0;
-}
-
-
-void FAnimCurveBase::SetCurveTypeFlags(int32 NewCurveTypeFlags)
-{
-	CurveTypeFlags = NewCurveTypeFlags;
-}
-
-int32 FAnimCurveBase::GetCurveTypeFlags() const
-{
-	return CurveTypeFlags;
-}
-
-////////////////////////////////////////////////////
-//  FFloatCurve
-
-// we don't want to have = operator. This only copies curves, but leaving naming and everything else intact. 
-void FFloatCurve::CopyCurve(FFloatCurve& SourceCurve)
-{
-	FloatCurve = SourceCurve.FloatCurve;
-}
-
-float FFloatCurve::Evaluate(float CurrentTime, float BlendWeight) const
-{
-	return FloatCurve.Eval(CurrentTime)*BlendWeight;
-}
-
-void FFloatCurve::UpdateOrAddKey(float NewKey, float CurrentTime)
-{
-	FloatCurve.UpdateOrAddKey(CurrentTime, NewKey);
-}
-
-void FFloatCurve::Resize(float NewLength)
-{
-	FloatCurve.ResizeTimeRange(0, NewLength);
-}
-////////////////////////////////////////////////////
-//  FVectorCurve
-
-// we don't want to have = operator. This only copies curves, but leaving naming and everything else intact. 
-void FVectorCurve::CopyCurve(FVectorCurve& SourceCurve)
-{
-	FMemory::Memcpy(FloatCurves, SourceCurve.FloatCurves);
-}
-
-FVector FVectorCurve::Evaluate(float CurrentTime, float BlendWeight) const
-{
-	FVector Value;
-
-	Value.X = FloatCurves[X].Eval(CurrentTime)*BlendWeight;
-	Value.Y = FloatCurves[Y].Eval(CurrentTime)*BlendWeight;
-	Value.Z = FloatCurves[Z].Eval(CurrentTime)*BlendWeight;
-
-	return Value;
-}
-
-void FVectorCurve::UpdateOrAddKey(const FVector& NewKey, float CurrentTime)
-{
-	FloatCurves[X].UpdateOrAddKey(CurrentTime, NewKey.X);
-	FloatCurves[Y].UpdateOrAddKey(CurrentTime, NewKey.Y);
-	FloatCurves[Z].UpdateOrAddKey(CurrentTime, NewKey.Z);
-}
-
-void FVectorCurve::Resize(float NewLength)
-{
-	FloatCurves[X].ResizeTimeRange(0, NewLength);
-	FloatCurves[Y].ResizeTimeRange(0, NewLength);
-	FloatCurves[Z].ResizeTimeRange(0, NewLength);
-}
-////////////////////////////////////////////////////
-//  FTransformCurve
-
-// we don't want to have = operator. This only copies curves, but leaving naming and everything else intact. 
-void FTransformCurve::CopyCurve(FTransformCurve& SourceCurve)
-{
-	TranslationCurve.CopyCurve(SourceCurve.TranslationCurve);
-	RotationCurve.CopyCurve(SourceCurve.RotationCurve);
-	ScaleCurve.CopyCurve(SourceCurve.ScaleCurve);
-}
-
-FTransform FTransformCurve::Evaluate(float CurrentTime, float BlendWeight) const
-{
-	FTransform Value;
-	Value.SetTranslation(TranslationCurve.Evaluate(CurrentTime, BlendWeight));
-	if (ScaleCurve.DoesContainKey())
-	{
-		Value.SetScale3D(ScaleCurve.Evaluate(CurrentTime, BlendWeight));
-	}
-	else
-	{
-		Value.SetScale3D(FVector(1.f));
-	}
-
-	// blend rotation float curve
-	FVector RotationAsVector = RotationCurve.Evaluate(CurrentTime, BlendWeight);
-	// pitch, yaw, roll order - please check AddKey function
-	FRotator Rotator(RotationAsVector.Y, RotationAsVector.Z, RotationAsVector.X);
-	Value.SetRotation(FQuat(Rotator));
-
-	return Value;
-}
-
-void FTransformCurve::UpdateOrAddKey(const FTransform& NewKey, float CurrentTime)
-{
-	TranslationCurve.UpdateOrAddKey(NewKey.GetTranslation(), CurrentTime);
-	// pitch, yaw, roll order - please check Evaluate function
-	FVector RotationAsVector;
-	FRotator Rotator = NewKey.GetRotation().Rotator();
-	RotationAsVector.X = Rotator.Roll;
-	RotationAsVector.Y = Rotator.Pitch;
-	RotationAsVector.Z = Rotator.Yaw;
-
-	RotationCurve.UpdateOrAddKey(RotationAsVector, CurrentTime);
-	ScaleCurve.UpdateOrAddKey(NewKey.GetScale3D(), CurrentTime);
-}
-
-void FTransformCurve::Resize(float NewLength)
-{
-	TranslationCurve.Resize(NewLength);
-	RotationCurve.Resize(NewLength);
-	ScaleCurve.Resize(NewLength);
-}
-
-/////////////////////////////////////////////////////
-// FRawCurveTracks
-
-void FRawCurveTracks::EvaluateCurveData(class UAnimInstance* Instance, float CurrentTime, float BlendWeight ) const
-{
-	// evaluate the curve data at the CurrentTime and add to Instance
-	for(auto CurveIter = FloatCurves.CreateConstIterator(); CurveIter; ++CurveIter)
-	{
-		const FFloatCurve& Curve = *CurveIter;
-
-		Instance->AddCurveValue(Curve.CurveUid, Curve.Evaluate(CurrentTime, BlendWeight), Curve.GetCurveTypeFlags());
-	}
-}
-
-#if WITH_EDITOR
-/**
- * Since we don't care about blending, we just change this decoration to OutCurves
- * @TODO : Fix this if we're saving vectorcurves and blending
- */
-void FRawCurveTracks::EvaluateTransformCurveData(USkeleton * Skeleton, TMap<FName, FTransform>&OutCurves, float CurrentTime, float BlendWeight) const
-{
-	check (Skeleton);
-	// evaluate the curve data at the CurrentTime and add to Instance
-	for(auto CurveIter = TransformCurves.CreateConstIterator(); CurveIter; ++CurveIter)
-	{
-		const FTransformCurve& Curve = *CurveIter;
-
-		// if disabled, do not handle
-		if (Curve.GetCurveTypeFlag(ACF_Disabled))
-		{
-			continue;
-		}
-
-		FSmartNameMapping* NameMapping = Skeleton->SmartNames.GetContainer(USkeleton::AnimTrackCurveMappingName);
-
-		// Add or retrieve curve
-		FName CurveName;
-		
-		// make sure it was added
-		if (ensure (NameMapping->GetName(Curve.CurveUid, CurveName)))
-		{
-			// note we're not checking Curve.GetCurveTypeFlags() yet
-			FTransform & Value = OutCurves.FindOrAdd(CurveName);
-			Value = Curve.Evaluate(CurrentTime, BlendWeight);
-		}
-	}
-}
-#endif
-FAnimCurveBase * FRawCurveTracks::GetCurveData(USkeleton::AnimCurveUID Uid, ESupportedCurveType SupportedCurveType /*= FloatType*/)
-{
-	switch (SupportedCurveType)
-	{
-#if WITH_EDITOR
-	case VectorType:
-		return GetCurveDataImpl<FVectorCurve>(VectorCurves, Uid);
-	case TransformType:
-		return GetCurveDataImpl<FTransformCurve>(TransformCurves, Uid);
-#endif // WITH_EDITOR
-	case FloatType:
-	default:
-		return GetCurveDataImpl<FFloatCurve>(FloatCurves, Uid);
-	}
-}
-
-bool FRawCurveTracks::DeleteCurveData(USkeleton::AnimCurveUID Uid, ESupportedCurveType SupportedCurveType /*= FloatType*/)
-{
-	switch(SupportedCurveType)
-	{
-#if WITH_EDITOR
-	case VectorType:
-		return DeleteCurveDataImpl<FVectorCurve>(VectorCurves, Uid);
-	case TransformType:
-		return DeleteCurveDataImpl<FTransformCurve>(TransformCurves, Uid);
-#endif // WITH_EDITOR
-	case FloatType:
-	default:
-		return DeleteCurveDataImpl<FFloatCurve>(FloatCurves, Uid);
-	}
-}
-
-bool FRawCurveTracks::AddCurveData(USkeleton::AnimCurveUID Uid, int32 CurveFlags /*= ACF_DefaultCurve*/, ESupportedCurveType SupportedCurveType /*= FloatType*/)
-{
-	switch(SupportedCurveType)
-	{
-#if WITH_EDITOR
-	case VectorType:
-		return AddCurveDataImpl<FVectorCurve>(VectorCurves, Uid, CurveFlags);
-	case TransformType:
-		return AddCurveDataImpl<FTransformCurve>(TransformCurves, Uid, CurveFlags);
-#endif // WITH_EDITOR
-	case FloatType:
-	default:
-		return AddCurveDataImpl<FFloatCurve>(FloatCurves, Uid, CurveFlags);
-	}
-}
-
-void FRawCurveTracks::Resize(float TotalLength)
-{
-	for (auto& Curve: FloatCurves)
-	{
-		Curve.Resize(TotalLength);
-	}
-
-#if WITH_EDITORONLY_DATA
-	for(auto& Curve: VectorCurves)
-	{
-		Curve.Resize(TotalLength);
-	}
-
-	for(auto& Curve: TransformCurves)
-	{
-		Curve.Resize(TotalLength);
-	}
-#endif
-}
-
-void FRawCurveTracks::Serialize(FArchive& Ar)
-{
-	// @TODO: If we're about to serialize vector curve, add here
-	if(Ar.UE4Ver() >= VER_UE4_SKELETON_ADD_SMARTNAMES)
-	{
-		for(FFloatCurve& Curve : FloatCurves)
-		{
-			Curve.Serialize(Ar);
-		}
-	}
-#if WITH_EDITORONLY_DATA
-	if( !Ar.IsCooking() )
-	{
-		if( Ar.UE4Ver() >= VER_UE4_ANIMATION_ADD_TRACKCURVES )
-		{
-			for( FTransformCurve& Curve : TransformCurves )
-			{
-				Curve.Serialize( Ar );
-			}
-
-		}
-	}
-#endif // WITH_EDITORONLY_DATA
-}
-
-void FRawCurveTracks::UpdateLastObservedNames(FSmartNameMapping* NameMapping, ESupportedCurveType SupportedCurveType /*= FloatType*/)
-{
-	switch(SupportedCurveType)
-	{
-#if WITH_EDITOR
-	case VectorType:
-		UpdateLastObservedNamesImpl<FVectorCurve>(VectorCurves, NameMapping);
-		break;
-	case TransformType:
-		UpdateLastObservedNamesImpl<FTransformCurve>(TransformCurves, NameMapping);
-		break;
-#endif // WITH_EDITOR
-	case FloatType:
-	default:
-		UpdateLastObservedNamesImpl<FFloatCurve>(FloatCurves, NameMapping);
-	}
-}
-
-bool FRawCurveTracks::DuplicateCurveData(USkeleton::AnimCurveUID ToCopyUid, USkeleton::AnimCurveUID NewUid, ESupportedCurveType SupportedCurveType /*= FloatType*/)
-{
-	switch(SupportedCurveType)
-	{
-#if WITH_EDITOR
-	case VectorType:
-		return DuplicateCurveDataImpl<FVectorCurve>(VectorCurves, ToCopyUid, NewUid);
-	case TransformType:
-		return DuplicateCurveDataImpl<FTransformCurve>(TransformCurves, ToCopyUid, NewUid);
-#endif // WITH_EDITOR
-	case FloatType:
-	default:
-		return DuplicateCurveDataImpl<FFloatCurve>(FloatCurves, ToCopyUid, NewUid);
-	}
-}
-
-///////////////////////////////////
-// @TODO: REFACTOR THIS IF WE'RE SERIALIZING VECTOR CURVES
-//
-// implementation template functions to accomodate FloatCurve and VectorCurve
-// for now vector curve isn't used in run-time, so it's useless outside of editor
-// so just to reduce cost of run-time, functionality is split. 
-// this split worries me a bit because if the name conflict happens this will break down w.r.t. smart naming
-// currently vector curve is not saved and not evaluated, so it will be okay since the name doesn't matter much, 
-// but this has to be refactored once we'd like to move onto serialize
-///////////////////////////////////
-template <typename DataType>
-DataType * FRawCurveTracks::GetCurveDataImpl(TArray<DataType> & Curves, USkeleton::AnimCurveUID Uid)
-{
-	for(DataType& Curve : Curves)
-	{
-		if(Curve.CurveUid == Uid)
-		{
-			return &Curve;
-		}
-	}
-
-	return NULL;
-}
-
-template <typename DataType>
-bool FRawCurveTracks::DeleteCurveDataImpl(TArray<DataType> & Curves, USkeleton::AnimCurveUID Uid)
-{
-	for(int32 Idx = 0; Idx < Curves.Num(); ++Idx)
-	{
-		if(Curves[Idx].CurveUid == Uid)
-		{
-			Curves.RemoveAt(Idx);
-			return true;
-		}
-	}
-
-	return false;
-}
-
-template <typename DataType>
-bool FRawCurveTracks::AddCurveDataImpl(TArray<DataType> & Curves, USkeleton::AnimCurveUID Uid, int32 CurveFlags)
-{
-	if(GetCurveDataImpl<DataType>(Curves, Uid) == NULL)
-	{
-		Curves.Add(DataType(Uid, CurveFlags));
-		return true;
-	}
-	return false;
-}
-
-template <typename DataType>
-void FRawCurveTracks::UpdateLastObservedNamesImpl(TArray<DataType> & Curves, FSmartNameMapping* NameMapping)
-{
-	if(NameMapping)
-	{
-		for(DataType& Curve : Curves)
-		{
-			NameMapping->GetName(Curve.CurveUid, Curve.LastObservedName);
-		}
-	}
-}
-
-template <typename DataType>
-bool FRawCurveTracks::DuplicateCurveDataImpl(TArray<DataType> & Curves, USkeleton::AnimCurveUID ToCopyUid, USkeleton::AnimCurveUID NewUid)
-{
-	DataType* ExistingCurve = GetCurveDataImpl<DataType>(Curves, ToCopyUid);
-	if(ExistingCurve && GetCurveDataImpl<DataType>(Curves, NewUid) == NULL)
-	{
-		// Add the curve to the track and set its data to the existing curve
-		Curves.Add(DataType(NewUid, ExistingCurve->GetCurveTypeFlags()));
-		Curves.Last().CopyCurve(*ExistingCurve);
-
-		return true;
-	}
-	return false;
-}
-
 /////////////////////////////////////////////////////
 
 UAnimSequenceBase::UAnimSequenceBase(const FObjectInitializer& ObjectInitializer)
@@ -729,8 +330,6 @@ void UAnimSequenceBase::TickAssetPlayerInstance(const FAnimTickRecord& Instance,
 
 	OnAssetPlayerTickedInternal(Context, PreviousTime, MoveDelta, Instance, InstanceOwner);
 
-	// Evaluate Curve data now - even if time did not move, we still need to return curve if it exists
-	EvaluateCurveData(InstanceOwner, CurrentTime, Instance.EffectiveBlendWeight);
 }
 
 #if WITH_EDITOR
@@ -894,11 +493,9 @@ uint8* UAnimSequenceBase::FindNotifyPropertyData(int32 NotifyIndex, UArrayProper
 
 
 /** Add curve data to Instance at the time of CurrentTime **/
-void UAnimSequenceBase::EvaluateCurveData(class UAnimInstance* Instance, float CurrentTime, float BlendWeight ) const
+void UAnimSequenceBase::EvaluateCurveData(FBlendedCurve& OutCurve, float CurrentTime ) const
 {
-	// if we have compression, this should change
-	// this is raw data evaluation
-	RawCurveData.EvaluateCurveData(Instance, CurrentTime, BlendWeight);
+	RawCurveData.EvaluateCurveData(OutCurve, CurrentTime);
 }
 
 void UAnimSequenceBase::Serialize(FArchive& Ar)
diff --git a/Engine/Source/Runtime/Engine/Private/Animation/AnimSingleNodeInstance.cpp b/Engine/Source/Runtime/Engine/Private/Animation/AnimSingleNodeInstance.cpp
--- a/Engine/Source/Runtime/Engine/Private/Animation/AnimSingleNodeInstance.cpp
+++ b/Engine/Source/Runtime/Engine/Private/Animation/AnimSingleNodeInstance.cpp
@@ -313,7 +313,7 @@ bool UAnimSingleNodeInstance::NativeEvaluateAnimation(FPoseContext& Output)
 		//@TODO: animrefactor: Seems like more code duplication than we need
 		if (UBlendSpaceBase* BlendSpace = Cast<UBlendSpaceBase>(CurrentAsset))
 		{
-			InternalBlendSpaceEvaluatePose(BlendSpace, BlendSampleData, Output.Pose);
+			InternalBlendSpaceEvaluatePose(BlendSpace, BlendSampleData, Output);
 		}
 		else if (UAnimSequence* Sequence = Cast<UAnimSequence>(CurrentAsset))
 		{
@@ -321,64 +321,55 @@ bool UAnimSingleNodeInstance::NativeEvaluateAnimation(FPoseContext& Output)
 			{
 				FAnimExtractContext ExtractionContext(CurrentTime, Sequence->bEnableRootMotion);
 
-				Sequence->GetAdditiveBasePose(Output.Pose, ExtractionContext);
+				Sequence->GetAdditiveBasePose(Output.Pose, Output.Curve, ExtractionContext);
 
 				FCompactPose AdditivePose;
+				FBlendedCurve AdditiveCurve;
 				AdditivePose.SetBoneContainer(&Output.Pose.GetBoneContainer());
-				Sequence->GetAnimationPose(AdditivePose, ExtractionContext);
+				AdditiveCurve.InitFrom(Output.Curve);
+				Sequence->GetAnimationPose(AdditivePose, AdditiveCurve, ExtractionContext);
 
-				if (Sequence->AdditiveAnimType == AAT_LocalSpaceBase)
-				{
-					FAnimationRuntime::AccumulateAdditivePose(Output.Pose, AdditivePose, 1.f);
-					Output.Pose.NormalizeRotations();
-				}
-				else
-				{
-					FAnimationRuntime::AccumulateMeshSpaceRotationAdditiveToLocalPose(Output.Pose, AdditivePose, 1.f);
-					Output.Pose.NormalizeRotations();
-				}
+				FAnimationRuntime::AccumulateAdditivePose(Output.Pose, AdditivePose, Output.Curve, AdditiveCurve, 1.f, Sequence->AdditiveAnimType);
+				Output.Pose.NormalizeRotations();
 			}
 			else
 			{
 				// if SkeletalMesh isn't there, we'll need to use skeleton
-				FAnimationRuntime::GetPoseFromSequence(Sequence, Output.Pose, FAnimExtractContext(CurrentTime, Sequence->bEnableRootMotion));
+				Sequence->GetAnimationPose(Output.Pose, Output.Curve, FAnimExtractContext(CurrentTime, Sequence->bEnableRootMotion));
 			}
 		}
 		else if (UAnimComposite* Composite = Cast<UAnimComposite>(CurrentAsset))
 		{
+			FAnimExtractContext ExtractionContext(CurrentTime, ShouldExtractRootMotion());
 			const FAnimTrack& AnimTrack = Composite->AnimationTrack;
 
 			// find out if this is additive animation
-			EAdditiveAnimationType AdditiveAnimType = AAT_None;
 			if (AnimTrack.IsAdditive())
 			{
-				FCompactPose AdditivePose;
-				AdditivePose.SetBoneContainer(&Output.Pose.GetBoneContainer());
-
-				AdditiveAnimType = AnimTrack.IsRotationOffsetAdditive()? AAT_RotationOffsetMeshSpace : AAT_LocalSpaceBase;
-				
-				// get base pose - for now we only support ref pose as base
-				Output.Pose.ResetToRefPose();
-				
-				//get the additive pose
-				FAnimationRuntime::GetPoseFromAnimTrack(AnimTrack, AdditivePose, FAnimExtractContext(CurrentTime, RootMotionMode == ERootMotionMode::RootMotionFromEverything));
-
-				// if additive, we should blend with source to make it fullbody
-				if (AdditiveAnimType == AAT_LocalSpaceBase)
+				if (Composite->PreviewBasePose)
 				{
-					FAnimationRuntime::AccumulateAdditivePose(Output.Pose, AdditivePose, 1.f);
-					Output.Pose.NormalizeRotations();
+					Composite->PreviewBasePose->GetAdditiveBasePose(Output.Pose, Output.Curve, ExtractionContext);
 				}
-				else if (AdditiveAnimType == AAT_RotationOffsetMeshSpace)
+				else
 				{
-					FAnimationRuntime::AccumulateMeshSpaceRotationAdditiveToLocalPose(Output.Pose, AdditivePose, 1.f);
-					Output.Pose.NormalizeRotations();
+					// get base pose - for now we only support ref pose as base
+					Output.Pose.ResetToRefPose();
 				}
+
+				EAdditiveAnimationType AdditiveAnimType = AnimTrack.IsRotationOffsetAdditive()? AAT_RotationOffsetMeshSpace : AAT_LocalSpaceBase;
+
+				FCompactPose AdditivePose;
+				FBlendedCurve AdditiveCurve;
+				AdditivePose.SetBoneContainer(&Output.Pose.GetBoneContainer());
+				AdditiveCurve.InitFrom(Output.Curve);
+				Composite->GetAnimationPose(AdditivePose, AdditiveCurve, ExtractionContext);
+
+				FAnimationRuntime::AccumulateAdditivePose(Output.Pose, AdditivePose, Output.Curve, AdditiveCurve, 1.f, AdditiveAnimType);
 			}
 			else
 			{
 				//doesn't handle additive yet
-				FAnimationRuntime::GetPoseFromAnimTrack(AnimTrack, Output.Pose, FAnimExtractContext(CurrentTime, RootMotionMode == ERootMotionMode::RootMotionFromEverything));
+				Composite->GetAnimationPose(Output.Pose, Output.Curve, ExtractionContext);
 			}
 		}
 		else if (UAnimMontage* Montage = Cast<UAnimMontage>(CurrentAsset))
@@ -388,7 +379,9 @@ bool UAnimSingleNodeInstance::NativeEvaluateAnimation(FPoseContext& Output)
 			if (Montage->SlotAnimTracks.Num() > 0)
 			{
 				FCompactPose SourcePose;
+				FBlendedCurve SourceCurve;
 				SourcePose.SetBoneContainer(&Output.Pose.GetBoneContainer());
+				SourceCurve.InitFrom(Output.Curve);
 			
 				if (Montage->IsValidAdditive())
 				{
@@ -396,7 +389,7 @@ bool UAnimSingleNodeInstance::NativeEvaluateAnimation(FPoseContext& Output)
 					// if montage is additive, we need to have base pose for the slot pose evaluate
 					if (Montage->PreviewBasePose && Montage->SequenceLength > 0.f)
 					{
-						Montage->PreviewBasePose->GetBonePose(SourcePose, FAnimExtractContext(CurrentTime));
+						Montage->PreviewBasePose->GetBonePose(SourcePose, SourceCurve, FAnimExtractContext(CurrentTime));
 					}
 					else
 #endif // WITH_EDITORONLY_DATA
@@ -409,7 +402,7 @@ bool UAnimSingleNodeInstance::NativeEvaluateAnimation(FPoseContext& Output)
 					SourcePose.ResetToRefPose();			
 				}
 
-				SlotEvaluatePose(Montage->SlotAnimTracks[0].SlotName, SourcePose, Output.Pose, 1.f);
+				SlotEvaluatePose(Montage->SlotAnimTracks[0].SlotName, SourcePose, SourceCurve, Output.Pose, Output.Curve, 1.f);
 			}
 		}
 	}
@@ -549,12 +542,6 @@ void UAnimSingleNodeInstance::SetPosition(float InPosition, bool bFireNotifies)
 
 			TriggerAnimNotifies(0.f);
 
-			// since this is singlenode instance, if position changes, we can't keep old morphtarget curves
-			// we clear it and evaluate curve here with new asset. 
-			MorphTargetCurves.Empty();
-			MaterialParameterCurves.Empty();
-			// Evaluate Curve data now - even if time did not move, we still need to return curve if it exists
-			SequenceBase->EvaluateCurveData(this, CurrentTime, 1.0);
 		}
 	}
 }
@@ -564,19 +551,23 @@ void UAnimSingleNodeInstance::SetBlendSpaceInput(const FVector& InBlendInput)
 	BlendSpaceInput = InBlendInput;
 }
 
-void UAnimSingleNodeInstance::InternalBlendSpaceEvaluatePose(class UBlendSpaceBase* BlendSpace, TArray<FBlendSampleData>& BlendSampleDataCache, FCompactPose& OutPose)
+void UAnimSingleNodeInstance::InternalBlendSpaceEvaluatePose(class UBlendSpaceBase* BlendSpace, TArray<FBlendSampleData>& BlendSampleDataCache, FPoseContext& OutContext)
 {
 	USkeletalMeshComponent* Component = GetSkelMeshComponent();
 
 	if (BlendSpace->IsValidAdditive())
 	{
+		FCompactPose& OutPose = OutContext.Pose;
+		FBlendedCurve& OutCurve = OutContext.Curve;
 		FCompactPose AdditivePose;
+		FBlendedCurve AdditiveCurve;
 		AdditivePose.SetBoneContainer(&OutPose.GetBoneContainer());
+		AdditiveCurve.InitFrom(OutCurve);
 
 #if WITH_EDITORONLY_DATA
 		if (BlendSpace->PreviewBasePose)
 		{
-			BlendSpace->PreviewBasePose->GetBonePose(/*out*/ OutPose, FAnimExtractContext(PreviewPoseCurrentTime));
+			BlendSpace->PreviewBasePose->GetBonePose(/*out*/ OutPose, /*out*/OutCurve, FAnimExtractContext(PreviewPoseCurrentTime));
 		}
 		else
 #endif // WITH_EDITORONLY_DATA
@@ -585,29 +576,15 @@ void UAnimSingleNodeInstance::InternalBlendSpaceEvaluatePose(class UBlendSpaceBa
 			OutPose.ResetToRefPose();
 		}
 
-		FAnimationRuntime::GetPoseFromBlendSpace(
-			BlendSpace,
-			BlendSampleDataCache,
-			/*out*/ AdditivePose);
+		BlendSpace->GetAnimationPose(BlendSampleDataCache, AdditivePose, AdditiveCurve);
 
-		if (BlendSpace->IsA(UAimOffsetBlendSpace::StaticClass()) ||
-			BlendSpace->IsA(UAimOffsetBlendSpace1D::StaticClass()) )
-		{
-			FAnimationRuntime::AccumulateMeshSpaceRotationAdditiveToLocalPose(OutPose, AdditivePose, 1.f);
-			OutPose.NormalizeRotations();
-		}
-		else
-		{
-			FAnimationRuntime::AccumulateAdditivePose(OutPose, AdditivePose, 1.f);
-			OutPose.NormalizeRotations();
-		}
+		enum EAdditiveAnimationType AdditiveType = BlendSpace->bRotationBlendInMeshSpace? AAT_RotationOffsetMeshSpace : AAT_LocalSpaceBase;
+
+		FAnimationRuntime::AccumulateAdditivePose(OutPose, AdditivePose, OutCurve, AdditiveCurve, 1.f, AdditiveType);
 	}
 	else
 	{
-		FAnimationRuntime::GetPoseFromBlendSpace(
-			BlendSpace,
-			BlendSampleDataCache,
-			/*out*/ OutPose);
+		BlendSpace->GetAnimationPose(BlendSampleDataCache, OutContext.Pose, OutContext.Curve);
 	}
 }
 
diff --git a/Engine/Source/Runtime/Engine/Private/Animation/AnimationRuntime.cpp b/Engine/Source/Runtime/Engine/Private/Animation/AnimationRuntime.cpp
--- a/Engine/Source/Runtime/Engine/Private/Animation/AnimationRuntime.cpp
+++ b/Engine/Source/Runtime/Engine/Private/Animation/AnimationRuntime.cpp
@@ -101,10 +101,82 @@ FORCEINLINE void BlendPose(const FCompactPose& SourcePose, FCompactPose& ResultP
 	}
 }
 
+FORCEINLINE void BlendCurves(const TArray<FBlendedCurve>& SourceCurves, const TArray<float>&	SourceWeights, FBlendedCurve& OutCurve)
+{
+	if (SourceCurves.Num() > 0)
+	{
+		OutCurve.Override(SourceCurves[0], SourceWeights[0]);
+
+		for(int32 CurveIndex=1; CurveIndex<SourceCurves.Num(); ++CurveIndex)
+		{
+			OutCurve.Accumulate(SourceCurves[CurveIndex], SourceWeights[CurveIndex]);
+		}
+	}
+}
+
+FORCEINLINE void BlendCurves(const TArray<const FBlendedCurve*>& SourceCurves, const TArray<float>&	SourceWeights, FBlendedCurve& OutCurve)
+{
+	if(SourceCurves.Num() > 0)
+	{
+		OutCurve.Override(*SourceCurves[0], SourceWeights[0]);
+
+		for(int32 CurveIndex=1; CurveIndex<SourceCurves.Num(); ++CurveIndex)
+		{
+			OutCurve.Accumulate(*SourceCurves[CurveIndex], SourceWeights[CurveIndex]);
+		}
+	}
+}
+
+FORCEINLINE void BlendCurves(const TArray<const FBlendedCurve*>& SourceCurves, const TArray<float>&	SourceWeights, FBlendedCurve& OutCurve, ECurveBlendOption::Type BlendOption)
+{
+	if(SourceCurves.Num() > 0)
+	{
+		if (BlendOption == ECurveBlendOption::Type::BlendByWeight)
+		{
+			BlendCurves(SourceCurves, SourceWeights, OutCurve);
+		}
+		else if (BlendOption == ECurveBlendOption::Type::NormalizeByWeight)
+		{
+			float SumOfWeight = 0.f;
+			for (const auto& Weight : SourceWeights)
+			{
+				SumOfWeight += Weight;
+			}
+
+			if (SumOfWeight > ZERO_ANIMWEIGHT_THRESH)
+			{
+				TArray<float>	NormalizeSourceWeights;
+				NormalizeSourceWeights.AddUninitialized(SourceWeights.Num());
+				for(int32 Idx=0; Idx<SourceWeights.Num(); ++Idx)
+				{
+					NormalizeSourceWeights[Idx] = SourceWeights[Idx] / SumOfWeight;
+				}
+
+				BlendCurves(SourceCurves, NormalizeSourceWeights, OutCurve);
+			}
+			else
+			{
+				BlendCurves(SourceCurves, SourceWeights, OutCurve);
+			}
+		}
+		else
+		{
+			OutCurve.Override(*SourceCurves[0], SourceWeights[0]);
+
+			for(int32 CurveIndex=1; CurveIndex<SourceCurves.Num(); ++CurveIndex)
+			{
+				OutCurve.Combine(*SourceCurves[CurveIndex]);
+			}
+		}
+	}
+}
+
 void FAnimationRuntime::BlendPosesTogether(
 	const TArray<FCompactPose>& SourcePoses,
+	const TArray<FBlendedCurve>& SourceCurves,
 	const TArray<float>&	SourceWeights,
-	/*out*/ FCompactPose& ResultPose)
+	/*out*/ FCompactPose& ResultPose, 
+	/*out*/ FBlendedCurve& ResultCurve)
 {
 	check(SourcePoses.Num() > 0);
 
@@ -120,12 +192,20 @@ void FAnimationRuntime::BlendPosesTogether(
 	{
 		ResultPose.NormalizeRotations();
 	}
+
+	// curve blending if exists
+	if (SourceCurves.Num() > 0)
+	{
+		BlendCurves(SourceCurves, SourceWeights, ResultCurve);
+	}
 }
 
 void FAnimationRuntime::BlendPosesTogether(
 	const TArray<const FCompactPose*>& SourcePoses,
+	const TArray<const FBlendedCurve*>& SourceCurves,
 	const TArray<float>&	SourceWeights,
-	/*out*/ FCompactPose& ResultPose)
+	/*out*/ FCompactPose& ResultPose, 
+	/*out*/ FBlendedCurve& ResultCurve)
 {
 	check(SourcePoses.Num() > 0);
 
@@ -141,19 +221,28 @@ void FAnimationRuntime::BlendPosesTogether(
 	{
 		ResultPose.NormalizeRotations();
 	}
+
+	if(SourceCurves.Num() > 0)
+	{
+		BlendCurves(SourceCurves, SourceWeights, ResultCurve);
+	}
 }
 
 void FAnimationRuntime::BlendTwoPosesTogether(
 	const FCompactPose& SourcePose1,
 	const FCompactPose& SourcePose2,
+	const FBlendedCurve& SourceCurve1,
+	const FBlendedCurve& SourceCurve2,
 	const float			WeightOfPose1,
-	/*out*/ FCompactPose& ResultPose)
+	/*out*/ FCompactPose& ResultPose,
+	/*out*/ FBlendedCurve& ResultCurve)
 {
 	BlendPose<ETransformBlendMode::Overwrite>(SourcePose1, ResultPose, WeightOfPose1);
 	BlendPose<ETransformBlendMode::Accumulate>(SourcePose2, ResultPose, 1.f - WeightOfPose1);
 
 	// Ensure that all of the resulting rotations are normalized
 	ResultPose.NormalizeRotations();
+	ResultCurve.Blend(SourceCurve1, SourceCurve2, 1.f - WeightOfPose1);
 }
 
 template <int32 TRANSFORM_BLEND_MODE>
@@ -218,7 +307,7 @@ void BlendPosePerBone(const TArray<int32>& PerBoneIndices, const FBlendSampleDat
  *							This array must be in strictly increasing order.
  */
 
-void FAnimationRuntime::BlendPosesTogetherPerBone(const TArray<FCompactPose>& SourcePoses, const UBlendSpaceBase* BlendSpace, const TArray<FBlendSampleData>& BlendSampleDataCache, /*out*/ FCompactPose& ResultPose)
+void FAnimationRuntime::BlendPosesTogetherPerBone(const TArray<FCompactPose>& SourcePoses, const TArray<FBlendedCurve>& SourceCurves, const UBlendSpaceBase* BlendSpace, const TArray<FBlendSampleData>& BlendSampleDataCache, /*out*/ FCompactPose& ResultPose, /*out*/ FBlendedCurve& ResultCurve)
 {
 	check(SourcePoses.Num() > 0);
 
@@ -226,7 +315,7 @@ void FAnimationRuntime::BlendPosesTogetherPerBone(const TArray<FCompactPose>& So
 
 	TArray<int32> PerBoneIndices;
 	PerBoneIndices.AddUninitialized(ResultPose.GetNumBones());
-	for (int BoneIndex = 0; BoneIndex < PerBoneIndices.Num(); ++BoneIndex)
+	for (int32 BoneIndex = 0; BoneIndex < PerBoneIndices.Num(); ++BoneIndex)
 	{
 		PerBoneIndices[BoneIndex] = BlendSpace->GetPerBoneInterpolationIndex(RequiredBoneIndices[BoneIndex], ResultPose.GetBoneContainer());
 	}
@@ -240,9 +329,21 @@ void FAnimationRuntime::BlendPosesTogetherPerBone(const TArray<FCompactPose>& So
 
 	// Ensure that all of the resulting rotations are normalized
 	ResultPose.NormalizeRotations();
+
+	if (SourceCurves.Num() > 0)
+	{
+		TArray<float> SourceWeights;
+		SourceWeights.AddUninitialized(BlendSampleDataCache.Num());
+		for (int32 CacheIndex=0; CacheIndex<BlendSampleDataCache.Num(); ++CacheIndex)
+		{
+			SourceWeights[CacheIndex] = BlendSampleDataCache[CacheIndex].TotalWeight;
+		}
+
+		BlendCurves(SourceCurves, SourceWeights, ResultCurve);
+	}
 }
 
-void FAnimationRuntime::BlendPosesTogetherPerBoneInMeshSpace(TArray<FCompactPose>& SourcePoses, const UBlendSpaceBase* BlendSpace, const TArray<FBlendSampleData>& BlendSampleDataCache, FCompactPose& ResultPose)
+void FAnimationRuntime::BlendPosesTogetherPerBoneInMeshSpace(TArray<FCompactPose>& SourcePoses, const TArray<FBlendedCurve>& SourceCurves, const UBlendSpaceBase* BlendSpace, const TArray<FBlendSampleData>& BlendSampleDataCache, FCompactPose& ResultPose, FBlendedCurve& ResultCurve)
 {
 	FQuat NewRotation;
 	USkeleton* Skeleton = BlendSpace->GetSkeleton();
@@ -269,7 +370,7 @@ void FAnimationRuntime::BlendPosesTogetherPerBoneInMeshSpace(TArray<FCompactPose
 	}
 
 	// now we have mesh space rotation, call BlendPosesTogetherPerBone
-	BlendPosesTogetherPerBone(SourcePoses, BlendSpace, BlendSampleDataCache, ResultPose);
+	BlendPosesTogetherPerBone(SourcePoses, SourceCurves, BlendSpace, BlendSampleDataCache, ResultPose, ResultCurve);
 
 	// now result atoms has the output with mesh space rotation. Convert back to local space, start from back
 	for (const FCompactPoseBoneIndex BoneIndex : ResultPose.ForEachBoneIndex())
@@ -283,20 +384,6 @@ void FAnimationRuntime::BlendPosesTogetherPerBoneInMeshSpace(TArray<FCompactPose
 		}
 	}
 }
-/**
- * Accumulates BlendPoses to ResultAtoms with BlendWeight. 
- * ResultAtoms += BlendPose*BlendWeight.
- * Result is NOT Normalized, as this is expected to be called several times with each 'BlendPoses' transform array to be accumulated.
- * 
- * @param	BlendPoses		Array to BoneTransforms to scale and accumulate.
- * @param	BlendWeight		Weight to apply to 'BlendPoses'
- * @param	RequiredBones	Indices of bones to process, in strictly increasing order.
- * @param	ResultAtoms		Array of BoneTransforms to accumulate to.
- */
-void FAnimationRuntime::BlendPosesAccumulate(const FCompactPose& SourcePose, const float BlendWeight, /*inout*/ FCompactPose& ResultPose)
-{
-	BlendPose<ETransformBlendMode::Accumulate>(SourcePose, ResultPose, BlendWeight);
-}
 
 void FAnimationRuntime::LerpBoneTransforms(TArray<FTransform> & A, const TArray<FTransform> & B, float Alpha, const TArray<FBoneIndexType> & RequiredBonesArray)
 {
@@ -343,95 +430,6 @@ void FAnimationRuntime::CombineWithAdditiveAnimations(int32 NumAdditivePoses, co
 	}
 }
 
-void FAnimationRuntime::GetPoseFromSequence(
-		const UAnimSequenceBase* InSequence,
-		/*out*/ FCompactPose& ResultPose,
-		const FAnimExtractContext& ExtractionContext)
-{
-	//@TODO: ANIM: ANIMCOMPOSITE: Some of the code in here cares about UAnimSequence members and won't work on something derived from a different branch of UAnimSequenceBase
-	const UAnimSequence* Sequence = Cast<const UAnimSequence>(InSequence);
-
-	if (!Sequence)
-	{
-#if 0
-		UE_LOG(LogAnimation, Log, TEXT("FAnimationRuntime::GetPoseFromSequence - %s - No animation data!"), *GetFName());
-#endif
-		ResultPose.ResetToRefPose();
-		return;
-	}
-
-	Sequence->GetAnimationPose(ResultPose, ExtractionContext);
-
-	// Check that all bone atoms coming from animation are normalized
-#if DO_CHECK && WITH_EDITORONLY_DATA
-	check(ResultPose.IsNormalized());
-#endif
-}
-
-void FAnimationRuntime::GetPoseFromAnimTrack(
-	const FAnimTrack& Track,
-	/*out*/ FCompactPose& ResultPose,
-	const FAnimExtractContext& ExtractionContext)
-{
-	TArray<FCompactPose> SourcePoses;
-	TArray<float> SourceWeights;
-	float TotalWeight = 0.f;
-
-	float CurrentTime = FMath::Clamp(ExtractionContext.CurrentTime, 0.f, Track.GetLength());
-
-	// first get all the montage instance weight this slot node has
-	for (int32 I = 0; I<Track.AnimSegments.Num(); ++I)
-	{
-		const FAnimSegment& AnimSegment = Track.AnimSegments[I];
-
-		float PositionInAnim = 0.f;
-		float Weight = 0.f;
-		UAnimSequenceBase* AnimRef = AnimSegment.GetAnimationData(CurrentTime, PositionInAnim, Weight);
-
-		// make this to be 1 function
-		if (AnimRef && (Weight > ZERO_ANIMWEIGHT_THRESH))
-		{
-			// todo anim: hack - until we fix animcomposite
-			UAnimSequence * Sequence = Cast<UAnimSequence>(AnimRef);
-			if (Sequence)
-			{
-				const int32 NewIndex = SourceWeights.AddUninitialized(1);
-				SourcePoses.Add(FCompactPose());
-
-				SourcePoses[NewIndex].SetBoneContainer(&ResultPose.GetBoneContainer());
-				SourceWeights[NewIndex] = Weight;
-				TotalWeight += Weight;
-
-				// Copy passed in Extraction Context, but override position and looping parameters.
-				FAnimExtractContext SequenceExtractionContext(ExtractionContext);
-				SequenceExtractionContext.CurrentTime = PositionInAnim;
-				SequenceExtractionContext.bExtractRootMotion &= Sequence->bEnableRootMotion;
-				FAnimationRuntime::GetPoseFromSequence(Sequence, SourcePoses[NewIndex], SequenceExtractionContext);
-			}
-		}
-	}
-
-	if (SourcePoses.Num() == 0)
-	{
-		ResultPose.ResetToRefPose();
-	}
-	else if (SourcePoses.Num() == 1)
-	{
-		ResultPose = SourcePoses[0];
-	}
-	else
-	{
-		// If we have SourcePoses.Num() > 0, then we will have a non zero weight.
-		check(TotalWeight >= ZERO_ANIMWEIGHT_THRESH);
-		for (int32 I = 0; I < SourceWeights.Num(); ++I)
-		{
-			// normalize I
-			SourceWeights[I] /= TotalWeight;
-		}
-		FAnimationRuntime::BlendPosesTogether(SourcePoses, SourceWeights, ResultPose);
-	}
-}
-
 void FAnimationRuntime::ConvertPoseToAdditive(FCompactPose& TargetPose, const FCompactPose& BasePose)
 {
 	for (FCompactPoseBoneIndex BoneIndex : BasePose.ForEachBoneIndex())
@@ -471,8 +469,24 @@ void FAnimationRuntime::ConvertMeshRotationPoseToLocalSpace(FCompactPose& Pose)
 		Pose[BoneIndex].SetRotation(LocalSpaceRotation);
 	}
 }
+void FAnimationRuntime::AccumulateAdditivePose(FCompactPose& BasePose, const FCompactPose& AdditivePose, FBlendedCurve& BaseCurve, const FBlendedCurve& AdditiveCurve, float Weight, enum EAdditiveAnimationType AdditiveType)
+{
+	if (AdditiveType == AAT_RotationOffsetMeshSpace)
+	{
+		AccumulateMeshSpaceRotationAdditiveToLocalPose(BasePose, AdditivePose, BaseCurve, AdditiveCurve, Weight);
+	}
+	else
+	{
+		AccumulateLocalSpaceAdditivePose(BasePose, AdditivePose, BaseCurve, AdditiveCurve, Weight);
+	}
+
+	// if curve exists, accumulate with the weight, 
+	BaseCurve.Accumulate(AdditiveCurve, Weight);
+	// normalize
+	BasePose.NormalizeRotations();
+}
 
-void FAnimationRuntime::AccumulateAdditivePose(FCompactPose& BasePose, const FCompactPose& AdditivePose, float Weight)
+void FAnimationRuntime::AccumulateLocalSpaceAdditivePose(FCompactPose& BasePose, const FCompactPose& AdditivePose, FBlendedCurve& BaseCurve, const FBlendedCurve& AdditiveCurve, float Weight)
 {
 	if (Weight > ZERO_ANIMWEIGHT_THRESH)
 	{
@@ -498,7 +512,7 @@ void FAnimationRuntime::AccumulateAdditivePose(FCompactPose& BasePose, const FCo
 	}
 }
 
-void FAnimationRuntime::AccumulateMeshSpaceRotationAdditiveToLocalPose(FCompactPose& BasePose, const FCompactPose& MeshSpaceRotationAdditive, float Weight)
+void FAnimationRuntime::AccumulateMeshSpaceRotationAdditiveToLocalPose(FCompactPose& BasePose, const FCompactPose& MeshSpaceRotationAdditive, FBlendedCurve& BaseCurve, const FBlendedCurve& AdditiveCurve, float Weight)
 {
 	if (Weight > ZERO_ANIMWEIGHT_THRESH)
 	{
@@ -506,7 +520,7 @@ void FAnimationRuntime::AccumulateMeshSpaceRotationAdditiveToLocalPose(FCompactP
 		FAnimationRuntime::ConvertPoseToMeshRotation(BasePose);
 
 		// Add MeshSpaceRotAdditive to it
-		FAnimationRuntime::AccumulateAdditivePose(BasePose, MeshSpaceRotationAdditive, Weight);
+		FAnimationRuntime::AccumulateLocalSpaceAdditivePose(BasePose, MeshSpaceRotationAdditive, BaseCurve, AdditiveCurve, Weight);
 
 		// Convert back to local space
 		FAnimationRuntime::ConvertMeshRotationPoseToLocalSpace(BasePose);
@@ -563,89 +577,6 @@ void FAnimationRuntime::ApplyWeightToTransform(const FBoneContainer& RequiredBon
 	}			
 }
 
-void FAnimationRuntime::GetPoseFromBlendSpace(
-	UBlendSpaceBase * BlendSpace,
-	TArray<FBlendSampleData>& BlendSampleDataCache,
-	/*out*/ FCompactPose& ResultPose)
-{
-	// if not valid parameters
-	if (BlendSpace == NULL || BlendSampleDataCache.Num() == 0)
-	{
-		ResultPose.ResetToRefPose();
-		return;
-	}
-
-	const int32 NumPoses = BlendSampleDataCache.Num();
-
-	TArray<FCompactPose> ChildrenPoses;
-	ChildrenPoses.AddZeroed(NumPoses);
-
-	TArray<float>	ChildrenWeights;
-	ChildrenWeights.AddZeroed(NumPoses);
-
-	for (FCompactPose& Pose : ChildrenPoses)
-	{
-		Pose.SetBoneContainer(&ResultPose.GetBoneContainer());
-	}
-
-	const TArray<struct FBlendSample> & SampleData = BlendSpace->GetBlendSamples();
-
-	// get all child atoms we interested in
-	for (int32 I = 0; I < BlendSampleDataCache.Num(); ++I)
-	{
-		FCompactPose& Pose = ChildrenPoses[I];
-
-		if (SampleData.IsValidIndex(BlendSampleDataCache[I].SampleDataIndex))
-		{
-			const FBlendSample& Sample = SampleData[BlendSampleDataCache[I].SampleDataIndex];
-			ChildrenWeights[I] = BlendSampleDataCache[I].GetWeight();
-
-			if (Sample.Animation)
-			{
-				const float Time = FMath::Clamp<float>(BlendSampleDataCache[I].Time, 0.f, Sample.Animation->SequenceLength);
-
-				// first one always fills up the source one
-				FAnimationRuntime::GetPoseFromSequence(Sample.Animation, Pose, FAnimExtractContext(Time, true));
-			}
-			else
-			{
-				Pose.ResetToRefPose();
-			}
-		}
-		else
-		{
-			Pose.ResetToRefPose();
-		}
-	}
-
-	if (BlendSpace->PerBoneBlend.Num() > 0)
-	{
-		if (BlendSpace->IsValidAdditive())
-		{
-			FAnimationRuntime::BlendPosesTogetherPerBone(ChildrenPoses, BlendSpace, BlendSampleDataCache, ResultPose);
-		}
-		else
-		{
-			if (BlendSpace->bRotationBlendInMeshSpace)
-			{
-				FAnimationRuntime::BlendPosesTogetherPerBoneInMeshSpace(ChildrenPoses, BlendSpace, BlendSampleDataCache, ResultPose);
-			}
-			else
-			{
-				FAnimationRuntime::BlendPosesTogetherPerBone(ChildrenPoses, BlendSpace, BlendSampleDataCache, ResultPose);
-			}
-		}
-	}
-	else
-	{
-		FAnimationRuntime::BlendPosesTogether(ChildrenPoses, ChildrenWeights, ResultPose);
-	}
-
-	// Once all the accumulation and blending has been done, normalize rotations.
-	ResultPose.NormalizeRotations();
-}
-
-
 /* from % from OutKeyIndex1, meaning (CurrentKeyIndex(float)-OutKeyIndex1)/(OutKeyIndex2-OutKeyIndex1) */
 void FAnimationRuntime::GetKeyIndicesFromTime(int32& OutKeyIndex1, int32& OutKeyIndex2, float& OutAlpha, const float Time, const int32 NumFrames, const float SequenceLength)
 {
@@ -853,8 +784,12 @@ void FAnimationRuntime::ExcludeBonesWithNoParents(const TArray<int32> & BoneIndi
 void FAnimationRuntime::BlendMeshPosesPerBoneWeights(
 		struct FCompactPose& BasePose,
 		const TArray<struct FCompactPose>& BlendPoses,
-		const TArray<FPerBoneBlendWeight> & BoneBlendWeights,
-		/*out*/ FCompactPose& OutPose)
+		struct FBlendedCurve& BaseCurve,
+		const TArray<struct FBlendedCurve>& BlendedCurves,
+		const TArray<FPerBoneBlendWeight>& BoneBlendWeights,
+		ECurveBlendOption::Type CurveBlendOption,
+		/*out*/ FCompactPose& OutPose,
+		/*out*/ struct FBlendedCurve& OutCurve)
 {
 	check(BasePose.GetNumBones() == BoneBlendWeights.Num());
 
@@ -868,6 +803,11 @@ void FAnimationRuntime::BlendMeshPosesPerBoneWeights(
 	BlendRotations.AddUninitialized(BasePose.GetNumBones());
 	TargetRotations.AddUninitialized(BasePose.GetNumBones());
 
+	int32 PoseNum = BlendPoses.Num();
+
+	TArray<float> MaxPoseWeights;
+	MaxPoseWeights.AddZeroed(PoseNum);
+
 	for (FCompactPoseBoneIndex BoneIndex : BasePose.ForEachBoneIndex())
 	{
 		const int32 PoseIndex = BoneBlendWeights[BoneIndex.GetInt()].SourceIndex;
@@ -897,6 +837,8 @@ void FAnimationRuntime::BlendMeshPosesPerBoneWeights(
 		FTransform BlendAtom;
 
 		const float BlendWeight = FMath::Clamp(BoneBlendWeights[BoneIndex.GetInt()].BlendWeight, 0.f, 1.f);
+		MaxPoseWeights[PoseIndex] = FMath::Max(MaxPoseWeights[PoseIndex], BlendWeight);
+
 		if (BlendWeight < ZERO_ANIMWEIGHT_THRESH)
 		{
 			BlendAtom = BaseAtom;
@@ -929,15 +871,45 @@ void FAnimationRuntime::BlendMeshPosesPerBoneWeights(
 			OutPose[BoneIndex].SetRotation(LocalBlendQuat);
 		}
 	}
+
+	// time to blend curves
+	// the way we blend curve per bone
+	// is to find out max weight per that pose, and then apply that weight to the curve
+	{
+		TArray<const FBlendedCurve*> SourceCurves;
+		TArray<float> SourceWegihts;
+
+		SourceCurves.SetNumUninitialized(PoseNum+1);
+		SourceWegihts.SetNumUninitialized(PoseNum+1);
+
+		SourceCurves[0] = &BaseCurve;
+		SourceWegihts[0] = 1.f;
+
+		for(int32 Idx=0; Idx<PoseNum; ++Idx)
+		{
+			SourceCurves[Idx+1] = &BlendedCurves[Idx];
+			SourceWegihts[Idx+1] = MaxPoseWeights[Idx];
+		}
+
+		BlendCurves(SourceCurves, SourceWegihts, OutCurve, CurveBlendOption);
+	}
 }
 
 void FAnimationRuntime::BlendLocalPosesPerBoneWeights(
 	FCompactPose& BasePose,
 	const TArray<FCompactPose>& BlendPoses,
+	struct FBlendedCurve& BaseCurve,
+	const TArray<struct FBlendedCurve>& BlendedCurves,
 	const TArray<FPerBoneBlendWeight> & BoneBlendWeights,
-	/*out*/ FCompactPose& OutPose)
+	ECurveBlendOption::Type CurveBlendOption,
+	/*out*/ FCompactPose& OutPose, 
+	/*out*/ struct FBlendedCurve& OutCurve)
 {
 	check(BasePose.GetNumBones() == BoneBlendWeights.Num());
+	int32 PoseNum = BlendPoses.Num();
+
+	TArray<float> MaxPoseWeights;
+	MaxPoseWeights.AddZeroed(PoseNum);
 
 	for (FCompactPoseBoneIndex BoneIndex : BasePose.ForEachBoneIndex())
 	{
@@ -945,6 +917,7 @@ void FAnimationRuntime::BlendLocalPosesPerBoneWeights(
 		const FTransform& BaseAtom = BasePose[BoneIndex];
 
 		const float BlendWeight = FMath::Clamp(BoneBlendWeights[BoneIndex.GetInt()].BlendWeight, 0.f, 1.f);
+		MaxPoseWeights[PoseIndex] = FMath::Max(MaxPoseWeights[PoseIndex], BlendWeight);
 
 		if (BlendWeight < ZERO_ANIMWEIGHT_THRESH)
 		{
@@ -962,6 +935,28 @@ void FAnimationRuntime::BlendLocalPosesPerBoneWeights(
 			OutPose[BoneIndex] = BlendAtom;
 		}
 	}
+
+	// time to blend curves
+	// the way we blend curve per bone
+	// is to find out max weight per that pose, and then apply that weight to the curve
+	{
+		TArray<const FBlendedCurve*> SourceCurves;
+		TArray<float> SourceWegihts;
+
+		SourceCurves.SetNumUninitialized(PoseNum+1);
+		SourceWegihts.SetNumUninitialized(PoseNum+1);
+
+		SourceCurves[0] = &BaseCurve;
+		SourceWegihts[0] = 1.f;
+
+		for (int32 Idx=0; Idx<PoseNum; ++Idx)
+		{
+			SourceCurves[Idx+1] = &BlendedCurves[Idx];
+			SourceWegihts[Idx+1] = MaxPoseWeights[Idx];
+		}
+		
+		BlendCurves(SourceCurves, SourceWegihts, OutCurve, CurveBlendOption);
+	}
 }
 
 void FAnimationRuntime::UpdateDesiredBoneWeight(const TArray<FPerBoneBlendWeight>& SrcBoneBlendWeights, TArray<FPerBoneBlendWeight>& TargetBoneBlendWeights, const TArray<float>& BlendWeights)
@@ -984,7 +979,16 @@ void FAnimationRuntime::UpdateDesiredBoneWeight(const TArray<FPerBoneBlendWeight
 	}
 }
 
-void FAnimationRuntime::BlendPosesPerBoneFilter(struct FCompactPose& BasePose, const TArray<struct FCompactPose>& BlendPoses, struct FCompactPose& OutPose, TArray<FPerBoneBlendWeight>& BoneBlendWeights, bool bMeshSpaceRotationBlending)
+void FAnimationRuntime::BlendPosesPerBoneFilter(
+	struct FCompactPose& BasePose, 
+	const TArray<struct FCompactPose>& BlendPoses, 
+	struct FBlendedCurve& BaseCurve, 
+	const TArray<struct FBlendedCurve>& BlendedCurves, 
+	struct FCompactPose& OutPose, 
+	struct FBlendedCurve& OutCurve, 
+	TArray<FPerBoneBlendWeight>& BoneBlendWeights, 
+	bool bMeshSpaceRotationBlending, 
+	ECurveBlendOption::Type CurveBlendOption)
 {
 	ensure(OutPose.GetNumBones() == BasePose.GetNumBones());
 
@@ -992,11 +996,11 @@ void FAnimationRuntime::BlendPosesPerBoneFilter(struct FCompactPose& BasePose, c
 	{
 		if (bMeshSpaceRotationBlending)
 		{
-			BlendMeshPosesPerBoneWeights(BasePose, BlendPoses, BoneBlendWeights, OutPose);
+			BlendMeshPosesPerBoneWeights(BasePose, BlendPoses, BaseCurve, BlendedCurves, BoneBlendWeights, CurveBlendOption, OutPose, OutCurve);
 		}
 		else
 		{
-			BlendLocalPosesPerBoneWeights(BasePose, BlendPoses, BoneBlendWeights, OutPose);
+			BlendLocalPosesPerBoneWeights(BasePose, BlendPoses, BaseCurve, BlendedCurves, BoneBlendWeights, CurveBlendOption, OutPose, OutCurve);
 		}
 	}
 	else // if no blendpose, outpose = basepose
diff --git a/Engine/Source/Runtime/Engine/Private/Animation/BlendSpace.cpp b/Engine/Source/Runtime/Engine/Private/Animation/BlendSpace.cpp
--- a/Engine/Source/Runtime/Engine/Private/Animation/BlendSpace.cpp
+++ b/Engine/Source/Runtime/Engine/Private/Animation/BlendSpace.cpp
@@ -107,7 +107,7 @@ FVector2D UBlendSpace::CalculateThreshold() const
 
 bool UBlendSpace::IsValidAdditive()  const
 {
-	return IsValidAdditiveInternal(AAT_LocalSpaceBase);
+	return IsValidAdditiveInternal(AAT_LocalSpaceBase) || IsValidAdditiveInternal(AAT_RotationOffsetMeshSpace);
 }
 
 void UBlendSpace::SnapToBorder(FBlendSample& Sample) const
diff --git a/Engine/Source/Runtime/Engine/Private/Animation/BlendSpaceBase.cpp b/Engine/Source/Runtime/Engine/Private/Animation/BlendSpaceBase.cpp
--- a/Engine/Source/Runtime/Engine/Private/Animation/BlendSpaceBase.cpp
+++ b/Engine/Source/Runtime/Engine/Private/Animation/BlendSpaceBase.cpp
@@ -275,8 +275,6 @@ void UBlendSpaceBase::TickAssetPlayerInstance(const FAnimTickRecord& Instance, c
 								Context.RootMotionMovementParams.AccumulateWithBlend(Sample.Animation->ExtractRootMotion(PrevSampleDataTime, DeltaTimePosition, Instance.bLooping), SampleEntry.GetWeight());
 							}
 
-							// handle curves
-							Sample.Animation->EvaluateCurveData(InstanceOwner, CurrentSampleDataTime, Instance.EffectiveBlendWeight * SampleEntry.GetWeight());
 
 							UE_LOG(LogAnimation, Verbose, TEXT("%d. Blending animation(%s) with %f weight at time %0.2f"), I+1, *Sample.Animation->GetName(), SampleEntry.GetWeight(), CurrentSampleDataTime);
 						}
@@ -613,11 +611,19 @@ bool UBlendSpaceBase::ValidateSampleInput(FBlendSample & BlendSample, int32 Orig
 	// make sure we get same kinds of samples(additive or nonadditive)
 	if (SampleData.Num() > 0 && BlendSample.Animation)
 	{
-		if (IsValidAdditive() != (BlendSample.Animation->IsValidAdditive()))
+		bool bIsAdditive = IsValidAdditive();
+		if (bIsAdditive != (BlendSample.Animation->IsValidAdditive()))
 		{
 			UE_LOG(LogAnimation, Log, TEXT("Adding sample failed. Please add same kinds of sequence (additive/non-additive)."));
 			return false;
 		}
+
+		// make sure it's same additive if it is additive
+		if (bIsAdditive && !IsValidAdditiveInternal(BlendSample.Animation->AdditiveAnimType))
+		{
+			UE_LOG(LogAnimation, Log, TEXT("Adding sample failed. Please add same kinds of additive sequence (loca/mesh)."));
+			return false;
+		}
 	}
 
 	const USkeleton* MySkeleton = GetSkeleton();
@@ -658,6 +664,9 @@ void UBlendSpaceBase::ValidateSampleData()
 			continue;
 		}
 
+		// set rotation blend in mesh space
+		bRotationBlendInMeshSpace = IsValidAdditiveInternal(AAT_RotationOffsetMeshSpace);
+
 		// we need data to be snapped on the border
 		// otherwise, you have this grid area that doesn't have valid 
 		// sample points. Usually users will put it there
@@ -925,6 +934,86 @@ void UBlendSpaceBase::FillupGridElements(const TArray<FVector> & PointList, cons
 	}
 }
 
+void UBlendSpaceBase::GetAnimationPose(TArray<FBlendSampleData>& BlendSampleDataCache, /*out*/ FCompactPose& OutPose, /*out*/ FBlendedCurve& OutCurve)
+{
+	if(BlendSampleDataCache.Num() == 0)
+	{
+		OutPose.ResetToRefPose();
+		return;
+	}
+
+	const int32 NumPoses = BlendSampleDataCache.Num();
+
+	TArray<FCompactPose> ChildrenPoses;
+	ChildrenPoses.AddZeroed(NumPoses);
+
+	TArray<FBlendedCurve> ChildrenCurves;
+	ChildrenCurves.AddZeroed(NumPoses);
+
+	TArray<float>	ChildrenWeights;
+	ChildrenWeights.AddZeroed(NumPoses);
+
+	for(int32 ChildrenIdx=0; ChildrenIdx<ChildrenPoses.Num(); ++ChildrenIdx)
+	{
+		ChildrenPoses[ChildrenIdx].SetBoneContainer(&OutPose.GetBoneContainer());
+		ChildrenCurves[ChildrenIdx].InitFrom(OutCurve);
+	}
+
+	// get all child atoms we interested in
+	for(int32 I = 0; I < BlendSampleDataCache.Num(); ++I)
+	{
+		FCompactPose& Pose = ChildrenPoses[I];
+
+		if(SampleData.IsValidIndex(BlendSampleDataCache[I].SampleDataIndex))
+		{
+			const FBlendSample& Sample = SampleData[BlendSampleDataCache[I].SampleDataIndex];
+			ChildrenWeights[I] = BlendSampleDataCache[I].GetWeight();
+
+			if(Sample.Animation)
+			{
+				const float Time = FMath::Clamp<float>(BlendSampleDataCache[I].Time, 0.f, Sample.Animation->SequenceLength);
+
+				// first one always fills up the source one
+				Sample.Animation->GetAnimationPose(Pose, ChildrenCurves[I], FAnimExtractContext(Time, true));
+			}
+			else
+			{
+				Pose.ResetToRefPose();
+			}
+		}
+		else
+		{
+			Pose.ResetToRefPose();
+		}
+	}
+
+	if(PerBoneBlend.Num() > 0)
+	{
+		if(IsValidAdditive())
+		{
+			if(bRotationBlendInMeshSpace)
+			{
+				FAnimationRuntime::BlendPosesTogetherPerBoneInMeshSpace(ChildrenPoses, ChildrenCurves, this, BlendSampleDataCache, OutPose, OutCurve);
+			}
+			else
+			{
+				FAnimationRuntime::BlendPosesTogetherPerBone(ChildrenPoses, ChildrenCurves, this, BlendSampleDataCache, OutPose, OutCurve);
+			}
+		}
+		else
+		{
+			FAnimationRuntime::BlendPosesTogetherPerBone(ChildrenPoses, ChildrenCurves, this, BlendSampleDataCache, OutPose, OutCurve);
+		}
+	}
+	else
+	{
+		FAnimationRuntime::BlendPosesTogether(ChildrenPoses, ChildrenCurves, ChildrenWeights, OutPose, OutCurve);
+	}
+
+	// Once all the accumulation and blending has been done, normalize rotations.
+	OutPose.NormalizeRotations();
+}
+
 #if WITH_EDITOR
 bool UBlendSpaceBase::GetAllAnimationSequencesReferred(TArray<UAnimSequence*>& AnimationSequences)
 {
diff --git a/Engine/Source/Runtime/Engine/Private/Components/PoseableMeshComponent.cpp b/Engine/Source/Runtime/Engine/Private/Components/PoseableMeshComponent.cpp
--- a/Engine/Source/Runtime/Engine/Private/Components/PoseableMeshComponent.cpp
+++ b/Engine/Source/Runtime/Engine/Private/Components/PoseableMeshComponent.cpp
@@ -5,7 +5,7 @@
 =============================================================================*/
 
 #include "EnginePrivate.h"
-#include "AnimTree.h"
+#include "Animation/AnimStats.h"
 #include "Animation/AnimInstance.h"
 #include "Components/PoseableMeshComponent.h"
 
diff --git a/Engine/Source/Runtime/Engine/Private/Components/SkeletalMeshComponent.cpp b/Engine/Source/Runtime/Engine/Private/Components/SkeletalMeshComponent.cpp
--- a/Engine/Source/Runtime/Engine/Private/Components/SkeletalMeshComponent.cpp
+++ b/Engine/Source/Runtime/Engine/Private/Components/SkeletalMeshComponent.cpp
@@ -13,7 +13,7 @@
 #include "AnimationUtils.h"
 #include "AnimationRuntime.h"
 #include "PhysXASync.h"
-#include "AnimTree.h"
+#include "Animation/AnimStats.h"
 #include "Animation/AnimNodeBase.h"
 #include "Animation/VertexAnim/VertexAnimation.h"
 #include "Particles/ParticleSystemComponent.h"
@@ -934,7 +934,7 @@ void USkeletalMeshComponent::RecalcRequiredBones(int32 LODIndex)
 	CachedSpaceBases.Empty();
 }
 
-void USkeletalMeshComponent::EvaluateAnimation(const USkeletalMesh* InSkeletalMesh, UAnimInstance* InAnimInstance, TArray<FTransform>& OutLocalAtoms, TArray<FActiveVertexAnim>& OutVertexAnims, FVector& OutRootBoneTranslation) const
+void USkeletalMeshComponent::EvaluateAnimation(const USkeletalMesh* InSkeletalMesh, UAnimInstance* InAnimInstance, TArray<FTransform>& OutLocalAtoms, TArray<FActiveVertexAnim>& OutVertexAnims, FVector& OutRootBoneTranslation, FBlendedCurve& OutCurve) const
 {
 	SCOPE_CYCLE_COUNTER(STAT_AnimBlendTime);
 
@@ -957,13 +957,14 @@ void USkeletalMeshComponent::EvaluateAnimation(const USkeletalMesh* InSkeletalMe
 			
 			// Run the anim blueprint
 			InAnimInstance->EvaluateAnimation(EvaluationContext);
-
+			// Move the curves
+			OutCurve.MoveFrom(EvaluationContext.Curve);
+			
 			// can we avoid that copy?
 			if( EvaluationContext.Pose.GetNumBones() > 0 )
 			{
 				// Make sure rotations are normalized to account for accumulation of errors.
 				EvaluationContext.Pose.NormalizeRotations();
-
 				for (const FCompactPoseBoneIndex BoneIndex : EvaluationContext.Pose.ForEachBoneIndex())
 				{
 					FMeshPoseBoneIndex MeshPoseBoneIndex = EvaluationContext.Pose.GetBoneContainer().MakeMeshPoseIndex(BoneIndex);
@@ -1015,7 +1016,7 @@ void USkeletalMeshComponent::UpdateSlaveComponent()
 	Super::UpdateSlaveComponent();
 }
 
-void USkeletalMeshComponent::PerformAnimationEvaluation(const USkeletalMesh* InSkeletalMesh, UAnimInstance* InAnimInstance, TArray<FTransform>& OutSpaceBases, TArray<FTransform>& OutLocalAtoms, TArray<FActiveVertexAnim>& OutVertexAnims, FVector& OutRootBoneTranslation) const
+void USkeletalMeshComponent::PerformAnimationEvaluation(const USkeletalMesh* InSkeletalMesh, UAnimInstance* InAnimInstance, TArray<FTransform>& OutSpaceBases, TArray<FTransform>& OutLocalAtoms, TArray<FActiveVertexAnim>& OutVertexAnims, FVector& OutRootBoneTranslation, FBlendedCurve& OutCurve) const
 {
 	SCOPE_CYCLE_COUNTER(STAT_PerformAnimEvaluation);
 	// Can't do anything without a SkeletalMesh
@@ -1026,7 +1027,7 @@ void USkeletalMeshComponent::PerformAnimationEvaluation(const USkeletalMesh* InS
 	}
 
 	// evaluate pure animations, and fill up LocalAtoms
-	EvaluateAnimation(InSkeletalMesh, InAnimInstance, OutLocalAtoms, OutVertexAnims, OutRootBoneTranslation);
+	EvaluateAnimation(InSkeletalMesh, InAnimInstance, OutLocalAtoms, OutVertexAnims, OutRootBoneTranslation, OutCurve);
 	// Fill SpaceBases from LocalAtoms
 	FillSpaceBases(InSkeletalMesh, OutLocalAtoms, OutSpaceBases);
 }
@@ -1035,6 +1036,21 @@ static TAutoConsoleVariable<int32> CVarHiPriSkinnedMeshesTicks(
 	TEXT("tick.HiPriSkinnedMeshes"),1,
 	TEXT("If > 0, then schedule the skinned component ticks in a tick group before other ticks."));
 
+int32 GetCurveNumber(USkeleton* Skeleton)
+{
+	check (Skeleton);
+	// get all curve list
+	if(const FSmartNameMapping* Mapping = Skeleton->SmartNames.GetContainer(USkeleton::AnimCurveMappingName))
+	{
+		TArray<FSmartNameMapping::UID> UIDList;
+		Mapping->FillUidArray(UIDList);
+
+		return UIDList.Num();
+	}
+
+	return 0;
+}
+
 void USkeletalMeshComponent::RefreshBoneTransforms(FActorComponentTickFunction* TickFunction)
 {
 	SCOPE_CYCLE_COUNTER(STAT_RefreshBoneTransforms);
@@ -1066,6 +1082,10 @@ void USkeletalMeshComponent::RefreshBoneTransforms(FActorComponentTickFunction*
 									  || (LocalAtoms.Num() != CachedLocalAtoms.Num())
 									 || (GetNumSpaceBases() != CachedSpaceBases.Num()));
 
+
+	const bool bInvalidCachedCurve = bDoEvaluationRateOptimization && 
+									CachedCurve.Num() != GetCurveNumber(SkeletalMesh->Skeleton);
+
 	const bool bShouldDoEvaluation = !bDoEvaluationRateOptimization || bInvalidCachedBones || !AnimUpdateRateParams->ShouldSkipEvaluation();
 
 	const bool bDoPAE = !!CVarUseParallelAnimationEvaluation.GetValueOnGameThread() && FApp::ShouldUseThreadingForPerformance();
@@ -1091,12 +1111,13 @@ void USkeletalMeshComponent::RefreshBoneTransforms(FActorComponentTickFunction*
 	
 	AnimEvaluationContext.bDoInterpolation = bDoEvaluationRateOptimization && !bInvalidCachedBones && AnimUpdateRateParams->ShouldInterpolateSkippedFrames();
 	AnimEvaluationContext.bDuplicateToCacheBones = bInvalidCachedBones || (bDoEvaluationRateOptimization && AnimEvaluationContext.bDoEvaluation && !AnimEvaluationContext.bDoInterpolation);
-
+	AnimEvaluationContext.bDuplicateToCacheCurve = bInvalidCachedCurve || (bDoEvaluationRateOptimization && AnimEvaluationContext.bDoEvaluation && !AnimEvaluationContext.bDoInterpolation);
 	if (!bDoEvaluationRateOptimization)
 	{
 		//If we aren't optimizing clear the cached local atoms
 		CachedLocalAtoms.Empty();
 		CachedSpaceBases.Empty();
+		CachedCurve.Empty();
 	}
 
 	if (bDoParallelEvaluation)
@@ -1137,17 +1158,18 @@ void USkeletalMeshComponent::RefreshBoneTransforms(FActorComponentTickFunction*
 		{
 			if (AnimEvaluationContext.bDoInterpolation)
 			{
-				PerformAnimationEvaluation(SkeletalMesh, AnimScriptInstance, CachedSpaceBases, CachedLocalAtoms, ActiveVertexAnims, RootBoneTranslation);
+				PerformAnimationEvaluation(SkeletalMesh, AnimScriptInstance, CachedSpaceBases, CachedLocalAtoms, ActiveVertexAnims, RootBoneTranslation, CachedCurve);
 			}
 			else
 			{
-				PerformAnimationEvaluation(SkeletalMesh, AnimScriptInstance, GetEditableSpaceBases(), LocalAtoms, ActiveVertexAnims, RootBoneTranslation);
+				PerformAnimationEvaluation(SkeletalMesh, AnimScriptInstance, GetEditableSpaceBases(), LocalAtoms, ActiveVertexAnims, RootBoneTranslation, AnimEvaluationContext.Curve);
 			}
 		}
 		else if (!AnimEvaluationContext.bDoInterpolation)
 		{
 			LocalAtoms = CachedLocalAtoms;
 			GetEditableSpaceBases() = CachedSpaceBases;
+			AnimEvaluationContext.Curve = CachedCurve;
 		}
 
 		PostAnimEvaluation(AnimEvaluationContext);
@@ -1162,11 +1184,7 @@ void USkeletalMeshComponent::RefreshBoneTransforms(FActorComponentTickFunction*
 
 void USkeletalMeshComponent::PostAnimEvaluation(FAnimationEvaluationContext& EvaluationContext)
 {
-	if (AnimScriptInstance)
-	{
-		AnimScriptInstance->PostAnimEvaluation();
-	}
-
+	FBlendedCurve EvaluatedCurve = EvaluationContext.Curve;
 	AnimEvaluationContext.Clear();
 
 	SCOPE_CYCLE_COUNTER(STAT_PostAnimEvaluation);
@@ -1176,6 +1194,11 @@ void USkeletalMeshComponent::PostAnimEvaluation(FAnimationEvaluationContext& Eva
 		CachedLocalAtoms = LocalAtoms;
 	}
 
+	if (EvaluationContext.bDuplicateToCacheCurve)
+	{
+		CachedCurve = EvaluatedCurve;
+	}
+
 	if (EvaluationContext.bDoInterpolation)
 	{
 		SCOPE_CYCLE_COUNTER(STAT_InterpolateSkippedFrames);
@@ -1189,6 +1212,9 @@ void USkeletalMeshComponent::PostAnimEvaluation(FAnimationEvaluationContext& Eva
 			GetEditableSpaceBases() = GetSpaceBases();
 		}
 		FAnimationRuntime::LerpBoneTransforms(GetEditableSpaceBases(), CachedSpaceBases, Alpha, RequiredBones);
+
+		// interpolate curve
+		EvaluatedCurve.Blend(EvaluatedCurve, CachedCurve, Alpha);
 	}
 
 	bNeedToFlipSpaceBaseBuffers = true;
@@ -1219,6 +1245,13 @@ void USkeletalMeshComponent::PostAnimEvaluation(FAnimationEvaluationContext& Eva
 		UpdateOverlaps();
 	}
 
+	if(AnimScriptInstance)
+	{
+		// curve update happens first
+		AnimScriptInstance->UpdateCurves(EvaluatedCurve);
+		AnimScriptInstance->PostAnimEvaluation();
+	}
+
 	MarkRenderDynamicDataDirty();
 }
 
diff --git a/Engine/Source/Runtime/Engine/Private/Components/SkinnedMeshComponent.cpp b/Engine/Source/Runtime/Engine/Private/Components/SkinnedMeshComponent.cpp
--- a/Engine/Source/Runtime/Engine/Private/Components/SkinnedMeshComponent.cpp
+++ b/Engine/Source/Runtime/Engine/Private/Components/SkinnedMeshComponent.cpp
@@ -10,7 +10,7 @@
 #include "SkeletalRenderCPUSkin.h"
 #include "SkeletalRenderGPUSkin.h"
 #include "AnimationUtils.h"
-#include "AnimTree.h"
+#include "Animation/AnimStats.h"
 #include "Animation/VertexAnim/MorphTarget.h"
 #include "ComponentReregisterContext.h"
 #include "Engine/SkeletalMeshSocket.h"
diff --git a/Engine/Source/Runtime/Engine/Private/PhysicsEngine/BodySetup.cpp b/Engine/Source/Runtime/Engine/Private/PhysicsEngine/BodySetup.cpp
--- a/Engine/Source/Runtime/Engine/Private/PhysicsEngine/BodySetup.cpp
+++ b/Engine/Source/Runtime/Engine/Private/PhysicsEngine/BodySetup.cpp
@@ -8,7 +8,7 @@
 #include "PhysicsPublic.h"
 #include "PhysicsEngine/BodySetup.h"
 #include "TargetPlatform.h"
-#include "AnimTree.h"
+#include "Animation/AnimStats.h"
 
 #if WITH_PHYSX
 	#include "PhysXSupport.h"
diff --git a/Engine/Source/Runtime/Engine/Private/PhysicsEngine/PhysAnim.cpp b/Engine/Source/Runtime/Engine/Private/PhysicsEngine/PhysAnim.cpp
--- a/Engine/Source/Runtime/Engine/Private/PhysicsEngine/PhysAnim.cpp
+++ b/Engine/Source/Runtime/Engine/Private/PhysicsEngine/PhysAnim.cpp
@@ -5,7 +5,7 @@
 =============================================================================*/ 
 
 #include "EnginePrivate.h"
-#include "AnimTree.h"
+#include "Animation/AnimStats.h"
 #include "SkeletalRenderPublic.h"
 #include "Components/LineBatchComponent.h"
 #if WITH_PHYSX
diff --git a/Engine/Source/Runtime/Engine/Public/AnimTree.h b/Engine/Source/Runtime/Engine/Public/AnimTree.h
deleted file mode 100644
diff --git a/Engine/Source/Runtime/Engine/Public/Animation/AnimCurveTypes.h b/Engine/Source/Runtime/Engine/Public/Animation/AnimCurveTypes.h
new file mode 100644
--- /dev/null
+++ b/Engine/Source/Runtime/Engine/Public/Animation/AnimCurveTypes.h
@@ -0,0 +1,467 @@
+// Copyright 1998-2015 Epic Games, Inc. All Rights Reserved.
+
+#pragma once
+
+#include "Curves/CurveBase.h"
+#include "AnimTypes.h"
+#include "AnimCurveTypes.generated.h"
+
+/** native enum for curve types **/
+enum EAnimCurveFlags
+{
+	// Used as morph target curve
+	ACF_DrivesMorphTarget	= 0x00000001,
+	// Used as triggering event
+	ACF_TriggerEvent		= 0x00000002,
+	// Is editable in Sequence Editor
+	ACF_Editable			= 0x00000004,
+	// Used as a material curve
+	ACF_DrivesMaterial		= 0x00000008,
+	// Is a metadata 'curve'
+	ACF_Metadata			= 0x00000010,
+	// motifies bone track
+	ACF_DriveTrack			= 0x00000020,
+	// disabled, right now it's used by track
+	ACF_Disabled			= 0x00000040,
+
+	// default flag when created
+	ACF_DefaultCurve		= ACF_TriggerEvent | ACF_Editable,
+	// curves created from Morph Target
+	ACF_MorphTargetCurve	= ACF_DrivesMorphTarget
+};
+
+/**
+ * Float curve data for one track
+ */
+
+USTRUCT()
+struct FAnimCurveBase
+{
+	GENERATED_USTRUCT_BODY()
+
+	// Last observed name of the curve. We store this so we can recover from situations that
+	// mean the skeleton doesn't have a mapped name for our UID (such as a user saving the an
+	// animation but not the skeleton).
+	UPROPERTY()
+	FName		LastObservedName;
+
+	/* For smart naming - management purpose - i.e. rename/delete */
+	FSmartNameMapping::UID CurveUid;
+
+private:
+	/** Curve Type Flags */
+	UPROPERTY()
+	int32		CurveTypeFlags;
+
+public:
+	FAnimCurveBase(){}
+
+	FAnimCurveBase(FName InName, int32 InCurveTypeFlags)
+		: LastObservedName(InName)
+		, CurveTypeFlags(InCurveTypeFlags)
+	{	
+	}
+
+	FAnimCurveBase(USkeleton::AnimCurveUID Uid, int32 InCurveTypeFlags)
+		: CurveUid(Uid)
+		, CurveTypeFlags(InCurveTypeFlags)
+	{}
+
+	// To be able to use typedef'd types we need to serialize manually
+	virtual void Serialize(FArchive& Ar)
+	{
+		if(Ar.UE4Ver() >= VER_UE4_SKELETON_ADD_SMARTNAMES)
+		{
+			Ar << CurveUid;
+		}
+	}
+
+	/**
+	 * Set InFlag to bValue
+	 */
+	ENGINE_API void SetCurveTypeFlag(EAnimCurveFlags InFlag, bool bValue);
+
+	/**
+	 * Toggle the value of the specified flag
+	 */
+	ENGINE_API void ToggleCurveTypeFlag(EAnimCurveFlags InFlag);
+
+	/**
+	 * Return true if InFlag is set, false otherwise 
+	 */
+	ENGINE_API bool GetCurveTypeFlag(EAnimCurveFlags InFlag) const;
+
+	/**
+	 * Set CurveTypeFlags to NewCurveTypeFlags
+	 * This just overwrites CurveTypeFlags
+	 */
+	void SetCurveTypeFlags(int32 NewCurveTypeFlags);
+	/** 
+	 * returns CurveTypeFlags
+	 */
+	int32 GetCurveTypeFlags() const;
+};
+
+USTRUCT()
+struct FFloatCurve : public FAnimCurveBase
+{
+	GENERATED_USTRUCT_BODY()
+
+	/** Curve data for float. */
+	UPROPERTY()
+	FRichCurve	FloatCurve;
+
+	FFloatCurve(){}
+
+	FFloatCurve(FName InName, int32 InCurveTypeFlags)
+		: FAnimCurveBase(InName, InCurveTypeFlags)
+	{
+	}
+
+	FFloatCurve(USkeleton::AnimCurveUID Uid, int32 InCurveTypeFlags)
+		: FAnimCurveBase(Uid, InCurveTypeFlags)
+	{}
+
+	virtual void Serialize(FArchive& Ar) override
+	{
+		FAnimCurveBase::Serialize(Ar);
+	}
+
+	// we don't want to have = operator. This only copies curves, but leaving naming and everything else intact. 
+	void CopyCurve(FFloatCurve& SourceCurve);
+	float Evaluate(float CurrentTime) const;
+	void UpdateOrAddKey(float NewKey, float CurrentTime);
+	void Resize(float NewLength);
+};
+
+USTRUCT()
+struct FVectorCurve : public FAnimCurveBase
+{
+	GENERATED_USTRUCT_BODY()
+
+	enum EIndex
+	{
+		X = 0, 
+		Y, 
+		Z, 
+		Max
+	};
+
+	/** Curve data for float. */
+	UPROPERTY()
+	FRichCurve	FloatCurves[3];
+
+	FVectorCurve(){}
+
+	FVectorCurve(FName InName, int32 InCurveTypeFlags)
+		: FAnimCurveBase(InName, InCurveTypeFlags)
+	{
+	}
+
+	FVectorCurve(USkeleton::AnimCurveUID Uid, int32 InCurveTypeFlags)
+		: FAnimCurveBase(Uid, InCurveTypeFlags)
+	{}
+
+	virtual void Serialize(FArchive& Ar) override
+	{
+		FAnimCurveBase::Serialize(Ar);
+	}
+
+	// we don't want to have = operator. This only copies curves, but leaving naming and everything else intact. 
+	void CopyCurve(FVectorCurve& SourceCurve);
+	FVector Evaluate(float CurrentTime, float BlendWeight) const;
+	void UpdateOrAddKey(const FVector& NewKey, float CurrentTime);
+	bool DoesContainKey() const { return (FloatCurves[0].GetNumKeys() > 0 || FloatCurves[1].GetNumKeys() > 0 || FloatCurves[2].GetNumKeys() > 0);}
+	void Resize(float NewLength);
+};
+
+USTRUCT()
+struct FTransformCurve: public FAnimCurveBase
+{
+	GENERATED_USTRUCT_BODY()
+
+	/** Curve data for each transform. */
+	UPROPERTY()
+	FVectorCurve	TranslationCurve;
+
+	/** Rotation curve - right now we use euler because quat also doesn't provide linear interpolation - curve editor can't handle quat interpolation
+	 * If you hit gimbal lock, you should add extra key to fix it. This will cause gimbal lock. 
+	 * @TODO: Eventually we'll need FRotationCurve that would contain rotation curve - that will interpolate as slerp or as quaternion 
+	 */
+	UPROPERTY()
+	FVectorCurve	RotationCurve;
+
+	UPROPERTY()
+	FVectorCurve	ScaleCurve;
+
+	FTransformCurve(){}
+
+	FTransformCurve(FName InName, int32 InCurveTypeFlags)
+		: FAnimCurveBase(InName, InCurveTypeFlags)
+	{
+	}
+
+	FTransformCurve(USkeleton::AnimCurveUID Uid, int32 InCurveTypeFlags)
+		: FAnimCurveBase(Uid, InCurveTypeFlags)
+	{}
+
+	virtual void Serialize(FArchive& Ar) override
+	{
+		FAnimCurveBase::Serialize(Ar);
+	}
+
+	// we don't want to have = operator. This only copies curves, but leaving naming and everything else intact. 
+	void CopyCurve(FTransformCurve& SourceCurve);
+	FTransform Evaluate(float CurrentTime, float BlendWeight) const;
+	ENGINE_API void UpdateOrAddKey(const FTransform& NewKey, float CurrentTime);
+	void Resize(float NewLength);
+};
+
+/**
+* This is array of curves that run when collecting curves natively 
+*/
+struct FCurveElement
+{
+	/** Curve Value */
+	float					Value;
+	/** Curve Flags - this does OR ops when it gets blended*/
+	int32					Flags;
+
+	FCurveElement(float InValue, int32	InFlags)
+		:  Value(InValue), Flags(InFlags)
+	{}
+};
+
+/**
+ * This struct is used to create curve snap shot of current time when extracted
+ */
+struct ENGINE_API FBlendedCurve
+{
+	/**
+	* List of curve elements for this pose
+	*/
+	TArray<FCurveElement> Elements;
+	TArray<FSmartNameMapping::UID> UIDList;
+
+	/**
+	 * constructor
+	 */
+	FBlendedCurve()
+		: bInitialized(false)
+	{
+	}
+
+	/**
+	 * constructor
+	 */
+	FBlendedCurve(const class UAnimInstance* AnimInstance);
+
+	/**
+	 * constructor
+	 */
+	FBlendedCurve(const class USkeleton* Skeleton)
+	{
+		check (Skeleton);
+		InitFrom(Skeleton);
+	}
+
+	/** Initialize Curve Data from following data */
+	void InitFrom(const class USkeleton* Skeleton);
+	void InitFrom(const FBlendedCurve& InCurveToInitFrom)
+	{
+		UIDList = InCurveToInitFrom.UIDList;
+		Elements.Empty(UIDList.Num());
+		Elements.AddZeroed(UIDList.Num());
+		bInitialized = true;
+	}
+
+	/** Set value of InUID to InValue */
+	void Set(USkeleton::AnimCurveUID InUid, float InValue, int32 InFlags);
+
+	/**
+	 * Blend (A, B) using Alpha, same as Lerp
+	 */
+	void Blend(const FBlendedCurve& A, const FBlendedCurve& B, float Alpha);
+	/**
+	 * Convert current curves to Additive (this - BaseCurve) if same found
+	 */
+	void ConvertToAdditive(const FBlendedCurve& BaseCurve);
+	/**
+	 * Accumulate the input curve with input Weight
+	 */
+	void Accumulate(const FBlendedCurve& AdditiveCurve, float Weight);
+
+	/**
+	 * This doesn't blend but combine MAX(current weight, curvetocombine weight)
+	 */
+	void Combine(const FBlendedCurve& CurveToCombine);
+
+	/**
+	 * Override with inupt curve * weight
+	 */
+	void Override(const FBlendedCurve& CurveToOverrideFrom, float Weight);
+	/**
+	 * Override with inupt curve 
+	 */
+	void Override(const FBlendedCurve& CurveToOverrideFrom);
+
+	/** Return number of elements */
+	int32 Num() const { return Elements.Num(); }
+
+	/** CopyFrom as expected. */
+	void CopyFrom(const FBlendedCurve& CurveToCopyFrom);
+	/** CopyFrom/MoveFrom as expected. Once moved, it is invalid */
+	void MoveFrom(FBlendedCurve& CurveToMoveFrom);
+
+	/** Empty */
+	void Empty()
+	{
+		UIDList.Empty();
+		Elements.Empty();
+	}
+private:
+	/**  Whether initialized or not */
+	bool bInitialized;
+	/** Empty and allocate Count number */
+	void Reset(int32 Count);
+};
+/**
+ * Raw Curve data for serialization
+ */
+USTRUCT()
+struct FRawCurveTracks
+{
+	GENERATED_USTRUCT_BODY()
+
+	enum ESupportedCurveType
+	{
+		FloatType,
+		VectorType,
+		TransformType,
+		Max, 
+	};
+
+	UPROPERTY()
+	TArray<FFloatCurve>		FloatCurves;
+
+#if WITH_EDITORONLY_DATA
+	/**
+	 * @note : Currently VectorCurves are not evaluated or used for anything else but transient data for modifying bone track
+	 *			Note that it doesn't have UPROPERTY tag yet. In the future, we'd like this to be serialized, but not for now
+	 **/
+	UPROPERTY(transient)
+	TArray<FVectorCurve>	VectorCurves;
+
+	/**
+	 * @note : TransformCurves are used to edit additive animation in editor. 
+	 **/
+	UPROPERTY()
+	TArray<FTransformCurve>	TransformCurves;
+#endif // #if WITH_EDITOR_DATA
+
+	/**
+	 * Evaluate curve data at the time CurrentTime, and add to Instance. It only evaluates Float Curve for now
+	 *
+	 * return true if curve exists, false otherwise
+	 */
+	void EvaluateCurveData( FBlendedCurve& Curves, float CurrentTime ) const;
+
+#if WITH_EDITOR
+	/**
+	 *	Evaluate transform curves 
+	 */
+	ENGINE_API void EvaluateTransformCurveData(USkeleton * Skeleton, TMap<FName, FTransform>&OutCurves, float CurrentTime, float BlendWeight) const;
+#endif // WITH_EDITOR
+	/**
+	 * Find curve data based on the curve UID
+	 */
+	ENGINE_API FAnimCurveBase * GetCurveData(USkeleton::AnimCurveUID Uid, ESupportedCurveType SupportedCurveType = FloatType);
+	/**
+	 * Add new curve from the provided UID and return true if success
+	 * bVectorInterpCurve == true, then it will create FVectorCuve, otherwise, FFloatCurve
+	 */
+	ENGINE_API bool AddCurveData(USkeleton::AnimCurveUID Uid, int32 CurveFlags = ACF_DefaultCurve, ESupportedCurveType SupportedCurveType = FloatType);
+
+	/**
+	 * Delete curve data 
+	 */
+	ENGINE_API bool DeleteCurveData(USkeleton::AnimCurveUID Uid, ESupportedCurveType SupportedCurveType = FloatType);
+	/**
+	 * Duplicate curve data
+	 * 
+	 */
+	ENGINE_API bool DuplicateCurveData(USkeleton::AnimCurveUID ToCopyUid, USkeleton::AnimCurveUID NewUid, ESupportedCurveType SupportedCurveType = FloatType);
+
+	/**
+	 * Updates the LastObservedName field of the curves from the provided name container
+	 */
+	ENGINE_API void UpdateLastObservedNames(FSmartNameMapping* NameMapping, ESupportedCurveType SupportedCurveType = FloatType);
+
+	/** 
+	 * Serialize
+	 */
+	void Serialize(FArchive& Ar);
+
+	/*
+	 * resize curve length. If longer, it doesn't do any. If shorter, remove previous keys and add new key to the end of the frame. 
+	 */
+	void Resize(float TotalLength);
+	/** 
+	 * Clear all keys
+	 */
+	void Empty()
+	{
+		FloatCurves.Empty();
+#if WITH_EDITORONLY_DATA
+		VectorCurves.Empty();
+		TransformCurves.Empty();
+#endif
+	}
+
+	void SortFloatCurvesByUID()
+	{
+		struct FCurveSortByUid
+		{
+			FORCEINLINE bool operator()(const FFloatCurve& A, const FFloatCurve& B) const
+			{
+				return (A.CurveUid < B.CurveUid);
+			}
+		};
+
+		FloatCurves.Sort(FCurveSortByUid());
+	}
+private:
+	/** 
+	 * Adding vector curve support - this is all transient data for now. This does not save and all these data will be baked into RawAnimationData
+	 */
+
+	/**
+	 * Find curve data based on the curve UID
+	 */
+	template <typename DataType>
+	DataType * GetCurveDataImpl(TArray<DataType>& Curves, USkeleton::AnimCurveUID Uid);
+
+	/**
+	 * Add new curve from the provided UID and return true if success
+	 * bVectorInterpCurve == true, then it will create FVectorCuve, otherwise, FFloatCurve
+	 */
+	template <typename DataType>
+	bool AddCurveDataImpl(TArray<DataType>& Curves, USkeleton::AnimCurveUID Uid, int32 CurveFlags);
+	/**
+	 * Delete curve data 
+	 */
+	template <typename DataType>
+	bool DeleteCurveDataImpl(TArray<DataType>& Curves, USkeleton::AnimCurveUID Uid);
+	/**
+	 * Duplicate curve data
+	 * 
+	 */
+	template <typename DataType>
+	bool DuplicateCurveDataImpl(TArray<DataType>& Curves, USkeleton::AnimCurveUID ToCopyUid, USkeleton::AnimCurveUID NewUid);
+
+	/**
+	 * Updates the LastObservedName field of the curves from the provided name container
+	 */
+	template <typename DataType>
+	void UpdateLastObservedNamesImpl(TArray<DataType>& Curves, FSmartNameMapping* NameMapping);
+};
\ No newline at end of file
diff --git a/Engine/Source/Runtime/Engine/Public/Animation/AnimStats.h b/Engine/Source/Runtime/Engine/Public/Animation/AnimStats.h
new file mode 100644
--- /dev/null
+++ b/Engine/Source/Runtime/Engine/Public/Animation/AnimStats.h
@@ -0,0 +1,41 @@
+// Copyright 1998-2015 Epic Games, Inc. All Rights Reserved.
+
+/*=============================================================================
+	AnimStats.h:	Animation stats
+=============================================================================*/
+
+#pragma once
+
+/** Skeletal stats */
+DECLARE_CYCLE_STAT_EXTERN(TEXT("SkinnedMeshComp Tick"), STAT_SkinnedMeshCompTick, STATGROUP_Anim, );
+DECLARE_CYCLE_STAT_EXTERN(TEXT("TickUpdateRate"), STAT_TickUpdateRate, STATGROUP_Anim, );
+DECLARE_CYCLE_STAT_EXTERN(TEXT("Anim Tick Time"), STAT_AnimTickTime, STATGROUP_Anim, );
+DECLARE_CYCLE_STAT_EXTERN(TEXT("StateMachine Update"), STAT_AnimStateMachineUpdate, STATGROUP_Anim, );
+DECLARE_CYCLE_STAT_EXTERN(TEXT("StateMachine Find Transition"), STAT_AnimStateMachineFindTransition, STATGROUP_Anim, );
+DECLARE_CYCLE_STAT_EXTERN(TEXT("RefreshBoneTransforms"), STAT_RefreshBoneTransforms, STATGROUP_Anim, );
+DECLARE_CYCLE_STAT_EXTERN(TEXT("Evaluate Anim"), STAT_AnimBlendTime, STATGROUP_Anim, );
+DECLARE_CYCLE_STAT_EXTERN(TEXT("Perform Anim Evaluation"), STAT_PerformAnimEvaluation, STATGROUP_Anim, );
+DECLARE_CYCLE_STAT_EXTERN(TEXT("Post Anim Evaluation"), STAT_PostAnimEvaluation, STATGROUP_Anim, );
+DECLARE_CYCLE_STAT_EXTERN(TEXT("GraphTime"), STAT_AnimGraphEvaluate, STATGROUP_Anim, );
+DECLARE_CYCLE_STAT_EXTERN(TEXT("Eval Pose"), STAT_AnimNativeEvaluatePoses, STATGROUP_Anim, );
+DECLARE_CYCLE_STAT_EXTERN(TEXT("Trigger Notifies"), STAT_AnimTriggerAnimNotifies, STATGROUP_Anim, );
+DECLARE_CYCLE_STAT_EXTERN(TEXT("Anim Decompression"), STAT_GetAnimationPose, STATGROUP_Anim, );
+DECLARE_CYCLE_STAT_EXTERN(TEXT("Blend Poses"), STAT_AnimNativeBlendPoses, STATGROUP_Anim, );
+DECLARE_CYCLE_STAT_EXTERN(TEXT("Copy Pose"), STAT_AnimNativeCopyPoses, STATGROUP_Anim, );
+DECLARE_CYCLE_STAT_EXTERN(TEXT("StateMachine Eval"), STAT_AnimStateMachineEvaluate, STATGROUP_Anim, );
+DECLARE_CYCLE_STAT_EXTERN(TEXT("FillSpaceBases"), STAT_SkelComposeTime, STATGROUP_Anim, );
+DECLARE_CYCLE_STAT_EXTERN(TEXT("InterpolateSkippedFrames"), STAT_InterpolateSkippedFrames, STATGROUP_Anim, );
+DECLARE_CYCLE_STAT_EXTERN(TEXT("UpdateKinematicBonesToAnim"), STAT_UpdateRBBones, STATGROUP_Anim, );
+DECLARE_CYCLE_STAT_EXTERN(TEXT("UpdateRBJointsMotors"), STAT_UpdateRBJoints, STATGROUP_Anim, );
+DECLARE_CYCLE_STAT_EXTERN(TEXT("UpdateLocalToWorldAndOverlaps"), STAT_UpdateLocalToWorldAndOverlaps, STATGROUP_Anim, );
+DECLARE_CYCLE_STAT_EXTERN(TEXT("SkelComp UpdateTransform"), STAT_SkelCompUpdateTransform, STATGROUP_Anim, );
+DECLARE_CYCLE_STAT_EXTERN(TEXT("MeshObject Update"), STAT_MeshObjectUpdate, STATGROUP_Anim, );
+DECLARE_CYCLE_STAT_EXTERN(TEXT("Update SkelMesh Bounds"), STAT_UpdateSkelMeshBounds, STATGROUP_Anim, );
+DECLARE_CYCLE_STAT_EXTERN(TEXT("BlendInPhysics"), STAT_BlendInPhysics, STATGROUP_Anim, );
+DECLARE_CYCLE_STAT_EXTERN(TEXT("SkinPerPolyVertices);"), STAT_SkinPerPolyVertices, STATGROUP_Anim, );
+DECLARE_CYCLE_STAT_EXTERN(TEXT("UpdateTriMeshVertices"), STAT_UpdateTriMeshVertices, STATGROUP_Anim, );
+DECLARE_CYCLE_STAT_EXTERN(TEXT("AnimGameThreadTime"), STAT_AnimGameThreadTime, STATGROUP_Anim, );
+DECLARE_CYCLE_STAT_EXTERN(TEXT("BlueprintUpdateAnimation"), STAT_BlueprintUpdateAnimation, STATGROUP_Anim, );
+DECLARE_CYCLE_STAT_EXTERN(TEXT("NativeUpdateAnimation"), STAT_NativeUpdateAnimation, STATGROUP_Anim, );
+
+
diff --git a/Engine/Source/Runtime/Engine/Public/Animation/AnimTypes.h b/Engine/Source/Runtime/Engine/Public/Animation/AnimTypes.h
new file mode 100644
--- /dev/null
+++ b/Engine/Source/Runtime/Engine/Public/Animation/AnimTypes.h
@@ -0,0 +1,256 @@
+// Copyright 1998-2015 Epic Games, Inc. All Rights Reserved.
+
+#pragma once
+
+#include "Animation/AnimLinkableElement.h"
+#include "AnimTypes.generated.h"
+
+namespace EAnimEventTriggerOffsets
+{
+	enum Type
+	{
+		OffsetBefore,
+		OffsetAfter,
+		NoOffset
+	};
+}
+
+ENGINE_API float GetTriggerTimeOffsetForType(EAnimEventTriggerOffsets::Type OffsetType);
+
+/** Ticking method for AnimNotifies in AnimMontages. */
+UENUM()
+namespace EMontageNotifyTickType
+{
+	enum Type
+	{
+		/** Queue notifies, and trigger them at the end of the evaluation phase (faster). Not suitable for changing sections or montage position. */
+		Queued,
+		/** Trigger notifies as they are encountered (Slower). Suitable for changing sections or montage position. */
+		BranchingPoint,
+	};
+}
+
+/** Filtering method for deciding whether to trigger a notify. */
+UENUM()
+namespace ENotifyFilterType
+{
+	enum Type
+	{
+		/** No filtering. */
+		NoFiltering,
+
+		/** Filter By Skeletal Mesh LOD. */
+		LOD,
+	};
+}
+
+/**
+ * Triggers an animation notify.  Each AnimNotifyEvent contains an AnimNotify object
+ * which has its Notify method called and passed to the animation.
+ */
+USTRUCT()
+struct FAnimNotifyEvent : public FAnimLinkableElement
+{
+	GENERATED_USTRUCT_BODY()
+
+	/** The user requested time for this notify */
+	UPROPERTY()
+	float DisplayTime_DEPRECATED;
+
+	/** An offset from the DisplayTime to the actual time we will trigger the notify, as we cannot always trigger it exactly at the time the user wants */
+	UPROPERTY()
+	float TriggerTimeOffset;
+
+	/** An offset similar to TriggerTimeOffset but used for the end scrub handle of a notify state's duration */
+	UPROPERTY()
+	float EndTriggerTimeOffset;
+
+	UPROPERTY(EditAnywhere, BlueprintReadWrite, Category=AnimNotifyEvent)
+	float TriggerWeightThreshold;
+
+	UPROPERTY(EditAnywhere, BlueprintReadOnly, Category=AnimNotifyEvent)
+	FName NotifyName;
+
+	UPROPERTY(EditAnywhere, Instanced, BlueprintReadWrite, Category=AnimNotifyEvent)
+	class UAnimNotify * Notify;
+
+	UPROPERTY(EditAnywhere, Instanced, BlueprintReadWrite, Category=AnimNotifyEvent)
+	class UAnimNotifyState * NotifyStateClass;
+
+	UPROPERTY()
+	float Duration;
+
+	/** Linkable element to use for the end handle representing a notify state duration */
+	UPROPERTY()
+	FAnimLinkableElement EndLink;
+
+	/** If TRUE, this notify has been converted from an old BranchingPoint. */
+	UPROPERTY()
+	bool bConvertedFromBranchingPoint;
+
+	UPROPERTY(EditAnywhere, BlueprintReadWrite, Category=AnimNotifyEvent)
+	TEnumAsByte<EMontageNotifyTickType::Type> MontageTickType;
+
+	/** Defines the chance of of this notify triggering, 0 = No Chance, 1 = Always triggers */
+	UPROPERTY(EditAnywhere, BlueprintReadWrite, Category = AnimNotifyTriggerSettings, meta = (ClampMin = "0.0", ClampMax = "1.0"))
+	float NotifyTriggerChance;
+
+	/** Defines a method for filtering notifies (stopping them triggering) e.g. by looking at the meshes current LOD */
+	UPROPERTY(EditAnywhere, BlueprintReadWrite, Category = AnimNotifyTriggerSettings)
+	TEnumAsByte<ENotifyFilterType::Type> NotifyFilterType;
+
+	/** LOD to start filtering this notify from.*/
+	UPROPERTY(EditAnywhere, BlueprintReadWrite, Category = AnimNotifyTriggerSettings, meta = (ClampMin = "0"))
+	int32 NotifyFilterLOD;
+
+#if WITH_EDITORONLY_DATA
+	/** Color of Notify in editor */
+	UPROPERTY()
+	FColor NotifyColor;
+#endif // WITH_EDITORONLY_DATA
+
+	/** 'Track' that the notify exists on, used for visual placement in editor and sorting priority in runtime */
+	UPROPERTY()
+	int32 TrackIndex;
+
+	FAnimNotifyEvent()
+		: DisplayTime_DEPRECATED(0)
+		, TriggerTimeOffset(0)
+		, EndTriggerTimeOffset(0)
+		, TriggerWeightThreshold(ZERO_ANIMWEIGHT_THRESH)
+#if WITH_EDITORONLY_DATA
+		, Notify(NULL)
+#endif // WITH_EDITORONLY_DATA
+		, NotifyStateClass(NULL)
+		, Duration(0)
+		, bConvertedFromBranchingPoint(false)
+		, MontageTickType(EMontageNotifyTickType::Queued)
+		, NotifyTriggerChance(1.f)
+		, NotifyFilterType(ENotifyFilterType::NoFiltering)
+		, NotifyFilterLOD(0)
+#if WITH_EDITORONLY_DATA
+		, TrackIndex(0)
+#endif // WITH_EDITORONLY_DATA
+	{
+	}
+
+	/** Updates trigger offset based on a combination of predicted offset and current offset */
+	ENGINE_API void RefreshTriggerOffset(EAnimEventTriggerOffsets::Type PredictedOffsetType);
+
+	/** Updates end trigger offset based on a combination of predicted offset and current offset */
+	ENGINE_API void RefreshEndTriggerOffset(EAnimEventTriggerOffsets::Type PredictedOffsetType);
+
+	/** Returns the actual trigger time for this notify. In some cases this may be different to the DisplayTime that has been set */
+	ENGINE_API float GetTriggerTime() const;
+
+	/** Returns the actual end time for a state notify. In some cases this may be different from DisplayTime + Duration */
+	ENGINE_API float GetEndTriggerTime() const;
+
+	ENGINE_API float GetDuration() const;
+
+	ENGINE_API void SetDuration(float NewDuration);
+
+	/** Returns true is this AnimNotify is a BranchingPoint */
+	ENGINE_API bool IsBranchingPoint() const
+	{
+		return (MontageTickType == EMontageNotifyTickType::BranchingPoint) && GetLinkedMontage();
+	}
+
+	/** Returns true if this is blueprint derived notifies **/
+	bool IsBlueprintNotify() const
+	{
+		return Notify != NULL || NotifyStateClass != NULL;
+	}
+
+	bool operator ==(const FAnimNotifyEvent& Other) const
+	{
+		return(
+			(Notify && Notify == Other.Notify) || 
+			(NotifyStateClass && NotifyStateClass == Other.NotifyStateClass) ||
+			(!IsBlueprintNotify() && NotifyName == Other.NotifyName)
+			);
+	}
+
+	/** This can be used with the Sort() function on a TArray of FAnimNotifyEvents to sort the notifies array by time, earliest first. */
+	ENGINE_API bool operator <(const FAnimNotifyEvent& Other) const;
+
+	ENGINE_API virtual void SetTime(float NewTime, EAnimLinkMethod::Type ReferenceFrame = EAnimLinkMethod::Absolute) override;
+};
+
+// Used by UAnimSequenceBase::SortNotifies() to sort its Notifies array
+FORCEINLINE bool FAnimNotifyEvent::operator<(const FAnimNotifyEvent& Other) const
+{
+	float ATime = GetTriggerTime();
+	float BTime = Other.GetTriggerTime();
+
+	// When we've got the same time sort on the track index. Explicitly
+	// using SMALL_NUMBER here incase the underlying default changes as
+	// notifies can have an offset of KINDA_SMALL_NUMBER to be consider
+	// distinct
+	if (FMath::IsNearlyEqual(ATime, BTime, SMALL_NUMBER))
+	{
+		return TrackIndex < Other.TrackIndex;
+	}
+	else
+	{
+		return ATime < BTime;
+	}
+}
+
+/**
+ * Keyframe position data for one track.  Pos(i) occurs at Time(i).  Pos.Num() always equals Time.Num().
+ */
+USTRUCT()
+struct FAnimNotifyTrack
+{
+	GENERATED_USTRUCT_BODY()
+
+	UPROPERTY()
+	FName TrackName;
+
+	UPROPERTY()
+	FLinearColor TrackColor;
+
+	TArray<FAnimNotifyEvent*> Notifies;
+
+	FAnimNotifyTrack()
+		: TrackName(TEXT(""))
+		, TrackColor(FLinearColor::White)
+	{
+	}
+
+	FAnimNotifyTrack(FName InTrackName, FLinearColor InTrackColor)
+		: TrackName(InTrackName)
+		, TrackColor(InTrackColor)
+	{
+	}
+};
+
+/** 
+ * Indicates whether an animation is additive, and what kind.
+ */
+UENUM()
+enum EAdditiveAnimationType
+{
+	/** No additive. */
+	AAT_None  UMETA(DisplayName="No additive"),
+	/* Create Additive from LocalSpace Base. */
+	AAT_LocalSpaceBase UMETA(DisplayName="Local Space"),
+	/* Create Additive from MeshSpace Rotation Only, Translation still will be LocalSpace. */
+	AAT_RotationOffsetMeshSpace UMETA(DisplayName="Mesh Space"),
+	AAT_MAX,
+};
+
+UENUM()
+namespace ECurveBlendOption
+{
+	enum Type
+	{
+		/** Find Max Weight of curve and use that weight. */
+		MaxWeight,
+		/** Normalize By Sum of Weight and use it to blend. */
+		NormalizeByWeight,
+		/** Blend By Weight without normalizing*/
+		BlendByWeight
+	};
+}
\ No newline at end of file
diff --git a/Engine/Source/Runtime/Engine/Public/AnimationRuntime.h b/Engine/Source/Runtime/Engine/Public/AnimationRuntime.h
--- a/Engine/Source/Runtime/Engine/Public/AnimationRuntime.h
+++ b/Engine/Source/Runtime/Engine/Public/AnimationRuntime.h
@@ -6,16 +6,15 @@
 =============================================================================*/ 
 
 #pragma once
-#include "Animation/AnimSequenceBase.h"
 #include "BonePose.h"
+#include "AnimTypes.h"
+#include "AnimCurveTypes.h"
 
 struct FInputBlendPose;
 struct FA2CSPose;
 struct FA2Pose;
-struct FAnimTrack;
 struct FPerBoneBlendWeight;
 class UBlendSpaceBase;
-class UAnimSequenceBase;
 typedef TArray<FTransform> FTransformArrayA2;
 /** In AnimationRunTime Library, we extract animation data based on Skeleton hierarchy, not ref pose hierarchy. 
 	Ref pose will need to be re-mapped later **/
@@ -42,8 +41,10 @@ public:
 	*/
 	static void BlendPosesTogether(
 		const TArray<FCompactPose>& SourcePoses,
+		const TArray<FBlendedCurve>& SourceCurves,
 		const TArray<float>&	SourceWeights,
-		/*out*/ FCompactPose& ResultPose);
+		/*out*/ FCompactPose& ResultPose, 
+		/*out*/ FBlendedCurve& ResultCurve);
 
 	/**
 	* Blends together a set of poses, each with a given weight.
@@ -56,8 +57,10 @@ public:
 	*/
 	static void BlendPosesTogether(
 		const TArray<const FCompactPose*>& SourcePoses,
+		const TArray<const FBlendedCurve*>& SourceCurves,
 		const TArray<float>&	SourceWeights,
-		/*out*/ FCompactPose& ResultPose);
+		/*out*/ FCompactPose& ResultPose, 
+		/*out*/ FBlendedCurve& ResultCurve);
 
 	/**
 	* Blends together two poses.
@@ -70,8 +73,11 @@ public:
 	static void BlendTwoPosesTogether(
 		const FCompactPose& SourcePose1,
 		const FCompactPose& SourcePose2,
+		const FBlendedCurve& SourceCurve1,
+		const FBlendedCurve& SourceCurve2,
 		const float			WeightOfPose1,
-		/*out*/ FCompactPose& ResultPose);
+		/*out*/ FCompactPose& ResultPose,
+		/*out*/ FBlendedCurve& ResultCurve);
 
 	/**
 	* Blends together a set of poses, each with a given weight.
@@ -83,9 +89,11 @@ public:
 	*/
 	static void BlendPosesTogetherPerBone(
 		const TArray<FCompactPose>& SourcePoses,
+		const TArray<FBlendedCurve>& SourceCurves,
 		const UBlendSpaceBase* BlendSpace,
 		const TArray<FBlendSampleData>& BlendSampleDataCache,
-		/*out*/ FCompactPose& ResultPose);
+		/*out*/ FCompactPose& ResultPose, 
+		/*out*/ FBlendedCurve& ResultCurve);
 
 	/**
 	* Blends together a set of poses, each with a given weight.
@@ -97,9 +105,11 @@ public:
 	*/
 	static void BlendPosesTogetherPerBoneInMeshSpace(
 		TArray<FCompactPose>& SourcePoses,
+		const TArray<FBlendedCurve>& SourceCurves,
 		const UBlendSpaceBase* BlendSpace,
 		const TArray<FBlendSampleData>& BlendSampleDataCache,
-		/*out*/ FCompactPose& ResultPose);
+		/*out*/ FCompactPose& ResultPose,
+		/*out*/ FBlendedCurve& ResultCurve);
 
 	/**
 	* Blend Poses per bone weights : The BasePoses + BlendPoses(SourceIndex) * Blend Weights(BoneIndex)
@@ -112,7 +122,16 @@ public:
 	* having multiple bone names with multiple weights, and filtering through which one is correct one
 	* I assume all those things should be determined before coming here and this only cares about weights
 	**/
-	static void BlendPosesPerBoneFilter(struct FCompactPose& BasePose, const TArray<struct FCompactPose>& BlendPoses, struct FCompactPose& OutPose, TArray<FPerBoneBlendWeight>& BoneBlendWeights, bool bMeshSpaceRotationBlending);
+	static void BlendPosesPerBoneFilter(
+		struct FCompactPose& BasePose, 
+		const TArray<struct FCompactPose>& BlendPoses, 
+		struct FBlendedCurve& BaseCurve, 
+		const TArray<struct FBlendedCurve>& BlendCurves, 
+		struct FCompactPose& OutPose, 
+		struct FBlendedCurve& OutCurve, 
+		TArray<FPerBoneBlendWeight>& BoneBlendWeights, 
+		bool bMeshSpaceRotationBlending,
+		enum ECurveBlendOption::Type CurveBlendOption);
 
 	static void UpdateDesiredBoneWeight(const TArray<FPerBoneBlendWeight>& SrcBoneBlendWeights, TArray<FPerBoneBlendWeight>& TargetBoneBlendWeights, const TArray<float>& BlendWeights);
 
@@ -129,16 +148,6 @@ public:
 		const FBoneContainer& RequiredBones,
 		/*inout*/ FTransformArrayA2& Atoms);
 
-	static void GetPoseFromSequence(
-		const UAnimSequenceBase* InSequence,
-		/*out*/ FCompactPose& ResultPose,
-		const FAnimExtractContext& ExtractionContext);
-
-	static void GetPoseFromAnimTrack(
-		const FAnimTrack& Track,
-		/*out*/ FCompactPose& ResultPose,
-		const FAnimExtractContext& ExtractionContext);
-
 	/** Fill ref pose **/
 	static void FillWithRefPose(TArray<FTransform> & OutAtoms, const FBoneContainer& RequiredBones);
 
@@ -159,25 +168,15 @@ public:
 	/** Convert a MeshSpaceRotation pose to Local Space. Rotations are NOT normalized. */
 	static void ConvertMeshRotationPoseToLocalSpace(FCompactPose& Pose);
 
+	/** Accumulate Additive Pose based on AdditiveType*/
+	static void AccumulateAdditivePose(FCompactPose& BasePose, const FCompactPose& AdditivePose, FBlendedCurve& BaseCurve, const FBlendedCurve& AdditiveCurve, float Weight, enum EAdditiveAnimationType AdditiveType);
+
 	/** Accumulates weighted AdditivePose to BasePose. Rotations are NOT normalized. */
-	static void AccumulateAdditivePose(FCompactPose& BasePose, const FCompactPose& AdditivePose, float Weight);
+	static void AccumulateLocalSpaceAdditivePose(FCompactPose& BasePose, const FCompactPose& AdditivePose, FBlendedCurve& BaseCurve, const FBlendedCurve& AdditiveCurve, float Weight);
 
 	/** Accumulate a MeshSpaceRotation Additive pose to a local pose. Rotations are NOT normalized */
-	static void AccumulateMeshSpaceRotationAdditiveToLocalPose(FCompactPose& BasePose, const FCompactPose& MeshSpaceRotationAdditive, float Weight);
+	static void AccumulateMeshSpaceRotationAdditiveToLocalPose(FCompactPose& BasePose, const FCompactPose& MeshSpaceRotationAdditive, FBlendedCurve& BaseCurve, const FBlendedCurve& AdditiveCurve, float Weight);
 
-	/**
-	 * Accumulates BlendPoses to ResultAtoms with BlendWeight. 
-	 * ResultAtoms += BlendPose*BlendWeight.
-	 * Result is NOT Normalized, as this is expected to be called several times with each 'BlendPoses' transform array to be accumulated.
-	 * 
-	 * @param	SourcePose		BoneTransforms to scale and accumulate.
-	 * @param	BlendWeight		Weight to apply to 'BlendPoses'
-	 * @param	ResultPose		BoneTransforms to accumulate to.
-	 */
-	static void BlendPosesAccumulate(
-		const FCompactPose& SourcePose,
-		const float BlendWeight,
-		/*inout*/ FCompactPose& ResultPose);
 
 	/** Lerp for BoneTransforms. Stores results in A. Performs A = Lerp(A, B, Alpha);
 	 * @param A : In/Out transform array.
@@ -187,13 +186,6 @@ public:
 	 */
 	static void LerpBoneTransforms(TArray<FTransform> & A, const TArray<FTransform> & B, float Alpha, const TArray<FBoneIndexType> & RequiredBonesArray);
 
-	/**
-	* Get Pose From BlendSpace
-	*/
-	static void GetPoseFromBlendSpace(
-		UBlendSpaceBase * BlendSpace,
-		TArray<FBlendSampleData>& BlendSampleDataCache,
-		/*out*/ FCompactPose& ResultPose);
 
 	/** 
 	 * Advance CurrentTime to CurrentTime + MoveDelta. 
@@ -275,8 +267,12 @@ private:
 	static void BlendMeshPosesPerBoneWeights(
 				struct FCompactPose& BasePose,
 				const TArray<struct FCompactPose>& BlendPoses,
+				struct FBlendedCurve& BaseCurve, 
+				const TArray<struct FBlendedCurve>& BlendedCurves, 
 				const TArray<FPerBoneBlendWeight> & BoneBlendWeights,
-				/*out*/ FCompactPose& OutPose);
+				enum ECurveBlendOption::Type CurveBlendOption,
+				/*out*/ FCompactPose& OutPose, 
+				/*out*/ struct FBlendedCurve& OutCurve);
 
 	/** 
 	* Blend Poses per bone weights : The BasePoses + BlendPoses(SourceIndex) * Blend Weights(BoneIndex)
@@ -294,7 +290,11 @@ private:
 	static void BlendLocalPosesPerBoneWeights(
 			FCompactPose& BasePose,
 			const TArray<FCompactPose>& BlendPoses,
+			struct FBlendedCurve& BaseCurve, 
+			const TArray<struct FBlendedCurve>& BlendedCurves, 
 			const TArray<FPerBoneBlendWeight> & BoneBlendWeights,
-			/*out*/ FCompactPose& OutPose);
+			enum ECurveBlendOption::Type CurveBlendOption,
+			/*out*/ FCompactPose& OutPose, 
+			/*out*/ struct FBlendedCurve& OutCurve);
 };
 
diff --git a/Engine/Source/Runtime/Engine/Public/BonePose.h b/Engine/Source/Runtime/Engine/Public/BonePose.h
--- a/Engine/Source/Runtime/Engine/Public/BonePose.h
+++ b/Engine/Source/Runtime/Engine/Public/BonePose.h
@@ -217,6 +217,16 @@ public:
 		BoneContainer = NULL;
 	}
 
+	// Copy bone transform from SrcPose to this
+	void CopyBonesFrom(FCompactPose& SrcPose)
+	{
+		if (this != &SrcPose)
+		{
+			Bones = SrcPose.Bones;
+			BoneContainer = SrcPose.BoneContainer;
+		}
+	}
+
 	// Sets this pose to its ref pose
 	void ResetToRefPose()
 	{
diff --git a/Engine/Source/Runtime/Engine/Public/Engine.h b/Engine/Source/Runtime/Engine/Public/Engine.h
--- a/Engine/Source/Runtime/Engine/Public/Engine.h
+++ b/Engine/Source/Runtime/Engine/Public/Engine.h
@@ -164,6 +164,7 @@
 #include "Components/SceneCaptureComponentCube.h"
 #include "Components/WindDirectionalSourceComponent.h"
 #include "Components/TimelineComponent.h"
+#include "Animation/AnimCurveTypes.h"
 #include "Animation/AnimSequenceBase.h"
 #include "Animation/AnimCompositeBase.h"
 #include "Animation/AnimComposite.h"
