PR #1292: Allow creation of standalone GameInstance in editor (Contributed by [REDACTED_PERSON])

diff --git a/Engine/Source/Runtime/Engine/Private/GameInstance.cpp b/Engine/Source/Runtime/Engine/Private/GameInstance.cpp
--- a/Engine/Source/Runtime/Engine/Private/GameInstance.cpp
+++ b/Engine/Source/Runtime/Engine/Private/GameInstance.cpp
@@ -16,7 +16,6 @@
 #if WITH_EDITOR
 #include "UnrealEd.h"
 #endif
-#include "Engine/GameEngine.h"
 
 
 UGameInstance::UGameInstance(const FObjectInitializer& ObjectInitializer)
@@ -88,10 +87,8 @@ void UGameInstance::Shutdown()
 
 void UGameInstance::InitializeStandalone()
 {
-	UGameEngine* const Engine = CastChecked<UGameEngine>(GetEngine());
-
 	// Creates the world context. This should be the only WorldContext that ever gets created for this GameInstance.
-	WorldContext = &Engine->CreateNewWorldContext(EWorldType::Game);
+	WorldContext = &GetEngine()->CreateNewWorldContext(EWorldType::Game);
 	WorldContext->OwningGameInstance = this;
 
 	// In standalone create a dummy world from the beginning to avoid issues of not having a world until LoadMap gets us our real world
