Temporary fix for compile error due to integration

diff --git a/Engine/Source/Runtime/Engine/Public/BonePose.h b/Engine/Source/Runtime/Engine/Public/BonePose.h
--- a/Engine/Source/Runtime/Engine/Public/BonePose.h
+++ b/Engine/Source/Runtime/Engine/Public/BonePose.h
@@ -217,6 +217,16 @@ public:
 		BoneContainer = NULL;
 	}
 
+	// Moves transform data out of the supplied InTransforms. InTransform will be left empty
+	void MoveBonesFrom(TArray<FTransform>& InTransforms)
+	{
+		// if number doesn't match it is not compatible to start with
+		if (InTransforms.Num() == Bones.Num())
+		{
+			Bones = MoveTemp(InTransforms);
+		}
+	}
+
 	// Copy bone transform from SrcPose to this
 	void CopyBonesFrom(FCompactPose& SrcPose)
 	{
