Fixed compile error.

diff --git a/Engine/Source/Runtime/Engine/Private/Animation/AnimSingleNodeInstance.cpp b/Engine/Source/Runtime/Engine/Private/Animation/AnimSingleNodeInstance.cpp
--- a/Engine/Source/Runtime/Engine/Private/Animation/AnimSingleNodeInstance.cpp
+++ b/Engine/Source/Runtime/Engine/Private/Animation/AnimSingleNodeInstance.cpp
@@ -346,11 +346,13 @@ bool UAnimSingleNodeInstance::NativeEvaluateAnimation(FPoseContext& Output)
 			// find out if this is additive animation
 			if (AnimTrack.IsAdditive())
 			{
+#if WITH_EDITORONLY_DATA
 				if (Composite->PreviewBasePose)
 				{
 					Composite->PreviewBasePose->GetAdditiveBasePose(Output.Pose, Output.Curve, ExtractionContext);
 				}
 				else
+#endif
 				{
 					// get base pose - for now we only support ref pose as base
 					Output.Pose.ResetToRefPose();
