Fixes inconsistencies between camera relative world space and translated world space

diff --git a/Engine/Shaders/MaterialTemplate.usf b/Engine/Shaders/MaterialTemplate.usf
--- a/Engine/Shaders/MaterialTemplate.usf
+++ b/Engine/Shaders/MaterialTemplate.usf
@@ -266,40 +266,40 @@ struct FMaterialVertexParameters
 	FMaterialParticleParameters Particle;
 };
 
-float3 GetWorldPosition(FMaterialVertexParameters Parameters)
+float3 GetTranslatedWorldPosition(FMaterialVertexParameters Parameters)
 {
 	return Parameters.WorldPosition;
 }
 
-float3 GetPrevWorldPosition(FMaterialVertexParameters Parameters)
+float3 GetPrevTranslatedWorldPosition(FMaterialVertexParameters Parameters)
 {
 	// Previous world position and current world position are sharing the
 	// same attribute in Parameters, because in BasePassVertexShader.usf
 	// and in VelocityShader.usf, we are regenerating a Parameters from
 	// VertexFactoryGetPreviousWorldPosition() instead of
 	// VertexFactoryGetWorldPosition().
-	return GetWorldPosition(Parameters);
+	return GetTranslatedWorldPosition(Parameters);
 }
 
-float3 GetWorldPosition_NoOffsets(FMaterialVertexParameters Parameters)
+float3 GetWorldPosition(FMaterialVertexParameters Parameters)
 {
-	return GetWorldPosition(Parameters) - View.PreViewTranslation;
+	return GetTranslatedWorldPosition(Parameters) - View.PreViewTranslation;
 }
 
-float3 GetPrevWorldPosition_NoOffsets(FMaterialVertexParameters Parameters)
+float3 GetPrevWorldPosition(FMaterialVertexParameters Parameters)
 {
-	return GetPrevWorldPosition(Parameters) - View.PrevPreViewTranslation;
+	return GetPrevTranslatedWorldPosition(Parameters) - View.PrevPreViewTranslation;
 }
 
 //TODO(bug UE-17131): We should compute world displacement for the previous frame
 float3 GetWorldPosition(FMaterialTessellationParameters Parameters)
 {
-	return Parameters.WorldPosition + View.PreViewTranslation;
+	return Parameters.WorldPosition;
 }
 
-float3 GetWorldPosition_NoOffsets(FMaterialTessellationParameters Parameters)
+float3 GetTranslatedWorldPosition(FMaterialTessellationParameters Parameters)
 {
-	return Parameters.WorldPosition;
+	return Parameters.WorldPosition + View.PreViewTranslation;
 }
 
 float3 GetWorldPosition(FMaterialPixelParameters Parameters)
@@ -307,17 +307,17 @@ float3 GetWorldPosition(FMaterialPixelParameters Parameters)
 	return Parameters.WorldPosition;
 }
 
-float3 GetWorldPosition_NoOffsets(FMaterialPixelParameters Parameters)
+float3 GetWorldPosition_NoMaterialOffsets(FMaterialPixelParameters Parameters)
 {
 	return Parameters.WorldPosition_NoOffsets;
 }
 
-float3 GetWorldPosition_CamRelative(FMaterialPixelParameters Parameters)
+float3 GetTranslatedWorldPosition(FMaterialPixelParameters Parameters)
 {
 	return Parameters.WorldPosition_CamRelative;
 }
 
-float3 GetWorldPosition_NoOffsets_CamRelative(FMaterialPixelParameters Parameters)
+float3 GetTranslatedWorldPosition_NoMaterialOffsets(FMaterialPixelParameters Parameters)
 {
 	return Parameters.WorldPosition_NoOffsets_CamRelative;
 }
@@ -807,7 +807,7 @@ float3 GetSpeedTreeVertexOffsetInner(FMaterialVertexParameters Parameters, int G
 		}
 	#else
 		float3x3 LocalToWorld = (float3x3)Primitive.LocalToWorld;
-		float3 LocalPosition = mul(float4(GetWorldPosition_NoOffsets(Parameters), 1), Primitive.WorldToLocal).xyz;
+		float3 LocalPosition = mul(float4(GetWorldPosition(Parameters), 1), Primitive.WorldToLocal).xyz;
 	#endif
 
 	float3 TreePos = GetObjectWorldPosition(Parameters);
diff --git a/Engine/Source/Runtime/Engine/Classes/Materials/MaterialExpressionTransformPosition.h b/Engine/Source/Runtime/Engine/Classes/Materials/MaterialExpressionTransformPosition.h
--- a/Engine/Source/Runtime/Engine/Classes/Materials/MaterialExpressionTransformPosition.h
+++ b/Engine/Source/Runtime/Engine/Classes/Materials/MaterialExpressionTransformPosition.h
@@ -8,16 +8,16 @@
 UENUM()
 enum EMaterialPositionTransformSource
 {
-	/** Local space */
+	/** Local Space */
 	TRANSFORMPOSSOURCE_Local UMETA(DisplayName="Local Space"),
 	
-	/** World space with no shader offsets applied */
-	TRANSFORMPOSSOURCE_World UMETA(DisplayName="World Space (No Shader Offsets)"),
+	/** Absolute World Space */
+	TRANSFORMPOSSOURCE_World UMETA(DisplayName="Absolute World Space"),
 	
-	/** World space with all shader offsets applied */
-	TRANSFORMPOSSOURCE_TranslatedWorld  UMETA(DisplayName="World Space"),
+	/** Camera Relative World Space */
+	TRANSFORMPOSSOURCE_TranslatedWorld  UMETA(DisplayName="Camera Relative World Space"),
 
-	/** View space */
+	/** View Space */
 	TRANSFORMPOSSOURCE_View  UMETA(DisplayName="View Space"),
 	
 	TRANSFORMPOSSOURCE_MAX,
diff --git a/Engine/Source/Runtime/Engine/Classes/Materials/MaterialExpressionWorldPosition.h b/Engine/Source/Runtime/Engine/Classes/Materials/MaterialExpressionWorldPosition.h
--- a/Engine/Source/Runtime/Engine/Classes/Materials/MaterialExpressionWorldPosition.h
+++ b/Engine/Source/Runtime/Engine/Classes/Materials/MaterialExpressionWorldPosition.h
@@ -9,17 +9,17 @@
 UENUM()
 enum EWorldPositionIncludedOffsets
 {
-	/** Absolute world position with all shader offsets applied */
-	WPT_Default UMETA(DisplayName="Absolute World Position"),
+	/** Absolute world position with all material shader offsets applied */
+	WPT_Default UMETA(DisplayName="Absolute World Position (Including Material Shader Offsets)"),
 
-	/** Absolute world position with no shader offsets applied */
-	WPT_ExcludeAllShaderOffsets UMETA(DisplayName="Absolute World Position (No Shader Offsets)"),
+	/** Absolute world position with no material shader offsets applied */
+	WPT_ExcludeAllShaderOffsets UMETA(DisplayName="Absolute World Position (Excluding Material Shader Offsets)"),
 
-	/** Camera relative world position with all shader offsets applied */
-	WPT_CameraRelative UMETA(DisplayName="Camera Relative World Position"),
+	/** Camera relative world position with all material shader offsets applied */
+	WPT_CameraRelative UMETA(DisplayName="Camera Relative World Position (Including Material Shader Offsets)"),
 
-	/** Camera relative world position with no shader offsets applied */
-	WPT_CameraRelativeNoOffsets UMETA(DisplayName="Camera Relative World Position (No Shader Offsets"),
+	/** Camera relative world position with no material shader offsets applied */
+	WPT_CameraRelativeNoOffsets UMETA(DisplayName="Camera Relative World Position (Excluding Material Shader Offsets)"),
 
 	WPT_MAX
 };
@@ -28,7 +28,7 @@ UCLASS(collapsecategories, hidecategories=Object, MinimalAPI)
 class UMaterialExpressionWorldPosition : public UMaterialExpression
 {
 	GENERATED_UCLASS_BODY()
-
+	
 	UPROPERTY(EditAnywhere, Category=UMaterialExpressionWorldPosition, meta=(DisplayName = "Shader Offsets"))
 	TEnumAsByte<EWorldPositionIncludedOffsets> WorldPositionShaderOffset;
 	
diff --git a/Engine/Source/Runtime/Engine/Private/Materials/HLSLMaterialTranslator.h b/Engine/Source/Runtime/Engine/Private/Materials/HLSLMaterialTranslator.h
--- a/Engine/Source/Runtime/Engine/Private/Materials/HLSLMaterialTranslator.h
+++ b/Engine/Source/Runtime/Engine/Private/Materials/HLSLMaterialTranslator.h
@@ -2092,7 +2092,7 @@ protected:
 
 	virtual int32 WorldPosition(EWorldPositionIncludedOffsets WorldPositionIncludedOffsets) override
 	{
-		const TCHAR * WorldPositionSuffix = nullptr; 
+		FString FunctionNamePattern;
 
 		// If this material has no expressions for world position offset or world displacement, the non-offset world position will
 		// be exactly the same as the offset one, so there is no point bringing in the extra code.
@@ -2102,48 +2102,27 @@ protected:
 		{
 		case WPT_Default:
 			{
-				// TODO: in tessellation, no offset should actually return the world position without including the
-				// world position offset
-				WorldPositionSuffix = TEXT("");
-
-				if (ShaderFrequency != SF_Pixel)
-				{
-					// TODO: there is a miss understanding between no offset, translated world position and camera relative:
-					//		what the artist call camera relative, is the translated world position
-					//		what they call no offsets, is without apply material's world position offset
-					WorldPositionSuffix = TEXT("_NoOffsets");
-				}
-
+				FunctionNamePattern = TEXT("Get<PREV>WorldPosition");
 				break;
 			}
 
 		case WPT_ExcludeAllShaderOffsets:
 			{
 				bNeedsWorldPositionExcludingShaderOffsets = true;
-				WorldPositionSuffix = TEXT("_NoOffsets");
+				FunctionNamePattern = TEXT("Get<PREV>WorldPosition<NO_MATERIAL_OFFSETS>");
 				break;
 			}
 
 		case WPT_CameraRelative:
 			{
-				if (ShaderFrequency != SF_Pixel)
-				{
-					return Errorf(TEXT("Cannot access camera relative world position outside of the pixel shader"));
-				}
-
-				WorldPositionSuffix = TEXT("_CamRelative");
+				FunctionNamePattern = TEXT("Get<PREV>TranslatedWorldPosition");
 				break;
 			}
 
 		case WPT_CameraRelativeNoOffsets:
 			{
-				if (ShaderFrequency != SF_Pixel)
-				{
-					return Errorf(TEXT("Cannot access camera relative world position outside of the pixel shader"));
-				}
-
 				bNeedsWorldPositionExcludingShaderOffsets = true;
-				WorldPositionSuffix = TEXT("_NoOffsets_CamRelative");
+				FunctionNamePattern = TEXT("Get<PREV>TranslatedWorldPosition<NO_MATERIAL_OFFSETS>");
 				break;
 			}
 
@@ -2154,17 +2133,21 @@ protected:
 			}
 		}
 
-		check(WorldPositionSuffix);
+		// If compiling for the previous frame in the vertex shader
+		FunctionNamePattern.ReplaceInline(TEXT("<PREV>"), bCompilingPreviousFrame && ShaderFrequency == SF_Vertex ? TEXT("Prev") : TEXT(""));
 		
-		if (bCompilingPreviousFrame && ShaderFrequency == SF_Vertex)
+		if (ShaderFrequency == SF_Pixel)
 		{
-			//TODO(bug UE-17131): We should compute world displacement for the previous frame
-			return AddInlinedCodeChunk(MCT_Float3, TEXT("GetPrevWorldPosition%s(Parameters)"), WorldPositionSuffix);
+			// No material offset only available in the vertex shader.
+			// TODO: should also be available in the tesselation shader
+			FunctionNamePattern.ReplaceInline(TEXT("<NO_MATERIAL_OFFSETS>"), TEXT("_NoMaterialOffsets"));
 		}
 		else
 		{
-			return AddInlinedCodeChunk(MCT_Float3, TEXT("GetWorldPosition%s(Parameters)"), WorldPositionSuffix);
+			FunctionNamePattern.ReplaceInline(TEXT("<NO_MATERIAL_OFFSETS>"), TEXT(""));
 		}
+
+		return AddInlinedCodeChunk(MCT_Float3, TEXT("%s(Parameters)"), *FunctionNamePattern);
 	}
 
 	virtual int32 ObjectWorldPosition() override
diff --git a/Engine/Source/Runtime/Engine/Private/Materials/MaterialExpressions.cpp b/Engine/Source/Runtime/Engine/Private/Materials/MaterialExpressions.cpp
--- a/Engine/Source/Runtime/Engine/Private/Materials/MaterialExpressions.cpp
+++ b/Engine/Source/Runtime/Engine/Private/Materials/MaterialExpressions.cpp
@@ -5806,6 +5806,7 @@ UMaterialExpressionWorldPosition::UMaterialExpressionWorldPosition(const FObject
 
 int32 UMaterialExpressionWorldPosition::Compile(class FMaterialCompiler* Compiler, int32 OutputIndex, int32 MultiplexIndex)
 {
+	// TODO: should use a separate check box for Including/Excluding Material Shader Offsets
 	return Compiler->WorldPosition(WorldPositionShaderOffset);
 }
 
@@ -5815,25 +5816,25 @@ void UMaterialExpressionWorldPosition::GetCaption(TArray<FString>& OutCaptions)
 	{
 	case WPT_Default:
 		{
-			OutCaptions.Add(NSLOCTEXT("MaterialExpressions", "WorldPositonText", "Absolute World Position").ToString());
+			OutCaptions.Add(NSLOCTEXT("MaterialExpressions", "WorldPositonText", "Absolute World Position (Including Material Offsets)").ToString());
 			break;
 		}
 
 	case WPT_ExcludeAllShaderOffsets:
 		{
-			OutCaptions.Add(NSLOCTEXT("MaterialExpressions", "WorldPositonExcludingOffsetsText", "Absolute World Position (Excluding Offsets)").ToString());
+			OutCaptions.Add(NSLOCTEXT("MaterialExpressions", "WorldPositonExcludingOffsetsText", "Absolute World Position (Excluding Material Offsets)").ToString());
 			break;
 		}
 
 	case WPT_CameraRelative:
 		{
-			OutCaptions.Add(NSLOCTEXT("MaterialExpressions", "CamRelativeWorldPositonText", "Camera Relative World Position").ToString());
+			OutCaptions.Add(NSLOCTEXT("MaterialExpressions", "CamRelativeWorldPositonText", "Camera Relative World Position (Including Material Offsets)").ToString());
 			break;
 		}
 
 	case WPT_CameraRelativeNoOffsets:
 		{
-			OutCaptions.Add(NSLOCTEXT("MaterialExpressions", "CamRelativeWorldPositonExcludingOffsetsText", "Camera Relative World Position (Excluding Offsets)").ToString());
+			OutCaptions.Add(NSLOCTEXT("MaterialExpressions", "CamRelativeWorldPositonExcludingOffsetsText", "Camera Relative World Position (Excluding Material Offsets)").ToString());
 			break;
 		}
 
