UE4 - Fix some scene captures not working on fullscreen (UE-11226)

diff --git a/Engine/Source/Runtime/Engine/Private/SceneView.cpp b/Engine/Source/Runtime/Engine/Private/SceneView.cpp
--- a/Engine/Source/Runtime/Engine/Private/SceneView.cpp
+++ b/Engine/Source/Runtime/Engine/Private/SceneView.cpp
@@ -1375,7 +1375,14 @@ void FSceneView::EndFinalPostprocessSettings(const FSceneViewInitOptions& ViewIn
 
 		// Window full screen mode with upscaling
 		bool bFullscreen = false;
-		if (GEngine && GEngine->GameViewport && GEngine->GameViewport->GetWindow().IsValid())
+		bool bIsSceneCapture = false;
+		if (ViewInitOptions.ViewFamily && ViewInitOptions.ViewFamily->Views.Num() == 1)
+		{
+			auto* View = ViewInitOptions.ViewFamily->Views[0];
+			bIsSceneCapture = View ? View->bIsSceneCapture : false;
+		}
+
+		if (!bIsSceneCapture && GEngine && GEngine->GameViewport && GEngine->GameViewport->GetWindow().IsValid())
 		{
 			bFullscreen = GEngine->GameViewport->GetWindow()->GetWindowMode() != EWindowMode::Windowed;
 		}
