Replace accidentally removed code that prevents newly spawned tick functions that haven't yet ticked when they are removed from the task level.

diff --git a/Engine/Source/Runtime/Engine/Private/TickTaskManager.cpp b/Engine/Source/Runtime/Engine/Private/TickTaskManager.cpp
--- a/Engine/Source/Runtime/Engine/Private/TickTaskManager.cpp
+++ b/Engine/Source/Runtime/Engine/Private/TickTaskManager.cpp
@@ -860,6 +860,10 @@ public:
 			check(bFound); // otherwise you changed TickState while the tick function was registered. Call SetTickFunctionEnable instead.
 			break;
 		}
+		if (bTickNewlySpawned)
+		{
+			NewlySpawnedTickFunctions.Remove(TickFunction);
+		}
 	}
 
 private:
