Rework animation track segment looping logic to reduce floating point error problems

diff --git a/Engine/Source/Runtime/Engine/Private/Animation/AnimCompositeBase.cpp b/Engine/Source/Runtime/Engine/Private/Animation/AnimCompositeBase.cpp
--- a/Engine/Source/Runtime/Engine/Private/Animation/AnimCompositeBase.cpp
+++ b/Engine/Source/Runtime/Engine/Private/Animation/AnimCompositeBase.cpp
@@ -99,7 +99,8 @@ void FAnimSegment::GetAnimNotifiesFromTrackPositions(const float& PreviousTrackP
 			// The track can be playing backwards and the animation can be playing backwards, so we
 			// need to combine to work out what direction we are traveling through the animation
 			bool bAnimPlayingBackwards = bTrackPlayingBackwards ^ (ValidPlayRate < 0.f);
-			
+			const float ResetStartPosition = bAnimPlayingBackwards ? AnimEndTime : AnimStartTime;
+
 			// Abstract out end point since animation can be playing forward or backward.
 			const float AnimEndPoint = bAnimPlayingBackwards ? AnimStartTime : AnimEndTime;
 
@@ -111,9 +112,8 @@ void FAnimSegment::GetAnimNotifiesFromTrackPositions(const float& PreviousTrackP
 				// If our time left is shorter than time to end point, no problem. End there.
 				if( FMath::Abs(TrackTimeToGo) < FMath::Abs(TrackTimeToAnimEndPoint) )
 				{
-					const float AnimEndPosition = ConvertTrackPosToAnimPos(CurrentTrackPosition);
-					// Make sure we have not wrapped around, positions should be contiguous.
-					check(bAnimPlayingBackwards ? (AnimEndPosition <= AnimStartPosition) : (AnimStartPosition <= AnimEndPosition));
+					const float AnimPlayRate = ValidPlayRate * (bTrackPlayingBackwards ? -1.f : 1.f);
+					const float AnimEndPosition = (TrackTimeToGo * AnimPlayRate) + AnimStartPosition;
 					AnimSequence->GetAnimNotifiesFromDeltaPositions(AnimStartPosition, AnimEndPosition, OutActiveNotifies);
 					break;
 				}
@@ -121,14 +121,12 @@ void FAnimSegment::GetAnimNotifiesFromTrackPositions(const float& PreviousTrackP
 				else
 				{
 					// Add that piece for extraction.
-					// Make sure we have not wrapped around, positions should be contiguous.
-					check(bAnimPlayingBackwards ? (AnimEndPoint <= AnimStartPosition) : (AnimStartPosition <= AnimEndPoint));
 					AnimSequence->GetAnimNotifiesFromDeltaPositions(AnimStartPosition, AnimEndPoint, OutActiveNotifies);
 
 					// decrease our TrackTimeToGo if we have to do another iteration.
 					// and put ourselves back at the beginning of the animation.
 					TrackTimeToGo -= TrackTimeToAnimEndPoint;
-					AnimStartPosition = bAnimPlayingBackwards ? AnimEndTime : AnimStartTime;
+					AnimStartPosition = ResetStartPosition;
 				}
 			}
 		}
@@ -173,6 +171,7 @@ void FAnimSegment::GetRootMotionExtractionStepsForTrackRange(TArray<FRootMotionE
 			// The track can be playing backwards and the animation can be playing backwards, so we
 			// need to combine to work out what direction we are traveling through the animation
 			bool bAnimPlayingBackwards = bTrackPlayingBackwards ^ (ValidPlayRate < 0.f);
+			const float ResetStartPosition = bAnimPlayingBackwards ? AnimEndTime : AnimStartTime;
 
 			// Abstract out end point since animation can be playing forward or backward.
 			const float AnimEndPoint = bAnimPlayingBackwards ? AnimStartTime : AnimEndTime;
@@ -185,9 +184,8 @@ void FAnimSegment::GetRootMotionExtractionStepsForTrackRange(TArray<FRootMotionE
 				// If our time left is shorter than time to end point, no problem. End there.
 				if( FMath::Abs(TrackTimeToGo) < FMath::Abs(TrackTimeToAnimEndPoint) )
 				{
-					const float AnimEndPosition = ConvertTrackPosToAnimPos(EndTrackPosition);
-					// Make sure we have not wrapped around, positions should be contiguous.
-					check(bAnimPlayingBackwards ? (AnimEndPosition <= AnimStartPosition) : (AnimStartPosition <= AnimEndPosition));
+					const float AnimPlayRate = ValidPlayRate * (bTrackPlayingBackwards ? -1.f : 1.f);
+					const float AnimEndPosition = (TrackTimeToGo * AnimPlayRate) + AnimStartPosition;
 					RootMotionExtractionSteps.Add(FRootMotionExtractionStep(AnimSequence, AnimStartPosition, AnimEndPosition));
 					break;
 				}
@@ -195,14 +193,12 @@ void FAnimSegment::GetRootMotionExtractionStepsForTrackRange(TArray<FRootMotionE
 				else
 				{
 					// Add that piece for extraction.
-					// Make sure we have not wrapped around, positions should be contiguous.
-					check(bAnimPlayingBackwards ? (AnimEndPoint <= AnimStartPosition) : (AnimStartPosition <= AnimEndPoint));
 					RootMotionExtractionSteps.Add(FRootMotionExtractionStep(AnimSequence, AnimStartPosition, AnimEndPoint));
 
 					// decrease our TrackTimeToGo if we have to do another iteration.
 					// and put ourselves back at the beginning of the animation.
 					TrackTimeToGo -= TrackTimeToAnimEndPoint;
-					AnimStartPosition = bAnimPlayingBackwards ? AnimEndTime : AnimStartTime;
+					AnimStartPosition = ResetStartPosition;
 				}
 			}
 		}
