Fix compile error

diff --git a/Engine/Source/Runtime/Engine/Private/Animation/AnimCompositeBase.cpp b/Engine/Source/Runtime/Engine/Private/Animation/AnimCompositeBase.cpp
--- a/Engine/Source/Runtime/Engine/Private/Animation/AnimCompositeBase.cpp
+++ b/Engine/Source/Runtime/Engine/Private/Animation/AnimCompositeBase.cpp
@@ -8,6 +8,7 @@
 #include "AnimationUtils.h"
 #include "Animation/AnimCompositeBase.h"
 #include "Animation/AnimSequence.h"
+#include "AnimationRuntime.h"
 
 ///////////////////////////////////////////////////////
 // FAnimSegment
diff --git a/Engine/Source/Runtime/Engine/Private/Animation/AnimSequence.cpp b/Engine/Source/Runtime/Engine/Private/Animation/AnimSequence.cpp
--- a/Engine/Source/Runtime/Engine/Private/Animation/AnimSequence.cpp
+++ b/Engine/Source/Runtime/Engine/Private/Animation/AnimSequence.cpp
@@ -15,6 +15,7 @@
 #include "Animation/AnimNotifies/AnimNotifyState.h"
 #include "Animation/Rig.h"
 #include "EditorFramework/AssetImportData.h"
+#include "Animation/AnimStats.h"
 
 #define USE_SLERP 0
 
