UE-17625 - Launch Game Crash with None Distribution Method in Cascade

Already submitted fix for "noclear" property support to SPropertyEditorEditInline in CL 2600315. However a number of particle modules do not correctly handle the resetting of distributions, and instead reset to NULL. Michael Schoell last fixed this back in CL 1610305 for most modules and some newer modules also required a fix up.

diff --git a/Engine/Source/Runtime/Engine/Classes/Particles/Location/ParticleModuleSourceMovement.h b/Engine/Source/Runtime/Engine/Classes/Particles/Location/ParticleModuleSourceMovement.h
--- a/Engine/Source/Runtime/Engine/Classes/Particles/Location/ParticleModuleSourceMovement.h
+++ b/Engine/Source/Runtime/Engine/Classes/Particles/Location/ParticleModuleSourceMovement.h
@@ -18,6 +18,9 @@ class UParticleModuleSourceMovement : public UParticleModuleLocationBase
 	struct FRawDistributionVector SourceMovementScale;
 
 	//Begin UObject Interface
+#if WITH_EDITOR
+	virtual void PostEditChangeProperty(FPropertyChangedEvent& PropertyChangedEvent) override;
+#endif // WITH_EDITOR
 	virtual void PostInitProperties() override;
 	//End UObject Interface
 
@@ -28,6 +31,9 @@ class UParticleModuleSourceMovement : public UParticleModuleLocationBase
 		return false;
 	}
 	//End UParticleModule Interface
+
+	/** Initializes the default values for this property */
+	void InitializeDefaults();
 };
 
 
diff --git a/Engine/Source/Runtime/Engine/Classes/Particles/Rotation/ParticleModuleMeshRotation.h b/Engine/Source/Runtime/Engine/Classes/Particles/Rotation/ParticleModuleMeshRotation.h
--- a/Engine/Source/Runtime/Engine/Classes/Particles/Rotation/ParticleModuleMeshRotation.h
+++ b/Engine/Source/Runtime/Engine/Classes/Particles/Rotation/ParticleModuleMeshRotation.h
@@ -22,6 +22,9 @@ class UParticleModuleMeshRotation : public UParticleModuleRotationBase
 	uint32 bInheritParent:1;
 
 	//Begin UObject Interface
+#if WITH_EDITOR
+	virtual void PostEditChangeProperty(FPropertyChangedEvent& PropertyChangedEvent) override;
+#endif // WITH_EDITOR
 	virtual void	PostInitProperties() override;
 	//End UObject Interface
 
@@ -30,6 +33,9 @@ class UParticleModuleMeshRotation : public UParticleModuleRotationBase
 	virtual bool	TouchesMeshRotation() const override { return true; }
 	//End UParticleModule Interface
 
+	/** Initializes the default values for this property */
+	void InitializeDefaults();
+
 	/**
 	 *	Extended version of spawn, allows for using a random stream for distribution value retrieval
 	 *
diff --git a/Engine/Source/Runtime/Engine/Classes/Particles/RotationRate/ParticleModuleMeshRotationRate.h b/Engine/Source/Runtime/Engine/Classes/Particles/RotationRate/ParticleModuleMeshRotationRate.h
--- a/Engine/Source/Runtime/Engine/Classes/Particles/RotationRate/ParticleModuleMeshRotationRate.h
+++ b/Engine/Source/Runtime/Engine/Classes/Particles/RotationRate/ParticleModuleMeshRotationRate.h
@@ -18,6 +18,9 @@ class UParticleModuleMeshRotationRate : public UParticleModuleRotationRateBase
 	struct FRawDistributionVector StartRotationRate;
 
 	//Begin UObject Interface
+#if WITH_EDITOR
+	virtual void PostEditChangeProperty(FPropertyChangedEvent& PropertyChangedEvent) override;
+#endif // WITH_EDITOR
 	virtual void	PostInitProperties() override;
 	//End UObject Interface
 
@@ -27,6 +30,9 @@ class UParticleModuleMeshRotationRate : public UParticleModuleRotationRateBase
 	virtual bool	TouchesMeshRotation() const override { return true; }
 	// End UParticleModule Interface
 
+	/** Initializes the default values for this property */
+	void InitializeDefaults();
+
 	/**
 	 *	Extended version of spawn, allows for using a random stream for distribution value retrieval
 	 *
diff --git a/Engine/Source/Runtime/Engine/Classes/Particles/RotationRate/ParticleModuleMeshRotationRateMultiplyLife.h b/Engine/Source/Runtime/Engine/Classes/Particles/RotationRate/ParticleModuleMeshRotationRateMultiplyLife.h
--- a/Engine/Source/Runtime/Engine/Classes/Particles/RotationRate/ParticleModuleMeshRotationRateMultiplyLife.h
+++ b/Engine/Source/Runtime/Engine/Classes/Particles/RotationRate/ParticleModuleMeshRotationRateMultiplyLife.h
@@ -19,6 +19,9 @@ class UParticleModuleMeshRotationRateMultiplyLife : public UParticleModuleRotati
 
 
 	//Begin UObject Interface
+#if WITH_EDITOR
+	virtual void PostEditChangeProperty(FPropertyChangedEvent& PropertyChangedEvent) override;
+#endif // WITH_EDITOR
 	virtual void	PostInitProperties() override;
 	//End UObject Interface
 
@@ -27,6 +30,9 @@ class UParticleModuleMeshRotationRateMultiplyLife : public UParticleModuleRotati
 	virtual void	Update(FParticleEmitterInstance* Owner, int32 Offset, float DeltaTime) override;
 	virtual void SetToSensibleDefaults(UParticleEmitter* Owner) override;
 	// End UParticleModule Interface
+
+	/** Initializes the default values for this property */
+	void InitializeDefaults();
 };
 
 
diff --git a/Engine/Source/Runtime/Engine/Classes/Particles/RotationRate/ParticleModuleMeshRotationRateOverLife.h b/Engine/Source/Runtime/Engine/Classes/Particles/RotationRate/ParticleModuleMeshRotationRateOverLife.h
--- a/Engine/Source/Runtime/Engine/Classes/Particles/RotationRate/ParticleModuleMeshRotationRateOverLife.h
+++ b/Engine/Source/Runtime/Engine/Classes/Particles/RotationRate/ParticleModuleMeshRotationRateOverLife.h
@@ -25,6 +25,9 @@ class UParticleModuleMeshRotationRateOverLife : public UParticleModuleRotationRa
 	uint32 bScaleRotRate:1;
 
 	//Begin UObject Interface
+#if WITH_EDITOR
+	virtual void PostEditChangeProperty(FPropertyChangedEvent& PropertyChangedEvent) override;
+#endif // WITH_EDITOR
 	virtual void	PostInitProperties() override;
 	//End UObject Interface
 
@@ -34,6 +37,9 @@ class UParticleModuleMeshRotationRateOverLife : public UParticleModuleRotationRa
 	virtual void SetToSensibleDefaults(UParticleEmitter* Owner) override;
 	virtual bool TouchesMeshRotation() const override { return true; }
 	// End UParticleModule Interface
+
+	/** Initializes the default values for this property */
+	void InitializeDefaults();
 };
 
 
diff --git a/Engine/Source/Runtime/Engine/Classes/Particles/VectorField/ParticleModuleVectorFieldScale.h b/Engine/Source/Runtime/Engine/Classes/Particles/VectorField/ParticleModuleVectorFieldScale.h
--- a/Engine/Source/Runtime/Engine/Classes/Particles/VectorField/ParticleModuleVectorFieldScale.h
+++ b/Engine/Source/Runtime/Engine/Classes/Particles/VectorField/ParticleModuleVectorFieldScale.h
@@ -17,7 +17,13 @@ class UParticleModuleVectorFieldScale : public UParticleModuleVectorFieldBase
 	UPROPERTY(EditAnywhere, Category=VectorField)
 	class UDistributionFloat* VectorFieldScale;
 
+	/** Initializes the default values for this property */
+	void InitializeDefaults();
+
 	// Begin UObject Interface
+#if WITH_EDITOR
+	virtual void PostEditChangeProperty(FPropertyChangedEvent& PropertyChangedEvent) override;
+#endif // WITH_EDITOR
 	virtual void PostInitProperties() override;
 	// End UObject Interface
 
diff --git a/Engine/Source/Runtime/Engine/Classes/Particles/VectorField/ParticleModuleVectorFieldScaleOverLife.h b/Engine/Source/Runtime/Engine/Classes/Particles/VectorField/ParticleModuleVectorFieldScaleOverLife.h
--- a/Engine/Source/Runtime/Engine/Classes/Particles/VectorField/ParticleModuleVectorFieldScaleOverLife.h
+++ b/Engine/Source/Runtime/Engine/Classes/Particles/VectorField/ParticleModuleVectorFieldScaleOverLife.h
@@ -17,7 +17,13 @@ class UParticleModuleVectorFieldScaleOverLife : public UParticleModuleVectorFiel
 	UPROPERTY(EditAnywhere, Category=VectorField)
 	class UDistributionFloat* VectorFieldScaleOverLife;
 
+	/** Initializes the default values for this property */
+	void InitializeDefaults();
+
 	// Begin UObject Interface
+#if WITH_EDITOR
+	virtual void PostEditChangeProperty(FPropertyChangedEvent& PropertyChangedEvent) override;
+#endif // WITH_EDITOR
 	virtual void PostInitProperties() override;
 	// End UObject Interface
 
diff --git a/Engine/Source/Runtime/Engine/Private/Particles/ParticleModules.cpp b/Engine/Source/Runtime/Engine/Private/Particles/ParticleModules.cpp
--- a/Engine/Source/Runtime/Engine/Private/Particles/ParticleModules.cpp
+++ b/Engine/Source/Runtime/Engine/Private/Particles/ParticleModules.cpp
@@ -873,10 +873,9 @@ UParticleModuleSourceMovement::UParticleModuleSourceMovement(const FObjectInitia
 	bFinalUpdateModule = true;
 }
 
-void UParticleModuleSourceMovement::PostInitProperties()
+void UParticleModuleSourceMovement::InitializeDefaults()
 {
-	Super::PostInitProperties();
-	if (!HasAnyFlags(RF_ClassDefaultObject | RF_NeedLoad))
+	if (!SourceMovementScale.Distribution)
 	{
 		UDistributionVectorConstant* DistributionSourceMovementScale = NewObject<UDistributionVectorConstant>(this, TEXT("DistributionSourceMovementScale"));
 		DistributionSourceMovementScale->Constant = FVector(1.0f, 1.0f, 1.0f);
@@ -884,6 +883,23 @@ void UParticleModuleSourceMovement::PostInitProperties()
 	}
 }
 
+void UParticleModuleSourceMovement::PostInitProperties()
+{
+	Super::PostInitProperties();
+	if (!HasAnyFlags(RF_ClassDefaultObject | RF_NeedLoad))
+	{
+		InitializeDefaults();
+	}
+}
+
+#if WITH_EDITOR
+void UParticleModuleSourceMovement::PostEditChangeProperty(FPropertyChangedEvent& PropertyChangedEvent)
+{
+	InitializeDefaults();
+	Super::PostEditChangeProperty(PropertyChangedEvent);
+}
+#endif // WITH_EDITOR
+
 void UParticleModuleSourceMovement::FinalUpdate(FParticleEmitterInstance* Owner, int32 Offset, float DeltaTime)
 {
 	Super::FinalUpdate(Owner, Offset, DeltaTime);
@@ -1142,10 +1158,9 @@ UParticleModuleMeshRotation::UParticleModuleMeshRotation(const FObjectInitialize
 	bInheritParent = false;
 }
 
-void UParticleModuleMeshRotation::PostInitProperties()
+void UParticleModuleMeshRotation::InitializeDefaults()
 {
-	Super::PostInitProperties();
-	if (!HasAnyFlags(RF_ClassDefaultObject | RF_NeedLoad))
+	if (!StartRotation.Distribution)
 	{
 		UDistributionVectorUniform* DistributionStartRotation = NewObject<UDistributionVectorUniform>(this, TEXT("DistributionStartRotation"));
 		DistributionStartRotation->Min = FVector(0.0f, 0.0f, 0.0f);
@@ -1154,6 +1169,24 @@ void UParticleModuleMeshRotation::PostInitProperties()
 	}
 }
 
+
+void UParticleModuleMeshRotation::PostInitProperties()
+{
+	Super::PostInitProperties();
+	if (!HasAnyFlags(RF_ClassDefaultObject | RF_NeedLoad))
+	{
+		InitializeDefaults();
+	}
+}
+
+#if WITH_EDITOR
+void UParticleModuleMeshRotation::PostEditChangeProperty(FPropertyChangedEvent& PropertyChangedEvent)
+{
+	InitializeDefaults();
+	Super::PostEditChangeProperty(PropertyChangedEvent);
+}
+#endif // WITH_EDITOR
+
 void UParticleModuleMeshRotation::Spawn(FParticleEmitterInstance* Owner, int32 Offset, float SpawnTime, FBaseParticle* ParticleBase)
 {
 	SpawnEx(Owner, Offset, SpawnTime, NULL, ParticleBase);
@@ -1240,10 +1273,9 @@ UParticleModuleMeshRotationRate::UParticleModuleMeshRotationRate(const FObjectIn
 	bUpdateModule = true;
 }
 
-void UParticleModuleMeshRotationRate::PostInitProperties()
+void UParticleModuleMeshRotationRate::InitializeDefaults()
 {
-	Super::PostInitProperties();
-	if (!HasAnyFlags(RF_ClassDefaultObject | RF_NeedLoad))
+	if (!StartRotationRate.Distribution)
 	{
 		UDistributionVectorUniform* DistributionStartRotationRate = NewObject<UDistributionVectorUniform>(this, TEXT("DistributionStartRotationRate"));
 		DistributionStartRotationRate->Min = FVector(0.0f, 0.0f, 0.0f);
@@ -1252,6 +1284,23 @@ void UParticleModuleMeshRotationRate::PostInitProperties()
 	}
 }
 
+void UParticleModuleMeshRotationRate::PostInitProperties()
+{
+	Super::PostInitProperties();
+	if (!HasAnyFlags(RF_ClassDefaultObject | RF_NeedLoad))
+	{
+		InitializeDefaults();
+	}
+}
+
+#if WITH_EDITOR
+void UParticleModuleMeshRotationRate::PostEditChangeProperty(FPropertyChangedEvent& PropertyChangedEvent)
+{
+	InitializeDefaults();
+	Super::PostEditChangeProperty(PropertyChangedEvent);
+}
+#endif // WITH_EDITOR
+
 void UParticleModuleMeshRotationRate::Spawn(FParticleEmitterInstance* Owner, int32 Offset, float SpawnTime, FBaseParticle* ParticleBase)
 {
 	SpawnEx(Owner, Offset, SpawnTime, NULL, ParticleBase);
@@ -1334,16 +1383,32 @@ UParticleModuleMeshRotationRateMultiplyLife::UParticleModuleMeshRotationRateMult
 	bSpawnModule = true;
 	bUpdateModule = true;
 }
+void UParticleModuleMeshRotationRateMultiplyLife::InitializeDefaults()
+{
+	if (!LifeMultiplier.Distribution)
+	{
+		LifeMultiplier.Distribution = NewObject<UDistributionVectorConstant>(this, TEXT("DistributionLifeMultiplier"));
+	}
+}
 
 void UParticleModuleMeshRotationRateMultiplyLife::PostInitProperties()
 {
 	Super::PostInitProperties();
 	if (!HasAnyFlags(RF_ClassDefaultObject | RF_NeedLoad))
 	{
-		LifeMultiplier.Distribution = NewObject<UDistributionVectorConstant>(this, TEXT("DistributionLifeMultiplier"));
+		InitializeDefaults();
 	}
 }
 
+#if WITH_EDITOR
+void UParticleModuleMeshRotationRateMultiplyLife::PostEditChangeProperty(FPropertyChangedEvent& PropertyChangedEvent)
+{
+	InitializeDefaults();
+	Super::PostEditChangeProperty(PropertyChangedEvent);
+}
+#endif // WITH_EDITOR
+
+
 void UParticleModuleMeshRotationRateMultiplyLife::Spawn(FParticleEmitterInstance* Owner, int32 Offset, float SpawnTime, FBaseParticle* ParticleBase)
 {
 	const int32 MeshRotationOffset = Owner->GetMeshRotationOffset();
@@ -1394,15 +1459,31 @@ UParticleModuleMeshRotationRateOverLife::UParticleModuleMeshRotationRateOverLife
 	bUpdateModule = true;
 }
 
+void UParticleModuleMeshRotationRateOverLife::InitializeDefaults()
+{
+	if (!RotRate.Distribution)
+	{
+		RotRate.Distribution = NewObject<UDistributionVectorConstantCurve>(this, TEXT("DistributionRotRate"));
+	}
+}
+
 void UParticleModuleMeshRotationRateOverLife::PostInitProperties()
 {
 	Super::PostInitProperties();
 	if (!HasAnyFlags(RF_ClassDefaultObject | RF_NeedLoad))
 	{
-		RotRate.Distribution = NewObject<UDistributionVectorConstantCurve>(this, TEXT("DistributionRotRate"));
+		InitializeDefaults();
 	}
 }
 
+#if WITH_EDITOR
+void UParticleModuleMeshRotationRateOverLife::PostEditChangeProperty(FPropertyChangedEvent& PropertyChangedEvent)
+{
+	InitializeDefaults();
+	Super::PostEditChangeProperty(PropertyChangedEvent);
+}
+#endif // WITH_EDITOR
+
 void UParticleModuleMeshRotationRateOverLife::Spawn(FParticleEmitterInstance* Owner, int32 Offset, float SpawnTime, FBaseParticle* ParticleBase)
 {
 	const int32 MeshRotationOffset = Owner->GetMeshRotationOffset();
diff --git a/Engine/Source/Runtime/Engine/Private/Particles/ParticleModules_VectorField.cpp b/Engine/Source/Runtime/Engine/Private/Particles/ParticleModules_VectorField.cpp
--- a/Engine/Source/Runtime/Engine/Private/Particles/ParticleModules_VectorField.cpp
+++ b/Engine/Source/Runtime/Engine/Private/Particles/ParticleModules_VectorField.cpp
@@ -43,16 +43,31 @@ UParticleModuleVectorFieldScale::UParticleModuleVectorFieldScale(const FObjectIn
 {
 }
 
-void UParticleModuleVectorFieldScale::PostInitProperties()
+void UParticleModuleVectorFieldScale::InitializeDefaults()
 {
-	Super::PostInitProperties();
-	if (!HasAnyFlags(RF_ClassDefaultObject | RF_NeedLoad))
+	if (!VectorFieldScale)
 	{
 		UDistributionFloatConstant* DistributionVectorFieldScale = NewObject<UDistributionFloatConstant>(this, TEXT("DistributionVectorFieldScale"));
 		DistributionVectorFieldScale->Constant = 1.0f;
 		VectorFieldScale = DistributionVectorFieldScale;
 	}
 }
+void UParticleModuleVectorFieldScale::PostInitProperties()
+{
+	Super::PostInitProperties();
+	if (!HasAnyFlags(RF_ClassDefaultObject | RF_NeedLoad))
+	{
+		InitializeDefaults();
+	}
+}
+
+#if WITH_EDITOR
+void UParticleModuleVectorFieldScale::PostEditChangeProperty(FPropertyChangedEvent& PropertyChangedEvent)
+{
+	InitializeDefaults();
+	Super::PostEditChangeProperty(PropertyChangedEvent);
+}
+#endif // WITH_EDITOR
 
 void UParticleModuleVectorFieldScale::CompileModule(FParticleEmitterBuildInfo& EmitterInfo)
 {
@@ -84,10 +99,9 @@ UParticleModuleVectorFieldScaleOverLife::UParticleModuleVectorFieldScaleOverLife
 {
 }
 
-void UParticleModuleVectorFieldScaleOverLife::PostInitProperties()
+void UParticleModuleVectorFieldScaleOverLife::InitializeDefaults()
 {
-	Super::PostInitProperties();
-	if (!HasAnyFlags(RF_ClassDefaultObject | RF_NeedLoad))
+	if (!VectorFieldScaleOverLife)
 	{
 		UDistributionFloatConstant* DistributionVectorFieldScaleOverLife = NewObject<UDistributionFloatConstant>(this, TEXT("DistributionVectorFieldScaleOverLife"));
 		DistributionVectorFieldScaleOverLife->Constant = 1.0f;
@@ -95,6 +109,23 @@ void UParticleModuleVectorFieldScaleOverLife::PostInitProperties()
 	}
 }
 
+void UParticleModuleVectorFieldScaleOverLife::PostInitProperties()
+{
+	Super::PostInitProperties();
+	if (!HasAnyFlags(RF_ClassDefaultObject | RF_NeedLoad))
+	{
+		InitializeDefaults();
+	}
+}
+
+#if WITH_EDITOR
+void UParticleModuleVectorFieldScaleOverLife::PostEditChangeProperty(FPropertyChangedEvent& PropertyChangedEvent)
+{
+	InitializeDefaults();
+	Super::PostEditChangeProperty(PropertyChangedEvent);
+}
+#endif // WITH_EDITOR
+
 void UParticleModuleVectorFieldScaleOverLife::CompileModule(FParticleEmitterBuildInfo& EmitterInfo)
 {
 	EmitterInfo.VectorFieldScaleOverLife.ScaleByDistribution(VectorFieldScaleOverLife);
