Move modulated shadow parameters to directional light.

diff --git a/Engine/Source/Runtime/Engine/Classes/Components/DirectionalLightComponent.h b/Engine/Source/Runtime/Engine/Classes/Components/DirectionalLightComponent.h
--- a/Engine/Source/Runtime/Engine/Classes/Components/DirectionalLightComponent.h
+++ b/Engine/Source/Runtime/Engine/Classes/Components/DirectionalLightComponent.h
@@ -119,6 +119,18 @@ class UDirectionalLightComponent : public ULightComponent
 	UPROPERTY(EditAnywhere, BlueprintReadOnly, Category=Light, meta=(ShowOnlyInnerProperties))
 	struct FLightmassDirectionalLightSettings LightmassSettings;
 
+	/**
+	* Whether the light should cast modulated shadows from dynamic objects (mobile only).  Also requires Cast Shadows to be set to True.
+	**/
+	UPROPERTY(EditAnywhere, BlueprintReadOnly, Category = Light, AdvancedDisplay)
+	uint32 bCastModulatedShadows : 1;
+
+	/**
+	* Color to modulate against the scene color when rendering modulated shadows. (mobile only)
+	**/
+	UPROPERTY(BlueprintReadOnly, interp, Category = Light, meta = (HideAlphaChannel), AdvancedDisplay)
+	FColor ModulatedShadowColor;
+
 	UPROPERTY(EditAnywhere, AdvancedDisplay, BlueprintReadOnly, Category=Light, meta=(DisplayName = "Atmosphere Sun Light"))
 	uint32 bUsedAsAtmosphereSunLight : 1;
 
diff --git a/Engine/Source/Runtime/Engine/Classes/Components/LightComponentBase.h b/Engine/Source/Runtime/Engine/Classes/Components/LightComponentBase.h
--- a/Engine/Source/Runtime/Engine/Classes/Components/LightComponentBase.h
+++ b/Engine/Source/Runtime/Engine/Classes/Components/LightComponentBase.h
@@ -60,18 +60,6 @@ class ENGINE_API ULightComponentBase : public USceneComponent
 	UPROPERTY(EditAnywhere, BlueprintReadOnly, Category=Light, AdvancedDisplay)
 	uint32 CastDynamicShadows:1;
 
-	/**
-	* Whether the light should cast modulated shadows from dynamic objects (mobile only).  Also requires Cast Shadows to be set to True.
-	**/
-	UPROPERTY(EditAnywhere, BlueprintReadOnly, Category = Light, AdvancedDisplay)
-	uint32 bCastModulatedShadows : 1;
-
-	/**
-	* Color to modulate against the scene color when rendering modulated shadows. (mobile only)
-	**/
-	UPROPERTY(BlueprintReadOnly, interp, Category = Light, meta = (HideAlphaChannel), AdvancedDisplay)
-	FColor ModulatedShadowColor;
-
 	/** Whether the light affects translucency or not.  Disabling this can save GPU time when there are many small lights. */
 	UPROPERTY(EditAnywhere, BlueprintReadOnly, Category=Light, AdvancedDisplay)
 	uint32 bAffectTranslucentLighting:1;
diff --git a/Engine/Source/Runtime/Engine/Private/Components/DirectionalLightComponent.cpp b/Engine/Source/Runtime/Engine/Private/Components/DirectionalLightComponent.cpp
--- a/Engine/Source/Runtime/Engine/Private/Components/DirectionalLightComponent.cpp
+++ b/Engine/Source/Runtime/Engine/Private/Components/DirectionalLightComponent.cpp
@@ -127,6 +127,9 @@ public:
 			FarShadowDistance = Component->FarShadowDistance;
 			FarShadowCascadeCount = Component->FarShadowCascadeCount;
 		}
+
+		bCastModulatedShadows = Component->bCastModulatedShadows;
+		ModulatedShadowColor = FLinearColor(Component->ModulatedShadowColor);
 	}
 
 	void UpdateLightShaftOverrideDirection_GameThread(const UDirectionalLightComponent* Component)
@@ -680,6 +683,8 @@ UDirectionalLightComponent::UDirectionalLightComponent(const FObjectInitializer&
 	IndirectLightingIntensity = 1.0f;
 	CastTranslucentShadows = true;
 	bUseInsetShadowsForMovableObjects = true;
+
+	ModulatedShadowColor = FColor(128, 128, 128);
 }
 
 #if WITH_EDITOR
@@ -758,6 +763,11 @@ bool UDirectionalLightComponent::CanEditChange(const UProperty* InProperty) cons
 		{
 			return bEnableLightShaftOcclusion;
 		}
+
+		if (PropertyName == GET_MEMBER_NAME_STRING_CHECKED(UDirectionalLightComponent, ModulatedShadowColor))
+		{
+			return bCastModulatedShadows;
+		}
 	}
 
 	return Super::CanEditChange(InProperty);
diff --git a/Engine/Source/Runtime/Engine/Private/Components/LightComponent.cpp b/Engine/Source/Runtime/Engine/Private/Components/LightComponent.cpp
--- a/Engine/Source/Runtime/Engine/Private/Components/LightComponent.cpp
+++ b/Engine/Source/Runtime/Engine/Private/Components/LightComponent.cpp
@@ -220,7 +220,6 @@ FLightSceneProxy::FLightSceneProxy(const ULightComponent* InLightComponent)
 	, bCastStaticShadow(InLightComponent->CastShadows && InLightComponent->CastStaticShadows)
 	, bCastTranslucentShadows(InLightComponent->CastTranslucentShadows)
 	, bCastShadowsFromCinematicObjectsOnly(InLightComponent->bCastShadowsFromCinematicObjectsOnly)
-	, bCastModulatedShadows(InLightComponent->bCastModulatedShadows)
 	, bAffectTranslucentLighting(InLightComponent->bAffectTranslucentLighting)
 	, bUsedAsAtmosphereSunLight(InLightComponent->IsUsedAsAtmosphereSunLight())
 	, bAffectDynamicIndirectLighting(InLightComponent->bAffectDynamicIndirectLighting)
@@ -233,6 +232,7 @@ FLightSceneProxy::FLightSceneProxy(const ULightComponent* InLightComponent)
 	, StatId(InLightComponent->GetStatID(true))
 	, FarShadowDistance(0)
 	, FarShadowCascadeCount(0)
+	, bCastModulatedShadows(false)
 {
 	// Brightness in Lumens
 	float LightBrightness = InLightComponent->ComputeLightBrightness();
@@ -268,8 +268,6 @@ FLightSceneProxy::FLightSceneProxy(const ULightComponent* InLightComponent)
 	LightFunctionDisabledBrightness = LightComponent->DisabledBrightness;
 
 	StaticShadowDepthMap = &LightComponent->StaticShadowDepthMap;
-
-	ModulatedShadowColor = FLinearColor(LightComponent->ModulatedShadowColor);
 }
 
 bool FLightSceneProxy::ShouldCreatePerObjectShadowsForDynamicObjects() const
@@ -303,7 +301,6 @@ ULightComponentBase::ULightComponentBase(const FObjectInitializer& ObjectInitial
 	Brightness_DEPRECATED = 3.1415926535897932f;
 	Intensity = 3.1415926535897932f;
 	LightColor = FColor::White;
-	ModulatedShadowColor = FColor(128, 128, 128);
 	bAffectsWorld = true;
 	CastShadows = true;
 	CastStaticShadows = true;
@@ -519,11 +516,6 @@ bool ULightComponent::CanEditChange(const UProperty* InProperty) const
 		{
 			return bUseTemperature;
 		}
-
-		if (PropertyName == GET_MEMBER_NAME_STRING_CHECKED(ULightComponent, ModulatedShadowColor))
-		{
-			return bCastModulatedShadows;
-		}
 	}
 
 	return Super::CanEditChange(InProperty);
@@ -581,6 +573,8 @@ void ULightComponent::PostEditChangeProperty(FPropertyChangedEvent& PropertyChan
 		PropertyName != GET_MEMBER_NAME_STRING_CHECKED(UDirectionalLightComponent, OcclusionMaskDarkness) &&
 		PropertyName != GET_MEMBER_NAME_STRING_CHECKED(UDirectionalLightComponent, OcclusionDepthRange) &&
 		PropertyName != GET_MEMBER_NAME_STRING_CHECKED(UDirectionalLightComponent, LightShaftOverrideDirection) &&
+		PropertyName != GET_MEMBER_NAME_STRING_CHECKED(UDirectionalLightComponent, bCastModulatedShadows) &&
+		PropertyName != GET_MEMBER_NAME_STRING_CHECKED(UDirectionalLightComponent, ModulatedShadowColor) &&
 		// Properties that should only unbuild lighting for a Static light (can be changed dynamically on a Stationary light)
 		(PropertyName != GET_MEMBER_NAME_STRING_CHECKED(ULightComponent, Intensity) || Mobility == EComponentMobility::Static) &&
 		(PropertyName != GET_MEMBER_NAME_STRING_CHECKED(ULightComponent, LightColor) || Mobility == EComponentMobility::Static) &&
diff --git a/Engine/Source/Runtime/Engine/Public/SceneManagement.h b/Engine/Source/Runtime/Engine/Public/SceneManagement.h
--- a/Engine/Source/Runtime/Engine/Public/SceneManagement.h
+++ b/Engine/Source/Runtime/Engine/Public/SceneManagement.h
@@ -918,9 +918,6 @@ protected:
 
 	const uint32 bCastShadowsFromCinematicObjectsOnly : 1;
 
-	/** Whether the light will cast modulated shadows when using the forward renderer (mobile). */
-	const uint32 bCastModulatedShadows : 1;
-
 	/** Whether the light affects translucency or not.  Disabling this can save GPU time when there are many small lights. */
 	const uint32 bAffectTranslucentLighting : 1;
 
@@ -954,6 +951,9 @@ protected:
 	/** Only for whole scene directional lights, 0: no FarShadowCascades, otherwise the count of cascades between WholeSceneDynamicShadowRadius and FarShadowDistance that are covered by distant shadow cascades. */
 	uint32 FarShadowCascadeCount;
 
+	/** Whether the light will cast modulated shadows when using the forward renderer (mobile). */
+	uint32 bCastModulatedShadows : 1;
+
 	/** Modulated shadow color. */
 	FLinearColor ModulatedShadowColor;
 	/**
