Fix networked client case for gameinstance properly loading starting level while waiting for pending connection

diff --git a/Engine/Source/Runtime/Engine/Private/GameInstance.cpp b/Engine/Source/Runtime/Engine/Private/GameInstance.cpp
--- a/Engine/Source/Runtime/Engine/Private/GameInstance.cpp
+++ b/Engine/Source/Runtime/Engine/Private/GameInstance.cpp
@@ -323,10 +323,11 @@ void UGameInstance::StartGameInstance()
 		UE_LOG(LogLoad, Error, TEXT("%s"), *FString::Printf(TEXT("Failed to enter %s: %s. Please check the log for errors."), *URL.Map, *Error));
 
 		// the map specified on the command-line couldn't be loaded.  ask the user if we should load the default map instead
-		if (FCString::Stricmp(*URL.Map, *DefaultMap) != 0)
+		if (FCString::Stricmp(Parm, *DefaultMap) != 0)
 		{
 			const FText Message = FText::Format(NSLOCTEXT("Engine", "MapNotFound", "The map specified on the commandline '{0}' could not be found. Would you like to load the default map instead?"), FText::FromString(URL.Map));
-			if (FMessageDialog::Open(EAppMsgType::OkCancel, Message) != EAppReturnType::Ok)
+			if (   FCString::Stricmp(*URL.Map, *DefaultMap) != 0  
+				&& FMessageDialog::Open(EAppMsgType::OkCancel, Message) != EAppReturnType::Ok)
 			{
 				// user canceled (maybe a typo while attempting to run a commandlet)
 				FPlatformMisc::RequestExit(false);
