Fixed shadow variable warnings in AnimCompositeBase.cpp.

diff --git a/Engine/Source/Runtime/Engine/Private/Animation/AnimCompositeBase.cpp b/Engine/Source/Runtime/Engine/Private/Animation/AnimCompositeBase.cpp
--- a/Engine/Source/Runtime/Engine/Private/Animation/AnimCompositeBase.cpp
+++ b/Engine/Source/Runtime/Engine/Private/Animation/AnimCompositeBase.cpp
@@ -113,8 +113,8 @@ void FAnimSegment::GetAnimNotifiesFromTrackPositions(const float& PreviousTrackP
 				// If our time left is shorter than time to end point, no problem. End there.
 				if( FMath::Abs(TrackTimeToGo) < FMath::Abs(TrackTimeToAnimEndPoint) )
 				{
-					const float AnimPlayRate = ValidPlayRate * (bTrackPlayingBackwards ? -1.f : 1.f);
-					const float AnimEndPosition = (TrackTimeToGo * AnimPlayRate) + AnimStartPosition;
+					const float PlayRate = ValidPlayRate * (bTrackPlayingBackwards ? -1.f : 1.f);
+					const float AnimEndPosition = (TrackTimeToGo * PlayRate) + AnimStartPosition;
 					AnimSequence->GetAnimNotifiesFromDeltaPositions(AnimStartPosition, AnimEndPosition, OutActiveNotifies);
 					break;
 				}
@@ -185,8 +185,8 @@ void FAnimSegment::GetRootMotionExtractionStepsForTrackRange(TArray<FRootMotionE
 				// If our time left is shorter than time to end point, no problem. End there.
 				if( FMath::Abs(TrackTimeToGo) < FMath::Abs(TrackTimeToAnimEndPoint) )
 				{
-					const float AnimPlayRate = ValidPlayRate * (bTrackPlayingBackwards ? -1.f : 1.f);
-					const float AnimEndPosition = (TrackTimeToGo * AnimPlayRate) + AnimStartPosition;
+					const float PlayRate = ValidPlayRate * (bTrackPlayingBackwards ? -1.f : 1.f);
+					const float AnimEndPosition = (TrackTimeToGo * PlayRate) + AnimStartPosition;
 					RootMotionExtractionSteps.Add(FRootMotionExtractionStep(AnimSequence, AnimStartPosition, AnimEndPosition));
 					break;
 				}
