Fixed shadow variable warning in FSceneView::EndFinalPostprocessSettings()

diff --git a/Engine/Source/Runtime/Engine/Private/SceneView.cpp b/Engine/Source/Runtime/Engine/Private/SceneView.cpp
--- a/Engine/Source/Runtime/Engine/Private/SceneView.cpp
+++ b/Engine/Source/Runtime/Engine/Private/SceneView.cpp
@@ -1375,14 +1375,14 @@ void FSceneView::EndFinalPostprocessSettings(const FSceneViewInitOptions& ViewIn
 
 		// Window full screen mode with upscaling
 		bool bFullscreen = false;
-		bool bIsSceneCapture = false;
+		bool bSceneCapture = false;
 		if (ViewInitOptions.ViewFamily && ViewInitOptions.ViewFamily->Views.Num() == 1)
 		{
 			auto* View = ViewInitOptions.ViewFamily->Views[0];
-			bIsSceneCapture = View ? View->bIsSceneCapture : false;
+			bSceneCapture = View ? View->bIsSceneCapture : false;
 		}
 
-		if (!bIsSceneCapture && GEngine && GEngine->GameViewport && GEngine->GameViewport->GetWindow().IsValid())
+		if (!bSceneCapture && GEngine && GEngine->GameViewport && GEngine->GameViewport->GetWindow().IsValid())
 		{
 			bFullscreen = GEngine->GameViewport->GetWindow()->GetWindowMode() != EWindowMode::Windowed;
 		}
