Turn warnings back on for simulating without collision. This is too helpful to turn off. Advanced users will just have to work around it. We can re-evaluate once we hear of an actual usecase.

#UE-17697

diff --git a/Engine/Source/Runtime/Engine/Private/Components/PrimitiveComponent.cpp b/Engine/Source/Runtime/Engine/Private/Components/PrimitiveComponent.cpp
--- a/Engine/Source/Runtime/Engine/Private/Components/PrimitiveComponent.cpp
+++ b/Engine/Source/Runtime/Engine/Private/Components/PrimitiveComponent.cpp
@@ -896,6 +896,12 @@ bool UPrimitiveComponent::ShouldCreatePhysicsState() const
 	bool bShouldCreatePhysicsState = IsRegistered() && (bAlwaysCreatePhysicsState || IsCollisionEnabled());
 
 #if WITH_EDITOR
+	if (BodyInstance.bSimulatePhysics && (GetCollisionEnabled() == ECollisionEnabled::NoCollision || GetCollisionEnabled() == ECollisionEnabled::QueryOnly))
+	{
+		FMessageLog("PIE").Warning(FText::Format(LOCTEXT("InvalidSimulateOptions", "Invalid Simulate Options: Body ({0}) is set to simulate physics but Collision Enabled is incompatible"),
+			FText::FromString(GetReadableName())));
+	}
+
 	// if it shouldn't create physics state, but if world wants to enable trace collision for components, allow it
 	if (!bShouldCreatePhysicsState && World && World->bEnableTraceCollision)
 	{
