Fix incorrectly updating mass and sim properties when non simulated

#UE-17705

diff --git a/Engine/Source/Runtime/Engine/Private/PhysicsEngine/BodyInstance.cpp b/Engine/Source/Runtime/Engine/Private/PhysicsEngine/BodyInstance.cpp
--- a/Engine/Source/Runtime/Engine/Private/PhysicsEngine/BodyInstance.cpp
+++ b/Engine/Source/Runtime/Engine/Private/PhysicsEngine/BodyInstance.cpp
@@ -4417,13 +4417,13 @@ void FBodyInstance::InitDynamicProperties_AssumesLocked()
 {
 	if(PxRigidDynamic* RigidActor = GetPxRigidDynamic_AssumesLocked())
 	{
-		UpdateMassProperties();
-		UpdateDampingProperties();
-		SetMaxAngularVelocity(MaxAngularVelocity, false);
-		SetMaxDepenetrationVelocity(bOverrideMaxDepenetrationVelocity ? MaxDepenetrationVelocity : UPhysicsSettings::Get()->MaxDepenetrationVelocity);
-
-		if(ShouldInstanceSimulatingPhysics())
+		if (ShouldInstanceSimulatingPhysics())
 		{
+			UpdateMassProperties();
+			UpdateDampingProperties();
+			SetMaxAngularVelocity(MaxAngularVelocity, false);
+			SetMaxDepenetrationVelocity(bOverrideMaxDepenetrationVelocity ? MaxDepenetrationVelocity : UPhysicsSettings::Get()->MaxDepenetrationVelocity);
+
 			RigidActor->setLinearVelocity(U2PVector(InitialLinearVelocity));
 		}
 
