Improved Data Table handling of User Defined Enums

UE-17445 - Data Table Enums Fail to import correctly

The Data Table editor UI will now show you the friendly name of the enums in its grid panel (matching what the property panel shows).

The Data Table import will now fallback to trying to import using a friendly name if no matching enum entry was found for the given name. This avoids the need to use the standard User Defined Enum names of "NewEnumeration0", "NewEnumeration1", etc, but does make the CSV/JSON data less robust against renaming the enum entries. CSV/JSON exporting will still use the internal names.

diff --git a/Engine/Source/Editor/UnrealEd/Private/DataTableEditorUtils.cpp b/Engine/Source/Editor/UnrealEd/Private/DataTableEditorUtils.cpp
--- a/Engine/Source/Editor/UnrealEd/Private/DataTableEditorUtils.cpp
+++ b/Engine/Source/Editor/UnrealEd/Private/DataTableEditorUtils.cpp
@@ -233,7 +233,7 @@ void FDataTableEditorUtils::CacheDataTableForEditing(const UDataTable* DataTable
 				const UProperty* Prop = StructProps[ColumnIndex];
 				FDataTableEditorColumnHeaderDataPtr CachedColumnData = OutAvailableColumns[ColumnIndex];
 
-				const FText CellText = FText::FromString(DataTableUtils::GetPropertyValueAsString(Prop, RowData));
+				const FText CellText = DataTableUtils::GetPropertyValueAsText(Prop, RowData);
 				CachedRowData->CellData.Add(CellText);
 
 				const float CellWidth = FontMeasure->Measure(CellText, CellTextStyle.Font).X + CellPadding;
diff --git a/Engine/Source/Runtime/Engine/Classes/Engine/DataTable.h b/Engine/Source/Runtime/Engine/Classes/Engine/DataTable.h
--- a/Engine/Source/Runtime/Engine/Classes/Engine/DataTable.h
+++ b/Engine/Source/Runtime/Engine/Classes/Engine/DataTable.h
@@ -2,12 +2,10 @@
 
 #pragma once
 
+#include "DataTableUtils.h" // Needed here for LogDataTable
 #include "DataTable.generated.h"
 
 
-ENGINE_API DECLARE_LOG_CATEGORY_EXTERN(LogDataTable, Log, All);
-
-
 // forward declare JSON writer
 template <class CharType>
 struct TPrettyJsonPrintPolicy;
diff --git a/Engine/Source/Runtime/Engine/Private/DataTable.cpp b/Engine/Source/Runtime/Engine/Private/DataTable.cpp
--- a/Engine/Source/Runtime/Engine/Private/DataTable.cpp
+++ b/Engine/Source/Runtime/Engine/Private/DataTable.cpp
@@ -2,13 +2,10 @@
 
 #include "EnginePrivate.h"
 #include "Engine/DataTable.h"
-#include "DataTableUtils.h"
 #include "DataTableCSV.h"
 #include "DataTableJSON.h"
 #include "EditorFramework/AssetImportData.h"
 
-DEFINE_LOG_CATEGORY(LogDataTable);
-
 ENGINE_API const FString FDataTableRowHandle::Unknown(TEXT("UNKNOWN"));
 ENGINE_API const FString FDataTableCategoryHandle::Unknown(TEXT("UNKNOWN"));
 
diff --git a/Engine/Source/Runtime/Engine/Private/DataTableUtils.cpp b/Engine/Source/Runtime/Engine/Private/DataTableUtils.cpp
--- a/Engine/Source/Runtime/Engine/Private/DataTableUtils.cpp
+++ b/Engine/Source/Runtime/Engine/Private/DataTableUtils.cpp
@@ -3,12 +3,39 @@
 #include "EnginePrivate.h"
 #include "DataTableUtils.h"
 
+DEFINE_LOG_CATEGORY(LogDataTable);
+
 FString DataTableUtils::AssignStringToProperty(const FString& InString, const UProperty* InProp, uint8* InData)
 {
 	FStringOutputDevice ImportError;
-	if(InProp != NULL && IsSupportedTableProperty(InProp))
+	if(InProp && IsSupportedTableProperty(InProp))
 	{
-		InProp->ImportText(*InString, InProp->ContainerPtrToValuePtr<uint8>(InData), PPF_None, NULL, &ImportError);
+		FString StringToImport = InString;
+
+		const UByteProperty* EnumProp = Cast<const UByteProperty>(InProp);
+		if(EnumProp && EnumProp->Enum)
+		{
+			// Enum properties may use the friendly name in their import data, however the UPropertyByte::ImportText function will only accept the internal enum entry name
+			// Detect if we're using a friendly name for an entry, and if so, try and map it to the correct internal name before performing the import
+			const int32 EnumIndex = EnumProp->Enum->FindEnumIndex(*InString);
+			if(EnumIndex == INDEX_NONE)
+			{
+				// Couldn't find a match for the name we were given, try and find a match using the friendly names
+				for(int32 EnumEntryIndex = 0; EnumEntryIndex < EnumProp->Enum->NumEnums(); ++EnumEntryIndex)
+				{
+					const FText FriendlyEnumEntryName = EnumProp->Enum->GetEnumText(EnumEntryIndex);
+					if(FriendlyEnumEntryName.ToString() == InString)
+					{
+						// Get the corresponding internal name and warn the user that we're using this fallback
+						StringToImport = EnumProp->Enum->GetEnumName(EnumEntryIndex);
+						UE_LOG(LogDataTable, Warning, TEXT("Could not a find matching enum entry for '%s', but did find a matching display name. Will import using the enum entry corresponding to that display name ('%s')"), *InString, *StringToImport);
+						break;
+					}
+				}
+			}
+		}
+		
+		InProp->ImportText(*StringToImport, InProp->ContainerPtrToValuePtr<uint8>(InData), PPF_None, nullptr, &ImportError);
 	}
 
 	FString Error = ImportError;
@@ -17,11 +44,25 @@ FString DataTableUtils::AssignStringToProperty(const FString& InString, const UP
 
 FString DataTableUtils::GetPropertyValueAsString(const UProperty* InProp, uint8* InData)
 {
-	FString Result(TEXT(""));
+	FString Result;
+
+	if(InProp && IsSupportedTableProperty(InProp))
+	{
+		InProp->ExportText_InContainer(0, Result, InData, InData, nullptr, PPF_None);
+	}
+
+	return Result;
+}
+
+FText DataTableUtils::GetPropertyValueAsText(const UProperty* InProp, uint8* InData)
+{
+	FText Result;
 
-	if(InProp != NULL && IsSupportedTableProperty(InProp))
+	if(InProp && IsSupportedTableProperty(InProp))
 	{
-		InProp->ExportText_InContainer(0, Result, InData, InData, NULL, PPF_None);
+		FString ExportedString;
+		InProp->ExportText_InContainer(0, ExportedString, InData, InData, nullptr, PPF_PropertyWindow);
+		Result = FText::FromString(ExportedString);
 	}
 
 	return Result;
diff --git a/Engine/Source/Runtime/Engine/Public/DataTableUtils.h b/Engine/Source/Runtime/Engine/Public/DataTableUtils.h
--- a/Engine/Source/Runtime/Engine/Public/DataTableUtils.h
+++ b/Engine/Source/Runtime/Engine/Public/DataTableUtils.h
@@ -2,14 +2,19 @@
 
 #pragma once
 
+ENGINE_API DECLARE_LOG_CATEGORY_EXTERN(LogDataTable, Log, All);
+
 namespace DataTableUtils
 {
 	/** Util to assign a value (given as a string) to a struct property. */
 	ENGINE_API FString AssignStringToProperty(const FString& InString, const UProperty* InProp, uint8* InData);
 
-	/** Util to assign get a property as a string */
+	/** Util to get a property as a string */
 	ENGINE_API FString GetPropertyValueAsString(const UProperty* InProp, uint8* InData);
 
+	/** Util to get a property as text (this will use the display name of the value where available - use GetPropertyValueAsString if you need an internal identifier) */
+	ENGINE_API FText GetPropertyValueAsText(const UProperty* InProp, uint8* InData);
+
 	/** Util to get all property names from a struct */
 	ENGINE_API TArray<FName> GetStructPropertyNames(UStruct* InStruct);
 
