Don't call begin play on a component that has travelled seamlessly from another level

diff --git a/Engine/Source/Runtime/Engine/Private/Components/ActorComponent.cpp b/Engine/Source/Runtime/Engine/Private/Components/ActorComponent.cpp
--- a/Engine/Source/Runtime/Engine/Private/Components/ActorComponent.cpp
+++ b/Engine/Source/Runtime/Engine/Private/Components/ActorComponent.cpp
@@ -813,7 +813,7 @@ void UActorComponent::RegisterComponentWithWorld(UWorld* InWorld)
 
 	if (MyOwner && MyOwner->HasActorBegunPlay())
 	{
-		if (bWantsBeginPlay)
+		if (!bHasBegunPlay && bWantsBeginPlay)
 		{
 			BeginPlay();
 		}
