Make AudioListenerComponent a weak object pointer to prevent cases where it is referencing a component that no longer exists.

diff --git a/Engine/Source/Runtime/Engine/Classes/GameFramework/PlayerController.h b/Engine/Source/Runtime/Engine/Classes/GameFramework/PlayerController.h
--- a/Engine/Source/Runtime/Engine/Classes/GameFramework/PlayerController.h
+++ b/Engine/Source/Runtime/Engine/Classes/GameFramework/PlayerController.h
@@ -1295,7 +1295,7 @@ protected:
 	/** Whether to override the normal audio listener positioning method */
 	uint32 bOverrideAudioListener:1;
 	/** Component that is currently driving the audio listener position/orientation */
-	USceneComponent* AudioListenerComponent;
+	TWeakObjectPtr<USceneComponent> AudioListenerComponent;
 	/** Currently overridden location of audio listener */
 	FVector AudioListenerLocationOverride;
 	/** Currently overridden rotation of audio listener */
diff --git a/Engine/Source/Runtime/Engine/Private/PlayerController.cpp b/Engine/Source/Runtime/Engine/Private/PlayerController.cpp
--- a/Engine/Source/Runtime/Engine/Private/PlayerController.cpp
+++ b/Engine/Source/Runtime/Engine/Private/PlayerController.cpp
@@ -2503,10 +2503,11 @@ void APlayerController::GetAudioListenerPosition(FVector& OutLocation, FVector&
 
 	if (bOverrideAudioListener)
 	{
-		if (AudioListenerComponent != nullptr)
+		USceneComponent* ListenerComponent = AudioListenerComponent.Get();
+		if (ListenerComponent != nullptr)
 		{
-			ViewRotation = AudioListenerComponent->GetComponentRotation() + AudioListenerRotationOverride;
-			ViewLocation = AudioListenerComponent->GetComponentLocation() + ViewRotation.RotateVector(AudioListenerLocationOverride);
+			ViewRotation = ListenerComponent->GetComponentRotation() + AudioListenerRotationOverride;
+			ViewLocation = ListenerComponent->GetComponentLocation() + ViewRotation.RotateVector(AudioListenerLocationOverride);
 		}
 		else
 		{
