Fix it so simulating complexAsSimple gives a warning. Also compute mass data even if not simulating, as long as not ComplexAsSimple

diff --git a/Engine/Source/Runtime/Engine/Private/PhysicsEngine/BodyInstance.cpp b/Engine/Source/Runtime/Engine/Private/PhysicsEngine/BodyInstance.cpp
--- a/Engine/Source/Runtime/Engine/Private/PhysicsEngine/BodyInstance.cpp
+++ b/Engine/Source/Runtime/Engine/Private/PhysicsEngine/BodyInstance.cpp
@@ -2394,6 +2394,10 @@ void FBodyInstance::SetInstanceSimulatePhysics(bool bSimulate, bool bMaintainPhy
 			{
 				FMessageLog("PIE").Warning(FText::Format(LOCTEXT("SimPhysStatic", "Trying to simulate physics on ''{0}'' but it is static."),
 					FText::FromString(GetPathNameSafe(OwnerComponentInst))));
+			}else if(BodySetup.IsValid() && BodySetup->GetCollisionTraceFlag() == ECollisionTraceFlag::CTF_UseComplexAsSimple)
+			{
+				FMessageLog("PIE").Warning(FText::Format(LOCTEXT("SimComplexAsSimple", "Trying to simulate physics on ''{0}'' but it has ComplexAsSimple collision."),
+					FText::FromString(GetPathNameSafe(OwnerComponentInst))));
 			}
 		}
 #endif
@@ -4417,7 +4421,9 @@ void FBodyInstance::InitDynamicProperties_AssumesLocked()
 {
 	if(PxRigidDynamic* RigidActor = GetPxRigidDynamic_AssumesLocked())
 	{
-		if (ShouldInstanceSimulatingPhysics())
+		//A non simulated body may become simulated at runtime, so we need to compute its mass.
+		//However, this is not supported for complexAsSimple since a trimesh cannot itself be simulated, it can only be used for collision of other simple shapes.
+		if (BodySetup->GetCollisionTraceFlag() != ECollisionTraceFlag::CTF_UseComplexAsSimple)
 		{
 			UpdateMassProperties();
 			UpdateDampingProperties();
@@ -4425,6 +4431,19 @@ void FBodyInstance::InitDynamicProperties_AssumesLocked()
 			SetMaxDepenetrationVelocity(bOverrideMaxDepenetrationVelocity ? MaxDepenetrationVelocity : UPhysicsSettings::Get()->MaxDepenetrationVelocity);
 
 			RigidActor->setLinearVelocity(U2PVector(InitialLinearVelocity));
+		}else
+		{
+#if !(UE_BUILD_SHIPPING || UE_BUILD_TEST)
+			if (ShouldInstanceSimulatingPhysics())
+			{
+				if(UPrimitiveComponent* OwnerComponentInst = OwnerComponent.Get())
+				{
+					FMessageLog("PIE").Warning(FText::Format(LOCTEXT("SimComplexAsSimple", "Trying to simulate physics on ''{0}'' but it has ComplexAsSimple collision."),
+						FText::FromString(GetPathNameSafe(OwnerComponentInst))));
+				}
+				
+			}
+#endif
 		}
 
 		float SleepEnergyThresh = RigidActor->getSleepThreshold();
