#ue4 - Fix Character::CanJump() event no longer showing up as overrideable in the BP event menu.

UE-17783

diff --git a/Engine/Source/Runtime/Engine/Classes/GameFramework/Character.h b/Engine/Source/Runtime/Engine/Classes/GameFramework/Character.h
--- a/Engine/Source/Runtime/Engine/Classes/GameFramework/Character.h
+++ b/Engine/Source/Runtime/Engine/Classes/GameFramework/Character.h
@@ -198,7 +198,7 @@ typedef FBasedMovementInfo FRepRelativeMovement;
  * @see https://docs.unrealengine.com/latest/INT/Gameplay/Framework/Pawn/Character/
  */ 
 
-UCLASS(config=Game, BlueprintType, hideCategories=("Pawn|Character|InternalEvents"), meta=(ShortTooltip="A character is a type of Pawn that includes the ability to walk around."))
+UCLASS(config=Game, BlueprintType, meta=(ShortTooltip="A character is a type of Pawn that includes the ability to walk around."))
 class ENGINE_API ACharacter : public APawn
 {
 	GENERATED_BODY()
@@ -440,7 +440,7 @@ protected:
 	 * @Return Whether the character can jump in the current state. 
 	 */
 
-	UFUNCTION(BlueprintNativeEvent, Category="Pawn|Character|InternalEvents", meta=(DisplayName="CanJump"))
+	UFUNCTION(BlueprintNativeEvent, Category="Pawn|Character", meta=(DisplayName="CanJump"))
 	bool CanJumpInternal() const;
 	virtual bool CanJumpInternal_Implementation() const;
 
