Removed unused InactivePlayerArray variable from GameState. (UE-14073)

diff --git a/Engine/Source/Runtime/Engine/Classes/GameFramework/GameState.h b/Engine/Source/Runtime/Engine/Classes/GameFramework/GameState.h
--- a/Engine/Source/Runtime/Engine/Classes/GameFramework/GameState.h
+++ b/Engine/Source/Runtime/Engine/Classes/GameFramework/GameState.h
@@ -74,10 +74,6 @@ public:
 	UPROPERTY(BlueprintReadOnly, Category=GameState)
 	TArray<class APlayerState*> PlayerArray;
 
-	/** This list mirrors the GameMode's list of inactive PlayerState objects. */
-	UPROPERTY()
-	TArray<class APlayerState*> InactivePlayerArray;
-
 	/** GameMode class notification callback. */
 	UFUNCTION()
 	virtual void OnRep_GameModeClass();
diff --git a/Engine/Source/Runtime/Engine/Private/GameState.cpp b/Engine/Source/Runtime/Engine/Private/GameState.cpp
--- a/Engine/Source/Runtime/Engine/Private/GameState.cpp
+++ b/Engine/Source/Runtime/Engine/Private/GameState.cpp
@@ -135,21 +135,7 @@ void AGameState::AddPlayerState(APlayerState* PlayerState)
 	if (!PlayerState->bIsInactive)
 	{
 		// make sure no duplicates
-		for (int32 i=0; i<PlayerArray.Num(); i++)
-		{
-			if (PlayerArray[i] == PlayerState)
-				return;
-		}
-
-		PlayerArray.Add(PlayerState);
-	}
-	else
-	{
-		// Add once only
-		if (InactivePlayerArray.Find(PlayerState) == INDEX_NONE)
-		{
-			InactivePlayerArray.Add(PlayerState);
-		}
+		PlayerArray.AddUnique(PlayerState);
 	}
 }
 
