#ue4 - Fix analog input being broken for network clients.

- Non-zero values were being scaled up to 1 erroneously.

UE-17792

diff --git a/Engine/Source/Runtime/Engine/Private/Components/CharacterMovementComponent.cpp b/Engine/Source/Runtime/Engine/Private/Components/CharacterMovementComponent.cpp
--- a/Engine/Source/Runtime/Engine/Private/Components/CharacterMovementComponent.cpp
+++ b/Engine/Source/Runtime/Engine/Private/Components/CharacterMovementComponent.cpp
@@ -6559,7 +6559,8 @@ void UCharacterMovementComponent::MoveAutonomous
 	UpdateFromCompressedFlags(CompressedFlags);
 	CharacterOwner->CheckJumpInput(DeltaTime);
 
-	Acceleration = ScaleInputAcceleration( ConstrainInputAcceleration(NewAccel) );
+	Acceleration = ConstrainInputAcceleration(NewAccel);
+	Acceleration = Acceleration.GetClampedToMaxSize(GetMaxAcceleration());
 	AnalogInputModifier = ComputeAnalogInputModifier();
 	
 	PerformMovement(DeltaTime);
