Integrated PR #1289: Check if AController is not NULL (Thanks trdwll!)

diff --git a/Engine/Source/Runtime/Engine/Private/GameMode.cpp b/Engine/Source/Runtime/Engine/Private/GameMode.cpp
--- a/Engine/Source/Runtime/Engine/Private/GameMode.cpp
+++ b/Engine/Source/Runtime/Engine/Private/GameMode.cpp
@@ -1454,7 +1454,7 @@ bool AGameMode::CanSpectate_Implementation( APlayerController* Viewer, APlayerSt
 
 void AGameMode::ChangeName( AController* Other, const FString& S, bool bNameChange )
 {
-	if( !S.IsEmpty() )
+	if( Other && !S.IsEmpty() )
 	{
 		Other->PlayerState->SetPlayerName(S);
 	}
