Integrted PR #1319: Do not use GameMode::DefaultPawnClass directly (Thanks slonopotamus!)

Fixes log message potentially outputting misleading message with incorrect class name.

diff --git a/Engine/Source/Runtime/Engine/Private/GameMode.cpp b/Engine/Source/Runtime/Engine/Private/GameMode.cpp
--- a/Engine/Source/Runtime/Engine/Private/GameMode.cpp
+++ b/Engine/Source/Runtime/Engine/Private/GameMode.cpp
@@ -1341,10 +1341,11 @@ APawn* AGameMode::SpawnDefaultPawnFor_Implementation(AController* NewPlayer, AAc
 	FActorSpawnParameters SpawnInfo;
 	SpawnInfo.Instigator = Instigator;	
 	SpawnInfo.ObjectFlags |= RF_Transient;	// We never want to save default player pawns into a map
-	APawn* ResultPawn = GetWorld()->SpawnActor<APawn>(GetDefaultPawnClassForController(NewPlayer), StartLocation, StartRotation, SpawnInfo );
+	UClass* PawnClass = GetDefaultPawnClassForController(NewPlayer);
+	APawn* ResultPawn = GetWorld()->SpawnActor<APawn>(PawnClass, StartLocation, StartRotation, SpawnInfo );
 	if ( ResultPawn == NULL )
 	{
-		UE_LOG(LogGameMode, Warning, TEXT("Couldn't spawn Pawn of type %s at %s"), *GetNameSafe(DefaultPawnClass), *StartSpot->GetName());
+		UE_LOG(LogGameMode, Warning, TEXT("Couldn't spawn Pawn of type %s at %s"), *GetNameSafe(PawnClass), *StartSpot->GetName());
 	}
 	return ResultPawn;
 }
