#ue4 - Hide GetSocketQuaternion from BPs (mark it deprecated for BPs). We don't expose FQuat anywhere else in BPs.

https://answers.unrealengine.com/questions/223815/48-p1-can-get-quaternions-but-can-not-make-variabl.html

diff --git a/Engine/Source/Runtime/Engine/Classes/Components/SceneComponent.h b/Engine/Source/Runtime/Engine/Classes/Components/SceneComponent.h
--- a/Engine/Source/Runtime/Engine/Classes/Components/SceneComponent.h
+++ b/Engine/Source/Runtime/Engine/Classes/Components/SceneComponent.h
@@ -465,7 +465,7 @@ public:
 	 * @param InSocketName Name of the socket or the bone to get the transform 
 	 * @return Socket transform in world space if socket if found. Otherwise it will return component's transform in world space.
 	 */
-	UFUNCTION(BlueprintCallable, Category="Utilities|Transformation", meta=(Keywords="Bone"))
+	UFUNCTION(BlueprintCallable, Category="Utilities|Transformation", meta=(Keywords="Bone", DeprecatedFunction, DeprecationMessage="Use GetSocketRotation instead, Quat is not fully supported in blueprints."))
 	virtual FQuat GetSocketQuaternion(FName InSocketName) const;
 
 	/** 
