#ue4 - Make sure client and server perform movement with the same acceleration value after network rounding.

diff --git a/Engine/Source/Runtime/Engine/Private/Components/CharacterMovementComponent.cpp b/Engine/Source/Runtime/Engine/Private/Components/CharacterMovementComponent.cpp
--- a/Engine/Source/Runtime/Engine/Private/Components/CharacterMovementComponent.cpp
+++ b/Engine/Source/Runtime/Engine/Private/Components/CharacterMovementComponent.cpp
@@ -6059,9 +6059,11 @@ void UCharacterMovementComponent::ReplicateMoveToServer(float DeltaTime, const F
 		}
 	}
 
-	// Perform the move locally
-	Acceleration = NewMove->Acceleration;
+	// Acceleration should match what we send to the server, plus any other restrictions the server also enforces (see MoveAutonomous).
+	Acceleration = NewMove->Acceleration.GetClampedToMaxSize(GetMaxAcceleration());
 	AnalogInputModifier = ComputeAnalogInputModifier(); // recompute since acceleration may have changed.
+
+	// Perform the move locally
 	CharacterOwner->ClientRootMotionParams.Clear();
 	PerformMovement(NewMove->DeltaTime);
 
