Export GetPixelFormatEnum in ENGINE_API

diff --git a/Engine/Source/Runtime/Engine/Classes/Engine/Texture.h b/Engine/Source/Runtime/Engine/Classes/Engine/Texture.h
--- a/Engine/Source/Runtime/Engine/Classes/Engine/Texture.h
+++ b/Engine/Source/Runtime/Engine/Classes/Engine/Texture.h
@@ -777,7 +777,7 @@ public:
 	/**
 	 * Retrieves the pixel format enum for enum <-> string conversions.
 	 */
-	static class UEnum* GetPixelFormatEnum();
+	ENGINE_API static class UEnum* GetPixelFormatEnum();
 
 };
 
