Make sure we clear navigation collision buffers before gathering new data. Fixes cooking crash in KiteDemo

diff --git a/Engine/Source/Runtime/Engine/Classes/AI/Navigation/NavCollision.h b/Engine/Source/Runtime/Engine/Classes/AI/Navigation/NavCollision.h
--- a/Engine/Source/Runtime/Engine/Classes/AI/Navigation/NavCollision.h
+++ b/Engine/Source/Runtime/Engine/Classes/AI/Navigation/NavCollision.h
@@ -101,7 +101,7 @@ class ENGINE_API UNavCollision : public UObject
 	void GetNavigationModifier(FCompositeNavModifier& Modifier, const FTransform& LocalToWorld);
 
 	/** Read collisions data */
-	bool GatherCollision();
+	void GatherCollision();
 
 protected:
 	void ClearCollision();
diff --git a/Engine/Source/Runtime/Engine/Private/AI/Navigation/NavCollision.cpp b/Engine/Source/Runtime/Engine/Private/AI/Navigation/NavCollision.cpp
--- a/Engine/Source/Runtime/Engine/Private/AI/Navigation/NavCollision.cpp
+++ b/Engine/Source/Runtime/Engine/Private/AI/Navigation/NavCollision.cpp
@@ -157,8 +157,7 @@ void UNavCollision::Setup(UBodySetup* BodySetup)
 
 	// Make sure all are cleared before we start
 	ClearCollision(); 
-
-	bool bCalculated = false;
+		
 	if (ShouldUseConvexCollision())
 	{
 		// Find or create cooked navcollision data
@@ -171,30 +170,26 @@ void UNavCollision::Setup(UBodySetup* BodySetup)
 				// Create physics objects
 				FNavCollisionDataReader CookedDataReader(*FormatData, TriMeshCollision, ConvexCollision, ConvexShapeIndices);
 
-				bCalculated = true;
+				bHasConvexGeometry = true;
 			}
 		}
 		else
 		{
-			bCalculated = GatherCollision();
+			GatherCollision();
 		}
 	}
-
-	bHasConvexGeometry = bCalculated;
 }
 
-bool UNavCollision::GatherCollision()
+void UNavCollision::GatherCollision()
 {
-	bool bCalculated = false;
 	UStaticMesh* StaticMeshOuter = Cast<UStaticMesh>(GetOuter());
 	// get data from owner
 	if (StaticMeshOuter && StaticMeshOuter->BodySetup)
 	{
+		ClearCollision();
 		NavigationHelper::GatherCollision(StaticMeshOuter->BodySetup, this);
-		bCalculated = true;
+		bHasConvexGeometry = true;
 	}
-
-	return bCalculated;
 }
 
 void UNavCollision::ClearCollision()
@@ -244,7 +239,7 @@ void UNavCollision::GetNavigationModifier(FCompositeNavModifier& Modifier, const
 		// rebuild collision data if needed
 		if (!bHasConvexGeometry)
 		{
-			bHasConvexGeometry = GatherCollision();
+			GatherCollision();
 		}
 
 		if (ConvexCollision.VertexBuffer.Num() > 0)
