missed files for CL# 2605108

diff --git a/Engine/Source/Runtime/Engine/Classes/GameFramework/WorldSettings.h b/Engine/Source/Runtime/Engine/Classes/GameFramework/WorldSettings.h
--- a/Engine/Source/Runtime/Engine/Classes/GameFramework/WorldSettings.h
+++ b/Engine/Source/Runtime/Engine/Classes/GameFramework/WorldSettings.h
@@ -471,6 +471,7 @@ class ENGINE_API AWorldSettings : public AInfo, public IInterface_AssetUserData
 
 public:
 	// Begin UObject interface.
+	virtual void PostLoad() override;
 #if WITH_EDITOR
 	virtual void PostEditChangeProperty(FPropertyChangedEvent& PropertyChangedEvent) override;
 #endif // WITH_EDITOR
diff --git a/Engine/Source/Runtime/Engine/Private/WorldSettings.cpp b/Engine/Source/Runtime/Engine/Private/WorldSettings.cpp
--- a/Engine/Source/Runtime/Engine/Private/WorldSettings.cpp
+++ b/Engine/Source/Runtime/Engine/Private/WorldSettings.cpp
@@ -211,6 +211,19 @@ void AWorldSettings::RemoveUserDataOfClass(TSubclassOf<UAssetUserData> InUserDat
 	}
 }
 
+void AWorldSettings::PostLoad()
+{
+	Super::PostLoad();
+
+#if WITH_EDITOR
+	for (FHierarchicalSimplification Entry : HierarchicalLODSetup)
+	{
+		Entry.ProxySetting.PostLoadDeprecated();	
+	}
+#endif// WITH_EDITOR
+}
+
+
 #if WITH_EDITOR
 
 void AWorldSettings::CheckForErrors()
