fixing non editor build

diff --git a/Engine/Source/Runtime/Engine/Classes/Engine/Level.h b/Engine/Source/Runtime/Engine/Classes/Engine/Level.h
--- a/Engine/Source/Runtime/Engine/Classes/Engine/Level.h
+++ b/Engine/Source/Runtime/Engine/Classes/Engine/Level.h
@@ -230,8 +230,6 @@ private:
 	friend class FLightmassProcessor;
 };
 
-#if WITH_EDITOR
-
 USTRUCT()
 struct ENGINE_API FLevelSimplificationDetails
 {
@@ -299,8 +297,6 @@ struct ENGINE_API FLevelSimplificationDetails
 	bool operator == (const FLevelSimplificationDetails& Other) const;
 };
 
-#endif //WITH_EDITOR
-
 //
 // The level object.  Contains the level's actor list, BSP information, and brush list.
 // Every Level has a World as its Outer and can be used as the PersistentLevel, however,
