Add some range checks to UBodySetup::UpdateTriMeshVertices

diff --git a/Engine/Source/Runtime/Engine/Private/PhysicsEngine/BodySetup.cpp b/Engine/Source/Runtime/Engine/Private/PhysicsEngine/BodySetup.cpp
--- a/Engine/Source/Runtime/Engine/Private/PhysicsEngine/BodySetup.cpp
+++ b/Engine/Source/Runtime/Engine/Private/PhysicsEngine/BodySetup.cpp
@@ -865,9 +865,12 @@ void UBodySetup::UpdateTriMeshVertices(const TArray<FVector> & NewPositions)
 #if WITH_PHYSX
 	if (TriMeshes.Num())
 	{
-		
+		check(TriMeshes[0] != nullptr);
+		PxU32 PNumVerts = TriMeshes[0]->getNbVertices(); // Get num of verts we expect
 		PxVec3 * PNewPositions = TriMeshes[0]->getVerticesForModification();	//we only update the first trimesh. We assume this per poly case is not updating welded trimeshes
-		for (int32 i = 0; i < NewPositions.Num(); ++i)
+
+		int32 NumToCopy = FMath::Min<int32>(PNumVerts, NewPositions.Num()); // Make sure we don't write off end of array provided
+		for (int32 i = 0; i < NumToCopy; ++i)
 		{
 			PNewPositions[i] = U2PVector(NewPositions[i]);
 		}
