PR #1314: Update CurveBase.cpp (Contributed by [REDACTED_PERSON])

diff --git a/Engine/Source/Runtime/Engine/Private/CurveBase.cpp b/Engine/Source/Runtime/Engine/Private/CurveBase.cpp
--- a/Engine/Source/Runtime/Engine/Private/CurveBase.cpp
+++ b/Engine/Source/Runtime/Engine/Private/CurveBase.cpp
@@ -910,7 +910,7 @@ void UCurveBase::GetValueRange(float& MinValue, float& MaxValue) const
 			Curves[i].CurveToEdit->GetValueRange(CurveMin, CurveMax);
 
 			MinValue = FMath::Min(CurveMin, MinValue);
-			MaxValue = FMath::Min(CurveMax, MaxValue);
+			MaxValue = FMath::Max(CurveMax, MaxValue);
 		}
 	}
 }
