Avoiding vertex buffer copy during navigation collision exporting

diff --git a/Engine/Source/Runtime/Engine/Private/AI/Navigation/NavCollision.cpp b/Engine/Source/Runtime/Engine/Private/AI/Navigation/NavCollision.cpp
--- a/Engine/Source/Runtime/Engine/Private/AI/Navigation/NavCollision.cpp
+++ b/Engine/Source/Runtime/Engine/Private/AI/Navigation/NavCollision.cpp
@@ -245,14 +245,13 @@ void UNavCollision::GetNavigationModifier(FCompositeNavModifier& Modifier, const
 		if (ConvexCollision.VertexBuffer.Num() > 0)
 		{
 			int32 LastVertIndex = 0;
-			TArray<FVector> Verts(ConvexCollision.VertexBuffer);
 
 			for (int32 i = 0; i < ConvexShapeIndices.Num(); i++)
 			{
 				int32 FirstVertIndex = LastVertIndex;
 				LastVertIndex = ConvexShapeIndices.IsValidIndex(i + 1) ? ConvexShapeIndices[i + 1] : ConvexCollision.VertexBuffer.Num();
 
-				FAreaNavModifier AreaMod(Verts, FirstVertIndex, LastVertIndex, ENavigationCoordSystem::Unreal, LocalToWorld, UseAreaClass);
+				FAreaNavModifier AreaMod(ConvexCollision.VertexBuffer, FirstVertIndex, LastVertIndex, ENavigationCoordSystem::Unreal, LocalToWorld, UseAreaClass);
 				AreaMod.SetIncludeAgentHeight(true);
 				Modifier.Add(AreaMod);
 			}
diff --git a/Engine/Source/Runtime/Engine/Private/AI/Navigation/NavigationModifier.cpp b/Engine/Source/Runtime/Engine/Private/AI/Navigation/NavigationModifier.cpp
--- a/Engine/Source/Runtime/Engine/Private/AI/Navigation/NavigationModifier.cpp
+++ b/Engine/Source/Runtime/Engine/Private/AI/Navigation/NavigationModifier.cpp
@@ -167,13 +167,23 @@ FAreaNavModifier::FAreaNavModifier(const FBox& Box, const FTransform& LocalToWor
 FAreaNavModifier::FAreaNavModifier(const TArray<FVector>& InPoints, ENavigationCoordSystem::Type CoordType, const FTransform& LocalToWorld, const TSubclassOf<UNavArea> InAreaClass)
 {
 	Init(InAreaClass);
-	SetConvex(InPoints, 0, InPoints.Num(), CoordType, LocalToWorld);
+	SetConvex(InPoints.GetData(), 0, InPoints.Num(), CoordType, LocalToWorld);
 }
 
 FAreaNavModifier::FAreaNavModifier(const TArray<FVector>& InPoints, const int32 FirstIndex, const int32 LastIndex, ENavigationCoordSystem::Type CoordType, const FTransform& LocalToWorld, const TSubclassOf<UNavArea> InAreaClass)
 {
+	check(InPoints.IsValidIndex(FirstIndex) && InPoints.IsValidIndex(LastIndex-1));
+
+	Init(InAreaClass);
+	SetConvex(InPoints.GetData(), FirstIndex, LastIndex, CoordType, LocalToWorld);
+}
+
+FAreaNavModifier::FAreaNavModifier(const TNavStatArray<FVector>& InPoints, const int32 FirstIndex, const int32 LastIndex, ENavigationCoordSystem::Type CoordType, const FTransform& LocalToWorld, const TSubclassOf<UNavArea> InAreaClass)
+{
+	check(InPoints.IsValidIndex(FirstIndex) && InPoints.IsValidIndex(LastIndex-1));
+
 	Init(InAreaClass);
-	SetConvex(InPoints, FirstIndex, LastIndex, CoordType, LocalToWorld);
+	SetConvex(InPoints.GetData(), FirstIndex, LastIndex, CoordType, LocalToWorld);
 }
 
 FAreaNavModifier::FAreaNavModifier(const UBrushComponent* BrushComponent, const TSubclassOf<UNavArea> InAreaClass)
@@ -194,7 +204,7 @@ FAreaNavModifier::FAreaNavModifier(const UBrushComponent* BrushComponent, const
 	}
 
 	Init(InAreaClass);
-	SetConvex(Verts, 0, Verts.Num(), ENavigationCoordSystem::Unreal, BrushComponent->ComponentToWorld);
+	SetConvex(Verts.GetData(), 0, Verts.Num(), ENavigationCoordSystem::Unreal, BrushComponent->ComponentToWorld);
 }
 
 void FAreaNavModifier::GetCylinder(FCylinderNavAreaData& Data) const
@@ -282,11 +292,11 @@ void FAreaNavModifier::SetBox(const FBox& Box, const FTransform& LocalToWorld)
 	}
 	else
 	{
-		SetConvex(Corners, 0, Corners.Num(), ENavigationCoordSystem::Unreal, FTransform::Identity);
+		SetConvex(Corners.GetData(), 0, Corners.Num(), ENavigationCoordSystem::Unreal, FTransform::Identity);
 	}
 }
 
-void FAreaNavModifier::SetConvex(const TArray<FVector>& InPoints, const int32 FirstIndex, const int32 LastIndex, ENavigationCoordSystem::Type CoordType, const FTransform& LocalToWorld)
+void FAreaNavModifier::SetConvex(const FVector* InPoints, const int32 FirstIndex, const int32 LastIndex, ENavigationCoordSystem::Type CoordType, const FTransform& LocalToWorld)
 {
 	FConvexNavAreaData ConvexData;
 	ConvexData.MinZ = MAX_FLT;
diff --git a/Engine/Source/Runtime/Engine/Public/AI/NavigationModifier.h b/Engine/Source/Runtime/Engine/Public/AI/NavigationModifier.h
--- a/Engine/Source/Runtime/Engine/Public/AI/NavigationModifier.h
+++ b/Engine/Source/Runtime/Engine/Public/AI/NavigationModifier.h
@@ -74,6 +74,7 @@ struct ENGINE_API FAreaNavModifier : public FNavigationModifier
 	FAreaNavModifier(const FBox& Box, const FTransform& LocalToWorld, const TSubclassOf<UNavArea> AreaClass);
 	FAreaNavModifier(const TArray<FVector>& Points, ENavigationCoordSystem::Type CoordType, const FTransform& LocalToWorld, const TSubclassOf<UNavArea> AreaClass);
 	FAreaNavModifier(const TArray<FVector>& Points, const int32 FirstIndex, const int32 LastIndex, ENavigationCoordSystem::Type CoordType, const FTransform& LocalToWorld, const TSubclassOf<UNavArea> AreaClass);
+	FAreaNavModifier(const TNavStatArray<FVector>& Points, const int32 FirstIndex, const int32 LastIndex, ENavigationCoordSystem::Type CoordType, const FTransform& LocalToWorld, const TSubclassOf<UNavArea> AreaClass);
 	FAreaNavModifier(const UBrushComponent* BrushComponent, const TSubclassOf<UNavArea> AreaClass);
 
 	FORCEINLINE const FBox& GetBounds() const { return Bounds; }
@@ -102,7 +103,7 @@ protected:
 	uint8 bIncludeAgentHeight : 1;
 
 	void Init(const TSubclassOf<UNavArea> AreaClass);
-	void SetConvex(const TArray<FVector>& Points, const int32 FirstIndex, const int32 LastIndex, ENavigationCoordSystem::Type CoordType, const FTransform& LocalToWorld);
+	void SetConvex(const FVector* Points, const int32 FirstIndex, const int32 LastIndex, ENavigationCoordSystem::Type CoordType, const FTransform& LocalToWorld);
 	void SetBox(const FBox& Box, const FTransform& LocalToWorld);
 };
 
