UE-17106 - Change the dynamic parameter default value back to 1.

Added DefaultValue property to DynamicParameter material node. This defaults to 1 for new materials but 0 for existing ones to maintain current behavior.
This default is used for all vertex factories other than those for particles. In particle vertex factories, the dynamic parameter is used, as before. Sometimes coming from vertex data, others coming from a default vertex buffer filled with 1.0.

Also, defaulted the DynamicParameter on the cascade side to 1.0 to improve workflow that side too.

diff --git a/Engine/Shaders/MaterialTemplate.usf b/Engine/Shaders/MaterialTemplate.usf
--- a/Engine/Shaders/MaterialTemplate.usf
+++ b/Engine/Shaders/MaterialTemplate.usf
@@ -68,6 +68,15 @@ struct FMaterialParticleParameters
 	float2 Size;
 };
 
+float4 GetDynamicParameter(FMaterialParticleParameters Parameters, float4 Default)
+{
+	#if PARTICLE_FACTORY
+	return Parameters.DynamicParameter;
+	#else
+	return Default;
+	#endif
+}
+
 /** 
  * Parameters needed by pixel shader material inputs.
  * These are independent of vertex factory.
diff --git a/Engine/Source/Editor/MaterialEditor/Private/MaterialEditor.cpp b/Engine/Source/Editor/MaterialEditor/Private/MaterialEditor.cpp
--- a/Engine/Source/Editor/MaterialEditor/Private/MaterialEditor.cpp
+++ b/Engine/Source/Editor/MaterialEditor/Private/MaterialEditor.cpp
@@ -2993,7 +2993,7 @@ UMaterialExpression* FMaterialEditor::CreateNewMaterialExpression(UClass* NewExp
 		UMaterialExpressionDynamicParameter* DynamicExpression = Cast<UMaterialExpressionDynamicParameter>(NewExpression);
 		if (DynamicExpression)
 		{
-			DynamicExpression->UpdateDynamicParameterNames();
+			DynamicExpression->UpdateDynamicParameterProperties();
 		}
 
 		Material->AddExpressionParameter(NewExpression, Material->EditorParameters);
@@ -3552,9 +3552,9 @@ void FMaterialEditor::NotifyPostChange( const FPropertyChangedEvent& PropertyCha
 				{
 					Material->UpdateExpressionParameterName(SelectedNode->MaterialExpression);
 				}
-				else if (NameOfPropertyThatChanged == FName(TEXT("ParamNames")))
+				else if (SelectedNode->MaterialExpression->IsA(UMaterialExpressionDynamicParameter::StaticClass()))
 				{
-					Material->UpdateExpressionDynamicParameterNames(SelectedNode->MaterialExpression);
+					Material->UpdateExpressionDynamicParameters(SelectedNode->MaterialExpression);
 				}
 				else
 				{
diff --git a/Engine/Source/Runtime/Core/Public/UObject/ObjectVersion.h b/Engine/Source/Runtime/Core/Public/UObject/ObjectVersion.h
--- a/Engine/Source/Runtime/Core/Public/UObject/ObjectVersion.h
+++ b/Engine/Source/Runtime/Core/Public/UObject/ObjectVersion.h
@@ -542,6 +542,8 @@ enum EUnrealEngineObjectUE4Version
 	VER_UE4_CHARACTER_MOVEMENT_ADD_BRAKING_FRICTION,
 	// Removed TTransArrays from UModel
 	VER_UE4_BSP_UNDO_FIX,
+	// Added default value to dynamic parameter.
+	VER_UE4_DYNAMIC_PARAMETER_DEFAULT_VALUE,
 	// -----<new versions can be added before this line>-------------------------------------------------
 	// - this needs to be the last line (see note below)
 	VER_UE4_AUTOMATIC_VERSION_PLUS_ONE,
diff --git a/Engine/Source/Runtime/Engine/Classes/Materials/Material.h b/Engine/Source/Runtime/Engine/Classes/Materials/Material.h
--- a/Engine/Source/Runtime/Engine/Classes/Materials/Material.h
+++ b/Engine/Source/Runtime/Engine/Classes/Materials/Material.h
@@ -1110,12 +1110,12 @@ public:
 	ENGINE_API virtual bool HasDuplicateDynamicParameters(const UMaterialExpression* Expression);
 
 	/**
-	 * Iterate through all of the expression nodes and fix up changed names on
-	 * matching dynamic parameters when a name change occurs.
+	 * Iterate through all of the expression nodes and fix up changed properties on
+	 * matching dynamic parameters when a change occurs.
 	 *
 	 * @param	Expression	The expression dynamic parameter.
 	 */
-	ENGINE_API virtual void UpdateExpressionDynamicParameterNames(const UMaterialExpression* Expression);
+	ENGINE_API virtual void UpdateExpressionDynamicParameters(const UMaterialExpression* Expression);
 
 	/**
 	 * Get the name of a parameter.
diff --git a/Engine/Source/Runtime/Engine/Classes/Materials/MaterialExpressionDynamicParameter.h b/Engine/Source/Runtime/Engine/Classes/Materials/MaterialExpressionDynamicParameter.h
--- a/Engine/Source/Runtime/Engine/Classes/Materials/MaterialExpressionDynamicParameter.h
+++ b/Engine/Source/Runtime/Engine/Classes/Materials/MaterialExpressionDynamicParameter.h
@@ -23,9 +23,13 @@ class UMaterialExpressionDynamicParameter : public UMaterialExpression
 	UPROPERTY(EditAnywhere, editfixedsize, Category=MaterialExpressionDynamicParameter)
 	TArray<FString> ParamNames;
 
+	UPROPERTY(EditAnywhere, Category = MaterialExpressionDynamicParameter)
+	FLinearColor DefaultValue;
+
 	// Begin UObject Interface
 #if WITH_EDITOR
 	virtual void PostEditChangeProperty(FPropertyChangedEvent& PropertyChangedEvent) override;
+	virtual void PostLoad();
 #endif // WITH_EDITOR
 	// End UObject Interface
 
@@ -40,17 +44,17 @@ class UMaterialExpressionDynamicParameter : public UMaterialExpression
 
 	/**
 	 * Iterate through all of the expression nodes until we find another 
-	 * dynamic parameter we can copy the parameter names from
+	 * dynamic parameter we can copy the properties from
 	 */
-	ENGINE_API void UpdateDynamicParameterNames();
+	ENGINE_API void UpdateDynamicParameterProperties();
 
 	/**
-	 * Copy the parameters names from the specified dynamic parameter
+	 * Copy the properties from the specified dynamic parameter
 	 *
-	 * @param	FromParam	The param to copy the names from
+	 * @param	FromParam	The param to copy from
 	 * @return	true if sucessful
 	 */
-	bool CopyDynamicParameterNames(const UMaterialExpressionDynamicParameter* FromParam);
+	bool CopyDynamicParameterProperties(const UMaterialExpressionDynamicParameter* FromParam);
 };
 
 
diff --git a/Engine/Source/Runtime/Engine/Private/Materials/HLSLMaterialTranslator.h b/Engine/Source/Runtime/Engine/Private/Materials/HLSLMaterialTranslator.h
--- a/Engine/Source/Runtime/Engine/Private/Materials/HLSLMaterialTranslator.h
+++ b/Engine/Source/Runtime/Engine/Private/Materials/HLSLMaterialTranslator.h
@@ -3360,7 +3360,7 @@ protected:
 		return TransformBase(SourceCoordBasis, DestCoordBasis, A, 1);
 	}
 
-	virtual int32 DynamicParameter() override
+	virtual int32 DynamicParameter(FLinearColor& DefaultValue) override
 	{
 		if (ShaderFrequency != SF_Vertex && ShaderFrequency != SF_Pixel && ShaderFrequency != SF_Compute)
 		{
@@ -3369,9 +3369,11 @@ protected:
 
 		bNeedsParticleDynamicParameter = true;
 
+		int32 Default = Constant4(DefaultValue.R, DefaultValue.G, DefaultValue.B, DefaultValue.A);
 		return AddInlinedCodeChunk(
 			MCT_Float4,
-			TEXT("Parameters.Particle.DynamicParameter")
+			TEXT("GetDynamicParameter(Parameters.Particle, %s)"),
+			*GetParameterCode(Default)
 			);
 	}
 
diff --git a/Engine/Source/Runtime/Engine/Private/Materials/Material.cpp b/Engine/Source/Runtime/Engine/Private/Materials/Material.cpp
--- a/Engine/Source/Runtime/Engine/Private/Materials/Material.cpp
+++ b/Engine/Source/Runtime/Engine/Private/Materials/Material.cpp
@@ -2851,7 +2851,7 @@ bool UMaterial::HasDuplicateDynamicParameters(const UMaterialExpression* Express
 }
 
 
-void UMaterial::UpdateExpressionDynamicParameterNames(const UMaterialExpression* Expression)
+void UMaterial::UpdateExpressionDynamicParameters(const UMaterialExpression* Expression)
 {
 	const UMaterialExpressionDynamicParameter* DynParam = Cast<UMaterialExpressionDynamicParameter>(Expression);
 	if (DynParam)
@@ -2859,7 +2859,7 @@ void UMaterial::UpdateExpressionDynamicParameterNames(const UMaterialExpression*
 		for (int32 ExpIndex = 0; ExpIndex < Expressions.Num(); ExpIndex++)
 		{
 			UMaterialExpressionDynamicParameter* CheckParam = Cast<UMaterialExpressionDynamicParameter>(Expressions[ExpIndex]);
-			if (CheckParam && CheckParam->CopyDynamicParameterNames(DynParam))
+			if (CheckParam && CheckParam->CopyDynamicParameterProperties(DynParam))
 			{
 #if WITH_EDITOR
 				CheckParam->GraphNode->ReconstructNode();
diff --git a/Engine/Source/Runtime/Engine/Private/Materials/MaterialExpressions.cpp b/Engine/Source/Runtime/Engine/Private/Materials/MaterialExpressions.cpp
--- a/Engine/Source/Runtime/Engine/Private/Materials/MaterialExpressions.cpp
+++ b/Engine/Source/Runtime/Engine/Private/Materials/MaterialExpressions.cpp
@@ -4375,13 +4375,15 @@ UMaterialExpressionDynamicParameter::UMaterialExpressionDynamicParameter(const F
 	Outputs.Add(FExpressionOutput(TEXT(""), 1, 0, 0, 1, 0));
 	Outputs.Add(FExpressionOutput(TEXT(""), 1, 0, 0, 0, 1));
 
+	DefaultValue = FLinearColor::White;
+
 	bShaderInputData = true;
 }
 
 
 int32 UMaterialExpressionDynamicParameter::Compile( FMaterialCompiler* Compiler, int32 OutputIndex, int32 MultiplexIndex )
 {
-	return Compiler->DynamicParameter();
+	return Compiler->DynamicParameter(DefaultValue);
 }
 
 
@@ -4441,20 +4443,30 @@ void UMaterialExpressionDynamicParameter::PostEditChangeProperty(FPropertyChange
 
 #endif // WITH_EDITOR
 
-void UMaterialExpressionDynamicParameter::UpdateDynamicParameterNames()
+void UMaterialExpressionDynamicParameter::PostLoad()
+{
+	Super::PostLoad();
+
+	if (GetLinkerUE4Version() < VER_UE4_DYNAMIC_PARAMETER_DEFAULT_VALUE)
+	{
+		DefaultValue = FLinearColor::Black;//Old data should default to 0.0f;
+	}
+}
+
+void UMaterialExpressionDynamicParameter::UpdateDynamicParameterProperties()
 {
 	check(Material);
 	for (int32 ExpIndex = 0; ExpIndex < Material->Expressions.Num(); ExpIndex++)
 	{
 		const UMaterialExpressionDynamicParameter* DynParam = Cast<UMaterialExpressionDynamicParameter>(Material->Expressions[ExpIndex]);
-		if (CopyDynamicParameterNames(DynParam))
+		if (CopyDynamicParameterProperties(DynParam))
 		{
 			break;
 		}
 	}
 }
 
-bool UMaterialExpressionDynamicParameter::CopyDynamicParameterNames(const UMaterialExpressionDynamicParameter* FromParam)
+bool UMaterialExpressionDynamicParameter::CopyDynamicParameterProperties(const UMaterialExpressionDynamicParameter* FromParam)
 {
 	if (FromParam && (FromParam != this))
 	{
@@ -4462,6 +4474,7 @@ bool UMaterialExpressionDynamicParameter::CopyDynamicParameterNames(const UMater
 		{
 			ParamNames[NameIndex] = FromParam->ParamNames[NameIndex];
 		}
+		DefaultValue = FromParam->DefaultValue;
 		return true;
 	}
 	return false;
diff --git a/Engine/Source/Runtime/Engine/Private/Particles/ParticleModules_Parameter.cpp b/Engine/Source/Runtime/Engine/Private/Particles/ParticleModules_Parameter.cpp
--- a/Engine/Source/Runtime/Engine/Private/Particles/ParticleModules_Parameter.cpp
+++ b/Engine/Source/Runtime/Engine/Private/Particles/ParticleModules_Parameter.cpp
@@ -64,10 +64,15 @@ void UParticleModuleParameterDynamic::PostInitProperties()
 		};
 		static FConstructorStatics ConstructorStatics;
 
-		DynamicParams.Add(FEmitterDynamicParameter(ConstructorStatics.NAME_None, false, EDPV_UserSet, NewObject<UDistributionFloatConstant>(this, TEXT("DistributionParam1"))));
-		DynamicParams.Add(FEmitterDynamicParameter(ConstructorStatics.NAME_None, false, EDPV_UserSet, NewObject<UDistributionFloatConstant>(this, TEXT("DistributionParam2"))));
-		DynamicParams.Add(FEmitterDynamicParameter(ConstructorStatics.NAME_None, false, EDPV_UserSet, NewObject<UDistributionFloatConstant>(this, TEXT("DistributionParam3"))));
-		DynamicParams.Add(FEmitterDynamicParameter(ConstructorStatics.NAME_None, false, EDPV_UserSet, NewObject<UDistributionFloatConstant>(this, TEXT("DistributionParam4"))));
+		UDistributionFloatConstant* NewParam1 = NewObject<UDistributionFloatConstant>(this, TEXT("DistributionParam1")); NewParam1->Constant = 1.0f;
+		UDistributionFloatConstant* NewParam2 = NewObject<UDistributionFloatConstant>(this, TEXT("DistributionParam2")); NewParam2->Constant = 1.0f;
+		UDistributionFloatConstant* NewParam3 = NewObject<UDistributionFloatConstant>(this, TEXT("DistributionParam3")); NewParam3->Constant = 1.0f;
+		UDistributionFloatConstant* NewParam4 = NewObject<UDistributionFloatConstant>(this, TEXT("DistributionParam4")); NewParam4->Constant = 1.0f;
+
+		DynamicParams.Add(FEmitterDynamicParameter(ConstructorStatics.NAME_None, false, EDPV_UserSet, NewParam1));
+		DynamicParams.Add(FEmitterDynamicParameter(ConstructorStatics.NAME_None, false, EDPV_UserSet, NewParam2));
+		DynamicParams.Add(FEmitterDynamicParameter(ConstructorStatics.NAME_None, false, EDPV_UserSet, NewParam3));
+		DynamicParams.Add(FEmitterDynamicParameter(ConstructorStatics.NAME_None, false, EDPV_UserSet, NewParam4));
 	}
 }
 
diff --git a/Engine/Source/Runtime/Engine/Public/MaterialCompiler.h b/Engine/Source/Runtime/Engine/Public/MaterialCompiler.h
--- a/Engine/Source/Runtime/Engine/Public/MaterialCompiler.h
+++ b/Engine/Source/Runtime/Engine/Public/MaterialCompiler.h
@@ -151,7 +151,7 @@ public:
 	virtual int32 TransformVector(EMaterialCommonBasis SourceCoordBasis, EMaterialCommonBasis DestCoordBasis, int32 A) = 0;
 	virtual int32 TransformPosition(EMaterialCommonBasis SourceCoordBasis, EMaterialCommonBasis DestCoordBasis, int32 A) = 0;
 
-	virtual int32 DynamicParameter() = 0;
+	virtual int32 DynamicParameter(FLinearColor& DefaultValue) = 0;
 	virtual int32 LightmapUVs() = 0;
 
 	virtual int32 LightmassReplace(int32 Realtime, int32 Lightmass) = 0;
@@ -315,7 +315,7 @@ public:
 		return Compiler->TransformPosition(SourceCoordBasis, DestCoordBasis, A);
 	}
 
-	virtual int32 DynamicParameter() override { return Compiler->DynamicParameter(); }
+	virtual int32 DynamicParameter(FLinearColor& DefaultValue) override { return Compiler->DynamicParameter(DefaultValue); }
 	virtual int32 LightmapUVs() override { return Compiler->LightmapUVs(); }
 
 	virtual int32 LightmassReplace(int32 Realtime, int32 Lightmass) override { return Realtime; }
