UE-17686 Fix typo in OnActorEndOverlap comment

diff --git a/Engine/Source/Runtime/Engine/Classes/GameFramework/Actor.h b/Engine/Source/Runtime/Engine/Classes/GameFramework/Actor.h
--- a/Engine/Source/Runtime/Engine/Classes/GameFramework/Actor.h
+++ b/Engine/Source/Runtime/Engine/Classes/GameFramework/Actor.h
@@ -525,7 +525,7 @@ public:
 	FActorBeginOverlapSignature OnActorBeginOverlap;
 
 	/** 
-	 *	Called when another actor steps overlapping this actor. 
+	 *	Called when another actor stops overlapping this actor. 
 	 *	@note Components on both this and the other Actor must have bGenerateOverlapEvents set to true to generate overlap events.
 	 */
 	UPROPERTY(BlueprintAssignable, Category="Collision")
