UE-17191 Fix bIgnoreOwningActor option on URadialForceComponent only working for impulses (for forces)

diff --git a/Engine/Source/Runtime/Engine/Private/PhysicsEngine/RadialForceComponent.cpp b/Engine/Source/Runtime/Engine/Private/PhysicsEngine/RadialForceComponent.cpp
--- a/Engine/Source/Runtime/Engine/Private/PhysicsEngine/RadialForceComponent.cpp
+++ b/Engine/Source/Runtime/Engine/Private/PhysicsEngine/RadialForceComponent.cpp
@@ -43,6 +43,13 @@ void URadialForceComponent::TickComponent(float DeltaTime, enum ELevelTick TickT
 
 		FCollisionQueryParams Params(AddForceOverlapName, false);
 		Params.bTraceAsyncScene = true; // want to hurt stuff in async scene
+
+		// Ignore owner actor if desired
+		if (bIgnoreOwningActor)
+		{
+			Params.AddIgnoredActor(GetOwner());
+		}
+
 		GetWorld()->OverlapMultiByObjectType(Overlaps, Origin, FQuat::Identity, CollisionObjectQueryParams, FCollisionShape::MakeSphere(Radius), Params);
 
 		// Iterate over each and apply force
