UE-14186 Add EqualEqual_TransformTransform and NearlyEqual_TransformTransform to BP math library

diff --git a/Engine/Source/Runtime/Engine/Classes/Kismet/KismetMathLibrary.h b/Engine/Source/Runtime/Engine/Classes/Kismet/KismetMathLibrary.h
--- a/Engine/Source/Runtime/Engine/Classes/Kismet/KismetMathLibrary.h
+++ b/Engine/Source/Runtime/Engine/Classes/Kismet/KismetMathLibrary.h
@@ -1435,6 +1435,19 @@ class ENGINE_API UKismetMathLibrary : public UBlueprintFunctionLibrary
 	UFUNCTION(BlueprintPure, Category="Math|Interpolation")
 	static FTransform TInterpTo(const FTransform& Current, const FTransform& Target, float DeltaTime, float InterpSpeed);
 
+	/** Returns true if transform A is equal to transform B */
+	UFUNCTION(BlueprintPure, meta = (DisplayName = "Equal Transform", CompactNodeTitle = "==", Keywords = "== equal"), Category="Math|Transform")
+	static bool EqualEqual_TransformTransform(const FTransform& A, const FTransform& B);
+
+	/** 
+	 *	Returns true if transform A is nearly equal to B 
+	 *	@param LocationTolerance	How close position of transforms need to be to be considered equal
+	 *	@param RotationTolerance	How close rotations of transforms need to be to be considered equal
+	 *	@param Scale3DTolerance		How close scale of transforms need to be to be considered equal
+	 */
+	UFUNCTION(BlueprintPure, meta = (DisplayName = "Nearly Equal (transform)", Keywords = "== equal"), Category = "Math|Transform")
+	static bool NearlyEqual_TransformTransform(const FTransform& A, const FTransform& B, float LocationTolerance = 1.e-4f, float RotationTolerance = 1.e-4f, float Scale3DTolerance = 1.e-4f);
+
 	//
 	// Vector2D functions
 	//
diff --git a/Engine/Source/Runtime/Engine/Private/KismetMathLibrary.cpp b/Engine/Source/Runtime/Engine/Private/KismetMathLibrary.cpp
--- a/Engine/Source/Runtime/Engine/Private/KismetMathLibrary.cpp
+++ b/Engine/Source/Runtime/Engine/Private/KismetMathLibrary.cpp
@@ -1126,6 +1126,21 @@ FTransform UKismetMathLibrary::TInterpTo(const FTransform& Current, const FTrans
 	return TLerp(Current, Target, Alpha);
 }
 
+bool UKismetMathLibrary::EqualEqual_TransformTransform(const FTransform& A, const FTransform& B)
+{
+	return NearlyEqual_TransformTransform(A, B);
+}
+
+bool UKismetMathLibrary::NearlyEqual_TransformTransform(const FTransform& A, const FTransform& B, float LocationTolerance, float RotationTolerance, float Scale3DTolerance)
+{
+	return 
+		A.GetRotation().Equals(B.GetRotation(), RotationTolerance) && 
+		A.GetTranslation().Equals(B.GetTranslation(), LocationTolerance) && 
+		A.GetScale3D().Equals(B.GetScale3D(), Scale3DTolerance);
+}
+
+
+
 FVector2D UKismetMathLibrary::Add_Vector2DVector2D(FVector2D A, FVector2D B)
 {
 	return A + B;
