ADDED UE-11699 static meshes: Feature - Bounds push instead of scale

diff --git a/Engine/Source/Editor/UnrealEd/Private/Factories/EditorFactories.cpp b/Engine/Source/Editor/UnrealEd/Private/Factories/EditorFactories.cpp
--- a/Engine/Source/Editor/UnrealEd/Private/Factories/EditorFactories.cpp
+++ b/Engine/Source/Editor/UnrealEd/Private/Factories/EditorFactories.cpp
@@ -5728,6 +5728,10 @@ EReimportResult::Type UReimportFbxStaticMeshFactory::Reimport( UObject* Obj )
 				(UNavCollision*)StaticDuplicateObject(Mesh->NavCollision, GetTransientPackage(), nullptr) :
 				nullptr;
 
+			// preserve extended bound settings
+			const FVector PositiveBoundsExtension = Mesh->PositiveBoundsExtension;
+			const FVector NegativeBoundsExtension = Mesh->NegativeBoundsExtension;
+
 			if (FFbxImporter->ReimportStaticMesh(Mesh, ImportData))
 			{
 				UE_LOG(LogEditorFactories, Log, TEXT("-- imported successfully") );
@@ -5745,6 +5749,10 @@ EReimportResult::Type UReimportFbxStaticMeshFactory::Reimport( UObject* Obj )
 					NavCollision->Rename(NULL, Mesh, REN_DontCreateRedirectors | REN_DoNotDirty);
 				}
 
+				// Restore bounds extension settings
+				Mesh->PositiveBoundsExtension = PositiveBoundsExtension;
+				Mesh->NegativeBoundsExtension = NegativeBoundsExtension;
+
 				// Try to find the outer package so we can dirty it up
 				if (Mesh->GetOuter())
 				{
diff --git a/Engine/Source/Runtime/Core/Public/UObject/ObjectVersion.h b/Engine/Source/Runtime/Core/Public/UObject/ObjectVersion.h
--- a/Engine/Source/Runtime/Core/Public/UObject/ObjectVersion.h
+++ b/Engine/Source/Runtime/Core/Public/UObject/ObjectVersion.h
@@ -544,6 +544,8 @@ enum EUnrealEngineObjectUE4Version
 	VER_UE4_BSP_UNDO_FIX,
 	// Added default value to dynamic parameter.
 	VER_UE4_DYNAMIC_PARAMETER_DEFAULT_VALUE,
+	// Added ExtendedBounds to StaticMesh
+	VER_UE4_STATIC_MESH_EXTENDED_BOUNDS,
 	// -----<new versions can be added before this line>-------------------------------------------------
 	// - this needs to be the last line (see note below)
 	VER_UE4_AUTOMATIC_VERSION_PLUS_ONE,
diff --git a/Engine/Source/Runtime/Engine/Classes/Engine/StaticMesh.h b/Engine/Source/Runtime/Engine/Classes/Engine/StaticMesh.h
--- a/Engine/Source/Runtime/Engine/Classes/Engine/StaticMesh.h
+++ b/Engine/Source/Runtime/Engine/Classes/Engine/StaticMesh.h
@@ -419,6 +419,16 @@ class UStaticMesh : public UObject, public IInterface_CollisionDataProvider, pub
 	/** Data that is only available if this static mesh is an imported SpeedTree */
 	TSharedPtr<class FSpeedTreeWind> SpeedTreeWind;
 
+	/** Bound extension values in the positive direction of XYZ, positive value increases bound size */
+	UPROPERTY(EditDefaultsOnly, AdvancedDisplay, Category = StaticMesh)
+	FVector PositiveBoundsExtension;
+	/** Bound extension values in the negative direction of XYZ, positive value increases bound size */
+	UPROPERTY(EditDefaultsOnly, AdvancedDisplay, Category = StaticMesh)
+	FVector NegativeBoundsExtension;
+	/** Original mesh bounds extended with Positive/NegativeBoundsExtension */
+	UPROPERTY()
+	FBoxSphereBounds ExtendedBounds;
+
 protected:
 	/**
 	 * Index of an element to ignore while gathering streaming texture factors.
@@ -435,7 +445,6 @@ public:
 	/** Pre-build navigation collision */
 	UPROPERTY(VisibleAnywhere, transient, duplicatetransient, Instanced, Category = Navigation)
 	class UNavCollision* NavCollision;
-	
 public:
 	/**
 	 * Default constructor
@@ -502,7 +511,7 @@ public:
 	 */
 	ENGINE_API FBoxSphereBounds GetBounds() const;
 
-	/** Returns the bounding box, in local space, of the StaticMesh asset */
+	/** Returns the bounding box, in local space including bounds extension(s), of the StaticMesh asset */
 	UFUNCTION(BlueprintCallable, Category="StaticMesh")
 	ENGINE_API FBox GetBoundingBox() const;
 
@@ -633,5 +642,8 @@ private:
 	 * Caches derived renderable data.
 	 */
 	void CacheDerivedData();
+
+	/** Calculates the extended bounds */
+	void CalculateExtendedBounds();
 #endif // #if WITH_EDITOR
 };
diff --git a/Engine/Source/Runtime/Engine/Private/StaticMesh.cpp b/Engine/Source/Runtime/Engine/Private/StaticMesh.cpp
--- a/Engine/Source/Runtime/Engine/Private/StaticMesh.cpp
+++ b/Engine/Source/Runtime/Engine/Private/StaticMesh.cpp
@@ -1368,22 +1368,12 @@ bool UStaticMesh::HasValidRenderData() const
 
 FBoxSphereBounds UStaticMesh::GetBounds() const
 {
-	FBoxSphereBounds Bounds(ForceInit);
-	if (RenderData)
-	{
-		Bounds = RenderData->Bounds;
-	}
-	return Bounds;
+	return ExtendedBounds;
 }
 
 FBox UStaticMesh::GetBoundingBox() const
 {
-	FBox BoundingBox(FVector(0,0,0), FVector(0,0,0));
-	if (RenderData != nullptr)
-	{
-		BoundingBox = RenderData->Bounds.GetBox();
-	}
-	return BoundingBox;
+	return ExtendedBounds.GetBox();
 }
 
 float UStaticMesh::GetStreamingTextureFactor(int32 RequestedUVIndex)
@@ -1472,6 +1462,11 @@ void UStaticMesh::PostEditChangeProperty(FPropertyChangedEvent& PropertyChangedE
 		GEngine->TriggerStreamingDataRebuild();
 	}
 
+	if (PropertyChangedEvent.MemberProperty && ( (PropertyChangedEvent.MemberProperty->GetFName() == GET_MEMBER_NAME_CHECKED(UStaticMesh, PositiveBoundsExtension) ) || ( PropertyChangedEvent.MemberProperty->GetFName() == GET_MEMBER_NAME_CHECKED(UStaticMesh, NegativeBoundsExtension) ) ))
+	{
+		// Update the extended bounds
+		CalculateExtendedBounds();
+	}
 	AutoLODPixelError = FMath::Max(AutoLODPixelError, 1.0f);
 
 	if (!bAutoComputeLODScreenSize
@@ -1759,6 +1754,29 @@ void UStaticMesh::CacheDerivedData()
 		}
 	}
 }
+
+void UStaticMesh::CalculateExtendedBounds()
+{
+	FBoxSphereBounds Bounds(ForceInit);
+	if (RenderData)
+	{
+		Bounds = RenderData->Bounds;
+	}
+
+	// Convert to Min and Max
+	FVector Min = Bounds.Origin - Bounds.BoxExtent;
+	FVector Max = Bounds.Origin + Bounds.BoxExtent;
+	// Apply bound extensions
+	Min -= NegativeBoundsExtension;
+	Max += PositiveBoundsExtension;
+	// Convert back to Origin, Extent and update SphereRadius
+	Bounds.Origin = (Min + Max) / 2;
+	Bounds.BoxExtent = (Max - Min) / 2;	
+	Bounds.SphereRadius = Bounds.BoxExtent.GetAbsMax();
+
+	ExtendedBounds = Bounds;
+}
+
 #endif // #if WITH_EDITORONLY_DATA
 
 #if WITH_EDITORONLY_DATA
@@ -1982,6 +2000,11 @@ void UStaticMesh::PostLoad()
 		InitResources();
 	}
 
+	if (GetLinkerUE4Version() < VER_UE4_STATIC_MESH_EXTENDED_BOUNDS)
+	{
+		CalculateExtendedBounds();
+	}
+
 	// We want to always have a BodySetup, its used for per-poly collision as well
 	if(BodySetup == NULL)
 	{
diff --git a/Engine/Source/Runtime/Engine/Private/StaticMeshBuild.cpp b/Engine/Source/Runtime/Engine/Private/StaticMeshBuild.cpp
--- a/Engine/Source/Runtime/Engine/Private/StaticMeshBuild.cpp
+++ b/Engine/Source/Runtime/Engine/Private/StaticMeshBuild.cpp
@@ -141,6 +141,9 @@ void UStaticMesh::Build(bool bSilent, TArray<FText>* OutErrors)
 
 	}
 
+	// Calculate extended bounds
+	CalculateExtendedBounds();
+
 	if(!bSilent)
 	{
 		GWarn->EndSlowTask();
