UE-17411 Change in physical material override on instance now correctly updates mass properties

diff --git a/Engine/Source/Runtime/Engine/Private/PhysicsEngine/BodyInstance.cpp b/Engine/Source/Runtime/Engine/Private/PhysicsEngine/BodyInstance.cpp
--- a/Engine/Source/Runtime/Engine/Private/PhysicsEngine/BodyInstance.cpp
+++ b/Engine/Source/Runtime/Engine/Private/PhysicsEngine/BodyInstance.cpp
@@ -2948,8 +2948,11 @@ void FBodyInstance::SetPhysMaterialOverride( UPhysicalMaterial* NewPhysMaterial
 	// Save ref to PhysicalMaterial
 	PhysMaterialOverride = NewPhysMaterial;
 
-	// Go through the chain of physical materials and update the NxActor
+	// Go through the chain of physical materials and update the shapes 
 	UpdatePhysicalMaterials();
+
+	// Because physical material has changed, we need to update the mass
+	UpdateMassProperties();
 }
 
 UPhysicalMaterial* FBodyInstance::GetSimplePhysicalMaterial() const
