Fixed build breakage, WITH_EDITOR define missing

diff --git a/Engine/Source/Runtime/Engine/Private/StaticMesh.cpp b/Engine/Source/Runtime/Engine/Private/StaticMesh.cpp
--- a/Engine/Source/Runtime/Engine/Private/StaticMesh.cpp
+++ b/Engine/Source/Runtime/Engine/Private/StaticMesh.cpp
@@ -2000,10 +2000,12 @@ void UStaticMesh::PostLoad()
 		InitResources();
 	}
 
+#if WITH_EDITOR
 	if (GetLinkerUE4Version() < VER_UE4_STATIC_MESH_EXTENDED_BOUNDS)
 	{
 		CalculateExtendedBounds();
 	}
+#endif // #if WITH_EDITOR
 
 	// We want to always have a BodySetup, its used for per-poly collision as well
 	if(BodySetup == NULL)
