Compile fix

diff --git a/Engine/Source/Runtime/Engine/Classes/Materials/MaterialExpressionDynamicParameter.h b/Engine/Source/Runtime/Engine/Classes/Materials/MaterialExpressionDynamicParameter.h
--- a/Engine/Source/Runtime/Engine/Classes/Materials/MaterialExpressionDynamicParameter.h
+++ b/Engine/Source/Runtime/Engine/Classes/Materials/MaterialExpressionDynamicParameter.h
@@ -29,8 +29,9 @@ class UMaterialExpressionDynamicParameter : public UMaterialExpression
 	// Begin UObject Interface
 #if WITH_EDITOR
 	virtual void PostEditChangeProperty(FPropertyChangedEvent& PropertyChangedEvent) override;
-	virtual void PostLoad();
 #endif // WITH_EDITOR
+
+	virtual void PostLoad() override;
 	// End UObject Interface
 
 	// Begin UMaterialExpression Interface
