UE-17862 Add TimeSecondsToString function to convert float seconds to formatted string

diff --git a/Engine/Source/Runtime/Engine/Classes/Kismet/KismetStringLibrary.h b/Engine/Source/Runtime/Engine/Classes/Kismet/KismetStringLibrary.h
--- a/Engine/Source/Runtime/Engine/Classes/Kismet/KismetStringLibrary.h
+++ b/Engine/Source/Runtime/Engine/Classes/Kismet/KismetStringLibrary.h
@@ -435,4 +435,8 @@ class ENGINE_API UKismetStringLibrary : public UBlueprintFunctionLibrary
 	/** @return the substring from Start position for Count characters. */
 	UFUNCTION(BlueprintPure, Category = "Utilities|String")
 		static FString Mid(const FString& SourceString, int32 Start, int32 Count);
+
+	/** Convert a number of seconds into minutes:seconds.milliseconds format string */
+	UFUNCTION(BlueprintPure,  Category = "Utilities|String")
+	static FString TimeSecondsToString(float InSeconds);
 };
diff --git a/Engine/Source/Runtime/Engine/Private/KismetStringLibrary.cpp b/Engine/Source/Runtime/Engine/Private/KismetStringLibrary.cpp
--- a/Engine/Source/Runtime/Engine/Private/KismetStringLibrary.cpp
+++ b/Engine/Source/Runtime/Engine/Private/KismetStringLibrary.cpp
@@ -414,3 +414,15 @@ FString UKismetStringLibrary::Mid(const FString& SourceString, int32 Start, int3
 	return SourceString.Mid(Start, Count);
 }
 
+FString UKismetStringLibrary::TimeSecondsToString(float InSeconds)
+{
+	// Get whole minutes
+	const int32 NumMinutes = FMath::FloorToInt(InSeconds/60.f);
+	// Get seconds not part of whole minutes
+	const int32 NumSeconds = FMath::FloorToInt(InSeconds-(NumMinutes*60.f));
+	// Get fraction of non-whole seconds, convert to 100th of a second, then floor to get whole 100ths
+	const int32 NumCentiseconds = FMath::FloorToInt((InSeconds - FMath::FloorToFloat(InSeconds)) * 100.f);
+
+	// Create string, including leading zeroes
+	return FString::Printf(TEXT("%02d:%02d:%02d"), NumMinutes, NumSeconds, NumCentiseconds);
+}
