UE4 - fix runtime blueprint use of atmospheric component UE-17850

diff --git a/Engine/Source/Runtime/Engine/Classes/Atmosphere/AtmosphericFogComponent.h b/Engine/Source/Runtime/Engine/Classes/Atmosphere/AtmosphericFogComponent.h
--- a/Engine/Source/Runtime/Engine/Classes/Atmosphere/AtmosphericFogComponent.h
+++ b/Engine/Source/Runtime/Engine/Classes/Atmosphere/AtmosphericFogComponent.h
@@ -261,6 +261,11 @@ public:
 	ENGINE_API void InitResource();
 	ENGINE_API void ReleaseResource();
 
+	// Begin UActorComponent interface.
+	virtual FActorComponentInstanceData* GetComponentInstanceData() const override;
+	// End UActorComponent interface.
+
+	void ApplyComponentInstanceData(class FAtmospherePrecomputeInstanceData* ComponentInstanceData);
 	const FAtmospherePrecomputeParameters& GetPrecomputeParameters() const { return PrecomputeParams;  }
 
 private:
diff --git a/Engine/Source/Runtime/Engine/Private/Atmosphere/Atmosphere.cpp b/Engine/Source/Runtime/Engine/Private/Atmosphere/Atmosphere.cpp
--- a/Engine/Source/Runtime/Engine/Private/Atmosphere/Atmosphere.cpp
+++ b/Engine/Source/Runtime/Engine/Private/Atmosphere/Atmosphere.cpp
@@ -718,6 +718,115 @@ void UAtmosphericFogComponent::Serialize(FArchive& Ar)
 	}
 }
 
+/** Used to store lightmap data during RerunConstructionScripts */
+class FAtmospherePrecomputeInstanceData : public FSceneComponentInstanceData
+{
+public:
+	FAtmospherePrecomputeInstanceData(const UAtmosphericFogComponent* SourceComponent)
+		: FSceneComponentInstanceData(SourceComponent)
+	{}
+
+	virtual ~FAtmospherePrecomputeInstanceData()
+	{}
+
+	virtual void ApplyToComponent(UActorComponent* Component, const ECacheApplyPhase CacheApplyPhase) override
+	{
+		FSceneComponentInstanceData::ApplyToComponent(Component, CacheApplyPhase);
+		CastChecked<UAtmosphericFogComponent>(Component)->ApplyComponentInstanceData(this);
+	}
+
+	struct FAtmospherePrecomputeParameters PrecomputeParameter;
+
+	FByteBulkData TransmittanceData;
+	FByteBulkData IrradianceData;
+	FByteBulkData InscatterData;
+};
+
+// Backup the precomputed data before re-running Blueprint construction script
+FActorComponentInstanceData* UAtmosphericFogComponent::GetComponentInstanceData() const
+{
+	FActorComponentInstanceData* InstanceData = nullptr;
+
+	if (TransmittanceData.GetElementCount() && IrradianceData.GetElementCount() && InscatterData.GetElementCount() && PrecomputeCounter == EValid)
+	{
+		// Allocate new struct for holding light map data
+		FAtmospherePrecomputeInstanceData* PrecomputedData = new FAtmospherePrecomputeInstanceData(this);
+		InstanceData = PrecomputedData;
+
+		// Fill in info
+		PrecomputedData->PrecomputeParameter = PrecomputeParams;
+		{
+			int32 TotalByte = TransmittanceData.GetBulkDataSize();
+			PrecomputedData->TransmittanceData.Lock(LOCK_READ_WRITE);
+			void* OutData = PrecomputedData->TransmittanceData.Realloc(TotalByte);
+			TransmittanceData.GetCopy(&OutData, false);
+			PrecomputedData->TransmittanceData.Unlock();
+		}
+
+		{
+			int32 TotalByte = IrradianceData.GetBulkDataSize();
+			PrecomputedData->IrradianceData.Lock(LOCK_READ_WRITE);
+			void* OutData = PrecomputedData->IrradianceData.Realloc(TotalByte);
+			IrradianceData.GetCopy(&OutData, false);
+			PrecomputedData->IrradianceData.Unlock();
+		}
+
+		{
+			int32 TotalByte = InscatterData.GetBulkDataSize();
+			PrecomputedData->InscatterData.Lock(LOCK_READ_WRITE);
+			void* OutData = PrecomputedData->InscatterData.Realloc(TotalByte);
+			InscatterData.GetCopy(&OutData, false);
+			PrecomputedData->InscatterData.Unlock();
+		}
+	}
+	else
+	{
+		InstanceData = Super::GetComponentInstanceData();
+	}
+
+	return InstanceData;
+}
+
+// Restore the precomputed data after re-running Blueprint construction script
+void UAtmosphericFogComponent::ApplyComponentInstanceData(FAtmospherePrecomputeInstanceData* PrecomputedData)
+{
+	check(PrecomputedData);
+
+	if (PrecomputedData->PrecomputeParameter != GetPrecomputeParameters())
+	{
+		return;
+	}
+
+	FComponentReregisterContext ReregisterContext(this);
+	ReleaseResource();
+
+	{
+		int32 TotalByte = PrecomputedData->TransmittanceData.GetBulkDataSize();
+		TransmittanceData.Lock(LOCK_READ_WRITE);
+		void* OutData = TransmittanceData.Realloc(TotalByte);
+		PrecomputedData->TransmittanceData.GetCopy(&OutData, false);
+		TransmittanceData.Unlock();
+	}
+
+	{
+		int32 TotalByte = PrecomputedData->IrradianceData.GetBulkDataSize();
+		IrradianceData.Lock(LOCK_READ_WRITE);
+		void* OutData = IrradianceData.Realloc(TotalByte);
+		PrecomputedData->IrradianceData.GetCopy(&OutData, false);
+		IrradianceData.Unlock();
+	}
+
+	{
+		int32 TotalByte = PrecomputedData->InscatterData.GetBulkDataSize();
+		InscatterData.Lock(LOCK_READ_WRITE);
+		void* OutData = InscatterData.Realloc(TotalByte);
+		PrecomputedData->InscatterData.GetCopy(&OutData, false);
+		InscatterData.Unlock();
+	}
+
+	PrecomputeCounter = EValid;
+	InitResource();
+}
 
 /**
  * Gets called any time cvars change (on the main thread), we check if r.Atmosphere has changed and update the components.
