Hide some categories from RadialForceComponent that don't make sense

diff --git a/Engine/Source/Runtime/Engine/Classes/PhysicsEngine/RadialForceComponent.h b/Engine/Source/Runtime/Engine/Classes/PhysicsEngine/RadialForceComponent.h
--- a/Engine/Source/Runtime/Engine/Classes/PhysicsEngine/RadialForceComponent.h
+++ b/Engine/Source/Runtime/Engine/Classes/PhysicsEngine/RadialForceComponent.h
@@ -7,7 +7,7 @@
 /**
  *	Used to emit a radial force or impulse that can affect physics objects and or destructible objects.
  */
-UCLASS(hidecategories=(Object, Mobility), ClassGroup=Physics, showcategories=Trigger, meta=(BlueprintSpawnableComponent), MinimalAPI)
+UCLASS(hidecategories=(Object, Mobility, LOD, Physics), ClassGroup=Physics, showcategories=Trigger, meta=(BlueprintSpawnableComponent), MinimalAPI)
 class URadialForceComponent : public USceneComponent
 {
 	GENERATED_UCLASS_BODY()
