Fixed evaluations of tracks doing nothing if the current eval time ends up beyond the bounds of a section.  Different track types may handle this behavior differently.  for example shot tracks should only cut cameras if directly on the section.
Fixed evaluations of integral curves not returning the correct value when jumping across keyframes.

diff --git a/Engine/Source/Runtime/Engine/Classes/Curves/CurveBase.h b/Engine/Source/Runtime/Engine/Classes/Curves/CurveBase.h
--- a/Engine/Source/Runtime/Engine/Classes/Curves/CurveBase.h
+++ b/Engine/Source/Runtime/Engine/Classes/Curves/CurveBase.h
@@ -560,7 +560,7 @@ public:
 	virtual bool IsKeyHandleValid(FKeyHandle KeyHandle) const override;
 
 	/** Evaluates the value of an array of keys at a time */
-	int32 Evaluate(float Time) const;
+	int32 Evaluate(float Time, int32 DefaultValue = 0) const;
 
 	/** Const iterator for the keys, so the indices and handles stay valid */
 	TArray<FIntegralKey>::TConstIterator GetKeyIterator() const;
diff --git a/Engine/Source/Runtime/Engine/Private/CurveBase.cpp b/Engine/Source/Runtime/Engine/Private/CurveBase.cpp
--- a/Engine/Source/Runtime/Engine/Private/CurveBase.cpp
+++ b/Engine/Source/Runtime/Engine/Private/CurveBase.cpp
@@ -1020,16 +1020,37 @@ bool FIntegralCurve::IsKeyHandleValid(FKeyHandle KeyHandle) const
 	return bValid;
 }
 
-int32 FIntegralCurve::Evaluate(float Time) const
+int32 FIntegralCurve::Evaluate(float Time, int32 DefaultValue) const
 {
-	for (int32 i = 0; i < Keys.Num(); ++i)
+	int32 ReturnVal = DefaultValue;
+	if( Keys.Num() == 0 )
 	{
-		if (Time < Keys[i].Time)
+		ReturnVal = DefaultValue;
+	}
+	else if( Keys.Num() < 2 || Time < Keys[0].Time )
+	{
+		// There is only one key or the time is before the first value. Return the first value
+		ReturnVal = Keys[0].Value;
+	}
+	else if( Time < Keys[Keys.Num()-1].Time )
+	{
+		// The key is in the range of Key[0] to Keys[Keys.Num()-1].  Find it by searching
+		for(int32 i = 0; i < Keys.Num(); ++i)
 		{
-			return Keys[FMath::Max(0, i - 1)].Value;
+			if(Time < Keys[i].Time)
+			{
+				ReturnVal = Keys[FMath::Max(0, i - 1)].Value;
+				break;
+			}
 		}
 	}
-	return Keys.Num() ? Keys.Last().Value : 0;
+	else
+	{
+		// Key is beyon the last point in the curve.  Return it's value
+		ReturnVal = Keys[Keys.Num() - 1].Value;
+	}
+
+	return ReturnVal;
 }
 
 TArray<FIntegralKey>::TConstIterator FIntegralCurve::GetKeyIterator() const
diff --git a/Engine/Source/Runtime/MovieScene/Private/MovieSceneCommonHelpers.cpp b/Engine/Source/Runtime/MovieScene/Private/MovieSceneCommonHelpers.cpp
--- a/Engine/Source/Runtime/MovieScene/Private/MovieSceneCommonHelpers.cpp
+++ b/Engine/Source/Runtime/MovieScene/Private/MovieSceneCommonHelpers.cpp
@@ -30,18 +30,15 @@ TArray<UMovieSceneSection*> MovieSceneHelpers::GetTraversedSections( const TArra
 
 UMovieSceneSection* MovieSceneHelpers::FindSectionAtTime( const TArray<UMovieSceneSection*>& Sections, float Time )
 {
-	for( int32 SectionIndex = 0; SectionIndex < Sections.Num(); ++SectionIndex )
+	if( Sections.Num() == 1 )
 	{
-		UMovieSceneSection* Section = Sections[SectionIndex];
-
-		//@todo Sequencer - There can be multiple sections overlapping in time. Returning instantly does not account for that.
-		if( Section->IsTimeWithinSection( Time ) )
-		{
-			return Section;
-		}
+		return Sections[0];
+	}
+	else
+	{
+		return MovieSceneHelpers::FindNearestSectionAtTime( Sections, Time );
 	}
 
-	return NULL;
 }
 
 UMovieSceneSection* MovieSceneHelpers::FindNearestSectionAtTime( const TArray<UMovieSceneSection*>& Sections, float Time )
diff --git a/Engine/Source/Runtime/MovieScene/Private/MovieSceneInstance.cpp b/Engine/Source/Runtime/MovieScene/Private/MovieSceneInstance.cpp
--- a/Engine/Source/Runtime/MovieScene/Private/MovieSceneInstance.cpp
+++ b/Engine/Source/Runtime/MovieScene/Private/MovieSceneInstance.cpp
@@ -6,6 +6,7 @@
 FMovieSceneInstance::FMovieSceneInstance( UMovieScene& InMovieScene )
 	: MovieScene( &InMovieScene )
 {
+	TimeRange = MovieScene->GetTimeRange();
 }
 
 void FMovieSceneInstance::SaveState()
@@ -66,6 +67,8 @@ void FMovieSceneInstance::Update( float Position, float LastPosition, class IMov
 
 void FMovieSceneInstance::RefreshInstance( IMovieScenePlayer& Player )
 {
+	TimeRange = MovieScene->GetTimeRange();
+
 	UMovieSceneTrack* ShotTrack = MovieScene->GetShotTrack();
 
 	TSharedRef<FMovieSceneInstance> ThisInstance = AsShared();
diff --git a/Engine/Source/Runtime/MovieScene/Public/MovieSceneInstance.h b/Engine/Source/Runtime/MovieScene/Public/MovieSceneInstance.h
--- a/Engine/Source/Runtime/MovieScene/Public/MovieSceneInstance.h
+++ b/Engine/Source/Runtime/MovieScene/Public/MovieSceneInstance.h
@@ -36,6 +36,13 @@ public:
 
 	/** @return The movie scene being instanced */
 	UMovieScene* GetMovieScene() { return MovieScene.Get(); }
+
+	/** 
+	 * @return The time range for the movie scene 
+	 * Note this is not necessarily how long the instance will be playing.  The parent could have changed the length 
+	 */
+	TRange<float> GetMovieSceneTimeRange() const { return TimeRange; }
+
 private:
 	void RefreshInstanceMap( const TArray<UMovieSceneTrack*>& Tracks, const TArray<UObject*>& RuntimeObjects, FMovieSceneInstanceMap& TrackInstances, class IMovieScenePlayer& Player );
 private:
@@ -58,6 +65,8 @@ private:
 	FMovieSceneInstanceMap MasterTrackInstances;
 	/** All object binding instances */
 	TMap<FGuid, FMovieSceneObjectBindingInstance> ObjectBindingInstances;
+	/** Cached time range for the movie scene */
+	TRange<float> TimeRange;
 };
 
 
diff --git a/Engine/Source/Runtime/MovieSceneTracks/Private/TrackInstances/SubMovieSceneTrackInstance.cpp b/Engine/Source/Runtime/MovieSceneTracks/Private/TrackInstances/SubMovieSceneTrackInstance.cpp
--- a/Engine/Source/Runtime/MovieSceneTracks/Private/TrackInstances/SubMovieSceneTrackInstance.cpp
+++ b/Engine/Source/Runtime/MovieSceneTracks/Private/TrackInstances/SubMovieSceneTrackInstance.cpp
@@ -72,10 +72,7 @@ void FSubMovieSceneTrackInstance::Update( float Position, float LastPosition, co
 
 		TSharedPtr<FMovieSceneInstance>& Instance = SubMovieSceneInstances.FindChecked( Section );
 
-		UMovieScene* MovieScene = Section->GetMovieScene();
-
-		// @todo Sequencer we should not do this every update
-		TRange<float> TimeRange = MovieScene->GetTimeRange();
+		TRange<float> TimeRange = Instance->GetMovieSceneTimeRange();
 
 		// Position for the movie scene needs to be in local space
 		float LocalDelta = TimeRange.GetLowerBoundValue() - Section->GetStartTime();
