[UE-17742] Fix a potential crash during SCS component registration.

change summary:
- modified RegisterInstancedComponent() in SimpleConstructionScript.cpp to make a local copy of the AttachChildren array before iteration, in case one of the child components removes itself from the hierarchy during registration.

diff --git a/Engine/Source/Runtime/Engine/Private/SimpleConstructionScript.cpp b/Engine/Source/Runtime/Engine/Private/SimpleConstructionScript.cpp
--- a/Engine/Source/Runtime/Engine/Private/SimpleConstructionScript.cpp
+++ b/Engine/Source/Runtime/Engine/Private/SimpleConstructionScript.cpp
@@ -27,7 +27,9 @@ namespace
 			USceneComponent* InstancedSceneComponent = Cast<USceneComponent>(InstancedComponent);
 			if(InstancedSceneComponent != nullptr)
 			{
-				for(auto InstancedChildComponent : InstancedSceneComponent->AttachChildren)
+				// Make a copy of the array here in case something alters the AttachChildren array during registration (e.g. Physics)
+				TArray<USceneComponent*> AttachChildren = InstancedSceneComponent->AttachChildren;
+				for(auto InstancedChildComponent : AttachChildren)
 				{
 					RegisterInstancedComponent(InstancedChildComponent);
 				}
