Expose skeletal mesh motor control to blueprints

#UE-17963

diff --git a/Engine/Source/Runtime/Engine/Classes/Components/SkeletalMeshComponent.h b/Engine/Source/Runtime/Engine/Classes/Components/SkeletalMeshComponent.h
--- a/Engine/Source/Runtime/Engine/Classes/Components/SkeletalMeshComponent.h
+++ b/Engine/Source/Runtime/Engine/Classes/Components/SkeletalMeshComponent.h
@@ -1016,13 +1016,16 @@ public:
 	UFUNCTION(BlueprintCallable, Category="Components|SkeletalMesh")
 	void AccumulateAllBodiesBelowPhysicsBlendWeight(const FName& InBoneName, float AddPhysicsBlendWeight, bool bSkipCustomPhysicsType = false );
 
-	/** Enable or Disable AngularPositionDrive */
+	/** Enable or Disable AngularPositionDrive. If motor is in SLERP mode it will be turned on if either EnableSwingDrive OR EnableTwistDrive are enabled. In Twist and Swing mode the twist and the swing can be controlled individually.*/
+	UFUNCTION(BlueprintCallable, Category = "Components|SkeletalMesh")
 	void SetAllMotorsAngularPositionDrive(bool bEnableSwingDrive, bool bEnableTwistDrive, bool bSkipCustomPhysicsType = false);
 
-	/** Enable or Disable AngularVelocityDrive based on a list of bone names */
+	/** Enable or Disable AngularVelocityDrive. If motor is in SLERP mode it will be turned on if either EnableSwingDrive OR EnableTwistDrive are enabled. In Twist and Swing mode the twist and the swing can be controlled individually.*/
+	UFUNCTION(BlueprintCallable, Category = "Components|SkeletalMesh")
 	void SetAllMotorsAngularVelocityDrive(bool bEnableSwingDrive, bool bEnableTwistDrive, bool bSkipCustomPhysicsType = false);
 
-	/** Set Angular Drive motors params for all constraint instance */
+	/** Set Angular Drive motors params for all constraint instances */
+	UFUNCTION(BlueprintCallable, Category = "Components|SkeletalMesh")
 	void SetAllMotorsAngularDriveParams(float InSpring, float InDamping, float InForceLimit, bool bSkipCustomPhysicsType = false);
 
 	/** Enable or Disable AngularPositionDrive based on a list of bone names */
diff --git a/Engine/Source/Runtime/Engine/Private/PhysicsEngine/ConstraintInstance.cpp b/Engine/Source/Runtime/Engine/Private/PhysicsEngine/ConstraintInstance.cpp
--- a/Engine/Source/Runtime/Engine/Private/PhysicsEngine/ConstraintInstance.cpp
+++ b/Engine/Source/Runtime/Engine/Private/PhysicsEngine/ConstraintInstance.cpp
@@ -967,28 +967,25 @@ void FConstraintInstance::SetAngularVelocityTarget(const FVector& InVelTarget)
 void FConstraintInstance::SetAngularDriveParams(float InSpring, float InDamping, float InForceLimit)
 {
 #if WITH_PHYSX
-	if (bAngularOrientationDrive || bAngularVelocityDrive)
+	ExecuteOnUnbrokenJointReadWrite([&] (PxD6Joint* Joint)
 	{
-		ExecuteOnUnbrokenJointReadWrite([&] (PxD6Joint* Joint)
-		{
-			const float AngularForceLimit = InForceLimit > 0.0f ? InForceLimit : PX_MAX_F32;
-			const float DriveSpring = bAngularOrientationDrive ? InSpring : 0.0f;
-			const float DriveDamping = bAngularVelocityDrive ? InDamping : 0.0f;
+		const float AngularForceLimit = InForceLimit > 0.0f ? InForceLimit : PX_MAX_F32;
+		const float DriveSpring = bAngularOrientationDrive ? InSpring : 0.0f;
+		const float DriveDamping = bAngularVelocityDrive ? InDamping : 0.0f;
 
-			if (AngularDriveMode == EAngularDriveMode::SLERP)
-			{
-				Joint->setDrive(PxD6Drive::eTWIST, PxD6JointDrive());
-				Joint->setDrive(PxD6Drive::eSWING, PxD6JointDrive());
-				Joint->setDrive(PxD6Drive::eSLERP, (bEnableSwingDrive || bEnableTwistDrive) ? PxD6JointDrive(DriveSpring, DriveDamping, AngularForceLimit, bIsAccelerationDrive) : PxD6JointDrive());				
-			}
-			else
-			{
-				Joint->setDrive(PxD6Drive::eTWIST, bEnableTwistDrive ? PxD6JointDrive(DriveSpring, DriveDamping, AngularForceLimit, bIsAccelerationDrive) : PxD6JointDrive());
-				Joint->setDrive(PxD6Drive::eSWING, bEnableSwingDrive ? PxD6JointDrive(DriveSpring, DriveDamping, AngularForceLimit, bIsAccelerationDrive) : PxD6JointDrive());
-				Joint->setDrive(PxD6Drive::eSLERP, PxD6JointDrive());
-			}
-		});
-	}
+		if (AngularDriveMode == EAngularDriveMode::SLERP)
+		{
+			Joint->setDrive(PxD6Drive::eTWIST, PxD6JointDrive());
+			Joint->setDrive(PxD6Drive::eSWING, PxD6JointDrive());
+			Joint->setDrive(PxD6Drive::eSLERP, (bEnableSwingDrive || bEnableTwistDrive) ? PxD6JointDrive(DriveSpring, DriveDamping, AngularForceLimit, bIsAccelerationDrive) : PxD6JointDrive());				
+		}
+		else
+		{
+			Joint->setDrive(PxD6Drive::eTWIST, bEnableTwistDrive ? PxD6JointDrive(DriveSpring, DriveDamping, AngularForceLimit, bIsAccelerationDrive) : PxD6JointDrive());
+			Joint->setDrive(PxD6Drive::eSWING, bEnableSwingDrive ? PxD6JointDrive(DriveSpring, DriveDamping, AngularForceLimit, bIsAccelerationDrive) : PxD6JointDrive());
+			Joint->setDrive(PxD6Drive::eSLERP, PxD6JointDrive());
+		}
+	});
 #endif
 
 	AngularDriveSpring = InSpring;
