Compile Fix! Part Deux

diff --git a/Engine/Source/Runtime/Engine/Classes/Materials/MaterialExpressionDynamicParameter.h b/Engine/Source/Runtime/Engine/Classes/Materials/MaterialExpressionDynamicParameter.h
--- a/Engine/Source/Runtime/Engine/Classes/Materials/MaterialExpressionDynamicParameter.h
+++ b/Engine/Source/Runtime/Engine/Classes/Materials/MaterialExpressionDynamicParameter.h
@@ -10,8 +10,8 @@
 #include "Materials/MaterialExpression.h"
 #include "MaterialExpressionDynamicParameter.generated.h"
 
-UCLASS(collapsecategories, hidecategories=Object, MinimalAPI)
-class UMaterialExpressionDynamicParameter : public UMaterialExpression
+UCLASS(collapsecategories, hidecategories=Object)
+class ENGINE_API UMaterialExpressionDynamicParameter : public UMaterialExpression
 {
 	GENERATED_UCLASS_BODY()
 
