#ue4 - Make Kismet Transform comparison use the faster methods which don't require conversion when the vectorized FTransform is used.

diff --git a/Engine/Source/Runtime/Engine/Private/KismetMathLibrary.cpp b/Engine/Source/Runtime/Engine/Private/KismetMathLibrary.cpp
--- a/Engine/Source/Runtime/Engine/Private/KismetMathLibrary.cpp
+++ b/Engine/Source/Runtime/Engine/Private/KismetMathLibrary.cpp
@@ -1134,9 +1134,9 @@ bool UKismetMathLibrary::EqualEqual_TransformTransform(const FTransform& A, cons
 bool UKismetMathLibrary::NearlyEqual_TransformTransform(const FTransform& A, const FTransform& B, float LocationTolerance, float RotationTolerance, float Scale3DTolerance)
 {
 	return 
-		A.GetRotation().Equals(B.GetRotation(), RotationTolerance) && 
-		A.GetTranslation().Equals(B.GetTranslation(), LocationTolerance) && 
-		A.GetScale3D().Equals(B.GetScale3D(), Scale3DTolerance);
+		FTransform::AreRotationsEqual(A, B, RotationTolerance) && 
+		FTransform::AreTranslationsEqual(A, B, LocationTolerance) && 
+		FTransform::AreScale3DsEqual(A, B, Scale3DTolerance);
 }
 
 
