Fix substepping applying forces and setting interpolating target when the body's kinematic flag has changed. Also fix linear velocity being set on kinematic actors during initialization

#UE-14634

diff --git a/Engine/Source/Runtime/Engine/Private/PhysicsEngine/BodyInstance.cpp b/Engine/Source/Runtime/Engine/Private/PhysicsEngine/BodyInstance.cpp
--- a/Engine/Source/Runtime/Engine/Private/PhysicsEngine/BodyInstance.cpp
+++ b/Engine/Source/Runtime/Engine/Private/PhysicsEngine/BodyInstance.cpp
@@ -4432,8 +4432,6 @@ void FBodyInstance::InitDynamicProperties_AssumesLocked()
 			UpdateDampingProperties();
 			SetMaxAngularVelocity(MaxAngularVelocity, false);
 			SetMaxDepenetrationVelocity(bOverrideMaxDepenetrationVelocity ? MaxDepenetrationVelocity : UPhysicsSettings::Get()->MaxDepenetrationVelocity);
-
-			RigidActor->setLinearVelocity(U2PVector(InitialLinearVelocity));
 		}else
 		{
 #if !(UE_BUILD_SHIPPING || UE_BUILD_TEST)
@@ -4449,6 +4447,11 @@ void FBodyInstance::InitDynamicProperties_AssumesLocked()
 #endif
 		}
 
+		if (ShouldInstanceSimulatingPhysics())
+		{
+			RigidActor->setLinearVelocity(U2PVector(InitialLinearVelocity));
+		}
+
 		float SleepEnergyThresh = RigidActor->getSleepThreshold();
 		if (SleepFamily == ESleepFamily::Sensitive)
 		{
diff --git a/Engine/Source/Runtime/Engine/Private/PhysicsEngine/PhysSubstepTasks.cpp b/Engine/Source/Runtime/Engine/Private/PhysicsEngine/PhysSubstepTasks.cpp
--- a/Engine/Source/Runtime/Engine/Private/PhysicsEngine/PhysSubstepTasks.cpp
+++ b/Engine/Source/Runtime/Engine/Private/PhysicsEngine/PhysSubstepTasks.cpp
@@ -168,6 +168,14 @@ void FPhysSubstepTask::ApplyCustomPhysics(const FPhysTarget& PhysTarget, FBodyIn
 #endif
 }
 
+#if WITH_PHYSX
+bool IsKinematicHelper(const PxRigidBody* PRigidBody)
+{
+	const bool bIsKinematic = PRigidBody->getRigidDynamicFlags() & PxRigidDynamicFlag::eKINEMATIC;
+	return bIsKinematic;
+}
+#endif
+
 /** Applies forces - Assumes caller has obtained writer lock */
 void FPhysSubstepTask::ApplyForces_AssumesLocked(const FPhysTarget& PhysTarget, FBodyInstance* BodyInstance)
 {
@@ -281,11 +289,16 @@ void FPhysSubstepTask::SubstepInterpolation(float InAlpha, float DeltaTime)
 		//We should only be iterating over actors that belong to this scene
 		check(PRigidBody->getScene() == PScene);
 
-		ApplyCustomPhysics(PhysTarget, BodyInstance, DeltaTime);
-		ApplyForces_AssumesLocked(PhysTarget, BodyInstance);
-		ApplyTorques_AssumesLocked(PhysTarget, BodyInstance);
-		ApplyRadialForces_AssumesLocked(PhysTarget, BodyInstance);
-		InterpolateKinematicActor_AssumesLocked(PhysTarget, BodyInstance, InAlpha);
+		if (!IsKinematicHelper(PRigidBody))
+		{
+			ApplyCustomPhysics(PhysTarget, BodyInstance, DeltaTime);
+			ApplyForces_AssumesLocked(PhysTarget, BodyInstance);
+			ApplyTorques_AssumesLocked(PhysTarget, BodyInstance);
+			ApplyRadialForces_AssumesLocked(PhysTarget, BodyInstance);
+		}else
+		{
+			InterpolateKinematicActor_AssumesLocked(PhysTarget, BodyInstance, InAlpha);
+		}
 	}
 
 	/** Final substep */
