CameraComponent::SetCameraMesh will now update the proxy mesh if it exists (UE-16011)

diff --git a/Engine/Source/Runtime/Engine/Classes/Camera/CameraComponent.h b/Engine/Source/Runtime/Engine/Classes/Camera/CameraComponent.h
--- a/Engine/Source/Runtime/Engine/Classes/Camera/CameraComponent.h
+++ b/Engine/Source/Runtime/Engine/Classes/Camera/CameraComponent.h
@@ -102,7 +102,7 @@ protected:
 
 public:
 
-	virtual void SetCameraMesh(UStaticMesh* Mesh) { CameraMesh = Mesh; }
+	virtual void SetCameraMesh(UStaticMesh* Mesh);
 
 #endif
 public:
diff --git a/Engine/Source/Runtime/Engine/Private/Camera/CameraComponent.cpp b/Engine/Source/Runtime/Engine/Private/Camera/CameraComponent.cpp
--- a/Engine/Source/Runtime/Engine/Private/Camera/CameraComponent.cpp
+++ b/Engine/Source/Runtime/Engine/Private/Camera/CameraComponent.cpp
@@ -135,8 +135,20 @@ void UCameraComponent::PostLoad()
 	SetDeprecatedControllerViewRotation(*this, bUsePawnControlRotation);
 }
 
-
 #if WITH_EDITORONLY_DATA
+
+ void UCameraComponent::SetCameraMesh(UStaticMesh* Mesh)
+ {
+	 if (Mesh != CameraMesh)
+	 {
+		 CameraMesh = Mesh;
+
+		 if (ProxyMeshComponent)
+		 {
+			 ProxyMeshComponent->SetStaticMesh(Mesh);
+		 }
+	 }
+ }
 void UCameraComponent::RefreshVisualRepresentation()
 {
 	if (DrawFrustum != NULL)
@@ -178,7 +190,7 @@ void UCameraComponent::RestoreFrustumColor()
 		//Cam->DrawFrustum->FrustumColor = DefCam->DrawFrustum->FrustumColor;
 	}
 }
-#endif
+#endif	// WITH_EDITORONLY_DATA
 
 #if WITH_EDITOR
 void UCameraComponent::PostEditChangeProperty(FPropertyChangedEvent& PropertyChangedEvent)
@@ -187,9 +199,9 @@ void UCameraComponent::PostEditChangeProperty(FPropertyChangedEvent& PropertyCha
 
 #if WITH_EDITORONLY_DATA
 	RefreshVisualRepresentation();
-#endif
+#endif	// WITH_EDITORONLY_DATA
 }
-#endif
+#endif	// WITH_EDITOR
 
 void UCameraComponent::Serialize(FArchive& Ar)
 {
@@ -248,7 +260,7 @@ void UCameraComponent::CheckForErrors()
 			->AddToken(FMapErrorToken::Create(FMapErrors::CameraAspectRatioIsZero));
 	}
 }
-#endif
+#endif	// WITH_EDITOR
 
 
 void SetDeprecatedControllerViewRotation(UCameraComponent& Component, bool bValue)
