#ue4 - Improve documentation for BrakingFriction.

diff --git a/Engine/Source/Runtime/Engine/Classes/GameFramework/CharacterMovementComponent.h b/Engine/Source/Runtime/Engine/Classes/GameFramework/CharacterMovementComponent.h
--- a/Engine/Source/Runtime/Engine/Classes/GameFramework/CharacterMovementComponent.h
+++ b/Engine/Source/Runtime/Engine/Classes/GameFramework/CharacterMovementComponent.h
@@ -233,6 +233,7 @@ public:
 	/**
 	 * Friction (drag) coefficient applied when braking (whenever Acceleration = 0, or if character is exceeding max speed).
 	 * Only used if bUseSeparateBrakingFriction setting is true, otherwise current friction such as GroundFriction is used.
+	 * This is the current value, used in all movement modes; if this is not desired, override it when movement mode changes.
 	 * @see bUseSeparateBrakingFriction
 	 */
 	UPROPERTY(Category="Character Movement (General Settings)", EditAnywhere, BlueprintReadWrite, meta=(ClampMin="0", UIMin="0", EditCondition="bUseSeparateBrakingFriction"))
