Improve clarity of Tick Interval units

diff --git a/Engine/Source/Runtime/Engine/Classes/Engine/EngineBaseTypes.h b/Engine/Source/Runtime/Engine/Classes/Engine/EngineBaseTypes.h
--- a/Engine/Source/Runtime/Engine/Classes/Engine/EngineBaseTypes.h
+++ b/Engine/Source/Runtime/Engine/Classes/Engine/EngineBaseTypes.h
@@ -209,8 +209,8 @@ private:
 	float RelativeTickCooldown;
 public:
 
-	/** The frequency at which this tick function will be executed.  If less than or equal to 0 then it will tick every frame */
-	UPROPERTY(EditDefaultsOnly, Category="Tick")
+	/** The frequency in seconds at which this tick function will be executed.  If less than or equal to 0 then it will tick every frame */
+	UPROPERTY(EditDefaultsOnly, Category="Tick", meta=(DisplayName="Tick Interval (secs)"))
 	float TickInterval;
 
 	/** Back pointer to the FTickTaskLevel containing this tick function if it is registered **/
