Core: flush server logs in a deterministic time (UEPLAT-941).

- Use ServerFlushLogInterval in GameEngine to control the maximum time limit when logs should be flushed (30 seconds by default).

diff --git a/Engine/Config/BaseEngine.ini b/Engine/Config/BaseEngine.ini
--- a/Engine/Config/BaseEngine.ini
+++ b/Engine/Config/BaseEngine.ini
@@ -1317,6 +1317,7 @@ UIScaleCurve=(EditorCurveData=(Keys=((Time=480,Value=0.444),(Time=720,Value=0.66
 
 [/Script/Engine.GameEngine]
 MaxDeltaTime=0
+ServerFlushLogInterval=30
 
 [Engine.StartupPackages]
 bSerializeStartupPackagesFromMemory=true
diff --git a/Engine/Source/Runtime/Engine/Classes/Engine/GameEngine.h b/Engine/Source/Runtime/Engine/Classes/Engine/GameEngine.h
--- a/Engine/Source/Runtime/Engine/Classes/Engine/GameEngine.h
+++ b/Engine/Source/Runtime/Engine/Classes/Engine/GameEngine.h
@@ -23,6 +23,10 @@ class ENGINE_API UGameEngine
 	UPROPERTY(config)
 	float MaxDeltaTime;
 
+	/** Maximium time (in seconds) between the flushes of the logs on the server (best effort). If 0, this will happen every tick. */
+	UPROPERTY(config)
+	float ServerFlushLogInterval;
+
 	UPROPERTY(transient)
 	UGameInstance* GameInstance;
 
@@ -149,4 +153,7 @@ private:
 
 	virtual void HandleNetworkFailure_NotifyGameInstance(UWorld* World, UNetDriver* NetDriver, ENetworkFailure::Type FailureType) override;
 	virtual void HandleTravelFailure_NotifyGameInstance(UWorld* World, ETravelFailure::Type FailureType) override;
+
+	/** Last time the logs have been flushed. */
+	double LastTimeLogsFlushed;
 };
diff --git a/Engine/Source/Runtime/Engine/Private/GameEngine.cpp b/Engine/Source/Runtime/Engine/Private/GameEngine.cpp
--- a/Engine/Source/Runtime/Engine/Private/GameEngine.cpp
+++ b/Engine/Source/Runtime/Engine/Private/GameEngine.cpp
@@ -458,6 +458,7 @@ void UGameEngine::Init(IEngineLoop* InEngineLoop)
 	}
 
 	bCheckForMovieCapture = true;
+	LastTimeLogsFlushed = FPlatformTime::Seconds();
 
 	// Attach the viewport client to a new viewport.
 	if(ViewportClient)
@@ -822,7 +823,17 @@ void UGameEngine::Tick( float DeltaSeconds, bool bIdleMode )
 		HotReload->Tick();
 	}
 
-	if (!IsRunningDedicatedServer() && !IsRunningCommandlet())
+	if (IsRunningDedicatedServer())
+	{
+		double CurrentTime = FPlatformTime::Seconds();
+		if (CurrentTime - LastTimeLogsFlushed > static_cast<double>(ServerFlushLogInterval))
+		{
+			GLog->Flush();
+
+			LastTimeLogsFlushed = FPlatformTime::Seconds();
+		}
+	}
+	else if (!IsRunningCommandlet())
 	{
 		// Clean up the game viewports that have been closed.
 		CleanupGameViewport();
