Fixed a crash when a SkelVert/SurfLocaton when it accessed a models color buffer which didn't exist.
Looks like these were recently made optional.

diff --git a/Engine/Source/Runtime/Engine/Private/Components/SkinnedMeshComponent.cpp b/Engine/Source/Runtime/Engine/Private/Components/SkinnedMeshComponent.cpp
--- a/Engine/Source/Runtime/Engine/Private/Components/SkinnedMeshComponent.cpp
+++ b/Engine/Source/Runtime/Engine/Private/Components/SkinnedMeshComponent.cpp
@@ -1655,13 +1655,19 @@ void USkinnedMeshComponent::ComputeSkinnedPositions(TArray<FVector> & OutPositio
 
 FColor USkinnedMeshComponent::GetVertexColor(int32 VertexIndex) const
 {
-	// Fail if no mesh
+	// Fail if no mesh or no color vertex buffer.
+	FColor FallbackColor = FColor(255, 255, 255, 255);
 	if (!SkeletalMesh || !MeshObject)
 	{
-		return FColor(255, 255, 255, 255);
+		return FallbackColor;
 	}
 
 	FStaticLODModel& Model = MeshObject->GetSkeletalMeshResource().LODModels[0];
+	
+	if (!Model.ColorVertexBuffer.IsInitialized())
+	{
+		return FallbackColor;
+	}
 
 	// Find the chunk and vertex within that chunk, and skinning type, for this vertex.
 	int32 ChunkIndex;
