Adding default materials to root as they appear to be GCd in some cases.
In response to https://udn.unrealengine.com/questions/251160/gdefaultmaterials-in-materialcpp.html

diff --git a/Engine/Source/Runtime/Engine/Private/Materials/Material.cpp b/Engine/Source/Runtime/Engine/Private/Materials/Material.cpp
--- a/Engine/Source/Runtime/Engine/Private/Materials/Material.cpp
+++ b/Engine/Source/Runtime/Engine/Private/Materials/Material.cpp
@@ -355,9 +355,11 @@ void UMaterialInterface::InitDefaultMaterials()
 #if WITH_EDITOR
 		GPowerToRoughnessMaterialFunction = LoadObject< UMaterialFunction >( NULL, TEXT("/Engine/Functions/Engine_MaterialFunctions01/Shading/PowerToRoughness.PowerToRoughness"), NULL, LOAD_None, NULL );
 		checkf( GPowerToRoughnessMaterialFunction, TEXT("Cannot load PowerToRoughness") );
+		GPowerToRoughnessMaterialFunction->AddToRoot();
 
 		GConvertFromDiffSpecMaterialFunction = LoadObject< UMaterialFunction >( NULL, TEXT("/Engine/Functions/Engine_MaterialFunctions01/Shading/ConvertFromDiffSpec.ConvertFromDiffSpec"), NULL, LOAD_None, NULL );
 		checkf( GConvertFromDiffSpecMaterialFunction, TEXT("Cannot load ConvertFromDiffSpec") );
+		GConvertFromDiffSpecMaterialFunction->AddToRoot();
 #endif
 
 		for (int32 Domain = 0; Domain < MD_MAX; ++Domain)
@@ -370,6 +372,7 @@ void UMaterialInterface::InitDefaultMaterials()
 					GDefaultMaterials[Domain] = LoadObject<UMaterial>(NULL,GDefaultMaterialNames[Domain],NULL,LOAD_None,NULL);
 					checkf(GDefaultMaterials[Domain] != NULL, TEXT("Cannot load default material '%s'"), GDefaultMaterialNames[Domain]);
 				}
+				GDefaultMaterials[Domain]->AddToRoot();
 			}
 		}
 		bInitialized = true;
