Non-Linear blend support for blend lists and state transitions

diff --git a/Engine/Source/Editor/AnimGraph/Classes/AnimStateTransitionNode.h b/Engine/Source/Editor/AnimGraph/Classes/AnimStateTransitionNode.h
--- a/Engine/Source/Editor/AnimGraph/Classes/AnimStateTransitionNode.h
+++ b/Engine/Source/Editor/AnimGraph/Classes/AnimStateTransitionNode.h
@@ -31,8 +31,14 @@ class UAnimStateTransitionNode : public UAnimStateNodeBase
 	float CrossfadeDuration;
 
 	// The type of blending to use in the crossfade
-	UPROPERTY(EditAnywhere, Config, Category=Transition)
-	TEnumAsByte<ETransitionBlendMode::Type> CrossfadeMode;
+	UPROPERTY()
+	TEnumAsByte<ETransitionBlendMode::Type> CrossfadeMode_DEPRECATED;
+
+	UPROPERTY(EditAnywhere, Category=Transition)
+	EAlphaBlendOption BlendMode;
+
+	UPROPERTY(EditAnywhere, Category=Transition)
+	UCurveFloat* CustomBlendCurve;
 
 	// Try setting the rule automatically based on the player node remaining time and the CrossfadeDuration, ignoring the internal time
 	UPROPERTY(EditAnywhere, Category=Transition)
diff --git a/Engine/Source/Editor/AnimGraph/Private/AnimStateTransitionNode.cpp b/Engine/Source/Editor/AnimGraph/Private/AnimStateTransitionNode.cpp
--- a/Engine/Source/Editor/AnimGraph/Private/AnimStateTransitionNode.cpp
+++ b/Engine/Source/Editor/AnimGraph/Private/AnimStateTransitionNode.cpp
@@ -68,7 +68,7 @@ UAnimStateTransitionNode::UAnimStateTransitionNode(const FObjectInitializer& Obj
 {
 
 	CrossfadeDuration = 0.2f;
-	CrossfadeMode = ETransitionBlendMode::TBM_Cubic;
+	BlendMode = EAlphaBlendOption::HermiteCubic;
 	bSharedRules = false;
 	SharedRulesGuid.Invalidate();
 	bSharedCrossfade = false;
@@ -107,6 +107,22 @@ void UAnimStateTransitionNode::PostLoad()
 	{
 		FAnimStateTransitionNodeSharedCrossfadeHelper().MakeSureGuidExists(this);
 	}
+
+	if(GetLinkerUE4Version() < VER_UE4_ADDED_NON_LINEAR_TRANSITION_BLENDS)
+	{
+		switch(CrossfadeMode_DEPRECATED)
+		{
+			case ETransitionBlendMode::TBM_Linear:
+				BlendMode = EAlphaBlendOption::Linear;
+				break;
+			case ETransitionBlendMode::TBM_Cubic:
+				// Old cubic was actually an in/out hermite polynomial (FMath::SmoothStep)
+				BlendMode = EAlphaBlendOption::HermiteCubic;
+				break;
+			default:
+				break;
+		}
+	}
 }
 
 void UAnimStateTransitionNode::PostPasteNode()
@@ -352,7 +368,8 @@ void UAnimStateTransitionNode::UseSharedCrossfade(const UAnimStateTransitionNode
 void UAnimStateTransitionNode::CopyCrossfadeSettings(const UAnimStateTransitionNode* SrcNode)
 {
 	CrossfadeDuration = SrcNode->CrossfadeDuration;
-	CrossfadeMode = SrcNode->CrossfadeMode;
+	CrossfadeMode_DEPRECATED = SrcNode->CrossfadeMode_DEPRECATED;
+	BlendMode = SrcNode->BlendMode;
 	SharedCrossfadeIdx = SrcNode->SharedCrossfadeIdx;
 	SharedCrossfadeName = SrcNode->SharedCrossfadeName;
 	SharedCrossfadeGuid = SrcNode->SharedCrossfadeGuid;
diff --git a/Engine/Source/Editor/DetailCustomizations/Private/AnimTransitionNodeDetails.cpp b/Engine/Source/Editor/DetailCustomizations/Private/AnimTransitionNodeDetails.cpp
--- a/Engine/Source/Editor/DetailCustomizations/Private/AnimTransitionNodeDetails.cpp
+++ b/Engine/Source/Editor/DetailCustomizations/Private/AnimTransitionNodeDetails.cpp
@@ -55,7 +55,7 @@ void FAnimTransitionNodeDetails::CustomizeDetails( IDetailLayoutBuilder& DetailB
 		// All of the blend related settings are ignored, so hide them.
 		DetailBuilder.HideProperty(GET_MEMBER_NAME_CHECKED(UAnimStateTransitionNode, Bidirectional));
 		DetailBuilder.HideProperty(GET_MEMBER_NAME_CHECKED(UAnimStateTransitionNode, CrossfadeDuration));
-		DetailBuilder.HideProperty(GET_MEMBER_NAME_CHECKED(UAnimStateTransitionNode, CrossfadeMode));
+		DetailBuilder.HideProperty(GET_MEMBER_NAME_CHECKED(UAnimStateTransitionNode, BlendMode));
 		DetailBuilder.HideProperty(GET_MEMBER_NAME_CHECKED(UAnimStateTransitionNode, LogicType));
 		DetailBuilder.HideProperty(GET_MEMBER_NAME_CHECKED(UAnimStateTransitionNode, PriorityOrder));
 	}
@@ -153,7 +153,8 @@ void FAnimTransitionNodeDetails::CustomizeDetails( IDetailLayoutBuilder& DetailB
 
 		//@TODO: Gate editing these on shared non-authorative ones
 		CrossfadeCategory.AddProperty(GET_MEMBER_NAME_CHECKED(UAnimStateTransitionNode, CrossfadeDuration)).DisplayName( LOCTEXT("DurationLabel", "Duration") );
-		CrossfadeCategory.AddProperty(GET_MEMBER_NAME_CHECKED(UAnimStateTransitionNode, CrossfadeMode)).DisplayName( LOCTEXT("ModeLabel", "Mode") );
+		CrossfadeCategory.AddProperty(GET_MEMBER_NAME_CHECKED(UAnimStateTransitionNode, BlendMode)).DisplayName( LOCTEXT("ModeLabel", "Mode") );
+		CrossfadeCategory.AddProperty(GET_MEMBER_NAME_CHECKED(UAnimStateTransitionNode, CustomBlendCurve)).DisplayName(LOCTEXT("CurveLabel", "Custom Blend Curve"));
 
 		// Add a button that is only visible when blend logic type is custom
 		CrossfadeCategory.AddCustomRow( LOCTEXT("EditBlendGraph", "Edit Blend Graph") )
diff --git a/Engine/Source/Editor/GraphEditor/Private/AnimationStateNodes/SGraphNodeAnimTransition.cpp b/Engine/Source/Editor/GraphEditor/Private/AnimationStateNodes/SGraphNodeAnimTransition.cpp
--- a/Engine/Source/Editor/GraphEditor/Private/AnimationStateNodes/SGraphNodeAnimTransition.cpp
+++ b/Engine/Source/Editor/GraphEditor/Private/AnimationStateNodes/SGraphNodeAnimTransition.cpp
@@ -305,13 +305,6 @@ FString SGraphNodeAnimTransition::GetCurrentDuration() const
 	return FString::Printf(TEXT("%.2f seconds"), TransNode->CrossfadeDuration);
 }
 
-FString SGraphNodeAnimTransition::GetCurrentAlphaCurveMode() const
-{
-	UAnimStateTransitionNode* TransNode = CastChecked<UAnimStateTransitionNode>(GraphNode);
-
-	return (TransNode->CrossfadeMode == ETransitionBlendMode::TBM_Linear) ? TEXT("Linear") : TEXT("Cubic");
-}
-
 void SGraphNodeAnimTransition::OnMouseEnter(const FGeometry& MyGeometry, const FPointerEvent& MouseEvent)
 {
 	UAnimStateTransitionNode* TransNode = CastChecked<UAnimStateTransitionNode>(GraphNode);
diff --git a/Engine/Source/Editor/GraphEditor/Private/AnimationStateNodes/SGraphNodeAnimTransition.h b/Engine/Source/Editor/GraphEditor/Private/AnimationStateNodes/SGraphNodeAnimTransition.h
--- a/Engine/Source/Editor/GraphEditor/Private/AnimationStateNodes/SGraphNodeAnimTransition.h
+++ b/Engine/Source/Editor/GraphEditor/Private/AnimationStateNodes/SGraphNodeAnimTransition.h
@@ -45,5 +45,4 @@ private:
 	TSharedRef<SWidget> GenerateRichTooltip();
 
 	FString GetCurrentDuration() const;
-	FString GetCurrentAlphaCurveMode() const;
 };
diff --git a/Engine/Source/Editor/KismetCompiler/Private/AnimBlueprintCompiler.cpp b/Engine/Source/Editor/KismetCompiler/Private/AnimBlueprintCompiler.cpp
--- a/Engine/Source/Editor/KismetCompiler/Private/AnimBlueprintCompiler.cpp
+++ b/Engine/Source/Editor/KismetCompiler/Private/AnimBlueprintCompiler.cpp
@@ -805,7 +805,8 @@ void FAnimBlueprintCompiler::ProcessStateMachine(UAnimGraphNode_StateMachineBase
 			BakedTransition.StartNotify = FindOrAddNotify(TransitionNode->TransitionStart);
 			BakedTransition.EndNotify = FindOrAddNotify(TransitionNode->TransitionEnd);
 			BakedTransition.InterruptNotify = FindOrAddNotify(TransitionNode->TransitionInterrupt);
-			BakedTransition.CrossfadeMode = TransitionNode->CrossfadeMode;
+			BakedTransition.BlendMode = TransitionNode->BlendMode;
+			BakedTransition.CustomCurve = TransitionNode->CustomBlendCurve;
 			BakedTransition.LogicType = TransitionNode->LogicType;
 
 			UAnimStateNodeBase* PreviousState = TransitionNode->GetPreviousState();
diff --git a/Engine/Source/Runtime/Core/Public/Math/AlphaBlendType.h b/Engine/Source/Runtime/Core/Public/Math/AlphaBlendType.h
--- a/Engine/Source/Runtime/Core/Public/Math/AlphaBlendType.h
+++ b/Engine/Source/Runtime/Core/Public/Math/AlphaBlendType.h
@@ -37,6 +37,7 @@ FORCEINLINE float AlphaToBlendType(float InAlpha, uint8 BlendType)
 /**
  * Alpha Blend Type
  */
+DEPRECATED(4.9, "FTAlphaBlend is deprecated, please use FAlphaBlend instead")
 struct FTAlphaBlend
 {
 	/** Internal Lerped value for Alpha */
diff --git a/Engine/Source/Runtime/Core/Public/UObject/ObjectVersion.h b/Engine/Source/Runtime/Core/Public/UObject/ObjectVersion.h
--- a/Engine/Source/Runtime/Core/Public/UObject/ObjectVersion.h
+++ b/Engine/Source/Runtime/Core/Public/UObject/ObjectVersion.h
@@ -546,6 +546,8 @@ enum EUnrealEngineObjectUE4Version
 	VER_UE4_DYNAMIC_PARAMETER_DEFAULT_VALUE,
 	// Added ExtendedBounds to StaticMesh
 	VER_UE4_STATIC_MESH_EXTENDED_BOUNDS,
+	// Added non-linear blending to anim transitions, deprecating old types
+	VER_UE4_ADDED_NON_LINEAR_TRANSITION_BLENDS,
 	// -----<new versions can be added before this line>-------------------------------------------------
 	// - this needs to be the last line (see note below)
 	VER_UE4_AUTOMATIC_VERSION_PLUS_ONE,
diff --git a/Engine/Source/Runtime/Engine/Classes/Animation/AnimNode_BlendListBase.h b/Engine/Source/Runtime/Engine/Classes/Animation/AnimNode_BlendListBase.h
--- a/Engine/Source/Runtime/Engine/Classes/Animation/AnimNode_BlendListBase.h
+++ b/Engine/Source/Runtime/Engine/Classes/Animation/AnimNode_BlendListBase.h
@@ -2,6 +2,7 @@
 
 #pragma once
 
+#include "AlphaBlend.h"
 #include "AnimNodeBase.h"
 #include "AnimNode_BlendListBase.generated.h"
 
@@ -17,6 +18,15 @@ public:
 	UPROPERTY(EditAnywhere, EditFixedSize, BlueprintReadWrite, Category=Config, meta=(PinShownByDefault))
 	TArray<float> BlendTime;
 
+	UPROPERTY(EditAnywhere, Category=BlendType)
+	EAlphaBlendOption BlendType;
+
+	UPROPERTY(EditAnywhere, Category=BlendType)
+	UCurveFloat* CustomBlendCurve;
+
+	UPROPERTY()
+	TArray<FAlphaBlend> Blends;
+
 protected:
 	UPROPERTY()
 	TArray<float> BlendWeights;
diff --git a/Engine/Source/Runtime/Engine/Classes/Animation/AnimNode_StateMachine.h b/Engine/Source/Runtime/Engine/Classes/Animation/AnimNode_StateMachine.h
--- a/Engine/Source/Runtime/Engine/Classes/Animation/AnimNode_StateMachine.h
+++ b/Engine/Source/Runtime/Engine/Classes/Animation/AnimNode_StateMachine.h
@@ -23,8 +23,8 @@ struct FAnimationActiveTransitionEntry
 	// Duration of this cross-fade (may be shorter than the nominal duration specified by the state machine if the target state had non-zero weight at the start)
 	float CrossfadeDuration;
 
-	// Blend type (type of curve applied to time)
-	TEnumAsByte<ETransitionBlendMode::Type> CrossfadeMode;
+	// Type of blend to use
+	EAlphaBlendOption BlendOption;
 
 	// Is this transition active?
 	bool bActive;
@@ -58,8 +58,8 @@ struct FAnimationActiveTransitionEntry
 
 public:
 	FAnimationActiveTransitionEntry();
-	FAnimationActiveTransitionEntry(int32 NextStateID, float ExistingWeightOfNextState, int32 PreviousStateID, const FAnimationTransitionBetweenStates& ReferenceTransitionInfo);
-	
+	FAnimationActiveTransitionEntry(int32 NextStateID, float ExistingWeightOfNextState, FAnimationActiveTransitionEntry* ExistingTransitionForNextState, int32 PreviousStateID, const FAnimationTransitionBetweenStates& ReferenceTransitionInfo);
+
 	void InitializeCustomGraphLinks(const FAnimationUpdateContext& Context, const FBakedStateExitTransition& TransitionRule);
 
 	void Update(const FAnimationUpdateContext& Context, int32 CurrentStateIndex, bool &OutFinished);
@@ -70,8 +70,11 @@ public:
 	bool Serialize(FArchive& Ar);
 
 protected:
-	float CalculateInverseAlpha(ETransitionBlendMode::Type BlendType, float InFraction) const;
+	float CalculateInverseAlpha(EAlphaBlendOption BlendMode, float InFraction) const;
 	float CalculateAlpha(float InFraction) const;
+
+	// Blend object to handle alpha interpolation
+	FAlphaBlend Blend;
 };
 
 USTRUCT()
diff --git a/Engine/Source/Runtime/Engine/Classes/Animation/AnimStateMachineTypes.h b/Engine/Source/Runtime/Engine/Classes/Animation/AnimStateMachineTypes.h
--- a/Engine/Source/Runtime/Engine/Classes/Animation/AnimStateMachineTypes.h
+++ b/Engine/Source/Runtime/Engine/Classes/Animation/AnimStateMachineTypes.h
@@ -2,6 +2,7 @@
 
 #pragma once
 
+#include "AlphaBlend.h"
 #include "AnimStateMachineTypes.generated.h"
 
 //@TODO: Document
@@ -128,7 +129,10 @@ struct FAnimationTransitionBetweenStates : public FAnimationStateBase
 	int32 InterruptNotify;
 
 	UPROPERTY()
-	TEnumAsByte<ETransitionBlendMode::Type> CrossfadeMode;
+	EAlphaBlendOption BlendMode;
+
+	UPROPERTY()
+	UCurveFloat* CustomCurve;
 
 	UPROPERTY()
 	TEnumAsByte<ETransitionLogicType::Type> LogicType;
@@ -145,7 +149,8 @@ struct FAnimationTransitionBetweenStates : public FAnimationStateBase
 		, StartNotify(INDEX_NONE)
 		, EndNotify(INDEX_NONE)
 		, InterruptNotify(INDEX_NONE)
-		, CrossfadeMode(ETransitionBlendMode::TBM_Cubic)
+		, BlendMode(EAlphaBlendOption::CubicInOut)
+		, CustomCurve(nullptr)
 		, LogicType(ETransitionLogicType::TLT_StandardBlend)
 #if WITH_EDITOR
 		, ReverseTransition(false)
diff --git a/Engine/Source/Runtime/Engine/Private/AlphaBlend.cpp b/Engine/Source/Runtime/Engine/Private/AlphaBlend.cpp
new file mode 100644
--- /dev/null
+++ b/Engine/Source/Runtime/Engine/Private/AlphaBlend.cpp
@@ -0,0 +1,165 @@
+// Copyright 1998-2015 Epic Games, Inc. All Rights Reserved.
+
+#include "EnginePrivate.h"
+#include "AlphaBlend.h"
+
+FAlphaBlend::FAlphaBlend() 
+	: BlendOption(EAlphaBlendOption::Linear)
+	, BeginValue(0.0f)
+	, DesiredValue(1.0f)
+	, BlendTime(0.2f)
+	, CustomCurve(nullptr)
+{
+	Reset();
+}
+
+void FAlphaBlend::SetBlendTime(float InBlendTime)
+{
+	BlendTime = FMath::Max(InBlendTime, 0.f);
+}
+
+void FAlphaBlend::Reset()
+{
+	AlphaLerp = 0.0f;
+	AlphaBlend = 0.0f;
+	BlendedValue = BeginValue;
+
+	// Set alpha target to full - will also handle zero blend times
+	SetTarget(1.0f);
+}
+
+bool FAlphaBlend::GetToggleStatus()
+{
+	return (AlphaTarget > 0.f);
+}
+
+void FAlphaBlend::Toggle(bool bEnable)
+{
+	ConditionalSetTarget(bEnable ? 1.f : 0.f);
+}
+
+void FAlphaBlend::ConditionalSetTarget(float InAlphaTarget)
+{
+	if( AlphaTarget != InAlphaTarget )
+	{
+		SetTarget(InAlphaTarget);
+	}
+}
+
+void FAlphaBlend::SetTarget(float InAlphaTarget)
+{
+	// Clamp parameters to valid range
+	AlphaTarget = FMath::Clamp<float>(InAlphaTarget, 0.f, 1.f);
+
+	// if blend time is zero, transition now, don't wait to call update.
+	if( BlendTime <= 0.f )
+	{
+		AlphaLerp = AlphaTarget;
+		AlphaBlend = AlphaToBlendOption();
+		BlendTimeRemaining = 0.f;
+		BlendedValue = BeginValue + (DesiredValue - BeginValue) * AlphaBlend;
+	}
+	else
+	{
+		// Blend time is to go all the way, so scale that by how much we have to travel
+		BlendTimeRemaining = BlendTime * FMath::Abs(AlphaTarget - AlphaLerp);
+	}
+}
+
+void FAlphaBlend::Update(float InDeltaTime)
+{
+	// Make sure passed in delta time is positive
+	check(InDeltaTime >= 0.f);
+
+	if( !IsComplete() )
+	{
+		if( BlendTimeRemaining > InDeltaTime )
+		{
+			const float BlendDelta = AlphaTarget - AlphaLerp; 
+			AlphaLerp += (BlendDelta / BlendTimeRemaining) * InDeltaTime;
+			BlendTimeRemaining -= InDeltaTime;
+
+			AlphaBlend = AlphaToBlendOption();
+			BlendedValue = BeginValue + (DesiredValue - BeginValue) * AlphaBlend;
+		}
+		else
+		{
+			BlendTimeRemaining = 0.f; 
+			AlphaLerp = 1.0f;
+			AlphaBlend = 1.0f;
+			BlendedValue = DesiredValue;
+		}
+	}
+}
+
+float FAlphaBlend::AlphaToBlendOption()
+{
+	return AlphaToBlendOption(AlphaLerp, BlendOption, CustomCurve);
+}
+
+float FAlphaBlend::AlphaToBlendOption(float InAlpha, EAlphaBlendOption InBlendOption, UCurveFloat* InCustomCurve)
+{
+	switch(InBlendOption)
+	{
+		case EAlphaBlendOption::Sinusoidal:		return FMath::Clamp<float>((FMath::Sin(InAlpha * PI - HALF_PI) + 1.f) / 2.f, 0.f, 1.f);
+		case EAlphaBlendOption::Cubic:			return FMath::Clamp<float>(FMath::CubicInterp<float>(0.f, 0.f, 1.f, 0.f, InAlpha), 0.f, 1.f);
+		case EAlphaBlendOption::QuadraticInOut: return FMath::Clamp<float>(FMath::InterpEaseInOut<float>(0.f, 1.f, InAlpha, 2), 0.f, 1.f);
+		case EAlphaBlendOption::CubicInOut:		return FMath::Clamp<float>(FMath::InterpEaseInOut<float>(0.f, 1.f, InAlpha, 3), 0.f, 1.f);
+		case EAlphaBlendOption::HermiteCubic:	return FMath::Clamp<float>(FMath::SmoothStep(0.0f, 1.0f, InAlpha), 0.0f, 1.0f);
+		case EAlphaBlendOption::QuarticInOut:	return FMath::Clamp<float>(FMath::InterpEaseInOut<float>(0.f, 1.f, InAlpha, 4), 0.f, 1.f);
+		case EAlphaBlendOption::QuinticInOut:	return FMath::Clamp<float>(FMath::InterpEaseInOut<float>(0.f, 1.f, InAlpha, 5), 0.f, 1.f);
+		case EAlphaBlendOption::CircularIn:		return FMath::Clamp<float>(FMath::InterpCircularIn<float>(0.0f, 1.0f, InAlpha), 0.0f, 1.0f);
+		case EAlphaBlendOption::CircularOut:	return FMath::Clamp<float>(FMath::InterpCircularOut<float>(0.0f, 1.0f, InAlpha), 0.0f, 1.0f);
+		case EAlphaBlendOption::CircularInOut:	return FMath::Clamp<float>(FMath::InterpCircularInOut<float>(0.0f, 1.0f, InAlpha), 0.0f, 1.0f);
+		case EAlphaBlendOption::ExpIn:			return FMath::Clamp<float>(FMath::InterpExpoIn<float>(0.0f, 1.0f, InAlpha), 0.0f, 1.0f);
+		case EAlphaBlendOption::ExpOut:			return FMath::Clamp<float>(FMath::InterpExpoOut<float>(0.0f, 1.0f, InAlpha), 0.0f, 1.0f);
+		case EAlphaBlendOption::ExpInOut:		return FMath::Clamp<float>(FMath::InterpExpoInOut<float>(0.0f, 1.0f, InAlpha), 0.0f, 1.0f);
+		case EAlphaBlendOption::Custom:
+		{
+			if(InCustomCurve)
+			{
+				float Min;
+				float Max;
+				InCustomCurve->GetTimeRange(Min, Max);
+				return FMath::Clamp<float>(InCustomCurve->GetFloatValue(Min + (Max - Min) * InAlpha), 0.0f, 1.0f);
+			}
+		}
+	}
+
+	return InAlpha;
+}
+
+void FAlphaBlend::SetBlendOption(EAlphaBlendOption InBlendOption)
+{
+	BlendOption = InBlendOption;
+
+	// Recalculate the current blended value
+	AlphaBlend = AlphaToBlendOption();
+	BlendedValue = BeginValue + (DesiredValue - BeginValue) * AlphaBlend;
+}
+
+float FAlphaBlend::GetBlendedValue()
+{
+	return BlendedValue;
+}
+
+void FAlphaBlend::SetValueRange(float Begin, float Desired)
+{
+	BeginValue = Begin;
+	DesiredValue = Desired;
+
+	// Convert to new range
+	BlendedValue = BeginValue + (DesiredValue - BeginValue) * AlphaBlend;
+}
+
+void FAlphaBlend::SetAlpha(float InAlpha)
+{
+	AlphaLerp = FMath::Clamp(InAlpha, 0.0f, 1.0f);
+	AlphaBlend = AlphaToBlendOption();
+	BlendedValue = BeginValue + (DesiredValue - BeginValue) * AlphaBlend;
+}
+
+bool FAlphaBlend::IsComplete()
+{
+	return AlphaLerp == 1.0f;
+}
diff --git a/Engine/Source/Runtime/Engine/Private/Animation/AnimNode_StateMachine.cpp b/Engine/Source/Runtime/Engine/Private/Animation/AnimNode_StateMachine.cpp
--- a/Engine/Source/Runtime/Engine/Private/Animation/AnimNode_StateMachine.cpp
+++ b/Engine/Source/Runtime/Engine/Private/Animation/AnimNode_StateMachine.cpp
@@ -15,7 +15,7 @@ FAnimationActiveTransitionEntry::FAnimationActiveTransitionEntry()
 	: ElapsedTime(0.0f)
 	, Alpha(0.0f)
 	, CrossfadeDuration(0.0f)
-	, CrossfadeMode(ETransitionBlendMode::TBM_Cubic)
+	, BlendOption(EAlphaBlendOption::HermiteCubic)
 	, bActive(false)
 	, NextState(INDEX_NONE)
 	, PreviousState(INDEX_NONE)
@@ -26,10 +26,10 @@ FAnimationActiveTransitionEntry::FAnimationActiveTransitionEntry()
 {
 }
 
-FAnimationActiveTransitionEntry::FAnimationActiveTransitionEntry(int32 NextStateID, float ExistingWeightOfNextState, int32 PreviousStateID, const FAnimationTransitionBetweenStates& ReferenceTransitionInfo)
+FAnimationActiveTransitionEntry::FAnimationActiveTransitionEntry(int32 NextStateID, float ExistingWeightOfNextState, FAnimationActiveTransitionEntry* ExistingTransitionForNextState, int32 PreviousStateID, const FAnimationTransitionBetweenStates& ReferenceTransitionInfo)
 	: ElapsedTime(0.0f)
 	, Alpha(0.0f)
-	, CrossfadeMode(ReferenceTransitionInfo.CrossfadeMode)
+	, BlendOption(ReferenceTransitionInfo.BlendMode)
 	, bActive(true)
 	, NextState(NextStateID)
 	, PreviousState(PreviousStateID)
@@ -38,18 +38,23 @@ FAnimationActiveTransitionEntry::FAnimationActiveTransitionEntry(int32 NextState
 	, InterruptNotify(ReferenceTransitionInfo.InterruptNotify)
 	, LogicType(ReferenceTransitionInfo.LogicType)
 {
-	// Adjust the duration of the blend based on the target weight to get us there quicker if we're closer
 	const float Scaler = 1.0f - ExistingWeightOfNextState;
-	CrossfadeDuration = ReferenceTransitionInfo.CrossfadeDuration * CalculateInverseAlpha(CrossfadeMode, Scaler);
+	CrossfadeDuration = ReferenceTransitionInfo.CrossfadeDuration * CalculateInverseAlpha(BlendOption, Scaler);
+
+	Blend.SetBlendTime(CrossfadeDuration);
+	Blend.SetBlendOption(BlendOption);
+	Blend.CustomCurve = ReferenceTransitionInfo.CustomCurve;
+	Blend.SetValueRange(0.0f, 1.0f);
+	Blend.Reset();
 }
 
-float FAnimationActiveTransitionEntry::CalculateInverseAlpha(ETransitionBlendMode::Type BlendType, float InFraction) const
+float FAnimationActiveTransitionEntry::CalculateInverseAlpha(EAlphaBlendOption BlendMode, float InFraction) const
 {
-	if (BlendType == ETransitionBlendMode::TBM_Cubic)
+	if(BlendMode == EAlphaBlendOption::HermiteCubic)
 	{
-		const float A = 4.0f/3.0f;
+		const float A = 4.0f / 3.0f;
 		const float B = -2.0f;
-		const float C = 5.0f/3.0f;
+		const float C = 5.0f / 3.0f;
 
 		const float T = InFraction;
 		const float TT = InFraction*InFraction;
@@ -88,33 +93,14 @@ void FAnimationActiveTransitionEntry::Update(const FAnimationUpdateContext& Cont
 	if (bActive)
 	{
 		ElapsedTime += Context.GetDeltaTime();
-		if (ElapsedTime >= CrossfadeDuration)
+		Blend.Update(Context.GetDeltaTime());
+		Alpha = FAlphaBlend::AlphaToBlendOption(ElapsedTime / CrossfadeDuration, Blend.BlendOption, Blend.CustomCurve); //Blend.GetBlendedValue();
+
+		if(Blend.IsComplete())
 		{
 			bActive = false;
 			bOutFinished = true;
 		}
-
-		if(CrossfadeDuration <= 0.0f)
-		{
-			Alpha = 1.0f;
-		}
-		else
-		{
-			Alpha = CalculateAlpha(ElapsedTime / CrossfadeDuration);
-		}
-
-	}
-}
-
-float FAnimationActiveTransitionEntry::CalculateAlpha(float InFraction) const
-{
-	if (CrossfadeMode == ETransitionBlendMode::TBM_Cubic)
-	{
-		return FMath::SmoothStep(0.0f, 1.0f, InFraction);
-	}
-	else
-	{
-		return FMath::Clamp<float>(InFraction, 0.0f, 1.0f);
 	}
 }
 
@@ -322,16 +308,27 @@ void FAnimNode_StateMachine::Update(const FAnimationUpdateContext& Context)
 			// Get the current weight of the next state, which may be non-zero
 			const float ExistingWeightOfNextState = GetStateWeight(NextState);
 
+			FAnimationActiveTransitionEntry* PreviousTransitionForNextState = nullptr;
+			for(int32 i = ActiveTransitionArray.Num() - 1 ;  i >= 0 ; ++i)
+			{
+				FAnimationActiveTransitionEntry& TransitionEntry = ActiveTransitionArray[i];
+				if(TransitionEntry.PreviousState == NextState)
+				{
+					PreviousTransitionForNextState = &TransitionEntry;
+					break;
+				}
+			}
+
 			// Push the transition onto the stack
 			const FAnimationTransitionBetweenStates& ReferenceTransition = GetTransitionInfo(PotentialTransition.TransitionRule->TransitionIndex);
-			FAnimationActiveTransitionEntry* NewTransition = new (ActiveTransitionArray) FAnimationActiveTransitionEntry(NextState, ExistingWeightOfNextState, PreviousState, ReferenceTransition);
+			FAnimationActiveTransitionEntry* NewTransition = new (ActiveTransitionArray) FAnimationActiveTransitionEntry(NextState, ExistingWeightOfNextState, PreviousTransitionForNextState, PreviousState, ReferenceTransition);
 			NewTransition->InitializeCustomGraphLinks(Context, *(PotentialTransition.TransitionRule));
 
 #if WITH_EDITORONLY_DATA
 			NewTransition->SourceTransitionIndices = PotentialTransition.SourceTransitionIndices;
 #endif
 
-			if (!bFirstUpdate)
+			if(!bFirstUpdate)
 			{
 				Context.AnimInstance->AddAnimNotifyFromGeneratedClass(NewTransition->StartNotify);
 			}
diff --git a/Engine/Source/Runtime/Engine/Public/AlphaBlend.h b/Engine/Source/Runtime/Engine/Public/AlphaBlend.h
new file mode 100644
--- /dev/null
+++ b/Engine/Source/Runtime/Engine/Public/AlphaBlend.h
@@ -0,0 +1,146 @@
+// Copyright 1998-2015 Epic Games, Inc. All Rights Reserved.
+
+#pragma once
+
+#include "AlphaBlend.generated.h"
+
+class UCurveFloat;
+
+UENUM()
+enum class EAlphaBlendOption : uint8
+{
+	// Linear interpolation
+	Linear = 0,
+	// Cubic-in interpolation
+	Cubic,
+	// Hermite-Cubic
+	HermiteCubic,
+	// Sinusoidal interpolation
+	Sinusoidal,
+	// Quadratic in-out interpolation
+	QuadraticInOut,
+	// Cubic in-out interpolation
+	CubicInOut,
+	// Quartic in-out interpolation
+	QuarticInOut,
+	// Quintic in-out interpolation
+	QuinticInOut,
+	// Circular-in interpolation
+	CircularIn,
+	// Circular-out interpolation
+	CircularOut,
+	// Circular in-out interpolation
+	CircularInOut,
+	// Exponential-in interpolation
+	ExpIn,
+	// Exponential-Out interpolation
+	ExpOut,
+	// Exponential in-out interpolation
+	ExpInOut,
+	// Custom interpolation, will use custom curve inside an FAlphaBlend or linear if none has been set
+	Custom,
+};
+
+/**
+ * Alpha Blend Type
+ */
+USTRUCT()
+struct FAlphaBlend
+{
+	GENERATED_BODY()
+
+	/** Type of blending used (Linear, Cubic, etc.) */
+	UPROPERTY(EditAnywhere, Category = "Blend")
+	EAlphaBlendOption BlendOption;
+
+	UPROPERTY(EditAnywhere, Category = "Blend")
+	float BeginValue;
+
+	UPROPERTY(EditAnywhere, Category = "Blend")
+	float DesiredValue;
+
+	UPROPERTY(EditAnywhere, Category = "Blend")
+	float	BlendTime;
+
+	UPROPERTY(EditAnywhere, Category = "Blend")
+	UCurveFloat* CustomCurve;
+
+	// Constructor
+	FAlphaBlend();
+
+	/** Update transition blend time */
+	void SetBlendTime(float InBlendTime);
+
+	/** Sets the method used to blend the value */
+	void SetBlendOption(EAlphaBlendOption InBlendOption);
+
+	/** Sets the range of values to map to the interpolation */
+	void SetValueRange(float Begin, float Desired);
+
+	/** Sets the begin value for the blended value */
+	void SetBeginValue(float InBegin);
+
+	/** Sets the final desired value for the blended value */
+	void SetDesiredValue(float InDesired);
+
+	/** Sets the Lerp alpha value directly */
+	void SetAlpha(float InAlpha);
+
+	/** Reset to zero / restart the blend */
+	void Reset();
+
+	/** Returns true if Target is > 0.f, or false otherwise */
+	bool GetToggleStatus();
+
+	/** Enable (1.f) or Disable (0.f) */
+	void Toggle(bool bEnable);
+
+	/** SetTarget, but check that we're actually setting a different target */
+	void ConditionalSetTarget(float InAlphaTarget);
+
+	/** Set Target for interpolation */
+	void SetTarget(float InAlphaTarget);
+
+	/** Update interpolation, has to be called once every frame */
+	void Update(float InDeltaTime);
+
+	/** Converts internal lerped alpha into the output alpha type */
+	float AlphaToBlendOption();
+
+	/** Converts InAlpha from a linear 0...1 value into the output alpha described by InBlendOption 
+	 *  @param InAlpha In linear 0...1 alpha
+	 *  @param InBlendOption The type of blend to use
+	 *  @param InCustomCurve The curve to use when blend option is set to custom
+	 */
+	static float AlphaToBlendOption(float InAlpha, EAlphaBlendOption InBlendOption, UCurveFloat* InCustomCurve = nullptr);
+
+	/** Gets the current 0..1 alpha value */
+	float GetAlpha();
+
+	/** Gets the current blended value */
+	float GetBlendedValue();
+
+	/** Gets whether or not the blend is complete */
+	bool IsComplete();
+
+protected:
+	/** Internal Lerped value for Alpha */
+	UPROPERTY()
+	float	AlphaLerp;
+
+	/** Resulting Alpha value, between 0.f and 1.f */
+	UPROPERTY()
+	float	AlphaBlend;
+
+	/** Target to reach */
+	UPROPERTY()
+	float	AlphaTarget;
+
+	/** Time left to reach target */
+	UPROPERTY()
+	float	BlendTimeRemaining;
+
+	/** The current blended value derived from the begin and desired values */
+	UPROPERTY()
+	float BlendedValue;
+};
\ No newline at end of file
