Missed file for graph blending

diff --git a/Engine/Source/Runtime/Engine/Private/Animation/AnimNode_BlendListBase.cpp b/Engine/Source/Runtime/Engine/Private/Animation/AnimNode_BlendListBase.cpp
--- a/Engine/Source/Runtime/Engine/Private/Animation/AnimNode_BlendListBase.cpp
+++ b/Engine/Source/Runtime/Engine/Private/Animation/AnimNode_BlendListBase.cpp
@@ -27,8 +27,20 @@ void FAnimNode_BlendListBase::Initialize(const FAnimationInitializeContext& Cont
 
 	RemainingBlendTimes.Empty(NumPoses);
 	RemainingBlendTimes.AddZeroed(NumPoses);
+	Blends.Empty(NumPoses);
+	Blends.AddZeroed(NumPoses);
 
 	LastActiveChildIndex = INDEX_NONE;
+
+	for(int32 i = 0 ; i < Blends.Num() ; ++i)
+	{
+		FAlphaBlend& Blend = Blends[i];
+
+		Blend.SetBlendTime(0.0f);
+		Blend.SetBlendOption(BlendType);
+		Blend.CustomCurve = CustomBlendCurve;
+	}
+	Blends[0].SetAlpha(1.0f);
 }
 
 void FAnimNode_BlendListBase::CacheBones(const FAnimationCacheBonesContext& Context) 
@@ -71,19 +83,36 @@ void FAnimNode_BlendListBase::Update(const FAnimationUpdateContext& Context)
 			{
 				RemainingBlendTimes[i] = RemainingBlendTime;
 			}
+
+			for(int32 i = 0; i < Blends.Num(); ++i)
+			{
+				FAlphaBlend& Blend = Blends[i];
+
+				Blend.SetBlendTime(RemainingBlendTime);
+
+				if(i == ChildIndex)
+				{
+					Blend.SetValueRange(BlendWeights[i], 1.0f);
+				}
+				else
+				{
+					Blend.SetValueRange(BlendWeights[i], 0.0f);
+				}
+
+				Blend.Reset();
+			}
 		}
 
 		// Advance the weights
 		//@TODO: This means we advance even in a frame where the target weights/times just got modified; is that desirable?
 		float SumWeight = 0.0f;
-		for (int32 i = 0; i < RemainingBlendTimes.Num(); ++i)
+		for (int32 i = 0; i < Blends.Num(); ++i)
 		{
-			float& RemainingBlendTime = RemainingBlendTimes[i];
 			float& BlendWeight = BlendWeights[i];
 
-			const float DesiredWeight = (i == ChildIndex) ? 1.0f : 0.0f;
-
-			FAnimationRuntime::TickBlendWeight(Context.GetDeltaTime(), DesiredWeight, BlendWeight, RemainingBlendTime);
+			FAlphaBlend& Blend = Blends[i];
+			Blend.Update(Context.GetDeltaTime());
+			BlendWeight = Blend.GetBlendedValue();
 
 			SumWeight += BlendWeight;
 		}
