Refactor physx event callback to include phys scene info

diff --git a/Engine/Source/Runtime/Engine/Private/PhysicsEngine/PhysScene.cpp b/Engine/Source/Runtime/Engine/Private/PhysicsEngine/PhysScene.cpp
--- a/Engine/Source/Runtime/Engine/Private/PhysicsEngine/PhysScene.cpp
+++ b/Engine/Source/Runtime/Engine/Private/PhysicsEngine/PhysScene.cpp
@@ -107,8 +107,7 @@ FPhysScene::FPhysScene()
 	{
 		CPUDispatcher = new FPhysXCPUDispatcher();
 	}
-	// Create sim event callback
-	SimEventCallback = new FPhysXSimEventCallback();
+
 #endif	//#if WITH_PHYSX
 
 	// initialize console variable - this console variable change requires it to restart scene. 
@@ -183,7 +182,6 @@ FPhysScene::~FPhysScene()
 
 #if WITH_PHYSX
 	GPhysCommandHandler->DeferredDeleteCPUDispathcer(CPUDispatcher);
-	GPhysCommandHandler->DeferredDeleteSimEventCallback(SimEventCallback);
 #endif	//#if WITH_PHYSX
 }
 
@@ -1248,6 +1246,9 @@ void FPhysScene::InitPhysScene(uint32 SceneType)
 #if WITH_PHYSX
 	PhysxUserData = FPhysxUserData(this);
 
+	// Create sim event callback
+	SimEventCallback[SceneType] = new FPhysXSimEventCallback(this, SceneType);
+
 	// Include scene descriptor in loop, so that we might vary it with scene type
 	PxSceneDesc PSceneDesc(GPhysXSDK->getTolerancesScale());
 	PSceneDesc.cpuDispatcher = CPUDispatcher;
@@ -1258,7 +1259,7 @@ void FPhysScene::InitPhysScene(uint32 SceneType)
 	PSceneDesc.filterShaderDataSize = sizeof(PhysSceneShaderInfo);
 
 	PSceneDesc.filterShader = PhysXSimFilterShader;
-	PSceneDesc.simulationEventCallback = SimEventCallback;
+	PSceneDesc.simulationEventCallback = SimEventCallback[SceneType];
 
 	if(UPhysicsSettings::Get()->bEnablePCM)
 	{
@@ -1434,6 +1435,7 @@ void FPhysScene::TermPhysScene(uint32 SceneType)
 
 		// @todo block on any running scene before calling this
 		GPhysCommandHandler->DeferredRelease(PScene);
+		GPhysCommandHandler->DeferredDeleteSimEventCallback(SimEventCallback[SceneType]);
 
 		// Commands may have accumulated as the scene is terminated - flush any commands for this scene.
 		DeferredCommandHandler.Flush();
diff --git a/Engine/Source/Runtime/Engine/Private/PhysicsEngine/PhysXSupport.cpp b/Engine/Source/Runtime/Engine/Private/PhysicsEngine/PhysXSupport.cpp
--- a/Engine/Source/Runtime/Engine/Private/PhysicsEngine/PhysXSupport.cpp
+++ b/Engine/Source/Runtime/Engine/Private/PhysicsEngine/PhysXSupport.cpp
@@ -433,14 +433,7 @@ void FPhysXSimEventCallback::onContact(const PxContactPairHeader& PairHeader, co
 		}
 	}
 
-	// Get the Scene. 
-	const PxScene* PScene = PActor0->getScene();
-	check(PScene == PActor1->getScene());
-	FPhysScene* PhysScene = FPhysxUserData::Get<FPhysScene>(PScene->userData);
-	check(PhysScene);
-
-	const EPhysicsSceneType SceneType = PhysScene->GetPhysXScene(PST_Sync) == PScene ? PST_Sync : PST_Async;
-	TArray<FCollisionNotifyInfo>& PendingCollisionNotifies = PhysScene->GetPendingCollisionNotifies(SceneType);
+	TArray<FCollisionNotifyInfo>& PendingCollisionNotifies = OwningScene->GetPendingCollisionNotifies(SceneType);
 
 	uint32 PreAddingCollisionNotify = PendingCollisionNotifies.Num() - 1;
 	TArray<int32> PairNotifyMapping = FBodyInstance::AddCollisionNotifyInfo(BodyInst0, BodyInst1, Pairs, NumPairs, PendingCollisionNotifies);
@@ -525,6 +518,16 @@ void FPhysXSimEventCallback::onConstraintBreak( PxConstraintInfo* constraints, P
 	}
 }
 
+void FPhysXSimEventCallback::onWake(PxActor** actors, PxU32 count)
+{
+
+}
+
+void FPhysXSimEventCallback::onSleep(PxActor** actors, PxU32 count)
+{
+
+}
+
 //////////////////////////////////////////////////////////////////////////
 // FPhysXCPUDispatcher
 
diff --git a/Engine/Source/Runtime/Engine/Private/PhysicsEngine/PhysXSupport.h b/Engine/Source/Runtime/Engine/Private/PhysicsEngine/PhysXSupport.h
--- a/Engine/Source/Runtime/Engine/Private/PhysicsEngine/PhysXSupport.h
+++ b/Engine/Source/Runtime/Engine/Private/PhysicsEngine/PhysXSupport.h
@@ -854,15 +854,23 @@ PxFilterFlags PhysXSimFilterShader(	PxFilterObjectAttributes attributes0, PxFilt
 									PxFilterObjectAttributes attributes1, PxFilterData filterData1,
 									PxPairFlags& pairFlags, const void* constantBlock, PxU32 constantBlockSize );
 
+class FPhysScene;
 
 /** Event callback used to notify engine about various collision events */
 class FPhysXSimEventCallback : public PxSimulationEventCallback
 {
+public:
+	FPhysXSimEventCallback(FPhysScene* InOwningScene, int32 InSceneType) : OwningScene(InOwningScene), SceneType(InSceneType){}
+
 	virtual void onConstraintBreak(PxConstraintInfo* constraints, PxU32 count) override;
-	virtual void onWake(PxActor** actors, PxU32 count) override {}
-	virtual void onSleep(PxActor** actors, PxU32 count) override {}
+	virtual void onWake(PxActor** actors, PxU32 count) override;
+	virtual void onSleep(PxActor** actors, PxU32 count) override;
 	virtual void onTrigger(PxTriggerPair* pairs, PxU32 count) override {}
 	virtual void onContact(const PxContactPairHeader& PairHeader, const PxContactPair* Pairs, PxU32 NumPairs) override;
+
+private:	
+	FPhysScene* OwningScene;
+	int32 SceneType;
 };
 
 /** Used to dispatch physx tasks to task graph */
diff --git a/Engine/Source/Runtime/Engine/Public/PhysicsPublic.h b/Engine/Source/Runtime/Engine/Public/PhysicsPublic.h
--- a/Engine/Source/Runtime/Engine/Public/PhysicsPublic.h
+++ b/Engine/Source/Runtime/Engine/Public/PhysicsPublic.h
@@ -369,7 +369,7 @@ private:
 	/** Dispatcher for CPU tasks */
 	class PxCpuDispatcher*			CPUDispatcher;
 	/** Simulation event callback object */
-	class FPhysXSimEventCallback*			SimEventCallback;
+	class FPhysXSimEventCallback*			SimEventCallback[PST_MAX];
 #if WITH_VEHICLE
 	/** Vehicle scene */
 	class FPhysXVehicleManager*			VehicleManager;
