Added comment description to GameSession.h for API reference documentation.

diff --git a/Engine/Source/Runtime/Engine/Classes/GameFramework/GameSession.h b/Engine/Source/Runtime/Engine/Classes/GameFramework/GameSession.h
--- a/Engine/Source/Runtime/Engine/Classes/GameFramework/GameSession.h
+++ b/Engine/Source/Runtime/Engine/Classes/GameFramework/GameSession.h
@@ -11,6 +11,10 @@
 class UWorld;
 class APlayerController;
 
+/**
+Acts as a game-specific wrapper around the session interface. The game code makes calls to this when it needs to interact with the session interface.
+A game session exists only the server, while running an online game.
+*/
 UCLASS(config=Game, notplaceable)
 class ENGINE_API AGameSession : public AInfo
 {
