Fix bug where welded BodyInstances were not having their WeldParent set correctly (was setting it on the root by mistake).
Fixes bug where welded overlap components do not generate overlap events.

#UE-18046

diff --git a/Engine/Source/Runtime/Engine/Private/PhysicsEngine/BodyInstance.cpp b/Engine/Source/Runtime/Engine/Private/PhysicsEngine/BodyInstance.cpp
--- a/Engine/Source/Runtime/Engine/Private/PhysicsEngine/BodyInstance.cpp
+++ b/Engine/Source/Runtime/Engine/Private/PhysicsEngine/BodyInstance.cpp
@@ -1842,41 +1842,42 @@ bool FBodyInstance::Weld(FBodyInstance* TheirBody, const FTransform& TheirTM)
 
 	ExecuteOnPhysicsReadWrite([&]
 	{
-	SCOPE_CYCLE_COUNTER(STAT_UpdatePhysMats);
+		SCOPE_CYCLE_COUNTER(STAT_UpdatePhysMats);
 
-		WeldParent = this;
+		TheirBody->WeldParent = this;
 
-	UPhysicalMaterial* SimplePhysMat = GetSimplePhysicalMaterial();
-	TArray<UPhysicalMaterial*> ComplexPhysMats = GetComplexPhysicalMaterials();
-	PxMaterial* PSimpleMat = SimplePhysMat->GetPhysXMaterial();
+		UPhysicalMaterial* SimplePhysMat = GetSimplePhysicalMaterial();
+		TArray<UPhysicalMaterial*> ComplexPhysMats = GetComplexPhysicalMaterials();
+		PxMaterial* PSimpleMat = SimplePhysMat->GetPhysXMaterial();
 
-	FShapeData ShapeData;
-	GetFilterData_AssumesLocked(ShapeData);
-	GetShapeFlags_AssumesLocked(ShapeData, ShapeData.CollisionEnabled, BodySetup->GetCollisionTraceFlag() == CTF_UseComplexAsSimple);
+		FShapeData ShapeData;
+		GetFilterData_AssumesLocked(ShapeData);
+		GetShapeFlags_AssumesLocked(ShapeData, ShapeData.CollisionEnabled, BodySetup->GetCollisionTraceFlag() == CTF_UseComplexAsSimple);
 
-	//child body gets placed into the same scenes as parent body
-	if (PxRigidActor* MyBody = RigidActorSync)
-	{
-			TheirBody->BodySetup->AddShapesToRigidActor_AssumesLocked(this, MyBody, PST_Sync, Scale3D, PSimpleMat, ComplexPhysMats, ShapeData, RelativeTM, &PNewShapes);
-	}
+		//child body gets placed into the same scenes as parent body
+		if (PxRigidActor* MyBody = RigidActorSync)
+		{
+				TheirBody->BodySetup->AddShapesToRigidActor_AssumesLocked(this, MyBody, PST_Sync, Scale3D, PSimpleMat, ComplexPhysMats, ShapeData, RelativeTM, &PNewShapes);
+		}
 
-	if (PxRigidActor* MyBody = RigidActorAsync)
-	{
-			TheirBody->BodySetup->AddShapesToRigidActor_AssumesLocked(this, MyBody, PST_Sync, Scale3D, PSimpleMat, ComplexPhysMats, ShapeData, RelativeTM, &PNewShapes);
-	}
+		if (PxRigidActor* MyBody = RigidActorAsync)
+		{
+				TheirBody->BodySetup->AddShapesToRigidActor_AssumesLocked(this, MyBody, PST_Sync, Scale3D, PSimpleMat, ComplexPhysMats, ShapeData, RelativeTM, &PNewShapes);
+		}
+
+		for (int32 ShapeIdx = 0; ShapeIdx < PNewShapes.Num(); ++ShapeIdx)
+		{
+			PxShape* PShape = PNewShapes[ShapeIdx];
+			ShapeToBodiesMap.Add(PShape, FWeldInfo(TheirBody, RelativeTM));
+		}
 
 
-	for (int32 ShapeIdx = 0; ShapeIdx < PNewShapes.Num(); ++ShapeIdx)
-	{
-		PxShape* PShape = PNewShapes[ShapeIdx];
-		ShapeToBodiesMap.Add(PShape, FWeldInfo(TheirBody, RelativeTM));
-	}
 
-	PostShapeChange();
+		PostShapeChange();
 
-	//remove their body from scenes
-	TermBodyHelper(TheirBody->SceneIndexSync, TheirBody->RigidActorSync, TheirBody);
-	TermBodyHelper(TheirBody->SceneIndexAsync, TheirBody->RigidActorAsync, TheirBody);
+		//remove their body from scenes
+		TermBodyHelper(TheirBody->SceneIndexSync, TheirBody->RigidActorSync, TheirBody);
+		TermBodyHelper(TheirBody->SceneIndexAsync, TheirBody->RigidActorAsync, TheirBody);
 	});
 	
 
@@ -1928,7 +1929,7 @@ void FBodyInstance::UnWeld(FBodyInstance* TheirBI)
 		PostShapeChange();
 	}
 
-		WeldParent = nullptr;
+		TheirBI->WeldParent = nullptr;
 	});
 #endif
 }
@@ -4067,7 +4068,7 @@ FTransform RootSpaceToWeldedSpace(const FBodyInstance* BI, const FTransform& Roo
 
 bool FBodyInstance::OverlapMulti(TArray<struct FOverlapResult>& InOutOverlaps, const class UWorld* World, const FTransform* pWorldToComponent, const FVector& Pos, const FQuat& Quat, ECollisionChannel TestChannel, const struct FComponentQueryParams& Params, const struct FCollisionResponseParams& ResponseParams, const struct FCollisionObjectQueryParams& ObjectQueryParams) const
 {
-	if ( (IsValidBodyInstance() || (WeldParent && WeldParent->IsValidBodyInstance())) == false )
+	if ( !IsValidBodyInstance()  && (!WeldParent || !WeldParent->IsValidBodyInstance()))
 	{
 		UE_LOG(LogCollision, Log, TEXT("FBodyInstance::OverlapMulti : (%s) No physics data"), *GetBodyDebugName());
 		return false;
