Changed 3 KismetMathLibrary functions to BlueprintPure (PointsAreCoplanar, LinePlaneIntersection, LinePlaneIntersection_OriginNormal)

#merge

diff --git a/Engine/Source/Runtime/Engine/Classes/Kismet/KismetMathLibrary.h b/Engine/Source/Runtime/Engine/Classes/Kismet/KismetMathLibrary.h
--- a/Engine/Source/Runtime/Engine/Classes/Kismet/KismetMathLibrary.h
+++ b/Engine/Source/Runtime/Engine/Classes/Kismet/KismetMathLibrary.h
@@ -1653,7 +1653,7 @@ class ENGINE_API UKismetMathLibrary : public UBlueprintFunctionLibrary
 	 *
 	 * @return Whether the points are relatively coplanar, based on the tolerance
 	 */
-	UFUNCTION(BlueprintCallable, Category = "Math|Geometry")
+	UFUNCTION(BlueprintPure, Category = "Math|Geometry")
 	static bool PointsAreCoplanar(const TArray<FVector>& Points, float Tolerance = 0.1f);
 
 	//
@@ -1666,7 +1666,7 @@ class ENGINE_API UKismetMathLibrary : public UBlueprintFunctionLibrary
 	 * @param		Intersection - The point of intersection between the line and the plane
 	 * @return		True if the intersection test was successful.
 	 */
-	UFUNCTION(BlueprintCallable, Category = "Math|Intersection")
+	UFUNCTION(BlueprintPure, Category = "Math|Intersection")
 	static bool LinePlaneIntersection(const FVector& LineStart, const FVector& LineEnd, const FPlane& APlane, float& T, FVector& Intersection);
 
 	/**
@@ -1675,6 +1675,6 @@ class ENGINE_API UKismetMathLibrary : public UBlueprintFunctionLibrary
 	 * @param		Intersection - The point of intersection between the line and the plane
 	 * @return		True if the intersection test was successful.
 	 */
-	UFUNCTION(BlueprintCallable, Category = "Math|Intersection", meta = (DisplayName = "Line Plane Intersection (Origin & Normal)"))
+	UFUNCTION(BlueprintPure, Category = "Math|Intersection", meta = (DisplayName = "Line Plane Intersection (Origin & Normal)"))
 	static bool LinePlaneIntersection_OriginNormal(const FVector& LineStart, const FVector& LineEnd, FVector PlaneOrigin, FVector PlaneNormal, float& T, FVector& Intersection);
 };
