- Added GameplayStatics::DeprojectScreenToWorld and ProjectWorldToScreen
- added simpler version of FSceneView::DeprojectScreenToWorld that takes a single view proj matrix
- added FSceneView::ProjectWorldToScreen
- added PlayerController::GetLocalPlayer()

#merge

diff --git a/Engine/Source/Runtime/Engine/Classes/GameFramework/PlayerController.h b/Engine/Source/Runtime/Engine/Classes/GameFramework/PlayerController.h
--- a/Engine/Source/Runtime/Engine/Classes/GameFramework/PlayerController.h
+++ b/Engine/Source/Runtime/Engine/Classes/GameFramework/PlayerController.h
@@ -1173,6 +1173,9 @@ public:
 	/** Associate a new UPlayer with this PlayerController. */
 	virtual void SetPlayer(UPlayer* Player);
 
+	/** Returns the ULocalPlayer for this controller if it exists, or null otherwise */
+	class ULocalPlayer* GetLocalPlayer() const;
+
 	/**
 	 * Called client-side to smoothly interpolate received TargetViewRotation (result is in BlendedTargetViewRotation)
 	 * @param TargetPawn   is the pawn which is the current ViewTarget
diff --git a/Engine/Source/Runtime/Engine/Classes/Kismet/GameplayStatics.h b/Engine/Source/Runtime/Engine/Classes/Kismet/GameplayStatics.h
--- a/Engine/Source/Runtime/Engine/Classes/Kismet/GameplayStatics.h
+++ b/Engine/Source/Runtime/Engine/Classes/Kismet/GameplayStatics.h
@@ -560,5 +560,25 @@ class ENGINE_API UGameplayStatics : public UBlueprintFunctionLibrary
 	*/
 	UFUNCTION(BlueprintCallable, Category = "Foliage", meta = (WorldContext = "WorldContextObject", UnsafeDuringActorConstruction = "true"))
 	static int32 GrassOverlappingSphereCount(UObject* WorldContextObject, const UStaticMesh* StaticMesh, FVector CenterPosition, float Radius);
+
+	/** 
+	 * Transforms the given 2D screen space coordinate into a 3D world-space point and direction
+	 * @param Player			Deproject using this player's view.
+	 * @param ScreenPosition	2D screen space to deproject
+	 * @param WorldPosition		(out) Corresponding 3D position in world space
+	 * @param WorldDirection	(out) World space direction vector away from the camera at the given 2d poiunt
+	 */
+	UFUNCTION(BlueprintPure, Category = "Utilities", meta = (Keywords = "unproject"))
+	static bool DeprojectScreenToWorld(APlayerController const* Player, const FVector2D& ScreenPosition, FVector& WorldPosition, FVector& WorldDirection);
+
+	/** 
+	 * Transforms the given 3D world-space point into a its 2D screen space coordinate. 
+	 * @param Player			Project using this player's view.
+	 * @param WorldPosition		World position to project.
+	 * @param ScreenPosition	(out) Corresponding 2D position in screen space
+	 */
+	UFUNCTION(BlueprintPure, Category = "Utilities")
+	static bool ProjectWorldToScreen(APlayerController const* Player, const FVector& WorldPosition, FVector2D& ScreenPosition);
+
 };
 
diff --git a/Engine/Source/Runtime/Engine/Private/GameplayStatics.cpp b/Engine/Source/Runtime/Engine/Private/GameplayStatics.cpp
--- a/Engine/Source/Runtime/Engine/Private/GameplayStatics.cpp
+++ b/Engine/Source/Runtime/Engine/Private/GameplayStatics.cpp
@@ -1599,4 +1599,45 @@ int32 UGameplayStatics::GrassOverlappingSphereCount(UObject* WorldContextObject,
 	return Count;
 }
 
+
+bool UGameplayStatics::DeprojectScreenToWorld(APlayerController const* Player, const FVector2D& ScreenPosition, FVector& WorldPosition, FVector& WorldDirection)
+{
+	ULocalPlayer* const LP = Player ? Player->GetLocalPlayer() : nullptr;
+	if (LP && LP->ViewportClient)
+	{
+		// get the projection data
+		FSceneViewProjectionData ProjectionData;
+		if (LP->GetProjectionData(LP->ViewportClient->Viewport, eSSP_FULL, /*out*/ ProjectionData))
+		{
+			FMatrix const InvViewProjMatrix = ProjectionData.ComputeViewProjectionMatrix().InverseFast();
+			FSceneView::DeprojectScreenToWorld(ScreenPosition, ProjectionData.GetConstrainedViewRect(), InvViewProjMatrix, /*out*/ WorldPosition, /*out*/ WorldDirection);
+			return true;
+		}
+	}
+
+	// something went wrong, zero things and return false
+	WorldPosition = FVector::ZeroVector;
+	WorldDirection = FVector::ZeroVector;
+	return false;
+}
+
+bool UGameplayStatics::ProjectWorldToScreen(APlayerController const* Player, const FVector& WorldPosition, FVector2D& ScreenPosition)
+{
+	ULocalPlayer* const LP = Player ? Player->GetLocalPlayer() : nullptr;
+	if (LP && LP->ViewportClient)
+	{
+		// get the projection data
+		FSceneViewProjectionData ProjectionData;
+		if (LP->GetProjectionData(LP->ViewportClient->Viewport, eSSP_FULL, /*out*/ ProjectionData))
+		{
+			FMatrix const ViewProjectionMatrix = ProjectionData.ComputeViewProjectionMatrix();
+			FSceneView::ProjectWorldToScreen(WorldPosition, ProjectionData.GetConstrainedViewRect(), ViewProjectionMatrix, ScreenPosition);
+			return true;
+		}
+	}
+
+	ScreenPosition = FVector2D::ZeroVector;
+	return false;
+}
+
 #undef LOCTEXT_NAMESPACE
\ No newline at end of file
diff --git a/Engine/Source/Runtime/Engine/Private/PlayerController.cpp b/Engine/Source/Runtime/Engine/Private/PlayerController.cpp
--- a/Engine/Source/Runtime/Engine/Private/PlayerController.cpp
+++ b/Engine/Source/Runtime/Engine/Private/PlayerController.cpp
@@ -3733,6 +3733,11 @@ void APlayerController::SetPlayer( UPlayer* InPlayer )
 	ReceivedPlayer();
 }
 
+ULocalPlayer* APlayerController::GetLocalPlayer() const
+{
+	return Cast<ULocalPlayer>(Player);
+}
+
 void APlayerController::TickPlayerInput(const float DeltaSeconds, const bool bGamePaused)
 {
 	SCOPE_CYCLE_COUNTER(STAT_PC_TickInput);
diff --git a/Engine/Source/Runtime/Engine/Private/SceneView.cpp b/Engine/Source/Runtime/Engine/Private/SceneView.cpp
--- a/Engine/Source/Runtime/Engine/Private/SceneView.cpp
+++ b/Engine/Source/Runtime/Engine/Private/SceneView.cpp
@@ -678,10 +678,8 @@ FVector FSceneView::Deproject(const FPlane& ScreenPoint) const
 
 void FSceneView::DeprojectFVector2D(const FVector2D& ScreenPos, FVector& out_WorldOrigin, FVector& out_WorldDirection) const
 {
-	const FMatrix InverseViewMatrix = ViewMatrices.ViewMatrix.InverseFast();
-	const FMatrix InvProjectionMatrix = ViewMatrices.GetInvProjMatrix();
-	
-	DeprojectScreenToWorld(ScreenPos, UnscaledViewRect, InverseViewMatrix, InvProjectionMatrix, out_WorldOrigin, out_WorldDirection);
+	const FMatrix InvViewProjMatrix = ViewMatrices.GetInvViewProjMatrix();
+	DeprojectScreenToWorld(ScreenPos, UnscaledViewRect, InvViewProjectionMatrix, out_WorldOrigin, out_WorldDirection);
 }
 
 void FSceneView::DeprojectScreenToWorld(const FVector2D& ScreenPos, const FIntRect& ViewRect, const FMatrix& InvViewMatrix, const FMatrix& InvProjectionMatrix, FVector& out_WorldOrigin, FVector& out_WorldDirection)
@@ -734,10 +732,77 @@ void FSceneView::DeprojectScreenToWorld(const FVector2D& ScreenPos, const FIntRe
 	out_WorldDirection = RayDirWorldSpace.GetSafeNormal();
 }
 
+void FSceneView::DeprojectScreenToWorld(const FVector2D& ScreenPos, const FIntRect& ViewRect, const FMatrix& InvViewProjMatrix, FVector& out_WorldOrigin, FVector& out_WorldDirection)
+{
+	float PixelX = FMath::TruncToFloat(ScreenPos.X);
+	float PixelY = FMath::TruncToFloat(ScreenPos.Y);
+
+	// Get the eye position and direction of the mouse cursor in two stages (inverse transform projection, then inverse transform view).
+	// This avoids the numerical instability that occurs when a view matrix with large translation is composed with a projection matrix
+
+	// Get the pixel coordinates into 0..1 normalized coordinates within the constrained view rectangle
+	const float NormalizedX = (PixelX - ViewRect.Min.X) / ((float)ViewRect.Width());
+	const float NormalizedY = (PixelY - ViewRect.Min.Y) / ((float)ViewRect.Height());
+
+	// Get the pixel coordinates into -1..1 projection space
+	const float ScreenSpaceX = (NormalizedX - 0.5f) * 2.0f;
+	const float ScreenSpaceY = ((1.0f - NormalizedY) - 0.5f) * 2.0f;
+
+	// The start of the raytrace is defined to be at mousex,mousey,1 in projection space (z=1 is near, z=0 is far - this gives us better precision)
+	// To get the direction of the raytrace we need to use any z between the near and the far plane, so let's use (mousex, mousey, 0.5)
+	const FVector4 RayStartProjectionSpace = FVector4(ScreenSpaceX, ScreenSpaceY, 1.0f, 1.0f);
+	const FVector4 RayEndProjectionSpace = FVector4(ScreenSpaceX, ScreenSpaceY, 0.5f, 1.0f);
+
+	// Projection (changing the W coordinate) is not handled by the FMatrix transforms that work with vectors, so multiplications
+	// by the projection matrix should use homogeneous coordinates (i.e. FPlane).
+	const FVector4 HGRayStartWorldSpace = InvViewProjMatrix.TransformFVector4(RayStartProjectionSpace);
+	const FVector4 HGRayEndWorldSpace = InvViewProjMatrix.TransformFVector4(RayEndProjectionSpace);
+	FVector RayStartWorldSpace(HGRayStartWorldSpace.X, HGRayStartWorldSpace.Y, HGRayStartWorldSpace.Z);
+	FVector RayEndWorldSpace(HGRayEndWorldSpace.X, HGRayEndWorldSpace.Y, HGRayEndWorldSpace.Z);
+	// divide vectors by W to undo any projection and get the 3-space coordinate 
+	if (HGRayStartWorldSpace.W != 0.0f)
+	{
+		RayStartWorldSpace /= HGRayStartWorldSpace.W;
+	}
+	if (HGRayEndWorldSpace.W != 0.0f)
+	{
+		RayEndWorldSpace /= HGRayEndWorldSpace.W;
+	}
+	const FVector RayDirWorldSpace = (RayEndWorldSpace - RayStartWorldSpace).GetSafeNormal();
+
+	// Finally, store the results in the outputs
+	out_WorldOrigin = RayStartWorldSpace;
+	out_WorldDirection = RayDirWorldSpace;
+}
+
+void FSceneView::ProjectWorldToScreen(const FVector& WorldPosition, const FIntRect& ViewRect, const FMatrix& ViewProjectionMatrix, FVector2D& out_ScreenPos)
+{
+	FPlane Result = ViewProjectionMatrix.TransformFVector4(FVector4(WorldPosition, 1.f));
+	if (Result.W == 0)
+	{
+		Result.W = KINDA_SMALL_NUMBER;
+	}
+
+	// the result of this will be x and y coords in -1..1 projection space
+	const float RHW = 1.0f / Result.W;
+	FPlane PosInScreenSpace = FPlane(Result.X * RHW, Result.Y * RHW, Result.Z * RHW, Result.W);
+
+	// Move from projection space to normalized 0..1 UI space
+	const float NormalizedX = (PosInScreenSpace.X / 2.f) + 0.5f;
+	const float NormalizedY = 1.f - (PosInScreenSpace.Y / 2.f) - 0.5f;
+
+	FVector2D RayStartViewRectSpace(
+		(float)ViewRect.Min.X + (NormalizedX * (float)ViewRect.Width()),
+		(float)ViewRect.Min.Y + (NormalizedY * (float)ViewRect.Height())
+		);
+
+	out_ScreenPos = RayStartViewRectSpace;
+}
+
+
 #define LERP_PP(NAME) if(Src.bOverride_ ## NAME)	Dest . NAME = FMath::Lerp(Dest . NAME, Src . NAME, Weight);
 #define IF_PP(NAME) if(Src.bOverride_ ## NAME && Src . NAME)
 
-
 // @param Weight 0..1
 void FSceneView::OverridePostProcessSettings(const FPostProcessSettings& Src, float Weight)
 {
diff --git a/Engine/Source/Runtime/Engine/Public/SceneView.h b/Engine/Source/Runtime/Engine/Public/SceneView.h
--- a/Engine/Source/Runtime/Engine/Public/SceneView.h
+++ b/Engine/Source/Runtime/Engine/Public/SceneView.h
@@ -633,14 +633,16 @@ public:
 	 */
 	FVector Deproject(const FPlane& ScreenPoint) const;
 
-	/** transforms 2D screen coordinates into a 3D world-space origin and direction 
+	/** 
+	 * Transforms 2D screen coordinates into a 3D world-space origin and direction 
 	 * @param ScreenPos - screen coordinates in pixels
 	 * @param out_WorldOrigin (out) - world-space origin vector
 	 * @param out_WorldDirection (out) - world-space direction vector
 	 */
 	void DeprojectFVector2D(const FVector2D& ScreenPos, FVector& out_WorldOrigin, FVector& out_WorldDirection) const;
 
-	/** transforms 2D screen coordinates into a 3D world-space origin and direction 
+	/** 
+	 * Transforms 2D screen coordinates into a 3D world-space origin and direction 
 	 * @param ScreenPos - screen coordinates in pixels
 	 * @param ViewRect - view rectangle
 	 * @param InvViewMatrix - inverse view matrix
@@ -650,6 +652,18 @@ public:
 	 */
 	static void DeprojectScreenToWorld(const FVector2D& ScreenPos, const FIntRect& ViewRect, const FMatrix& InvViewMatrix, const FMatrix& InvProjMatrix, FVector& out_WorldOrigin, FVector& out_WorldDirection);
 
+	/** Overload to take a single combined view projection matrix. */
+	static void DeprojectScreenToWorld(const FVector2D& ScreenPos, const FIntRect& ViewRect, const FMatrix& InvViewProjMatrix, FVector& out_WorldOrigin, FVector& out_WorldDirection);
+
+	/** 
+	 * Transforms 3D world-space origin into 2D screen coordinates
+	 * @param WorldPosition - the 3d world point to transform
+	 * @param ViewRect - view rectangle
+	 * @param ViewProjectionMatrix - combined view projection matrix
+	 * @param out_ScreenPos (out) - screen coordinates in pixels
+	 */
+	static void ProjectWorldToScreen(const FVector& WorldPosition, const FIntRect& ViewRect, const FMatrix& ViewProjectionMatrix, FVector2D& out_ScreenPos);
+
 	inline FVector GetViewRight() const { return ViewMatrices.ViewMatrix.GetColumn(0); }
 	inline FVector GetViewUp() const { return ViewMatrices.ViewMatrix.GetColumn(1); }
 	inline FVector GetViewDirection() const { return ViewMatrices.ViewMatrix.GetColumn(2); }
