Removed old commented code

diff --git a/Engine/Source/Runtime/Engine/Private/Actor.cpp b/Engine/Source/Runtime/Engine/Private/Actor.cpp
--- a/Engine/Source/Runtime/Engine/Private/Actor.cpp
+++ b/Engine/Source/Runtime/Engine/Private/Actor.cpp
@@ -1,6 +1,5 @@
 // Copyright 1998-2015 Epic Games, Inc. All Rights Reserved.
 
-
 #include "EnginePrivate.h"
 #include "Engine/InputDelegateBinding.h"
 #include "Engine/LevelStreamingPersistent.h"
@@ -27,7 +26,6 @@
 #include "Kismet/GameplayStatics.h"
 #include "PhysicalMaterials/PhysicalMaterial.h"
 
-//DEFINE_LOG_CATEGORY_STATIC(LogActor, Log, All);
 DEFINE_LOG_CATEGORY(LogActor);
 
 DEFINE_STAT(STAT_GetComponentsTime);
