#ue4 - Improved docs for BrakingFriction and bUseSeparateBrakingFriction to clarify that it applies to all movement modes.

diff --git a/Engine/Source/Runtime/Engine/Classes/GameFramework/CharacterMovementComponent.h b/Engine/Source/Runtime/Engine/Classes/GameFramework/CharacterMovementComponent.h
--- a/Engine/Source/Runtime/Engine/Classes/GameFramework/CharacterMovementComponent.h
+++ b/Engine/Source/Runtime/Engine/Classes/GameFramework/CharacterMovementComponent.h
@@ -233,7 +233,8 @@ public:
 	/**
 	 * Friction (drag) coefficient applied when braking (whenever Acceleration = 0, or if character is exceeding max speed).
 	 * Only used if bUseSeparateBrakingFriction setting is true, otherwise current friction such as GroundFriction is used.
-	 * This is the current value, used in all movement modes; if this is not desired, override it when movement mode changes.
+	 * This is the current value, used in all movement modes; if this is not desired, override it when movement mode changes,
+	 * or only enable bUseSeparateBrakingFriction in desired movement modes.
 	 * @see bUseSeparateBrakingFriction
 	 */
 	UPROPERTY(Category="Character Movement (General Settings)", EditAnywhere, BlueprintReadWrite, meta=(ClampMin="0", UIMin="0", EditCondition="bUseSeparateBrakingFriction"))
@@ -337,6 +338,7 @@ public:
 	/**
 	 * If true, BrakingFriction will be used to slow the character to a stop (when there is no Acceleration).
 	 * If false, braking uses the same friction passed to CalcVelocity() (ie GroundFriction when walking), multiplied by 2.
+	 * This setting applies to all movement modes; if only desired in certain modes, consider toggling it when movement modes change.
 	 * @see BrakingFriction
 	 */
 	UPROPERTY(Category="Character Movement (General Settings)", EditDefaultsOnly, BlueprintReadWrite)
