Fixed small error in caching actor rotation for a UModel.

diff --git a/Engine/Source/Runtime/Engine/Private/Model.cpp b/Engine/Source/Runtime/Engine/Private/Model.cpp
--- a/Engine/Source/Runtime/Engine/Private/Model.cpp
+++ b/Engine/Source/Runtime/Engine/Private/Model.cpp
@@ -360,7 +360,7 @@ void UModel::PostLoad()
 			OwnerLocationWhenLastBuilt = Owner->GetActorLocation();
 			OwnerPrepivotWhenLastBuilt = Owner->GetPrePivot();
 			OwnerScaleWhenLastBuilt = Owner->GetActorScale();
-			OwnerRotationWhenLastBuilt = -Owner->GetActorRotation();
+			OwnerRotationWhenLastBuilt = Owner->GetActorRotation();
 			bCachedOwnerTransformValid = true;
 		}
 #endif
