#ue4 - Perf: Optimized and simplified IsFalling(), IsMovingOnGround(), IsSwimming(), IsFlying().

diff --git a/Engine/Source/Runtime/Engine/Classes/GameFramework/CharacterMovementComponent.h b/Engine/Source/Runtime/Engine/Classes/GameFramework/CharacterMovementComponent.h
--- a/Engine/Source/Runtime/Engine/Classes/GameFramework/CharacterMovementComponent.h
+++ b/Engine/Source/Runtime/Engine/Classes/GameFramework/CharacterMovementComponent.h
@@ -1884,6 +1884,18 @@ public:
 };
 
 
+FORCEINLINE ACharacter* UCharacterMovementComponent::GetCharacterOwner() const
+{
+	return CharacterOwner;
+}
+
+FORCEINLINE_DEBUGGABLE bool UCharacterMovementComponent::IsWalking() const
+{
+	return IsMovingOnGround();
+}
+
+
+
 /** FSavedMove_Character represents a saved move on the client that has been sent to the server and might need to be played back. */
 class ENGINE_API FSavedMove_Character
 {
@@ -2118,10 +2130,6 @@ public:
 	float UpdateTimeStampAndDeltaTime(float DeltaTime, ACharacter & CharacterOwner, class UCharacterMovementComponent & CharacterMovementComponent);
 };
 
-FORCEINLINE ACharacter* UCharacterMovementComponent::GetCharacterOwner() const
-{
-	return CharacterOwner;
-}
 
 class ENGINE_API FNetworkPredictionData_Server_Character : public FNetworkPredictionData_Server
 {
diff --git a/Engine/Source/Runtime/Engine/Private/Components/CharacterMovementComponent.cpp b/Engine/Source/Runtime/Engine/Private/Components/CharacterMovementComponent.cpp
--- a/Engine/Source/Runtime/Engine/Private/Components/CharacterMovementComponent.cpp
+++ b/Engine/Source/Runtime/Engine/Private/Components/CharacterMovementComponent.cpp
@@ -2235,46 +2235,22 @@ float UCharacterMovementComponent::ImmersionDepth()
 
 bool UCharacterMovementComponent::IsFlying() const
 {
-	if (!CharacterOwner || !UpdatedComponent)
-	{
-		return false;
-	}
-
-	return MovementMode == MOVE_Flying;
+	return (MovementMode == MOVE_Flying) && UpdatedComponent;
 }
 
-
-
 bool UCharacterMovementComponent::IsMovingOnGround() const
 {
-	if (!CharacterOwner || !UpdatedComponent)
-	{
-		return false;
-	}
-
-	return (MovementMode == MOVE_Walking) || (MovementMode == MOVE_NavWalking);
+	return ((MovementMode == MOVE_Walking) || (MovementMode == MOVE_NavWalking)) && UpdatedComponent;
 }
 
-
-
 bool UCharacterMovementComponent::IsFalling() const
 {
-	if (!CharacterOwner || !UpdatedComponent)
-	{
-		return false;
-	}
-
-	return MovementMode == MOVE_Falling;
+	return (MovementMode == MOVE_Falling) && UpdatedComponent;
 }
 
 bool UCharacterMovementComponent::IsSwimming() const
 {
-	if (!CharacterOwner || !UpdatedComponent)
-	{
-		return false;
-	}
-
-	return MovementMode == MOVE_Swimming;
+	return (MovementMode == MOVE_Swimming) && UpdatedComponent;
 }
 
 bool UCharacterMovementComponent::IsCrouching() const
@@ -2282,11 +2258,6 @@ bool UCharacterMovementComponent::IsCrouching() const
 	return CharacterOwner && CharacterOwner->bIsCrouched;
 }
 
-bool UCharacterMovementComponent::IsWalking() const
-{
-	return IsMovingOnGround();
-}
-
 void UCharacterMovementComponent::CalcVelocity(float DeltaTime, float Friction, bool bFluid, float BrakingDeceleration)
 {
 	// Do not update velocity when using root motion
