Allow custom attenuation curves

diff --git a/Engine/Source/Editor/DetailCustomizations/Private/AttenuationSettingsCustomizations.cpp b/Engine/Source/Editor/DetailCustomizations/Private/AttenuationSettingsCustomizations.cpp
--- a/Engine/Source/Editor/DetailCustomizations/Private/AttenuationSettingsCustomizations.cpp
+++ b/Engine/Source/Editor/DetailCustomizations/Private/AttenuationSettingsCustomizations.cpp
@@ -85,6 +85,9 @@ void FAttenuationSettingsCustomization::CustomizeChildren( TSharedRef<IPropertyH
 		ChildBuilder.AddChildProperty(SpatializationAlgorithmHandle.ToSharedRef());
 	}
 
+	IDetailPropertyRow& CustomCurveRow = ChildBuilder.AddChildProperty(PropertyHandles.FindChecked(GET_MEMBER_NAME_CHECKED(FAttenuationSettings, CustomAttenuationCurve)).ToSharedRef());
+	CustomCurveRow.Visibility(TAttribute<EVisibility>(this, &FAttenuationSettingsCustomization::IsCustomCurveSelected));
+
 	IDetailPropertyRow& dbAttenuationRow = ChildBuilder.AddChildProperty(PropertyHandles.FindChecked(GET_MEMBER_NAME_CHECKED(FAttenuationSettings, dBAttenuationAtMax)).ToSharedRef());
 	dbAttenuationRow.Visibility(TAttribute<EVisibility>(this, &FAttenuationSettingsCustomization::IsNaturalSoundSelected));
 
@@ -188,7 +191,7 @@ void FAttenuationSettingsCustomization::CustomizeChildren( TSharedRef<IPropertyH
 	ChildBuilder.AddChildProperty(PropertyHandles.FindChecked(GET_MEMBER_NAME_CHECKED(FAttenuationSettings, LPFRadiusMin)).ToSharedRef());
 	ChildBuilder.AddChildProperty(PropertyHandles.FindChecked(GET_MEMBER_NAME_CHECKED(FAttenuationSettings, LPFRadiusMax)).ToSharedRef());
 
-	if (PropertyHandles.Num() != 13)
+	if (PropertyHandles.Num() != 14)
 	{
 		FString PropertyList;
 		for (auto It(PropertyHandles.CreateConstIterator()); It; ++It)
@@ -247,4 +250,14 @@ EVisibility FAttenuationSettingsCustomization::IsNaturalSoundSelected() const
 	const ESoundDistanceModel DistanceAlgorithm = (ESoundDistanceModel)DistanceAlgorithmValue;
 
 	return (DistanceAlgorithm == ATTENUATION_NaturalSound ? EVisibility::Visible : EVisibility::Hidden);
+}
+
+EVisibility FAttenuationSettingsCustomization::IsCustomCurveSelected() const
+{
+	uint8 DistanceAlgorithmValue;
+	DistanceAlgorithmHandle->GetValue(DistanceAlgorithmValue);
+
+	const ESoundDistanceModel DistanceAlgorithm = (ESoundDistanceModel)DistanceAlgorithmValue;
+
+	return (DistanceAlgorithm == ATTENUATION_Custom ? EVisibility::Visible : EVisibility::Hidden);
 }
\ No newline at end of file
diff --git a/Engine/Source/Editor/DetailCustomizations/Private/AttenuationSettingsCustomizations.h b/Engine/Source/Editor/DetailCustomizations/Private/AttenuationSettingsCustomizations.h
--- a/Engine/Source/Editor/DetailCustomizations/Private/AttenuationSettingsCustomizations.h
+++ b/Engine/Source/Editor/DetailCustomizations/Private/AttenuationSettingsCustomizations.h
@@ -18,6 +18,7 @@ protected:
 	EVisibility IsCapsuleSelected() const;
 	EVisibility IsConeSelected() const;
 	EVisibility IsNaturalSoundSelected() const;
+	EVisibility IsCustomCurveSelected() const;
 
 	TSharedPtr< IPropertyHandle > AttenuationShapeHandle;
 	TSharedPtr< IPropertyHandle > DistanceAlgorithmHandle;
diff --git a/Engine/Source/Runtime/Engine/Classes/Sound/SoundAttenuation.h b/Engine/Source/Runtime/Engine/Classes/Sound/SoundAttenuation.h
--- a/Engine/Source/Runtime/Engine/Classes/Sound/SoundAttenuation.h
+++ b/Engine/Source/Runtime/Engine/Classes/Sound/SoundAttenuation.h
@@ -3,7 +3,7 @@
 #pragma once
 
 #include "CoreUObject.h"
-
+#include "Curves/CurveFloat.h"
 #include "SoundAttenuation.generated.h"
 
 UENUM()
@@ -14,6 +14,7 @@ enum ESoundDistanceModel
 	ATTENUATION_Inverse,
 	ATTENUATION_LogReverse,
 	ATTENUATION_NaturalSound,
+	ATTENUATION_Custom,
 	ATTENUATION_MAX,
 };
 
@@ -70,6 +71,9 @@ struct ENGINE_API FAttenuationSettings
 	UPROPERTY(EditAnywhere, BlueprintReadWrite, Category=Attenuation )
 	TEnumAsByte<enum ESoundDistanceModel> DistanceAlgorithm;
 
+	UPROPERTY(EditAnywhere, BlueprintReadWrite, Category=Attenuation)
+	FRuntimeFloatCurve CustomAttenuationCurve;
+
 	UPROPERTY()
 	TEnumAsByte<enum ESoundDistanceCalc> DistanceType_DEPRECATED;
 
@@ -156,6 +160,7 @@ struct ENGINE_API FAttenuationSettings
 
 private:
 
+	float AttenuationEval(const float Distance, const float Falloff, const float dBAttenuationAtMax) const;
 	float AttenuationEvalBox(const FTransform& SoundLocation, const FVector ListenerLocation) const;
 	float AttenuationEvalCapsule(const FTransform& SoundLocation, const FVector ListenerLocation) const;
 	float AttenuationEvalCone(const FTransform& SoundLocation, const FVector ListenerLocation) const;
diff --git a/Engine/Source/Runtime/Engine/Private/SoundAttenuation.cpp b/Engine/Source/Runtime/Engine/Private/SoundAttenuation.cpp
--- a/Engine/Source/Runtime/Engine/Private/SoundAttenuation.cpp
+++ b/Engine/Source/Runtime/Engine/Private/SoundAttenuation.cpp
@@ -77,7 +77,7 @@ float FAttenuationSettings::GetMaxDimension() const
  *
  * @return Attenuation value (between 0.0 and 1.0)
  */
-float AttenuationEval(const ESoundDistanceModel DistanceModel, const float Distance, const float Falloff, const float dBAttenuationAtMax)
+float FAttenuationSettings::AttenuationEval(const float Distance, const float Falloff, const float dBAttenuationAtMax) const
 {
 	if( Distance > Falloff )
 	{
@@ -86,7 +86,7 @@ float AttenuationEval(const ESoundDistanceModel DistanceModel, const float Dista
 	// UsedMinRadius is the point at which to start attenuating
 	else if( Distance > 0.f )
 	{
-		switch(DistanceModel)
+		switch(DistanceAlgorithm)
 		{
 			case ATTENUATION_Linear:
 
@@ -109,6 +109,10 @@ float AttenuationEval(const ESoundDistanceModel DistanceModel, const float Dista
 				check( dBAttenuationAtMax <= 0.0f );
 				return FMath::Pow( 10.0f, ( ( Distance / Falloff ) * dBAttenuationAtMax ) / 20.0f );
 			}
+
+			case ATTENUATION_Custom:
+
+				return CustomAttenuationCurve.GetRichCurveConst()->Eval(Distance/Falloff);
 		}
 	}
 	return 1.0f;
@@ -119,7 +123,7 @@ float FAttenuationSettings::AttenuationEvalBox(const FTransform& SoundTransform,
 	const float DistanceSq = ComputeSquaredDistanceFromBoxToPoint(-AttenuationShapeExtents, AttenuationShapeExtents,SoundTransform.InverseTransformPositionNoScale(ListenerLocation));
 	if (DistanceSq < FalloffDistance * FalloffDistance)
 	{
-		return AttenuationEval(DistanceAlgorithm, FMath::Sqrt(DistanceSq), FalloffDistance, dBAttenuationAtMax);
+		return AttenuationEval(FMath::Sqrt(DistanceSq), FalloffDistance, dBAttenuationAtMax);
 	}
 
 	return 0.f;
@@ -145,7 +149,7 @@ float FAttenuationSettings::AttenuationEvalCapsule(const FTransform& SoundTransf
 		Distance = FMath::PointDistToSegment(ListenerLocation, StartPoint, EndPoint) - CapsuleRadius;
 	}
 
-	return AttenuationEval(DistanceAlgorithm, Distance, FalloffDistance, dBAttenuationAtMax);
+	return AttenuationEval(Distance, FalloffDistance, dBAttenuationAtMax);
 }
 
 float FAttenuationSettings::AttenuationEvalCone(const FTransform& SoundTransform, const FVector ListenerLocation) const
@@ -157,12 +161,12 @@ float FAttenuationSettings::AttenuationEvalCone(const FTransform& SoundTransform
 	const FVector Origin = SoundTransform.GetTranslation() - (SoundForward * ConeOffset);
 
 	const float Distance = FMath::Max(FVector::Dist( Origin, ListenerLocation ) - AttenuationShapeExtents.X, 0.f);
-	VolumeMultiplier *= AttenuationEval(DistanceAlgorithm, Distance, FalloffDistance, dBAttenuationAtMax);
+	VolumeMultiplier *= AttenuationEval(Distance, FalloffDistance, dBAttenuationAtMax);
 
 	if (VolumeMultiplier > 0.f)
 	{
 		const float theta = FMath::RadiansToDegrees(fabsf(FMath::Acos( FVector::DotProduct(SoundForward, (ListenerLocation - Origin).GetSafeNormal()))));
-		VolumeMultiplier *= AttenuationEval(DistanceAlgorithm, theta - AttenuationShapeExtents.Y, AttenuationShapeExtents.Z, dBAttenuationAtMax);
+		VolumeMultiplier *= AttenuationEval(theta - AttenuationShapeExtents.Y, AttenuationShapeExtents.Z, dBAttenuationAtMax);
 	}
 
 	return VolumeMultiplier;
@@ -193,7 +197,7 @@ void FAttenuationSettings::ApplyAttenuation( const FTransform& SoundTransform, c
 		case EAttenuationShape::Sphere:
 			{
 				const float Distance = FMath::Max(FVector::Dist( SoundTransform.GetTranslation(), ListenerLocation ) - AttenuationShapeExtents.X, 0.f);
-				Volume *= AttenuationEval(DistanceAlgorithm, Distance, FalloffDistance, dBAttenuationAtMax);
+				Volume *= AttenuationEval(Distance, FalloffDistance, dBAttenuationAtMax);
 				break;
 			}
 
