- Simplified & optimized project/deproject functions in PlayerController
- consted a couple LocalPlayer functions

diff --git a/Engine/Source/Runtime/Engine/Classes/Engine/LocalPlayer.h b/Engine/Source/Runtime/Engine/Classes/Engine/LocalPlayer.h
--- a/Engine/Source/Runtime/Engine/Classes/Engine/LocalPlayer.h
+++ b/Engine/Source/Runtime/Engine/Classes/Engine/LocalPlayer.h
@@ -220,7 +220,7 @@ protected:
 	 * @param OutViewInfo - Upon return contains the view information for the player.
 	 * @param StereoPass - Which stereoscopic pass, if any, to get the viewport for.  This will include eye offsetting
 	 */
-	void GetViewPoint(FMinimalViewInfo& OutViewInfo, EStereoscopicPass StereoPass = eSSP_FULL);
+	void GetViewPoint(FMinimalViewInfo& OutViewInfo, EStereoscopicPass StereoPass = eSSP_FULL) const;
 
 	/** @todo document */
 	void ExecMacro( const TCHAR* Filename, FOutputDevice& Ar );
@@ -394,7 +394,7 @@ public:
 	 * @param	ProjectionData			The structure to be filled with projection data
 	 * @return  False if there is no viewport, or if the Actor is null
 	 */
-	bool GetProjectionData(FViewport* Viewport, EStereoscopicPass StereoPass, FSceneViewProjectionData& ProjectionData);
+	bool GetProjectionData(FViewport* Viewport, EStereoscopicPass StereoPass, FSceneViewProjectionData& ProjectionData) const;
 
 	/**
 	 * Determines whether this player is the first and primary player on their machine.
diff --git a/Engine/Source/Runtime/Engine/Private/LocalPlayer.cpp b/Engine/Source/Runtime/Engine/Private/LocalPlayer.cpp
--- a/Engine/Source/Runtime/Engine/Private/LocalPlayer.cpp
+++ b/Engine/Source/Runtime/Engine/Private/LocalPlayer.cpp
@@ -334,7 +334,7 @@ public:
 	 * @param OutFOV			The FOV at which the view was locked.
 	 * @returns true if the view is locked, false if it is not.
 	 */
-	bool GetViewPoint(ULocalPlayer* Player, FVector& OutViewLocation, FRotator& OutViewRotation, float& OutFOV)
+	bool GetViewPoint(ULocalPlayer const* Player, FVector& OutViewLocation, FRotator& OutViewRotation, float& OutFOV)
 	{
 		FPlayerState PlayerState = PlayerStates.GetAnnotation(Player);
 		if (PlayerState.bLocked)
@@ -350,7 +350,7 @@ public:
 	/**
 	 * Returns true if the player's viewpoint is locked.
 	 */
-	bool IsViewLocked(ULocalPlayer* Player)
+	bool IsViewLocked(ULocalPlayer const* Player)
 	{
 		FPlayerState PlayerState = PlayerStates.GetAnnotation(Player);
 		return PlayerState.bLocked;
@@ -604,7 +604,7 @@ FAutoConsoleCommand FLockedViewState::CmdCopyLockedViews(
 	FConsoleCommandDelegate::CreateStatic(FLockedViewState::CopyLockedViews)
 	);
 
-void ULocalPlayer::GetViewPoint(FMinimalViewInfo& OutViewInfo, EStereoscopicPass StereoPass)
+void ULocalPlayer::GetViewPoint(FMinimalViewInfo& OutViewInfo, EStereoscopicPass StereoPass) const
 {
 	if (FLockedViewState::Get().GetViewPoint(this, OutViewInfo.Location, OutViewInfo.Rotation, OutViewInfo.FOV) == false
 		&& PlayerController != NULL)
@@ -867,7 +867,7 @@ bool ULocalPlayer::GetPixelPoint(const FVector& InPoint, FVector2D& OutPoint, co
 	return bInFrontOfCamera;
 }
 
-bool ULocalPlayer::GetProjectionData(FViewport* Viewport, EStereoscopicPass StereoPass, FSceneViewProjectionData& ProjectionData)
+bool ULocalPlayer::GetProjectionData(FViewport* Viewport, EStereoscopicPass StereoPass, FSceneViewProjectionData& ProjectionData) const
 {
 	// If the actor
 	if ((Viewport == NULL) || (PlayerController == NULL) || (Viewport->GetSizeXY().X == 0) || (Viewport->GetSizeXY().Y == 0))
diff --git a/Engine/Source/Runtime/Engine/Private/PlayerController.cpp b/Engine/Source/Runtime/Engine/Private/PlayerController.cpp
--- a/Engine/Source/Runtime/Engine/Private/PlayerController.cpp
+++ b/Engine/Source/Runtime/Engine/Private/PlayerController.cpp
@@ -1717,7 +1717,7 @@ bool APlayerController::GetHitResultUnderCursor(ECollisionChannel TraceChannel,
 		}
 	}
 
-	if(!bHit)	//If there was no hit we reset the results. This is redundent but helps Blueprint users
+	if(!bHit)	//If there was no hit we reset the results. This is redundant but helps Blueprint users
 	{
 		HitResult = FHitResult();
 	}
@@ -1738,7 +1738,7 @@ bool APlayerController::GetHitResultUnderCursorByChannel(ETraceTypeQuery TraceCh
 		}
 	}
 
-	if(!bHit)	//If there was no hit we reset the results. This is redundent but helps Blueprint users
+	if(!bHit)	//If there was no hit we reset the results. This is redundant but helps Blueprint users
 	{
 		HitResult = FHitResult();
 	}
@@ -1835,127 +1835,64 @@ bool APlayerController::GetHitResultUnderFingerForObjects(ETouchIndex::Type Fing
 
 bool APlayerController::DeprojectMousePositionToWorld(FVector& WorldLocation, FVector& WorldDirection) const
 {
-	bool bSuccessfulDeproject = false;
-
-	ULocalPlayer* LocalPlayer = Cast<ULocalPlayer>(Player);
-	if (LocalPlayer && LocalPlayer->ViewportClient )
+	ULocalPlayer* const LocalPlayer = GetLocalPlayer();
+	if (LocalPlayer && LocalPlayer->ViewportClient)
 	{
 		FVector2D ScreenPosition;
 		if (LocalPlayer->ViewportClient->GetMousePosition(ScreenPosition))
 		{
-			bSuccessfulDeproject = DeprojectScreenPositionToWorld(ScreenPosition.X, ScreenPosition.Y, WorldLocation, WorldDirection);
+			return UGameplayStatics::DeprojectScreenToWorld(this, ScreenPosition, WorldLocation, WorldDirection);
 		}
 	}
 
-	return bSuccessfulDeproject;
+	return false;
 }
 
 bool APlayerController::DeprojectScreenPositionToWorld(float ScreenX, float ScreenY, FVector& WorldLocation, FVector& WorldDirection) const
 {
-	bool bSuccessfulDeproject = false;
-
-	ULocalPlayer* LocalPlayer = Cast<ULocalPlayer>(Player);
-
-	if (LocalPlayer && LocalPlayer->ViewportClient && LocalPlayer->ViewportClient->Viewport)
-	{
-		// Create a view family for the game viewport
-		FSceneViewFamilyContext ViewFamily(FSceneViewFamily::ConstructionValues(
-			LocalPlayer->ViewportClient->Viewport,
-			GetWorld()->Scene,
-			LocalPlayer->ViewportClient->EngineShowFlags)
-			.SetRealtimeUpdate(true));
-
-		// Calculate a view where the player is to update the streaming from the players start location
-		FVector ViewLocation;
-		FRotator ViewRotation;
-		FSceneView* SceneView = LocalPlayer->CalcSceneView(&ViewFamily, /*out*/ ViewLocation, /*out*/ ViewRotation, LocalPlayer->ViewportClient->Viewport);
-
-		if (SceneView)
-		{
-			const FVector2D ScreenPosition(ScreenX, ScreenY);
-			SceneView->DeprojectFVector2D(ScreenPosition, WorldLocation, WorldDirection);
-
-			bSuccessfulDeproject = true;
-		}
-	}
-
-	return bSuccessfulDeproject;
+	return UGameplayStatics::DeprojectScreenToWorld(this, FVector2D(ScreenX, ScreenY), WorldLocation, WorldDirection);
 }
 
 
 bool APlayerController::ProjectWorldLocationToScreen(FVector WorldLocation, FVector2D& ScreenLocation) const
 {
-	FVector FullScreenLocation;
-	const bool bSuccess = ProjectWorldLocationToScreenWithDistance(WorldLocation, FullScreenLocation);
-	ScreenLocation = FVector2D(FullScreenLocation.X, FullScreenLocation.Y);
-	return bSuccess;
+	return UGameplayStatics::ProjectWorldToScreen(this, WorldLocation, ScreenLocation);
 }
 
 bool APlayerController::ProjectWorldLocationToScreenWithDistance(FVector WorldLocation, FVector& ScreenLocation) const
 {
-	ULocalPlayer* LocalPlayer = Cast<ULocalPlayer>(Player);
-	if (LocalPlayer != NULL && LocalPlayer->ViewportClient != NULL && LocalPlayer->ViewportClient->Viewport != NULL)
+	FVector2D ScreenLoc2D;
+	if (UGameplayStatics::ProjectWorldToScreen(this, WorldLocation, ScreenLoc2D))
 	{
-		// Create a view family for the game viewport
-		FSceneViewFamilyContext ViewFamily( FSceneViewFamily::ConstructionValues(
-			LocalPlayer->ViewportClient->Viewport,
-			GetWorld()->Scene,
-			LocalPlayer->ViewportClient->EngineShowFlags )
-			.SetRealtimeUpdate(true) );
-
-		// Calculate a view where the player is to update the streaming from the players start location
-		FVector ViewLocation;
-		FRotator ViewRotation;
-		FSceneView* SceneView = LocalPlayer->CalcSceneView( &ViewFamily, /*out*/ ViewLocation, /*out*/ ViewRotation, LocalPlayer->ViewportClient->Viewport );
-
-		if (SceneView) 
+		// find distance
+		ULocalPlayer const* const LP = GetLocalPlayer();
+		if (LP && LP->ViewportClient)
 		{
-			FVector2D ScreenLocation2D;
-			const bool bSuccess = SceneView->WorldToPixel(WorldLocation, ScreenLocation2D);
-			ScreenLocation = FVector(ScreenLocation2D.X, ScreenLocation2D.Y, FVector::Dist(ViewLocation, WorldLocation));
-			return bSuccess;
+			// get the projection data
+			FSceneViewProjectionData ProjectionData;
+			if (LP->GetProjectionData(LP->ViewportClient->Viewport, eSSP_FULL, /*out*/ ProjectionData))
+			{
+				ScreenLocation = FVector(ScreenLoc2D.X, ScreenLoc2D.Y, FVector::Dist(ProjectionData.ViewOrigin, WorldLocation));
+			}
 		}
 	}
 
 	return false;
 }
 
-
 bool APlayerController::GetHitResultAtScreenPosition(const FVector2D ScreenPosition, const ECollisionChannel TraceChannel, const FCollisionQueryParams& CollisionQueryParams, FHitResult& HitResult) const
 {
 	// Early out if we clicked on a HUD hitbox
-	if( GetHUD() != NULL && GetHUD()->GetHitBoxAtCoordinates(ScreenPosition, true) )
+	if (GetHUD() != NULL && GetHUD()->GetHitBoxAtCoordinates(ScreenPosition, true))
 	{
 		return false;
 	}
 
-	ULocalPlayer* LocalPlayer = Cast<ULocalPlayer>(Player);
-
-	if (LocalPlayer != NULL && LocalPlayer->ViewportClient != NULL && LocalPlayer->ViewportClient->Viewport != NULL)
+	FVector WorldOrigin;
+	FVector WorldDirection;
+	if (UGameplayStatics::DeprojectScreenToWorld(this, ScreenPosition, WorldOrigin, WorldDirection) == true)
 	{
-		UGameViewportClient* ViewportClient = LocalPlayer->ViewportClient;
-
-		// Create a view family for the game viewport
-		FSceneViewFamilyContext ViewFamily( FSceneViewFamily::ConstructionValues(
-			ViewportClient->Viewport,
-			GetWorld()->Scene,
-			ViewportClient->EngineShowFlags)
-			.SetRealtimeUpdate(true) );
-
-
-		// Calculate a view where the player is to update the streaming from the players start location
-		FVector ViewLocation;
-		FRotator ViewRotation;
-		FSceneView* SceneView = LocalPlayer->CalcSceneView(&ViewFamily, /*out*/ ViewLocation, /*out*/ ViewRotation, ViewportClient->Viewport);
-
-		if (SceneView)
-		{
-			FVector WorldOrigin;
-			FVector WorldDirection;
-			SceneView->DeprojectFVector2D(ScreenPosition, WorldOrigin, WorldDirection);
-
-			return GetWorld()->LineTraceSingleByChannel(HitResult, WorldOrigin, WorldOrigin + WorldDirection * HitResultTraceDistance, TraceChannel, CollisionQueryParams);
-		}
+		return GetWorld()->LineTraceSingleByChannel(HitResult, WorldOrigin, WorldOrigin + WorldDirection * HitResultTraceDistance, TraceChannel, CollisionQueryParams);
 	}
 
 	return false;
@@ -1981,33 +1918,12 @@ bool APlayerController::GetHitResultAtScreenPosition(const FVector2D ScreenPosit
 		return false;
 	}
 
-	ULocalPlayer* LocalPlayer = Cast<ULocalPlayer>(Player);
-
-	if (LocalPlayer != NULL && LocalPlayer->ViewportClient != NULL && LocalPlayer->ViewportClient->Viewport != NULL)
+	FVector WorldOrigin;
+	FVector WorldDirection;
+	if (UGameplayStatics::DeprojectScreenToWorld(this, ScreenPosition, WorldOrigin, WorldDirection) == true)
 	{
-		// Create a view family for the game viewport
-		FSceneViewFamilyContext ViewFamily( FSceneViewFamily::ConstructionValues(
-			LocalPlayer->ViewportClient->Viewport,
-			GetWorld()->Scene,
-			LocalPlayer->ViewportClient->EngineShowFlags )
-			.SetRealtimeUpdate(true) );
-
-
-		// Calculate a view where the player is to update the streaming from the players start location
-		FVector ViewLocation;
-		FRotator ViewRotation;
-		FSceneView* SceneView = LocalPlayer->CalcSceneView( &ViewFamily, /*out*/ ViewLocation, /*out*/ ViewRotation, LocalPlayer->ViewportClient->Viewport );
-
-		if (SceneView)
-		{
-			FVector WorldOrigin;
-			FVector WorldDirection;
-			SceneView->DeprojectFVector2D(ScreenPosition, WorldOrigin, WorldDirection);
-
-			FCollisionObjectQueryParams ObjParam(ObjectTypes);
-			
-			return GetWorld()->LineTraceSingleByObjectType(HitResult, WorldOrigin, WorldOrigin + WorldDirection * HitResultTraceDistance, ObjParam, FCollisionQueryParams("ClickableTrace", bTraceComplex));
-		}
+		FCollisionObjectQueryParams const ObjParam(ObjectTypes);
+		return GetWorld()->LineTraceSingleByObjectType(HitResult, WorldOrigin, WorldOrigin + WorldDirection * HitResultTraceDistance, ObjParam, FCollisionQueryParams("ClickableTrace", bTraceComplex));
 	}
 
 	return false;
