Fixed CameraModifier blending issue with blends longer than 1 second.

diff --git a/Engine/Source/Runtime/Engine/Private/Camera/CameraModifier.cpp b/Engine/Source/Runtime/Engine/Private/Camera/CameraModifier.cpp
--- a/Engine/Source/Runtime/Engine/Private/Camera/CameraModifier.cpp
+++ b/Engine/Source/Runtime/Engine/Private/Camera/CameraModifier.cpp
@@ -64,14 +64,17 @@ void UCameraModifier::UpdateAlpha(float DeltaTime)
 	// interpolate!
 	if (BlendTime <= 0.f)
 	{
+		// no blendtime means no blending, just go directly to target alpha
 		Alpha = TargetAlpha;
 	}
-	else if (BlendTime > TargetAlpha)
+	else if (Alpha > TargetAlpha)
 	{
+		// interpolate downward to target, while protecting against overshooting
 		Alpha = FMath::Max<float>(Alpha - DeltaTime / BlendTime, TargetAlpha);
 	}
 	else
 	{
+		// interpolate upward to target, while protecting against overshooting
 		Alpha = FMath::Min<float>(Alpha + DeltaTime / BlendTime, TargetAlpha);
 	}
 }
