Added "project"/"deproject" keywords to BP-accessible PlayerController projection/deprojection functions.

diff --git a/Engine/Source/Runtime/Engine/Classes/GameFramework/PlayerController.h b/Engine/Source/Runtime/Engine/Classes/GameFramework/PlayerController.h
--- a/Engine/Source/Runtime/Engine/Classes/GameFramework/PlayerController.h
+++ b/Engine/Source/Runtime/Engine/Classes/GameFramework/PlayerController.h
@@ -466,18 +466,18 @@ public:
 	bool GetHitResultUnderFingerForObjects(ETouchIndex::Type FingerIndex, const  TArray<TEnumAsByte<EObjectTypeQuery> > & ObjectTypes, bool bTraceComplex, FHitResult& HitResult) const;
 
 	/** Convert current mouse 2D position to World Space 3D position and direction. Returns false if unable to determine value. **/
-	UFUNCTION(BlueprintCallable, Category="Game|Player", meta=(DisplayName="ConvertMouseLocationToWorldSpace"))
+	UFUNCTION(BlueprintCallable, Category = "Game|Player", meta = (DisplayName = "ConvertMouseLocationToWorldSpace", Keywords = "deproject"))
 	bool DeprojectMousePositionToWorld(FVector& WorldLocation, FVector& WorldDirection) const;
 
 	/** Convert current mouse 2D position to World Space 3D position and direction. Returns false if unable to determine value. **/
-	UFUNCTION(BlueprintCallable, Category = "Game|Player", meta = (DisplayName = "ConvertScreenLocationToWorldSpace"))
+	UFUNCTION(BlueprintCallable, Category = "Game|Player", meta = (DisplayName = "ConvertScreenLocationToWorldSpace", Keywords = "deproject"))
 	bool DeprojectScreenPositionToWorld(float ScreenX, float ScreenY, FVector& WorldLocation, FVector& WorldDirection) const;
 
 	/**
 	 * Convert a World Space 3D position into a 2D Screen Space position.
 	 * @return true if the world coordinate was successfully projected to the screen.
 	 */
-	UFUNCTION(BlueprintCallable, Category = "Game|Player", meta = ( DisplayName = "ConvertWorldLocationToScreenLocation" ))
+	UFUNCTION(BlueprintCallable, Category = "Game|Player", meta = (DisplayName = "ConvertWorldLocationToScreenLocation", Keywords = "project"))
 	bool ProjectWorldLocationToScreen(FVector WorldLocation, FVector2D& ScreenLocation) const;
 
 	/**
