Added PrecomputedAOMask material node - allows procedural texturing using Ambient Occlusion computed by Lightmass
* bGenerateAmbientOcclusionMaterialMask must be enabled in WorldSettings for the mask to be produced by a lighting build.  The mask is stored in a BC4 texure which adds half a byte per lightmap texel.

diff --git a/Engine/Shaders/BasePassPixelShader.usf b/Engine/Shaders/BasePassPixelShader.usf
--- a/Engine/Shaders/BasePassPixelShader.usf
+++ b/Engine/Shaders/BasePassPixelShader.usf
@@ -653,6 +653,13 @@ void Main(
 		GetMaterialCoverageAndClipping(MaterialParameters);
 	}
 
+#if HQ_TEXTURE_LIGHTMAP && USES_AO_MATERIAL_MASK && !MATERIAL_SHADINGMODEL_UNLIT
+	float2 LightmapUV0, LightmapUV1;
+	GetLightMapCoordinates(Interpolants, LightmapUV0, LightmapUV1);
+	// Must be computed before GetMaterialBaseColor
+	MaterialParameters.AOMaterialMask = GetAOMaterialMask(LightmapUV0 * float2(1, 2));
+#endif
+
 	// Store the results in local variables and reuse instead of calling the functions multiple times.
 	half3 BaseColor = GetMaterialBaseColor( MaterialParameters );
 	half  Metallic = GetMaterialMetallic( MaterialParameters );
diff --git a/Engine/Shaders/LightmapCommon.usf b/Engine/Shaders/LightmapCommon.usf
--- a/Engine/Shaders/LightmapCommon.usf
+++ b/Engine/Shaders/LightmapCommon.usf
@@ -12,6 +12,9 @@ SamplerState LightMapSampler;
 Texture2D SkyOcclusionTexture;
 SamplerState SkyOcclusionSampler;
 
+Texture2D AOMaterialMaskTexture;
+SamplerState AOMaterialMaskSampler;
+
 half4 LightMapScale[NUM_LIGHTMAP_COEFFICIENTS];
 half4 LightMapAdd[NUM_LIGHTMAP_COEFFICIENTS];
 
@@ -103,6 +106,18 @@ float4 GetSkyBentNormalAndOcclusion(float2 LightmapUV)
 	return TextureValue;
 }
 
+float GetAOMaterialMask(float2 LightmapUV)
+{
+#if SUPPORTS_INDEPENDENT_SAMPLERS
+	// Share sampler with lightmaps to give artists more available
+	float TextureValue = Texture2DSample(AOMaterialMaskTexture, LightMapSampler, LightmapUV.xy).x;
+#else
+	float TextureValue = Texture2DSample(AOMaterialMaskTexture, AOMaterialMaskSampler, LightmapUV.xy).x;
+#endif
+	// Undo sqrt which allocated more precision toward 0
+	return TextureValue * TextureValue;
+}
+
 #endif
 
 #if STATICLIGHTING_TEXTUREMASK && STATICLIGHTING_SIGNEDDISTANCEFIELD
diff --git a/Engine/Shaders/MaterialTemplate.usf b/Engine/Shaders/MaterialTemplate.usf
--- a/Engine/Shaders/MaterialTemplate.usf
+++ b/Engine/Shaders/MaterialTemplate.usf
@@ -152,6 +152,8 @@ struct FMaterialPixelParameters
 	/** Offset applied to the lighting position for translucency, used to break up aliasing artifacts. */
 	half3 LightingPositionOffset;
 
+	float AOMaterialMask;
+
 #if LIGHTMAP_UV_ACCESS
 	float2	LightmapUVs;
 #endif
@@ -199,6 +201,7 @@ FMaterialPixelParameters MakeInitializedMaterialPixelParameters()
 	MPP.Particle.MotionBlurFade = 0;
 	MPP.Particle.Velocity = 0;
 	MPP.LightingPositionOffset = 0;
+	MPP.AOMaterialMask = 0;
 	MPP.Particle.Color = 0;
 	MPP.Particle.PositionAndSize = 0;
 	MPP.Particle.MacroUV = 0;
diff --git a/Engine/Source/Editor/UnrealEd/Private/Lightmass/Lightmass.cpp b/Engine/Source/Editor/UnrealEd/Private/Lightmass/Lightmass.cpp
--- a/Engine/Source/Editor/UnrealEd/Private/Lightmass/Lightmass.cpp
+++ b/Engine/Source/Editor/UnrealEd/Private/Lightmass/Lightmass.cpp
@@ -1896,6 +1896,7 @@ void FLightmassExporter::WriteSceneSettings( Lightmass::FSceneFileHeader& Scene
 	}
 	{
 		Scene.AmbientOcclusionSettings.bUseAmbientOcclusion = LevelSettings.bUseAmbientOcclusion;
+		Scene.AmbientOcclusionSettings.bGenerateAmbientOcclusionMaterialMask = LevelSettings.bGenerateAmbientOcclusionMaterialMask;
 		Scene.AmbientOcclusionSettings.bVisualizeAmbientOcclusion = LevelSettings.bVisualizeAmbientOcclusion;
 		Scene.AmbientOcclusionSettings.DirectIlluminationOcclusionFraction = LevelSettings.DirectIlluminationOcclusionFraction;
 		Scene.AmbientOcclusionSettings.IndirectIlluminationOcclusionFraction = LevelSettings.IndirectIlluminationOcclusionFraction;
diff --git a/Engine/Source/Programs/UnrealLightmass/Private/ImportExport/LightmassScene.cpp b/Engine/Source/Programs/UnrealLightmass/Private/ImportExport/LightmassScene.cpp
--- a/Engine/Source/Programs/UnrealLightmass/Private/ImportExport/LightmassScene.cpp
+++ b/Engine/Source/Programs/UnrealLightmass/Private/ImportExport/LightmassScene.cpp
@@ -407,7 +407,6 @@ void FScene::ApplyStaticLightingScale()
 	SceneConstants.SmallestTexelRadius *= SceneConstants.StaticLightingLevelScale;
 	MeshAreaLightSettings.MeshAreaLightSimplifyCornerDistanceThreshold *= SceneConstants.StaticLightingLevelScale;
 	MeshAreaLightSettings.MeshAreaLightGeneratedDynamicLightSurfaceOffset *= SceneConstants.StaticLightingLevelScale;
-	AmbientOcclusionSettings.MaxOcclusionDistance *= SceneConstants.StaticLightingLevelScale;
 	DynamicObjectSettings.FirstSurfaceSampleLayerHeight *= SceneConstants.StaticLightingLevelScale;
 	DynamicObjectSettings.SurfaceLightSampleSpacing *= SceneConstants.StaticLightingLevelScale;
 	DynamicObjectSettings.SurfaceSampleLayerHeightSpacing *= SceneConstants.StaticLightingLevelScale;
diff --git a/Engine/Source/Programs/UnrealLightmass/Private/Lighting/FinalGather.cpp b/Engine/Source/Programs/UnrealLightmass/Private/Lighting/FinalGather.cpp
--- a/Engine/Source/Programs/UnrealLightmass/Private/Lighting/FinalGather.cpp
+++ b/Engine/Source/Programs/UnrealLightmass/Private/Lighting/FinalGather.cpp
@@ -462,6 +462,13 @@ FLinearColor FStaticLightingSystem::FinalGatherSample(
 	const FVector4 TangentPathDirection = Vertex.TransformWorldVectorToTangent(WorldPathDirection);
 	checkSlow(TangentPathDirection.IsUnit3());
 
+#if ALLOW_LIGHTMAP_SAMPLE_DEBUGGING
+	if (bDebugThisTexel)
+	{
+		int32 asdf = 0;
+	}
+#endif
+
 	FVector4 SampleOffset(0,0,0);
 	if (GeneralSettings.bAccountForTexelSize)
 	{
@@ -506,7 +513,9 @@ FLinearColor FStaticLightingSystem::FinalGatherSample(
 
 		if (IntersectionDistance < AmbientOcclusionSettings.MaxOcclusionDistance)
 		{
-			FinalGatherInfo.NumSamplesOccluded += 1.0f / RayIntersection.Mesh->GetFullyOccludedSamplesFraction(RayIntersection.ElementIndex);
+			const float DistanceFraction = IntersectionDistance / AmbientOcclusionSettings.MaxOcclusionDistance;
+			const float DistanceWeight = 1.0f - 1.0f * DistanceFraction * DistanceFraction;
+			FinalGatherInfo.NumSamplesOccluded += DistanceWeight / RayIntersection.Mesh->GetFullyOccludedSamplesFraction(RayIntersection.ElementIndex);
 		}
 
 		// Only continue if the ray hit the frontface of the polygon, otherwise the ray started inside a mesh
@@ -905,9 +914,6 @@ public:
 				FRefinementTraversalContext NodeContext = (*CurrentNodesToRefine)[NodeIndex];
 				FLinearColor SubsampledRadiance = FLinearColor::Black;
 				FLightingCacheGatherInfo SubsampleGatherInfo; 
-				// This final gather info is currently being discarded 
-				//@todo - use it to improve AO quality
-				FFinalGatherInfo SubsampleFinalGatherInfo;
 
 				for (int32 SubThetaIndex = 0; SubThetaIndex < NumSubsamples; SubThetaIndex++)
 				{
@@ -949,6 +955,7 @@ public:
 
 							FVector UnoccludedSkyVector;
 							FLinearColor StationarySkyLighting;
+							FFinalGatherInfo SubsampleFinalGatherInfo;
 
 							const FLinearColor SubsampleLighting = LightingSystem.FinalGatherSample(
 								Mapping,
@@ -964,7 +971,7 @@ public:
 								UnoccludedSkyVector,
 								StationarySkyLighting);
 
-							FreeNode->Element = FRefinementElement(FLightingAndOcclusion(SubsampleLighting, UnoccludedSkyVector, StationarySkyLighting), FVector2D(Fraction1, Fraction2));
+							FreeNode->Element = FRefinementElement(FLightingAndOcclusion(SubsampleLighting, UnoccludedSkyVector, StationarySkyLighting, SubsampleFinalGatherInfo.NumSamplesOccluded), FVector2D(Fraction1, Fraction2));
 
 							Stats.NumRefiningFinalGatherSamples[RefinementDepth]++;
 
@@ -1051,12 +1058,11 @@ SampleType FStaticLightingSystem::IncomingRadianceAdaptive(
 
 	const double StartBaseTraceTime = FPlatformTime::Seconds();
 
-	FFinalGatherInfo FinalGatherInfo;
-
 	const int32 NumThetaSteps = FMath::TruncToInt(FMath::Sqrt(UniformHemisphereSamples.Num() / (float)PI) + .5f);
 	const int32 NumPhiSteps = UniformHemisphereSamples.Num() / NumThetaSteps;
 	checkSlow(NumThetaSteps * NumPhiSteps == UniformHemisphereSamples.Num());
 
+	int32 NumBackfaceHits = 0;
 	FUniformHemisphereRefinementGrid RefinementGrid(NumThetaSteps, NumPhiSteps);
 
 	// Initialize the root level of the refinement grid with lighting values
@@ -1069,6 +1075,7 @@ SampleType FStaticLightingSystem::IncomingRadianceAdaptive(
 			
 			FVector UnoccludedSkyVector;
 			FLinearColor StationarySkyLighting;
+			FFinalGatherInfo FinalGatherInfo;
 
 			const FLinearColor Radiance = FinalGatherSample(
 				Mapping,
@@ -1084,7 +1091,8 @@ SampleType FStaticLightingSystem::IncomingRadianceAdaptive(
 				UnoccludedSkyVector,
 				StationarySkyLighting);
 
-			RefinementGrid.SetRootElement(ThetaIndex, PhiIndex, FRefinementElement(FLightingAndOcclusion(Radiance, UnoccludedSkyVector, StationarySkyLighting), UniformHemisphereSampleUniforms[SampleIndex]));
+			NumBackfaceHits += FinalGatherInfo.NumBackfaceHits;
+			RefinementGrid.SetRootElement(ThetaIndex, PhiIndex, FRefinementElement(FLightingAndOcclusion(Radiance, UnoccludedSkyVector, StationarySkyLighting, FinalGatherInfo.NumSamplesOccluded), UniformHemisphereSampleUniforms[SampleIndex]));
 		}
 	}
 
@@ -1092,7 +1100,7 @@ SampleType FStaticLightingSystem::IncomingRadianceAdaptive(
 
 	MappingContext.Stats.BaseFinalGatherSampleTime += EndBaseTraceTime - StartBaseTraceTime;
 	MappingContext.Stats.NumBaseFinalGatherSamples += NumThetaSteps * NumPhiSteps;
-	GatherInfo.BackfacingHitsFraction = FinalGatherInfo.NumBackfaceHits / (float)UniformHemisphereSamples.Num();
+	GatherInfo.BackfacingHitsFraction = NumBackfaceHits / (float)UniformHemisphereSamples.Num();
 
 	// Refine if we are not hidden inside some geometry
 	const bool bRefine = GatherInfo.BackfacingHitsFraction < .5f || bIntersectingSurface;
@@ -1123,8 +1131,16 @@ SampleType FStaticLightingSystem::IncomingRadianceAdaptive(
 
 	MappingContext.Stats.RefiningFinalGatherSampleTime += EndRefiningTime - EndBaseTraceTime;
 
+#if ALLOW_LIGHTMAP_SAMPLE_DEBUGGING
+	if (bDebugThisTexel)
+	{
+		int32 TempBreak = 0;
+	}
+#endif
+
 	SampleType IncomingRadiance;
 	FVector CombinedSkyUnoccludedDirection(0);
+	float NumSamplesOccluded = 0;
 
 	// Accumulate lighting from all samples
 	for (int32 ThetaIndex = 0; ThetaIndex < NumThetaSteps; ThetaIndex++)
@@ -1154,13 +1170,13 @@ SampleType FStaticLightingSystem::IncomingRadianceAdaptive(
 			IncomingRadiance.AddIncomingRadiance(Radiance, SampleWeight, TangentPathDirection, WorldPathDirection);
 			IncomingRadiance.AddIncomingStationarySkyLight(FilteredLighting.StationarySkyLighting, SampleWeight, TangentPathDirection, WorldPathDirection);
 			checkSlow(IncomingRadiance.AreFloatsValid());
+			NumSamplesOccluded += FilteredLighting.NumSamplesOccluded;
 		}
 	}
 
-	//@todo - AO is not adaptively sampled
 	// Calculate the fraction of samples which were occluded
 	const float MaterialElementFullyOccludedSamplesFraction = Mapping ? Mapping->Mesh->GetFullyOccludedSamplesFraction(ElementIndex) : 1.0f;
-	const float OcclusionFraction = FMath::Min(FinalGatherInfo.NumSamplesOccluded / (AmbientOcclusionSettings.FullyOccludedSamplesFraction * MaterialElementFullyOccludedSamplesFraction * UniformHemisphereSamples.Num()), 1.0f);
+	const float OcclusionFraction = FMath::Min(NumSamplesOccluded / (AmbientOcclusionSettings.FullyOccludedSamplesFraction * MaterialElementFullyOccludedSamplesFraction * UniformHemisphereSamples.Num()), 1.0f);
 	// Constant which maintains an integral of .5 for the unclamped exponential function applied to occlusion below
 	// An integral of .5 is important because it makes an image with a uniform distribution of occlusion values stay the same brightness with different exponents.
 	// As a result, OcclusionExponent just controls contrast and doesn't affect brightness.
diff --git a/Engine/Source/Programs/UnrealLightmass/Private/Lighting/LightingSystem.cpp b/Engine/Source/Programs/UnrealLightmass/Private/Lighting/LightingSystem.cpp
--- a/Engine/Source/Programs/UnrealLightmass/Private/Lighting/LightingSystem.cpp
+++ b/Engine/Source/Programs/UnrealLightmass/Private/Lighting/LightingSystem.cpp
@@ -66,6 +66,8 @@ void FGatheredLightSample::AddWeighted(const FGatheredLightSample& OtherSample,f
 	SHCorrection += OtherSample.SHCorrection * Weight;
 	IncidentLighting += OtherSample.IncidentLighting * Weight;
 	SkyOcclusion += OtherSample.SkyOcclusion * Weight;
+	AOMaterialMask += OtherSample.AOMaterialMask * Weight;
+
 }
 
 void FGatheredLightSample::ApplyOcclusion(float Occlusion)
@@ -129,6 +131,8 @@ FLightSample FGatheredLightMapSample::ConvertToLightSample(bool bDebugThisSample
 	NewSample.SkyOcclusion[1] = HighQuality.SkyOcclusion.Y;
 	NewSample.SkyOcclusion[2] = HighQuality.SkyOcclusion.Z;
 
+	NewSample.AOMaterialMask = HighQuality.AOMaterialMask;
+
 	return NewSample;
 }
 
diff --git a/Engine/Source/Programs/UnrealLightmass/Private/Lighting/LightingSystem.h b/Engine/Source/Programs/UnrealLightmass/Private/Lighting/LightingSystem.h
--- a/Engine/Source/Programs/UnrealLightmass/Private/Lighting/LightingSystem.h
+++ b/Engine/Source/Programs/UnrealLightmass/Private/Lighting/LightingSystem.h
@@ -41,12 +41,15 @@ public:
 	/** Sky bent normal, points toward the most unoccluded direction, and the length is the visibility amount (0 = occluded, 1 = visible). */
 	FVector SkyOcclusion;
 
+	float AOMaterialMask;
+
 	/** Initialization constructor. */
 	FGatheredLightSample()
 	{
 		SHCorrection = 0.0f;
 		IncidentLighting = FLinearColor(0, 0, 0, 0);
 		SkyOcclusion = FVector(0);
+		AOMaterialMask = 0;
 	}
 
 	FGatheredLightSample(EForceInit)
@@ -54,6 +57,7 @@ public:
 		SHCorrection = 0.0f;
 		IncidentLighting = FLinearColor(0, 0, 0, 0);
 		SkyOcclusion = FVector(0);
+		AOMaterialMask = 0;
 	}
 
 	/**
@@ -93,6 +97,7 @@ public:
 		Result.SHCorrection = SHCorrection * Scalar;
 		Result.IncidentLighting = IncidentLighting * Scalar;
 		Result.SkyOcclusion = SkyOcclusion * Scalar;
+		Result.AOMaterialMask = AOMaterialMask * Scalar;
 		return Result;
 	}
 
@@ -103,6 +108,7 @@ public:
 		Result.SHCorrection = SHCorrection + SampleB.SHCorrection;
 		Result.IncidentLighting = IncidentLighting + SampleB.IncidentLighting;
 		Result.SkyOcclusion = SkyOcclusion + SampleB.SkyOcclusion;
+		Result.AOMaterialMask = AOMaterialMask + SampleB.AOMaterialMask;
 		return Result;
 	}
 };
@@ -1336,27 +1342,30 @@ public:
 	FLinearColor Lighting;
 	FVector UnoccludedSkyVector;
 	FLinearColor StationarySkyLighting;
+	float NumSamplesOccluded;
 
 	FLightingAndOcclusion() :
 		Lighting(ForceInit),
 		UnoccludedSkyVector(FVector(0)),
-		StationarySkyLighting(FLinearColor::Black)
+		StationarySkyLighting(FLinearColor::Black),
+		NumSamplesOccluded(0)
 	{}
 
-	FLightingAndOcclusion(const FLinearColor& InLighting, FVector InUnoccludedSkyVector, const FLinearColor& InStationarySkyLighting) : 
+	FLightingAndOcclusion(const FLinearColor& InLighting, FVector InUnoccludedSkyVector, const FLinearColor& InStationarySkyLighting, float InNumSamplesOccluded) : 
 		Lighting(InLighting),
 		UnoccludedSkyVector(InUnoccludedSkyVector),
-		StationarySkyLighting(InStationarySkyLighting)
+		StationarySkyLighting(InStationarySkyLighting),
+		NumSamplesOccluded(InNumSamplesOccluded)
 	{}
 
 	friend inline FLightingAndOcclusion operator+ (const FLightingAndOcclusion& A, const FLightingAndOcclusion& B)
 	{
-		return FLightingAndOcclusion(A.Lighting + B.Lighting, A.UnoccludedSkyVector + B.UnoccludedSkyVector, A.StationarySkyLighting + B.StationarySkyLighting);
+		return FLightingAndOcclusion(A.Lighting + B.Lighting, A.UnoccludedSkyVector + B.UnoccludedSkyVector, A.StationarySkyLighting + B.StationarySkyLighting, A.NumSamplesOccluded + B.NumSamplesOccluded);
 	}
 
 	friend inline FLightingAndOcclusion operator/ (const FLightingAndOcclusion& A, float Divisor)
 	{
-		return FLightingAndOcclusion(A.Lighting / Divisor, A.UnoccludedSkyVector / Divisor, A.StationarySkyLighting / Divisor);
+		return FLightingAndOcclusion(A.Lighting / Divisor, A.UnoccludedSkyVector / Divisor, A.StationarySkyLighting / Divisor, A.NumSamplesOccluded / Divisor);
 	}
 };
 
diff --git a/Engine/Source/Programs/UnrealLightmass/Private/Lighting/LightmapData.cpp b/Engine/Source/Programs/UnrealLightmass/Private/Lighting/LightmapData.cpp
--- a/Engine/Source/Programs/UnrealLightmass/Private/Lighting/LightmapData.cpp
+++ b/Engine/Source/Programs/UnrealLightmass/Private/Lighting/LightmapData.cpp
@@ -276,6 +276,9 @@ namespace Lightmass
 			// Sqrt on length to allocate more precision near 0
 			DestCoefficients.SkyOcclusion[3] = (uint8)FMath::Clamp<int32>( FMath::RoundToInt( FMath::Sqrt(BentNormalLength) * 255.0f ), 0, 255 );
 
+			// Sqrt to allocate more precision near 0
+			DestCoefficients.AOMaterialMask = (uint8)FMath::Clamp<int32>( FMath::RoundToInt( FMath::Sqrt(SourceSample.AOMaterialMask) * 255.0f ), 0, 255 ); 
+
 			{
 				float L, U, V, W;
 				GetLUVW( SourceSample.Coefficients[0], L, U, V, W );
diff --git a/Engine/Source/Programs/UnrealLightmass/Private/Lighting/TextureMapping.cpp b/Engine/Source/Programs/UnrealLightmass/Private/Lighting/TextureMapping.cpp
--- a/Engine/Source/Programs/UnrealLightmass/Private/Lighting/TextureMapping.cpp
+++ b/Engine/Source/Programs/UnrealLightmass/Private/Lighting/TextureMapping.cpp
@@ -492,13 +492,13 @@ void FStaticLightingSystem::ProcessTextureMapping(FStaticLightingTextureMapping*
 					if (DistanceToDebugTexelSq < 40000
 						|| X == Scene.DebugInput.LocalX && Y == Scene.DebugInput.LocalY)
 					{
-						FDebugStaticLightingVertex DebugVertex;
-						DebugVertex.VertexNormal = FVector4(TexelToVertex.WorldTangentZ);
-						DebugVertex.VertexPosition = TexelToVertex.WorldPosition;
-						DebugOutput.Vertices.Add(DebugVertex);
-
 						if (X == Scene.DebugInput.LocalX && Y == Scene.DebugInput.LocalY)
-						{							
+						{		
+							FDebugStaticLightingVertex DebugVertex;
+							DebugVertex.VertexNormal = FVector4(TexelToVertex.WorldTangentZ);
+							DebugVertex.VertexPosition = TexelToVertex.WorldPosition;
+							DebugOutput.Vertices.Add(DebugVertex);
+
 							DebugOutput.SelectedVertexIndices.Add(DebugOutput.Vertices.Num() - 1);
 							DebugOutput.SampleRadius = TexelToVertex.TexelRadius;
 						}
@@ -2969,6 +2969,7 @@ void FStaticLightingSystem::ProcessInterpolateTask(FInterpolateIndirectTaskDescr
 
 				// Apply occlusion to indirect lighting and add this texel's indirect lighting to its running total
 				CurrentLightSample.AddWeighted(IndirectLighting, 1);
+				CurrentLightSample.HighQuality.AOMaterialMask = IndirectLighting.Occlusion;
 
 				// Stationary sky light contribution goes into low quality lightmap only, bent normal sky shadowing will be exported separately
 				CurrentLightSample.LowQuality.AddWeighted(IndirectLighting.StationarySkyLighting, 1);
diff --git a/Engine/Source/Programs/UnrealLightmass/Public/ImportExport.h b/Engine/Source/Programs/UnrealLightmass/Public/ImportExport.h
--- a/Engine/Source/Programs/UnrealLightmass/Public/ImportExport.h
+++ b/Engine/Source/Programs/UnrealLightmass/Public/ImportExport.h
@@ -84,6 +84,7 @@ namespace Lightmass
 			SkyOcclusion[0] = 0;
 			SkyOcclusion[1] = 0;
 			SkyOcclusion[2] = 0;
+			AOMaterialMask = 0;
 		}
 
 		/** 
@@ -95,6 +96,8 @@ namespace Lightmass
 
 		float SkyOcclusion[3];
 
+		float AOMaterialMask;
+
 		/** True if this sample maps to a valid point on a triangle.  This is only meaningful for texture lightmaps. */
 		bool bIsMapped;
 
@@ -121,6 +124,7 @@ namespace Lightmass
 		uint8 Coverage;
 		uint8 Coefficients[LM_NUM_STORED_LIGHTMAP_COEF][4];
 		uint8 SkyOcclusion[4];
+		uint8 AOMaterialMask;
 	};
 
 	struct FLightMapDataBase
diff --git a/Engine/Source/Programs/UnrealLightmass/Public/SceneExport.h b/Engine/Source/Programs/UnrealLightmass/Public/SceneExport.h
--- a/Engine/Source/Programs/UnrealLightmass/Public/SceneExport.h
+++ b/Engine/Source/Programs/UnrealLightmass/Public/SceneExport.h
@@ -213,6 +213,13 @@ public:
 	 */
 	bool bUseAmbientOcclusion;
 
+	/** 
+	 * Whether to generate textures storing the AO computed by Lightmass.
+	 * These can be accessed through the PrecomputedAmbientOcclusion material node, 
+	 * Which is useful for blending between material layers on environment assets.
+	 */
+	bool bGenerateAmbientOcclusionMaterialMask;
+
 	/** Debugging - whether to only show the ambient occlusion term, useful for seeing the impact of AO settings in isolation. */
 	bool bVisualizeAmbientOcclusion;
 
diff --git a/Engine/Source/Runtime/Core/Public/UObject/ObjectVersion.h b/Engine/Source/Runtime/Core/Public/UObject/ObjectVersion.h
--- a/Engine/Source/Runtime/Core/Public/UObject/ObjectVersion.h
+++ b/Engine/Source/Runtime/Core/Public/UObject/ObjectVersion.h
@@ -548,6 +548,9 @@ enum EUnrealEngineObjectUE4Version
 	VER_UE4_STATIC_MESH_EXTENDED_BOUNDS,
 	// Added non-linear blending to anim transitions, deprecating old types
 	VER_UE4_ADDED_NON_LINEAR_TRANSITION_BLENDS,
+	// AO Material Mask texture
+	VER_UE4_AO_MATERIAL_MASK,
+
 	// -----<new versions can be added before this line>-------------------------------------------------
 	// - this needs to be the last line (see note below)
 	VER_UE4_AUTOMATIC_VERSION_PLUS_ONE,
diff --git a/Engine/Source/Runtime/Engine/Classes/GameFramework/WorldSettings.h b/Engine/Source/Runtime/Engine/Classes/GameFramework/WorldSettings.h
--- a/Engine/Source/Runtime/Engine/Classes/GameFramework/WorldSettings.h
+++ b/Engine/Source/Runtime/Engine/Classes/GameFramework/WorldSettings.h
@@ -93,6 +93,14 @@ struct FLightmassWorldInfoSettings
 	UPROPERTY(EditAnywhere, Category=LightmassOcclusion)
 	uint32 bUseAmbientOcclusion:1;
 
+	/** 
+	 * Whether to generate textures storing the AO computed by Lightmass.
+	 * These can be accessed through the PrecomputedAmbientOcclusion material node, 
+	 * Which is useful for blending between material layers on environment assets.
+	 */
+	UPROPERTY(EditAnywhere, Category=LightmassOcclusion)
+	uint32 bGenerateAmbientOcclusionMaterialMask:1;
+
 	/** How much of the AO to apply to direct lighting. */
 	UPROPERTY(EditAnywhere, Category=LightmassOcclusion, meta=(UIMin = "0", UIMax = "1"))
 	float DirectIlluminationOcclusionFraction;
@@ -145,6 +153,7 @@ struct FLightmassWorldInfoSettings
 		, EmissiveBoost(1.0f)
 		, DiffuseBoost(1.0f)
 		, bUseAmbientOcclusion(false)
+		, bGenerateAmbientOcclusionMaterialMask(false)
 		, DirectIlluminationOcclusionFraction(0.5f)
 		, IndirectIlluminationOcclusionFraction(1.0f)
 		, OcclusionExponent(1.0f)
diff --git a/Engine/Source/Runtime/Engine/Classes/Materials/MaterialExpressionPrecomputedAOMask.h b/Engine/Source/Runtime/Engine/Classes/Materials/MaterialExpressionPrecomputedAOMask.h
new file mode 100644
--- /dev/null
+++ b/Engine/Source/Runtime/Engine/Classes/Materials/MaterialExpressionPrecomputedAOMask.h
@@ -0,0 +1,25 @@
+// Copyright 1998-2015 Epic Games, Inc. All Rights Reserved.
+
+/**
+ *
+ *	A material expression that allows access to the precomputed AO mask generated when bGenerateAmbientOcclusionMaterialMask is enabled in Lightmass World settings.
+ */
+
+#pragma once
+#include "Materials/MaterialExpression.h"
+#include "MaterialExpressionPrecomputedAOMask.generated.h"
+
+UCLASS(collapsecategories, hidecategories=Object)
+class UMaterialExpressionPrecomputedAOMask : public UMaterialExpression
+{
+	GENERATED_UCLASS_BODY()
+
+
+	// Begin UMaterialExpression Interface
+	virtual int32 Compile(class FMaterialCompiler* Compiler, int32 OutputIndex, int32 MultiplexIndex) override;
+	virtual void GetCaption(TArray<FString>& OutCaptions) const override;
+	// End UMaterialExpression Interface
+};
+
+
+
diff --git a/Engine/Source/Runtime/Engine/Private/LightMap.cpp b/Engine/Source/Runtime/Engine/Private/LightMap.cpp
--- a/Engine/Source/Runtime/Engine/Private/LightMap.cpp
+++ b/Engine/Source/Runtime/Engine/Private/LightMap.cpp
@@ -344,6 +344,7 @@ enum ELightmapTextureType
 {
 	LTT_Coefficients = NUM_STORED_LIGHTMAP_COEF,
 	LTT_SkyOcclusion,
+	LTT_AOMaterialMask,
 	LTT_Num
 };
 
@@ -384,6 +385,7 @@ struct FLightMapPendingTexture : public FTextureLayout
 	/** Helper data to keep track of the asynchronous tasks for the 4 lightmap textures. */
 	ULightMapTexture2D*				Textures[NUM_STORED_LIGHTMAP_COEF];
 	ULightMapTexture2D*				SkyOcclusionTexture;
+	ULightMapTexture2D*				AOMaterialMaskTexture;
 
 	TArray<FLightMapAllocation*>	Allocations;
 	UObject*						Outer;
@@ -393,7 +395,6 @@ struct FLightMapPendingTexture : public FTextureLayout
 
 	/** Lightmap streaming flags that must match in order to be stored in this texture.	*/
 	ELightMapFlags					LightmapFlags;
-
 	// Optimization to quickly test if a new allocation won't fit
 	// Primarily of benefit to instanced mesh lightmaps
 	int32							UnallocatedTexels;
@@ -403,6 +404,7 @@ struct FLightMapPendingTexture : public FTextureLayout
 	FLightMapPendingTexture(UWorld* InWorld, uint32 InSizeX,uint32 InSizeY)
 		: FTextureLayout(4, 4, InSizeX, InSizeY, /* PowerOfTwo */ true, /* AlignByFour */ true) // Min size is 4x4 in case of block compression.
 		, SkyOcclusionTexture(nullptr)
+		, AOMaterialMaskTexture(nullptr)
 		, OwningWorld(InWorld)
 		, Bounds(FBox(0))
 		, LightmapFlags(LMF_None)
@@ -446,6 +448,7 @@ struct FLightMapPendingTexture : public FTextureLayout
 private:
 	FName GetLightmapName(int32 TextureIndex, int32 CoefficientIndex);
 	FName GetSkyOcclusionTextureName(int32 TextureIndex);
+	FName GetAOMaterialMaskTextureName(int32 TextureIndex);
 };
 
 TArray<FLightMapPendingTexture::FAsyncLightMapCacheTask> FLightMapPendingTexture::TotalAsyncTasks;
@@ -465,7 +468,20 @@ void FLightMapPendingTexture::FinishCompletedTasks( int32 NumUnfinishedTasksAllo
 		{
 			FAsyncLightMapCacheTask Task = TotalAsyncTasks[TaskIndex];
 			FLightMapPendingTexture* PendingTexture = Task.Texture;
-			ULightMapTexture2D* LightMapTexture = Task.TextureType == LTT_SkyOcclusion ? PendingTexture->SkyOcclusionTexture : PendingTexture->Textures[(int32)Task.TextureType];
+			ULightMapTexture2D* LightMapTexture;
+
+			if (Task.TextureType == LTT_SkyOcclusion)
+			{
+				LightMapTexture = PendingTexture->SkyOcclusionTexture;
+			}
+			else if (Task.TextureType == LTT_AOMaterialMask)
+			{
+				LightMapTexture = PendingTexture->AOMaterialMaskTexture;
+			}
+			else
+			{
+				LightMapTexture = PendingTexture->Textures[(int32)Task.TextureType];
+			}
 
 			if (LightMapTexture->IsAsyncCacheComplete())
 			{
@@ -495,7 +511,20 @@ void FLightMapPendingTexture::FinishCompletedTasks( int32 NumUnfinishedTasksAllo
  */
 void FLightMapPendingTexture::FinishEncoding( ELightmapTextureType TextureType )
 {
-	UTexture2D* Texture2D = TextureType == LTT_SkyOcclusion ? SkyOcclusionTexture : Textures[(int32)TextureType];
+	UTexture2D* Texture2D;
+
+	if (TextureType == LTT_SkyOcclusion)
+	{
+		Texture2D = SkyOcclusionTexture;
+	}
+	else if (TextureType == LTT_AOMaterialMask)
+	{
+		Texture2D = AOMaterialMaskTexture;
+	}
+	else
+	{
+		Texture2D = Textures[(int32)TextureType];
+	}
 
 	Texture2D->FinishCachePlatformData();
 	Texture2D->UpdateResource();
@@ -631,7 +660,7 @@ bool FLightMapPendingTexture::AddElement(FLightMapAllocationGroup& AllocationGro
 /** Whether to color each lightmap texture with a different (random) color. */
 bool GVisualizeLightmapTextures = false;
 
-static void GenerateLightmapMipsAndDilate(int32 NumMips, int32 TextureSizeX, int32 TextureSizeY, FColor TextureColor, FColor** MipData, int8** MipCoverageData)
+static void GenerateLightmapMipsAndDilateColor(int32 NumMips, int32 TextureSizeX, int32 TextureSizeY, FColor TextureColor, FColor** MipData, int8** MipCoverageData)
 {
 	for(int32 MipIndex = 1; MipIndex < NumMips; MipIndex++)
 	{
@@ -790,6 +819,165 @@ static void GenerateLightmapMipsAndDilate(int32 NumMips, int32 TextureSizeX, int
 	}
 }
 
+static void GenerateLightmapMipsAndDilateByte(int32 NumMips, int32 TextureSizeX, int32 TextureSizeY, uint8 TextureColor, uint8** MipData, int8** MipCoverageData)
+{
+	for(int32 MipIndex = 1; MipIndex < NumMips; MipIndex++)
+	{
+		const int32 SourceMipSizeX = FMath::Max(1, TextureSizeX >> (MipIndex - 1));
+		const int32 SourceMipSizeY = FMath::Max(1, TextureSizeY >> (MipIndex - 1));
+		const int32 DestMipSizeX = FMath::Max(1, TextureSizeX >> MipIndex);
+		const int32 DestMipSizeY = FMath::Max(1, TextureSizeY >> MipIndex);
+
+		// Downsample the previous mip-level, taking into account which texels are mapped.
+		uint8* NextMipData = MipData[MipIndex];
+		uint8* LastMipData = MipData[MipIndex - 1];
+
+		int8* NextMipCoverageData = MipCoverageData[ MipIndex ];
+		int8* LastMipCoverageData = MipCoverageData[ MipIndex - 1 ];
+
+		const int32 MipFactorX = SourceMipSizeX / DestMipSizeX;
+		const int32 MipFactorY = SourceMipSizeY / DestMipSizeY;
+
+		//@todo - generate mips before encoding lightmaps!  
+		// Currently we are filtering in the encoded space, similar to generating mips of sRGB textures in sRGB space
+		for(int32 Y = 0; Y < DestMipSizeY; Y++)
+		{
+			for(int32 X = 0; X < DestMipSizeX; X++)
+			{
+				float AccumulatedColor = 0;
+				uint32 Coverage = 0;
+
+				const uint32 MinSourceY = (Y + 0) * MipFactorY;
+				const uint32 MaxSourceY = (Y + 1) * MipFactorY;
+				for(uint32 SourceY = MinSourceY; SourceY < MaxSourceY; SourceY++)
+				{
+					const uint32 MinSourceX = (X + 0) * MipFactorX;
+					const uint32 MaxSourceX = (X + 1) * MipFactorX;
+					for(uint32 SourceX = MinSourceX; SourceX < MaxSourceX; SourceX++)
+					{
+						const uint8& SourceColor = LastMipData[SourceY * SourceMipSizeX + SourceX];
+						int8 SourceCoverage = LastMipCoverageData[ SourceY * SourceMipSizeX + SourceX ];
+						if( SourceCoverage )
+						{
+							AccumulatedColor += SourceColor / 255.0f * SourceCoverage;
+							Coverage += SourceCoverage;
+						}
+					}
+				}
+				uint8& DestColor = NextMipData[Y * DestMipSizeX + X];
+				int8& DestCoverage = NextMipCoverageData[Y * DestMipSizeX + X];
+				if ( GVisualizeLightmapTextures )
+				{
+					DestColor = TextureColor;
+					DestCoverage = 127;
+				}
+				else if(Coverage)
+				{
+					DestColor = (uint8)FMath::Clamp<int32>(FMath::TruncToInt(AccumulatedColor / Coverage * 255.f), 0, 255);
+					DestCoverage = Coverage / (MipFactorX * MipFactorY);
+				}
+				else
+				{
+					DestColor = 0;
+					DestCoverage = 0;
+				}
+			}
+		}
+	}
+
+	// Expand texels which are mapped into adjacent texels which are not mapped to avoid artifacts when using texture filtering.
+	for(int32 MipIndex = 0; MipIndex < NumMips; MipIndex++)
+	{
+		uint8* MipLevelData = MipData[MipIndex];
+		int8* MipLevelCoverageData = MipCoverageData[MipIndex];
+
+		uint32 MipSizeX = FMath::Max(1,TextureSizeX >> MipIndex);
+		uint32 MipSizeY = FMath::Max(1,TextureSizeY >> MipIndex);
+		for(uint32 DestY = 0;DestY < MipSizeY;DestY++)
+		{
+			for(uint32 DestX = 0; DestX < MipSizeX; DestX++)
+			{
+				uint8& DestColor = MipLevelData[DestY * MipSizeX + DestX];
+				int8& DestCoverage = MipLevelCoverageData[DestY * MipSizeX + DestX];
+				if(DestCoverage == 0)
+				{
+					float AccumulatedColor = 0;
+					uint32 Coverage = 0;
+
+					const int32 MinSourceY = FMath::Max((int32)DestY - 1, (int32)0);
+					const int32 MaxSourceY = FMath::Min((int32)DestY + 1, (int32)MipSizeY - 1);
+					for(int32 SourceY = MinSourceY; SourceY <= MaxSourceY; SourceY++)
+					{
+						const int32 MinSourceX = FMath::Max((int32)DestX - 1, (int32)0);
+						const int32 MaxSourceX = FMath::Min((int32)DestX + 1, (int32)MipSizeX - 1);
+						for(int32 SourceX = MinSourceX; SourceX <= MaxSourceX; SourceX++)
+						{
+							uint8& SourceColor = MipLevelData[SourceY * MipSizeX + SourceX];
+							int8 SourceCoverage = MipLevelCoverageData[SourceY * MipSizeX + SourceX];
+							if( SourceCoverage > 0 )
+							{
+								static const uint32 Weights[3][3] =
+								{
+									{ 1, 255, 1 },
+									{ 255, 0, 255 },
+									{ 1, 255, 1 },
+								};
+								AccumulatedColor += SourceColor / 255.0f * SourceCoverage * Weights[SourceX - DestX + 1][SourceY - DestY + 1];
+								Coverage += SourceCoverage * Weights[SourceX - DestX + 1][SourceY - DestY + 1];
+							}
+						}
+					}
+
+					if(Coverage)
+					{
+						DestColor = (uint8)FMath::Clamp<int32>(FMath::TruncToInt(AccumulatedColor / Coverage * 255.f), 0, 255);
+						DestCoverage = -1;
+					}
+				}
+			}
+		}
+	}
+
+	// Fill zero coverage texels with closest colors using mips
+	for(int32 MipIndex = NumMips - 2; MipIndex >= 0; MipIndex--)
+	{
+		const int32 DstMipSizeX = FMath::Max(1, TextureSizeX >> MipIndex);
+		const int32 DstMipSizeY = FMath::Max(1, TextureSizeY >> MipIndex);
+		const int32 SrcMipSizeX = FMath::Max(1, TextureSizeX >> (MipIndex + 1));
+		const int32 SrcMipSizeY = FMath::Max(1, TextureSizeY >> (MipIndex + 1));
+
+		// Source from higher mip, taking into account which texels are mapped.
+		uint8* DstMipData = MipData[MipIndex];
+		uint8* SrcMipData = MipData[MipIndex + 1];
+
+		int8* DstMipCoverageData = MipCoverageData[ MipIndex ];
+		int8* SrcMipCoverageData = MipCoverageData[ MipIndex + 1 ];
+
+		for(int32 DstY = 0; DstY < DstMipSizeY; DstY++)
+		{
+			for(int32 DstX = 0; DstX < DstMipSizeX; DstX++)
+			{
+				const uint32 SrcX = DstX / 2;
+				const uint32 SrcY = DstY / 2;
+
+				const uint8& SrcColor = SrcMipData[ SrcY * SrcMipSizeX + SrcX ];
+				int8 SrcCoverage = SrcMipCoverageData[ SrcY * SrcMipSizeX + SrcX ];
+
+				uint8& DstColor = DstMipData[ DstY * DstMipSizeX + DstX ];
+				int8& DstCoverage = DstMipCoverageData[ DstY * DstMipSizeX + DstX ];
+
+				// Point upsample mip data for zero coverage texels
+				// TODO bilinear upsample
+				if( SrcCoverage != 0 && DstCoverage == 0 )
+				{
+					DstColor = SrcColor;
+					DstCoverage = SrcCoverage;
+				}
+			}
+		}
+	}
+}
+
 void FLightMapPendingTexture::StartEncoding()
 {
 	GLightmapCounter++;
@@ -809,7 +997,6 @@ void FLightMapPendingTexture::StartEncoding()
 		if (Allocation->bHasSkyShadowing)
 		{
 			bNeedsSkyOcclusionTexture = true;
-			break;
 		}
 	}
 
@@ -889,7 +1076,7 @@ void FLightMapPendingTexture::StartEncoding()
 			}
 		}
 
-		GenerateLightmapMipsAndDilate(NumMips, TextureSizeX, TextureSizeY, TextureColor, MipData, MipCoverageData);
+		GenerateLightmapMipsAndDilateColor(NumMips, TextureSizeX, TextureSizeY, TextureColor, MipData, MipCoverageData);
 
 		// Unlock all mip levels.
 		for (int32 MipIndex = 0; MipIndex < NumMips; ++MipIndex)
@@ -903,6 +1090,96 @@ void FLightMapPendingTexture::StartEncoding()
 		NumOutstandingAsyncTasks++;
 	}
 	
+	const FLightmassWorldInfoSettings* LightmassWorldSettings = OwningWorld.IsValid() ? &(OwningWorld->GetWorldSettings()->LightmassSettings) : NULL;
+
+	if (LightmassWorldSettings && LightmassWorldSettings->bUseAmbientOcclusion && LightmassWorldSettings->bGenerateAmbientOcclusionMaterialMask)
+	{
+		auto Texture = NewObject<ULightMapTexture2D>(Outer, GetAOMaterialMaskTextureName(GLightmapCounter));
+		AOMaterialMaskTexture = Texture;
+
+		Texture->Source.Init2DWithMipChain(GetSizeX(), GetSizeY(), TSF_G8);
+		Texture->MipGenSettings = TMGS_LeaveExistingMips;
+		int32 NumMips = Texture->Source.GetNumMips();
+		Texture->SRGB = false;
+		Texture->Filter	= GUseBilinearLightmaps ? TF_Default : TF_Nearest;
+		Texture->LODGroup = TEXTUREGROUP_Lightmap;
+		Texture->LightmapFlags = ELightMapFlags( LightmapFlags );
+		Texture->CompressionNoAlpha = false;
+		Texture->CompressionNone = !GCompressLightmaps;
+		// BC4
+		Texture->CompressionSettings = TC_Alpha;
+
+		int32 TextureSizeX = Texture->Source.GetSizeX();
+		int32 TextureSizeY = Texture->Source.GetSizeY();
+
+		int32 StartBottom = GetSizeX() * GetSizeY();
+
+		// Lock all mip levels.
+		uint8* MipData[MAX_TEXTURE_MIP_COUNT] = {0};
+		int8* MipCoverageData[MAX_TEXTURE_MIP_COUNT] = {0};
+		for (int32 MipIndex = 0; MipIndex < NumMips; ++MipIndex)
+		{
+			MipData[MipIndex] = (uint8*)Texture->Source.LockMip(MipIndex);
+
+			const int32 MipSizeX = FMath::Max( 1, TextureSizeX >> MipIndex );
+			const int32 MipSizeY = FMath::Max( 1, TextureSizeY >> MipIndex );
+			MipCoverageData[ MipIndex ] = (int8*)FMemory::Malloc( MipSizeX * MipSizeY );
+		}
+
+		// Create the uncompressed top mip-level.
+		uint8* TopMipData = MipData[0];
+		FMemory::Memzero( TopMipData, TextureSizeX * TextureSizeY * sizeof(uint8) );
+		FMemory::Memzero( MipCoverageData[0], TextureSizeX * TextureSizeY );
+
+		FIntRect TextureRect( MAX_int32, MAX_int32, MIN_int32, MIN_int32 );
+		for(int32 AllocationIndex = 0;AllocationIndex < Allocations.Num();AllocationIndex++)
+		{
+			FLightMapAllocation* Allocation = Allocations[AllocationIndex];
+			// Link the light-map to the texture.
+			Allocation->LightMap->AOMaterialMaskTexture = Texture;
+			
+			// Skip encoding of this texture if we were asked not to bother
+			if( !Allocation->bSkipEncoding )
+			{
+				TextureRect.Min.X = FMath::Min<int32>( TextureRect.Min.X, Allocation->OffsetX );
+				TextureRect.Min.Y = FMath::Min<int32>( TextureRect.Min.Y, Allocation->OffsetY );
+				TextureRect.Max.X = FMath::Max<int32>( TextureRect.Max.X, Allocation->OffsetX + Allocation->MappedRect.Width() );
+				TextureRect.Max.Y = FMath::Max<int32>( TextureRect.Max.Y, Allocation->OffsetY + Allocation->MappedRect.Height() );
+
+				// Copy the raw data for this light-map into the raw texture data array.
+				for(int32 Y = Allocation->MappedRect.Min.Y; Y < Allocation->MappedRect.Max.Y; ++Y)
+				{
+					for(int32 X = Allocation->MappedRect.Min.X; X < Allocation->MappedRect.Max.X; ++X)
+					{
+						const FLightMapCoefficients& SourceCoefficients = Allocation->RawData[Y * Allocation->TotalSizeX + X];
+
+						int32 DestY = Y - Allocation->MappedRect.Min.Y + Allocation->OffsetY;
+						int32 DestX = X - Allocation->MappedRect.Min.X + Allocation->OffsetX;
+
+						uint8& DestValue = TopMipData[DestY * TextureSizeX + DestX];
+
+						DestValue = SourceCoefficients.AOMaterialMask;
+
+						int8& DestCoverage = MipCoverageData[0][ DestY * TextureSizeX + DestX ];
+						DestCoverage = SourceCoefficients.Coverage / 2;
+					}
+				}
+			}
+		}
+
+		GenerateLightmapMipsAndDilateByte(NumMips, TextureSizeX, TextureSizeY, TextureColor.R, MipData, MipCoverageData);
+
+		// Unlock all mip levels.
+		for (int32 MipIndex = 0; MipIndex < NumMips; ++MipIndex)
+		{
+			Texture->Source.UnlockMip(MipIndex);
+			FMemory::Free( MipCoverageData[ MipIndex ] );
+		}
+
+		Texture->BeginCachePlatformData();
+		new(TotalAsyncTasks) FAsyncLightMapCacheTask(this, LTT_AOMaterialMask);
+		NumOutstandingAsyncTasks++;
+	}
 
 	// Encode and compress the coefficient textures.
 	for(uint32 CoefficientIndex = 0; CoefficientIndex < NUM_STORED_LIGHTMAP_COEF; CoefficientIndex += 2)
@@ -1066,7 +1343,7 @@ void FLightMapPendingTexture::StartEncoding()
 		GNumLightmapTotalTexelsNonPow2 += TextureRect.Width() * TextureRect.Height();
 		GNumLightmapTextures++;
 
-		GenerateLightmapMipsAndDilate(NumMips, TextureSizeX, TextureSizeY, TextureColor, MipData, MipCoverageData);
+		GenerateLightmapMipsAndDilateColor(NumMips, TextureSizeX, TextureSizeY, TextureColor, MipData, MipCoverageData);
 
 		// Unlock all mip levels.
 		for (int32 MipIndex = 0; MipIndex < NumMips; ++MipIndex)
@@ -1155,6 +1432,23 @@ FName FLightMapPendingTexture::GetSkyOcclusionTextureName(int32 TextureIndex)
 	return FName(*PotentialName);
 }
 
+FName FLightMapPendingTexture::GetAOMaterialMaskTextureName(int32 TextureIndex)
+{
+	FString PotentialName = TEXT("");
+	UObject* ExistingObject = NULL;
+	int32 LightmapIndex = 0;
+	// Search for an unused name
+	do
+	{
+		PotentialName = FString(TEXT("AOMaterialMask")) + FString::FromInt(LightmapIndex) + TEXT("_") + FString::FromInt(TextureIndex);
+
+		ExistingObject = FindObject<UObject>(Outer, *PotentialName);
+		LightmapIndex++;
+	}
+	while (ExistingObject != NULL);
+	return FName(*PotentialName);
+}
+
 /** The light-maps which have not yet been encoded into textures. */
 static TArray<FLightMapAllocationGroup> PendingLightMaps;
 static uint32 PendingLightMapSize = 0;
@@ -1569,6 +1863,7 @@ FLightMap2D::FLightMap2D()
 	Textures[0] = NULL;
 	Textures[1] = NULL;
 	SkyOcclusionTexture = NULL;
+	AOMaterialMaskTexture = NULL;
 }
 
 FLightMap2D::FLightMap2D(const TArray<FGuid>& InLightGuids)
@@ -1577,6 +1872,7 @@ FLightMap2D::FLightMap2D(const TArray<FGuid>& InLightGuids)
 	Textures[0] = NULL;
 	Textures[1] = NULL;
 	SkyOcclusionTexture = NULL;
+	AOMaterialMaskTexture = NULL;
 }
 
 const UTexture2D* FLightMap2D::GetTexture(uint32 BasisIndex) const
@@ -1596,6 +1892,11 @@ UTexture2D* FLightMap2D::GetSkyOcclusionTexture()
 	return SkyOcclusionTexture;
 }
 
+UTexture2D* FLightMap2D::GetAOMaterialMaskTexture()
+{
+	return AOMaterialMaskTexture;
+}
+
 /**
  * Returns whether the specified basis has a valid lightmap texture or not.
  * @param	BasisIndex - The basis index.
@@ -1623,6 +1924,7 @@ void FLightMap2D::AddReferencedObjects( FReferenceCollector& Collector )
 	Collector.AddReferencedObject(Textures[0]);
 	Collector.AddReferencedObject(Textures[1]);
 	Collector.AddReferencedObject(SkyOcclusionTexture);
+	Collector.AddReferencedObject(AOMaterialMaskTexture);
 }
 
 void FLightMap2D::Serialize(FArchive& Ar)
@@ -1676,10 +1978,20 @@ void FLightMap2D::Serialize(FArchive& Ar)
 
 				ULightMapTexture2D* Dummy = NULL;
 				Ar << (bStripHQLightmaps ? Dummy : SkyOcclusionTexture);
+
+				if (Ar.UE4Ver() >= VER_UE4_AO_MATERIAL_MASK)
+				{
+					Ar << (bStripHQLightmaps ? Dummy : AOMaterialMaskTexture);
+				}
 			}
 			else
 			{
 				Ar << SkyOcclusionTexture;
+
+				if (Ar.UE4Ver() >= VER_UE4_AO_MATERIAL_MASK)
+				{
+					Ar << AOMaterialMaskTexture;
+				}
 			}
 		}
 		
@@ -1711,6 +2023,7 @@ void FLightMap2D::Serialize(FArchive& Ar)
 		if (!bAllowHighQualityLightMaps)
 		{
 			SkyOcclusionTexture = NULL;
+			AOMaterialMaskTexture = NULL;
 		}
 	}
 }
@@ -1726,7 +2039,7 @@ FLightMapInteraction FLightMap2D::GetInteraction(ERHIFeatureLevel::Type InFeatur
 	// When the FLightMap2D is first created, the textures aren't set, so that case needs to be handled.
 	if(bValidTextures)
 	{
-		return FLightMapInteraction::Texture(Textures, SkyOcclusionTexture, ScaleVectors, AddVectors, CoordinateScale, CoordinateBias, bHighQuality);
+		return FLightMapInteraction::Texture(Textures, SkyOcclusionTexture, AOMaterialMaskTexture, ScaleVectors, AddVectors, CoordinateScale, CoordinateBias, bHighQuality);
 	}
 
 	return FLightMapInteraction::None();
@@ -1894,6 +2207,11 @@ bool FQuantizedLightmapData::HasNonZeroData() const
 					}
 				}
 			}
+
+			if (LightmapSample.AOMaterialMask != 0)
+			{
+				return true;
+			}
 		}
 	}
 
diff --git a/Engine/Source/Runtime/Engine/Private/Materials/HLSLMaterialTranslator.h b/Engine/Source/Runtime/Engine/Private/Materials/HLSLMaterialTranslator.h
--- a/Engine/Source/Runtime/Engine/Private/Materials/HLSLMaterialTranslator.h
+++ b/Engine/Source/Runtime/Engine/Private/Materials/HLSLMaterialTranslator.h
@@ -181,6 +181,8 @@ protected:
 	uint32 bUsesParticleSubUVs : 1;
 	/** Boolean indicating using LightmapUvs */
 	uint32 bUsesLightmapUVs : 1;
+	/** Whether the material uses AO Material Mask */
+	uint32 bUsesAOMaterialMask : 1;
 	/** true if needs SpeedTree code */
 	uint32 bUsesSpeedTree : 1;
 	/** Boolean indicating the material uses worldspace position without shader offsets applied */
@@ -235,6 +237,7 @@ public:
 	,	bUsesSphericalParticleOpacity(false)
 	,   bUsesParticleSubUVs(false)
 	,	bUsesLightmapUVs(false)
+	,	bUsesAOMaterialMask(false)
 	,	bUsesSpeedTree(false)
 	,	bNeedsWorldPositionExcludingShaderOffsets(false)
 	,	bNeedsParticleSize(false)
@@ -597,6 +600,11 @@ public:
 			OutEnvironment.SetDefine(TEXT("LIGHTMAP_UV_ACCESS"),TEXT("1"));
 		}
 
+		if (bUsesAOMaterialMask)
+		{
+			OutEnvironment.SetDefine(TEXT("USES_AO_MATERIAL_MASK"),TEXT("1"));
+		}
+
 		if (bUsesSpeedTree)
 		{
 			OutEnvironment.SetDefine(TEXT("USES_SPEEDTREE"),TEXT("1"));
@@ -3400,6 +3408,29 @@ protected:
 		return ResultIdx;
 	}
 
+	virtual int32 PrecomputedAOMask() override
+	{
+		if (ShaderFrequency != SF_Pixel && ShaderFrequency != SF_Compute)
+		{
+			return NonPixelShaderExpressionError();
+		}
+
+		if (ErrorUnlessFeatureLevelSupported(ERHIFeatureLevel::SM4) == INDEX_NONE)
+		{
+			return INDEX_NONE;
+		}
+
+		bUsesAOMaterialMask = true;
+
+		int32 ResultIdx = INDEX_NONE;
+		FString CodeChunk = FString::Printf(TEXT("Parameters.AOMaterialMask"));
+		ResultIdx = AddCodeChunk(
+			MCT_Float,
+			*CodeChunk
+			);
+		return ResultIdx;
+	}
+
 	virtual int32 LightmassReplace(int32 Realtime, int32 Lightmass) override { return Realtime; }
 
 	virtual int32 GIReplace(int32 Direct, int32 StaticIndirect, int32 DynamicIndirect) override 
diff --git a/Engine/Source/Runtime/Engine/Private/Materials/MaterialExpressions.cpp b/Engine/Source/Runtime/Engine/Private/Materials/MaterialExpressions.cpp
--- a/Engine/Source/Runtime/Engine/Private/Materials/MaterialExpressions.cpp
+++ b/Engine/Source/Runtime/Engine/Private/Materials/MaterialExpressions.cpp
@@ -55,6 +55,7 @@
 #include "Materials/MaterialExpressionGIReplace.h"
 #include "Materials/MaterialExpressionIf.h"
 #include "Materials/MaterialExpressionLightmapUVs.h"
+#include "Materials/MaterialExpressionPrecomputedAOMask.h"
 #include "Materials/MaterialExpressionLightmassReplace.h"
 #include "Materials/MaterialExpressionLightVector.h"
 #include "Materials/MaterialExpressionLinearInterpolate.h"
@@ -7949,6 +7950,47 @@ void UMaterialExpressionLightmapUVs::GetCaption(TArray<FString>& OutCaptions) co
 }
 
 
+//
+//	UMaterialExpressionAOMaterialMask
+//
+UMaterialExpressionPrecomputedAOMask::UMaterialExpressionPrecomputedAOMask(const FObjectInitializer& ObjectInitializer)
+	: Super(ObjectInitializer)
+{
+	// Structure to hold one-time initialization
+	struct FConstructorStatics
+	{
+		FText NAME_Constants;
+		FConstructorStatics()
+			: NAME_Constants(LOCTEXT( "Constants", "Constants" ))
+		{
+		}
+	};
+	static FConstructorStatics ConstructorStatics;
+
+	bShowOutputNameOnPin = true;
+	bHidePreviewWindow = true;
+
+	MenuCategories.Add(ConstructorStatics.NAME_Constants);
+
+	Outputs.Reset();
+	Outputs.Add(FExpressionOutput(TEXT("")));
+
+	bShaderInputData = true;
+}
+
+int32 UMaterialExpressionPrecomputedAOMask::Compile( FMaterialCompiler* Compiler, int32 OutputIndex, int32 MultiplexIndex )
+{
+	return Compiler->PrecomputedAOMask();
+}
+
+	
+void UMaterialExpressionPrecomputedAOMask::GetCaption(TArray<FString>& OutCaptions) const
+{
+	OutCaptions.Add(TEXT("PrecomputedAOMask"));
+}
+
+
+
 //
 //	UMaterialExpressionLightmassReplace
 //
diff --git a/Engine/Source/Runtime/Engine/Private/ModelLight.cpp b/Engine/Source/Runtime/Engine/Private/ModelLight.cpp
--- a/Engine/Source/Runtime/Engine/Private/ModelLight.cpp
+++ b/Engine/Source/Runtime/Engine/Private/ModelLight.cpp
@@ -1149,6 +1149,8 @@ void UModel::ApplyStaticLighting()
 					{
 						DestSample.SkyOcclusion[ColorIndex] = SourceSample.SkyOcclusion[ColorIndex];
 					}
+
+					DestSample.AOMaterialMask = SourceSample.AOMaterialMask;
 				}
 			}
 
diff --git a/Engine/Source/Runtime/Engine/Private/SceneManagement.cpp b/Engine/Source/Runtime/Engine/Private/SceneManagement.cpp
--- a/Engine/Source/Runtime/Engine/Private/SceneManagement.cpp
+++ b/Engine/Source/Runtime/Engine/Private/SceneManagement.cpp
@@ -244,6 +244,7 @@ void FMeshElementCollector::AddMesh(int32 ViewIndex, FMeshBatch& MeshBatch)
 FLightMapInteraction FLightMapInteraction::Texture(
 	const class ULightMapTexture2D* const* InTextures,
 	const ULightMapTexture2D* InSkyOcclusionTexture,
+	const ULightMapTexture2D* InAOMaterialMaskTexture,
 	const FVector4* InCoefficientScales,
 	const FVector4* InCoefficientAdds,
 	const FVector2D& InCoordinateScale,
@@ -273,6 +274,7 @@ FLightMapInteraction FLightMapInteraction::Texture(
 #if ALLOW_HQ_LIGHTMAPS
 		Result.HighQualityTexture = InTextures[0];
 		Result.SkyOcclusionTexture = InSkyOcclusionTexture;
+		Result.AOMaterialMaskTexture = InAOMaterialMaskTexture;
 		for(uint32 CoefficientIndex = 0;CoefficientIndex < NUM_HQ_LIGHTMAP_COEF;CoefficientIndex++)
 		{
 			Result.HighQualityCoefficientScales[CoefficientIndex] = InCoefficientScales[CoefficientIndex];
diff --git a/Engine/Source/Runtime/Engine/Public/LightMap.h b/Engine/Source/Runtime/Engine/Public/LightMap.h
--- a/Engine/Source/Runtime/Engine/Public/LightMap.h
+++ b/Engine/Source/Runtime/Engine/Public/LightMap.h
@@ -217,6 +217,8 @@ public:
 	 */
 	ENGINE_API UTexture2D* GetSkyOcclusionTexture();
 
+	ENGINE_API UTexture2D* GetAOMaterialMaskTexture();
+
 	/**
 	 * Returns whether the specified basis has a valid lightmap texture or not.
 	 * @param	BasisIndex - The basis index.
@@ -292,6 +294,8 @@ protected:
 
 	ULightMapTexture2D* SkyOcclusionTexture;
 
+	ULightMapTexture2D* AOMaterialMaskTexture;
+
 	/** A scale to apply to the coefficients. */
 	FVector4 ScaleVectors[NUM_STORED_LIGHTMAP_COEF];
 
@@ -470,13 +474,18 @@ struct FLightMapCoefficients
 	uint8 Coverage;
 	uint8 Coefficients[NUM_STORED_LIGHTMAP_COEF][4];
 	uint8 SkyOcclusion[4];
+	uint8 AOMaterialMask;
 
 	/** Equality operator */
 	bool operator==( const FLightMapCoefficients& RHS ) const
 	{
 		return Coverage == RHS.Coverage &&
 			   Coefficients == RHS.Coefficients &&
-			   SkyOcclusion == RHS.SkyOcclusion;
+			   SkyOcclusion[0] == RHS.SkyOcclusion[0] &&
+			   SkyOcclusion[1] == RHS.SkyOcclusion[1] &&
+			   SkyOcclusion[2] == RHS.SkyOcclusion[2] &&
+			   SkyOcclusion[3] == RHS.SkyOcclusion[3] &&
+			   AOMaterialMask == RHS.AOMaterialMask;
 	}
 };
 
diff --git a/Engine/Source/Runtime/Engine/Public/MaterialCompiler.h b/Engine/Source/Runtime/Engine/Public/MaterialCompiler.h
--- a/Engine/Source/Runtime/Engine/Public/MaterialCompiler.h
+++ b/Engine/Source/Runtime/Engine/Public/MaterialCompiler.h
@@ -153,6 +153,7 @@ public:
 
 	virtual int32 DynamicParameter(FLinearColor& DefaultValue) = 0;
 	virtual int32 LightmapUVs() = 0;
+	virtual int32 PrecomputedAOMask()  = 0;
 
 	virtual int32 LightmassReplace(int32 Realtime, int32 Lightmass) = 0;
 	virtual int32 GIReplace(int32 Direct, int32 StaticIndirect, int32 DynamicIndirect) = 0;
@@ -317,6 +318,7 @@ public:
 
 	virtual int32 DynamicParameter(FLinearColor& DefaultValue) override { return Compiler->DynamicParameter(DefaultValue); }
 	virtual int32 LightmapUVs() override { return Compiler->LightmapUVs(); }
+	virtual int32 PrecomputedAOMask() override { return Compiler->PrecomputedAOMask(); }
 
 	virtual int32 LightmassReplace(int32 Realtime, int32 Lightmass) override { return Realtime; }
 	virtual int32 GIReplace(int32 Direct, int32 StaticIndirect, int32 DynamicIndirect) override { return Compiler->GIReplace(Direct, StaticIndirect, DynamicIndirect); }
diff --git a/Engine/Source/Runtime/Engine/Public/SceneManagement.h b/Engine/Source/Runtime/Engine/Public/SceneManagement.h
--- a/Engine/Source/Runtime/Engine/Public/SceneManagement.h
+++ b/Engine/Source/Runtime/Engine/Public/SceneManagement.h
@@ -250,6 +250,7 @@ public:
 	static FLightMapInteraction Texture(
 		const class ULightMapTexture2D* const* InTextures,
 		const ULightMapTexture2D* InSkyOcclusionTexture,
+		const ULightMapTexture2D* InAOMaterialMaskTexture,
 		const FVector4* InCoefficientScales,
 		const FVector4* InCoefficientAdds,
 		const FVector2D& InCoordinateScale,
@@ -259,6 +260,7 @@ public:
 	/** Default constructor. */
 	FLightMapInteraction():
 		SkyOcclusionTexture(NULL),
+		AOMaterialMaskTexture(NULL),
 		Type(LMIT_None)
 	{}
 
@@ -287,6 +289,16 @@ public:
 #endif
 	}
 
+	const ULightMapTexture2D* GetAOMaterialMaskTexture() const
+	{
+		check(Type == LMIT_Texture);
+#if ALLOW_HQ_LIGHTMAPS
+		return AOMaterialMaskTexture;
+#else
+		return NULL;
+#endif
+	}
+
 	const FVector4* GetScaleArray() const
 	{
 #if ALLOW_LQ_LIGHTMAPS && ALLOW_HQ_LIGHTMAPS
@@ -387,6 +399,7 @@ private:
 	FVector4 HighQualityCoefficientAdds[NUM_HQ_LIGHTMAP_COEF];
 	const class ULightMapTexture2D* HighQualityTexture;
 	const ULightMapTexture2D* SkyOcclusionTexture;
+	const ULightMapTexture2D* AOMaterialMaskTexture;
 #endif
 
 #if ALLOW_LQ_LIGHTMAPS
diff --git a/Engine/Source/Runtime/Renderer/Private/LightMapRendering.h b/Engine/Source/Runtime/Renderer/Private/LightMapRendering.h
--- a/Engine/Source/Runtime/Renderer/Private/LightMapRendering.h
+++ b/Engine/Source/Runtime/Renderer/Private/LightMapRendering.h
@@ -137,6 +137,8 @@ public:
 			LightMapSamplerParameter.Bind(ParameterMap,TEXT("LightMapSampler"));
 			SkyOcclusionTexture.Bind(ParameterMap,TEXT("SkyOcclusionTexture"));
 			SkyOcclusionSampler.Bind(ParameterMap,TEXT("SkyOcclusionSampler"));
+			AOMaterialMaskTexture.Bind(ParameterMap,TEXT("AOMaterialMaskTexture"));
+			AOMaterialMaskSampler.Bind(ParameterMap,TEXT("AOMaterialMaskSampler"));
 			LightMapScaleParameter.Bind(ParameterMap,TEXT("LightMapScale"));
 			LightMapAddParameter.Bind(ParameterMap,TEXT("LightMapAdd"));
 		}
@@ -177,6 +179,24 @@ public:
 				);
 		}
 
+		void SetAOMaterialMaskTexture(FRHICommandList& RHICmdList, FShader* PixelShader, const UTexture2D* AOMaterialMaskTextureValue) const
+		{
+			FTexture* TextureResource = GBlackTexture;
+
+			if (AOMaterialMaskTextureValue)
+			{
+				TextureResource = AOMaterialMaskTextureValue->Resource;
+			}
+
+			SetTextureParameter(
+				RHICmdList, 
+				PixelShader->GetPixelShader(),
+				AOMaterialMaskTexture,
+				AOMaterialMaskSampler,
+				TextureResource
+				);
+		}
+
 		void SetLightMapScale(FRHICommandList& RHICmdList, FShader* PixelShader,const FLightMapInteraction& LightMapInteraction) const
 		{
 			const FPixelShaderRHIParamRef ShaderRHI = PixelShader->GetPixelShader();
@@ -191,6 +211,8 @@ public:
 			Ar << LightMapSamplerParameter;
 			Ar << SkyOcclusionTexture;
 			Ar << SkyOcclusionSampler;
+			Ar << AOMaterialMaskTexture;
+			Ar << AOMaterialMaskSampler;
 			Ar << LightMapScaleParameter;
 			Ar << LightMapAddParameter;
 		}
@@ -200,6 +222,8 @@ public:
 		FShaderResourceParameter LightMapSamplerParameter;
 		FShaderResourceParameter SkyOcclusionTexture;
 		FShaderResourceParameter SkyOcclusionSampler;
+		FShaderResourceParameter AOMaterialMaskTexture;
+		FShaderResourceParameter AOMaterialMaskSampler;
 		FShaderParameter LightMapScaleParameter;
 		FShaderParameter LightMapAddParameter;
 	};
@@ -273,6 +297,7 @@ public:
 			PixelShaderParameters->SetLightMapScale(RHICmdList, PixelShader,LightMapInteraction);
 			PixelShaderParameters->SetLightMapTexture(RHICmdList, PixelShader, LightMapInteraction.GetTexture(AllowHighQualityLightmaps(View.GetFeatureLevel())));
 			PixelShaderParameters->SetSkyOcclusionTexture(RHICmdList, PixelShader, LightMapInteraction.GetSkyOcclusionTexture());
+			PixelShaderParameters->SetAOMaterialMaskTexture(RHICmdList, PixelShader, LightMapInteraction.GetAOMaterialMaskTexture());
 		}
 	}
 
