Introduce sleep and wake events that user can hook into

#UE-15346

diff --git a/Engine/Source/Runtime/Engine/Classes/Components/PrimitiveComponent.h b/Engine/Source/Runtime/Engine/Classes/Components/PrimitiveComponent.h
--- a/Engine/Source/Runtime/Engine/Classes/Components/PrimitiveComponent.h
+++ b/Engine/Source/Runtime/Engine/Classes/Components/PrimitiveComponent.h
@@ -96,6 +96,10 @@ DECLARE_DYNAMIC_MULTICAST_DELEGATE_FourParams( FComponentHitSignature, class AAc
 DECLARE_DYNAMIC_MULTICAST_DELEGATE_FiveParams( FComponentBeginOverlapSignature,class AActor*, OtherActor, class UPrimitiveComponent*, OtherComp, int32, OtherBodyIndex, bool, bFromSweep, const FHitResult &, SweepResult);
 /** Delegate for notification of end of overlap with a specific component */
 DECLARE_DYNAMIC_MULTICAST_DELEGATE_ThreeParams( FComponentEndOverlapSignature, class AActor*, OtherActor, class UPrimitiveComponent*, OtherComp, int32, OtherBodyIndex);
+/** Delegate for notification when a wake event is fired by physics*/
+DECLARE_DYNAMIC_MULTICAST_DELEGATE_OneParam(FComponentWakeSignature, FName, BoneName);
+/** Delegate for notification when a sleep event is fired by physics*/
+DECLARE_DYNAMIC_MULTICAST_DELEGATE_OneParam(FComponentSleepSignature, FName, BoneName);
 
 DECLARE_DYNAMIC_MULTICAST_DELEGATE_OneParam( FComponentBeginCursorOverSignature, UPrimitiveComponent*, TouchedComponent );
 DECLARE_DYNAMIC_MULTICAST_DELEGATE_OneParam( FComponentEndCursorOverSignature, UPrimitiveComponent*, TouchedComponent );
@@ -650,6 +654,18 @@ public:
 	UPROPERTY(BlueprintAssignable, Category="Collision")
 	FComponentEndOverlapSignature OnComponentEndOverlap;
 
+	/** 
+	 *	Event called when the underlying physics objects is woken up
+	 */
+	UPROPERTY(BlueprintAssignable, Category="Collision")
+	FComponentWakeSignature OnComponentWake;
+
+	/** 
+	 *	Event called when the underlying physics objects is put to sleep
+	 */
+	UPROPERTY(BlueprintAssignable, Category = "Collision")
+	FComponentSleepSignature OnComponentSleep;
+
 	/** Event called when the mouse cursor is moved over this component and mouse over events are enabled in the player controller */
 	UPROPERTY(BlueprintAssignable, Category="Input|Mouse Input")
 	FComponentBeginCursorOverSignature OnBeginCursorOver;
@@ -1412,6 +1428,12 @@ public:
 	 */
 	void DispatchBlockingHit(AActor& Owner, FHitResult const& BlockingHit);
 
+	/**
+	 * Dispatch notification for wake events and propagate to any welded bodies
+	 */
+
+	void DispatchWakeEvents(int32 WakeEvent, FName BoneName);
+
 	/**
 	 * Set collision params on OutParams (such as CollisionResponse, bTraceAsyncScene) to match the settings on this PrimitiveComponent.
 	 */
diff --git a/Engine/Source/Runtime/Engine/Classes/PhysicsEngine/BodyInstance.h b/Engine/Source/Runtime/Engine/Classes/PhysicsEngine/BodyInstance.h
--- a/Engine/Source/Runtime/Engine/Classes/PhysicsEngine/BodyInstance.h
+++ b/Engine/Source/Runtime/Engine/Classes/PhysicsEngine/BodyInstance.h
@@ -239,6 +239,10 @@ public:
 	UPROPERTY(EditAnywhere, AdvancedDisplay, BlueprintReadOnly, Category = Physics, meta = (editcondition = "bSimulatePhysics"))
 	uint32 bStartAwake:1;
 
+	/**	Should 'wake/sleep' events fire when this object is woken up or put to sleep by the physics simulation. */
+	UPROPERTY(EditAnywhere,AdvancedDisplay, BlueprintReadOnly, Category = Physics)
+	uint32 bGenerateWakeEvents : 1;
+
 	/** If true, it will update mass when scale changes **/
 	UPROPERTY()
 	uint32 bUpdateMassWhenScaleChanges:1;
diff --git a/Engine/Source/Runtime/Engine/Private/Components/PrimitiveComponent.cpp b/Engine/Source/Runtime/Engine/Private/Components/PrimitiveComponent.cpp
--- a/Engine/Source/Runtime/Engine/Private/Components/PrimitiveComponent.cpp
+++ b/Engine/Source/Runtime/Engine/Private/Components/PrimitiveComponent.cpp
@@ -1691,6 +1691,39 @@ void UPrimitiveComponent::DispatchBlockingHit(AActor& Owner, FHitResult const& B
 	}
 }
 
+void UPrimitiveComponent::DispatchWakeEvents(int32 WakeEvent, FName BoneName)
+{
+	FBodyInstance* RootBI = GetBodyInstance(BoneName, false);
+	if(RootBI)
+	{
+		if(RootBI->bGenerateWakeEvents)
+		{
+			if (WakeEvent == SleepEvent::SET_Wakeup)
+			{
+				OnComponentWake.Broadcast(BoneName);
+			}else
+			{
+				OnComponentSleep.Broadcast(BoneName);
+			}
+		}
+	}
+	
+	//now update children that are welded
+	for(USceneComponent* SceneComp : AttachChildren)
+	{
+		if(UPrimitiveComponent* PrimComp = Cast<UPrimitiveComponent>(SceneComp))
+		{
+			if(FBodyInstance* BI = PrimComp->GetBodyInstance(BoneName, false))
+			{
+				if(BI->WeldParent == RootBI)
+				{
+					PrimComp->DispatchWakeEvents(WakeEvent, BoneName);	
+				}
+			}
+		}
+	}
+}
+
 
 bool UPrimitiveComponent::IsNavigationRelevant() const 
 { 
diff --git a/Engine/Source/Runtime/Engine/Private/PhysicsEngine/BodyInstance.cpp b/Engine/Source/Runtime/Engine/Private/PhysicsEngine/BodyInstance.cpp
--- a/Engine/Source/Runtime/Engine/Private/PhysicsEngine/BodyInstance.cpp
+++ b/Engine/Source/Runtime/Engine/Private/PhysicsEngine/BodyInstance.cpp
@@ -225,6 +225,7 @@ FBodyInstance::FBodyInstance()
 , bAutoWeld(false)
 , bWelded(false)
 , bStartAwake(true)
+, bGenerateWakeEvents(false)
 , bUpdateMassWhenScaleChanges(false)
 , bLockTranslation(true)
 , bLockRotation(true)
@@ -1062,6 +1063,11 @@ struct FInitBodiesHelper
 			{
 				PNewDynamic->setRigidDynamicFlag(PxRigidDynamicFlag::eKINEMATIC, true);
 			}
+
+			if(Instance->bGenerateWakeEvents)
+			{
+				PNewDynamic->setActorFlag(PxActorFlag::eSEND_SLEEP_NOTIFIES, true);
+			}
 		}
 
 		return PNewDynamic;
@@ -1857,12 +1863,20 @@ bool FBodyInstance::Weld(FBodyInstance* TheirBody, const FTransform& TheirTM)
 		//child body gets placed into the same scenes as parent body
 		if (PxRigidActor* MyBody = RigidActorSync)
 		{
-				TheirBody->BodySetup->AddShapesToRigidActor_AssumesLocked(this, MyBody, PST_Sync, Scale3D, PSimpleMat, ComplexPhysMats, ShapeData, RelativeTM, &PNewShapes);
+			TheirBody->BodySetup->AddShapesToRigidActor_AssumesLocked(this, MyBody, PST_Sync, Scale3D, PSimpleMat, ComplexPhysMats, ShapeData, RelativeTM, &PNewShapes);
+			if (TheirBody->bGenerateWakeEvents)
+			{
+				MyBody->setActorFlag(PxActorFlag::eSEND_SLEEP_NOTIFIES, true);
+			}
 		}
 
 		if (PxRigidActor* MyBody = RigidActorAsync)
 		{
-				TheirBody->BodySetup->AddShapesToRigidActor_AssumesLocked(this, MyBody, PST_Sync, Scale3D, PSimpleMat, ComplexPhysMats, ShapeData, RelativeTM, &PNewShapes);
+			TheirBody->BodySetup->AddShapesToRigidActor_AssumesLocked(this, MyBody, PST_Sync, Scale3D, PSimpleMat, ComplexPhysMats, ShapeData, RelativeTM, &PNewShapes);
+			if (TheirBody->bGenerateWakeEvents)
+			{
+				MyBody->setActorFlag(PxActorFlag::eSEND_SLEEP_NOTIFIES, true);
+			}
 		}
 
 		for (int32 ShapeIdx = 0; ShapeIdx < PNewShapes.Num(); ++ShapeIdx)
diff --git a/Engine/Source/Runtime/Engine/Private/PhysicsEngine/PhysScene.cpp b/Engine/Source/Runtime/Engine/Private/PhysicsEngine/PhysScene.cpp
--- a/Engine/Source/Runtime/Engine/Private/PhysicsEngine/PhysScene.cpp
+++ b/Engine/Source/Runtime/Engine/Private/PhysicsEngine/PhysScene.cpp
@@ -357,13 +357,15 @@ void FPhysScene::AddTorque_AssumesLocked(FBodyInstance* BodyInstance, const FVec
 }
 
 #if WITH_PHYSX
-void FPhysScene::RemoveActiveBody(FBodyInstance* BodyInstance, uint32 SceneType)
+void FPhysScene::RemoveActiveBody_AssumesLocked(FBodyInstance* BodyInstance, uint32 SceneType)
 {
 	int32 BodyIndex = ActiveBodyInstances[SceneType].Find(BodyInstance);
 	if (BodyIndex != INDEX_NONE)
 	{
 		ActiveBodyInstances[SceneType][BodyIndex] = nullptr;
 	}
+
+	PendingSleepEvents[SceneType].Remove(BodyInstance->GetPxRigidActor_AssumesLocked(SceneType));
 }
 #endif
 void FPhysScene::TermBody_AssumesLocked(FBodyInstance* BodyInstance)
@@ -388,8 +390,8 @@ void FPhysScene::TermBody_AssumesLocked(FBodyInstance* BodyInstance)
 	}
 
 #if WITH_PHYSX
-	RemoveActiveBody(BodyInstance, PST_Sync);
-	RemoveActiveBody(BodyInstance, PST_Async);
+	RemoveActiveBody_AssumesLocked(BodyInstance, PST_Sync);
+	RemoveActiveBody_AssumesLocked(BodyInstance, PST_Async);
 #endif
 }
 
@@ -903,6 +905,25 @@ void FPhysScene::DispatchPhysNotifications_AssumesLocked()
 
 	PendingApexDamageManager->PendingDamageEvents.Empty();
 #endif
+
+#if WITH_PHYSX
+	for (int32 SceneType = 0; SceneType < PST_MAX; ++SceneType)
+	{
+		for (auto MapItr = PendingSleepEvents[SceneType].CreateIterator(); MapItr; ++MapItr)
+		{
+			PxActor* Actor = MapItr.Key();
+			if(FBodyInstance* BodyInstance = FPhysxUserData::Get<FBodyInstance>(Actor->userData))
+			{
+				if(UPrimitiveComponent* PrimitiveComponent = BodyInstance->OwnerComponent.Get())
+				{
+					PrimitiveComponent->DispatchWakeEvents(MapItr.Value(), BodyInstance->BodySetup->BoneName);
+				}
+			}
+		}
+
+		PendingSleepEvents[SceneType].Empty();
+	}
+#endif
 }
 
 void FPhysScene::SetUpForFrame(const FVector* NewGrav, float InDeltaSeconds, float InMaxPhysicsDeltaTime)
@@ -1448,6 +1469,11 @@ void FPhysScene::TermPhysScene(uint32 SceneType)
 
 #if WITH_PHYSX
 
+void FPhysScene::AddPendingSleepingEvent(PxActor* Actor, SleepEvent::Type SleepEventType, int32 SceneType)
+{
+	PendingSleepEvents[SceneType].FindOrAdd(Actor) = SleepEventType;
+}
+
 void FPhysScene::FDeferredSceneData::FlushDeferredActors()
 {
 	check(AddInstances.Num() == AddActors.Num());
diff --git a/Engine/Source/Runtime/Engine/Public/PhysicsPublic.h b/Engine/Source/Runtime/Engine/Public/PhysicsPublic.h
--- a/Engine/Source/Runtime/Engine/Public/PhysicsPublic.h
+++ b/Engine/Source/Runtime/Engine/Public/PhysicsPublic.h
@@ -252,6 +252,15 @@ private:
 	/** Array of commands waiting to execute once simulation is done */
 	TArray<FPhysPendingCommand> PendingCommands;
 };
+namespace SleepEvent
+{
+	enum Type
+	{
+		SET_Wakeup,
+		SET_Sleep
+	};
+
+}
 
 /** Container object for a physics engine 'scene'. */
 
@@ -321,6 +330,8 @@ public:
 	 *	@param SceneType - The scene type to add the actor to
 	 */
 	void DeferRemoveActor(FBodyInstance* OwningInstance, PxActor* Actor, EPhysicsSceneType SceneType);
+
+	void AddPendingSleepingEvent(PxActor* Actor, SleepEvent::Type SleepEventType, int32 SceneType);
 #endif
 
 private:
@@ -598,7 +609,7 @@ private:
 
 	/** Fetch results from simulation and get the active transforms. Make sure to lock before calling this function as the fetch and data you use must be treated as an atomic operation */
 	void UpdateActiveTransforms(uint32 SceneType);
-	void RemoveActiveBody(FBodyInstance* BodyInstance, uint32 SceneType);
+	void RemoveActiveBody_AssumesLocked(FBodyInstance* BodyInstance, uint32 SceneType);
 
 #endif
 
@@ -624,6 +635,10 @@ private:
 
 	/** Map from SkeletalMeshComponent UniqueID to a pointer to the collision disable table inside its PhysicsAsset */
 	TMap< uint32, TMap<struct FRigidBodyIndexPair, bool>* >		CollisionDisableTableLookup;
+
+#if WITH_PHYSX
+	TMap<PxActor*, SleepEvent::Type> PendingSleepEvents[PST_MAX];
+#endif
 };
 
 /**
