#ANIM: added post eval event graph, and native eval functions
#ANIM: exposed BoneIsChildOf to blueprint

diff --git a/Engine/Source/Runtime/Engine/Classes/Animation/AnimInstance.h b/Engine/Source/Runtime/Engine/Classes/Animation/AnimInstance.h
--- a/Engine/Source/Runtime/Engine/Classes/Animation/AnimInstance.h
+++ b/Engine/Source/Runtime/Engine/Classes/Animation/AnimInstance.h
@@ -409,6 +409,10 @@ public:
 	UFUNCTION(BlueprintImplementableEvent)
 	void BlueprintUpdateAnimation(float DeltaTimeX);
 
+	/** Executed after the Animation is evaluated */
+	UFUNCTION(BlueprintImplementableEvent)
+	void BlueprintPostEvaluateAnimation();
+
 	bool CanTransitionSignature() const;
 	
 	UFUNCTION()
@@ -776,24 +780,29 @@ public:
 	// Updates the montage data used for evaluation based on the current playing montages
 	void UpdateMontageEvaluationData();
 
-	//@TODO: Better comments
-	virtual void EvaluateAnimation(struct FPoseContext& Output);
-	virtual void PostAnimEvaluation() {}
-
+	// Animation phase trigger
+	// start with initialize
+	// update happens in every tick
+	// evaluate happens when condition is met - i.e. depending on your skeletalmeshcomponent update flag
+	// post eval happens after evaluation is done
+	// uninitialize happens when owner is unregistered
 	void InitializeAnimation();
 	void UpdateAnimation(float DeltaSeconds);
+	void EvaluateAnimation(struct FPoseContext& Output);
+	void PostEvaluateAnimation();
 	void UninitializeAnimation();
 
+	// the below functions are the native overrides for each phase
 	// Native initialization override point
 	virtual void NativeInitializeAnimation();
-
 	// Native update override point
 	virtual void NativeUpdateAnimation(float DeltaSeconds);
-
 	// Native evaluate override point.
 	// @return true if this function is implemented, false otherwise.
 	// Note: the node graph will not be evaluated if this function returns true
 	virtual bool NativeEvaluateAnimation(FPoseContext& Output);
+	// Native Post Evaluate override point
+	virtual void NativePostEvaluateAnimation();
 
 	// Sets up a native transition delegate between states with PrevStateName and NextStateName, in the state machine with name MachineName.
 	// Note that a transition already has to exist for this to succeed
diff --git a/Engine/Source/Runtime/Engine/Classes/Animation/AnimSingleNodeInstance.h b/Engine/Source/Runtime/Engine/Classes/Animation/AnimSingleNodeInstance.h
--- a/Engine/Source/Runtime/Engine/Classes/Animation/AnimSingleNodeInstance.h
+++ b/Engine/Source/Runtime/Engine/Classes/Animation/AnimSingleNodeInstance.h
@@ -59,7 +59,7 @@ class ENGINE_API UAnimSingleNodeInstance : public UAnimInstance
 	virtual void NativeInitializeAnimation() override;
 	virtual void NativeUpdateAnimation(float DeltaTimeX) override;
 	virtual bool NativeEvaluateAnimation(FPoseContext& Output) override;
-	virtual void PostAnimEvaluation() override;
+	virtual void NativePostEvaluateAnimation() override;
 
 protected:
 	virtual void Montage_Advance(float DeltaTime) override;
diff --git a/Engine/Source/Runtime/Engine/Classes/Components/SkinnedMeshComponent.h b/Engine/Source/Runtime/Engine/Classes/Components/SkinnedMeshComponent.h
--- a/Engine/Source/Runtime/Engine/Classes/Components/SkinnedMeshComponent.h
+++ b/Engine/Source/Runtime/Engine/Classes/Components/SkinnedMeshComponent.h
@@ -741,6 +741,7 @@ public:
 	 * @return true if child (strictly, not same). false otherwise
 	 * Note - will return false if ChildBoneIndex is the same as ParentBoneIndex ie. must be strictly a child.
 	 */
+	UFUNCTION(BlueprintCallable, Category="Components|SkinnedMesh")
 	bool BoneIsChildOf(FName BoneName, FName ParentBoneName) const;
 
 	/** 
diff --git a/Engine/Source/Runtime/Engine/Private/Animation/AnimInstance.cpp b/Engine/Source/Runtime/Engine/Private/Animation/AnimInstance.cpp
--- a/Engine/Source/Runtime/Engine/Private/Animation/AnimInstance.cpp
+++ b/Engine/Source/Runtime/Engine/Private/Animation/AnimInstance.cpp
@@ -47,8 +47,8 @@ DEFINE_STAT(STAT_AnimTickTime);
 DEFINE_STAT(STAT_SkinnedMeshCompTick);
 DEFINE_STAT(STAT_TickUpdateRate);
 DEFINE_STAT(STAT_PerformAnimEvaluation);
-DEFINE_STAT(STAT_PostAnimEvaluation);
 DEFINE_STAT(STAT_BlueprintUpdateAnimation);
+DEFINE_STAT(STAT_BlueprintPostEvaluateAnimation);
 DEFINE_STAT(STAT_NativeUpdateAnimation);
 
 DECLARE_CYCLE_STAT_EXTERN(TEXT("Anim Init Time"), STAT_AnimInitTime, STATGROUP_Anim, );
@@ -481,6 +481,16 @@ void UAnimInstance::EvaluateAnimation(FPoseContext& Output)
 	}
 }
 
+void UAnimInstance::PostEvaluateAnimation()
+{
+	NativePostEvaluateAnimation();
+
+	{
+		SCOPE_CYCLE_COUNTER(STAT_BlueprintPostEvaluateAnimation);
+		BlueprintPostEvaluateAnimation();
+	}
+}
+
 void UAnimInstance::NativeInitializeAnimation()
 {
 }
@@ -494,6 +504,10 @@ bool UAnimInstance::NativeEvaluateAnimation(FPoseContext& Output)
 	return false;
 }
 
+void UAnimInstance::NativePostEvaluateAnimation()
+{
+}
+
 void UAnimInstance::AddNativeTransitionBinding(const FName& MachineName, const FName& PrevStateName, const FName& NextStateName, const FCanTakeTransition& NativeTransitionDelegate)
 {
 	NativeTransitionBindings.Add(FNativeTransitionBinding(MachineName, PrevStateName, NextStateName, NativeTransitionDelegate));
diff --git a/Engine/Source/Runtime/Engine/Private/Animation/AnimSingleNodeInstance.cpp b/Engine/Source/Runtime/Engine/Private/Animation/AnimSingleNodeInstance.cpp
--- a/Engine/Source/Runtime/Engine/Private/Animation/AnimSingleNodeInstance.cpp
+++ b/Engine/Source/Runtime/Engine/Private/Animation/AnimSingleNodeInstance.cpp
@@ -417,11 +417,11 @@ bool UAnimSingleNodeInstance::NativeEvaluateAnimation(FPoseContext& Output)
 	return true;
 }
 
-void UAnimSingleNodeInstance::PostAnimEvaluation()
+void UAnimSingleNodeInstance::NativePostEvaluateAnimation()
 {
 	PostEvaluateAnimEvent.ExecuteIfBound();
 
-	Super::PostAnimEvaluation();
+	Super::NativePostEvaluateAnimation();
 }
 
 void UAnimSingleNodeInstance::Montage_Advance(float DeltaTime)
diff --git a/Engine/Source/Runtime/Engine/Private/Components/SkeletalMeshComponent.cpp b/Engine/Source/Runtime/Engine/Private/Components/SkeletalMeshComponent.cpp
--- a/Engine/Source/Runtime/Engine/Private/Components/SkeletalMeshComponent.cpp
+++ b/Engine/Source/Runtime/Engine/Private/Components/SkeletalMeshComponent.cpp
@@ -38,6 +38,7 @@ TAutoConsoleVariable<int32> CVarUseParallelAnimationEvaluation(TEXT("a.ParallelA
 
 DECLARE_CYCLE_STAT_EXTERN(TEXT("Anim Instance Spawn Time"), STAT_AnimSpawnTime, STATGROUP_Anim, );
 DEFINE_STAT(STAT_AnimSpawnTime);
+DEFINE_STAT(STAT_PostAnimEvaluation);
 
 class FParallelAnimationEvaluationTask
 {
@@ -1249,7 +1250,7 @@ void USkeletalMeshComponent::PostAnimEvaluation(FAnimationEvaluationContext& Eva
 	{
 		// curve update happens first
 		AnimScriptInstance->UpdateCurves(EvaluatedCurve);
-		AnimScriptInstance->PostAnimEvaluation();
+		AnimScriptInstance->PostEvaluateAnimation();
 	}
 
 	MarkRenderDynamicDataDirty();
diff --git a/Engine/Source/Runtime/Engine/Public/Animation/AnimStats.h b/Engine/Source/Runtime/Engine/Public/Animation/AnimStats.h
--- a/Engine/Source/Runtime/Engine/Public/Animation/AnimStats.h
+++ b/Engine/Source/Runtime/Engine/Public/Animation/AnimStats.h
@@ -37,5 +37,6 @@ DECLARE_CYCLE_STAT_EXTERN(TEXT("UpdateTriMeshVertices"), STAT_UpdateTriMeshVerti
 DECLARE_CYCLE_STAT_EXTERN(TEXT("AnimGameThreadTime"), STAT_AnimGameThreadTime, STATGROUP_Anim, );
 DECLARE_CYCLE_STAT_EXTERN(TEXT("BlueprintUpdateAnimation"), STAT_BlueprintUpdateAnimation, STATGROUP_Anim, );
 DECLARE_CYCLE_STAT_EXTERN(TEXT("NativeUpdateAnimation"), STAT_NativeUpdateAnimation, STATGROUP_Anim, );
+DECLARE_CYCLE_STAT_EXTERN(TEXT("BlueprintPostEvaluateAnimation"), STAT_BlueprintPostEvaluateAnimation, STATGROUP_Anim, );
 
 
