#ue4 - Improve docs on MovementComponent and FloatingPawnMovement to indicate that collision is wholly dependent on the UpdatedComponent (normally the root component).

diff --git a/Engine/Source/Runtime/Engine/Classes/GameFramework/FloatingPawnMovement.h b/Engine/Source/Runtime/Engine/Classes/GameFramework/FloatingPawnMovement.h
--- a/Engine/Source/Runtime/Engine/Classes/GameFramework/FloatingPawnMovement.h
+++ b/Engine/Source/Runtime/Engine/Classes/GameFramework/FloatingPawnMovement.h
@@ -11,6 +11,9 @@
 /**
  * FloatingPawnMovement is a movement component that provides simple movement for any Pawn class.
  * Limits on speed and acceleration are provided, while gravity is not implemented.
+ *
+ * Normally the root component of the owning actor is moved, however another component may be selected (see SetUpdatedComponent()).
+ * During swept (non-teleporting) movement only collision of UpdatedComponent is considered, attached components will teleport to the end location ignoring collision.
  */
 UCLASS(ClassGroup = Movement, meta = (BlueprintSpawnableComponent))
 class ENGINE_API UFloatingPawnMovement : public UPawnMovementComponent
diff --git a/Engine/Source/Runtime/Engine/Classes/GameFramework/MovementComponent.h b/Engine/Source/Runtime/Engine/Classes/GameFramework/MovementComponent.h
--- a/Engine/Source/Runtime/Engine/Classes/GameFramework/MovementComponent.h
+++ b/Engine/Source/Runtime/Engine/Classes/GameFramework/MovementComponent.h
@@ -42,6 +42,7 @@ enum class EPlaneConstraintAxisSetting : uint8
  *    - Utility functions for moving when there may be initial penetration (SafeMoveUpdatedComponent(), ResolvePenetration()).
  *    - Automatically registering the component tick and finding a component to move on the owning Actor.
  * Normally the root component of the owning actor is moved, however another component may be selected (see SetUpdatedComponent()).
+ * During swept (non-teleporting) movement only collision of UpdatedComponent is considered, attached components will teleport to the end location ignoring collision.
  */
 UCLASS(ClassGroup=Movement, abstract, BlueprintType)
 class ENGINE_API UMovementComponent : public UActorComponent
