#IMPORT: If you import zero scale joint as refpose, make sure to warn user, and don't use zero scale.
We can't inverse zero scale.

https://jira.ol.epicgames.net/browse/UE-16676

diff --git a/Engine/Source/Runtime/Engine/Private/SkeletalMesh.cpp b/Engine/Source/Runtime/Engine/Private/SkeletalMesh.cpp
--- a/Engine/Source/Runtime/Engine/Private/SkeletalMesh.cpp
+++ b/Engine/Source/Runtime/Engine/Private/SkeletalMesh.cpp
@@ -2501,8 +2501,19 @@ void USkeletalMesh::CalculateInvRefMatrices()
 				RefBases[b] = RefBases[b] * RefBases[Parent];
 			}
 
+			FVector XAxis, YAxis, ZAxis;
+
+			RefBases[b].GetScaledAxes(XAxis, YAxis, ZAxis);
+			if(	XAxis.IsNearlyZero(SMALL_NUMBER) &&
+				YAxis.IsNearlyZero(SMALL_NUMBER) &&
+				ZAxis.IsNearlyZero(SMALL_NUMBER))
+			{
+				// this is not allowed, warn them 
+				UE_LOG(LogSkeletalMesh, Warning, TEXT("Reference Pose for joint (%s) includes NIL matrix. Zero scale isn't allowed on ref pose. "), *RefSkeleton.GetBoneName(b).ToString());
+			}
+
 			// Precompute inverse so we can use from-refpose-skin vertices.
-			RefBasesInvMatrix[b] = RefBases[b].InverseFast(); 
+			RefBasesInvMatrix[b] = RefBases[b].Inverse(); 
 		}
 
 #if WITH_EDITORONLY_DATA
