#ue4 - Add BrakingFrictionFactor which controls how friction is used in ApplyVelocityBraking().

Note: this is 2 by default for historical reasons, but a value of 1 more closely matches actual drag.

diff --git a/Engine/Source/Runtime/Engine/Classes/GameFramework/CharacterMovementComponent.h b/Engine/Source/Runtime/Engine/Classes/GameFramework/CharacterMovementComponent.h
--- a/Engine/Source/Runtime/Engine/Classes/GameFramework/CharacterMovementComponent.h
+++ b/Engine/Source/Runtime/Engine/Classes/GameFramework/CharacterMovementComponent.h
@@ -198,10 +198,10 @@ public:
 
 	/**
 	 * Setting that affects movement control. Higher values allow faster changes in direction.
-	 * If bUseSeparateBrakingFriction is false, also affects the ability to stop more quickly when braking (whenever Acceleration is zero).
-	 * This property allows you to control how much friction is applied when moving across the ground, applying an opposing force that scales with current velocity.
+	 * If bUseSeparateBrakingFriction is false, also affects the ability to stop more quickly when braking (whenever Acceleration is zero), where it is multiplied by BrakingFrictionFactor.
+	 * When braking, this property allows you to control how much friction is applied when moving across the ground, applying an opposing force that scales with current velocity.
 	 * This can be used to simulate slippery surfaces such as ice or oil by changing the value (possibly based on the material pawn is standing on).
-	 * @see BrakingDecelerationWalking, BrakingFriction, bUseSeparateBrakingFriction
+	 * @see BrakingDecelerationWalking, BrakingFriction, bUseSeparateBrakingFriction, BrakingFrictionFactor
 	 */
 	UPROPERTY(Category="Character Movement: Walking", EditAnywhere, BlueprintReadWrite, meta=(ClampMin="0", UIMin="0"))
 	float GroundFriction;
@@ -231,15 +231,34 @@ public:
 	float MaxAcceleration;
 
 	/**
-	 * Friction (drag) coefficient applied when braking (whenever Acceleration = 0, or if character is exceeding max speed).
-	 * Only used if bUseSeparateBrakingFriction setting is true, otherwise current friction such as GroundFriction is used.
-	 * This is the current value, used in all movement modes; if this is not desired, override it when movement mode changes,
-	 * or only enable bUseSeparateBrakingFriction in desired movement modes.
-	 * @see bUseSeparateBrakingFriction
+	 * Factor used to multiply actual value of friction used when braking.
+	 * This applies to any friction value that is currently used, which may depend on bUseSeparateBrakingFriction.
+	 * @note This is 2 by default for historical reasons, a value of 1 gives the true drag equation.
+	 * @see bUseSeparateBrakingFriction, GroundFriction, BrakingFriction
+	 */
+	UPROPERTY(Category="Character Movement (General Settings)", EditAnywhere, BlueprintReadWrite, meta=(ClampMin="0", UIMin="0"))
+	float BrakingFrictionFactor;
+
+	/**
+	 * Friction (drag) coefficient applied when braking (whenever Acceleration = 0, or if character is exceeding max speed); actual value used is this multiplied by BrakingFrictionFactor.
+	 * When braking, this property allows you to control how much friction is applied when moving across the ground, applying an opposing force that scales with current velocity.
+	 * Braking is composed of friction (velocity-dependent drag) and constant deceleration.
+	 * This is the current value, used in all movement modes; if this is not desired, override it or bUseSeparateBrakingFriction when movement mode changes.
+	 * @note Only used if bUseSeparateBrakingFriction setting is true, otherwise current friction such as GroundFriction is used.
+	 * @see bUseSeparateBrakingFriction, BrakingFrictionFactor, GroundFriction, BrakingDecelerationWalking
 	 */
 	UPROPERTY(Category="Character Movement (General Settings)", EditAnywhere, BlueprintReadWrite, meta=(ClampMin="0", UIMin="0", EditCondition="bUseSeparateBrakingFriction"))
 	float BrakingFriction;
 
+	/**
+	 * If true, BrakingFriction will be used to slow the character to a stop (when there is no Acceleration).
+	 * If false, braking uses the same friction passed to CalcVelocity() (ie GroundFriction when walking), multiplied by BrakingFrictionFactor.
+	 * This setting applies to all movement modes; if only desired in certain modes, consider toggling it when movement modes change.
+	 * @see BrakingFriction
+	 */
+	UPROPERTY(Category="Character Movement (General Settings)", EditDefaultsOnly, BlueprintReadWrite)
+	uint32 bUseSeparateBrakingFriction:1;
+
 	/**
 	 * Deceleration when walking and not applying acceleration. This is a constant opposing force that directly lowers velocity by a constant value.
 	 * @see GroundFriction, MaxAcceleration
@@ -335,15 +354,6 @@ public:
 	UPROPERTY(Category="Character Movement (General Settings)", EditAnywhere, BlueprintReadWrite)
 	uint32 bOrientRotationToMovement:1;
 
-	/**
-	 * If true, BrakingFriction will be used to slow the character to a stop (when there is no Acceleration).
-	 * If false, braking uses the same friction passed to CalcVelocity() (ie GroundFriction when walking), multiplied by 2.
-	 * This setting applies to all movement modes; if only desired in certain modes, consider toggling it when movement modes change.
-	 * @see BrakingFriction
-	 */
-	UPROPERTY(Category="Character Movement (General Settings)", EditDefaultsOnly, BlueprintReadWrite)
-	uint32 bUseSeparateBrakingFriction : 1;
-
 protected:
 
 	/**
diff --git a/Engine/Source/Runtime/Engine/Private/Components/CharacterMovementComponent.cpp b/Engine/Source/Runtime/Engine/Private/Components/CharacterMovementComponent.cpp
--- a/Engine/Source/Runtime/Engine/Private/Components/CharacterMovementComponent.cpp
+++ b/Engine/Source/Runtime/Engine/Private/Components/CharacterMovementComponent.cpp
@@ -202,6 +202,7 @@ UCharacterMovementComponent::UCharacterMovementComponent(const FObjectInitialize
 	AirControlBoostVelocityThreshold = 25.f;
 	FallingLateralFriction = 0.f;
 	MaxAcceleration = 2048.0f;
+	BrakingFrictionFactor = 2.0f; // Historical value, 1 would be more appropriate.
 	BrakingDecelerationWalking = MaxAcceleration;
 	BrakingDecelerationFalling = 0.f;
 	BrakingDecelerationFlying = 0.f;
@@ -2692,7 +2693,8 @@ void UCharacterMovementComponent::ApplyVelocityBraking(float DeltaTime, float Fr
 		return;
 	}
 
-	Friction = FMath::Max(0.f, Friction);
+	const float FrictionFactor = FMath::Max(0.f, BrakingFrictionFactor);
+	Friction = FMath::Max(0.f, Friction * FrictionFactor);
 	BrakingDeceleration = FMath::Max(0.f, BrakingDeceleration);
 	const bool bZeroFriction = (Friction == 0.f);
 	const bool bZeroBraking = (BrakingDeceleration == 0.f);
@@ -2709,9 +2711,6 @@ void UCharacterMovementComponent::ApplyVelocityBraking(float DeltaTime, float Fr
 	float RemainingTime = DeltaTime;
 	const float MaxTimeStep = (1.0f / 33.0f);
 
-	// Old (legacy) friction doubled the affect of "Friction" here.
-	const float FrictionFactor = (bUseSeparateBrakingFriction ? 1.0f : 2.0f);
-
 	// Decelerate to brake to a stop
 	const FVector RevAccel = (bZeroBraking ? FVector::ZeroVector : (-BrakingDeceleration * Velocity.GetSafeNormal()));
 	while( RemainingTime >= MIN_TICK_TIME )
@@ -2721,7 +2720,7 @@ void UCharacterMovementComponent::ApplyVelocityBraking(float DeltaTime, float Fr
 		RemainingTime -= dt;
 
 		// apply friction and braking
-		Velocity = Velocity + ((-FrictionFactor * Friction) * Velocity + RevAccel) * dt ; 
+		Velocity = Velocity + ((-Friction) * Velocity + RevAccel) * dt ; 
 		
 		// Don't reverse direction
 		if ((Velocity | OldVel) <= 0.f)
