Fixed bad tooltips on material instance property overrides

diff --git a/Engine/Source/Runtime/Engine/Classes/Materials/MaterialInstanceBasePropertyOverrides.h b/Engine/Source/Runtime/Engine/Classes/Materials/MaterialInstanceBasePropertyOverrides.h
--- a/Engine/Source/Runtime/Engine/Classes/Materials/MaterialInstanceBasePropertyOverrides.h
+++ b/Engine/Source/Runtime/Engine/Classes/Materials/MaterialInstanceBasePropertyOverrides.h
@@ -43,11 +43,11 @@ struct ENGINE_API FMaterialInstanceBasePropertyOverrides
 	UPROPERTY(EditAnywhere, Category = Material, meta = (editcondition = "bOverride_ShadingModel"))
 	TEnumAsByte<EMaterialShadingModel> ShadingModel;
 
-	/** If the material is two sided or not. */
+	/** Indicates that the material should be rendered without backface culling and the normal should be flipped for backfaces. */
 	UPROPERTY(EditAnywhere, Category = Material, meta = (editcondition = "bOverride_TwoSided"))
 	uint32 TwoSided : 1;
 
-	/** If the material is two sided or not. */
+	/** Whether the material should support a dithered LOD transition when used with the foliage system. */
 	UPROPERTY(EditAnywhere, Category = Material, meta = (editcondition = "bOverride_DitheredLODTransition"))
 	uint32 DitheredLODTransition : 1;
 
