Make sure destination map "always-loaded" sub-levels are fully loaded on seamless travel
CL# 2608694
Merging using UE4-Orion_to_UE4

diff --git a/Engine/Source/Runtime/Engine/Private/World.cpp b/Engine/Source/Runtime/Engine/Private/World.cpp
--- a/Engine/Source/Runtime/Engine/Private/World.cpp
+++ b/Engine/Source/Runtime/Engine/Private/World.cpp
@@ -4830,6 +4830,9 @@ UWorld* FSeamlessTravelHandler::Tick()
 				}
 			}
 
+			// Make sure "always loaded" sub-levels are fully loaded
+			LoadedWorld->FlushLevelStreaming(EFlushLevelStreamingType::Visibility);
+			
 			UNavigationSystem::InitializeForWorld(LoadedWorld, FNavigationSystemRunMode::GameMode);
 
 			// Note that AI system will be created only if ai-system-creation conditions are met
