added navigation agent selectors to navmesh bounds volume, defining zones per agent
#ue4 UE-2021

diff --git a/Engine/Source/Runtime/Engine/Classes/AI/Navigation/NavMeshBoundsVolume.h b/Engine/Source/Runtime/Engine/Classes/AI/Navigation/NavMeshBoundsVolume.h
--- a/Engine/Source/Runtime/Engine/Classes/AI/Navigation/NavMeshBoundsVolume.h
+++ b/Engine/Source/Runtime/Engine/Classes/AI/Navigation/NavMeshBoundsVolume.h
@@ -14,6 +14,9 @@ class ANavMeshBoundsVolume : public AVolume
 {
 	GENERATED_UCLASS_BODY()
 
+	UPROPERTY(EditAnywhere, Category = Navigation)
+	FNavAgentSelector SupportedAgents;
+
 	// Begin AActor Interface
 	virtual void PostRegisterAllComponents() override;
 	virtual void PostUnregisterAllComponents() override;
diff --git a/Engine/Source/Runtime/Engine/Classes/AI/Navigation/NavigationTypes.h b/Engine/Source/Runtime/Engine/Classes/AI/Navigation/NavigationTypes.h
--- a/Engine/Source/Runtime/Engine/Classes/AI/Navigation/NavigationTypes.h
+++ b/Engine/Source/Runtime/Engine/Classes/AI/Navigation/NavigationTypes.h
@@ -87,11 +87,90 @@
 	FORCEINLINE bool HasFlag(ENavigationDirtyFlag::Type Flag) const { return (Flags & Flag) != 0; }
 };
 
+USTRUCT()
+struct ENGINE_API FNavAgentSelector
+{
+	GENERATED_USTRUCT_BODY()
+
+#if CPP
+	union
+	{
+		struct
+		{
+#endif
+			UPROPERTY(EditAnywhere, Category = Default)
+			uint32 bSupportsAgent0 : 1;
+			UPROPERTY(EditAnywhere, Category = Default)
+			uint32 bSupportsAgent1 : 1;
+			UPROPERTY(EditAnywhere, Category = Default)
+			uint32 bSupportsAgent2 : 1;
+			UPROPERTY(EditAnywhere, Category = Default)
+			uint32 bSupportsAgent3 : 1;
+			UPROPERTY(EditAnywhere, Category = Default)
+			uint32 bSupportsAgent4 : 1;
+			UPROPERTY(EditAnywhere, Category = Default)
+			uint32 bSupportsAgent5 : 1;
+			UPROPERTY(EditAnywhere, Category = Default)
+			uint32 bSupportsAgent6 : 1;
+			UPROPERTY(EditAnywhere, Category = Default)
+			uint32 bSupportsAgent7 : 1;
+			UPROPERTY(EditAnywhere, Category = Default)
+			uint32 bSupportsAgent8 : 1;
+			UPROPERTY(EditAnywhere, Category = Default)
+			uint32 bSupportsAgent9 : 1;
+			UPROPERTY(EditAnywhere, Category = Default)
+			uint32 bSupportsAgent10 : 1;
+			UPROPERTY(EditAnywhere, Category = Default)
+			uint32 bSupportsAgent11 : 1;
+			UPROPERTY(EditAnywhere, Category = Default)
+			uint32 bSupportsAgent12 : 1;
+			UPROPERTY(EditAnywhere, Category = Default)
+			uint32 bSupportsAgent13 : 1;
+			UPROPERTY(EditAnywhere, Category = Default)
+			uint32 bSupportsAgent14 : 1;
+			UPROPERTY(EditAnywhere, Category = Default)
+			uint32 bSupportsAgent15 : 1;
+#if CPP
+		};
+		uint32 PackedBits;
+	};
+#endif
+
+	FNavAgentSelector();
+
+	FORCEINLINE bool Contains(int32 AgentIndex) const
+	{
+		return (AgentIndex >= 0 && AgentIndex < 16) ? !!(PackedBits & (1 << AgentIndex)) : false;
+	}
+
+	FORCEINLINE bool IsInitialized() const
+	{
+		return (PackedBits & 0x80000000) != 0;
+	}
+
+	FORCEINLINE void MarkInitialized()
+	{
+		PackedBits |= 0x80000000;
+	}
+
+	bool Serialize(FArchive& Ar);
+};
+
+template<>
+struct TStructOpsTypeTraits< FNavAgentSelector > : public TStructOpsTypeTraitsBase
+{
+	enum
+	{
+		WithSerializer = true,
+	};
+};
+
 struct FNavigationBounds
 {
-	uint32						UniqueID;
-	FBox						AreaBox;	 
-	TWeakObjectPtr<ULevel>		Level;		// The level this bounds belongs to
+	uint32 UniqueID;
+	FBox AreaBox;
+	FNavAgentSelector SupportedAgents;
+	TWeakObjectPtr<ULevel> Level;		// The level this bounds belongs to
 
 	bool operator==(const FNavigationBounds& Other) const 
 	{ 
@@ -415,83 +494,5 @@
 }
 
 USTRUCT()
-struct ENGINE_API FNavAgentSelector
-{
-	GENERATED_USTRUCT_BODY()
-
-#if CPP
-	union
-	{
-		struct
-		{
-#endif
-			UPROPERTY(EditAnywhere, Category = Default)
-			uint32 bSupportsAgent0 : 1;
-			UPROPERTY(EditAnywhere, Category = Default)
-			uint32 bSupportsAgent1 : 1;
-			UPROPERTY(EditAnywhere, Category = Default)
-			uint32 bSupportsAgent2 : 1;
-			UPROPERTY(EditAnywhere, Category = Default)
-			uint32 bSupportsAgent3 : 1;
-			UPROPERTY(EditAnywhere, Category = Default)
-			uint32 bSupportsAgent4 : 1;
-			UPROPERTY(EditAnywhere, Category = Default)
-			uint32 bSupportsAgent5 : 1;
-			UPROPERTY(EditAnywhere, Category = Default)
-			uint32 bSupportsAgent6 : 1;
-			UPROPERTY(EditAnywhere, Category = Default)
-			uint32 bSupportsAgent7 : 1;
-			UPROPERTY(EditAnywhere, Category = Default)
-			uint32 bSupportsAgent8 : 1;
-			UPROPERTY(EditAnywhere, Category = Default)
-			uint32 bSupportsAgent9 : 1;
-			UPROPERTY(EditAnywhere, Category = Default)
-			uint32 bSupportsAgent10 : 1;
-			UPROPERTY(EditAnywhere, Category = Default)
-			uint32 bSupportsAgent11 : 1;
-			UPROPERTY(EditAnywhere, Category = Default)
-			uint32 bSupportsAgent12 : 1;
-			UPROPERTY(EditAnywhere, Category = Default)
-			uint32 bSupportsAgent13 : 1;
-			UPROPERTY(EditAnywhere, Category = Default)
-			uint32 bSupportsAgent14 : 1;
-			UPROPERTY(EditAnywhere, Category = Default)
-			uint32 bSupportsAgent15 : 1;
-#if CPP
-		};
-		uint32 PackedBits;
-	};
-#endif
-
-	FNavAgentSelector();
-
-	FORCEINLINE bool Contains(int32 AgentIndex) const
-	{
-		return (AgentIndex >= 0 && AgentIndex < 16) ? !!(PackedBits & (1 << AgentIndex)) : false;
-	}
-
-	FORCEINLINE bool IsInitialized() const
-	{
-		return (PackedBits & 0x80000000) != 0;
-	}
-
-	FORCEINLINE void MarkInitialized()
-	{
-		PackedBits |= 0x80000000;
-	}
-
-	bool Serialize(FArchive& Ar);
-};
-
-template<>
-struct TStructOpsTypeTraits< FNavAgentSelector > : public TStructOpsTypeTraitsBase
-{
-	enum
-	{
-		WithSerializer = true,
-	};
-};
-
-USTRUCT()
 struct ENGINE_API FNavDataConfig : public FNavAgentProperties
 {
diff --git a/Engine/Source/Runtime/Engine/Private/AI/Navigation/NavigationSystem.cpp b/Engine/Source/Runtime/Engine/Private/AI/Navigation/NavigationSystem.cpp
--- a/Engine/Source/Runtime/Engine/Private/AI/Navigation/NavigationSystem.cpp
+++ b/Engine/Source/Runtime/Engine/Private/AI/Navigation/NavigationSystem.cpp
@@ -2884,10 +2884,12 @@ void UNavigationSystem::OnNavigationBoundsUpdated(ANavMeshBoundsVolume* NavVolum
 	}
 
 	FNavigationBoundsUpdateRequest UpdateRequest;
-	UpdateRequest.NavBounds.UniqueID	= NavVolume->GetUniqueID();
-	UpdateRequest.NavBounds.AreaBox		= NavVolume->GetComponentsBoundingBox(true);
-	UpdateRequest.NavBounds.Level		= NavVolume->GetLevel();
-	UpdateRequest.UpdateRequest			= FNavigationBoundsUpdateRequest::Updated;
+	UpdateRequest.NavBounds.UniqueID = NavVolume->GetUniqueID();
+	UpdateRequest.NavBounds.AreaBox = NavVolume->GetComponentsBoundingBox(true);
+	UpdateRequest.NavBounds.Level = NavVolume->GetLevel();
+	UpdateRequest.NavBounds.SupportedAgents = NavVolume->SupportedAgents;
+	
+	UpdateRequest.UpdateRequest = FNavigationBoundsUpdateRequest::Updated;
 	AddNavigationBoundsUpdateRequest(UpdateRequest);
 }
 
@@ -2899,10 +2901,12 @@ void UNavigationSystem::OnNavigationBoundsAdded(ANavMeshBoundsVolume* NavVolume)
 	}
 
 	FNavigationBoundsUpdateRequest UpdateRequest;
-	UpdateRequest.NavBounds.UniqueID	= NavVolume->GetUniqueID();
-	UpdateRequest.NavBounds.AreaBox		= NavVolume->GetComponentsBoundingBox(true);
-	UpdateRequest.NavBounds.Level		= NavVolume->GetLevel();
-	UpdateRequest.UpdateRequest	= FNavigationBoundsUpdateRequest::Added;
+	UpdateRequest.NavBounds.UniqueID = NavVolume->GetUniqueID();
+	UpdateRequest.NavBounds.AreaBox = NavVolume->GetComponentsBoundingBox(true);
+	UpdateRequest.NavBounds.Level = NavVolume->GetLevel();
+	UpdateRequest.NavBounds.SupportedAgents = NavVolume->SupportedAgents;
+
+	UpdateRequest.UpdateRequest = FNavigationBoundsUpdateRequest::Added;
 	AddNavigationBoundsUpdateRequest(UpdateRequest);
 }
 
@@ -2914,10 +2918,12 @@ void UNavigationSystem::OnNavigationBoundsRemoved(ANavMeshBoundsVolume* NavVolum
 	}
 	
 	FNavigationBoundsUpdateRequest UpdateRequest;
-	UpdateRequest.NavBounds.UniqueID	= NavVolume->GetUniqueID();
-	UpdateRequest.NavBounds.AreaBox		= NavVolume->GetComponentsBoundingBox(true);
-	UpdateRequest.NavBounds.Level		= NavVolume->GetLevel();
-	UpdateRequest.UpdateRequest	= FNavigationBoundsUpdateRequest::Removed;
+	UpdateRequest.NavBounds.UniqueID = NavVolume->GetUniqueID();
+	UpdateRequest.NavBounds.AreaBox = NavVolume->GetComponentsBoundingBox(true);
+	UpdateRequest.NavBounds.Level = NavVolume->GetLevel();
+	UpdateRequest.NavBounds.SupportedAgents = NavVolume->SupportedAgents;
+
+	UpdateRequest.UpdateRequest = FNavigationBoundsUpdateRequest::Removed;
 	AddNavigationBoundsUpdateRequest(UpdateRequest);
 }
 
@@ -3002,19 +3008,21 @@ void UNavigationSystem::PerformNavigationBoundsUpdate(const TArray<FNavigationBo
 			{
 				if (ExistingElementId.IsValidId())
 				{
-					FBox ExistingBox = RegisteredNavBounds[ExistingElementId].AreaBox;
-
-					if (!(ExistingBox == Request.NavBounds.AreaBox))
+					const FBox ExistingBox = RegisteredNavBounds[ExistingElementId].AreaBox;
+					const bool bSameArea = (Request.NavBounds.AreaBox == ExistingBox);
+					if (!bSameArea)
 					{
 						UpdatedAreas.Add(ExistingBox);
-						RegisteredNavBounds[ExistingElementId] = Request.NavBounds;
 					}
+
+					// always assign new bounds data, it may have different properties (like supported agents)
+					RegisteredNavBounds[ExistingElementId] = Request.NavBounds;
 				}
 				else
 				{
 					ExistingElementId = RegisteredNavBounds.Add(Request.NavBounds);
 				}
-				
+
 				UpdatedAreas.Add(Request.NavBounds.AreaBox);
 			}
 
@@ -3050,9 +3058,10 @@ void UNavigationSystem::GatherNavigationBounds()
 		if (V != nullptr && !V->IsPendingKill())
 		{
 			FNavigationBounds NavBounds;
-			NavBounds.UniqueID		= V->GetUniqueID();
-			NavBounds.AreaBox		= V->GetComponentsBoundingBox(true);
-			NavBounds.Level			= V->GetLevel();
+			NavBounds.UniqueID = V->GetUniqueID();
+			NavBounds.AreaBox = V->GetComponentsBoundingBox(true);
+			NavBounds.Level = V->GetLevel();
+			NavBounds.SupportedAgents = V->SupportedAgents;
 			RegisteredNavBounds.Add(NavBounds);
 		}
 	}
diff --git a/Engine/Source/Runtime/Engine/Private/AI/Navigation/RecastNavMeshGenerator.cpp b/Engine/Source/Runtime/Engine/Private/AI/Navigation/RecastNavMeshGenerator.cpp
--- a/Engine/Source/Runtime/Engine/Private/AI/Navigation/RecastNavMeshGenerator.cpp
+++ b/Engine/Source/Runtime/Engine/Private/AI/Navigation/RecastNavMeshGenerator.cpp
@@ -3160,6 +3160,7 @@ void FRecastNavMeshGenerator::UpdateNavigationBounds()
 {
 	const UNavigationSystem* NavSys = UNavigationSystem::GetCurrent(GetWorld());
 	const TSet<FNavigationBounds>& NavigationBoundsSet = NavSys->GetNavigationBounds();
+	const int32 AgentIndex = NavSys->GetSupportedAgentIndex(DestNavMesh);
 
 	TotalNavBounds = FBox(0);
 	InclusionBounds.Empty(NavigationBoundsSet.Num());
@@ -3169,8 +3170,11 @@ void FRecastNavMeshGenerator::UpdateNavigationBounds()
 	{
 		for (const auto& NavigationBounds : NavigationBoundsSet)
 		{
-			InclusionBounds.Add(NavigationBounds.AreaBox);
-			TotalNavBounds+= NavigationBounds.AreaBox;
+			if (NavigationBounds.SupportedAgents.Contains(AgentIndex))
+			{
+				InclusionBounds.Add(NavigationBounds.AreaBox);
+				TotalNavBounds += NavigationBounds.AreaBox;
+			}
 		}
 	}
 	else
diff --git a/Engine/Source/Runtime/Engine/Private/NavMeshBoundsVolume.cpp b/Engine/Source/Runtime/Engine/Private/NavMeshBoundsVolume.cpp
--- a/Engine/Source/Runtime/Engine/Private/NavMeshBoundsVolume.cpp
+++ b/Engine/Source/Runtime/Engine/Private/NavMeshBoundsVolume.cpp
@@ -11,6 +11,7 @@ ANavMeshBoundsVolume::ANavMeshBoundsVolume(const FObjectInitializer& ObjectIniti
 	GetBrushComponent()->Mobility = EComponentMobility::Static;
 
 	BrushColor = FColor(200, 200, 200, 255);
+	SupportedAgents.MarkInitialized();
 
 	bColored = true;
 }
@@ -25,7 +26,9 @@ void ANavMeshBoundsVolume::PostEditChangeProperty( struct FPropertyChangedEvent&
 	if (GIsEditor && NavSys)
 	{
 		if (PropertyChangedEvent.Property == NULL ||
-			PropertyChangedEvent.Property->GetFName() == GET_MEMBER_NAME_CHECKED(ABrush, BrushBuilder))
+			PropertyChangedEvent.Property->GetFName() == GET_MEMBER_NAME_CHECKED(ABrush, BrushBuilder) ||
+			(PropertyChangedEvent.MemberProperty && 
+			 PropertyChangedEvent.MemberProperty->GetFName() == GET_MEMBER_NAME_CHECKED(ANavMeshBoundsVolume, SupportedAgents)))
 		{
 			NavSys->OnNavigationBoundsUpdated(this);
 		}
