fixed copying supported agents bits in newly created navigation area
#ue4

diff --git a/Engine/Source/Runtime/Engine/Classes/AI/Navigation/NavAreas/NavArea.h b/Engine/Source/Runtime/Engine/Classes/AI/Navigation/NavAreas/NavArea.h
--- a/Engine/Source/Runtime/Engine/Classes/AI/Navigation/NavAreas/NavArea.h
+++ b/Engine/Source/Runtime/Engine/Classes/AI/Navigation/NavAreas/NavArea.h
@@ -75,6 +75,7 @@ public:
 	virtual void PostInitProperties() override;
 	virtual void FinishDestroy() override;
 	virtual void PostLoad() override;
+	virtual void Serialize(FArchive& Ar) override;
 
 	FORCEINLINE uint16 GetAreaFlags() const { return AreaFlags; }
 	FORCEINLINE bool HasFlags(uint16 InFlags) const { return (InFlags & AreaFlags) != 0; }
diff --git a/Engine/Source/Runtime/Engine/Private/AI/Navigation/NavAreas/NavArea.cpp b/Engine/Source/Runtime/Engine/Private/AI/Navigation/NavAreas/NavArea.cpp
--- a/Engine/Source/Runtime/Engine/Private/AI/Navigation/NavAreas/NavArea.cpp
+++ b/Engine/Source/Runtime/Engine/Private/AI/Navigation/NavAreas/NavArea.cpp
@@ -67,6 +67,16 @@ void UNavArea::PostLoad()
 	}
 }
 
+void UNavArea::Serialize(FArchive& Ar)
+{
+	if (Ar.IsSaving() && !SupportedAgents.IsInitialized())
+	{
+		SupportedAgents.MarkInitialized();
+	}
+		
+	Super::Serialize(Ar);
+}
+
 FColor UNavArea::GetColor(UClass* AreaDefinitionClass)
 {
 	return AreaDefinitionClass ? AreaDefinitionClass->GetDefaultObject<UNavArea>()->DrawColor : FColor::Black;
