Mesh particle collisions now consider particle size.
To maintian old behaviour this is togglable via bCollisionsConsiderParticleSize in the MeshTypeData.
Old data has this set to false but new data will default to true.

diff --git a/Engine/Source/Runtime/Core/Public/UObject/ObjectVersion.h b/Engine/Source/Runtime/Core/Public/UObject/ObjectVersion.h
--- a/Engine/Source/Runtime/Core/Public/UObject/ObjectVersion.h
+++ b/Engine/Source/Runtime/Core/Public/UObject/ObjectVersion.h
@@ -552,6 +552,9 @@ enum EUnrealEngineObjectUE4Version
 	VER_UE4_AO_MATERIAL_MASK,
 	// Replaced navigation agents selection with single structure
 	VER_UE4_NAVIGATION_AGENT_SELECTOR,
+	// Mesh particle collisions consider particle size.
+	VER_UE4_MESH_PARTICLE_COLLISIONS_CONSIDER_PARTICLE_SIZE,
+
 
 	// -----<new versions can be added before this line>-------------------------------------------------
 	// - this needs to be the last line (see note below)
diff --git a/Engine/Source/Runtime/Engine/Classes/Particles/TypeData/ParticleModuleTypeDataMesh.h b/Engine/Source/Runtime/Engine/Classes/Particles/TypeData/ParticleModuleTypeDataMesh.h
--- a/Engine/Source/Runtime/Engine/Classes/Particles/TypeData/ParticleModuleTypeDataMesh.h
+++ b/Engine/Source/Runtime/Engine/Classes/Particles/TypeData/ParticleModuleTypeDataMesh.h
@@ -177,6 +177,13 @@ class UParticleModuleTypeDataMesh : public UParticleModuleTypeDataBase
 	UPROPERTY(EditAnywhere, Category=CameraFacing)
 	uint32 bFaceCameraDirectionRatherThanPosition:1;
 
+	/**
+	*	If true, all collisions for mesh particle on this emitter will take the particle size into account.
+	*	If false, particle size will be ignored in collision checks.
+	*/
+	UPROPERTY(EditAnywhere, Category = Collision)
+	uint32 bCollisionsConsiderPartilceSize : 1;
+
 	// Begin UObject Interface
 #if WITH_EDITOR
 	virtual void	PostEditChangeProperty(FPropertyChangedEvent& PropertyChangedEvent) override;
diff --git a/Engine/Source/Runtime/Engine/Private/Particles/ParticleModules.cpp b/Engine/Source/Runtime/Engine/Private/Particles/ParticleModules.cpp
--- a/Engine/Source/Runtime/Engine/Private/Particles/ParticleModules.cpp
+++ b/Engine/Source/Runtime/Engine/Private/Particles/ParticleModules.cpp
@@ -2966,6 +2966,7 @@ UParticleModuleTypeDataMesh::UParticleModuleTypeDataMesh(const FObjectInitialize
 	AxisLockOption = EPAL_NONE;
 	CameraFacingUpAxisOption_DEPRECATED = CameraFacing_NoneUP;
 	CameraFacingOption = XAxisFacing_NoUp;
+	bCollisionsConsiderPartilceSize = true;
 }
 
 FParticleEmitterInstance* UParticleModuleTypeDataMesh::CreateInstance(UParticleEmitter* InEmitterParent, UParticleSystemComponent* InComponent)
@@ -3001,6 +3002,10 @@ void UParticleModuleTypeDataMesh::Serialize(FArchive& Ar)
 		RPYDistribution->Max = oldOrient;
 		RPYDistribution->bIsDirty = true;
 	}
+	if (Ar.IsLoading() && Ar.UE4Ver() < VER_UE4_MESH_PARTICLE_COLLISIONS_CONSIDER_PARTICLE_SIZE)
+	{
+		bCollisionsConsiderPartilceSize = false;//Old data should default to the old behavior of not considering particle size.
+	}
 }
 
 void UParticleModuleTypeDataMesh::CreateDistribution()
diff --git a/Engine/Source/Runtime/Engine/Private/Particles/ParticleModules_Collision.cpp b/Engine/Source/Runtime/Engine/Private/Particles/ParticleModules_Collision.cpp
--- a/Engine/Source/Runtime/Engine/Private/Particles/ParticleModules_Collision.cpp
+++ b/Engine/Source/Runtime/Engine/Private/Particles/ParticleModules_Collision.cpp
@@ -347,6 +347,7 @@ void UParticleModuleCollision::Update(FParticleEmitterInstance* Owner, int32 Off
 		if (MeshType && MeshType->Mesh)
 		{
 			Extent = MeshType->Mesh->GetBounds().BoxExtent;
+			Extent = MeshType->bCollisionsConsiderPartilceSize ? Extent * Size : Extent;
 		}
 		
 		FHitResult Hit;
