PR #1295: Properly deform navigable geometry for spline mesh (Contributed by [REDACTED_PERSON])

diff --git a/Engine/Source/Runtime/Engine/Private/Components/SplineMeshComponent.cpp b/Engine/Source/Runtime/Engine/Private/Components/SplineMeshComponent.cpp
--- a/Engine/Source/Runtime/Engine/Private/Components/SplineMeshComponent.cpp
+++ b/Engine/Source/Runtime/Engine/Private/Components/SplineMeshComponent.cpp
@@ -794,13 +794,16 @@ bool USplineMeshComponent::DoCustomNavigableGeometryExport(FNavigableGeometryExp
 		{
 			if (NavCollision->bHasConvexGeometry)
 			{
+				FVector Mask = FVector(1, 1, 1);
+				GetAxisValue(Mask, ForwardAxis) = 0;
+
 				TArray<FVector> VertexBuffer;
 				VertexBuffer.Reserve(FMath::Max(NavCollision->ConvexCollision.VertexBuffer.Num(), NavCollision->TriMeshCollision.VertexBuffer.Num()));
 
 				for (int32 i = 0; i < NavCollision->ConvexCollision.VertexBuffer.Num(); ++i)
 				{
 					FVector Vertex = NavCollision->ConvexCollision.VertexBuffer[i];
-					Vertex = CalcSliceTransform(GetAxisValue(Vertex, ForwardAxis)).TransformPosition(Vertex);
+					Vertex = CalcSliceTransform(GetAxisValue(Vertex, ForwardAxis)).TransformPosition(Vertex * Mask);
 					VertexBuffer.Add(Vertex);
 				}
 				GeomExport.ExportCustomMesh(VertexBuffer.GetData(), VertexBuffer.Num(),
@@ -811,7 +814,7 @@ bool USplineMeshComponent::DoCustomNavigableGeometryExport(FNavigableGeometryExp
 				for (int32 i = 0; i < NavCollision->TriMeshCollision.VertexBuffer.Num(); ++i)
 				{
 					FVector Vertex = NavCollision->TriMeshCollision.VertexBuffer[i];
-					Vertex = CalcSliceTransform(GetAxisValue(Vertex, ForwardAxis)).TransformPosition(Vertex);
+					Vertex = CalcSliceTransform(GetAxisValue(Vertex, ForwardAxis)).TransformPosition(Vertex * Mask);
 					VertexBuffer.Add(Vertex);
 				}
 				GeomExport.ExportCustomMesh(VertexBuffer.GetData(), VertexBuffer.Num(),
