Consider wireframe materials to not write every pixel. This allows passes like pre-shadow stencil to render with correct rasterizer state. UE-17611

diff --git a/Engine/Source/Runtime/Engine/Public/MaterialShared.h b/Engine/Source/Runtime/Engine/Public/MaterialShared.h
--- a/Engine/Source/Runtime/Engine/Public/MaterialShared.h
+++ b/Engine/Source/Runtime/Engine/Public/MaterialShared.h
@@ -1109,7 +1109,7 @@ public:
 	/* Helper function to look at both IsMasked and IsDitheredLODTransition to determine if it writes every pixel */
 	ENGINE_API bool WritesEveryPixel() const
 	{
-		return !IsMasked() && !IsDitheredLODTransition();
+		return !IsMasked() && !IsDitheredLODTransition() && !IsWireframe();
 	}
 
 	/** 
