UE4 - fix crash related to looking for crash UE-17254 Undo causes crash in TScopedPointer<FStaticMeshRenderData>

diff --git a/Engine/Source/Runtime/Engine/Private/Components/StaticMeshComponent.cpp b/Engine/Source/Runtime/Engine/Private/Components/StaticMeshComponent.cpp
--- a/Engine/Source/Runtime/Engine/Private/Components/StaticMeshComponent.cpp
+++ b/Engine/Source/Runtime/Engine/Private/Components/StaticMeshComponent.cpp
@@ -1063,7 +1063,7 @@ void UStaticMeshComponent::PostEditUndo()
 	InitResources();
 
 	// Debug check command trying to track down undo related uninitialized resource
-	if (StaticMesh != NULL && StaticMesh->RenderData->LODResources.Num() > 0)
+	if (StaticMesh != NULL && StaticMesh->RenderData.IsValid() && StaticMesh->RenderData->LODResources.Num() > 0)
 	{
 		ENQUEUE_UNIQUE_RENDER_COMMAND_ONEPARAMETER(
 			ResourceCheckCommand,
@@ -1073,7 +1073,6 @@ void UStaticMeshComponent::PostEditUndo()
 			}
 		);
 	}
-
 	Super::PostEditUndo();
 }
 
