Revised comment on interp to component bPauseOnImpact to better reflect functionality.

diff --git a/Engine/Source/Runtime/Engine/Classes/Components/InterpToMovementComponent.h b/Engine/Source/Runtime/Engine/Classes/Components/InterpToMovementComponent.h
--- a/Engine/Source/Runtime/Engine/Classes/Components/InterpToMovementComponent.h
+++ b/Engine/Source/Runtime/Engine/Classes/Components/InterpToMovementComponent.h
@@ -70,7 +70,7 @@ class ENGINE_API UInterpToMovementComponent : public UMovementComponent
 	UPROPERTY(EditAnywhere, BlueprintReadWrite, Category=Control,meta=(UIMin=0.1f, ClampMin=0.1f))
 	float Duration;
 	
-	/** If true, will pause movement on impact, otherwise will flip movement direction. */
+	/** If true, will pause movement on impact. If false it will behave as if the end of the movement range was reached based on the BehaviourType. */
 	UPROPERTY(EditAnywhere, BlueprintReadWrite, Category=Behaviour)
 	uint32 bPauseOnImpact:1;
 
