Fixed an indexing crash in FAnimNode_StateMachine::Update

diff --git a/Engine/Source/Runtime/Engine/Private/Animation/AnimNode_StateMachine.cpp b/Engine/Source/Runtime/Engine/Private/Animation/AnimNode_StateMachine.cpp
--- a/Engine/Source/Runtime/Engine/Private/Animation/AnimNode_StateMachine.cpp
+++ b/Engine/Source/Runtime/Engine/Private/Animation/AnimNode_StateMachine.cpp
@@ -309,7 +309,7 @@ void FAnimNode_StateMachine::Update(const FAnimationUpdateContext& Context)
 			const float ExistingWeightOfNextState = GetStateWeight(NextState);
 
 			FAnimationActiveTransitionEntry* PreviousTransitionForNextState = nullptr;
-			for(int32 i = ActiveTransitionArray.Num() - 1 ;  i >= 0 ; ++i)
+			for(int32 i = ActiveTransitionArray.Num() - 1 ;  i >= 0 ; --i)
 			{
 				FAnimationActiveTransitionEntry& TransitionEntry = ActiveTransitionArray[i];
 				if(TransitionEntry.PreviousState == NextState)
