Fix shadow variable by using member variable instead of passing it in

diff --git a/Engine/Source/Runtime/Engine/Classes/Sound/SoundAttenuation.h b/Engine/Source/Runtime/Engine/Classes/Sound/SoundAttenuation.h
--- a/Engine/Source/Runtime/Engine/Classes/Sound/SoundAttenuation.h
+++ b/Engine/Source/Runtime/Engine/Classes/Sound/SoundAttenuation.h
@@ -160,7 +160,7 @@ struct ENGINE_API FAttenuationSettings
 
 private:
 
-	float AttenuationEval(const float Distance, const float Falloff, const float dBAttenuationAtMax) const;
+	float AttenuationEval(const float Distance, const float Falloff) const;
 	float AttenuationEvalBox(const FTransform& SoundLocation, const FVector ListenerLocation) const;
 	float AttenuationEvalCapsule(const FTransform& SoundLocation, const FVector ListenerLocation) const;
 	float AttenuationEvalCone(const FTransform& SoundLocation, const FVector ListenerLocation) const;
diff --git a/Engine/Source/Runtime/Engine/Private/SoundAttenuation.cpp b/Engine/Source/Runtime/Engine/Private/SoundAttenuation.cpp
--- a/Engine/Source/Runtime/Engine/Private/SoundAttenuation.cpp
+++ b/Engine/Source/Runtime/Engine/Private/SoundAttenuation.cpp
@@ -77,7 +77,7 @@ float FAttenuationSettings::GetMaxDimension() const
  *
  * @return Attenuation value (between 0.0 and 1.0)
  */
-float FAttenuationSettings::AttenuationEval(const float Distance, const float Falloff, const float dBAttenuationAtMax) const
+float FAttenuationSettings::AttenuationEval(const float Distance, const float Falloff) const
 {
 	if( Distance > Falloff )
 	{
@@ -123,7 +123,7 @@ float FAttenuationSettings::AttenuationEvalBox(const FTransform& SoundTransform,
 	const float DistanceSq = ComputeSquaredDistanceFromBoxToPoint(-AttenuationShapeExtents, AttenuationShapeExtents,SoundTransform.InverseTransformPositionNoScale(ListenerLocation));
 	if (DistanceSq < FalloffDistance * FalloffDistance)
 	{
-		return AttenuationEval(FMath::Sqrt(DistanceSq), FalloffDistance, dBAttenuationAtMax);
+		return AttenuationEval(FMath::Sqrt(DistanceSq), FalloffDistance);
 	}
 
 	return 0.f;
@@ -149,7 +149,7 @@ float FAttenuationSettings::AttenuationEvalCapsule(const FTransform& SoundTransf
 		Distance = FMath::PointDistToSegment(ListenerLocation, StartPoint, EndPoint) - CapsuleRadius;
 	}
 
-	return AttenuationEval(Distance, FalloffDistance, dBAttenuationAtMax);
+	return AttenuationEval(Distance, FalloffDistance);
 }
 
 float FAttenuationSettings::AttenuationEvalCone(const FTransform& SoundTransform, const FVector ListenerLocation) const
@@ -161,12 +161,12 @@ float FAttenuationSettings::AttenuationEvalCone(const FTransform& SoundTransform
 	const FVector Origin = SoundTransform.GetTranslation() - (SoundForward * ConeOffset);
 
 	const float Distance = FMath::Max(FVector::Dist( Origin, ListenerLocation ) - AttenuationShapeExtents.X, 0.f);
-	VolumeMultiplier *= AttenuationEval(Distance, FalloffDistance, dBAttenuationAtMax);
+	VolumeMultiplier *= AttenuationEval(Distance, FalloffDistance);
 
 	if (VolumeMultiplier > 0.f)
 	{
 		const float theta = FMath::RadiansToDegrees(fabsf(FMath::Acos( FVector::DotProduct(SoundForward, (ListenerLocation - Origin).GetSafeNormal()))));
-		VolumeMultiplier *= AttenuationEval(theta - AttenuationShapeExtents.Y, AttenuationShapeExtents.Z, dBAttenuationAtMax);
+		VolumeMultiplier *= AttenuationEval(theta - AttenuationShapeExtents.Y, AttenuationShapeExtents.Z);
 	}
 
 	return VolumeMultiplier;
@@ -197,7 +197,7 @@ void FAttenuationSettings::ApplyAttenuation( const FTransform& SoundTransform, c
 		case EAttenuationShape::Sphere:
 			{
 				const float Distance = FMath::Max(FVector::Dist( SoundTransform.GetTranslation(), ListenerLocation ) - AttenuationShapeExtents.X, 0.f);
-				Volume *= AttenuationEval(Distance, FalloffDistance, dBAttenuationAtMax);
+				Volume *= AttenuationEval(Distance, FalloffDistance);
 				break;
 			}
 
