UE-17471 Added comment to indicate WindDirectionalSource only affects speedtrees

diff --git a/Engine/Source/Runtime/Engine/Classes/Components/WindDirectionalSourceComponent.h b/Engine/Source/Runtime/Engine/Classes/Components/WindDirectionalSourceComponent.h
--- a/Engine/Source/Runtime/Engine/Classes/Components/WindDirectionalSourceComponent.h
+++ b/Engine/Source/Runtime/Engine/Classes/Components/WindDirectionalSourceComponent.h
@@ -4,6 +4,7 @@
 #pragma once
 #include "WindDirectionalSourceComponent.generated.h"
 
+/** Component that provides a directional wind source. Only affects SpeedTree assets. */
 UCLASS(collapsecategories, hidecategories=(Object, Mobility), editinlinenew)
 class UWindDirectionalSourceComponent : public USceneComponent
 {
diff --git a/Engine/Source/Runtime/Engine/Classes/Engine/WindDirectionalSource.h b/Engine/Source/Runtime/Engine/Classes/Engine/WindDirectionalSource.h
--- a/Engine/Source/Runtime/Engine/Classes/Engine/WindDirectionalSource.h
+++ b/Engine/Source/Runtime/Engine/Classes/Engine/WindDirectionalSource.h
@@ -3,6 +3,7 @@
 #include "GameFramework/Info.h"
 #include "WindDirectionalSource.generated.h"
 
+/** Actor that provides a directional wind source. Only affects SpeedTree assets. */
 UCLASS(ClassGroup=Wind, showcategories=(Rendering, "Utilities|Transformation"))
 class ENGINE_API AWindDirectionalSource : public AInfo
 {
