PR #1335:
Add GetPhysicsLinearVelocityAtPoint, SetCenterOfMass, SetMassOverrideInKg, GetClosestPointOnCollision to PrimitiveComponent
Allow specifying during for PrintString/PrintText
(Contributed by [REDACTED_PERSON])

diff --git a/Engine/Source/Runtime/Engine/Classes/Components/PrimitiveComponent.h b/Engine/Source/Runtime/Engine/Classes/Components/PrimitiveComponent.h
--- a/Engine/Source/Runtime/Engine/Classes/Components/PrimitiveComponent.h
+++ b/Engine/Source/Runtime/Engine/Classes/Components/PrimitiveComponent.h
@@ -874,6 +874,14 @@ public:
 	UFUNCTION(BlueprintCallable, Category="Physics")	
 	FVector GetPhysicsLinearVelocity(FName BoneName = NAME_None);
 
+	/**
+	*	Get the linear velocity of a point on a single body.
+	*	@param Point			Point is specified in world space.
+	*	@param BoneName			If a SkeletalMeshComponent, name of body to get velocity of. 'None' indicates root body.
+	*/
+	UFUNCTION(BlueprintCallable, Category = "Physics")
+	FVector GetPhysicsLinearVelocityAtPoint(FVector Point, FName BoneName = NAME_None);
+
 	/**
 	 *	Set the linear velocity of all bodies in this component.
 	 *
@@ -919,6 +927,14 @@ public:
 	UFUNCTION(BlueprintCallable, Category = "Physics")
 	FVector GetCenterOfMass(FName BoneName = NAME_None);
 
+	/**
+	*	Set the center of mass of a single body. This will offset the physx-calculated center of mass.
+	*	@param CenterOfMassOffset		User specified offset for the center of mass of this object, from the calculated location.
+	*	@param BoneName			If a SkeletalMeshComponent, name of body to set center of mass of. 'None' indicates root body.
+	*/
+	UFUNCTION(BlueprintCallable, Category = "Physics")
+	void SetCenterOfMass(FVector CenterOfMassOffset, FName BoneName = NAME_None);
+
 	/**
 	 *	'Wake' physics simulation for a single body.
 	 *	@param	BoneName	If a SkeletalMeshComponent, name of body to wake. 'None' indicates root body.
@@ -1209,6 +1225,20 @@ public:
 	 */
 	virtual float GetDistanceToCollision(const FVector& Point, FVector& ClosestPointOnCollision) const;
 
+	/**
+	* Returns the distance and closest point to the collision surface.
+	* Component must have simple collision to be queried for closest point.
+	*
+	* @param Point				World 3D vector
+	* @param OutPointOnBody		Point on the surface of collision closest to Point
+	* @param BoneName			If a SkeletalMeshComponent, name of body to set center of mass of. 'None' indicates root body.
+	*
+	* @return		Success if returns > 0.f, if returns 0.f, it is either not convex or inside of the point
+	*				If returns < 0.f, this primitive does not have collsion
+	*/
+	UFUNCTION(BlueprintCallable, Category = "Collision")
+	float GetClosestPointOnCollision(const FVector& Point, FVector& OutPointOnBody, FName BoneName = NAME_None) const;
+
 	/**
 	 * Creates a proxy to represent the primitive to the scene manager in the rendering thread.
 	 * @return The proxy object.
@@ -1539,6 +1569,13 @@ public:
 	UFUNCTION(BlueprintCallable, Category="Physics")
 	virtual void SetAllMassScale(float InMassScale = 1.f);
 
+	/**
+	*	Override the mass (in Kg) of a single physics body.
+	*	Set the Override Mass to false if you want to reset the body's mass to the auto-calculated physx mass.
+	*/
+	UFUNCTION(BlueprintCallable, Category = "Physics")
+	virtual void SetMassOverrideInKg(FName BoneName = NAME_None, float MassInKg = 1.f, bool bOverrideMass = true);
+
 	/** Returns the mass of this component in kg. */
 	UFUNCTION(BlueprintCallable, Category="Physics")
 	virtual float GetMass() const;
diff --git a/Engine/Source/Runtime/Engine/Classes/Kismet/KismetSystemLibrary.h b/Engine/Source/Runtime/Engine/Classes/Kismet/KismetSystemLibrary.h
--- a/Engine/Source/Runtime/Engine/Classes/Kismet/KismetSystemLibrary.h
+++ b/Engine/Source/Runtime/Engine/Classes/Kismet/KismetSystemLibrary.h
@@ -218,9 +218,10 @@ class ENGINE_API UKismetSystemLibrary : public UBlueprintFunctionLibrary
 	 * @param	bPrintToLog		Whether or not to print the output to the log
 	 * @param	bPrintToConsole	Whether or not to print the output to the console
 	 * @param	TextColor		Whether or not to print the output to the console
+	 * @param	Duration		The display duration (if Print to Screen is True). Using negative number will result in loading the duration time from the config.
 	 */
 	UFUNCTION(BlueprintCallable, meta=(WorldContext="WorldContextObject", CallableWithoutWorldContext, Keywords = "log print", AdvancedDisplay = "2"), Category="Utilities|String")
-	static void PrintString(UObject* WorldContextObject, const FString& InString = FString(TEXT("Hello")), bool bPrintToScreen = true, bool bPrintToLog = true, FLinearColor TextColor = FLinearColor(0.0,0.66,1.0));
+	static void PrintString(UObject* WorldContextObject, const FString& InString = FString(TEXT("Hello")), bool bPrintToScreen = true, bool bPrintToLog = true, FLinearColor TextColor = FLinearColor(0.0, 0.66, 1.0), float Duration = 2.f);
 
 	/**
 	 * Prints text to the log, and optionally, to the screen
@@ -231,9 +232,10 @@ class ENGINE_API UKismetSystemLibrary : public UBlueprintFunctionLibrary
 	 * @param	bPrintToLog		Whether or not to print the output to the log
 	 * @param	bPrintToConsole	Whether or not to print the output to the console
 	 * @param	TextColor		Whether or not to print the output to the console
+	 * @param	Duration		The display duration (if Print to Screen is True). Using negative number will result in loading the duration time from the config.
 	 */
 	UFUNCTION(BlueprintCallable, meta=(WorldContext="WorldContextObject", CallableWithoutWorldContext, Keywords = "log", AdvancedDisplay = "2"), Category="Utilities|Text")
-	static void PrintText(UObject* WorldContextObject, const FText InText = FText::FromString(TEXT("Hello")), bool bPrintToScreen = true, bool bPrintToLog = true, FLinearColor TextColor = FLinearColor(0.0,0.66,1.0));
+	static void PrintText(UObject* WorldContextObject, const FText InText = FText::FromString(TEXT("Hello")), bool bPrintToScreen = true, bool bPrintToLog = true, FLinearColor TextColor = FLinearColor(0.0, 0.66, 1.0), float Duration = 2.f);
 
 	/**
 	 * Prints a warning string to the log and the screen. Meant to be used as a way to inform the user that they misused the node.
diff --git a/Engine/Source/Runtime/Engine/Private/KismetSystemLibrary.cpp b/Engine/Source/Runtime/Engine/Private/KismetSystemLibrary.cpp
--- a/Engine/Source/Runtime/Engine/Private/KismetSystemLibrary.cpp
+++ b/Engine/Source/Runtime/Engine/Private/KismetSystemLibrary.cpp
@@ -164,7 +164,7 @@ UObject* UKismetSystemLibrary::Conv_InterfaceToObject(const FScriptInterface& In
 	return Interface.GetObject();
 }
 
-void UKismetSystemLibrary::PrintString(UObject* WorldContextObject, const FString& InString, bool bPrintToScreen, bool bPrintToLog, FLinearColor TextColor)
+void UKismetSystemLibrary::PrintString(UObject* WorldContextObject, const FString& InString, bool bPrintToScreen, bool bPrintToLog, FLinearColor TextColor, float Duration)
 {
 #if !(UE_BUILD_SHIPPING || UE_BUILD_TEST) // Do not Print in Shipping or Test
 
@@ -220,8 +220,7 @@ void UKismetSystemLibrary::PrintString(UObject* WorldContextObject, const FStrin
 	{
 		if (GAreScreenMessagesEnabled)
 		{
-			float Duration = 2.0f;
-			if ( GConfig )
+			if (GConfig && Duration < 0)
 			{
 				GConfig->GetFloat( TEXT("Kismet"), TEXT("PrintStringDuration"), Duration, GEngineIni );
 			}
@@ -235,10 +234,10 @@ void UKismetSystemLibrary::PrintString(UObject* WorldContextObject, const FStrin
 #endif
 }
 
-void UKismetSystemLibrary::PrintText(UObject* WorldContextObject, const FText InText, bool bPrintToScreen, bool bPrintToLog, FLinearColor TextColor)
+void UKismetSystemLibrary::PrintText(UObject* WorldContextObject, const FText InText, bool bPrintToScreen, bool bPrintToLog, FLinearColor TextColor, float Duration)
 {
 #if !(UE_BUILD_SHIPPING || UE_BUILD_TEST) // Do not Print in Shipping or Test
-	PrintString(WorldContextObject, InText.ToString(), bPrintToScreen, bPrintToLog, TextColor);
+	PrintString(WorldContextObject, InText.ToString(), bPrintToScreen, bPrintToLog, TextColor, Duration);
 #endif
 }
 
diff --git a/Engine/Source/Runtime/Engine/Private/PrimitiveComponentPhysics.cpp b/Engine/Source/Runtime/Engine/Private/PrimitiveComponentPhysics.cpp
--- a/Engine/Source/Runtime/Engine/Private/PrimitiveComponentPhysics.cpp
+++ b/Engine/Source/Runtime/Engine/Private/PrimitiveComponentPhysics.cpp
@@ -306,6 +306,16 @@ FVector UPrimitiveComponent::GetPhysicsLinearVelocity(FName BoneName)
 	return FVector(0,0,0);
 }
 
+FVector UPrimitiveComponent::GetPhysicsLinearVelocityAtPoint(FVector Point, FName BoneName)
+{
+	FBodyInstance* BI = GetBodyInstance(BoneName);
+	if (BI != NULL)
+	{
+		return BI->GetUnrealWorldVelocityAtPoint(Point);
+	}
+	return FVector(0, 0, 0);
+}
+
 void UPrimitiveComponent::SetAllPhysicsLinearVelocity(FVector NewVel,bool bAddToCurrent)
 {
 	SetPhysicsLinearVelocity(NewVel, bAddToCurrent, NAME_None);
@@ -349,6 +359,16 @@ FVector UPrimitiveComponent::GetCenterOfMass(FName BoneName)
 	return FVector::ZeroVector;
 }
 
+void UPrimitiveComponent::SetCenterOfMass(FVector CenterOfMassOffset, FName BoneName)
+{
+	if (FBodyInstance* ComponentBodyInstance = GetBodyInstance(BoneName))
+	{
+		WarnInvalidPhysicsOperations(LOCTEXT("SetCenterOfMass", "SetCenterOfMass"), nullptr);
+		ComponentBodyInstance->COMNudge = CenterOfMassOffset;
+		ComponentBodyInstance->UpdateMassProperties();
+	}
+}
+
 void UPrimitiveComponent::SetAllPhysicsAngularVelocity(FVector const& NewAngVel, bool bAddToCurrent)
 {
 	SetPhysicsAngularVelocity(NewAngVel, bAddToCurrent, NAME_None); 
@@ -448,7 +468,7 @@ void UPrimitiveComponent::SetMassScale(FName BoneName, float InMassScale)
 {
 	if (FBodyInstance* BI = GetBodyInstance(BoneName))
 	{
-		WarnInvalidPhysicsOperations(LOCTEXT("SetMassScale", "SetMassScale"), BI);
+		WarnInvalidPhysicsOperations(LOCTEXT("SetMassScale", "SetMassScale"), nullptr);
 		BI->SetMassScale(InMassScale);
 	}
 }
@@ -468,6 +488,18 @@ void UPrimitiveComponent::SetAllMassScale(float InMassScale)
 	SetMassScale(NAME_None, InMassScale);
 }
 
+void UPrimitiveComponent::SetMassOverrideInKg(FName BoneName, float MassInKg, bool bOverrideMass)
+{
+	FBodyInstance* BI = GetBodyInstance(BoneName);
+	if (BI)
+	{
+		WarnInvalidPhysicsOperations(LOCTEXT("SetCenterOfMass", "SetCenterOfMass"), nullptr);
+		BI->bOverrideMass = bOverrideMass;
+		BI->MassInKg = MassInKg;
+		BI->UpdateMassProperties();
+	}
+}
+
 float UPrimitiveComponent::GetMass() const
 {
 	if (FBodyInstance* BI = GetBodyInstance())
@@ -849,6 +881,19 @@ float UPrimitiveComponent::GetDistanceToCollision(const FVector& Point, FVector&
 	return -1.f;
 }
 
+float UPrimitiveComponent::GetClosestPointOnCollision(const FVector& Point, FVector& OutPointOnBody, FName BoneName) const
+{
+	OutPointOnBody = Point;
+
+	FBodyInstance* BodyInst = GetBodyInstance(BoneName);
+	if (BodyInst != NULL)
+	{
+		return BodyInst->GetDistanceToBody(Point, OutPointOnBody);
+	}
+
+	return -1.f;
+}
+
 bool UPrimitiveComponent::IsSimulatingPhysics(FName BoneName) const
 {
 	FBodyInstance* BodyInst = GetBodyInstance(BoneName);
