Saves Scalability state to the same ini file as it was loaded from, so that launching the editor doesn't reset Scalability settings set by the game.

diff --git a/Engine/Source/Runtime/Engine/Private/GameUserSettings.cpp b/Engine/Source/Runtime/Engine/Private/GameUserSettings.cpp
--- a/Engine/Source/Runtime/Engine/Private/GameUserSettings.cpp
+++ b/Engine/Source/Runtime/Engine/Private/GameUserSettings.cpp
@@ -263,7 +263,7 @@ void UGameUserSettings::LoadSettings( bool bForceReload/*=false*/ )
 	LoadConfig(GetClass(), *GGameUserSettingsIni);
 
 
-	// Note: Scalability::LoadState() should not be needed as we already loaed the settings earlier (needed so the engine can startup with that before the game is initialized)
+	// Note: Scalability::LoadState() should not be needed as we already loaded the settings earlier (needed so the engine can startup with that before the game is initialized)
 	ScalabilityQuality = Scalability::GetQualityLevels();
 
 	// Allow override using command-line settings
@@ -287,7 +287,8 @@ void UGameUserSettings::RequestResolutionChange(int32 InResolutionX, int32 InRes
 
 void UGameUserSettings::SaveSettings()
 {
-	Scalability::SaveState(GGameUserSettingsIni);
+	// Save the Scalability state to the same ini file as it was loaded from in FEngineLoop::Preinit
+	Scalability::SaveState(GIsEditor ? GEditorSettingsIni : GGameUserSettingsIni);
 	SaveConfig(CPF_Config, *GGameUserSettingsIni);
 }
 
