UE4 - Parallel Tick Queuing

diff --git a/Engine/Source/Runtime/Engine/Classes/Engine/EngineBaseTypes.h b/Engine/Source/Runtime/Engine/Classes/Engine/EngineBaseTypes.h
--- a/Engine/Source/Runtime/Engine/Classes/Engine/EngineBaseTypes.h
+++ b/Engine/Source/Runtime/Engine/Classes/Engine/EngineBaseTypes.h
@@ -149,29 +149,29 @@ protected:
 public:
 	/** Bool indicating that this function should execute even if the game is paused. Pause ticks are very limited in capabilities. **/
 	UPROPERTY(EditDefaultsOnly, Category="Tick", AdvancedDisplay)
-	uint32 bTickEvenWhenPaused:1;
+	uint8 bTickEvenWhenPaused:1;
 
 	/** If false, this tick function will never be registered and will never tick. Only settable in defaults. */
 	UPROPERTY()
-	uint32 bCanEverTick:1;
+	uint8 bCanEverTick:1;
 
 	/** If true, this tick function will start enabled, but can be disabled later on. */
 	UPROPERTY(EditDefaultsOnly, Category="Tick")
-	uint32 bStartWithTickEnabled:1;
+	uint8 bStartWithTickEnabled:1;
 
 	/** If we allow this tick to run on a dedicated server */
 	UPROPERTY(EditDefaultsOnly, Category="Tick", AdvancedDisplay)
-	uint32 bAllowTickOnDedicatedServer:1;
+	uint8 bAllowTickOnDedicatedServer:1;
 
 	/** Run this tick first within the tick group, presumably to start async tasks that must be completed with this tick group, hiding the latency. */
-	uint32 bHighPriority:1;
+	uint8 bHighPriority:1;
 
 	/** If false, this tick will run on the game thread, otherwise it will run on any thread in parallel with the game thread and in parallel with other "async ticks" **/
-	uint32 bRunOnAnyThread:1;
+	uint8 bRunOnAnyThread:1;
 
 private:
 	/** If true, means that this tick function is in the master array of tick functions **/
-	uint32 bRegistered:1;
+	uint8 bRegistered:1;
 
 	enum class ETickState : uint8
 	{
@@ -193,8 +193,8 @@ private:
 	/** Internal data to track if we have finshed visiting this tick function yet this frame **/
 	int32 TickQueuedGFrameCounter;
 
-	/** Completion handle for the task that will run this tick. Caution, this is no reset to nullptr until an unspecified future time **/
-	FGraphEventRef CompletionHandle;
+	/** Pointer to the task, only used during setup. This is often stale. **/
+	void *TaskPointer;
 
 	/** Prerequisites for this tick function **/
 	TArray<struct FTickPrerequisite> Prerequisites;
@@ -238,13 +238,11 @@ public:
 
 	/**
 	* Gets the current completion handle of this tick function, so it can be delayed until a later point when some additional
-	* tasks have been completed.  Only valid after TG_PreAsyncWork has started and then only until the TickFunction itself has
-	* started to run
+	* tasks have been completed.  Only valid after TG_PreAsyncWork has started and then only until the TickFunction finishes
+	* execution
 	**/
-	FGraphEventRef GetCompletionHandle() const
-	{
-		return CompletionHandle;
-	}
+	FGraphEventRef GetCompletionHandle() const;
+
 	/** 
 	* Gets the action tick group that this function will execute in this frame.
 	* Only valid after TG_PreAsyncWork has started through the end of the frame.
@@ -295,6 +293,13 @@ private:
 	 */
 	void QueueTickFunction(class FTickTaskSequencer& TTS, const struct FTickContext& TickContext);
 
+	/**
+	 * Queues a tick function for execution from the game thread
+	 * @param TickContext - context to tick in
+	 * @param StackForCycleDetection - Stack For Cycle Detection
+	 */
+	void QueueTickFunctionParallel(const struct FTickContext& TickContext, TArray<FTickFunction*, TInlineAllocator<4> >& StackForCycleDetection, struct FTickGroupCompletionItem* AllCompletionEvents, int32 Index, bool bWasInterval);
+
 	/** 
 	 * Logs the prerequisites
 	 */
@@ -321,6 +326,7 @@ private:
 	friend class FTickTaskSequencer;
 	friend class FTickTaskManager;
 	friend class FTickTaskLevel;
+	friend class FTickFunctionTask;
 };
 
 template<>
diff --git a/Engine/Source/Runtime/Engine/Private/TickTaskManager.cpp b/Engine/Source/Runtime/Engine/Private/TickTaskManager.cpp
--- a/Engine/Source/Runtime/Engine/Private/TickTaskManager.cpp
+++ b/Engine/Source/Runtime/Engine/Private/TickTaskManager.cpp
@@ -6,6 +6,7 @@
 
 #include "EnginePrivate.h"
 #include "TickTaskManagerInterface.h"
+#include "ParallelFor.h"
 
 DEFINE_LOG_CATEGORY_STATIC(LogTick, Log, All);
 
@@ -14,16 +15,25 @@ DECLARE_CYCLE_STAT(TEXT("Queue Ticks Wait"),STAT_QueueTicksWait,STATGROUP_Game);
 DECLARE_CYCLE_STAT(TEXT("Queue Tick Task"),STAT_QueueTickTask,STATGROUP_Game);
 DECLARE_CYCLE_STAT(TEXT("Post Queue Tick Task"),STAT_PostTickTask,STATGROUP_Game);
 DECLARE_CYCLE_STAT(TEXT("Cooldown Dequeuing"),STAT_DequeueCooldowns,STATGROUP_Game);
+DECLARE_CYCLE_STAT(TEXT("Gather Ticks for Parallel"),STAT_GatherTicksForParallel,STATGROUP_Game);
+DECLARE_CYCLE_STAT(TEXT("Finalize Parallel Queue"),STAT_FinalizeParallelQueue,STATGROUP_Game);
 DECLARE_CYCLE_STAT(TEXT("Schedule cooldowns"),STAT_ScheduleCooldowns,STATGROUP_Game);
 DECLARE_DWORD_COUNTER_STAT(TEXT("Ticks Queued"),STAT_TicksQueued,STATGROUP_Game);
 DECLARE_CYCLE_STAT(TEXT("TG_NewlySpawned"), STAT_TG_NewlySpawned, STATGROUP_TickGroups);
 
+static TAutoConsoleVariable<float> CVarStallStartFrame(
+	TEXT("tick.StallStartFrame"),
+	0.0f,
+	TEXT("Sleep for the given time in start frame. Time is given in ms. This is a debug option of little value."));
+
 static TAutoConsoleVariable<int32> CVarLogTicks(
-	TEXT("tick.LogTicks"),0,
+	TEXT("tick.LogTicks"),
+	0,
 	TEXT("Spew ticks for debugging."));
 
 static TAutoConsoleVariable<int32> CVarLogTicksShowPrerequistes(
-	TEXT("tick.ShowPrerequistes"),1,
+	TEXT("tick.ShowPrerequistes"),
+	1,
 	TEXT("When logging ticks, show the prerequistes; debugging."));
 
 static TAutoConsoleVariable<int32> CVarAllowAsyncComponentTicks(
@@ -31,6 +41,11 @@ static TAutoConsoleVariable<int32> CVarAllowAsyncComponentTicks(
 	0,
 	TEXT("Used to control async component ticks."));
 
+static TAutoConsoleVariable<int32> CVarAllowConcurrentQueue(
+	TEXT("tick.AllowConcurrentTickQueue"),
+	0,
+	TEXT("If true, queue ticks concurrently."));
+
 struct FTickContext
 {
 	/** Delta time to tick **/
@@ -66,73 +81,72 @@ struct FTickContext
 	}
 };
 
-/**
- * Class that handles the actual tick tasks and starting and completing tick groups
- */
-class FTickTaskSequencer
+/** Helper class define the task of ticking a component **/
+class FTickFunctionTask
 {
-	/** Helper class define the task of ticking a component **/
-	class FTickFunctionTask
-	{
-		/** Actor to tick **/
-		FTickFunction*			Target;
-		/** tick context, here thread is desired execution thread **/
-		FTickContext			Context;
-		/** If true, log each tick **/
-		bool					bLogTick; 
-		/** If true, log prereqs **/
-		bool					bLogTicksShowPrerequistes; 
-	public:
-		/** Constructor
-		 * @param InTarget - Function to tick
-		 * @param InContext - context to tick in, here thread is desired execution thread
-		**/
-		FORCEINLINE FTickFunctionTask(FTickFunction* InTarget, const FTickContext* InContext, bool InbLogTick, bool bInLogTicksShowPrerequistes)
-			: Target(InTarget)
-			, Context(*InContext)
-			, bLogTick(InbLogTick)
-			, bLogTicksShowPrerequistes(bInLogTicksShowPrerequistes)
-		{
-		}
-		FORCEINLINE TStatId GetStatId() const
-		{
-			RETURN_QUICK_DECLARE_CYCLE_STAT(FTickFunctionTask, STATGROUP_TaskGraphTasks);
-		}
-		/** return the thread for this task **/
-		FORCEINLINE ENamedThreads::Type GetDesiredThread()
+	/** Actor to tick **/
+	FTickFunction*			Target;
+	/** tick context, here thread is desired execution thread **/
+	FTickContext			Context;
+	/** If true, log each tick **/
+	bool					bLogTick; 
+	/** If true, log prereqs **/
+	bool					bLogTicksShowPrerequistes; 
+public:
+	/** Constructor
+		* @param InTarget - Function to tick
+		* @param InContext - context to tick in, here thread is desired execution thread
+	**/
+	FORCEINLINE FTickFunctionTask(FTickFunction* InTarget, const FTickContext* InContext, bool InbLogTick, bool bInLogTicksShowPrerequistes)
+		: Target(InTarget)
+		, Context(*InContext)
+		, bLogTick(InbLogTick)
+		, bLogTicksShowPrerequistes(bInLogTicksShowPrerequistes)
+	{
+	}
+	FORCEINLINE TStatId GetStatId() const
+	{
+		RETURN_QUICK_DECLARE_CYCLE_STAT(FTickFunctionTask, STATGROUP_TaskGraphTasks);
+	}
+	/** return the thread for this task **/
+	FORCEINLINE ENamedThreads::Type GetDesiredThread()
+	{
+		if (Target->bHighPriority)
 		{
-			if (Target->bHighPriority)
-			{
-				return ENamedThreads::HiPri(Context.Thread);
-			}
-			return Context.Thread;
-		}
-		FORCEINLINE static ESubsequentsMode::Type GetSubsequentsMode() 
-		{ 
-			return ESubsequentsMode::TrackSubsequents; 
+			return ENamedThreads::HiPri(Context.Thread);
 		}
-		/** 
-		 *	Actually execute the tick.
-		 *	@param	CurrentThread; the thread we are running on
-		 *	@param	MyCompletionGraphEvent; my completion event. Not always useful since at the end of DoWork, you can assume you are done and hence further tasks do not need you as a prerequisite. 
-		 *	However, MyCompletionGraphEvent can be useful for passing to other routines or when it is handy to set up subsequents before you actually do work.
-		 **/
-		void DoTask(ENamedThreads::Type CurrentThread, const FGraphEventRef& MyCompletionGraphEvent)
+		return Context.Thread;
+	}
+	FORCEINLINE static ESubsequentsMode::Type GetSubsequentsMode() 
+	{ 
+		return ESubsequentsMode::TrackSubsequents; 
+	}
+	/** 
+		*	Actually execute the tick.
+		*	@param	CurrentThread; the thread we are running on
+		*	@param	MyCompletionGraphEvent; my completion event. Not always useful since at the end of DoWork, you can assume you are done and hence further tasks do not need you as a prerequisite. 
+		*	However, MyCompletionGraphEvent can be useful for passing to other routines or when it is handy to set up subsequents before you actually do work.
+		**/
+	void DoTask(ENamedThreads::Type CurrentThread, const FGraphEventRef& MyCompletionGraphEvent)
+	{
+		if (bLogTick)
 		{
-			if (bLogTick)
+			UE_LOG(LogTick, Log, TEXT("tick %s %6d %2d %s"), Target->bHighPriority ? TEXT("*") : TEXT(" "),GFrameCounter, (int32)CurrentThread, *Target->DiagnosticMessage());
+			if (bLogTicksShowPrerequistes)
 			{
-				UE_LOG(LogTick, Log, TEXT("tick %s %6d %2d %s"), Target->bHighPriority ? TEXT("*") : TEXT(" "),GFrameCounter, (int32)CurrentThread, *Target->DiagnosticMessage());
-				if (bLogTicksShowPrerequistes)
-				{
-					Target->ShowPrerequistes();
-				}
+				Target->ShowPrerequistes();
 			}
-			Target->ExecuteTick(Context.DeltaSeconds, Context.TickType, CurrentThread, MyCompletionGraphEvent);
-			Target->CompletionHandle = NULL; // Allow the old completion handle to be recycled
 		}
-	};
-
+		Target->ExecuteTick(Context.DeltaSeconds, Context.TickType, CurrentThread, MyCompletionGraphEvent);
+		Target->TaskPointer = nullptr;  // This is stale and a good time to clear it for safety
+	}
+};
 
+/**
+ * Class that handles the actual tick tasks and starting and completing tick groups
+ */
+class FTickTaskSequencer
+{
 	/** Completion handles for each phase of ticks */
 	FGraphEventArray	TickCompletionEvents[TG_MAX];
 
@@ -172,7 +186,7 @@ public:
 	**/
 	FORCEINLINE static bool SingleThreadedMode()
 	{
-		if (IsRunningDedicatedServer() || FPlatformMisc::NumberOfCores() < 3)
+		if (!FApp::ShouldUseThreadingForPerformance() || IsRunningDedicatedServer() || FPlatformMisc::NumberOfCores() < 3)
 		{
 			return true;
 		}
@@ -200,24 +214,38 @@ public:
 		}
 		QUICK_SCOPE_CYCLE_COUNTER(STAT_StartTickTask_ConstructAndDispatchWhenReady);
 
-		TGraphTask<FTickFunctionTask>* Task = TGraphTask<FTickFunctionTask>::CreateTask(Prerequisites, TickContext.Thread).ConstructAndHold(TickFunction, &UseContext, bLogTicks, bLogTicksShowPrerequistes);
-		if (TickFunction->bHighPriority)
+		TickFunction->TaskPointer = TGraphTask<FTickFunctionTask>::CreateTask(Prerequisites, TickContext.Thread).ConstructAndHold(TickFunction, &UseContext, bLogTicks, bLogTicksShowPrerequistes);
+	}
+
+	/** Add a completion handle to a tick group **/
+	FORCEINLINE void AddTickTaskCompletion(ETickingGroup TickGroup, TGraphTask<FTickFunctionTask>* Task, bool bHiPri)
+	{
+		checkSlow(TickGroup >=0 && TickGroup < TG_MAX);
+		if (bHiPri)
 		{
-			HiPriTickTasks[TickFunction->ActualTickGroup].Add(Task);
+			HiPriTickTasks[TickGroup].Add(Task);
 		}
 		else
 		{
-			TickTasks[TickFunction->ActualTickGroup].Add(Task);
+			TickTasks[TickGroup].Add(Task);
 		}
-		TickFunction->CompletionHandle = Task->GetCompletionEvent();
+		new (TickCompletionEvents[TickGroup]) FGraphEventRef(Task->GetCompletionEvent());
 	}
-
-	/** Add a completion handle to a tick group **/
-	FORCEINLINE void AddTickTaskCompletion(ETickingGroup TickGroup, const FGraphEventRef& CompletionHandle)
+	/**
+	 * Start a component tick task and add the completion handle
+	 *
+	 * @param	InPrerequisites - prerequisites that must be completed before this tick can begin
+	 * @param	TickFunction - the tick function to queue
+	 * @param	Context - tick context to tick in. Thread here is the current thread.
+	 */
+	FORCEINLINE void QueueTickTask(const FGraphEventArray* Prerequisites, FTickFunction* TickFunction, const FTickContext& TickContext)
 	{
-		checkSlow(TickGroup >=0 && TickGroup < TG_MAX);
-		new (TickCompletionEvents[TickGroup]) FGraphEventRef(CompletionHandle);
+		checkSlow(TickContext.Thread == ENamedThreads::GameThread);
+		StartTickTask(Prerequisites, TickFunction, TickContext);
+		TGraphTask<FTickFunctionTask>* Task = (TGraphTask<FTickFunctionTask>*)TickFunction->TaskPointer;
+		AddTickTaskCompletion(TickFunction->ActualTickGroup, Task, TickFunction->bHighPriority);
 	}
+
 	/**
 	 * Start a component tick task and add the completion handle
 	 *
@@ -225,11 +253,10 @@ public:
 	 * @param	TickFunction - the tick function to queue
 	 * @param	Context - tick context to tick in. Thread here is the current thread.
 	 */
-	FORCEINLINE void QueueTickTask(const FGraphEventArray* Prerequisites, FTickFunction* TickFunction, const FTickContext& TickContext)
+	FORCEINLINE void QueueTickTaskParallel(const FGraphEventArray* Prerequisites, FTickFunction* TickFunction, const FTickContext& TickContext)
 	{
 		checkSlow(TickContext.Thread == ENamedThreads::GameThread);
 		StartTickTask(Prerequisites, TickFunction, TickContext);
-		AddTickTaskCompletion(TickFunction->ActualTickGroup, TickFunction->CompletionHandle);
 	}
 
 	/** 
@@ -456,6 +483,53 @@ public:
 		return AllEnabledTickFunctions.Num() + CooldownTicksEnabled;
 	}
 
+	/**
+	 * Queues the ticks for this level
+	 * @param InContext - information about the tick
+	 */
+	void StartFrameParallel(const FTickContext& InContext, TArray<FTickFunction*>& AllTickFunctions)
+	{
+		check(!NewlySpawnedTickFunctions.Num()); // There shouldn't be any in here at this point in the frame
+		Context.TickGroup = TG_PrePhysics; // reset this to the start tick group
+		Context.DeltaSeconds = InContext.DeltaSeconds;
+		Context.TickType = InContext.TickType;
+		Context.Thread = ENamedThreads::GameThread;
+		bTickNewlySpawned = true;
+
+		{
+			SCOPE_CYCLE_COUNTER(STAT_DequeueCooldowns);
+			// Determine which cooled down ticks will be enabled this frame
+			float CumulativeCooldown = 0.f;
+			FTickFunction* TickFunction = AllCoolingDownTickFunctions.Head;
+			while (TickFunction)
+			{
+				if (CumulativeCooldown + TickFunction->RelativeTickCooldown >= Context.DeltaSeconds)
+				{
+					TickFunction->RelativeTickCooldown -= (Context.DeltaSeconds - CumulativeCooldown);
+					break;
+				}
+				CumulativeCooldown += TickFunction->RelativeTickCooldown;
+
+				TickFunction->TickState = FTickFunction::ETickState::Enabled;
+				// we store a bit in here if this came from an interval queue 
+				AllTickFunctions.Add((FTickFunction*)(UPTRINT(TickFunction) | 1));
+
+				TickFunctionsToReschedule.Add(FTickScheduleDetails(TickFunction, TickFunction->TickInterval - (Context.DeltaSeconds - CumulativeCooldown))); // Give credit for any overrun
+
+				AllCoolingDownTickFunctions.Head = TickFunction->Next;
+				TickFunction = TickFunction->Next;
+			}
+		}
+		{
+			SCOPE_CYCLE_COUNTER(STAT_GatherTicksForParallel);
+			for (TSet<FTickFunction*>::TIterator It(AllEnabledTickFunctions); It; ++It)
+			{
+				FTickFunction* TickFunction = *It;
+				AllTickFunctions.Add(TickFunction);
+			}
+		}
+	}
+
 	struct FTickScheduleDetails
 	{
 		FTickScheduleDetails(FTickFunction* InTickFunction, const float InCooldown)
@@ -467,29 +541,68 @@ public:
 		float Cooldown;
 	};
 
-	/* Puts a TickFunction in to the cooldown state*/
-	void ScheduleTickFunctionCooldowns(TArray<FTickScheduleDetails>& TickFunctionsToSchedule)
+	void RemoveAndRescheduleForInterval(FTickFunction* TickFunction)
 	{
-		SCOPE_CYCLE_COUNTER(STAT_ScheduleCooldowns);
+		verify(AllEnabledTickFunctions.Remove(TickFunction) == 1);
+		TickFunctionsToReschedule.Add(FTickScheduleDetails(TickFunction, TickFunction->TickInterval));
+	}
 
-		TickFunctionsToSchedule.Sort([](const FTickScheduleDetails& A, const FTickScheduleDetails& B)
+	/* Puts a TickFunction in to the cooldown state*/
+	void ScheduleTickFunctionCooldowns()
+	{
+		if (TickFunctionsToReschedule.Num() > 0)
 		{
-			return A.Cooldown < B.Cooldown;
-		});
+			SCOPE_CYCLE_COUNTER(STAT_ScheduleCooldowns);
 
-		int32 RescheduleIndex = 0;
-		float CumulativeCooldown = 0.f;
-		FTickFunction* PrevComparisonTickFunction = nullptr;
-		FTickFunction* ComparisonTickFunction = AllCoolingDownTickFunctions.Head;
-		while (ComparisonTickFunction && RescheduleIndex < TickFunctionsToSchedule.Num())
-		{
-			const float CooldownTime = TickFunctionsToSchedule[RescheduleIndex].Cooldown;
-			if ((CumulativeCooldown + ComparisonTickFunction->RelativeTickCooldown) > CooldownTime)
+			TickFunctionsToReschedule.Sort([](const FTickScheduleDetails& A, const FTickScheduleDetails& B)
+			{
+				return A.Cooldown < B.Cooldown;
+			});
+
+			int32 RescheduleIndex = 0;
+			float CumulativeCooldown = 0.f;
+			FTickFunction* PrevComparisonTickFunction = nullptr;
+			FTickFunction* ComparisonTickFunction = AllCoolingDownTickFunctions.Head;
+			while (ComparisonTickFunction && RescheduleIndex < TickFunctionsToReschedule.Num())
+			{
+				const float CooldownTime = TickFunctionsToReschedule[RescheduleIndex].Cooldown;
+				if ((CumulativeCooldown + ComparisonTickFunction->RelativeTickCooldown) > CooldownTime)
+				{
+					FTickFunction* TickFunction = TickFunctionsToReschedule[RescheduleIndex].TickFunction;
+					TickFunction->TickState = FTickFunction::ETickState::CoolingDown;
+					TickFunction->RelativeTickCooldown = CooldownTime - CumulativeCooldown;
+
+					if (PrevComparisonTickFunction)
+					{
+						PrevComparisonTickFunction->Next = TickFunction;
+					}
+					else
+					{
+						check(ComparisonTickFunction == AllCoolingDownTickFunctions.Head);
+						AllCoolingDownTickFunctions.Head = TickFunction;
+					}
+					TickFunction->Next = ComparisonTickFunction;
+					PrevComparisonTickFunction = TickFunction;
+					ComparisonTickFunction->RelativeTickCooldown -= TickFunction->RelativeTickCooldown;
+					CumulativeCooldown += TickFunction->RelativeTickCooldown;
+					++RescheduleIndex;
+				}
+				else
+				{
+					CumulativeCooldown += ComparisonTickFunction->RelativeTickCooldown;
+					PrevComparisonTickFunction = ComparisonTickFunction;
+					ComparisonTickFunction = ComparisonTickFunction->Next;
+				}
+			}
+			for ( ; RescheduleIndex < TickFunctionsToReschedule.Num(); ++RescheduleIndex)
 			{
-				FTickFunction* TickFunction = TickFunctionsToSchedule[RescheduleIndex].TickFunction;
+				FTickFunction* TickFunction = TickFunctionsToReschedule[RescheduleIndex].TickFunction;
+				const float CooldownTime = TickFunctionsToReschedule[RescheduleIndex].Cooldown;
+
 				TickFunction->TickState = FTickFunction::ETickState::CoolingDown;
 				TickFunction->RelativeTickCooldown = CooldownTime - CumulativeCooldown;
 
+				TickFunction->Next = nullptr;
 				if (PrevComparisonTickFunction)
 				{
 					PrevComparisonTickFunction->Next = TickFunction;
@@ -499,40 +612,11 @@ public:
 					check(ComparisonTickFunction == AllCoolingDownTickFunctions.Head);
 					AllCoolingDownTickFunctions.Head = TickFunction;
 				}
-				TickFunction->Next = ComparisonTickFunction;
 				PrevComparisonTickFunction = TickFunction;
-				ComparisonTickFunction->RelativeTickCooldown -= TickFunction->RelativeTickCooldown;
-				CumulativeCooldown += TickFunction->RelativeTickCooldown;
-				++RescheduleIndex;
-			}
-			else
-			{
-				CumulativeCooldown += ComparisonTickFunction->RelativeTickCooldown;
-				PrevComparisonTickFunction = ComparisonTickFunction;
-				ComparisonTickFunction = ComparisonTickFunction->Next;
-			}
-		}
-		for ( ; RescheduleIndex < TickFunctionsToSchedule.Num(); ++RescheduleIndex)
-		{
-			FTickFunction* TickFunction = TickFunctionsToSchedule[RescheduleIndex].TickFunction;
-			const float CooldownTime = TickFunctionsToSchedule[RescheduleIndex].Cooldown;
-
-			TickFunction->TickState = FTickFunction::ETickState::CoolingDown;
-			TickFunction->RelativeTickCooldown = CooldownTime - CumulativeCooldown;
 
-			TickFunction->Next = nullptr;
-			if (PrevComparisonTickFunction)
-			{
-				PrevComparisonTickFunction->Next = TickFunction;
-			}
-			else
-			{
-				check(ComparisonTickFunction == AllCoolingDownTickFunctions.Head);
-				AllCoolingDownTickFunctions.Head = TickFunction;
+				CumulativeCooldown += TickFunction->RelativeTickCooldown;
 			}
-			PrevComparisonTickFunction = TickFunction;
-
-			CumulativeCooldown += TickFunction->RelativeTickCooldown;
+			TickFunctionsToReschedule.Reset();
 		}
 	}
 
@@ -568,11 +652,7 @@ public:
 			}
 		}
 
-		if (TickFunctionsToReschedule.Num() > 0)
-		{
-			ScheduleTickFunctionCooldowns(TickFunctionsToReschedule);
-			TickFunctionsToReschedule.Reset();
-		}
+		ScheduleTickFunctionCooldowns();
 	}
 	/**
 	 * Queues the newly spawned ticks for this level
@@ -595,11 +675,7 @@ public:
 				TickFunctionsToReschedule.Add(FTickScheduleDetails(TickFunction, TickFunction->TickInterval));
 			}
 		}
-		if (TickFunctionsToReschedule.Num() > 0)
-		{
-			ScheduleTickFunctionCooldowns(TickFunctionsToReschedule);
-			TickFunctionsToReschedule.Reset();
-		}
+		ScheduleTickFunctionCooldowns();
 		NewlySpawnedTickFunctions.Empty();
 		return Num;
 	}
@@ -626,7 +702,7 @@ public:
 					TickFunction->TickVisitedGFrameCounter = GFrameCounter;
 					TickFunction->TickQueuedGFrameCounter = GFrameCounter;
 					TickFunction->ExecuteTick(InContext.DeltaSeconds, InContext.TickType, ENamedThreads::GameThread, FGraphEventRef());
-					TickFunction->CompletionHandle = NULL; // Allow the old completion handle to be recycled
+					TickFunction->TaskPointer = nullptr; // this is stale, clear it out now
 
 					TickFunctionsToReschedule.Add(FTickScheduleDetails(TickFunction, TickFunction->TickInterval - (InContext.DeltaSeconds - CumulativeCooldown))); // Give credit for any overrun
 				}
@@ -656,13 +732,12 @@ public:
 		for (TSet<FTickFunction*>::TIterator It(AllEnabledTickFunctions); It; ++It)
 		{
 			TickFunction = *It;
-			TickFunction->CompletionHandle = NULL; // might as well NULL this out to allow these handles to be recycled
+			TickFunction->TaskPointer = nullptr; // this is stale, clear it out now
 			if (TickFunction->bTickEvenWhenPaused && TickFunction->TickState == FTickFunction::ETickState::Enabled)
 			{
 				TickFunction->TickVisitedGFrameCounter = GFrameCounter;
 				TickFunction->TickQueuedGFrameCounter = GFrameCounter;
 				TickFunction->ExecuteTick(InContext.DeltaSeconds, InContext.TickType, ENamedThreads::GameThread, FGraphEventRef());
-				TickFunction->CompletionHandle = NULL; // Allow the old completion handle to be recycled
 
 				if (TickFunction->TickInterval > 0.f)
 				{
@@ -672,11 +747,7 @@ public:
 			}
 		}
 
-		if (TickFunctionsToReschedule.Num() > 0)
-		{
-			ScheduleTickFunctionCooldowns(TickFunctionsToReschedule);
-			TickFunctionsToReschedule.Reset();
-		}
+		ScheduleTickFunctionCooldowns();
 
 		check(!NewlySpawnedTickFunctions.Num()); // We don't support new spawns during pause ticks
 	}
@@ -910,6 +981,19 @@ private:
 	bool										bTickNewlySpawned;
 };
 
+/** Helper struct to hold completion items from parallel task. They are moved into a separate place for cache coherency **/
+struct FTickGroupCompletionItem
+{
+	/** Task created **/
+	TGraphTask<FTickFunctionTask>* Task;
+	/** Tick group to complete with **/
+	TEnumAsByte<ETickingGroup>	ActualTickGroup;
+	/** True if this was a misplaced interval tick that we need to deal with **/
+	bool bNeedsToBeRemovedFromTickListsAndRescheduled;
+	/** True if this is hi pri **/
+	bool bHiPri;
+};
+
 /** Class that aggregates the individual levels and deals with parallel tick setup **/
 class FTickTaskManager : public FTickTaskManagerInterface
 {
@@ -947,6 +1031,13 @@ public:
 	virtual void StartFrame(UWorld* InWorld, float InDeltaSeconds, ELevelTick InTickType) override
 	{
 		SCOPE_CYCLE_COUNTER(STAT_QueueTicks);
+#if !UE_BUILD_SHIPPING
+		if (CVarStallStartFrame.GetValueOnGameThread() > 0.0f)
+		{
+			QUICK_SCOPE_CYCLE_COUNTER(STAT_Tick_Intentional_Stall);
+			FPlatformProcess::Sleep(CVarStallStartFrame.GetValueOnGameThread() / 1000.0f);
+		}
+#endif
 		World = InWorld;
 		Context.TickGroup = TG_PrePhysics; // reset this to the start tick group
 		Context.DeltaSeconds = InDeltaSeconds;
@@ -956,16 +1047,71 @@ public:
 		bTickNewlySpawned = true;
 		TickTaskSequencer.StartFrame();
 		FillLevelList();
-		int32 TotalTickFunctions = 0;
-		for( int32 LevelIndex = 0; LevelIndex < LevelList.Num(); LevelIndex++ )
+		
+		int32 NumWorkerThread = 0;
+		bool bConcurrentQueue = false;
+		if (!FTickTaskSequencer::SingleThreadedMode())
 		{
-			TotalTickFunctions += LevelList[LevelIndex]->StartFrame(Context);
+			bConcurrentQueue = !!CVarAllowConcurrentQueue.GetValueOnGameThread();
 		}
-		INC_DWORD_STAT_BY(STAT_TicksQueued, TotalTickFunctions);
 
-		for( int32 LevelIndex = 0; LevelIndex < LevelList.Num(); LevelIndex++ )
+		if (!bConcurrentQueue)
 		{
-			LevelList[LevelIndex]->QueueAllTicks();
+			int32 TotalTickFunctions = 0;
+			for( int32 LevelIndex = 0; LevelIndex < LevelList.Num(); LevelIndex++ )
+			{
+				TotalTickFunctions += LevelList[LevelIndex]->StartFrame(Context);
+			}
+			INC_DWORD_STAT_BY(STAT_TicksQueued, TotalTickFunctions);
+			for( int32 LevelIndex = 0; LevelIndex < LevelList.Num(); LevelIndex++ )
+			{
+				LevelList[LevelIndex]->QueueAllTicks();
+			}
+		}
+		else
+		{
+			for( int32 LevelIndex = 0; LevelIndex < LevelList.Num(); LevelIndex++ )
+			{
+				LevelList[LevelIndex]->StartFrameParallel(Context, AllTickFunctions);
+			}
+			AllCompletionEvents.AddZeroed(AllTickFunctions.Num());
+			INC_DWORD_STAT_BY(STAT_TicksQueued, AllTickFunctions.Num());
+			ParallelFor(AllTickFunctions.Num(), 
+				[this](int32 Index)
+				{
+					// we store a bit in here if this came from an interval queue 
+					FTickFunction* TickFunction = AllTickFunctions[Index];
+					bool bWasInterval = !!(UPTRINT(TickFunction) & 1);
+					TickFunction = (FTickFunction*)(UPTRINT(TickFunction) & ~1);
+
+					TArray<FTickFunction*, TInlineAllocator<4> > StackForCycleDetection;
+					TickFunction->QueueTickFunctionParallel(Context, StackForCycleDetection, &AllCompletionEvents[0], Index, bWasInterval);
+				}
+			);
+
+			{
+				SCOPE_CYCLE_COUNTER(STAT_FinalizeParallelQueue);
+
+				FTickTaskSequencer& TTS = FTickTaskSequencer::Get();
+				for (const FTickGroupCompletionItem& Comp : AllCompletionEvents)
+				{
+					check(Comp.Task);
+					TTS.AddTickTaskCompletion(Comp.ActualTickGroup, Comp.Task, Comp.bHiPri);
+					if (Comp.bNeedsToBeRemovedFromTickListsAndRescheduled)
+					{
+						FTickFunction* TickFunction = AllTickFunctions[&Comp - &AllCompletionEvents[0]];
+						TickFunction->TickTaskLevel->RemoveAndRescheduleForInterval(TickFunction);
+					}
+				}
+			}
+
+			for( int32 LevelIndex = 0; LevelIndex < LevelList.Num(); LevelIndex++ )
+			{
+				LevelList[LevelIndex]->ScheduleTickFunctionCooldowns();
+			}
+
+			AllTickFunctions.Reset();
+			AllCompletionEvents.Reset();
 		}
 	}
 
@@ -1138,6 +1284,9 @@ private:
 	/** true during the tick phase, when true, tick function adds also go to the newly spawned list. **/
 	bool										bTickNewlySpawned;
 
+	/** true during the tick phase, when true, tick function adds also go to the newly spawned list. **/
+	TArray<FTickFunction*> AllTickFunctions;
+	TArray<FTickGroupCompletionItem> AllCompletionEvents;
 };
 
 
@@ -1257,6 +1406,14 @@ void FTickFunction::ShowPrerequistes(int32 Indent)
 	}
 }
 
+FGraphEventRef FTickFunction::GetCompletionHandle() const
+{
+	check(TaskPointer);
+	TGraphTask<FTickFunctionTask>* Task = (TGraphTask<FTickFunctionTask>*)TaskPointer;
+	return Task->GetCompletionEvent();
+}
+
+
 void FTickFunction::QueueTickFunction(FTickTaskSequencer& TTS, const struct FTickContext& TickContext)
 {
 	checkSlow(TickContext.Thread == ENamedThreads::GameThread); // we assume same thread here
@@ -1287,14 +1444,14 @@ void FTickFunction::QueueTickFunction(FTickTaskSequencer& TTS, const struct FTic
 						// this must be up the call stack, therefore this is a cycle
 						UE_LOG(LogTick, Warning, TEXT("While processing prerequisites for %s, could use %s because it would form a cycle."),*DiagnosticMessage(), *Prereq->DiagnosticMessage());
 					}
-					else if (!Prereq->CompletionHandle.GetReference())
+					else if (!Prereq->TaskPointer)
 					{
 						//ok UE_LOG(LogTick, Warning, TEXT("While processing prerequisites for %s, could use %s because it is disabled."),*DiagnosticMessage(), *Prereq->DiagnosticMessage());
 					}
 					else
 					{
 						MaxPrerequisiteTickGroup = FMath::Max<ETickingGroup>(MaxPrerequisiteTickGroup, Prereq->ActualTickGroup);
-						TaskPrerequisites.Add(Prereq->CompletionHandle);
+						TaskPrerequisites.Add(Prereq->GetCompletionHandle());
 					}
 				}
 			}
@@ -1316,6 +1473,99 @@ void FTickFunction::QueueTickFunction(FTickTaskSequencer& TTS, const struct FTic
 	}
 }
 
+void FTickFunction::QueueTickFunctionParallel(const struct FTickContext& TickContext, TArray<FTickFunction*, TInlineAllocator<4> >& StackForCycleDetection, struct FTickGroupCompletionItem* AllCompletionEvents, int32 Index, bool bWasInterval)
+{
+
+	bool bProcessTick;
+
+	int32 OldValue = *(volatile int32*)&TickVisitedGFrameCounter;
+	if (OldValue != GFrameCounter)
+	{
+		OldValue = FPlatformAtomics::InterlockedCompareExchange(&TickVisitedGFrameCounter , GFrameCounter, OldValue);
+	}
+	bProcessTick = OldValue != GFrameCounter;
+
+	if (bProcessTick)
+	{
+		check(bRegistered);
+		if (TickState != FTickFunction::ETickState::Disabled)
+		{
+			ETickingGroup MaxPrerequisiteTickGroup = ETickingGroup(0);
+
+			FGraphEventArray TaskPrerequisites;
+			if (Prerequisites.Num())
+			{
+				StackForCycleDetection.Push(this);
+				for (int32 PrereqIndex = 0; PrereqIndex < Prerequisites.Num(); PrereqIndex++)
+				{
+					FTickFunction* Prereq = Prerequisites[PrereqIndex].Get();
+					if (!Prereq)
+					{
+						// stale prereq, delete it
+						Prerequisites.RemoveAtSwap(PrereqIndex--);
+					}
+					else if (Prereq->bRegistered)
+					{
+						// recursive call to make sure my prerequisite is set up so I can use its completion handle
+						Prereq->QueueTickFunctionParallel(TickContext, StackForCycleDetection, AllCompletionEvents, -1, false);
+						if (!Prereq->TaskPointer)
+						{
+							//ok UE_LOG(LogTick, Warning, TEXT("While processing prerequisites for %s, could use %s because it is disabled."),*DiagnosticMessage(), *Prereq->DiagnosticMessage());
+						}
+						else
+						{
+							MaxPrerequisiteTickGroup = FMath::Max<ETickingGroup>(MaxPrerequisiteTickGroup, Prereq->ActualTickGroup);
+							TaskPrerequisites.Add(Prereq->GetCompletionHandle());
+						}
+					}
+				}
+				StackForCycleDetection.Pop();
+			}
+
+			// tick group is the max of the prerequisites, the current tick group, and the desired tick group
+			ETickingGroup MyActualTickGroup =  FMath::Max<ETickingGroup>(MaxPrerequisiteTickGroup, FMath::Max<ETickingGroup>(TickGroup,TickContext.TickGroup));
+			if (MyActualTickGroup == TG_DuringPhysics && TickGroup != TG_DuringPhysics)
+			{
+				MyActualTickGroup = ETickingGroup(MyActualTickGroup + 1); // if the tick was "promoted" to during async, but it was not designed for that, then it needs to go later
+			}
+			ActualTickGroup = MyActualTickGroup;
+
+			if (TickState == FTickFunction::ETickState::Enabled)
+			{
+				FTickTaskSequencer::Get().QueueTickTaskParallel(&TaskPrerequisites, this, TickContext);
+			}
+		}
+		FPlatformMisc::MemoryBarrier();
+		TickQueuedGFrameCounter = GFrameCounter;
+	}
+	else
+	{
+		// if we are not going to process it, we need to at least wait until the other thread finishes it
+		volatile int32* TickQueuedGFrameCounterPtr = &TickQueuedGFrameCounter;
+		if (*TickQueuedGFrameCounterPtr != GFrameCounter)
+		{
+			QUICK_SCOPE_CYCLE_COUNTER(STAT_FTickFunction_QueueTickFunctionParallel_Spin);
+			while (*TickQueuedGFrameCounterPtr != GFrameCounter)
+			{
+				FPlatformProcess::SleepNoStats(0);
+			}
+		}
+	}
+	if (Index >= 0)
+	{
+		TGraphTask<FTickFunctionTask>* Task = (TGraphTask<FTickFunctionTask>*)TaskPointer;
+		check(Task);
+		FTickGroupCompletionItem& Comp = *(AllCompletionEvents + Index);
+		check(!Comp.Task);
+		Comp.Task = Task;
+		Comp.ActualTickGroup = ActualTickGroup;
+		if (!bWasInterval && TickInterval > 0.f && TickState == FTickFunction::ETickState::Enabled)
+		{
+			Comp.bNeedsToBeRemovedFromTickListsAndRescheduled = true;
+		}
+		Comp.bHiPri = bHighPriority;
+	}
+}
 
 /**
  * Singleton to retrieve the global tick task manager
