Merging parallel particle fixes from orion branch

diff --git a/Engine/Source/Editor/UnrealEd/Private/Settings/EditorPerProjectUserSettings.cpp b/Engine/Source/Editor/UnrealEd/Private/Settings/EditorPerProjectUserSettings.cpp
--- a/Engine/Source/Editor/UnrealEd/Private/Settings/EditorPerProjectUserSettings.cpp
+++ b/Engine/Source/Editor/UnrealEd/Private/Settings/EditorPerProjectUserSettings.cpp
@@ -32,7 +32,7 @@ void UEditorPerProjectUserSettings::PostEditChangeProperty( FPropertyChangedEven
 
 	if (Name == FName(TEXT("bUseCurvesForDistributions")))
 	{
-		extern CORE_API uint32 GDistributionType;
+		extern ENGINE_API uint32 GDistributionType;
 		//GDistributionType == 0 for curves
 		GDistributionType = (bUseCurvesForDistributions) ? 0 : 1;
 	}
diff --git a/Engine/Source/Runtime/Engine/Classes/Distributions/Distribution.h b/Engine/Source/Runtime/Engine/Classes/Distributions/Distribution.h
--- a/Engine/Source/Runtime/Engine/Classes/Distributions/Distribution.h
+++ b/Engine/Source/Runtime/Engine/Classes/Distributions/Distribution.h
@@ -5,6 +5,8 @@
 #include "Math/CurveEdInterface.h"
 #include "Distribution.generated.h"
 
+extern ENGINE_API uint32 GDistributionType;
+
 UENUM()
 enum DistributionParamMode
 {
diff --git a/Engine/Source/Runtime/Engine/Classes/Distributions/DistributionFloat.h b/Engine/Source/Runtime/Engine/Classes/Distributions/DistributionFloat.h
--- a/Engine/Source/Runtime/Engine/Classes/Distributions/DistributionFloat.h
+++ b/Engine/Source/Runtime/Engine/Classes/Distributions/DistributionFloat.h
@@ -42,32 +42,49 @@ public:
 	}
 
 
-	#if WITH_EDITOR
-		/**`
-		 * Initialize a raw distribution from the original Unreal distribution
-		 */
-		void Initialize();
-	#endif
+#if WITH_EDITOR
+	/**`
+		* Initialize a raw distribution from the original Unreal distribution
+		*/
+	void Initialize();
+#endif
 			 
-		/**
-		 * Gets a pointer to the raw distribution if you can just call FRawDistribution::GetValue1 on it, otherwise NULL 
-		 */
-		const FRawDistribution* GetFastRawDistribution();
-
-		/**
-		 * Get the value at the specified F
-		 */
-		ENGINE_API float GetValue(float F=0.0f, UObject* Data=NULL, struct FRandomStream* InRandomStream = NULL);
-
-		/**
-		 * Get the min and max values
-		 */
-		void GetOutRange(float& MinOut, float& MaxOut);
-
-		/**
-		 * Is this distribution a uniform type? (ie, does it have two values per entry?)
-		 */
-		inline bool IsUniform() { return LookupTable.SubEntryStride != 0; }
+	/**
+		* Gets a pointer to the raw distribution if you can just call FRawDistribution::GetValue1 on it, otherwise NULL 
+		*/
+	const FRawDistribution* GetFastRawDistribution();
+
+	/**
+		* Get the value at the specified F
+		*/
+	ENGINE_API float GetValue(float F=0.0f, UObject* Data=NULL, struct FRandomStream* InRandomStream = NULL);
+
+	/**
+		* Get the min and max values
+		*/
+	void GetOutRange(float& MinOut, float& MaxOut);
+
+	/**
+		* Is this distribution a uniform type? (ie, does it have two values per entry?)
+		*/
+	inline bool IsUniform() { return LookupTable.SubEntryStride != 0; }
+
+	void InitLookupTable();
+
+	FORCEINLINE bool HasLookupTable()
+	{
+#if WITH_EDITOR
+		InitLookupTable();
+#endif
+		return GDistributionType != 0 && !LookupTable.IsEmpty();
+	}
+
+	FORCEINLINE bool OkForParallel()
+	{
+		HasLookupTable(); // initialize if required
+		return true; // even if they stay distributions, this should probably be ok as long as nobody is changing them at runtime
+		//return !Distribution || HasLookupTable();
+	}
 	
 };
 
diff --git a/Engine/Source/Runtime/Engine/Classes/Distributions/DistributionVector.h b/Engine/Source/Runtime/Engine/Classes/Distributions/DistributionVector.h
--- a/Engine/Source/Runtime/Engine/Classes/Distributions/DistributionVector.h
+++ b/Engine/Source/Runtime/Engine/Classes/Distributions/DistributionVector.h
@@ -72,33 +72,49 @@ public:
 	}
 
 
-	#if WITH_EDITOR
-		/**
-		 * Initialize a raw distribution from the original Unreal distribution
-		 */
-		void Initialize();
-	#endif
-
-		/**
-		 * Gets a pointer to the raw distribution if you can just call FRawDistribution::GetValue3 on it, otherwise NULL 
-		 */
-		const FRawDistribution *GetFastRawDistribution();
-
-		/**
-		 * Get the value at the specified F
-		 */
-		ENGINE_API FVector GetValue(float F=0.0f, UObject* Data=NULL, int32 LastExtreme=0, struct FRandomStream* InRandomStream = NULL);
-
-		/**
-		 * Get the min and max values
-		 */
-		void GetOutRange(float& MinOut, float& MaxOut);
-
-		/**
-		 * Is this distribution a uniform type? (ie, does it have two values per entry?)
-		 */
-		inline bool IsUniform() { return LookupTable.SubEntryStride != 0; }
-	
+#if WITH_EDITOR
+	/**
+	* Initialize a raw distribution from the original Unreal distribution
+	*/
+	void Initialize();
+#endif
+
+	/**
+	* Gets a pointer to the raw distribution if you can just call FRawDistribution::GetValue3 on it, otherwise NULL 
+	*/
+	const FRawDistribution *GetFastRawDistribution();
+
+	/**
+	* Get the value at the specified F
+	*/
+	ENGINE_API FVector GetValue(float F=0.0f, UObject* Data=NULL, int32 LastExtreme=0, struct FRandomStream* InRandomStream = NULL);
+
+	/**
+	* Get the min and max values
+	*/
+	void GetOutRange(float& MinOut, float& MaxOut);
+
+	/**
+	* Is this distribution a uniform type? (ie, does it have two values per entry?)
+	*/
+	inline bool IsUniform() { return LookupTable.SubEntryStride != 0; }
+
+	void InitLookupTable();
+
+	FORCEINLINE bool HasLookupTable()
+	{
+#if WITH_EDITOR
+		InitLookupTable();
+#endif
+		return GDistributionType != 0 && !LookupTable.IsEmpty();
+	}
+
+	FORCEINLINE bool OkForParallel()
+	{
+		HasLookupTable(); // initialize if required
+		return true; // even if they stay distributions, this should probably be ok as long as nobody is changing them at runtime
+		//return !Distribution || HasLookupTable();
+	}
 };
 
 UCLASS(abstract, customconstructor)
diff --git a/Engine/Source/Runtime/Engine/Classes/Particles/Camera/ParticleModuleCameraOffset.h b/Engine/Source/Runtime/Engine/Classes/Particles/Camera/ParticleModuleCameraOffset.h
--- a/Engine/Source/Runtime/Engine/Classes/Particles/Camera/ParticleModuleCameraOffset.h
+++ b/Engine/Source/Runtime/Engine/Classes/Particles/Camera/ParticleModuleCameraOffset.h
@@ -56,5 +56,7 @@ class UParticleModuleCameraOffset : public UParticleModuleCameraBase
 	virtual void Update(FParticleEmitterInstance* Owner, int32 Offset, float DeltaTime) override;
 	virtual uint32 RequiredBytes(FParticleEmitterInstance* Owner = NULL) override;
 	//End UParticleModule Interface
+
+	virtual bool CanTickInAnyThread() override;
 };
 
diff --git a/Engine/Source/Runtime/Engine/Classes/Particles/Light/ParticleModuleLight.h b/Engine/Source/Runtime/Engine/Classes/Particles/Light/ParticleModuleLight.h
--- a/Engine/Source/Runtime/Engine/Classes/Particles/Light/ParticleModuleLight.h
+++ b/Engine/Source/Runtime/Engine/Classes/Particles/Light/ParticleModuleLight.h
@@ -70,6 +70,8 @@ class UParticleModuleLight : public UParticleModuleLightBase
 	//End UParticleModule Interface
 
 	void SpawnEx(FParticleEmitterInstance* Owner, int32 Offset, float SpawnTime, struct FRandomStream* InRandomStream, FBaseParticle* ParticleBase);
+
+	virtual bool CanTickInAnyThread() override;
 };
 
 
diff --git a/Engine/Source/Runtime/Engine/Classes/Particles/Material/ParticleModuleMaterialBase.h b/Engine/Source/Runtime/Engine/Classes/Particles/Material/ParticleModuleMaterialBase.h
--- a/Engine/Source/Runtime/Engine/Classes/Particles/Material/ParticleModuleMaterialBase.h
+++ b/Engine/Source/Runtime/Engine/Classes/Particles/Material/ParticleModuleMaterialBase.h
@@ -12,7 +12,7 @@ class UParticleModuleMaterialBase : public UParticleModule
 
 	virtual bool CanTickInAnyThread() override
 	{
-		return false;
+		return true;
 	}
 };
 
diff --git a/Engine/Source/Runtime/Engine/Classes/Particles/Parameter/ParticleModuleParameterDynamic.h b/Engine/Source/Runtime/Engine/Classes/Particles/Parameter/ParticleModuleParameterDynamic.h
--- a/Engine/Source/Runtime/Engine/Classes/Particles/Parameter/ParticleModuleParameterDynamic.h
+++ b/Engine/Source/Runtime/Engine/Classes/Particles/Parameter/ParticleModuleParameterDynamic.h
@@ -161,7 +161,6 @@ class UParticleModuleParameterDynamic : public UParticleModuleParameterBase
 	{
 		float ScaleValue = 1.0f;
 		float DistributionValue = 1.0f;
-		float TimeValue = InDynParams.bUseEmitterTime ? Owner->EmitterTime : Particle.RelativeTime;
 		switch (InDynParams.ValueMethod)
 		{
 		case EDPV_VelocityX:
@@ -179,6 +178,7 @@ class UParticleModuleParameterDynamic : public UParticleModuleParameterBase
 
 		if ((InDynParams.bScaleVelocityByParamValue == true) || (InDynParams.ValueMethod == EDPV_UserSet))
 		{
+			float TimeValue = InDynParams.bUseEmitterTime ? Owner->EmitterTime : Particle.RelativeTime;
 			DistributionValue = InDynParams.ParamValue.GetValue(TimeValue, Owner->Component, InRandomStream);
 		}
 
@@ -204,6 +204,8 @@ class UParticleModuleParameterDynamic : public UParticleModuleParameterBase
 	 *	Set the UpdatesFlags and bUsesVelocity
 	 */
 	virtual	void UpdateUsageFlags();
+
+	virtual bool CanTickInAnyThread() override;
 };
 
 
diff --git a/Engine/Source/Runtime/Engine/Classes/Particles/ParticleEmitter.h b/Engine/Source/Runtime/Engine/Classes/Particles/ParticleEmitter.h
--- a/Engine/Source/Runtime/Engine/Classes/Particles/ParticleEmitter.h
+++ b/Engine/Source/Runtime/Engine/Classes/Particles/ParticleEmitter.h
@@ -313,17 +313,7 @@ class UParticleEmitter : public UObject
 	/**
 	 *   Calculate spawn rate multiplier based on global effects quality level and emitter's quality scale
  	 */
-	FORCEINLINE float GetQualityLevelSpawnRateMult()
-	{
-		float Level = (1 - Scalability::GetQualityLevels().EffectsQuality);
-		float Q = 1;
-		for (int i = 0; i < Level + 1; i++)
-		{
-
-			Q = Q*QualityLevelSpawnRateScale;
-		}
-		return Q;
-	}
+	float GetQualityLevelSpawnRateMult();
 
 	/** Returns true if the is emitter has any enabled LODs, false otherwise. */
 	bool HasAnyEnabledLODs()const;
diff --git a/Engine/Source/Runtime/Engine/Classes/Particles/ParticleSystemComponent.h b/Engine/Source/Runtime/Engine/Classes/Particles/ParticleSystemComponent.h
--- a/Engine/Source/Runtime/Engine/Classes/Particles/ParticleSystemComponent.h
+++ b/Engine/Source/Runtime/Engine/Classes/Particles/ParticleSystemComponent.h
@@ -134,6 +134,8 @@ struct FParticleSysParam
 
 };
 
+template <> struct TIsPODType<FParticleSysParam> { enum { Value = true }; };
+
 /**
  *	The base class for all particle event data.
  */
@@ -444,6 +446,35 @@ public:
 	UPROPERTY(BlueprintAssignable)
 	FOnSystemFinished OnSystemFinished;
 
+private:
+	/** Cached copy of the transform for async work */
+	FTransform AsyncComponentToWorld;
+	/** Cached copy of the instance params */
+	TArray<struct FParticleSysParam> AsyncInstanceParameters;
+	/** Is AsyncComponentToWorld etc valid? */
+	bool bAsyncDataCopyIsValid;
+public:
+
+	const FTransform& GetAsyncComponentToWorld()
+	{
+		if (!IsInGameThread())
+		{
+			check(bAsyncDataCopyIsValid); 
+			return AsyncComponentToWorld;
+		}
+		return ComponentToWorld;
+	}
+
+	const TArray<struct FParticleSysParam>& GetAsyncInstanceParameters()
+	{
+		if (!IsInGameThread())
+		{
+			check(bAsyncDataCopyIsValid); 
+			return AsyncInstanceParameters;
+		}
+		return InstanceParameters;
+	}
+
 	//
 	//	Beam-related script functions.
 	//
@@ -580,6 +611,17 @@ public:
 	 */
 	virtual bool GetVectorParameter(const FName InName, FVector& OutVector);
 
+	/**
+	 *	Retrieve the Vector parameter value for the given name...also looks for colors and floats and returns those
+	 *
+	 *	@param	InName		Name of the parameter
+	 *	@param	OutVector	The value of the parameter found
+	 *
+	 *	@return	true		Parameter was found - OutVector is valid
+	 *			false		Parameter was not found - OutVector is invalid
+	 */
+	virtual bool GetAnyVectorParameter(const FName InName, FVector& OutVector);
+
 	/**
 	 *	Retrieve the Color parameter value for the given name.
 	 *
@@ -722,12 +764,13 @@ public:
 	  */
 	int32 GetCurrentDetailMode() const;
 
-private:
 	/** Possibly parallel phase of TickComponent **/
 	void ComputeTickComponent_Concurrent();
 
 	/** After the possibly parallel phase of TickComponent, we fire events, etc **/
 	void FinalizeTickComponent();
+
+private:
 	/** Wait on the async task and call finalize on the tick **/
 	void WaitForAsyncAndFinalize(EForceAsyncWorkCompletion Behavior) const;
 
diff --git a/Engine/Source/Runtime/Engine/Private/Distributions.cpp b/Engine/Source/Runtime/Engine/Private/Distributions.cpp
--- a/Engine/Source/Runtime/Engine/Private/Distributions.cpp
+++ b/Engine/Source/Runtime/Engine/Private/Distributions.cpp
@@ -54,6 +54,7 @@ UDistribution::UDistribution(const FObjectInitializer& ObjectInitializer)
  */
 static void BuildConstantLookupTable( FDistributionLookupTable* OutTable, int32 ValuesPerEntry, const float* Values )
 {
+	check(IsInGameThread() || IsInAsyncLoadingThread());
 	check( OutTable != NULL );
 	check( Values != NULL );
 
@@ -93,6 +94,7 @@ static void BuildZeroLookupTable( FDistributionLookupTable* OutTable, int32 Valu
 template <typename DistributionType>
 void BuildLookupTable( FDistributionLookupTable* OutTable, const DistributionType* Distribution )
 {
+	check(IsInGameThread() || IsInAsyncLoadingThread());
 	// Always clear the table.
 	OutTable->Empty();
 
@@ -150,6 +152,7 @@ void BuildLookupTable( FDistributionLookupTable* OutTable, const DistributionTyp
  */
 static void AppendLookupTable( FDistributionLookupTable* Table, const FDistributionLookupTable& OtherTable )
 {
+	check(IsInGameThread() || IsInAsyncLoadingThread());
 	check( Table != NULL );
 	check( Table->GetValuesPerEntry() >= 1 && Table->GetValuesPerEntry() <= 3 );
 	check( OtherTable.GetValuesPerEntry() == 1 );
@@ -228,6 +231,7 @@ static void AppendLookupTable( FDistributionLookupTable* Table, const FDistribut
  */
 static void SliceLookupTable( FDistributionLookupTable* Table, int32 ChannelsToKeep )
 {
+	check(IsInGameThread() || IsInAsyncLoadingThread());
 	check( Table != NULL );
 	check( Table->GetValuesPerEntry() >= ChannelsToKeep );
 
@@ -280,6 +284,7 @@ static void SliceLookupTable( FDistributionLookupTable* Table, int32 ChannelsToK
  */
 static void ScaleLookupTableByConstant( FDistributionLookupTable* Table, float Scale )
 {
+	check(IsInGameThread() || IsInAsyncLoadingThread());
 	check( Table != NULL );
 
 	for ( int32 ValueIndex = 0; ValueIndex < Table->Values.Num(); ++ValueIndex )
@@ -296,6 +301,7 @@ static void ScaleLookupTableByConstant( FDistributionLookupTable* Table, float S
  */
 static void ScaleLookupTableByConstants( FDistributionLookupTable* Table, const float* Scale, int32 ValueCount )
 {
+	check(IsInGameThread() || IsInAsyncLoadingThread());
 	check( Table != NULL );
 	check( ValueCount == Table->GetValuesPerEntry() );
 
@@ -328,6 +334,7 @@ static void ScaleLookupTableByConstants( FDistributionLookupTable* Table, const
  */
 static void AddConstantToLookupTable( FDistributionLookupTable* Table, const float* Addend, int32 ValueCount )
 {
+	check(IsInGameThread() || IsInAsyncLoadingThread());
 	check( Table != NULL );
 	check( ValueCount == Table->GetValuesPerEntry() );
 
@@ -359,6 +366,7 @@ static void AddConstantToLookupTable( FDistributionLookupTable* Table, const flo
  */
 static void ScaleLookupTableByLookupTable( FDistributionLookupTable* Table, const FDistributionLookupTable& OtherTable )
 {
+	check(IsInGameThread() || IsInAsyncLoadingThread());
 	check( Table != NULL );
 	check( OtherTable.GetValuesPerEntry() == 1 || OtherTable.GetValuesPerEntry() == Table->GetValuesPerEntry() );
 
@@ -433,6 +441,7 @@ static void ScaleLookupTableByLookupTable( FDistributionLookupTable* Table, cons
  */
 static void AddLookupTableToLookupTable( FDistributionLookupTable* Table, const FDistributionLookupTable& OtherTable )
 {
+	check(IsInGameThread() || IsInAsyncLoadingThread());
 	check( Table != NULL );
 	check( OtherTable.GetValuesPerEntry() == 1 || OtherTable.GetValuesPerEntry() == Table->GetValuesPerEntry() );
 
@@ -587,6 +596,7 @@ static float ComputeLookupTableError( const FDistributionLookupTable& InTable1,
  */
 static void ResampleLookupTable( FDistributionLookupTable* OutTable, const FDistributionLookupTable& InTable, float MinIn, float MaxIn, int32 SampleCount )
 {
+	check(IsInGameThread() || IsInAsyncLoadingThread());
 	const int32 Stride = InTable.EntryStride;
 	const float OneOverTimeScale = (InTable.TimeScale == 0.0f) ? 0.0f : 1.0f / InTable.TimeScale;
 	const float TimeScale = (SampleCount > 1) ? ((MaxIn - MinIn) / (float)(SampleCount-1)) : 0.0f;
@@ -624,6 +634,7 @@ static void ResampleLookupTable( FDistributionLookupTable* OutTable, const FDist
  */
 static void OptimizeLookupTable( FDistributionLookupTable* Table, float ErrorThreshold )
 {
+	check(IsInGameThread() || IsInAsyncLoadingThread());
 	check( Table != NULL );
 	check( (Table->EntryCount & (Table->EntryCount-1)) == 0 );
 
@@ -867,6 +878,7 @@ void FRawDistributionFloat::Initialize()
 	{
 		return;
 	}
+	check(IsInGameThread() || IsInAsyncLoadingThread());
 
 	// always empty out the lookup table
 	LookupTable.Empty();
@@ -896,20 +908,12 @@ void FRawDistributionFloat::Initialize()
 
 float FRawDistributionFloat::GetValue(float F, UObject* Data, struct FRandomStream* InRandomStream)
 {
-#if WITH_EDITOR
-	// make sure it's up to date
-	if( GIsEditor || (Distribution && Distribution->bIsDirty) )
+	if (!HasLookupTable())
 	{
-		Initialize();
-		if (!Distribution && LookupTable.IsEmpty())
+		if (!Distribution)
 		{
 			return 0.0f;
 		}
-	}
-#endif
-	// Use the distribution if the lookup table is empty
-	if ( (LookupTable.IsEmpty() && Distribution) || GDistributionType == 0 )
-	{
 		return Distribution->GetValue(F, Data, InRandomStream);
 	}
 
@@ -923,21 +927,7 @@ float FRawDistributionFloat::GetValue(float F, UObject* Data, struct FRandomStre
 
 const FRawDistribution *FRawDistributionFloat::GetFastRawDistribution()
 {
-#if WITH_EDITOR
-	// make sure it's up to date
-	if( GIsEditor )
-	{
-		Initialize();
-	}
-#endif
-
-	if (!IsSimple())
-	{
-		return 0;
-	}
-
-	// Use the distribution if the lookup table is empty
-	if ( (LookupTable.IsEmpty() && Distribution) || GDistributionType == 0  )
+	if (!IsSimple() || !HasLookupTable())
 	{
 		return 0;
 	}
@@ -950,15 +940,7 @@ const FRawDistribution *FRawDistributionFloat::GetFastRawDistribution()
 
 void FRawDistributionFloat::GetOutRange(float& MinOut, float& MaxOut)
 {
-#if WITH_EDITOR
-	// make sure it's up to date
-	if( GIsEditor )
-	{
-		Initialize();
-	}
-#endif
-
-	if (LookupTable.IsEmpty() || GDistributionType == 0)
+	if (!HasLookupTable() || !Distribution)
 	{
 		check(Distribution);
 		Distribution->GetOutRange(MinOut, MaxOut);
@@ -970,6 +952,18 @@ void FRawDistributionFloat::GetOutRange(float& MinOut, float& MaxOut)
 	}
 }
 
+void FRawDistributionFloat::InitLookupTable()
+{
+#if WITH_EDITOR
+	// make sure it's up to date
+	if( GIsEditor || (Distribution && Distribution->bIsDirty) )
+	{
+		Initialize();
+	}
+#endif
+}
+
+
 bool UDistributionFloat::NeedsLoadForClient() const
 {
 	return true;
@@ -1043,6 +1037,7 @@ void FRawDistributionVector::Initialize()
 	{
 		return;
 	}
+	check(IsInGameThread() || IsInAsyncLoadingThread());
 
 	// always empty out the lookup table
 	LookupTable.Empty();
@@ -1074,20 +1069,12 @@ void FRawDistributionVector::Initialize()
 
 FVector FRawDistributionVector::GetValue(float F, UObject* Data, int32 Extreme, struct FRandomStream* InRandomStream)
 {
-#if WITH_EDITOR
-	// make sure it's up to date
-	if( GIsEditor || (Distribution && Distribution->bIsDirty) )
+	if (!HasLookupTable())
 	{
-		Initialize();
-		if (!Distribution && LookupTable.IsEmpty())
+		if (!Distribution)
 		{
 			return FVector::ZeroVector;
 		}
-	}
-#endif
-
-	if ( (LookupTable.IsEmpty() && Distribution) || GDistributionType == 0 )
-	{
 		return Distribution->GetValue(F, Data, Extreme, InRandomStream);
 	}
 
@@ -1101,20 +1088,7 @@ FVector FRawDistributionVector::GetValue(float F, UObject* Data, int32 Extreme,
 
 const FRawDistribution *FRawDistributionVector::GetFastRawDistribution()
 {
-#if WITH_EDITOR
-	// make sure it's up to date
-	if( GIsEditor || (Distribution && Distribution->bIsDirty) )
-	{
-		Initialize();
-	}
-#endif
-
-	if (!IsSimple()) 
-	{
-		return 0;
-	}
-
-	if ( (LookupTable.IsEmpty() && Distribution) || GDistributionType == 0 )
+	if (!IsSimple() || !HasLookupTable()) 
 	{
 		return 0;
 	}
@@ -1127,15 +1101,7 @@ const FRawDistribution *FRawDistributionVector::GetFastRawDistribution()
 
 void FRawDistributionVector::GetOutRange(float& MinOut, float& MaxOut)
 {
-#if WITH_EDITOR
-	// make sure it's up to date
-	if( GIsEditor )
-	{
-		Initialize();
-	}
-#endif
-
-	if (LookupTable.IsEmpty() || GDistributionType == 0)
+	if (!HasLookupTable() && Distribution)
 	{
 		check(Distribution);
 		Distribution->GetOutRange(MinOut, MaxOut);
@@ -1147,6 +1113,18 @@ void FRawDistributionVector::GetOutRange(float& MinOut, float& MaxOut)
 	}
 }
 
+void FRawDistributionVector::InitLookupTable()
+{
+#if WITH_EDITOR
+	// make sure it's up to date
+	if( GIsEditor || (Distribution && Distribution->bIsDirty) )
+	{
+		Initialize();
+	}
+#endif
+}
+
+
 bool UDistributionVector::NeedsLoadForClient() const
 {
 	return true;
diff --git a/Engine/Source/Runtime/Engine/Private/Particles/ParticleBeam2EmitterInstance.cpp b/Engine/Source/Runtime/Engine/Private/Particles/ParticleBeam2EmitterInstance.cpp
--- a/Engine/Source/Runtime/Engine/Private/Particles/ParticleBeam2EmitterInstance.cpp
+++ b/Engine/Source/Runtime/Engine/Private/Particles/ParticleBeam2EmitterInstance.cpp
@@ -1044,6 +1044,7 @@ void FParticleBeam2EmitterInstance::SetupBeamModifierModulesOffsets()
  */
 void FParticleBeam2EmitterInstance::ResolveSource()
 {
+	check(IsInGameThread());
 	if (BeamModule_Source)
 	{
 		if (BeamModule_Source->SourceName != NAME_None)
@@ -1089,6 +1090,7 @@ void FParticleBeam2EmitterInstance::ResolveSource()
  */
 void FParticleBeam2EmitterInstance::ResolveTarget()
 {
+	check(IsInGameThread());
 	if (BeamModule_Target)
 	{
 		if (BeamModule_Target->TargetName != NAME_None)
diff --git a/Engine/Source/Runtime/Engine/Private/Particles/ParticleBeamModules.cpp b/Engine/Source/Runtime/Engine/Private/Particles/ParticleBeamModules.cpp
--- a/Engine/Source/Runtime/Engine/Private/Particles/ParticleBeamModules.cpp
+++ b/Engine/Source/Runtime/Engine/Private/Particles/ParticleBeamModules.cpp
@@ -1669,6 +1669,7 @@ void UParticleModuleBeamSource::PostEditChangeProperty(FPropertyChangedEvent& Pr
 
 void UParticleModuleBeamSource::AutoPopulateInstanceProperties(UParticleSystemComponent* PSysComp)
 {
+	check(IsInGameThread());
 	bool	bFound	= false;
 
 	for (int32 i = 0; i < PSysComp->InstanceParameters.Num(); i++)
@@ -2050,6 +2051,7 @@ void UParticleModuleBeamTarget::PostEditChangeProperty(FPropertyChangedEvent& Pr
 
 void UParticleModuleBeamTarget::AutoPopulateInstanceProperties(UParticleSystemComponent* PSysComp)
 {
+	check(IsInGameThread());
 	bool	bFound	= false;
 
 	for (int32 i = 0; i < PSysComp->InstanceParameters.Num(); i++)
diff --git a/Engine/Source/Runtime/Engine/Private/Particles/ParticleComponents.cpp b/Engine/Source/Runtime/Engine/Private/Particles/ParticleComponents.cpp
--- a/Engine/Source/Runtime/Engine/Private/Particles/ParticleComponents.cpp
+++ b/Engine/Source/Runtime/Engine/Private/Particles/ParticleComponents.cpp
@@ -1501,6 +1501,18 @@ void UParticleEmitter::Build()
 	}
 }
 
+float UParticleEmitter::GetQualityLevelSpawnRateMult()
+{
+	static IConsoleVariable* EffectsQuality = IConsoleManager::Get().FindConsoleVariable(TEXT("sg.EffectsQuality"));
+	float Level = (1 - EffectsQuality->GetInt());
+	float Q = 1;
+	for (int i = 0; i < Level + 1; i++)
+	{
+		Q = Q*QualityLevelSpawnRateScale;
+	}
+	return Q;
+}
+
 bool UParticleEmitter::HasAnyEnabledLODs()const
 {
 	for (UParticleLODLevel* LodLevel : LODLevels)
@@ -2548,6 +2560,12 @@ void UParticleSystem::BuildEmitters()
 	}
 }
 
+static bool LogReasoningForAnyThreadTicking()
+{
+	static bool bLogThreadedParticleTicking = FParse::Param(FCommandLine::Get(), TEXT("LogThreadedParticleTicking"));
+	return bLogThreadedParticleTicking;
+}
+
 void UParticleSystem::ComputeCanTickInAnyThread()
 {
 	check(!bIsElligibleForAsyncTickComputed);
@@ -2571,7 +2589,14 @@ void UParticleSystem::ComputeCanTickInAnyThread()
 						if (Module && !Module->CanTickInAnyThread())
 						{
 							bIsElligibleForAsyncTick = false;
-							return;
+							if (LogReasoningForAnyThreadTicking())
+							{
+								UE_LOG(LogParticles, Display, TEXT("Cannot tick %s in parallel because module %s in Emitter %s cannot tick in in parallel."), *GetFullName(), *Module->GetFullName(), *Emitter->GetFullName());
+							}
+							else
+							{
+								return;
+							}
 						}
 					}
 				}
@@ -2579,6 +2604,10 @@ void UParticleSystem::ComputeCanTickInAnyThread()
 
 		}
 	}
+	if (LogReasoningForAnyThreadTicking() && bIsElligibleForAsyncTick)
+	{
+		UE_LOG(LogParticles, Display, TEXT("Can tick %s in parallel."), *GetFullName());
+	}
 }
 
 bool UParticleSystem::ContainsEmitterType(UClass* TypeData)
@@ -2654,6 +2683,7 @@ UParticleSystemComponent::UParticleSystemComponent(const FObjectInitializer& Obj
 #if WITH_EDITOR
 void UParticleSystemComponent::CheckForErrors()
 {
+	check(IsInGameThread());
 	ForceAsyncWorkCompletion(ENSURE_AND_STALL);
 	for (int32 IPIndex = 0; IPIndex < InstanceParameters.Num(); IPIndex++)
 	{
@@ -3343,6 +3373,7 @@ void UParticleSystemComponent::UpdateLODInformation()
 
 void UParticleSystemComponent::OrientZAxisTowardCamera()
 {
+	QUICK_SCOPE_CYCLE_COUNTER(STAT_UParticleSystemComponent_OrientZAxisTowardCamera);
 	ForceAsyncWorkCompletion(ENSURE_AND_STALL);
 
 	//@TODO: CAMERA: How does this work for stereo and/or split-screen?
@@ -3403,6 +3434,7 @@ void UParticleSystemComponent::PreEditChange(UProperty* PropertyThatWillChange)
 
 void UParticleSystemComponent::PostEditChangeChainProperty(FPropertyChangedChainEvent& PropertyChangedEvent)
 {
+	check(IsInGameThread());
 	ForceAsyncWorkCompletion(ENSURE_AND_STALL);
 	if (PropertyChangedEvent.PropertyChain.Num() > 0)
 	{
@@ -3472,6 +3504,63 @@ FBoxSphereBounds UParticleSystemComponent::CalcBounds(const FTransform& LocalToW
 	return FBoxSphereBounds(BoundingBox);
 }
 
+class FParticleFinalizeTask
+{
+	UParticleSystemComponent* Target;
+public:
+	FParticleFinalizeTask(UParticleSystemComponent* InTarget)
+		: Target(InTarget)
+	{
+
+	}
+
+	FORCEINLINE TStatId GetStatId() const
+	{
+		RETURN_QUICK_DECLARE_CYCLE_STAT(FParticleFinalizeTask, STATGROUP_TaskGraphTasks);
+	}
+
+	ENamedThreads::Type GetDesiredThread()
+	{
+		return ENamedThreads::GameThread;
+	}
+
+	static ESubsequentsMode::Type GetSubsequentsMode() { return ESubsequentsMode::TrackSubsequents; }
+
+	void DoTask(ENamedThreads::Type CurrentThread, const FGraphEventRef& MyCompletionGraphEvent)
+	{
+		Target->FinalizeTickComponent();
+	}
+};
+
+class FParticleAsyncTask
+{
+	UParticleSystemComponent* Target;
+public:
+	FParticleAsyncTask(UParticleSystemComponent* InTarget)
+		: Target(InTarget)
+	{
+
+	}
+
+	FORCEINLINE TStatId GetStatId() const
+	{
+		RETURN_QUICK_DECLARE_CYCLE_STAT(FParticleAsyncTask, STATGROUP_TaskGraphTasks);
+	}
+
+	ENamedThreads::Type GetDesiredThread()
+	{
+		return ENamedThreads::AnyThread;
+	}
+
+	static ESubsequentsMode::Type GetSubsequentsMode() { return ESubsequentsMode::TrackSubsequents; }
+
+	void DoTask(ENamedThreads::Type CurrentThread, const FGraphEventRef& MyCompletionGraphEvent)
+	{
+		Target->ComputeTickComponent_Concurrent();
+	}
+};
+
+
 void UParticleSystemComponent::TickComponent(float DeltaTime, enum ELevelTick TickType, FActorComponentTickFunction *ThisTickFunction)
 {
 	ForceAsyncWorkCompletion(ENSURE_AND_STALL);
@@ -3521,6 +3610,7 @@ void UParticleSystemComponent::TickComponent(float DeltaTime, enum ELevelTick Ti
 	} 
 	else
 	{
+		QUICK_SCOPE_CYCLE_COUNTER(STAT_UParticleSystemComponent_CheckForReset);
 		bool bRequiresReset = false;
 		for (int32 EmitterIndex = 0; EmitterIndex < EmitterInstances.Num(); ++EmitterIndex)
 		{
@@ -3556,6 +3646,7 @@ void UParticleSystemComponent::TickComponent(float DeltaTime, enum ELevelTick Ti
 			&&	AccumTickTime > ClampedMaxSecondsBeforeInactive//SecondsBeforeInactive
 			&&	GetWorld()->IsGameWorld() )
 		{
+			QUICK_SCOPE_CYCLE_COUNTER(STAT_UParticleSystemComponent_LOD_Inactive);
 			const float CurrentTimeSeconds = World->GetTimeSeconds();
 			if( CurrentTimeSeconds > (LastRenderTime + ClampedMaxSecondsBeforeInactive) )
 			{
@@ -3573,6 +3664,7 @@ void UParticleSystemComponent::TickComponent(float DeltaTime, enum ELevelTick Ti
 		AccumLODDistanceCheckTime += DeltaTime;
 		if (AccumLODDistanceCheckTime > Template->LODDistanceCheckTime)
 		{
+			QUICK_SCOPE_CYCLE_COUNTER(STAT_UParticleSystemComponent_LOD);
 			AccumLODDistanceCheckTime = 0.0f;
 
 			if (ShouldComputeLODFromGameThread())
@@ -3615,16 +3707,20 @@ void UParticleSystemComponent::TickComponent(float DeltaTime, enum ELevelTick Ti
 		DeltaTime = Template->UpdateTime_Delta;
 	}
 
-	// Clear out the events.
-	SpawnEvents.Empty();
-	DeathEvents.Empty();
-	CollisionEvents.Empty();
-	BurstEvents.Empty();
-	TotalActiveParticles = 0;
-	bNeedsFinalize = true;
-	if (!ThisTickFunction || !CanTickInAnyThread() || FXConsoleVariables::bFreezeParticleSimulation || !FXConsoleVariables::bAllowAsyncTick)
 	{
-		bDisallowAsync = true;
+		QUICK_SCOPE_CYCLE_COUNTER(STAT_UParticleSystemComponent_ResetAndCheckParallel);
+		// Clear out the events.
+		SpawnEvents.Reset();
+		DeathEvents.Reset();
+		CollisionEvents.Reset();
+		BurstEvents.Reset();
+		TotalActiveParticles = 0;
+		bNeedsFinalize = true;
+		if (!ThisTickFunction || !CanTickInAnyThread() || FXConsoleVariables::bFreezeParticleSimulation || !FXConsoleVariables::bAllowAsyncTick ||
+			GDistributionType == 0) // this may not be absolutely required, however if you are using distributions it will be glacial anyway. If you want to get rid of this, note that some modules use this indirectly as their criteria for CanTickInAnyThread
+		{
+			bDisallowAsync = true;
+		}
 	}
 	if (bDisallowAsync)
 	{
@@ -3636,26 +3732,23 @@ void UParticleSystemComponent::TickComponent(float DeltaTime, enum ELevelTick Ti
 	}
 	else
 	{
-		// set up async task and the game thread task to finalize the results.
-		DECLARE_CYCLE_STAT(TEXT("FSimpleDelegateGraphTask.AsyncParticleTick"),
-			STAT_FSimpleDelegateGraphTask_AsyncParticleTick,
-			STATGROUP_TaskGraphTasks);
+		QUICK_SCOPE_CYCLE_COUNTER(STAT_UParticleSystemComponent_QueueTasks);
+		bAsyncDataCopyIsValid = true;
+		AsyncComponentToWorld = ComponentToWorld;
+		AsyncInstanceParameters = InstanceParameters;
 
-		AsyncWork = FSimpleDelegateGraphTask::CreateAndDispatchWhenReady(
-			FSimpleDelegateGraphTask::FDelegate::CreateUObject(this, &UParticleSystemComponent::ComputeTickComponent_Concurrent),
-			GET_STATID(STAT_FSimpleDelegateGraphTask_AsyncParticleTick), NULL, ENamedThreads::AnyThread
-		);
-
-		DECLARE_CYCLE_STAT(TEXT("FSimpleDelegateGraphTask.FinalizeParticleTick"),
-			STAT_FSimpleDelegateGraphTask_FinalizeParticleTick,
-			STATGROUP_TaskGraphTasks);
-
-		FGraphEventRef Finalize = FSimpleDelegateGraphTask::CreateAndDispatchWhenReady(
-			FSimpleDelegateGraphTask::FDelegate::CreateUObject(this, &UParticleSystemComponent::FinalizeTickComponent),
-			GET_STATID(STAT_FSimpleDelegateGraphTask_FinalizeParticleTick), AsyncWork, ENamedThreads::GameThread
-		);
+		{
+			QUICK_SCOPE_CYCLE_COUNTER(STAT_UParticleSystemComponent_QueueAsync);
+			AsyncWork = TGraphTask<FParticleAsyncTask>::CreateTask(nullptr, ENamedThreads::GameThread).ConstructAndDispatchWhenReady(this);
+		}
 
-		ThisTickFunction->GetCompletionHandle()->DontCompleteUntil(Finalize);
+		{
+			QUICK_SCOPE_CYCLE_COUNTER(STAT_UParticleSystemComponent_QueueFinalize);
+			FGraphEventArray Prereqs;
+			Prereqs.Add(AsyncWork);
+			FGraphEventRef Finalize = TGraphTask<FParticleFinalizeTask>::CreateTask(&Prereqs, ENamedThreads::GameThread).ConstructAndDispatchWhenReady(this);
+			ThisTickFunction->GetCompletionHandle()->DontCompleteUntil(Finalize);
+		}
 	}
 }
 
@@ -3717,6 +3810,7 @@ void UParticleSystemComponent::ComputeTickComponent_Concurrent()
 void UParticleSystemComponent::FinalizeTickComponent()
 {
 	SCOPE_CYCLE_COUNTER(STAT_ParticleFinalizeTickTime);
+	bAsyncDataCopyIsValid = false;
 	if (!bNeedsFinalize)
 	{
 		return;
@@ -5045,7 +5139,7 @@ void UParticleSystemComponent::SetFloatParameter(FName Name, float Param)
 	{
 		return;
 	}
-	ForceAsyncWorkCompletion(STALL);
+	check(IsInGameThread());
 
 	// First see if an entry for this name already exists
 	for (int32 i = 0; i < InstanceParameters.Num(); i++)
@@ -5072,7 +5166,7 @@ void UParticleSystemComponent::SetFloatRandParameter(FName ParameterName,float P
 	{
 		return;
 	}
-	ForceAsyncWorkCompletion(STALL);
+	check(IsInGameThread());
 
 	// First see if an entry for this name already exists
 	for (int32 i = 0; i < InstanceParameters.Num(); i++)
@@ -5101,7 +5195,7 @@ void UParticleSystemComponent::SetVectorParameter(FName Name, FVector Param)
 	{
 		return;
 	}
-	ForceAsyncWorkCompletion(STALL);
+	check(IsInGameThread());
 
 	// First see if an entry for this name already exists
 	for (int32 i = 0; i < InstanceParameters.Num(); i++)
@@ -5128,7 +5222,7 @@ void UParticleSystemComponent::SetVectorRandParameter(FName ParameterName,const
 	{
 		return;
 	}
-	ForceAsyncWorkCompletion(STALL);
+	check(IsInGameThread());
 
 	// First see if an entry for this name already exists
 	for (int32 i = 0; i < InstanceParameters.Num(); i++)
@@ -5157,7 +5251,7 @@ void UParticleSystemComponent::SetColorParameter(FName Name, FLinearColor Param)
 	{
 		return;
 	}
-	ForceAsyncWorkCompletion(STALL);
+	check(IsInGameThread());
 
 	FColor NewColor(Param.ToFColor(true));
 
@@ -5186,7 +5280,7 @@ void UParticleSystemComponent::SetActorParameter(FName Name, AActor* Param)
 	{
 		return;
 	}
-	ForceAsyncWorkCompletion(STALL);
+	check(IsInGameThread());
 
 	// First see if an entry for this name already exists
 	for (int32 i = 0; i < InstanceParameters.Num(); i++)
@@ -5213,7 +5307,7 @@ void UParticleSystemComponent::SetMaterialParameter(FName Name, UMaterialInterfa
 	{
 		return;
 	}
-	ForceAsyncWorkCompletion(STALL);
+	check(IsInGameThread());
 
 	// First see if an entry for this name already exists
 	for (int32 i = 0; i < InstanceParameters.Num(); i++)
@@ -5244,9 +5338,10 @@ bool UParticleSystemComponent::GetFloatParameter(const FName InName,float& OutFl
 		return false;
 	}
 
-	for (int32 i = 0; i < InstanceParameters.Num(); i++)
+	const TArray<struct FParticleSysParam>& UseInstanceParameters = GetAsyncInstanceParameters();
+	for (int32 i = 0; i < UseInstanceParameters.Num(); i++)
 	{
-		FParticleSysParam& Param = InstanceParameters[i];
+		const FParticleSysParam& Param = UseInstanceParameters[i];
 		if (Param.Name == InName)
 		{
 			if (Param.ParamType == PSPT_Scalar)
@@ -5256,7 +5351,7 @@ bool UParticleSystemComponent::GetFloatParameter(const FName InName,float& OutFl
 			}
 			else if (Param.ParamType == PSPT_ScalarRand)
 			{
-				check(IsInGameThread());
+				// check(IsInGameThread()); this isn't exactly cool to call from multiple threads, but it isn't terrible.
 				OutFloat = Param.Scalar + (Param.Scalar_Low - Param.Scalar) * FMath::SRand();
 				return true;
 			}
@@ -5275,9 +5370,10 @@ bool UParticleSystemComponent::GetVectorParameter(const FName InName,FVector& Ou
 		return false;
 	}
 
-	for (int32 i = 0; i < InstanceParameters.Num(); i++)
+	const TArray<struct FParticleSysParam>& UseInstanceParameters = GetAsyncInstanceParameters();
+	for (int32 i = 0; i < UseInstanceParameters.Num(); i++)
 	{
-		FParticleSysParam& Param = InstanceParameters[i];
+		const FParticleSysParam& Param = UseInstanceParameters[i];
 		if (Param.Name == InName)
 		{
 			if (Param.ParamType == PSPT_Vector)
@@ -5298,6 +5394,56 @@ bool UParticleSystemComponent::GetVectorParameter(const FName InName,FVector& Ou
 	return false;
 }
 
+bool UParticleSystemComponent::GetAnyVectorParameter(const FName InName,FVector& OutVector)
+{
+	// Always fail if we pass in no name.
+	if(InName == NAME_None)
+	{
+		return false;
+	}
+
+	const TArray<struct FParticleSysParam>& UseInstanceParameters = GetAsyncInstanceParameters();
+	for (int32 i = 0; i < UseInstanceParameters.Num(); i++)
+	{
+		const FParticleSysParam& Param = UseInstanceParameters[i];
+		if (Param.Name == InName)
+		{
+			if (Param.ParamType == PSPT_Vector)
+			{
+				OutVector = Param.Vector;
+				return true;
+			}
+			if (Param.ParamType == PSPT_VectorRand)
+			{
+				check(IsInGameThread());
+				FVector RandValue(FMath::SRand(), FMath::SRand(), FMath::SRand());
+				OutVector = Param.Vector + (Param.Vector_Low - Param.Vector) * RandValue;
+				return true;
+			}
+			if (Param.ParamType == PSPT_Scalar)
+			{
+				float OutFloat = Param.Scalar;
+				OutVector = FVector(OutFloat, OutFloat, OutFloat);
+				return true;
+			}
+			if (Param.ParamType == PSPT_ScalarRand)
+			{
+				// check(IsInGameThread()); this isn't exactly cool to call from multiple threads, but it isn't terrible.
+				float OutFloat = Param.Scalar + (Param.Scalar_Low - Param.Scalar) * FMath::SRand();
+				OutVector = FVector(OutFloat, OutFloat, OutFloat);
+				return true;
+			}
+			if (Param.ParamType == PSPT_Color)
+			{
+				OutVector = FVector(FLinearColor(Param.Color));
+				return true;
+			}
+		}
+	}
+
+	return false;
+}
+
 
 bool UParticleSystemComponent::GetColorParameter(const FName InName,FLinearColor& OutColor)
 {
@@ -5306,10 +5452,11 @@ bool UParticleSystemComponent::GetColorParameter(const FName InName,FLinearColor
 	{
 		return false;
 	}
+	const TArray<struct FParticleSysParam>& UseInstanceParameters = GetAsyncInstanceParameters();
 
-	for (int32 i = 0; i < InstanceParameters.Num(); i++)
+	for (int32 i = 0; i < UseInstanceParameters.Num(); i++)
 	{
-		FParticleSysParam& Param = InstanceParameters[i];
+		const FParticleSysParam& Param = UseInstanceParameters[i];
 		if (Param.Name == InName && Param.ParamType == PSPT_Color)
 		{
 			OutColor = FLinearColor(Param.Color);
@@ -5329,9 +5476,10 @@ bool UParticleSystemComponent::GetActorParameter(const FName InName,class AActor
 		return false;
 	}
 
-	for (int32 i = 0; i < InstanceParameters.Num(); i++)
+	const TArray<struct FParticleSysParam>& UseInstanceParameters = GetAsyncInstanceParameters();
+	for (int32 i = 0; i < UseInstanceParameters.Num(); i++)
 	{
-		FParticleSysParam& Param = InstanceParameters[i];
+		const FParticleSysParam& Param = UseInstanceParameters[i];
 		if (Param.Name == InName && Param.ParamType == PSPT_Actor)
 		{
 			OutActor = Param.Actor;
@@ -5351,9 +5499,10 @@ bool UParticleSystemComponent::GetMaterialParameter(const FName InName,class UMa
 		return false;
 	}
 
-	for (int32 i = 0; i < InstanceParameters.Num(); i++)
+	const TArray<struct FParticleSysParam>& UseInstanceParameters = GetAsyncInstanceParameters();
+	for (int32 i = 0; i < UseInstanceParameters.Num(); i++)
 	{
-		FParticleSysParam& Param = InstanceParameters[i];
+		const FParticleSysParam& Param = UseInstanceParameters[i];
 		if (Param.Name == InName && Param.ParamType == PSPT_Material)
 		{
 			OutMaterial = Param.Material;
@@ -5367,6 +5516,7 @@ bool UParticleSystemComponent::GetMaterialParameter(const FName InName,class UMa
 
 void UParticleSystemComponent::ClearParameter(FName ParameterName, EParticleSysParamType ParameterType)
 {
+	check(IsInGameThread());
 	for (int32 i = 0; i < InstanceParameters.Num(); i++)
 	{
 		if (InstanceParameters[i].Name == ParameterName && (ParameterType == PSPT_None || InstanceParameters[i].ParamType == ParameterType))
diff --git a/Engine/Source/Runtime/Engine/Private/Particles/ParticleEmitterInstances.cpp b/Engine/Source/Runtime/Engine/Private/Particles/ParticleEmitterInstances.cpp
--- a/Engine/Source/Runtime/Engine/Private/Particles/ParticleEmitterInstances.cpp
+++ b/Engine/Source/Runtime/Engine/Private/Particles/ParticleEmitterInstances.cpp
@@ -3349,6 +3349,7 @@ SIZE_T FParticleMeshEmitterInstance::GetResourceSize(EResourceSizeMode::Type Mod
  */
 void FParticleMeshEmitterInstance::SetMeshMaterials( const TArray<UMaterialInterface*>& InMaterials )
 {
+	check(IsInGameThread());
 	CurrentMaterials = InMaterials;
 }
 
diff --git a/Engine/Source/Runtime/Engine/Private/Particles/ParticleGpuSimulation.cpp b/Engine/Source/Runtime/Engine/Private/Particles/ParticleGpuSimulation.cpp
--- a/Engine/Source/Runtime/Engine/Private/Particles/ParticleGpuSimulation.cpp
+++ b/Engine/Source/Runtime/Engine/Private/Particles/ParticleGpuSimulation.cpp
@@ -103,6 +103,7 @@ public:
 	 */
 	uint32 Allocate()
 	{
+		FScopeLock Lock(&CriticalSection);
 		if ( FreeTileCount > 0 )
 		{
 			FreeTileCount--;
@@ -117,6 +118,7 @@ public:
 	 */
 	void Free( uint32 TileIndex )
 	{
+		FScopeLock Lock(&CriticalSection);
 		check( TileIndex < GParticleSimulationTileCount );
 		check( FreeTileCount < GParticleSimulationTileCount );
 		FreeTiles[FreeTileCount] = TileIndex;
@@ -128,6 +130,7 @@ public:
 	 */
 	int32 GetFreeTileCount() const
 	{
+		FScopeLock Lock(&CriticalSection);
 		return FreeTileCount;
 	}
 
@@ -137,6 +140,8 @@ private:
 	uint32 FreeTiles[GParticleSimulationTileCount];
 	/** How many tiles are in the free list. */
 	int32 FreeTileCount;
+
+	mutable FCriticalSection CriticalSection;
 };
 
 /*-----------------------------------------------------------------------------
@@ -3063,7 +3068,6 @@ public:
 
 		FParticleSimulationResources* ParticleSimulationResources = FXSystem->GetParticleSimulationResources();
 		const int32 MinTileCount = GetMinTileCount();
-		const int32 CheckTileCount = ParticleSimulationResources->GetFreeTileCount();
 		int32 NumAllocated = 0;
 		while (AllocatedTiles.Num() < MinTileCount)
 		{
@@ -3081,7 +3085,6 @@ public:
 		}
 		
 #if TRACK_TILE_ALLOCATIONS
-		check(ParticleSimulationResources->GetFreeTileCount() + NumAllocated == CheckTileCount);
 		UE_LOG(LogParticles,VeryVerbose,
 			TEXT("%s|%s|0x%016x [Init] %d tiles"),
 			*Component->GetName(),*Component->Template->GetName(),(PTRINT)this, AllocatedTiles.Num());
@@ -3428,7 +3431,6 @@ private:
 		if (TilesToFree > 0)
 		{
 			FParticleSimulationResources* SimulationResources = FXSystem->GetParticleSimulationResources();
-			const int32 CheckTileCount = SimulationResources->GetFreeTileCount();
 			const int32 FirstTileIndex = AllocatedTiles.Num() - TilesToFree;
 			const int32 LastTileIndex = FirstTileIndex + TilesToFree;
 			for (int32 TileIndex = FirstTileIndex; TileIndex < LastTileIndex; ++TileIndex)
@@ -3439,9 +3441,6 @@ private:
 			AllocatedTiles.RemoveAt(FirstTileIndex, TilesToFree);
 			TileTimeOfDeath.RemoveAt(FirstTileIndex, TilesToFree);
 			Simulation->bDirty_GameThread = true;
-#if TRACK_TILE_ALLOCATIONS
-			check(SimulationResources->GetFreeTileCount() == CheckTileCount + TilesToFree);
-#endif // #if TRACK_TILE_ALLOCATIONS
 		}
 		return TilesToFree;
 	}
@@ -3459,7 +3458,6 @@ private:
 			if (!GIsRequestingExit)
 			{
 				FParticleSimulationResources* ParticleSimulationResources = FXSystem->GetParticleSimulationResources();
-				const int32 CheckTileCount = ParticleSimulationResources->GetFreeTileCount();
 				const int32 TileCount = AllocatedTiles.Num();
 				for ( int32 ActiveTileIndex = 0; ActiveTileIndex < TileCount; ++ActiveTileIndex )
 				{
@@ -3468,7 +3466,6 @@ private:
 				}
 				AllocatedTiles.Reset();
 #if TRACK_TILE_ALLOCATIONS
-				check(ParticleSimulationResources->GetFreeTileCount() == CheckTileCount + TileCount);
 				UE_LOG(LogParticles,VeryVerbose,
 					TEXT("%s|%s|0x%016p [ReleaseSimulationResources] %d tiles"),
 					*Component->GetName(),*Component->Template->GetName(),(PTRINT)this, AllocatedTiles.Num());
diff --git a/Engine/Source/Runtime/Engine/Private/Particles/ParticleModules.cpp b/Engine/Source/Runtime/Engine/Private/Particles/ParticleModules.cpp
--- a/Engine/Source/Runtime/Engine/Private/Particles/ParticleModules.cpp
+++ b/Engine/Source/Runtime/Engine/Private/Particles/ParticleModules.cpp
@@ -276,6 +276,7 @@ void UParticleModule::ChangeEditorColor(FColor& Color, UInterpCurveEdSetup* EdSe
 
 void UParticleModule::AutoPopulateInstanceProperties(UParticleSystemComponent* PSysComp)
 {
+	check(IsInGameThread());
 	for (TFieldIterator<UStructProperty> It(GetClass()); It; ++It)
 	{
 		// attempt to get a distribution from a random struct property
@@ -2788,6 +2789,11 @@ void UParticleModuleLight::PostEditChangeProperty(FPropertyChangedEvent& Propert
 }
 #endif // WITH_EDITOR
 
+bool UParticleModuleLight::CanTickInAnyThread()
+{
+	return BrightnessOverLife.OkForParallel() && ColorScaleOverLife.OkForParallel() && RadiusScale.OkForParallel() && LightExponent.OkForParallel();
+}
+
 void UParticleModuleLight::SpawnEx(FParticleEmitterInstance* Owner, int32 Offset, float SpawnTime, struct FRandomStream* InRandomStream, FBaseParticle* ParticleBase)
 {
 	SPAWN_INIT;
@@ -2867,7 +2873,6 @@ void UParticleModuleLight::SetToSensibleDefaults(UParticleEmitter* Owner)
 void UParticleModuleLight::Render3DPreview(FParticleEmitterInstance* Owner, const FSceneView* View,FPrimitiveDrawInterface* PDI)
 {
 #if WITH_EDITOR
-
 	if ((Owner == NULL) || (Owner->ActiveParticles <= 0) || 
 		(Owner->ParticleData == NULL) || (Owner->ParticleIndices == NULL))
 	{
@@ -2876,6 +2881,7 @@ void UParticleModuleLight::Render3DPreview(FParticleEmitterInstance* Owner, cons
 
 	if (bPreviewLightRadius)
 	{
+		check(IsInGameThread());
 		int32 Offset = 0;
 		UParticleLODLevel* LODLevel	= Owner->SpriteTemplate->GetCurrentLODLevel(Owner);
 		const bool bLocalSpace = LODLevel->RequiredModule->bUseLocalSpace;
@@ -4078,29 +4084,8 @@ bool UDistributionVectorParticleParameter::GetParamValue(UObject* Data, FName Pa
 	UParticleSystemComponent* ParticleComp = Cast<UParticleSystemComponent>(Data);
 	if(ParticleComp)
 	{
-		bFoundParam = ParticleComp->GetVectorParameter(ParameterName, OutVector);
-
-		// If we failed to get a Vector parameter with the given name, see if we can get a Color parameter or Float parameter
-		if(!bFoundParam)
-		{
-			FLinearColor OutColor;
-			bFoundParam = ParticleComp->GetColorParameter(ParameterName, OutColor);
-			if(bFoundParam)
-			{
-				OutVector = FVector(OutColor);
-			}
-			else
-			{
-				float OutFloat;
-				bFoundParam = ParticleComp->GetFloatParameter(ParameterName, OutFloat);
-				if(bFoundParam)
-				{
-					OutVector = FVector(OutFloat);
-				}
-			}
-		}
+		bFoundParam = ParticleComp->GetAnyVectorParameter(ParameterName, OutVector);
 	}
-
 	return bFoundParam;
 }
 
diff --git a/Engine/Source/Runtime/Engine/Private/Particles/ParticleModules_Camera.cpp b/Engine/Source/Runtime/Engine/Private/Particles/ParticleModules_Camera.cpp
--- a/Engine/Source/Runtime/Engine/Private/Particles/ParticleModules_Camera.cpp
+++ b/Engine/Source/Runtime/Engine/Private/Particles/ParticleModules_Camera.cpp
@@ -61,6 +61,12 @@ void UParticleModuleCameraOffset::PostEditChangeProperty(FPropertyChangedEvent&
 }
 #endif // WITH_EDITOR
 
+bool UParticleModuleCameraOffset::CanTickInAnyThread()
+{
+	return CameraOffset.OkForParallel();
+}
+
+
 void UParticleModuleCameraOffset::Spawn(FParticleEmitterInstance* Owner, int32 Offset, float SpawnTime, FBaseParticle* ParticleBase)
 {
 	check(Owner && Owner->Component)
@@ -71,7 +77,7 @@ void UParticleModuleCameraOffset::Spawn(FParticleEmitterInstance* Owner, int32 O
 	{
 		if (Owner != NULL && Owner->Component != NULL)
 		{
-			ScaleFactor = Owner->Component->ComponentToWorld.GetMaximumAxisScale();
+			ScaleFactor = Owner->Component->GetAsyncComponentToWorld().GetMaximumAxisScale();
 		}
 	}
 	SPAWN_INIT;
diff --git a/Engine/Source/Runtime/Engine/Private/Particles/ParticleModules_Location.cpp b/Engine/Source/Runtime/Engine/Private/Particles/ParticleModules_Location.cpp
--- a/Engine/Source/Runtime/Engine/Private/Particles/ParticleModules_Location.cpp
+++ b/Engine/Source/Runtime/Engine/Private/Particles/ParticleModules_Location.cpp
@@ -1728,6 +1728,7 @@ uint32 UParticleModuleLocationBoneSocket::PrepPerInstanceBlock(FParticleEmitterI
 
 void UParticleModuleLocationBoneSocket::AutoPopulateInstanceProperties(UParticleSystemComponent* PSysComp)
 {
+	check(IsInGameThread());
 	bool bFound = false;
 	for (int32 ParamIdx = 0; ParamIdx < PSysComp->InstanceParameters.Num(); ParamIdx++)
 	{
@@ -2406,6 +2407,7 @@ uint32 UParticleModuleLocationSkelVertSurface::RequiredBytesPerInstance(FParticl
 
 void UParticleModuleLocationSkelVertSurface::AutoPopulateInstanceProperties(UParticleSystemComponent* PSysComp)
 {
+	check(IsInGameThread());
 	bool bFound = false;
 	for (int32 ParamIdx = 0; ParamIdx < PSysComp->InstanceParameters.Num(); ParamIdx++)
 	{
diff --git a/Engine/Source/Runtime/Engine/Private/Particles/ParticleModules_Parameter.cpp b/Engine/Source/Runtime/Engine/Private/Particles/ParticleModules_Parameter.cpp
--- a/Engine/Source/Runtime/Engine/Private/Particles/ParticleModules_Parameter.cpp
+++ b/Engine/Source/Runtime/Engine/Private/Particles/ParticleModules_Parameter.cpp
@@ -119,6 +119,19 @@ void UParticleModuleParameterDynamic::PostLoad()
 	UpdateUsageFlags();
 }
 
+bool UParticleModuleParameterDynamic::CanTickInAnyThread()
+{
+	for (FEmitterDynamicParameter& Parm : DynamicParams)
+	{
+		if (!Parm.ParamValue.OkForParallel())
+		{
+			return false;
+		}
+	}
+	return true;
+}
+
+
 void UParticleModuleParameterDynamic::Spawn(FParticleEmitterInstance* Owner, int32 Offset, float SpawnTime, FBaseParticle* ParticleBase)
 {
 	SpawnEx(Owner, Offset, SpawnTime, NULL, ParticleBase);
diff --git a/Engine/Source/Runtime/Engine/Private/Particles/ParticleTrail2EmitterInstance.cpp b/Engine/Source/Runtime/Engine/Private/Particles/ParticleTrail2EmitterInstance.cpp
--- a/Engine/Source/Runtime/Engine/Private/Particles/ParticleTrail2EmitterInstance.cpp
+++ b/Engine/Source/Runtime/Engine/Private/Particles/ParticleTrail2EmitterInstance.cpp
@@ -2036,6 +2036,7 @@ void FParticleRibbonEmitterInstance::SetupTrailModules()
 
 void FParticleRibbonEmitterInstance::ResolveSource()
 {
+	check(IsInGameThread());
 	if (SourceModule && SourceModule->SourceName != NAME_None)
 	{
 		switch (SourceModule->SourceMethod)
@@ -3322,8 +3323,10 @@ float FParticleAnimTrailEmitterInstance::Spawn(float DeltaTime)
 	// Figure out spawn rate for this tick.
 	if (bProcessSpawnRate)
 	{
+		static IConsoleVariable* EffectsQuality = IConsoleManager::Get().FindConsoleVariable(TEXT("sg.EffectsQuality"));
+
 		float RateScale = LODLevel->SpawnModule->RateScale.GetValue(EmitterTime, Component) * LODLevel->SpawnModule->GetGlobalRateScale();
-		float QualityMult = 0.25f * (1 << Scalability::GetQualityLevels().EffectsQuality);
+		float QualityMult = 0.25f * (1 << EffectsQuality->GetInt());
 		RateScale *= SpriteTemplate->QualityLevelSpawnRateScale*QualityMult;
 		SpawnRate += LODLevel->SpawnModule->Rate.GetValue(EmitterTime, Component) * FMath::Clamp<float>(RateScale, 0.0f, RateScale);
 	}
diff --git a/Engine/Source/Runtime/Engine/Private/Particles/ParticleTrailModules.cpp b/Engine/Source/Runtime/Engine/Private/Particles/ParticleTrailModules.cpp
--- a/Engine/Source/Runtime/Engine/Private/Particles/ParticleTrailModules.cpp
+++ b/Engine/Source/Runtime/Engine/Private/Particles/ParticleTrailModules.cpp
@@ -109,6 +109,7 @@ void UParticleModuleTrailSource::PostEditChangeProperty(FPropertyChangedEvent& P
 
 void UParticleModuleTrailSource::AutoPopulateInstanceProperties(UParticleSystemComponent* PSysComp)
 {
+	check(IsInGameThread());
 	switch (SourceMethod)
 	{
 	case PET2SRCM_Actor:
