Add some clarification in comments for SetCenterOfMass and SetMassOverrideInKg for when body is welded to something else
Code formatting fixes to GetPhysicsLinearVelocity, GetPhysicsLinearVelocityAtPoint, SetMassOverrideInKg, GetClosestPointOnCollision

diff --git a/Engine/Source/Runtime/Engine/Classes/Components/PrimitiveComponent.h b/Engine/Source/Runtime/Engine/Classes/Components/PrimitiveComponent.h
--- a/Engine/Source/Runtime/Engine/Classes/Components/PrimitiveComponent.h
+++ b/Engine/Source/Runtime/Engine/Classes/Components/PrimitiveComponent.h
@@ -929,6 +929,7 @@ public:
 
 	/**
 	*	Set the center of mass of a single body. This will offset the physx-calculated center of mass.
+	*	Note that in the case where multiple bodies are attached together, the center of mass will be set for the entire group.
 	*	@param CenterOfMassOffset		User specified offset for the center of mass of this object, from the calculated location.
 	*	@param BoneName			If a SkeletalMeshComponent, name of body to set center of mass of. 'None' indicates root body.
 	*/
@@ -1571,6 +1572,7 @@ public:
 
 	/**
 	*	Override the mass (in Kg) of a single physics body.
+	*	Note that in the case where multiple bodies are attached together, the override mass will be set for the entire group.
 	*	Set the Override Mass to false if you want to reset the body's mass to the auto-calculated physx mass.
 	*/
 	UFUNCTION(BlueprintCallable, Category = "Physics")
diff --git a/Engine/Source/Runtime/Engine/Private/PrimitiveComponentPhysics.cpp b/Engine/Source/Runtime/Engine/Private/PrimitiveComponentPhysics.cpp
--- a/Engine/Source/Runtime/Engine/Private/PrimitiveComponentPhysics.cpp
+++ b/Engine/Source/Runtime/Engine/Private/PrimitiveComponentPhysics.cpp
@@ -298,8 +298,7 @@ void UPrimitiveComponent::SetPhysicsLinearVelocity(FVector NewVel, bool bAddToCu
 
 FVector UPrimitiveComponent::GetPhysicsLinearVelocity(FName BoneName)
 {
-	FBodyInstance* BI = GetBodyInstance(BoneName);
-	if(BI != NULL)
+	if (FBodyInstance* BI = GetBodyInstance(BoneName))
 	{
 		return BI->GetUnrealWorldVelocity();
 	}
@@ -308,8 +307,7 @@ FVector UPrimitiveComponent::GetPhysicsLinearVelocity(FName BoneName)
 
 FVector UPrimitiveComponent::GetPhysicsLinearVelocityAtPoint(FVector Point, FName BoneName)
 {
-	FBodyInstance* BI = GetBodyInstance(BoneName);
-	if (BI != NULL)
+	if (FBodyInstance* BI = GetBodyInstance(BoneName))
 	{
 		return BI->GetUnrealWorldVelocityAtPoint(Point);
 	}
@@ -490,8 +488,7 @@ void UPrimitiveComponent::SetAllMassScale(float InMassScale)
 
 void UPrimitiveComponent::SetMassOverrideInKg(FName BoneName, float MassInKg, bool bOverrideMass)
 {
-	FBodyInstance* BI = GetBodyInstance(BoneName);
-	if (BI)
+	if (FBodyInstance* BI = GetBodyInstance(BoneName))
 	{
 		WarnInvalidPhysicsOperations(LOCTEXT("SetCenterOfMass", "SetCenterOfMass"), nullptr);
 		BI->bOverrideMass = bOverrideMass;
@@ -885,8 +882,7 @@ float UPrimitiveComponent::GetClosestPointOnCollision(const FVector& Point, FVec
 {
 	OutPointOnBody = Point;
 
-	FBodyInstance* BodyInst = GetBodyInstance(BoneName);
-	if (BodyInst != NULL)
+	if (FBodyInstance* BodyInst = GetBodyInstance(BoneName))
 	{
 		return BodyInst->GetDistanceToBody(Point, OutPointOnBody);
 	}
