UE4 - ABTest of cvars functionality

diff --git a/Engine/Source/Runtime/Engine/Private/UnrealEngine.cpp b/Engine/Source/Runtime/Engine/Private/UnrealEngine.cpp
--- a/Engine/Source/Runtime/Engine/Private/UnrealEngine.cpp
+++ b/Engine/Source/Runtime/Engine/Private/UnrealEngine.cpp
@@ -1022,8 +1022,339 @@ void UEngine::TickDeferredCommands()
 	DeferredCommands.RemoveAt(0, DeferredCommandsCount);
 }
 
+#if !UE_BUILD_SHIPPING
+
+static TAutoConsoleVariable<int32> CVarABTestHistory(
+	TEXT("abtest.HistoryNum"),
+	1000,
+	TEXT("Number of history frames to use for stats."));
+
+static TAutoConsoleVariable<int32> CVarABTestReportNum(
+	TEXT("abtest.ReportNum"),
+	100,
+	TEXT("Number of frames between reports."));
+
+static TAutoConsoleVariable<int32> CVarABTestCoolDown(
+	TEXT("abtest.CoolDown"),
+	3,
+	TEXT("Number of frames to discard data after each command to cover threading."));
+
+static TAutoConsoleVariable<int32> CVarABTestMinFramesPerTrial(
+	TEXT("abtest.MinFramesPerTrial"),
+	5,
+	TEXT("The number of frames to run a given command before swithing; this is randomized."));
+
+static TAutoConsoleVariable<int32> CVarABTestNumResamples(
+	TEXT("abtest.NumResamples"),
+	256,
+	TEXT("The number of resamples to use to determine confidence."));
+
+struct FABTest
+{
+	FRandomStream Stream;
+	bool bABTestActive;
+	FString ABTestCmds[2];
+	int32 ABTestNumSamples;
+	int32 RemainingCoolDown;
+	int32 CurrentTest;
+	int32 RemainingTrial;
+	int32 RemainingPrint;
+	int32 SampleIndex;
+
+	int32 HistoryNum;
+	int32 ReportNum;
+	int32 CoolDown;
+	int32 MinFramesPerTrial;
+	int32 NumResamples;
+
+	struct FSample
+	{
+		uint32 Micros;
+		int32 TestIndex;
+		TBitArray<> InResamples;
+	};
+
+	TArray<FSample> Samples;
+	TArray<uint32> ResampleAccumulators;
+	TArray<uint32> ResampleCount;
+	uint32 Totals[2];
+	uint32 Counts[2];
+
+
+	FABTest()
+		: Stream(9567)
+		, bABTestActive(false)
+		, ABTestNumSamples(0)
+		, RemainingCoolDown(0)
+		, CurrentTest(0)
+		, RemainingTrial(0)
+		, RemainingPrint(0)
+		, SampleIndex(0)
+		, HistoryNum(0)
+		, ReportNum(0)
+		, CoolDown(0)
+		, MinFramesPerTrial(0)
+		, NumResamples(0)
+	{
+	}
+
+	void Tick()
+	{
+		static double LastTime = 0;
+		if (bABTestActive && RemainingCoolDown)
+		{
+			RemainingCoolDown--;
+		}
+		else if (bABTestActive)
+		{
+			float Delta = FPlatformTime::Seconds() - LastTime;
+			check(Delta > 0);
+			FSample* Sample = nullptr;
+			if (ABTestNumSamples < HistoryNum)
+			{
+				check(ABTestNumSamples == Samples.Num());
+				Sample = new (Samples) FSample();
+			}
+			else
+			{
+				Sample = &Samples[ABTestNumSamples % HistoryNum];
+
+				// remove the old data
+				check(Totals[Sample->TestIndex] > Sample->Micros);
+				Totals[Sample->TestIndex] -= Sample->Micros;
+				check(Counts[Sample->TestIndex]);
+				Counts[Sample->TestIndex]--;
+				check(Sample->InResamples.Num() == NumResamples);
+				for (TConstSetBitIterator<> CurIndexIt(Sample->InResamples); CurIndexIt; ++CurIndexIt)
+				{
+					int32 ResampleIndex = CurIndexIt.GetIndex();
+					check(ResampleAccumulators[ResampleIndex] > Sample->Micros);
+					ResampleAccumulators[ResampleIndex] -= Sample->Micros;
+					check(ResampleCount[ResampleIndex]);
+					ResampleCount[ResampleIndex]--;
+				}
+			}
+			ABTestNumSamples++;
+			Sample->Micros = uint32(Delta * 1000000.0f);
+			check(Sample->Micros);
+			Sample->TestIndex = CurrentTest;
+			Totals[Sample->TestIndex] += Sample->Micros;
+			Counts[Sample->TestIndex]++;
+			Sample->InResamples.Empty(NumResamples);
+			for (int32 ResampleIndex = 0; ResampleIndex < NumResamples; ResampleIndex++)
+			{
+				bool In = Stream.FRand() > 0.5f;
+				Sample->InResamples.Add(In);
+				if (In)
+				{
+					ResampleAccumulators[ResampleIndex] += Sample->Micros;
+					ResampleCount[ResampleIndex]++;
+				}
+			}
+			check(RemainingPrint);
+			if (!--RemainingPrint)
+			{
+				if (Counts[0] && Counts[1])
+				{
+					UE_LOG(LogConsoleResponse, Display, TEXT("      %7.4fms  (%4d samples)  A = '%s'"), float(Totals[0]) / float(Counts[0]) / 1000.0f, Counts[0], *ABTestCmds[0]);
+					UE_LOG(LogConsoleResponse, Display, TEXT("      %7.4fms  (%4d samples)  B = '%s'"), float(Totals[1]) / float(Counts[1]) / 1000.0f, Counts[1], *ABTestCmds[1]);
+
+					float Diff = (float(Totals[0]) / float(Counts[0]) / 1000.0f) - (float(Totals[1]) / float(Counts[1]) / 1000.0f);
+					bool bAIsFaster = false;
+					if (Diff < 0.0f)
+					{
+						Diff *= -1.0f;
+						bAIsFaster = true;
+					}
+					static TArray<float> Scores;
+					Scores.Reset();
+					uint32 TotalFrame = Totals[0] + Totals[1];
+					uint32 TotalCount = Counts[0] + Counts[1];
+					for (int32 ResampleIndex = 0; ResampleIndex < NumResamples; ResampleIndex++)
+					{
+						check(TotalFrame > ResampleAccumulators[ResampleIndex]);
+						uint32 Opposite = TotalFrame - ResampleAccumulators[ResampleIndex];
+						check(TotalCount > ResampleCount[ResampleIndex]);
+						uint32 OppositeCount = TotalCount - ResampleCount[ResampleIndex];
+
+						// this ABS here gives us better data but makes it one sided, so the 90% confidence interval is at index 80/100 etc
+						// it is legit because we could have easily done the inverse permutation
+						float SampleDiff = FMath::Abs(((ResampleAccumulators[ResampleIndex]) / float(ResampleCount[ResampleIndex]) / 1000.0f) - (float(Opposite) / float(OppositeCount) / 1000.0f));
+						Scores.Add(SampleDiff);
+
+					}
+					Scores.Sort();
+					int32 Conf = NumResamples;
+					for (int32 Trial = 0; Trial < NumResamples; Trial++)
+					{
+						if (Scores[Trial] > Diff)
+						{
+							Conf = Trial;
+							break;
+						}
+					}
+
+					float fConf = 1.0f - (0.5f + float(Conf) / float(NumResamples) / 2.0f);
+
+					if (bAIsFaster)
+					{
+						UE_LOG(LogConsoleResponse, Display, TEXT("      A is %7.4fms faster than B;  %3.0f%% chance this is noise."), Diff, fConf * 100.0f);
+					}
+					else
+					{
+						UE_LOG(LogConsoleResponse, Display, TEXT("      B is %7.4fms faster than A;  %3.0f%% chance this is noise."), Diff, fConf * 100.0f);
+					}
+
+					UE_LOG(LogConsoleResponse, Display, TEXT("----------------"));
+				}
+				else
+				{
+					UE_LOG(LogConsoleResponse, Display, TEXT("No Samples?"));
+				}
+				RemainingPrint = ReportNum;
+			}
+			check(RemainingTrial);
+			if (!--RemainingTrial)
+			{
+				SwitchTest(1 - CurrentTest);
+			}
+		}
+		LastTime = FPlatformTime::Seconds();
+	}
+
+	static FABTest& Get()
+	{
+		static FABTest Singleton;
+		return Singleton;
+	}
+	static void ABTestCmdFunc(const TArray<FString>& Args)
+	{
+		FString ABTestCmds[2];
+		if (Args.Num() == 1 && Args[0].Compare(FString(TEXT("stop")), ESearchCase::IgnoreCase) == 0)
+		{
+			Get().Stop();
+			return;
+		}
+		if (Args.Num() == 3 && !Args[0].StartsWith(TEXT("\"")))
+		{
+			ABTestCmds[0] = Args[0].TrimQuotes() + TEXT(" ") + Args[1].TrimQuotes();
+			ABTestCmds[1] = Args[0].TrimQuotes() + TEXT(" ") + Args[2].TrimQuotes();
+		}
+		else if (Args.Num() > 2 && Args[0].StartsWith(TEXT("\"")))
+		{
+			FString Work;
+			int32 Test = 0;
+			for (int32 Index = 0; Index < Args.Num(); Index++)
+			{
+				Work += Args[Index];
+				if (Work.Len() > 2 && Work.StartsWith(TEXT("\"")) && Work.EndsWith(TEXT("\"")))
+				{
+					ABTestCmds[Test++] = Work.TrimQuotes();
+					Work.Empty();
+					if (Test > 1)
+					{
+						break;
+					}
+				}
+				else if (Index + 1 < Args.Num())
+				{
+					Work += TEXT(" ");
+				}
+			}
+		}
+		else if (Args.Num() == 2)
+		{
+			ABTestCmds[0] = Args[0].TrimQuotes();
+			ABTestCmds[1] = Args[1].TrimQuotes();
+		}
+		else
+		{
+			UE_LOG(LogConsoleResponse, Display, TEXT("abtest command requires two (quoted) arguments or three args or 'stop'."));
+			UE_LOG(LogConsoleResponse, Display, TEXT("Example: abtest \"r.MyCVar 0\" \"r.MyCVar 1\""));
+			UE_LOG(LogConsoleResponse, Display, TEXT("Example: abtest r.MyCVar 0 1"));
+			return;
+		}
+		Get().Start(ABTestCmds);
+	}
+
+	void Stop()
+	{
+		UE_LOG(LogConsoleResponse, Display, TEXT("Stopping AB test."));
+		bABTestActive = false;
+	}
+	
+	void Start(FString* InABTestCmds)
+	{
+		ABTestCmds[0] = InABTestCmds[0];
+		ABTestCmds[1] = InABTestCmds[1];
+
+		HistoryNum = CVarABTestHistory.GetValueOnGameThread();
+		ReportNum = CVarABTestReportNum.GetValueOnGameThread();
+		CoolDown = CVarABTestCoolDown.GetValueOnGameThread();
+		MinFramesPerTrial = CVarABTestMinFramesPerTrial.GetValueOnGameThread();
+		NumResamples = CVarABTestNumResamples.GetValueOnGameThread();
+
+		Samples.Empty(HistoryNum);
+		ResampleAccumulators.Empty(NumResamples);
+		ResampleAccumulators.AddZeroed(NumResamples);
+		ResampleCount.Empty(NumResamples);
+		ResampleCount.AddZeroed(NumResamples);
+		ABTestNumSamples = 0;
+		Totals[0] = 0;
+		Totals[1] = 0;
+		Counts[0] = 0;
+		Counts[1] = 0;
+		RemainingPrint = ReportNum;
+
+		bABTestActive = true;
+		Get().SwitchTest(0);
+		UE_LOG(LogConsoleResponse, Display, TEXT("abtest started with A = '%s' and B = '%s'"), *ABTestCmds[0], *ABTestCmds[1]);
+	}
+
+	void SwitchTest(int32 Index)
+	{
+		RemainingCoolDown = CoolDown;
+		CurrentTest = Index;
+		RemainingTrial = Stream.RandRange(MinFramesPerTrial, MinFramesPerTrial * 3);
+		check(RemainingTrial);
+		GEngine->Exec(nullptr, *ABTestCmds[CurrentTest], *GLog);
+	}
+
+	static void StaticTick()
+	{
+		Get().Tick();
+	}
+	static bool IsActive()
+	{
+		return Get().bABTestActive;
+	}
+};
+
+static FAutoConsoleCommand ABTestCmd(
+	TEXT("abtest"),
+	TEXT("Provide two console commands or 'stop' to stop the abtest. Frames are timed with the two options, logging results over time."),
+	FConsoleCommandWithArgsDelegate::CreateStatic(&FABTest::ABTestCmdFunc)
+	);
+
+#else
+
+struct FABTest
+{
+	FORCEINLINE static void StaticTick()
+	{
+	}
+	FORCEINLINE static bool IsActive()
+	{
+		return false;
+	}
+};
+
+#endif
+
+
 void UEngine::UpdateTimeAndHandleMaxTickRate()
 {
+	FABTest::StaticTick();
 	// start at now minus a bit so we don't get a zero delta.
 	static double LastTime = FPlatformTime::Seconds() - 0.0001;
 	static bool bTimeWasManipulated = false;
@@ -1072,7 +1403,7 @@ void UEngine::UpdateTimeAndHandleMaxTickRate()
 		UpdateRunningAverageDeltaTime(DeltaTime);
 
 		// Get max tick rate based on network settings and current delta time.
-		const float MaxTickRate	= GetMaxTickRate( DeltaTime );
+		const float MaxTickRate	= FABTest::IsActive() ? 0.0f : GetMaxTickRate( DeltaTime );
 		float WaitTime		= 0;
 		// Convert from max FPS to wait time.
 		if( MaxTickRate > 0 )
@@ -1091,7 +1422,7 @@ void UEngine::UpdateTimeAndHandleMaxTickRate()
 
 			if (IsRunningDedicatedServer()) // We aren't so concerned about wall time with a server, lots of CPU is wasted spinning. I suspect there is more to do with sleeping and time on dedicated servers.
 			{
-				FPlatformProcess::Sleep(WaitTime);
+				FPlatformProcess::SleepNoStats(WaitTime);
 			}
 			else
 			{
@@ -1100,13 +1431,13 @@ void UEngine::UpdateTimeAndHandleMaxTickRate()
 				// up our timeslice.
 				if( WaitTime > 5 / 1000.f )
 				{
-					FPlatformProcess::Sleep( WaitTime - 0.002f );
+					FPlatformProcess::SleepNoStats( WaitTime - 0.002f );
 				}
 
 				// Give up timeslice for remainder of wait time.
 				while( FPlatformTime::Seconds() < WaitEndTime )
 				{
-					FPlatformProcess::Sleep( 0 );
+					FPlatformProcess::SleepNoStats( 0 );
 				}
 			}
 			FApp::SetCurrentTime(FPlatformTime::Seconds());
