PR #843: Exposing AActor::bTearOff to BP through function AActor::TearOff()

diff --git a/Engine/Source/Runtime/Engine/Classes/GameFramework/Actor.h b/Engine/Source/Runtime/Engine/Classes/GameFramework/Actor.h
--- a/Engine/Source/Runtime/Engine/Classes/GameFramework/Actor.h
+++ b/Engine/Source/Runtime/Engine/Classes/GameFramework/Actor.h
@@ -144,6 +144,10 @@ public:
 	UPROPERTY(replicated)
 	uint32 bTearOff:1;    
 
+	/** Networking - Server - TearOff this actor to stop replication to clients. Will set bTearOff to true. */
+	UFUNCTION(BlueprintCallable, Category = Replication)
+	virtual void TearOff();
+
 	/**
 	 * Whether we have already exchanged Role/RemoteRole on the client, as when removing then re-adding a streaming level.
 	 * Causes all initialization to be performed again even though the actor may not have actually been reloaded.
diff --git a/Engine/Source/Runtime/Engine/Private/Actor.cpp b/Engine/Source/Runtime/Engine/Private/Actor.cpp
--- a/Engine/Source/Runtime/Engine/Private/Actor.cpp
+++ b/Engine/Source/Runtime/Engine/Private/Actor.cpp
@@ -1699,6 +1699,16 @@ void AActor::Destroyed()
 	}
 }
 
+void AActor::TearOff()
+{
+	const ENetMode NetMode = GetNetMode();
+
+	if (NetMode == NM_ListenServer || NetMode == NM_DedicatedServer)
+	{
+		bTearOff = true;
+	}
+}
+
 void AActor::TornOff() {}
 
 void AActor::Reset()
