#ue4 - Fix crash when UpdatedComponent for Projectile or Rotating movement are marked pending kill and set to null in the next tick.

UE-18080

diff --git a/Engine/Source/Runtime/Engine/Private/Components/ProjectileMovementComponent.cpp b/Engine/Source/Runtime/Engine/Private/Components/ProjectileMovementComponent.cpp
--- a/Engine/Source/Runtime/Engine/Private/Components/ProjectileMovementComponent.cpp
+++ b/Engine/Source/Runtime/Engine/Private/Components/ProjectileMovementComponent.cpp
@@ -100,6 +100,11 @@ void UProjectileMovementComponent::TickComponent(float DeltaTime, enum ELevelTic
 
 	Super::TickComponent(DeltaTime, TickType, ThisTickFunction);
 
+	if (!IsValid(UpdatedComponent))
+	{
+		return;
+	}
+
 	AActor* ActorOwner = UpdatedComponent->GetOwner();
 	if ( !ActorOwner || !CheckStillInWorld() )
 	{
diff --git a/Engine/Source/Runtime/Engine/Private/Components/RotatingMovementComponent.cpp b/Engine/Source/Runtime/Engine/Private/Components/RotatingMovementComponent.cpp
--- a/Engine/Source/Runtime/Engine/Private/Components/RotatingMovementComponent.cpp
+++ b/Engine/Source/Runtime/Engine/Private/Components/RotatingMovementComponent.cpp
@@ -21,6 +21,11 @@ void URotatingMovementComponent::TickComponent(float DeltaTime, enum ELevelTick
 
 	Super::TickComponent(DeltaTime, TickType, ThisTickFunction);
 
+	if (!IsValid(UpdatedComponent))
+	{
+		return;
+	}
+
 	// Compute new rotation
 	const FQuat OldRotation = UpdatedComponent->GetComponentQuat();
 	const FQuat DeltaRotation = (RotationRate * DeltaTime).Quaternion();
