Fix crash when apex sim calls back with a fracture after the destructible has called DestroyPhysics. This can happen due to deferred deletion of destructible actors.

diff --git a/Engine/Source/Runtime/Engine/Private/Components/DestructibleComponent.cpp b/Engine/Source/Runtime/Engine/Private/Components/DestructibleComponent.cpp
--- a/Engine/Source/Runtime/Engine/Private/Components/DestructibleComponent.cpp
+++ b/Engine/Source/Runtime/Engine/Private/Components/DestructibleComponent.cpp
@@ -845,6 +845,11 @@ FTransform UDestructibleComponent::GetSocketTransform(FName InSocketName, ERelat
 #if WITH_APEX
 void UDestructibleComponent::Pair( int32 ChunkIndex, PxShape* PShape)
 {
+	if(ApexDestructibleActor == nullptr)	//since we do deferred deletion it's possible we've already meant to delete this so ignore any simulation callbacks
+	{
+		return;
+	}
+
 	FDestructibleChunkInfo* CI;
 	FPhysxUserData* UserData;
 
