UMG - Fixing a spelling error.

diff --git a/Engine/Source/Runtime/Engine/Classes/Materials/Material.h b/Engine/Source/Runtime/Engine/Classes/Materials/Material.h
--- a/Engine/Source/Runtime/Engine/Classes/Materials/Material.h
+++ b/Engine/Source/Runtime/Engine/Classes/Materials/Material.h
@@ -90,7 +90,7 @@ enum EMaterialDomain
 	MD_LightFunction UMETA(DisplayName="Light Function"),
 	/** The material will be used in a custom post process pass. */
 	MD_PostProcess UMETA(DisplayName="Post Process"),
-	/** The material wwill be used for UMG or Slate UI */
+	/** The material will be used for UMG or Slate UI */
 	MD_UI UMETA(DisplayName="User Interface"),
 
 	MD_MAX
@@ -674,9 +674,9 @@ public:
 	UPROPERTY()
 	uint32 bCanMaskedBeAssumedOpaque : 1;
 
- 	/** true if Material is masked and uses custom opacity */
- 	UPROPERTY()
- 	uint32 bIsMasked_DEPRECATED:1;
+	/** true if Material is masked and uses custom opacity */
+	UPROPERTY()
+	uint32 bIsMasked_DEPRECATED:1;
 
 	/** true if Material is the preview material used in the material editor. */
 	UPROPERTY(transient, duplicatetransient)
