#ue4 - Improve docs for AddYawInput and AddControllerYawInput etc, to point to the scaling values that are applied.

UE-9152

diff --git a/Engine/Source/Runtime/Engine/Classes/GameFramework/Pawn.h b/Engine/Source/Runtime/Engine/Classes/GameFramework/Pawn.h
--- a/Engine/Source/Runtime/Engine/Classes/GameFramework/Pawn.h
+++ b/Engine/Source/Runtime/Engine/Classes/GameFramework/Pawn.h
@@ -373,16 +373,31 @@ public:
 	UFUNCTION(BlueprintCallable, Category="Pawn|Input", meta=(Keywords="ConsumeInput"))
 	virtual FVector ConsumeMovementInputVector();
 
-	/** Add input (affecting Pitch) to the Controller's ControlRotation, if it is a local PlayerController. */
-	UFUNCTION(BlueprintCallable, Category="Pawn|Input", meta=(Keywords="up down"))
+	/**
+	 * Add input (affecting Pitch) to the Controller's ControlRotation, if it is a local PlayerController.
+	 * This value is multiplied by the PlayerController's InputPitchScale value.
+	 * @param Val Amount to add to Pitch. This value is multiplied by the PlayerController's InputPitchScale value.
+	 * @see PlayerController::InputPitchScale
+	 */
+	UFUNCTION(BlueprintCallable, Category="Pawn|Input", meta=(Keywords="up down addpitch"))
 	virtual void AddControllerPitchInput(float Val);
 
-	/** Add input (affecting Yaw) to the Controller's ControlRotation, if it is a local PlayerController. */
-	UFUNCTION(BlueprintCallable, Category="Pawn|Input", meta=(Keywords="left right turn"))
+	/**
+	 * Add input (affecting Yaw) to the Controller's ControlRotation, if it is a local PlayerController.
+	 * This value is multiplied by the PlayerController's InputYawScale value.
+	 * @param Val Amount to add to Yaw. This value is multiplied by the PlayerController's InputYawScale value.
+	 * @see PlayerController::InputYawScale
+	 */
+	UFUNCTION(BlueprintCallable, Category="Pawn|Input", meta=(Keywords="left right turn addyaw"))
 	virtual void AddControllerYawInput(float Val);
 
-	/** Add input (affecting Roll) to the Controller's ControlRotation, if it is a local PlayerController. */
-	UFUNCTION(BlueprintCallable, Category="Pawn|Input")
+	/**
+	 * Add input (affecting Roll) to the Controller's ControlRotation, if it is a local PlayerController.
+	 * This value is multiplied by the PlayerController's InputRollScale value.
+	 * @param Val Amount to add to Roll. This value is multiplied by the PlayerController's InputRollScale value.
+	 * @see PlayerController::InputRollScale
+	 */
+	UFUNCTION(BlueprintCallable, Category="Pawn|Input", meta=(Keywords="addroll"))
 	virtual void AddControllerRollInput(float Val);
 
 	/** Helper to see if move input is ignored. If our controller is a PlayerController, checks Controller->IsMoveInputIgnored(). */
diff --git a/Engine/Source/Runtime/Engine/Classes/GameFramework/PlayerController.h b/Engine/Source/Runtime/Engine/Classes/GameFramework/PlayerController.h
--- a/Engine/Source/Runtime/Engine/Classes/GameFramework/PlayerController.h
+++ b/Engine/Source/Runtime/Engine/Classes/GameFramework/PlayerController.h
@@ -1013,15 +1013,24 @@ public:
 	UFUNCTION(reliable, server, WithValidation)
 	void ServerToggleAILogging();
 
-	/** Add Pitch (look up) input */
+	/**
+	 * Add Pitch (look up) input. This value is multiplied by InputPitchScale.
+	 * @param Val Amount to add to Pitch. This value is multiplied by InputPitchScale.
+	 */
 	UFUNCTION(BlueprintCallable, Category="Game|Player", meta=(Keywords="up down"))
 	virtual void AddPitchInput(float Val);
 
-	/** Add Yaw (turn) input */
+	/**
+	 * Add Yaw (turn) input. This value is multiplied by InputYawScale.
+	 * @param Val Amount to add to Yaw. This value is multiplied by InputYawScale.
+	 */
 	UFUNCTION(BlueprintCallable, Category="Game|Player", meta=(Keywords="left right turn"))
 	virtual void AddYawInput(float Val);
 
-	/** Add Roll input */
+	/**
+	 * Add Roll input. This value is multiplied by InputRollScale.
+	 * @param Val Amount to add to Roll. This value is multiplied by InputRollScale.
+	 */
 	UFUNCTION(BlueprintCallable, Category="Game|Player")
 	virtual void AddRollInput(float Val);
 
