Add delegate for demo net driver seeking so we are alerted when the HTTP requests and fast forwarding are complete

Merging 2609459

//depot/UE4-UT/...

to //depot/UE4/...

diff --git a/Engine/Source/Runtime/Engine/Classes/Engine/DemoNetDriver.h b/Engine/Source/Runtime/Engine/Classes/Engine/DemoNetDriver.h
--- a/Engine/Source/Runtime/Engine/Classes/Engine/DemoNetDriver.h
+++ b/Engine/Source/Runtime/Engine/Classes/Engine/DemoNetDriver.h
@@ -6,6 +6,8 @@
 
 DECLARE_LOG_CATEGORY_EXTERN( LogDemo, Log, All );
 
+DECLARE_DELEGATE_OneParam(FOnGotoTimeDelegate, const bool);
+
 /**
  * Simulated network driver for recording and playing back game sessions.
  */
@@ -82,6 +84,8 @@ private:
 	// If a channel is associated with Actor, adds the channel's GUID to the list of GUIDs excluded from queuing bunches during scrubbing.
 	void		AddNonQueuedActorForScrubbing(AActor* Actor);
 
+	FOnGotoTimeDelegate OnGotoTimeDelegate;
+
 public:
 
 	// UNetDriver interface.
@@ -95,12 +99,13 @@ public:
 	virtual void ProcessRemoteFunction( class AActor* Actor, class UFunction* Function, void* Parameters, struct FOutParmRec* OutParms, struct FFrame* Stack, class UObject* SubObject = nullptr ) override;
 	virtual bool IsAvailable() const override { return true; }
 	void SkipTime(const float InTimeToSkip);
-	void GotoTimeInSeconds( const float TimeInSeconds );
 	bool InitConnectInternal( FString& Error );
 	virtual bool ShouldClientDestroyTearOffActors() const override;
 	virtual bool ShouldSkipRepNotifies() const override;
 	virtual bool ShouldQueueBunchesForActorGUID(FNetworkGUID InGUID) const override;
 
+	void GotoTimeInSeconds(const float TimeInSeconds, const FOnGotoTimeDelegate& InOnGotoTimeDelegate = FOnGotoTimeDelegate());
+
 public:
 
 	// FExec interface
@@ -126,6 +131,7 @@ public:
 	void SpawnDemoRecSpectator( UNetConnection* Connection );
 	void ResetDemoState();
 	void JumpToEndOfLiveReplay();
+	virtual bool IsFastForwarding() { return bIsFastForwarding; }
 
 	/**
 	 * Adds a join-in-progress user to the set of users associated with the currently recording replay (if any)
diff --git a/Engine/Source/Runtime/Engine/Private/DemoNetDriver.cpp b/Engine/Source/Runtime/Engine/Private/DemoNetDriver.cpp
--- a/Engine/Source/Runtime/Engine/Private/DemoNetDriver.cpp
+++ b/Engine/Source/Runtime/Engine/Private/DemoNetDriver.cpp
@@ -1217,12 +1217,14 @@ void UDemoNetDriver::SkipTime(const float InTimeToSkip)
 	TimeToSkip = InTimeToSkip;
 }
 
-void UDemoNetDriver::GotoTimeInSeconds( const float TimeInSeconds )
+void UDemoNetDriver::GotoTimeInSeconds(const float TimeInSeconds, const FOnGotoTimeDelegate& InOnGotoTimeDelegate)
 {
 	if ( QueuedGotoTimeInSeconds < 0 )
 	{
 		QueuedGotoTimeInSeconds = TimeInSeconds;
 	}
+
+	OnGotoTimeDelegate = InOnGotoTimeDelegate;
 }
 
 void UDemoNetDriver::JumpToEndOfLiveReplay()
@@ -1420,6 +1422,9 @@ void UDemoNetDriver::TickDemoPlayback( float DeltaSeconds )
 
 		const auto FastForwardTotalSeconds = FPlatformTime::Seconds() - FastForwardStartSeconds;
 
+		OnGotoTimeDelegate.ExecuteIfBound(true);
+		OnGotoTimeDelegate.Unbind();
+
 		UE_LOG( LogDemo, Log, TEXT( "Fast forward took %.2f seconds." ), FastForwardTotalSeconds );
 	}
 }
@@ -1510,6 +1515,8 @@ void UDemoNetDriver::CheckpointReady( const bool bSuccess, const int64 SkipExtra
 	{
 		UE_LOG( LogDemo, Warning, TEXT( "UDemoNetConnection::CheckpointReady: Failed to go to checkpoint." ) );
 		DemoCurrentTime = OldDemoCurrentTime;
+		OnGotoTimeDelegate.ExecuteIfBound(false);
+		OnGotoTimeDelegate.Unbind();
 		return;
 	}
 
