Add sanity checks to catch replay scrubbing issues

diff --git a/Engine/Source/Runtime/Engine/Private/NetConnection.cpp b/Engine/Source/Runtime/Engine/Private/NetConnection.cpp
--- a/Engine/Source/Runtime/Engine/Private/NetConnection.cpp
+++ b/Engine/Source/Runtime/Engine/Private/NetConnection.cpp
@@ -969,6 +969,7 @@ void UNetConnection::ReceivedPacket( FBitReader& Reader )
 			// Ignore if reliable packet has already been processed.
 			if ( Bunch.bReliable && Bunch.ChSequence <= InReliable[Bunch.ChIndex] )
 			{
+				check( !InternalAck );		// Should be impossible with 100% reliable connections
 				UE_LOG( LogNetTraffic, Log, TEXT( "UNetConnection::ReceivedPacket: Received outdated bunch (Channel %d Current Sequence %i)" ), Bunch.ChIndex, InReliable[Bunch.ChIndex] );
 				continue;
 			}
@@ -984,6 +985,8 @@ void UNetConnection::ReceivedPacket( FBitReader& Reader )
 				const bool ValidUnreliableOpen = Bunch.bOpen && (Bunch.bClose || Bunch.bPartial);
 				if (!ValidUnreliableOpen)
 				{
+					check( !InternalAck );		// Should be impossible with 100% reliable connections
+
 					UE_LOG( LogNetTraffic, Warning, TEXT( "      Received unreliable bunch before open (Channel %d Current Sequence %i)" ), Bunch.ChIndex, InReliable[Bunch.ChIndex] );
 					// Since we won't be processing this packet, don't ack it
 					// We don't want the sender to think this bunch was processed when it really wasn't
