Removed legacy DistanceFieldPenumbraScale material property

diff --git a/Engine/Source/Editor/UnrealEd/Classes/Editor/UnrealEdTypes.h b/Engine/Source/Editor/UnrealEd/Classes/Editor/UnrealEdTypes.h
--- a/Engine/Source/Editor/UnrealEd/Classes/Editor/UnrealEdTypes.h
+++ b/Engine/Source/Editor/UnrealEd/Classes/Editor/UnrealEdTypes.h
@@ -92,20 +92,11 @@ struct FLightmassParameterizedMaterialSettings
 	UPROPERTY(EditAnywhere, BlueprintReadWrite, Category=LightmassParameterizedMaterialSettings)
 	struct FLightmassScalarParameterValue ExportResolutionScale;
 
-	/** Scales the penumbra size of distance field shadows.  This is useful to get softer precomputed shadows on certain material types like foliage. */
-	UPROPERTY(EditAnywhere, BlueprintReadWrite, Category=LightmassParameterizedMaterialSettings)
-	struct FLightmassScalarParameterValue DistanceFieldPenumbraScale;
-
-
-
-		FLightmassParameterizedMaterialSettings()
+	FLightmassParameterizedMaterialSettings()
 		: EmissiveBoost(1.0f)
 		, DiffuseBoost(1.0f)
 		, ExportResolutionScale(1.0f)
-		, DistanceFieldPenumbraScale(1.0f)
-		{
-		}
-	
+	{}
 };
 
 
diff --git a/Engine/Source/Editor/UnrealEd/Private/PreviewMaterial.cpp b/Engine/Source/Editor/UnrealEd/Private/PreviewMaterial.cpp
--- a/Engine/Source/Editor/UnrealEd/Private/PreviewMaterial.cpp
+++ b/Engine/Source/Editor/UnrealEd/Private/PreviewMaterial.cpp
@@ -609,8 +609,6 @@ void UMaterialEditorInstanceConstant::CopyToSourceInstance()
 		SourceInstance->SetDiffuseBoost(LightmassSettings.DiffuseBoost.ParameterValue);
 		SourceInstance->SetOverrideExportResolutionScale(LightmassSettings.ExportResolutionScale.bOverride);
 		SourceInstance->SetExportResolutionScale(LightmassSettings.ExportResolutionScale.ParameterValue);
-		SourceInstance->SetOverrideDistanceFieldPenumbraScale(LightmassSettings.DistanceFieldPenumbraScale.bOverride);
-		SourceInstance->SetDistanceFieldPenumbraScale(LightmassSettings.DistanceFieldPenumbraScale.ParameterValue);
 
 		// Copy Refraction bias setting
 		SourceInstance->SetScalarParameterValueEditorOnly(TEXT("RefractionDepthBias"), RefractionDepthBias);
@@ -699,8 +697,6 @@ void UMaterialEditorInstanceConstant::SetSourceInstance(UMaterialInstanceConstan
 	LightmassSettings.DiffuseBoost.ParameterValue = SourceInstance->GetDiffuseBoost();
 	LightmassSettings.ExportResolutionScale.bOverride = SourceInstance->GetOverrideExportResolutionScale();
 	LightmassSettings.ExportResolutionScale.ParameterValue = SourceInstance->GetExportResolutionScale();
-	LightmassSettings.DistanceFieldPenumbraScale.bOverride = SourceInstance->GetOverrideDistanceFieldPenumbraScale();
-	LightmassSettings.DistanceFieldPenumbraScale.ParameterValue = SourceInstance->GetDistanceFieldPenumbraScale();
 
 	//Copy refraction settings
 	SourceInstance->GetRefractionSettings(RefractionDepthBias);
diff --git a/Engine/Source/Runtime/Engine/Classes/Materials/MaterialInstance.h b/Engine/Source/Runtime/Engine/Classes/Materials/MaterialInstance.h
--- a/Engine/Source/Runtime/Engine/Classes/Materials/MaterialInstance.h
+++ b/Engine/Source/Runtime/Engine/Classes/Materials/MaterialInstance.h
@@ -272,7 +272,6 @@ public:
 	virtual ENGINE_API float GetEmissiveBoost() const override;
 	virtual ENGINE_API float GetDiffuseBoost() const override;
 	virtual ENGINE_API float GetExportResolutionScale() const override;
-	virtual ENGINE_API float GetDistanceFieldPenumbraScale() const override;
 	virtual ENGINE_API bool GetTexturesInPropertyChain(EMaterialProperty InProperty, TArray<UTexture*>& OutTextures,
 		TArray<FName>* OutTextureParamNames, class FStaticParameterSet* InStaticParameterSet) override;
 	virtual ENGINE_API void RecacheUniformExpressions() const override;
diff --git a/Engine/Source/Runtime/Engine/Classes/Materials/MaterialInterface.h b/Engine/Source/Runtime/Engine/Classes/Materials/MaterialInterface.h
--- a/Engine/Source/Runtime/Engine/Classes/Materials/MaterialInterface.h
+++ b/Engine/Source/Runtime/Engine/Classes/Materials/MaterialInterface.h
@@ -129,10 +129,6 @@ struct FLightmassMaterialInterfaceSettings
 	UPROPERTY(EditAnywhere, BlueprintReadWrite, Category=Material)
 	float ExportResolutionScale;
 
-	/***/
-	UPROPERTY()
-	float DistanceFieldPenumbraScale;
-
 	/** Boolean override flags - only used in MaterialInstance* cases. */
 	/** If true, override the bCastShadowAsMasked setting of the parent material. */
 	UPROPERTY()
@@ -150,26 +146,16 @@ struct FLightmassMaterialInterfaceSettings
 	UPROPERTY()
 	uint32 bOverrideExportResolutionScale:1;
 
-	/** If true, override the distance field penumbra scale setting of the parent material. */
-	UPROPERTY()
-	uint32 bOverrideDistanceFieldPenumbraScale:1;
-
-
-
-		FLightmassMaterialInterfaceSettings()
+	FLightmassMaterialInterfaceSettings()
 		: bCastShadowAsMasked(false)
 		, EmissiveBoost(1.0f)
 		, DiffuseBoost(1.0f)
 		, ExportResolutionScale(1.0f)
-		, DistanceFieldPenumbraScale(1.0f)
 		, bOverrideCastShadowAsMasked(false)
 		, bOverrideEmissiveBoost(false)
 		, bOverrideDiffuseBoost(false)
 		, bOverrideExportResolutionScale(false)
-		, bOverrideDistanceFieldPenumbraScale(false)
-		{
-		}
-	
+	{}
 };
 
 UCLASS(abstract, BlueprintType,MinimalAPI)
@@ -419,11 +405,6 @@ public:
 		return LightmassSettings.bOverrideExportResolutionScale;
 	}
 
-	inline bool GetOverrideDistanceFieldPenumbraScale() const
-	{
-		return LightmassSettings.bOverrideDistanceFieldPenumbraScale;
-	}
-
 	/** @return	The bCastShadowAsMasked value for this material. */
 	virtual bool GetCastShadowAsMasked() const
 	{
@@ -449,11 +430,6 @@ public:
 		return FMath::Clamp(LightmassSettings.ExportResolutionScale, .1f, 10.0f);
 	}
 
-	virtual float GetDistanceFieldPenumbraScale() const
-	{
-		return LightmassSettings.DistanceFieldPenumbraScale;
-	}
-
 	/** @param	bInOverrideCastShadowAsMasked	The override CastShadowAsMasked setting to set. */
 	inline void SetOverrideCastShadowAsMasked(bool bInOverrideCastShadowAsMasked)
 	{
@@ -478,11 +454,6 @@ public:
 		LightmassSettings.bOverrideExportResolutionScale = bInOverrideExportResolutionScale;
 	}
 
-	inline void SetOverrideDistanceFieldPenumbraScale(bool bInOverrideDistanceFieldPenumbraScale)
-	{
-		LightmassSettings.bOverrideDistanceFieldPenumbraScale = bInOverrideDistanceFieldPenumbraScale;
-	}
-
 	/** @param	InCastShadowAsMasked	The CastShadowAsMasked value for this material. */
 	inline void SetCastShadowAsMasked(bool InCastShadowAsMasked)
 	{
@@ -507,11 +478,6 @@ public:
 		LightmassSettings.ExportResolutionScale = InExportResolutionScale;
 	}
 
-	inline void SetDistanceFieldPenumbraScale(float InDistanceFieldPenumbraScale)
-	{
-		LightmassSettings.DistanceFieldPenumbraScale = InDistanceFieldPenumbraScale;
-	}
-
 	/**
 	 *	Get all of the textures in the expression chain for the given property (ie fill in the given array with all textures in the chain).
 	 *
diff --git a/Engine/Source/Runtime/Engine/Private/Materials/Material.cpp b/Engine/Source/Runtime/Engine/Private/Materials/Material.cpp
--- a/Engine/Source/Runtime/Engine/Private/Materials/Material.cpp
+++ b/Engine/Source/Runtime/Engine/Private/Materials/Material.cpp
@@ -267,30 +267,15 @@ public:
 		}
 	}
 
-	virtual float GetDistanceFieldPenumbraScale() const { return DistanceFieldPenumbraScale; }
-
 	// FRenderResource interface.
 	virtual FString GetFriendlyName() const { return Material->GetName(); }
 
 	// Constructor.
 	FDefaultMaterialInstance(UMaterial* InMaterial,bool bInSelected,bool bInHovered):
 		FMaterialRenderProxy(bInSelected, bInHovered),
-		Material(InMaterial),
-		DistanceFieldPenumbraScale(1.0f)
+		Material(InMaterial)
 	{}
 
-	/** Called from the game thread to update DistanceFieldPenumbraScale. */
-	void GameThread_UpdateDistanceFieldPenumbraScale(float NewDistanceFieldPenumbraScale)
-	{
-		ENQUEUE_UNIQUE_RENDER_COMMAND_TWOPARAMETER(
-			UpdateDistanceFieldPenumbraScaleCommand,
-			float*,DistanceFieldPenumbraScale,&DistanceFieldPenumbraScale,
-			float,NewDistanceFieldPenumbraScale,NewDistanceFieldPenumbraScale,
-		{
-			*DistanceFieldPenumbraScale = NewDistanceFieldPenumbraScale;
-		});
-	}
-
 private:
 
 	/** Get the fallback material. */
@@ -300,9 +285,6 @@ private:
 	}
 
 	UMaterial* Material;
-
-	// maintained by the render thread
-	float DistanceFieldPenumbraScale;
 };
 
 #if WITH_EDITOR
@@ -2408,8 +2390,6 @@ void UMaterial::PropagateDataToMaterialProxy()
 	{
 		if (DefaultMaterialInstances[i])
 		{
-			DefaultMaterialInstances[i]->GameThread_UpdateDistanceFieldPenumbraScale(GetDistanceFieldPenumbraScale());
-
 			UpdateMaterialRenderProxy(*DefaultMaterialInstances[i]);
 		}
 	}
diff --git a/Engine/Source/Runtime/Engine/Private/Materials/MaterialInstance.cpp b/Engine/Source/Runtime/Engine/Private/Materials/MaterialInstance.cpp
--- a/Engine/Source/Runtime/Engine/Private/Materials/MaterialInstance.cpp
+++ b/Engine/Source/Runtime/Engine/Private/Materials/MaterialInstance.cpp
@@ -75,7 +75,6 @@ FMaterialInstanceResource::FMaterialInstanceResource(UMaterialInstance* InOwner,
 	: FMaterialRenderProxy(bInSelected, bInHovered)
 	, Parent(NULL)
 	, Owner(InOwner)
-	, DistanceFieldPenumbraScale(1.0f)
 	, GameThreadParent(NULL)
 	, OpacityMaskClipValue(0.3333333f)
 	, BlendMode(BLEND_Opaque)
@@ -238,18 +237,6 @@ bool FMaterialInstanceResource::GetTextureValue(
 	}
 }
 
-/** Called from the game thread to update DistanceFieldPenumbraScale. */
-void FMaterialInstanceResource::GameThread_UpdateDistanceFieldPenumbraScale(float NewDistanceFieldPenumbraScale)
-{
-	ENQUEUE_UNIQUE_RENDER_COMMAND_TWOPARAMETER(
-		UpdateDistanceFieldPenumbraScaleCommand,
-		float*,DistanceFieldPenumbraScale,&DistanceFieldPenumbraScale,
-		float,NewDistanceFieldPenumbraScale,NewDistanceFieldPenumbraScale,
-	{
-		*DistanceFieldPenumbraScale = NewDistanceFieldPenumbraScale;
-	});
-}
-
 void FMaterialInstanceResource::GameThread_UpdateOverridableBaseProperties(const UMaterialInterface* MaterialInterface)
 {
 	ENQUEUE_UNIQUE_RENDER_COMMAND_FOURPARAMETER(
@@ -283,8 +270,6 @@ void UMaterialInstance::PropagateDataToMaterialProxy()
 	{
 		if (Resources[i])
 		{
-			Resources[i]->GameThread_UpdateDistanceFieldPenumbraScale(GetDistanceFieldPenumbraScale()); 
-
 			UpdateMaterialRenderProxy(*Resources[i]);
 		}
 	}
@@ -2228,22 +2213,6 @@ float UMaterialInstance::GetExportResolutionScale() const
 	return 1.0f;
 }
 
-float UMaterialInstance::GetDistanceFieldPenumbraScale() const
-{
-	if (LightmassSettings.bOverrideDistanceFieldPenumbraScale)
-	{
-		return LightmassSettings.DistanceFieldPenumbraScale;
-	}
-
-	if (Parent)
-	{
-		return Parent->GetDistanceFieldPenumbraScale();
-	}
-
-	return 1.0f;
-}
-
-
 bool UMaterialInstance::GetTexturesInPropertyChain(EMaterialProperty InProperty, TArray<UTexture*>& OutTextures,  
 	TArray<FName>* OutTextureParamNames, class FStaticParameterSet* InStaticParameterSet)
 {
diff --git a/Engine/Source/Runtime/Engine/Private/Materials/MaterialInstanceSupport.h b/Engine/Source/Runtime/Engine/Private/Materials/MaterialInstanceSupport.h
--- a/Engine/Source/Runtime/Engine/Private/Materials/MaterialInstanceSupport.h
+++ b/Engine/Source/Runtime/Engine/Private/Materials/MaterialInstanceSupport.h
@@ -90,14 +90,10 @@ public:
 	virtual bool GetVectorValue(const FName ParameterName, FLinearColor* OutValue, const FMaterialRenderContext& Context) const override;
 	virtual bool GetScalarValue(const FName ParameterName,float* OutValue, const FMaterialRenderContext& Context) const override;
 	virtual bool GetTextureValue(const FName ParameterName,const UTexture** OutValue, const FMaterialRenderContext& Context) const override;
-	virtual float GetDistanceFieldPenumbraScale() const override { return DistanceFieldPenumbraScale; }
 
 	/** Sets the material instance's parent. */
 	void GameThread_SetParent(UMaterialInterface* InParent);
 
-	/** Called from the game thread to update DistanceFieldPenumbraScale. */
-	void GameThread_UpdateDistanceFieldPenumbraScale(float NewDistanceFieldPenumbraScale);
-	
 	/** Called from the game thread to update the overridable base properties in the proxy. */
 	void GameThread_UpdateOverridableBaseProperties(const UMaterialInterface* MaterialInterface);
 
@@ -174,9 +170,6 @@ private:
 	/** The UMaterialInstance which owns this resource. */
 	UMaterialInstance* Owner;
 
-	/** Scales the penumbra size of distance field shadows. */
-	float DistanceFieldPenumbraScale;
-
 	/** The game thread accessible parent of the material instance. */
 	UMaterialInterface* GameThreadParent;
 
diff --git a/Engine/Source/Runtime/Engine/Private/Materials/MaterialInterface.cpp b/Engine/Source/Runtime/Engine/Private/Materials/MaterialInterface.cpp
--- a/Engine/Source/Runtime/Engine/Private/Materials/MaterialInterface.cpp
+++ b/Engine/Source/Runtime/Engine/Private/Materials/MaterialInterface.cpp
@@ -159,7 +159,6 @@ void UMaterialInterface::PostEditChangeProperty(FPropertyChangedEvent& PropertyC
 	LightmassSettings.EmissiveBoost = FMath::Max(LightmassSettings.EmissiveBoost, 0.0f);
 	LightmassSettings.DiffuseBoost = FMath::Max(LightmassSettings.DiffuseBoost, 0.0f);
 	LightmassSettings.ExportResolutionScale = FMath::Clamp(LightmassSettings.ExportResolutionScale, 0.0f, 16.0f);
-	LightmassSettings.DistanceFieldPenumbraScale = FMath::Clamp(LightmassSettings.DistanceFieldPenumbraScale, 0.01f, 100.0f);
 
 	Super::PostEditChangeProperty(PropertyChangedEvent);
 }
diff --git a/Engine/Source/Runtime/Engine/Public/MaterialShared.h b/Engine/Source/Runtime/Engine/Public/MaterialShared.h
--- a/Engine/Source/Runtime/Engine/Public/MaterialShared.h
+++ b/Engine/Source/Runtime/Engine/Public/MaterialShared.h
@@ -1341,7 +1341,6 @@ public:
 	virtual bool GetVectorValue(const FName ParameterName, FLinearColor* OutValue, const FMaterialRenderContext& Context) const = 0;
 	virtual bool GetScalarValue(const FName ParameterName, float* OutValue, const FMaterialRenderContext& Context) const = 0;
 	virtual bool GetTextureValue(const FName ParameterName,const UTexture** OutValue, const FMaterialRenderContext& Context) const = 0;
-	virtual float GetDistanceFieldPenumbraScale() const { return 1.0f; }
 	bool IsSelected() const { return bSelected; }
 	bool IsHovered() const { return bHovered; }
 
