Move options parsing out of AGameMode and in to UGameplayStatics
Expose option parsing to blueprints

diff --git a/Engine/Source/Runtime/Engine/Classes/GameFramework/GameMode.h b/Engine/Source/Runtime/Engine/Classes/GameFramework/GameMode.h
--- a/Engine/Source/Runtime/Engine/Classes/GameFramework/GameMode.h
+++ b/Engine/Source/Runtime/Engine/Classes/GameFramework/GameMode.h
@@ -373,19 +373,24 @@ public:
 	//=========================================================================
 	// URL Parsing
 	/** Grab the next option from a string. */
-	bool GrabOption( FString& Options, FString& ResultString );
+	DEPRECATED(4.9, "Use UGameplayStatics::GrabOption instead")
+	static bool GrabOption( FString& Options, FString& ResultString );
 
 	/** Break up a key=value pair into its key and value. */
-	void GetKeyValue( const FString& Pair, FString& Key, FString& Value );
+	DEPRECATED(4.9, "Use UGameplayStatics::GetKeyValue instead")
+	static void GetKeyValue( const FString& Pair, FString& Key, FString& Value );
 
 	/* Find an option in the options string and return it. */
-	FString ParseOption( const FString& Options, const FString& InKey );
+	DEPRECATED(4.9, "Use UGameplayStatics::ParseOption instead")
+	static FString ParseOption( FString Options, const FString& InKey );
 
 	/** HasOption - return true if the option is specified on the command line. */
-	bool HasOption( const FString& Options, const FString& InKey );
+	DEPRECATED(4.9, "Use UGameplayStatics::HasOption instead")
+	static bool HasOption( FString Options, const FString& InKey );
 
 	/** Search array of options for ParseString and return the index or CurrentValue if not found*/
-	int32 GetIntOption( const FString& Options, const FString& ParseString, int32 CurrentValue);
+	DEPRECATED(4.9, "Use UGameplayStatics::GetIntOption instead")
+	static int32 GetIntOption( const FString& Options, const FString& ParseString, int32 CurrentValue);
 	//=========================================================================
 
 	/** 
diff --git a/Engine/Source/Runtime/Engine/Classes/Kismet/GameplayStatics.h b/Engine/Source/Runtime/Engine/Classes/Kismet/GameplayStatics.h
--- a/Engine/Source/Runtime/Engine/Classes/Kismet/GameplayStatics.h
+++ b/Engine/Source/Runtime/Engine/Classes/Kismet/GameplayStatics.h
@@ -580,5 +580,48 @@ class ENGINE_API UGameplayStatics : public UBlueprintFunctionLibrary
 	UFUNCTION(BlueprintPure, Category = "Utilities")
 	static bool ProjectWorldToScreen(APlayerController const* Player, const FVector& WorldPosition, FVector2D& ScreenPosition);
 
+	// Utility functions for interacting with Options strings
+
+	//=========================================================================
+	// URL Parsing
+
+	static bool GrabOption( FString& Options, FString& ResultString );
+
+	/** 
+	 * Break up a key=value pair into its key and value. 
+	 * @param Pair			The string containing a pair to split apart.
+	 * @param Key			(out) Key portion of Pair. If no = in string will be the same as Pair.
+	 * @param Value			(out) Value portion of Pair. If no = in string will be empty.
+	 */
+	UFUNCTION(BlueprintPure, Category="Game Options")
+	static void GetKeyValue( const FString& Pair, FString& Key, FString& Value );
+
+	/** 
+	 * Find an option in the options string and return it.
+	 * @param Options		The string containing the options.
+	 * @param Key			The key to find the value of in Options.
+	 * @return				The value associated with Key if Key found in Options string.
+	 */
+	UFUNCTION(BlueprintPure, Category="Game Options")
+	static FString ParseOption( FString Options, const FString& Key );
+
+	/** 
+	 * Returns whether a key exists in an options string.
+	 * @param Options		The string containing the options.
+	 * @param Key			The key to determine if it exists in Options.
+	 * @return				Whether Key was found in Options.
+	 */
+	UFUNCTION(BlueprintPure, Category="Game Options")
+	static bool HasOption( FString Options, const FString& InKey );
+
+	/** 
+	 * Find an option in the options string and return it as an integer.
+	 * @param Options		The string containing the options.
+	 * @param Key			The key to find the value of in Options.
+	 * @return				The value associated with Key as an integer if Key found in Options string, otherwise DefaultValue.
+	 */
+	UFUNCTION(BlueprintPure, Category="Game Options")
+	static int32 GetIntOption( const FString& Options, const FString& Key, int32 DefaultValue);
+
 };
 
diff --git a/Engine/Source/Runtime/Engine/Private/GameMode.cpp b/Engine/Source/Runtime/Engine/Private/GameMode.cpp
--- a/Engine/Source/Runtime/Engine/Private/GameMode.cpp
+++ b/Engine/Source/Runtime/Engine/Private/GameMode.cpp
@@ -620,8 +620,8 @@ void AGameMode::HandleMatchHasStarted()
 	GetWorldSettings()->NotifyMatchStarted();
 
 	// if passed in bug info, send player to right location
-	FString BugLocString = ParseOption(OptionsString, TEXT("BugLoc"));
-	FString BugRotString = ParseOption(OptionsString, TEXT("BugRot"));
+	const FString BugLocString = UGameplayStatics::ParseOption(OptionsString, TEXT("BugLoc"));
+	const FString BugRotString = UGameplayStatics::ParseOption(OptionsString, TEXT("BugRot"));
 	if( !BugLocString.IsEmpty() || !BugRotString.IsEmpty() )
 	{
 		for( FConstPlayerControllerIterator Iterator = GetWorld()->GetPlayerControllerIterator(); Iterator; ++Iterator )
@@ -629,9 +629,6 @@ void AGameMode::HandleMatchHasStarted()
 			APlayerController* PlayerController = *Iterator;
 			if( PlayerController->CheatManager != NULL )
 			{
-				//`log( "BugLocString:" @ BugLocString );
-				//`log( "BugRotString:" @ BugRotString );
-
 				PlayerController->CheatManager->BugItGoString( BugLocString, BugRotString );
 			}
 		}
@@ -1019,10 +1016,10 @@ FString AGameMode::InitNewPlayer(APlayerController* NewPlayerController, const T
 	FString ErrorMessage;
 
 	// Register the player with the session
-	GameSession->RegisterPlayer(NewPlayerController, UniqueId, HasOption(Options, TEXT("bIsFromInvite")));
+	GameSession->RegisterPlayer(NewPlayerController, UniqueId, UGameplayStatics::HasOption(Options, TEXT("bIsFromInvite")));
 
 	// Init player's name
-	FString InName = ParseOption(Options, TEXT("Name")).Left(20);
+	FString InName = UGameplayStatics::ParseOption(Options, TEXT("Name")).Left(20);
 	if (InName.IsEmpty())
 	{
 		InName = FString::Printf(TEXT("%s%i"), *DefaultPlayerName.ToString(), NewPlayerController->PlayerState->PlayerId);
@@ -1079,7 +1076,7 @@ APlayerController* AGameMode::Login(UPlayer* NewPlayer, ENetRole RemoteRole, con
 	}
 
 	// Set up spectating
-	bool bSpectator = FCString::Stricmp(*ParseOption(Options, TEXT("SpectatorOnly")), TEXT("1")) == 0;
+	bool bSpectator = FCString::Stricmp(*UGameplayStatics::ParseOption(Options, TEXT("SpectatorOnly")), TEXT("1")) == 0;
 	if (bSpectator || MustSpectate(NewPlayerController))
 	{
 		NewPlayerController->StartSpectatingOnly();
@@ -1177,82 +1174,27 @@ void AGameMode::ClearPause()
 
 bool AGameMode::GrabOption( FString& Options, FString& Result )
 {
-	if( Options.Left(1)==TEXT("?") )
-	{
-		// Get result.
-		Result = Options.Mid(1, MAX_int32);
-		if (Result.Contains(TEXT("?"), ESearchCase::CaseSensitive))
-		{
-			Result = Result.Left(Result.Find(TEXT("?"), ESearchCase::CaseSensitive));
-		}
-
-		// Update options.
-		Options = Options.Mid(1, MAX_int32);
-		if (Options.Contains(TEXT("?"), ESearchCase::CaseSensitive))
-		{
-			Options = Options.Mid(Options.Find(TEXT("?"), ESearchCase::CaseSensitive), MAX_int32);
-		}
-		else
-		{
-			Options = TEXT("");
-		}
-
-		return true;
-	}
-	else return false;
+	return UGameplayStatics::GrabOption(Options, Result);
 }
 
 void AGameMode::GetKeyValue( const FString& Pair, FString& Key, FString& Value )
 {
-	const int32 EqualSignIndex = Pair.Find(TEXT("="), ESearchCase::CaseSensitive);
-	if( EqualSignIndex != INDEX_NONE )
-	{
-		Key = Pair.Left(EqualSignIndex);
-		Value = Pair.Mid(EqualSignIndex + 1, MAX_int32);
-	}
-	else
-	{
-		Key = Pair;
-		Value = TEXT("");
-	}
+	return UGameplayStatics::GetKeyValue(Pair, Key, Value);
 }
 
-FString AGameMode::ParseOption( const FString& Options, const FString& InKey )
+FString AGameMode::ParseOption( FString Options, const FString& InKey )
 {
-	FString OptionsMod = Options;
-	FString Pair, Key, Value;
-	while( GrabOption( OptionsMod, Pair ) )
-	{
-		GetKeyValue( Pair, Key, Value );
-		if( FCString::Stricmp(*Key, *InKey) == 0 )
-			return Value;
-	}
-	return TEXT("");
+	return UGameplayStatics::ParseOption(Options, InKey);
 }
 
-bool AGameMode::HasOption( const FString& Options, const FString& InKey )
+bool AGameMode::HasOption( FString Options, const FString& InKey )
 {
-	FString OptionsMod = Options;
-	FString Pair, Key, Value;
-	while( GrabOption( OptionsMod, Pair ) )
-	{
-		GetKeyValue( Pair, Key, Value );
-		if( FCString::Stricmp(*Key, *InKey) == 0 )
-		{
-			return true;
-		}
-	}
-	return false;
+	return UGameplayStatics::HasOption(Options, InKey);
 }
 
 int32 AGameMode::GetIntOption( const FString& Options, const FString& ParseString, int32 CurrentValue)
 {
-	FString InOpt = ParseOption( Options, ParseString );
-	if ( !InOpt.IsEmpty() )
-	{
-		return FCString::Atoi(*InOpt);
-	}
-	return CurrentValue;
+	return UGameplayStatics::GetIntOption(Options, ParseString, CurrentValue);
 }
 
 FString AGameMode::GetDefaultGameClassPath(const FString& MapName, const FString& Options, const FString& Portal) const
diff --git a/Engine/Source/Runtime/Engine/Private/GameSession.cpp b/Engine/Source/Runtime/Engine/Private/GameSession.cpp
--- a/Engine/Source/Runtime/Engine/Private/GameSession.cpp
+++ b/Engine/Source/Runtime/Engine/Private/GameSession.cpp
@@ -122,8 +122,8 @@ void AGameSession::InitOptions( const FString& Options )
 	check(World);
 	AGameMode* const GameMode = World->GetAuthGameMode();
 
-	MaxPlayers = GameMode->GetIntOption( Options, TEXT("MaxPlayers"), MaxPlayers );
-	MaxSpectators = GameMode->GetIntOption( Options, TEXT("MaxSpectators"), MaxSpectators );
+	MaxPlayers = UGameplayStatics::GetIntOption( Options, TEXT("MaxPlayers"), MaxPlayers );
+	MaxSpectators = UGameplayStatics::GetIntOption( Options, TEXT("MaxSpectators"), MaxSpectators );
 	SessionName = GetDefault<APlayerState>(GameMode->PlayerStateClass)->SessionName;
 }
 
@@ -183,7 +183,7 @@ FString AGameSession::ApproveLogin(const FString& Options)
 	check(GameMode);
 
 	int32 SpectatorOnly = 0;
-	SpectatorOnly = GameMode->GetIntOption(Options, TEXT("SpectatorOnly"), SpectatorOnly);
+	SpectatorOnly = UGameplayStatics::GetIntOption(Options, TEXT("SpectatorOnly"), SpectatorOnly);
 
 	if (AtCapacity(SpectatorOnly == 1))
 	{
@@ -191,7 +191,7 @@ FString AGameSession::ApproveLogin(const FString& Options)
 	}
 
 	int32 SplitscreenCount = 0;
-	SplitscreenCount = GameMode->GetIntOption(Options, TEXT("SplitscreenCount"), SplitscreenCount);
+	SplitscreenCount = UGameplayStatics::GetIntOption(Options, TEXT("SplitscreenCount"), SplitscreenCount);
 
 	if (SplitscreenCount > MaxSplitscreensPerConnection)
 	{
diff --git a/Engine/Source/Runtime/Engine/Private/GameplayStatics.cpp b/Engine/Source/Runtime/Engine/Private/GameplayStatics.cpp
--- a/Engine/Source/Runtime/Engine/Private/GameplayStatics.cpp
+++ b/Engine/Source/Runtime/Engine/Private/GameplayStatics.cpp
@@ -1640,4 +1640,89 @@ bool UGameplayStatics::ProjectWorldToScreen(APlayerController const* Player, con
 	return false;
 }
 
+bool UGameplayStatics::GrabOption( FString& Options, FString& Result )
+{
+	FString QuestionMark(TEXT("?"));
+
+	if( Options.Left(1).Equals(QuestionMark, ESearchCase::CaseSensitive) )
+	{
+		// Get result.
+		Result = Options.Mid(1, MAX_int32);
+		if (Result.Contains(QuestionMark, ESearchCase::CaseSensitive))
+		{
+			Result = Result.Left(Result.Find(QuestionMark, ESearchCase::CaseSensitive));
+		}
+
+		// Update options.
+		Options = Options.Mid(1, MAX_int32);
+		if (Options.Contains(QuestionMark, ESearchCase::CaseSensitive))
+		{
+			Options = Options.Mid(Options.Find(QuestionMark, ESearchCase::CaseSensitive), MAX_int32);
+		}
+		else
+		{
+			Options = FString();
+		}
+
+		return true;
+	}
+
+	return false;
+}
+
+void UGameplayStatics::GetKeyValue( const FString& Pair, FString& Key, FString& Value )
+{
+	const int32 EqualSignIndex = Pair.Find(TEXT("="), ESearchCase::CaseSensitive);
+	if( EqualSignIndex != INDEX_NONE )
+	{
+		Key = Pair.Left(EqualSignIndex);
+		Value = Pair.Mid(EqualSignIndex + 1, MAX_int32);
+	}
+	else
+	{
+		Key = Pair;
+		Value = TEXT("");
+	}
+}
+
+FString UGameplayStatics::ParseOption( FString Options, const FString& Key )
+{
+	FString ReturnValue;
+	FString Pair, PairKey, PairValue;
+	while( GrabOption( Options, Pair ) )
+	{
+		GetKeyValue( Pair, PairKey, PairValue );
+		if (Key == PairKey)
+		{
+			ReturnValue = MoveTemp(PairValue);
+			break;
+		}
+	}
+	return ReturnValue;
+}
+
+bool UGameplayStatics::HasOption( FString Options, const FString& Key )
+{
+	FString Pair, PairKey, PairValue;
+	while( GrabOption( Options, Pair ) )
+	{
+		GetKeyValue( Pair, PairKey, PairValue );
+		if (Key == PairKey)
+		{
+			return true;
+		}
+	}
+	return false;
+}
+
+int32 UGameplayStatics::GetIntOption( const FString& Options, const FString& Key, int32 DefaultValue)
+{
+	const FString InOpt = ParseOption( Options, Key );
+	if ( !InOpt.IsEmpty() )
+	{
+		return FCString::Atoi(*InOpt);
+	}
+	return DefaultValue;
+}
+
 #undef LOCTEXT_NAMESPACE
\ No newline at end of file
