PR 1234: Added voice related exports to NetConnection.h

diff --git a/Engine/Source/Runtime/Engine/Classes/Engine/NetConnection.h b/Engine/Source/Runtime/Engine/Classes/Engine/NetConnection.h
--- a/Engine/Source/Runtime/Engine/Classes/Engine/NetConnection.h
+++ b/Engine/Source/Runtime/Engine/Classes/Engine/NetConnection.h
@@ -352,7 +352,7 @@ public:
 	 *
 	 * @return true if it should be sent on this connection, false otherwise
 	 */
-	bool ShouldReplicateVoicePacketFrom(const FUniqueNetId& Sender);
+	ENGINE_API bool ShouldReplicateVoicePacketFrom(const FUniqueNetId& Sender);
 	
 	/**
 	 * @hack: set to net connection currently inside CleanUp(), for HasClientLoadedCurrentWorld() to be able to find it during PlayerController
@@ -596,7 +596,7 @@ public:
 	/**
 	 * @return Finds the voice channel for this connection or NULL if none
 	 */
-	class UVoiceChannel* GetVoiceChannel();
+	ENGINE_API class UVoiceChannel* GetVoiceChannel();
 
 	void FlushDormancy(class AActor* Actor);
 
