[UE-17604] Handling UCS added components in the InheritableComponentHandler - avoiding a crash and setting us up for overridable UCS component properties.

diff --git a/Engine/Config/BaseEngine.ini b/Engine/Config/BaseEngine.ini
--- a/Engine/Config/BaseEngine.ini
+++ b/Engine/Config/BaseEngine.ini
@@ -625,6 +625,9 @@ bUseStreamingPause=false
 +TaggedPropertyRedirects=(ClassName="MemberReference",OldPropertyName="MemberParentClass",NewPropertyName="MemberParent")
 +TaggedPropertyRedirects=(ClassName="SimpleMemberReference",OldPropertyName="MemberParentClass",NewPropertyName="MemberParent")
 
++TaggedPropertyRedirects=(ClassName="ComponentKey",OldPropertyName="VariableName",NewPropertyName="SCSVariableName")
++TaggedPropertyRedirects=(ClassName="ComponentKey",OldPropertyName="VariableGuid",NewPropertyName="AssociatedGuid")
+
 ;Vim rename
 +TaggedPropertyRedirects=(ClassName="AnimInstance",OldPropertyName="VimVertexAnims",NewPropertyName="VertexAnims")
 
diff --git a/Engine/Source/Editor/Kismet/Private/SSCSEditor.cpp b/Engine/Source/Editor/Kismet/Private/SSCSEditor.cpp
--- a/Engine/Source/Editor/Kismet/Private/SSCSEditor.cpp
+++ b/Engine/Source/Editor/Kismet/Private/SSCSEditor.cpp
@@ -1237,7 +1237,7 @@ UActorComponent* FSCSEditorTreeNodeComponent::INTERNAL_GetOverridenComponentTemp
 	const bool BlueprintCanOverrideComponentFromKey = Key.IsValid()
 		&& Blueprint
 		&& Blueprint->ParentClass
-		&& Blueprint->ParentClass->IsChildOf(Key.OwnerClass);
+		&& Blueprint->ParentClass->IsChildOf(Key.GetComponentOwner());
 
 	if (BlueprintCanOverrideComponentFromKey)
 	{
diff --git a/Engine/Source/Editor/UnrealEd/Private/Editor.cpp b/Engine/Source/Editor/UnrealEd/Private/Editor.cpp
--- a/Engine/Source/Editor/UnrealEd/Private/Editor.cpp
+++ b/Engine/Source/Editor/UnrealEd/Private/Editor.cpp
@@ -724,18 +724,55 @@ namespace EditorUtilities
 							UInheritableComponentHandler* InheritableComponentHandler = Blueprint->GetInheritableComponentHandler(true);
 							if (InheritableComponentHandler)
 							{
+								FComponentKey Key;
+								FName const SourceComponentName = SourceComponent->GetFName();
+
 								BPGC = Cast<UBlueprintGeneratedClass>(BPGC->GetSuperClass());
-								USCS_Node* SCSNode = nullptr;
-								while (BPGC)
+								while (!Key.IsValid() && BPGC)
 								{
-									SCSNode = BPGC->SimpleConstructionScript->FindSCSNode(SourceComponent->GetFName());
-									BPGC = (SCSNode ? nullptr : Cast<UBlueprintGeneratedClass>(BPGC->GetSuperClass()));
+									USCS_Node* SCSNode = BPGC->SimpleConstructionScript->FindSCSNode(SourceComponentName);
+									if (!SCSNode)
+									{
+										UBlueprint* SuperBlueprint = CastChecked<UBlueprint>(BPGC->ClassGeneratedBy);
+										for (UActorComponent* ComponentTemplate : BPGC->ComponentTemplates)
+										{
+											if (ComponentTemplate->GetFName() == SourceComponentName)
+											{
+												if (UEdGraph* UCSGraph = FBlueprintEditorUtils::FindUserConstructionScript(SuperBlueprint))
+												{
+													TArray<UK2Node_AddComponent*> ComponentNodes;
+													UCSGraph->GetNodesOfClass<UK2Node_AddComponent>(ComponentNodes);
+
+													for (UK2Node_AddComponent* UCSNode : ComponentNodes)
+													{
+														if (ComponentTemplate == UCSNode->GetTemplateFromNode())
+														{
+															Key = FComponentKey(SuperBlueprint, FUCSComponentId(UCSNode));
+															break;
+														}
+													}
+												}
+												break;
+											}
+										}
+									}
+									else
+									{
+										Key = FComponentKey(SCSNode);
+										break;
+									}
+									BPGC = Cast<UBlueprintGeneratedClass>(BPGC->GetSuperClass());
 								}
-								check(SCSNode);
 
-								FComponentKey Key(SCSNode);
-								check(InheritableComponentHandler->GetOverridenComponentTemplate(Key) == nullptr);
-								TargetComponent = InheritableComponentHandler->CreateOverridenComponentTemplate(Key);
+								if (ensure(Key.IsValid()))
+								{
+									check(InheritableComponentHandler->GetOverridenComponentTemplate(Key) == nullptr);
+									TargetComponent = InheritableComponentHandler->CreateOverridenComponentTemplate(Key);
+								}
+								else
+								{
+									TargetComponent = nullptr;
+								}								
 							}
 						}
 					}
@@ -978,6 +1015,10 @@ namespace EditorUtilities
 		for( auto SourceComponentIter( SourceComponents.CreateConstIterator() ); SourceComponentIter; ++SourceComponentIter )
 		{
 			UActorComponent* SourceComponent = *SourceComponentIter;
+			if (SourceComponent->CreationMethod == EComponentCreationMethod::UserConstructionScript)
+			{
+				continue;
+			}
 			UActorComponent* TargetComponent = FindMatchingComponentInstance( SourceComponent, TargetActor, TargetComponents, TargetComponentIndex );
 
 			if( SourceComponent != nullptr && TargetComponent != nullptr )
diff --git a/Engine/Source/Editor/UnrealEd/Private/Kismet2/BlueprintEditorUtils.cpp b/Engine/Source/Editor/UnrealEd/Private/Kismet2/BlueprintEditorUtils.cpp
--- a/Engine/Source/Editor/UnrealEd/Private/Kismet2/BlueprintEditorUtils.cpp
+++ b/Engine/Source/Editor/UnrealEd/Private/Kismet2/BlueprintEditorUtils.cpp
@@ -452,6 +452,14 @@ void FParamsChangedHelper::EditCreateDelegates(UK2Node_CreateDelegate* CallSite)
 	}
 }
 
+//////////////////////////////////////
+// FUCSComponentId
+
+FUCSComponentId::FUCSComponentId(const UK2Node_AddComponent* UCSNode)
+	: GraphNodeGuid(UCSNode->NodeGuid)
+{
+}
+
 //////////////////////////////////////
 // FBlueprintEditorUtils
 
@@ -4164,6 +4172,11 @@ void FBlueprintEditorUtils::RenameComponentMemberVariable(UBlueprint* Blueprint,
 	{
 		Blueprint->Modify();
 
+		// Update the name
+		const FName OldName = Node->VariableName;
+		Node->Modify();
+		Node->VariableName = NewName;
+
 		// Rename Inheritable Component Templates
 		{
 			const FComponentKey Key(Node);
@@ -4175,17 +4188,12 @@ void FBlueprintEditorUtils::RenameComponentMemberVariable(UBlueprint* Blueprint,
 				if (InheritableComponentHandler && InheritableComponentHandler->GetOverridenComponentTemplate(Key))
 				{
 					InheritableComponentHandler->Modify();
-					InheritableComponentHandler->RenameTemplate(Key, NewName);
+					InheritableComponentHandler->RefreshTemplateName(Key);
 					InheritableComponentHandler->MarkPackageDirty();
 				}
 			}
 		}
 
-		// Update the name
-		const FName OldName = Node->VariableName;
-		Node->Modify();
-		Node->VariableName = NewName;
-
 		Node->NameWasModified();
 
 		// Update any existing references to the old name
@@ -6165,6 +6173,32 @@ bool FBlueprintEditorUtils::IsSCSComponentProperty(UObjectProperty* MemberProper
 	return false;
 }
 
+UActorComponent* FBlueprintEditorUtils::FindUCSComponentTemplate(const FComponentKey& ComponentKey)
+{
+	UActorComponent* FoundTemplate = nullptr;
+	if (ComponentKey.IsValid() && ComponentKey.IsUCSKey())
+	{
+		UBlueprint* Blueprint = Cast<UBlueprint>(ComponentKey.GetComponentOwner()->ClassGeneratedBy);
+		check(Blueprint != nullptr);
+
+		if (UEdGraph* UCSGraph = FBlueprintEditorUtils::FindUserConstructionScript(Blueprint))
+		{
+			TArray<UK2Node_AddComponent*> ComponentNodes;
+			UCSGraph->GetNodesOfClass<UK2Node_AddComponent>(ComponentNodes);
+
+			for (UK2Node_AddComponent* UCSNode : ComponentNodes)
+			{
+				if (UCSNode->NodeGuid == ComponentKey.GetAssociatedGuid())
+				{
+					FoundTemplate = UCSNode->GetTemplateFromNode();
+					break;
+				}
+			}
+		}
+	}
+	return FoundTemplate;
+}
+
 /** Temporary fix for cut-n-paste error that failed to carry transactional flags */
 void FBlueprintEditorUtils::UpdateTransactionalFlags(UBlueprint* Blueprint)
 {
diff --git a/Engine/Source/Editor/UnrealEd/Public/Kismet2/BlueprintEditorUtils.h b/Engine/Source/Editor/UnrealEd/Public/Kismet2/BlueprintEditorUtils.h
--- a/Engine/Source/Editor/UnrealEd/Public/Kismet2/BlueprintEditorUtils.h
+++ b/Engine/Source/Editor/UnrealEd/Public/Kismet2/BlueprintEditorUtils.h
@@ -7,7 +7,8 @@
 #include "EdGraphSchema_K2.h"
 #include "K2Node_EditablePinBase.h"
 
-class USCS_Node;
+class  USCS_Node;
+struct FComponentKey;
 
 /** 
   * Flags describing how to handle graph removal
@@ -81,6 +82,16 @@ public:
 
 };
 
+struct UNREALED_API FUCSComponentId
+{
+public:
+	FUCSComponentId(const class UK2Node_AddComponent* UCSNode);
+	FGuid GetAssociatedGuid() const { return GraphNodeGuid; }
+
+private:
+	FGuid GraphNodeGuid;
+};
+
 DECLARE_CYCLE_STAT_EXTERN(TEXT("Notify Blueprint Changed"), EKismetCompilerStats_NotifyBlueprintChanged, STATGROUP_KismetCompiler, );
 
 /** Array type for GetCompilerRelevantNodes()                                                                                                                                                                                                                                                                                                                                                                                                                                                                                                                                                                                 */
@@ -998,6 +1009,9 @@ public:
 	/** Determines if this property is associated with a component that would be displayed in the SCS editor */
 	static bool IsSCSComponentProperty(UObjectProperty* MemberProperty);
 
+	/** Attempts to match up the FComponentKey with a ComponentTemplate from the Blueprint's UCS */
+	static UActorComponent* FindUCSComponentTemplate(const FComponentKey& ComponentKey);
+
 	//////////////////////////////////////////////////////////////////////////
 	// Interface
 
diff --git a/Engine/Source/Runtime/Engine/Classes/Engine/InheritableComponentHandler.h b/Engine/Source/Runtime/Engine/Classes/Engine/InheritableComponentHandler.h
--- a/Engine/Source/Runtime/Engine/Classes/Engine/InheritableComponentHandler.h
+++ b/Engine/Source/Runtime/Engine/Classes/Engine/InheritableComponentHandler.h
@@ -4,38 +4,59 @@
 #pragma once
 #include "InheritableComponentHandler.generated.h"
 
-class USCS_Node;
-class UActorComponent;
+class  USCS_Node;
+class  UActorComponent;
+struct FUCSComponentId;
+class  UBlueprint;
 
 USTRUCT()
 struct ENGINE_API FComponentKey
 {
 	GENERATED_USTRUCT_BODY()
 
-	UPROPERTY()
-	UBlueprintGeneratedClass* OwnerClass;
-
-	UPROPERTY()
-	FName VariableName;
-
-	UPROPERTY()
-	FGuid VariableGuid;
-
 	FComponentKey()
 		: OwnerClass(nullptr)
 	{}
 
 	FComponentKey(USCS_Node* SCSNode);
+#if WITH_EDITOR
+	FComponentKey(UBlueprint* Blueprint, const FUCSComponentId& UCSComponentID);
+#endif 
 
 	bool Match(const FComponentKey OtherKey) const;
 
+	bool IsSCSKey() const
+	{
+		return (SCSVariableName != NAME_None) && AssociatedGuid.IsValid();
+	}
+
+	bool IsUCSKey() const
+	{
+		return AssociatedGuid.IsValid() && (SCSVariableName == NAME_None);
+	}
+
 	bool IsValid() const
 	{
-		return OwnerClass && (VariableName != NAME_None) && VariableGuid.IsValid();
+		return OwnerClass && AssociatedGuid.IsValid() && (!IsSCSKey() || (SCSVariableName != NAME_None));
 	}
 
 	USCS_Node* FindSCSNode() const;
 	UActorComponent* GetOriginalTemplate() const;
+	bool RefreshVariableName();
+
+	UClass* GetComponentOwner()  const { return OwnerClass; }
+	FName   GetSCSVariableName() const { return SCSVariableName; }
+	FGuid   GetAssociatedGuid()  const { return AssociatedGuid; }
+
+private: 
+	UPROPERTY()
+	UBlueprintGeneratedClass* OwnerClass;
+
+	UPROPERTY()
+	FName SCSVariableName;
+
+	UPROPERTY()
+	FGuid AssociatedGuid;
 };
 
 USTRUCT()
@@ -86,7 +107,7 @@ public:
 		return 0 == Records.Num();
 	}
 
-	bool RenameTemplate(FComponentKey OldKey, FName NewName);
+	bool RefreshTemplateName(FComponentKey OldKey);
 
 	FComponentKey FindKey(UActorComponent* ComponentTemplate) const;
 #endif
diff --git a/Engine/Source/Runtime/Engine/Private/InheritableComponentHandler.cpp b/Engine/Source/Runtime/Engine/Private/InheritableComponentHandler.cpp
--- a/Engine/Source/Runtime/Engine/Private/InheritableComponentHandler.cpp
+++ b/Engine/Source/Runtime/Engine/Private/InheritableComponentHandler.cpp
@@ -4,6 +4,10 @@
 #include "Engine/SCS_Node.h"
 #include "Engine/InheritableComponentHandler.h"
 
+#if WITH_EDITOR
+#include "Kismet2/BlueprintEditorUtils.h"
+#endif // WITH_EDITOR
+
 // UInheritableComponentHandler
 
 void UInheritableComponentHandler::PostLoad()
@@ -42,7 +46,7 @@ UActorComponent* UInheritableComponentHandler::CreateOverridenComponentTemplate(
 	if (!BestArchetype)
 	{
 		UE_LOG(LogBlueprint, Warning, TEXT("CreateOverridenComponentTemplate '%s': cannot find archetype for component '%s' from '%s'"),
-			*GetPathNameSafe(this), *Key.VariableName.ToString(), *GetPathNameSafe(Key.OwnerClass));
+			*GetPathNameSafe(this), *Key.GetSCSVariableName().ToString(), *GetPathNameSafe(Key.GetComponentOwner()));
 		return NULL;
 	}
 	ensure(Cast<UBlueprintGeneratedClass>(GetOuter()));
@@ -100,19 +104,17 @@ void UInheritableComponentHandler::ValidateTemplates()
 		bool bIsValidAndNecessary = false;
 		{
 			FComponentOverrideRecord& Record = Records[Index];
+			FComponentKey& ComponentKey = Record.ComponentKey;
 			
-			if (Record.ComponentKey.OwnerClass && Record.ComponentKey.VariableGuid.IsValid())
+			FName VarName = ComponentKey.GetSCSVariableName();
+			if (ComponentKey.RefreshVariableName())
 			{
-				// Update variable, based on value
-				auto SCSNode = Record.ComponentKey.FindSCSNode();
-				const FName ActualName = SCSNode ? SCSNode->GetVariableName() : NAME_None;
-				if ((ActualName != NAME_None) && (ActualName != Record.ComponentKey.VariableName))
-				{
-					UE_LOG(LogBlueprint, Log, TEXT("ValidateTemplates '%s': variable old name '%s' new name '%s'"),
-						*GetPathNameSafe(this), *Record.ComponentKey.VariableName.ToString(), *ActualName.ToString());
-					Record.ComponentKey.VariableName = ActualName;
-					MarkPackageDirty();
-				}
+				FName NewName = ComponentKey.GetSCSVariableName();
+				UE_LOG(LogBlueprint, Log, TEXT("ValidateTemplates '%s': variable old name '%s' new name '%s'"),
+					*GetPathNameSafe(this), *VarName.ToString(), *NewName.ToString());
+				VarName = NewName;
+
+				MarkPackageDirty();
 			}
 
 			if (IsRecordValid(Record))
@@ -124,13 +126,13 @@ void UInheritableComponentHandler::ValidateTemplates()
 				else
 				{
 					UE_LOG(LogBlueprint, Log, TEXT("ValidateTemplates '%s': overriden template is unnecessary - component '%s' from '%s'"),
-						*GetPathNameSafe(this), *Record.ComponentKey.VariableName.ToString(), *GetPathNameSafe(Record.ComponentKey.OwnerClass));
+						*GetPathNameSafe(this), *VarName.ToString(), *GetPathNameSafe(ComponentKey.GetComponentOwner()));
 				}
 			}
 			else
 			{
 				UE_LOG(LogBlueprint, Warning, TEXT("ValidateTemplates '%s': overriden template is invalid - component '%s' from '%s'"),
-					*GetPathNameSafe(this), *Record.ComponentKey.VariableName.ToString(), *GetPathNameSafe(Record.ComponentKey.OwnerClass));
+					*GetPathNameSafe(this), *VarName.ToString(), *GetPathNameSafe(ComponentKey.GetComponentOwner()));
 			}
 		}
 
@@ -237,7 +239,7 @@ UActorComponent* UInheritableComponentHandler::FindBestArchetype(FComponentKey K
 	UActorComponent* ClosestArchetype = nullptr;
 
 	auto ActualBPGC = Cast<UBlueprintGeneratedClass>(GetOuter());
-	if (ActualBPGC && Key.OwnerClass && (ActualBPGC != Key.OwnerClass))
+	if (ActualBPGC && Key.GetComponentOwner() && (ActualBPGC != Key.GetComponentOwner()))
 	{
 		ActualBPGC = Cast<UBlueprintGeneratedClass>(ActualBPGC->GetSuperClass());
 		while (!ClosestArchetype && ActualBPGC)
@@ -258,13 +260,13 @@ UActorComponent* UInheritableComponentHandler::FindBestArchetype(FComponentKey K
 	return ClosestArchetype;
 }
 
-bool UInheritableComponentHandler::RenameTemplate(FComponentKey OldKey, FName NewName)
+bool UInheritableComponentHandler::RefreshTemplateName(FComponentKey OldKey)
 {
 	for (auto& Record : Records)
 	{
 		if (Record.ComponentKey.Match(OldKey))
 		{
-			Record.ComponentKey.VariableName = NewName;
+			Record.ComponentKey.RefreshVariableName();
 			return true;
 		}
 	}
@@ -317,7 +319,7 @@ FComponentKey UInheritableComponentHandler::FindKey(const FName VariableName) co
 {
 	for (const FComponentOverrideRecord& Record : Records)
 	{
-		if (Record.ComponentKey.VariableName == VariableName || (Record.ComponentTemplate && Record.ComponentTemplate->GetFName() == VariableName))
+		if (Record.ComponentKey.GetSCSVariableName() == VariableName || (Record.ComponentTemplate && Record.ComponentTemplate->GetFName() == VariableName))
 		{
 			return Record.ComponentKey;
 		}
@@ -349,26 +351,63 @@ FComponentKey::FComponentKey(USCS_Node* SCSNode) : OwnerClass(nullptr)
 {
 	if (SCSNode)
 	{
-		auto ParentSCS = SCSNode->GetSCS();
-		OwnerClass = ParentSCS ? Cast<UBlueprintGeneratedClass>(ParentSCS->GetOwnerClass()) : nullptr;
-		VariableGuid = SCSNode->VariableGuid;
-		VariableName = SCSNode->GetVariableName();
+		USimpleConstructionScript* ParentSCS = SCSNode->GetSCS();
+		OwnerClass      = ParentSCS ? Cast<UBlueprintGeneratedClass>(ParentSCS->GetOwnerClass()) : nullptr;
+		AssociatedGuid  = SCSNode->VariableGuid;
+		SCSVariableName = SCSNode->GetVariableName();
 	}
 }
 
+#if WITH_EDITOR
+FComponentKey::FComponentKey(UBlueprint* Blueprint, const FUCSComponentId& UCSComponentID)
+{
+	OwnerClass     = CastChecked<UBlueprintGeneratedClass>(Blueprint->GeneratedClass);
+	AssociatedGuid = UCSComponentID.GetAssociatedGuid();
+}
+#endif // WITH_EDITOR
+
 bool FComponentKey::Match(const FComponentKey OtherKey) const
 {
-	return (OwnerClass == OtherKey.OwnerClass) && (VariableGuid == OtherKey.VariableGuid);
+	return (OwnerClass == OtherKey.OwnerClass) && (AssociatedGuid == OtherKey.AssociatedGuid);
 }
 
 USCS_Node* FComponentKey::FindSCSNode() const
 {
-	auto ParentSCS = (OwnerClass && IsValid()) ? OwnerClass->SimpleConstructionScript : nullptr;
-	return ParentSCS ? ParentSCS->FindSCSNodeByGuid(VariableGuid) : nullptr;
+	USimpleConstructionScript* ParentSCS = (OwnerClass && IsSCSKey()) ? OwnerClass->SimpleConstructionScript : nullptr;
+	return ParentSCS ? ParentSCS->FindSCSNodeByGuid(AssociatedGuid) : nullptr;
 }
 
 UActorComponent* FComponentKey::GetOriginalTemplate() const
 {
-	auto SCSNode = FindSCSNode();
-	return SCSNode ? SCSNode->ComponentTemplate : nullptr;
+	UActorComponent* ComponentTemplate = nullptr;
+	if (IsSCSKey())
+	{
+		if (USCS_Node* SCSNode = FindSCSNode())
+		{
+			ComponentTemplate = SCSNode->ComponentTemplate;
+		}
+	}
+#if WITH_EDITOR
+	else if (IsUCSKey())
+	{
+		ComponentTemplate = FBlueprintEditorUtils::FindUCSComponentTemplate(*this);
+	}
+#endif // WITH_EDITOR
+	return ComponentTemplate;
+}
+
+bool FComponentKey::RefreshVariableName()
+{
+	if (IsValid() && IsSCSKey())
+	{
+		USCS_Node* SCSNode = FindSCSNode();
+		const FName UpdatedName = SCSNode ? SCSNode->GetVariableName() : NAME_None;
+
+		if (UpdatedName != SCSVariableName)
+		{
+			SCSVariableName = UpdatedName;
+			return true;
+		}
+	}
+	return false;
 }
