Remove commented out old-style logging (`log)

diff --git a/Engine/Source/Runtime/Engine/Private/HUD.cpp b/Engine/Source/Runtime/Engine/Private/HUD.cpp
--- a/Engine/Source/Runtime/Engine/Private/HUD.cpp
+++ b/Engine/Source/Runtime/Engine/Private/HUD.cpp
@@ -514,7 +514,6 @@ void AHUD::AddDebugText_Implementation(const FString& DebugText,
 		}
 		else
 		{
-			//`log("Adding debug text:"@DebugText@"for actor:"@SrcActor);
 			// search for an existing entry
 			int32 Idx = 0;
 			if (!bSkipOverwriteCheck)
diff --git a/Engine/Source/Runtime/Engine/Private/PlayerController.cpp b/Engine/Source/Runtime/Engine/Private/PlayerController.cpp
--- a/Engine/Source/Runtime/Engine/Private/PlayerController.cpp
+++ b/Engine/Source/Runtime/Engine/Private/PlayerController.cpp
@@ -2312,7 +2312,6 @@ void APlayerController::SetCinematicMode( bool bInCinematicMode, bool bAffectsMo
 void APlayerController::SetIgnoreMoveInput( bool bNewMoveInput )
 {
 	IgnoreMoveInput = FMath::Max( IgnoreMoveInput + (bNewMoveInput ? +1 : -1), 0 );
-	//`Log("IgnoreMove: " $ IgnoreMoveInput);
 }
 
 void APlayerController::ResetIgnoreMoveInput()
@@ -2328,7 +2327,6 @@ bool APlayerController::IsMoveInputIgnored() const
 void APlayerController::SetIgnoreLookInput( bool bNewLookInput )
 {
 	IgnoreLookInput = FMath::Max( IgnoreLookInput + (bNewLookInput ? +1 : -1), 0 );
-	//`Log("IgnoreLook: " $ IgnoreLookInput);
 }
 
 void APlayerController::ResetIgnoreLookInput()
@@ -2863,7 +2861,6 @@ void APlayerController::ClientForceGarbageCollection_Implementation()
 
 void APlayerController::LevelStreamingStatusChanged(ULevelStreaming* LevelObject, bool bNewShouldBeLoaded, bool bNewShouldBeVisible, bool bNewShouldBlockOnLoad, int32 LODIndex )
 {
-	//`log( "LevelStreamingStatusChanged: " @ LevelObject @ bNewShouldBeLoaded @ bNewShouldBeVisible @ bNewShouldBeVisible );
 	ClientUpdateLevelStreamingStatus(LevelObject->GetWorldAssetPackageFName(),bNewShouldBeLoaded,bNewShouldBeVisible,bNewShouldBlockOnLoad,LODIndex);
 }
 
diff --git a/Engine/Source/Runtime/Engine/Private/UserInterface/Console.cpp b/Engine/Source/Runtime/Engine/Private/UserInterface/Console.cpp
--- a/Engine/Source/Runtime/Engine/Private/UserInterface/Console.cpp
+++ b/Engine/Source/Runtime/Engine/Private/UserInterface/Console.cpp
@@ -640,8 +640,6 @@ bool UConsole::InputChar_Typing( int32 ControllerId, const FString& Unicode )
 
 bool UConsole::InputKey_InputLine( int32 ControllerId, FKey Key, EInputEvent Event, float AmountDepressed, bool bGamepad)
 {
-	//`log(`location@`showvar(Key));
-
 	if ( Event == IE_Pressed )
 	{
 		bCaptureKeyInput = false;
diff --git a/Engine/Source/Runtime/Engine/Private/UserInterface/PlayerInput.cpp b/Engine/Source/Runtime/Engine/Private/UserInterface/PlayerInput.cpp
--- a/Engine/Source/Runtime/Engine/Private/UserInterface/PlayerInput.cpp
+++ b/Engine/Source/Runtime/Engine/Private/UserInterface/PlayerInput.cpp
@@ -1693,13 +1693,11 @@ void UPlayerInput::SetBind(FName BindName, const FString& Command)
 			if (DebugExecBindings[BindIndex].Key == BindKey)
 			{
 				DebugExecBindings[BindIndex].Command = CommandMod;
-				// `log("Binding '"@BindName@"' found, setting CommandMod '"@CommandMod@"'");
 				SaveConfig();
 				return;
 			}
 		}
 
-		// `log("Binding '"@BindName@"' NOT found, adding new binding with CommandMod '"@CommandMod@"'");
 		FKeyBind NewBind;
 		NewBind.Key = BindKey;
 		NewBind.Command = CommandMod;
