#ue4 - Prevent crash when overlap events (triggered by toggling bGenerateOverlaps or changing collision setting) destroy components or cause other overlaps to end.

- Also cleaned up redundant branches in EndComponentOverlap().

UE-13866

diff --git a/Engine/Source/Runtime/Engine/Private/Components/PrimitiveComponent.cpp b/Engine/Source/Runtime/Engine/Private/Components/PrimitiveComponent.cpp
--- a/Engine/Source/Runtime/Engine/Private/Components/PrimitiveComponent.cpp
+++ b/Engine/Source/Runtime/Engine/Private/Components/PrimitiveComponent.cpp
@@ -2012,13 +2012,14 @@ void UPrimitiveComponent::BeginComponentOverlap(const FOverlapInfo& OtherOverlap
 void UPrimitiveComponent::EndComponentOverlap(const FOverlapInfo& OtherOverlap, bool bDoNotifies, bool bNoNotifySelf)
 {
 	UPrimitiveComponent* OtherComp = OtherOverlap.OverlapInfo.Component.Get();
-
-	AActor* const OtherActor = OtherComp ? OtherComp->GetOwner() : NULL;
-	AActor* const MyActor = GetOwner();
+	if (OtherComp == nullptr)
+	{
+		return;
+	}
 
 	//	UE_LOG(LogActor, Log, TEXT("END OVERLAP! Self=%s SelfComp=%s, Other=%s, OtherComp=%s"), *GetNameSafe(this), *GetNameSafe(MyComp), *GetNameSafe(OtherActor), *GetNameSafe(OtherComp));
 
-	const int32 OtherOverlapIdx = OtherComp ? OtherComp->OverlappingComponents.Find(FOverlapInfo(this, INDEX_NONE)) : INDEX_NONE;
+	const int32 OtherOverlapIdx = OtherComp->OverlappingComponents.Find(FOverlapInfo(this, INDEX_NONE));
 	if (OtherOverlapIdx != INDEX_NONE)
 	{
 		OtherComp->OverlappingComponents.RemoveAtSwap(OtherOverlapIdx, 1, false);
@@ -2028,34 +2029,39 @@ void UPrimitiveComponent::EndComponentOverlap(const FOverlapInfo& OtherOverlap,
 	if (OverlapIdx != INDEX_NONE)
 	{
 		OverlappingComponents.RemoveAtSwap(OverlapIdx, 1, false);
-	
-		if ((OtherActor != NULL) && bDoNotifies)
-		{
-			if (!bNoNotifySelf && IsPrimCompValidAndAlive(this))
-			{
-				OnComponentEndOverlap.Broadcast(OtherActor, OtherComp, OtherOverlap.GetBodyIndex());
-			}
 
-			if (IsPrimCompValidAndAlive(OtherComp))
+		if (bDoNotifies)
+		{
+			AActor* const OtherActor = OtherComp->GetOwner();
+			AActor* const MyActor = GetOwner();
+			if (OtherActor)
 			{
-				OtherComp->OnComponentEndOverlap.Broadcast(MyActor, this, INDEX_NONE);
-			}
-		}
-	}
+				if (!bNoNotifySelf && IsPrimCompValidAndAlive(this))
+				{
+					OnComponentEndOverlap.Broadcast(OtherActor, OtherComp, OtherOverlap.GetBodyIndex());
+				}
 
-	// if this was the last touch on the other actor, notify that we've untouched the actor as well
-	if (bDoNotifies && MyActor && OtherActor && !MyActor->IsOverlappingActor(OtherActor) )
-	{			
-		if (IsActorValidToNotify(MyActor))
-		{
-			MyActor->NotifyActorEndOverlap(OtherActor);
-			MyActor->OnActorEndOverlap.Broadcast(OtherActor);
-		}
+				if (IsPrimCompValidAndAlive(OtherComp))
+				{
+					OtherComp->OnComponentEndOverlap.Broadcast(MyActor, this, INDEX_NONE);
+				}
+	
+				// if this was the last touch on the other actor by this actor, notify that we've untouched the actor as well
+				if (MyActor && !MyActor->IsOverlappingActor(OtherActor) )
+				{			
+					if (IsActorValidToNotify(MyActor))
+					{
+						MyActor->NotifyActorEndOverlap(OtherActor);
+						MyActor->OnActorEndOverlap.Broadcast(OtherActor);
+					}
 
-		if (IsActorValidToNotify(OtherActor))
-		{
-			OtherActor->NotifyActorEndOverlap(MyActor);
-			OtherActor->OnActorEndOverlap.Broadcast(MyActor);
+					if (IsActorValidToNotify(OtherActor))
+					{
+						OtherActor->NotifyActorEndOverlap(MyActor);
+						OtherActor->OnActorEndOverlap.Broadcast(MyActor);
+					}
+				}
+			}
 		}
 	}
 }
@@ -2394,10 +2400,7 @@ void UPrimitiveComponent::UpdateOverlaps(const TArray<FOverlapInfo>* NewPendingO
 				for (auto CompIt = OldOverlappingComponents.CreateIterator(); CompIt; ++CompIt)
 				{
 					const FOverlapInfo& OtherOverlap = *CompIt;
-					if (OtherOverlap.OverlapInfo.Component.IsValid())
-					{
-						EndComponentOverlap(OtherOverlap, bDoNotifies, false);
-					}
+					EndComponentOverlap(OtherOverlap, bDoNotifies, false);
 				}
 			}
 
@@ -2405,10 +2408,7 @@ void UPrimitiveComponent::UpdateOverlaps(const TArray<FOverlapInfo>* NewPendingO
 			for (auto CompIt = NewOverlappingComponents.CreateIterator(); CompIt; ++CompIt)
 			{
 				const FOverlapInfo& OtherOverlap = *CompIt;
-				if (OtherOverlap.OverlapInfo.Component.IsValid())
-				{
-					BeginComponentOverlap(OtherOverlap, bDoNotifies);
-				}
+				BeginComponentOverlap(OtherOverlap, bDoNotifies);
 			}
 		}
 	}
@@ -2416,12 +2416,12 @@ void UPrimitiveComponent::UpdateOverlaps(const TArray<FOverlapInfo>* NewPendingO
 	{
 		// bGenerateOverlapEvents is false or collision is disabled
 
-		// End all overlaps that exist, in case bGenerateOverlapEvents was true last tick (i.e. was just turned off)
-		// Iterate backwards since EndComponentOverlap will remove items from OverlappingComponents.
-		for (int32 OverlapIdx = OverlappingComponents.Num()-1; OverlapIdx >= 0; --OverlapIdx)
+		if (OverlappingComponents.Num() > 0)
 		{
-			const FOverlapInfo& OtherOverlap = OverlappingComponents[OverlapIdx];
-			if (OtherOverlap.OverlapInfo.Component.IsValid())
+			// End all overlaps that exist, in case bGenerateOverlapEvents was true last tick (i.e. was just turned off)
+			// Make a copy since EndComponentOverlap will remove items from OverlappingComponents.
+			auto OverlapsCopy = OverlappingComponents;
+			for (const FOverlapInfo& OtherOverlap : OverlapsCopy)
 			{
 				EndComponentOverlap(OtherOverlap, bDoNotifies, false);
 			}
