[UE-17500] Revised fix for component data cache save/restore after reconstruction.

change summary:
- added 'SourceComponentTemplate' to FActorComponentInstanceData (to additionally track source component archetype)
- modified FActorComponentInstanceData::MatchesComponent() to include a parameter for passing in the template reference
- modified FActorComponentInstanceData::FActorComponentInstanceData() to limit calculation of the serialization index to UCS components only (SCS components are now matched by template)
- modified FActorComponentInstanceData::FActorComponentInstanceData() and MatchesComponent() to calculate the serialization index relative to the archetype rather than by component class
- modified FActorComponentInstanceData::AddReferencedObjects() to include the template reference in order to keep it from getting GC'd while in use
- modified FComponentInstanceDataCache::FComponentInstanceDataCache() to avoid caching component instances that are based on the component class CDO (i.e. not a unique archetype)
- removed FActorComponentInstanceData::SourceComponentName and SourceComponentClass (no longer being used for matching)

diff --git a/Engine/Source/Runtime/Engine/Private/ComponentInstanceDataCache.cpp b/Engine/Source/Runtime/Engine/Private/ComponentInstanceDataCache.cpp
--- a/Engine/Source/Runtime/Engine/Private/ComponentInstanceDataCache.cpp
+++ b/Engine/Source/Runtime/Engine/Private/ComponentInstanceDataCache.cpp
@@ -4,7 +4,7 @@
 #include "ComponentInstanceDataCache.h"
 
 FActorComponentInstanceData::FActorComponentInstanceData()
-	: SourceComponentClass(nullptr)
+	: SourceComponentTemplate(nullptr)
 	, SourceComponentTypeSerializedIndex(-1)
 	, SourceComponentCreationMethod(EComponentCreationMethod::Native)
 {
@@ -13,35 +13,38 @@ FActorComponentInstanceData::FActorComponentInstanceData()
 FActorComponentInstanceData::FActorComponentInstanceData(const UActorComponent* SourceComponent)
 {
 	check(SourceComponent);
-	SourceComponentName = SourceComponent->GetFName();
-	SourceComponentClass = SourceComponent->GetClass();
+	SourceComponentTemplate = SourceComponent->GetArchetype();
 	SourceComponentCreationMethod = SourceComponent->CreationMethod;
 	SourceComponentTypeSerializedIndex = -1;
-	
-	AActor* ComponentOwner = SourceComponent->GetOwner();
-	if (ComponentOwner)
+
+	// UCS components can share the same template (e.g. an AddComponent node inside a loop), so we also cache their serialization index here (relative to the shared template) as a means for identification
+	if(SourceComponentCreationMethod == EComponentCreationMethod::UserConstructionScript)
 	{
-		bool bFound = false;
-		for (const UActorComponent* BlueprintCreatedComponent : ComponentOwner->BlueprintCreatedComponents)
+		AActor* ComponentOwner = SourceComponent->GetOwner();
+		if (ComponentOwner)
 		{
-			if (BlueprintCreatedComponent)
+			bool bFound = false;
+			for (const UActorComponent* BlueprintCreatedComponent : ComponentOwner->BlueprintCreatedComponents)
 			{
-				if (BlueprintCreatedComponent == SourceComponent)
-				{
-					++SourceComponentTypeSerializedIndex;
-					bFound = true;
-					break;
-				}
-				else if (   BlueprintCreatedComponent->GetClass() == SourceComponentClass
-						 && BlueprintCreatedComponent->CreationMethod == SourceComponentCreationMethod)
+				if (BlueprintCreatedComponent)
 				{
-					++SourceComponentTypeSerializedIndex;
+					if (BlueprintCreatedComponent == SourceComponent)
+					{
+						++SourceComponentTypeSerializedIndex;
+						bFound = true;
+						break;
+					}
+					else if (   BlueprintCreatedComponent->CreationMethod == SourceComponentCreationMethod
+						&& BlueprintCreatedComponent->GetArchetype() == SourceComponentTemplate)
+					{
+						++SourceComponentTypeSerializedIndex;
+					}
 				}
 			}
-		}
-		if (!bFound)
-		{
-			SourceComponentTypeSerializedIndex = -1;
+			if (!bFound)
+			{
+				SourceComponentTypeSerializedIndex = -1;
+			}
 		}
 	}
 
@@ -73,20 +76,18 @@ FActorComponentInstanceData::FActorComponentInstanceData(const UActorComponent*
 	}
 }
 
-bool FActorComponentInstanceData::MatchesComponent(const UActorComponent* Component) const
+bool FActorComponentInstanceData::MatchesComponent(const UActorComponent* Component, const UObject* ComponentTemplate) const
 {
 	bool bMatches = false;
-	if (   Component 
-		&& (Component->CreationMethod == SourceComponentCreationMethod)
-		&& (Component->GetClass() == SourceComponentClass))
+	if (   Component
+		&& Component->CreationMethod == SourceComponentCreationMethod
+		&& ComponentTemplate == SourceComponentTemplate)
 	{
-		if (   (SourceComponentCreationMethod != EComponentCreationMethod::UserConstructionScript)
-			&& (Component->GetFName() == SourceComponentName))
+		if (SourceComponentCreationMethod != EComponentCreationMethod::UserConstructionScript)
 		{
 			bMatches = true;
 		}
-		else if (SourceComponentTypeSerializedIndex >= 0
-			&& SourceComponentCreationMethod != EComponentCreationMethod::SimpleConstructionScript)
+		else if (SourceComponentTypeSerializedIndex >= 0)
 		{
 			int32 FoundSerializedComponentsOfType = -1;
 			AActor* ComponentOwner = Component->GetOwner();
@@ -95,7 +96,7 @@ bool FActorComponentInstanceData::MatchesComponent(const UActorComponent* Compon
 				for (const UActorComponent* BlueprintCreatedComponent : ComponentOwner->BlueprintCreatedComponents)
 				{
 					if (   BlueprintCreatedComponent
-						&& (BlueprintCreatedComponent->GetClass() == SourceComponentClass)
+						&& (BlueprintCreatedComponent->GetArchetype() == SourceComponentTemplate)
 						&& (BlueprintCreatedComponent->CreationMethod == SourceComponentCreationMethod)
 						&& (++FoundSerializedComponentsOfType == SourceComponentTypeSerializedIndex))
 					{
@@ -147,7 +148,7 @@ void FActorComponentInstanceData::ApplyToComponent(UActorComponent* Component, c
 
 void FActorComponentInstanceData::AddReferencedObjects(FReferenceCollector& Collector)
 {
-	Collector.AddReferencedObject(SourceComponentClass);
+	Collector.AddReferencedObject(SourceComponentTemplate);
 }
 
 FComponentInstanceDataCache::FComponentInstanceDataCache(const AActor* Actor)
@@ -163,10 +164,15 @@ FComponentInstanceDataCache::FComponentInstanceDataCache(const AActor* Actor)
 		{
 			if (Component->IsCreatedByConstructionScript()) // Only cache data from 'created by construction script' components
 			{
-				FActorComponentInstanceData* ComponentInstanceData = Component->GetComponentInstanceData();
-				if (ComponentInstanceData)
+				// Also exclude instances that are based on the component CDO (i.e. not a unique archetype) - these are not going to be editable anyway.
+				const UObject* ComponentTemplate = Component->GetArchetype();
+				if (ComponentTemplate && ComponentTemplate != Component->GetClass()->GetDefaultObject())
 				{
-					ComponentsInstanceData.Add(ComponentInstanceData);
+					FActorComponentInstanceData* ComponentInstanceData = Component->GetComponentInstanceData();
+					if (ComponentInstanceData)
+					{
+						ComponentsInstanceData.Add(ComponentInstanceData);
+					}
 				}
 			}
 			else if (Component->CreationMethod == EComponentCreationMethod::Instance)
@@ -200,18 +206,22 @@ void FComponentInstanceDataCache::ApplyToActor(AActor* Actor, const ECacheApplyP
 		Actor->GetComponents(Components);
 
 		// Apply per-instance data.
-		for (UActorComponent* Component : Components)
+		for (UActorComponent* ComponentInstance : Components)
 		{
-			if(Component->IsCreatedByConstructionScript()) // Only try and apply data to 'created by construction script' components
+			if(ComponentInstance && ComponentInstance->IsCreatedByConstructionScript()) // Only try and apply data to 'created by construction script' components
 			{
-				for (FActorComponentInstanceData* ComponentInstanceData : ComponentsInstanceData)
+				// Cache template here to avoid redundant calls in the loop below
+				if (const UObject* ComponentTemplate = ComponentInstance->GetArchetype())
 				{
-					if (	ComponentInstanceData 
-						&&	ComponentInstanceData->GetComponentClass() == Component->GetClass() // filter on class early to avoid unnecessary virtual and expensive tests
-						&&	ComponentInstanceData->MatchesComponent(Component))
+					for (FActorComponentInstanceData* ComponentInstanceData : ComponentsInstanceData)
 					{
-						ComponentInstanceData->ApplyToComponent(Component, CacheApplyPhase);
-						break;
+						if (	ComponentInstanceData
+							&&	ComponentInstanceData->GetComponentClass() == ComponentTemplate->GetClass() // filter on class early to avoid unnecessary virtual and expensive tests
+							&&	ComponentInstanceData->MatchesComponent(ComponentInstance, ComponentTemplate))
+						{
+							ComponentInstanceData->ApplyToComponent(ComponentInstance, CacheApplyPhase);
+							break;
+						}
 					}
 				}
 			}
diff --git a/Engine/Source/Runtime/Engine/Public/ComponentInstanceDataCache.h b/Engine/Source/Runtime/Engine/Public/ComponentInstanceDataCache.h
--- a/Engine/Source/Runtime/Engine/Public/ComponentInstanceDataCache.h
+++ b/Engine/Source/Runtime/Engine/Public/ComponentInstanceDataCache.h
@@ -24,7 +24,7 @@ public:
 	{}
 
 	/** Determines whether this component instance data matches the component */
-	bool MatchesComponent(const UActorComponent* Component) const;
+	bool MatchesComponent(const UActorComponent* Component, const UObject* ComponentTemplate) const;
 
 	/** Applies this component instance data to the supplied component */
 	virtual void ApplyToComponent(UActorComponent* Component, const ECacheApplyPhase CacheApplyPhase);
@@ -36,21 +36,17 @@ public:
 
 	bool ContainsSavedProperties() const { return SavedProperties.Num() > 0; }
 
-	UClass* GetComponentClass() const { return SourceComponentClass; }
+	const UClass* GetComponentClass() const { return SourceComponentTemplate ? SourceComponentTemplate->GetClass() : nullptr; }
 
 protected:
-	/** The name of the source component */
-	FName SourceComponentName;
-
-	/** The class type of the source component */
-	UClass* SourceComponentClass;
+	/** The template used to create the source component */
+	const UObject* SourceComponentTemplate;
 
 	/** The index of the source component in its owner's serialized array 
 		when filtered to just that component type */
 	int32 SourceComponentTypeSerializedIndex;
 
-	/** The index of the source component in its owner's serialized array 
-		when filtered to just that component type */
+	/** The method that was used to create the source component */
 	EComponentCreationMethod SourceComponentCreationMethod;
 
 	TArray<uint8> SavedProperties;
