Fix crash when adding an instance component with a native sub-component to an actor with a blueprint created root component.

diff --git a/Engine/Source/Runtime/Engine/Private/SimpleConstructionScript.cpp b/Engine/Source/Runtime/Engine/Private/SimpleConstructionScript.cpp
--- a/Engine/Source/Runtime/Engine/Private/SimpleConstructionScript.cpp
+++ b/Engine/Source/Runtime/Engine/Private/SimpleConstructionScript.cpp
@@ -455,10 +455,20 @@ void USimpleConstructionScript::ExecuteScriptOnActor(AActor* Actor, const FTrans
 				Actor->GetComponents(Components);
 				for (int32 Index = Components.Num()-1; Index >= 0; --Index)
 				{
-					if (Components[Index]->CreationMethod == EComponentCreationMethod::Instance)
+					USceneComponent* SceneComponent = Components[Index];
+					if (SceneComponent->CreationMethod == EComponentCreationMethod::Instance)
 					{
 						Components.RemoveAt(Index);
 					}
+					else
+					{
+						// Handle the native sub-component of an instance component case
+						USceneComponent* ParentSceneComponent = SceneComponent->GetTypedOuter<USceneComponent>();
+						if (ParentSceneComponent && ParentSceneComponent->CreationMethod == EComponentCreationMethod::Instance)
+						{
+							Components.RemoveAt(Index);
+						}
+					}
 				}
 
 				// Get the native root component; if it's not set, the first native scene component will be used as root. This matches what's done in the SCS editor.
@@ -512,7 +522,7 @@ void USimpleConstructionScript::ExecuteScriptOnActor(AActor* Actor, const FTrans
 				// get list of every component SCS created, in case some of them aren't in the attachment hierarchy any more (e.g. rigid bodies)
 				TInlineComponentArray<USceneComponent*> ComponentsAfterSCS;
 				Actor->GetComponents(ComponentsAfterSCS);
-				for (auto C : ComponentsAfterSCS)
+				for (USceneComponent* C : ComponentsAfterSCS)
 				{
 					if (Components.Contains(C) == false)
 					{
