#ue4 - Remove stale (null, pending kill) overlaps from OverlappingComponents array when an invalid overlap is found in UpdateOverlaps.

diff --git a/Engine/Source/Runtime/Engine/Private/Components/PrimitiveComponent.cpp b/Engine/Source/Runtime/Engine/Private/Components/PrimitiveComponent.cpp
--- a/Engine/Source/Runtime/Engine/Private/Components/PrimitiveComponent.cpp
+++ b/Engine/Source/Runtime/Engine/Private/Components/PrimitiveComponent.cpp
@@ -2400,7 +2400,15 @@ void UPrimitiveComponent::UpdateOverlaps(const TArray<FOverlapInfo>* NewPendingO
 				for (auto CompIt = OldOverlappingComponents.CreateIterator(); CompIt; ++CompIt)
 				{
 					const FOverlapInfo& OtherOverlap = *CompIt;
-					EndComponentOverlap(OtherOverlap, bDoNotifies, false);
+					if (OtherOverlap.OverlapInfo.Component.IsValid())
+					{
+						EndComponentOverlap(OtherOverlap, bDoNotifies, false);
+					}
+					else
+					{
+						// Remove stale item
+						OverlappingComponents.RemoveSingleSwap(OtherOverlap);
+					}
 				}
 			}
 
