added comments

diff --git a/Engine/Source/Runtime/Engine/Classes/Engine/Scene.h b/Engine/Source/Runtime/Engine/Classes/Engine/Scene.h
--- a/Engine/Source/Runtime/Engine/Classes/Engine/Scene.h
+++ b/Engine/Source/Runtime/Engine/Classes/Engine/Scene.h
@@ -881,11 +881,11 @@ struct FPostProcessSettings
 	UPROPERTY(interp, BlueprintReadWrite, Category=DepthOfField, AdvancedDisplay, meta=(ClampMin = "0.0", ClampMax = "1.0", editcondition = "bOverride_DepthOfFieldOcclusion", DisplayName = "Occlusion"))
 	float DepthOfFieldOcclusion;
 	
-	/** Color threshold to do full quality DOF */
+	/** Color threshold to do full quality DOF (BokehDOF only) */
 	UPROPERTY(interp, BlueprintReadWrite, Category=DepthOfField, AdvancedDisplay, meta=(ClampMin = "0.0", ClampMax = "10.0", editcondition = "bOverride_DepthOfFieldColorThreshold", DisplayName = "Color Threshold"))
 	float DepthOfFieldColorThreshold;
 
-	/** Size threshold to do full quality DOF */
+	/** Size threshold to do full quality DOF (BokehDOF only) */
 	UPROPERTY(interp, BlueprintReadWrite, Category=DepthOfField, AdvancedDisplay, meta=(ClampMin = "0.0", ClampMax = "1.0", editcondition = "bOverride_DepthOfFieldSizeThreshold", DisplayName = "Size Threshold"))
 	float DepthOfFieldSizeThreshold;
 	
