nicer UI for DOF PostProcessSettings (filter parameters on DOF type)

diff --git a/Engine/Source/Runtime/Engine/Private/PostProcessVolume.cpp b/Engine/Source/Runtime/Engine/Private/PostProcessVolume.cpp
--- a/Engine/Source/Runtime/Engine/Private/PostProcessVolume.cpp
+++ b/Engine/Source/Runtime/Engine/Private/PostProcessVolume.cpp
@@ -61,6 +61,39 @@ bool APostProcessVolume::CanEditChange(const UProperty* InProperty) const
 	{
 		FString PropertyName = InProperty->GetName();
 
+		// Settings, can be shared for multiple objects types (volume, component, camera, player)
+		{
+			if (PropertyName == GET_MEMBER_NAME_STRING_CHECKED(FPostProcessSettings, DepthOfFieldScale) ||
+				PropertyName == GET_MEMBER_NAME_STRING_CHECKED(FPostProcessSettings, DepthOfFieldMaxBokehSize) ||
+				PropertyName == GET_MEMBER_NAME_STRING_CHECKED(FPostProcessSettings, DepthOfFieldColorThreshold) ||
+				PropertyName == GET_MEMBER_NAME_STRING_CHECKED(FPostProcessSettings, DepthOfFieldSizeThreshold) ||
+				PropertyName == GET_MEMBER_NAME_STRING_CHECKED(FPostProcessSettings, DepthOfFieldBokehShape))
+			{
+				return Settings.DepthOfFieldMethod == EDepthOfFieldMethod::DOFM_BokehDOF;
+			}
+
+			if (PropertyName == GET_MEMBER_NAME_STRING_CHECKED(FPostProcessSettings, DepthOfFieldNearBlurSize) ||
+				PropertyName == GET_MEMBER_NAME_STRING_CHECKED(FPostProcessSettings, DepthOfFieldFarBlurSize) ||
+				PropertyName == GET_MEMBER_NAME_STRING_CHECKED(FPostProcessSettings, DepthOfFieldSkyFocusDistance))
+			{
+				return Settings.DepthOfFieldMethod == EDepthOfFieldMethod::DOFM_Gaussian;
+			}
+
+			if (PropertyName == GET_MEMBER_NAME_STRING_CHECKED(FPostProcessSettings, DepthOfFieldNearTransitionRegion) ||
+				PropertyName == GET_MEMBER_NAME_STRING_CHECKED(FPostProcessSettings, DepthOfFieldFarTransitionRegion) ||
+				PropertyName == GET_MEMBER_NAME_STRING_CHECKED(FPostProcessSettings, DepthOfFieldFocalRegion))
+			{
+				return Settings.DepthOfFieldMethod != EDepthOfFieldMethod::DOFM_CircleDOF;
+			}
+
+			if (PropertyName == GET_MEMBER_NAME_STRING_CHECKED(FPostProcessSettings, DepthOfFieldDepthBlurAmount) ||
+				PropertyName == GET_MEMBER_NAME_STRING_CHECKED(FPostProcessSettings, DepthOfFieldDepthBlurRadius) ||
+				PropertyName == GET_MEMBER_NAME_STRING_CHECKED(FPostProcessSettings, DepthOfFieldFstop))
+			{
+				return Settings.DepthOfFieldMethod == EDepthOfFieldMethod::DOFM_CircleDOF;
+			}
+		}
+
 		if (PropertyName == GET_MEMBER_NAME_STRING_CHECKED(APostProcessVolume, bEnabled))
 		{
 			return true;
