polish material property UI

diff --git a/Engine/Source/Runtime/Engine/Classes/Materials/Material.h b/Engine/Source/Runtime/Engine/Classes/Materials/Material.h
--- a/Engine/Source/Runtime/Engine/Classes/Materials/Material.h
+++ b/Engine/Source/Runtime/Engine/Classes/Materials/Material.h
@@ -341,7 +341,7 @@ private:
 public:
 
 	/** If BlendMode is BLEND_Masked, the surface is not rendered where OpacityMask < OpacityMaskClipValue. */
-	UPROPERTY(EditAnywhere, Category=Material)
+	UPROPERTY(EditAnywhere, Category=Material, AdvancedDisplay)
 	float OpacityMaskClipValue;
 
 	/** Adds to world position in the vertex shader. */
@@ -393,14 +393,14 @@ public:
 	FScalarMaterialInput PixelDepthOffset;
 
 	/** Indicates that the material should be rendered in the SeparateTranslucency Pass (not affected by DOF, requires bAllowSeparateTranslucency to be set in .ini). */
-	UPROPERTY(EditAnywhere, Category=Translucency, meta=(DisplayName = "Separate Translucency"))
+	UPROPERTY(EditAnywhere, Category=Translucency, meta=(DisplayName = "Separate Translucency"), AdvancedDisplay)
 	uint32 bEnableSeparateTranslucency:1;
 
 	/**
 	 * Indicates that the material should be rendered using responsive anti-aliasing. Improves sharpness of small moving particles such as sparks.
 	 * Only use for small moving features because it will cause aliasing of the background.
 	 */
-	UPROPERTY(EditAnywhere, Category=Translucency, meta=(DisplayName = "Responsive AA"))
+	UPROPERTY(EditAnywhere, Category=Translucency, meta=(DisplayName = "Responsive AA"), AdvancedDisplay)
 	uint32 bEnableResponsiveAA:1;
 
 	/** SSR on translucency */
@@ -469,7 +469,7 @@ public:
 	float TranslucentShadowStartOffset;
 
 	/** Whether to draw on top of opaque pixels even if behind them. This only has meaning for translucency. */
-	UPROPERTY(EditAnywhere, Category=Translucency)
+	UPROPERTY(EditAnywhere, Category=Translucency, AdvancedDisplay)
 	uint32 bDisableDepthTest:1;
 
 	/** Whether to generate spherical normals for particles that use this material. */
@@ -480,7 +480,7 @@ public:
 	 * Whether the material takes a tangent space normal or a world space normal as input.
 	 * (TangentSpace requires extra instructions but is often more convenient).
 	 */
-	UPROPERTY(EditAnywhere, Category=Material)
+	UPROPERTY(EditAnywhere, Category=Material, AdvancedDisplay)
 	uint32 bTangentSpaceNormal:1;
 
 	/**
