#ue4 - Add ETeleportType flag to MovementComponent::MoveUpdatedComponent(Impl), SafeMoveUpdatedComponent to match SceneComponent::MoveComponent.

- Note: currently always treated as not teleporting when within FScopedMovementUpdate (UE-17217 will address that).

UE-17552

diff --git a/Engine/Source/Runtime/Engine/Classes/GameFramework/MovementComponent.h b/Engine/Source/Runtime/Engine/Classes/GameFramework/MovementComponent.h
--- a/Engine/Source/Runtime/Engine/Classes/GameFramework/MovementComponent.h
+++ b/Engine/Source/Runtime/Engine/Classes/GameFramework/MovementComponent.h
@@ -250,14 +250,15 @@ public:
 	 * Moves our UpdatedComponent by the given Delta, and sets rotation to NewRotation. Respects the plane constraint, if enabled.
 	 * @note This simply calls the virtual MoveUpdatedComponentImpl() which can be overridden to implement custom behavior.
 	 * @note The overload taking rotation as an FQuat is slightly faster than the version using FRotator (which will be converted to an FQuat).
+	 * @note The 'Teleport' flag is currently always treated as 'None' (not teleporting) when used in an active FScopedMovementUpdate.
 	 * @return True if some movement occurred, false if no movement occurred. Result of any impact will be stored in OutHit.
 	 */
-	bool MoveUpdatedComponent(const FVector& Delta, const FQuat& NewRotation,    bool bSweep, FHitResult* OutHit = NULL);
-	bool MoveUpdatedComponent(const FVector& Delta, const FRotator& NewRotation, bool bSweep, FHitResult* OutHit = NULL);
+	bool MoveUpdatedComponent(const FVector& Delta, const FQuat& NewRotation,    bool bSweep, FHitResult* OutHit = NULL, ETeleportType Teleport = ETeleportType::None);
+	bool MoveUpdatedComponent(const FVector& Delta, const FRotator& NewRotation, bool bSweep, FHitResult* OutHit = NULL, ETeleportType Teleport = ETeleportType::None);
 
 protected:
 
-	virtual bool MoveUpdatedComponentImpl(const FVector& Delta, const FQuat& NewRotation, bool bSweep, FHitResult* OutHit = NULL);
+	virtual bool MoveUpdatedComponentImpl(const FVector& Delta, const FQuat& NewRotation, bool bSweep, FHitResult* OutHit = NULL, ETeleportType Teleport = ETeleportType::None);
 
 public:
 	
@@ -266,17 +267,18 @@ public:
 	 * Respects the plane constraint, if enabled.
 	 * @return True if some movement occurred, false if no movement occurred. Result of any impact will be stored in OutHit.
 	 */
-	UFUNCTION(BlueprintCallable, Category="Components|Movement", meta=(DisplayName = "MoveUpdatedComponent"))
-	bool K2_MoveUpdatedComponent(FVector Delta, FRotator NewRotation, FHitResult& OutHit, bool bSweep = true);
+	UFUNCTION(BlueprintCallable, Category="Components|Movement", meta=(DisplayName = "MoveUpdatedComponent", AdvancedDisplay="bTeleport"))
+	bool K2_MoveUpdatedComponent(FVector Delta, FRotator NewRotation, FHitResult& OutHit, bool bSweep = true, bool bTeleport = false);
 
 	/**
 	 * Calls MoveUpdatedComponent(), handling initial penetrations by calling ResolvePenetration().
 	 * If this adjustment succeeds, the original movement will be attempted again.
-	 * @note The overload taking rotation as an FQuat is slightly faster than the version using FRotator (which will be converted to an FQuat)..
+	 * @note The overload taking rotation as an FQuat is slightly faster than the version using FRotator (which will be converted to an FQuat).
+	 * @note The 'Teleport' flag is currently always treated as 'None' (not teleporting) when used in an active FScopedMovementUpdate.
 	 * @return result of the final MoveUpdatedComponent() call.
 	 */
-	bool SafeMoveUpdatedComponent(const FVector& Delta, const FQuat& NewRotation,    bool bSweep, FHitResult& OutHit);
-	bool SafeMoveUpdatedComponent(const FVector& Delta, const FRotator& NewRotation, bool bSweep, FHitResult& OutHit);
+	bool SafeMoveUpdatedComponent(const FVector& Delta, const FQuat& NewRotation,    bool bSweep, FHitResult& OutHit, ETeleportType Teleport = ETeleportType::None);
+	bool SafeMoveUpdatedComponent(const FVector& Delta, const FRotator& NewRotation, bool bSweep, FHitResult& OutHit, ETeleportType Teleport = ETeleportType::None);
 
 	/**
 	 * Calculate a movement adjustment to try to move out of a penetration from a failed move.
@@ -451,19 +453,19 @@ inline EPlaneConstraintAxisSetting UMovementComponent::GetPlaneConstraintAxisSet
 	return PlaneConstraintAxisSetting;
 }
 
-FORCEINLINE_DEBUGGABLE bool UMovementComponent::SafeMoveUpdatedComponent(const FVector& Delta, const FRotator& NewRotation, bool bSweep, FHitResult& OutHit)
+FORCEINLINE_DEBUGGABLE bool UMovementComponent::SafeMoveUpdatedComponent(const FVector& Delta, const FRotator& NewRotation, bool bSweep, FHitResult& OutHit, ETeleportType Teleport)
 {
-	return SafeMoveUpdatedComponent(Delta, NewRotation.Quaternion(), bSweep, OutHit);
+	return SafeMoveUpdatedComponent(Delta, NewRotation.Quaternion(), bSweep, OutHit, Teleport);
 }
 
-FORCEINLINE_DEBUGGABLE bool UMovementComponent::MoveUpdatedComponent(const FVector& Delta, const FQuat& NewRotation, bool bSweep, FHitResult* OutHit)
+FORCEINLINE_DEBUGGABLE bool UMovementComponent::MoveUpdatedComponent(const FVector& Delta, const FQuat& NewRotation, bool bSweep, FHitResult* OutHit, ETeleportType Teleport)
 {
-	return MoveUpdatedComponentImpl(Delta, NewRotation, bSweep, OutHit);
+	return MoveUpdatedComponentImpl(Delta, NewRotation, bSweep, OutHit, Teleport);
 }
 
-FORCEINLINE_DEBUGGABLE bool UMovementComponent::MoveUpdatedComponent(const FVector& Delta, const FRotator& NewRotation, bool bSweep, FHitResult* OutHit)
+FORCEINLINE_DEBUGGABLE bool UMovementComponent::MoveUpdatedComponent(const FVector& Delta, const FRotator& NewRotation, bool bSweep, FHitResult* OutHit, ETeleportType Teleport)
 {
-	return MoveUpdatedComponentImpl(Delta, NewRotation.Quaternion(), bSweep, OutHit);
+	return MoveUpdatedComponentImpl(Delta, NewRotation.Quaternion(), bSweep, OutHit, Teleport);
 }
 
 FORCEINLINE_DEBUGGABLE bool UMovementComponent::ResolvePenetration(const FVector& Adjustment, const FHitResult& Hit, const FQuat& NewRotation)
diff --git a/Engine/Source/Runtime/Engine/Private/Components/MovementComponent.cpp b/Engine/Source/Runtime/Engine/Private/Components/MovementComponent.cpp
--- a/Engine/Source/Runtime/Engine/Private/Components/MovementComponent.cpp
+++ b/Engine/Source/Runtime/Engine/Private/Components/MovementComponent.cpp
@@ -458,25 +458,25 @@ void UMovementComponent::SnapUpdatedComponentToPlane()
 
 
 
-bool UMovementComponent::MoveUpdatedComponentImpl( const FVector& Delta, const FQuat& NewRotation, bool bSweep, FHitResult* OutHit)
+bool UMovementComponent::MoveUpdatedComponentImpl( const FVector& Delta, const FQuat& NewRotation, bool bSweep, FHitResult* OutHit, ETeleportType Teleport)
 {
 	if (UpdatedComponent)
 	{
 		const FVector NewDelta = ConstrainDirectionToPlane(Delta);
-		return UpdatedComponent->MoveComponent(NewDelta, NewRotation, bSweep, OutHit, MoveComponentFlags);
+		return UpdatedComponent->MoveComponent(NewDelta, NewRotation, bSweep, OutHit, MoveComponentFlags, Teleport);
 	}
 
 	return false;
 }
 
 
-bool UMovementComponent::K2_MoveUpdatedComponent(FVector Delta, FRotator NewRotation, FHitResult& OutHit, bool bSweep)
+bool UMovementComponent::K2_MoveUpdatedComponent(FVector Delta, FRotator NewRotation, FHitResult& OutHit, bool bSweep, bool bTeleport)
 {
-	return SafeMoveUpdatedComponent(Delta, NewRotation.Quaternion(), bSweep, OutHit);
+	return SafeMoveUpdatedComponent(Delta, NewRotation.Quaternion(), bSweep, OutHit, TeleportFlagToEnum(bTeleport));
 }
 
 
-bool UMovementComponent::SafeMoveUpdatedComponent(const FVector& Delta, const FQuat& NewRotation, bool bSweep, FHitResult& OutHit)
+bool UMovementComponent::SafeMoveUpdatedComponent(const FVector& Delta, const FQuat& NewRotation, bool bSweep, FHitResult& OutHit, ETeleportType Teleport)
 {
 	if (UpdatedComponent == NULL)
 	{
@@ -484,7 +484,7 @@ bool UMovementComponent::SafeMoveUpdatedComponent(const FVector& Delta, const FQ
 		return false;
 	}
 
-	bool bMoveResult = MoveUpdatedComponent(Delta, NewRotation, bSweep, &OutHit);
+	bool bMoveResult = MoveUpdatedComponent(Delta, NewRotation, bSweep, &OutHit, Teleport);
 
 	// Handle initial penetrations
 	if (OutHit.bStartPenetrating && UpdatedComponent)
@@ -493,7 +493,7 @@ bool UMovementComponent::SafeMoveUpdatedComponent(const FVector& Delta, const FQ
 		if (ResolvePenetration(RequestedAdjustment, OutHit, NewRotation))
 		{
 			// Retry original move
-			bMoveResult = MoveUpdatedComponent(Delta, NewRotation, bSweep, &OutHit);
+			bMoveResult = MoveUpdatedComponent(Delta, NewRotation, bSweep, &OutHit, Teleport);
 		}
 	}
 
@@ -548,7 +548,7 @@ bool UMovementComponent::ResolvePenetrationImpl(const FVector& ProposedAdjustmen
 		if (!bEncroached)
 		{
 			// Move without sweeping.
-			MoveUpdatedComponent(Adjustment, NewRotationQuat, false);
+			MoveUpdatedComponent(Adjustment, NewRotationQuat, false, nullptr, ETeleportType::TeleportPhysics);
 			UE_LOG(LogMovement, Verbose, TEXT("ResolvePenetration: teleport %s by %s"), *ActorOwner->GetName(), *Adjustment.ToString());
 			return true;
 		}
@@ -559,7 +559,7 @@ bool UMovementComponent::ResolvePenetrationImpl(const FVector& ProposedAdjustmen
 
 			// Try sweeping as far as possible...
 			FHitResult SweepOutHit(1.f);
-			bool bMoved = MoveUpdatedComponent(Adjustment, NewRotationQuat, true, &SweepOutHit);
+			bool bMoved = MoveUpdatedComponent(Adjustment, NewRotationQuat, true, &SweepOutHit, ETeleportType::TeleportPhysics);
 			UE_LOG(LogMovement, Verbose, TEXT("ResolvePenetration: sweep %s by %s (success = %d)"), *ActorOwner->GetName(), *Adjustment.ToString(), bMoved);
 			
 			// Still stuck?
@@ -570,7 +570,7 @@ bool UMovementComponent::ResolvePenetrationImpl(const FVector& ProposedAdjustmen
 				const FVector CombinedMTD = Adjustment + SecondMTD;
 				if (SecondMTD != Adjustment && !CombinedMTD.IsZero())
 				{
-					bMoved = MoveUpdatedComponent(CombinedMTD, NewRotationQuat, true);
+					bMoved = MoveUpdatedComponent(CombinedMTD, NewRotationQuat, true, nullptr, ETeleportType::TeleportPhysics);
 					UE_LOG(LogMovement, Verbose, TEXT("ResolvePenetration: sweep %s by %s (MTD combo success = %d)"), *ActorOwner->GetName(), *CombinedMTD.ToString(), bMoved);
 				}
 			}
@@ -582,7 +582,7 @@ bool UMovementComponent::ResolvePenetrationImpl(const FVector& ProposedAdjustmen
 				const FVector MoveDelta = ConstrainDirectionToPlane(Hit.TraceEnd - Hit.TraceStart);
 				if (!MoveDelta.IsZero())
 				{
-					bMoved = MoveUpdatedComponent(Adjustment + MoveDelta, NewRotationQuat, true);
+					bMoved = MoveUpdatedComponent(Adjustment + MoveDelta, NewRotationQuat, true, nullptr, ETeleportType::TeleportPhysics);
 					UE_LOG(LogMovement, Verbose, TEXT("ResolvePenetration: sweep %s by %s (adjusted attempt success = %d)"), *ActorOwner->GetName(), *(Adjustment + MoveDelta).ToString(), bMoved);
 				}
 			}	
