#ue4 - Mark bTeleport' flag as AdvancedDisplay wherever we do so for the bSweep flag.

diff --git a/Engine/Source/Runtime/Engine/Classes/Components/SceneComponent.h b/Engine/Source/Runtime/Engine/Classes/Components/SceneComponent.h
--- a/Engine/Source/Runtime/Engine/Classes/Components/SceneComponent.h
+++ b/Engine/Source/Runtime/Engine/Classes/Components/SceneComponent.h
@@ -252,7 +252,7 @@ public:
 
 
 	/** Set the rotation of this component relative to its parent */
-	UFUNCTION(BlueprintCallable, Category="Utilities|Transformation", meta=(DisplayName="SetRelativeRotation", AdvancedDisplay="bSweep,SweepHitResult"))
+	UFUNCTION(BlueprintCallable, Category="Utilities|Transformation", meta=(DisplayName="SetRelativeRotation", AdvancedDisplay="bSweep,SweepHitResult,bTeleport"))
 	void K2_SetRelativeRotation(FRotator NewRotation, bool bSweep, FHitResult& SweepHitResult, bool bTeleport);
 	void SetRelativeRotation(FRotator NewRotation, bool bSweep=false, FHitResult* OutSweepHitResult=nullptr, ETeleportType Teleport = ETeleportType::None);
 	void SetRelativeRotation(const FQuat& NewRotation, bool bSweep=false, FHitResult* OutSweepHitResult=nullptr, ETeleportType Teleport = ETeleportType::None);
@@ -280,7 +280,7 @@ public:
 	void AddRelativeLocation(FVector DeltaLocation, bool bSweep=false, FHitResult* OutSweepHitResult=nullptr, ETeleportType Teleport = ETeleportType::None);
 
 	/** Adds a delta the rotation of this component relative to its parent */
-	UFUNCTION(BlueprintCallable, Category="Utilities|Transformation", meta=(DisplayName="AddRelativeRotation", AdvancedDisplay="bSweep,SweepHitResult"))
+	UFUNCTION(BlueprintCallable, Category="Utilities|Transformation", meta=(DisplayName="AddRelativeRotation", AdvancedDisplay="bSweep,SweepHitResult,bTeleport"))
 	void K2_AddRelativeRotation(FRotator DeltaRotation, bool bSweep, FHitResult& SweepHitResult, bool bTeleport);
 	void AddRelativeRotation(FRotator DeltaRotation, bool bSweep=false, FHitResult* OutSweepHitResult=nullptr, ETeleportType Teleport = ETeleportType::None);
 	void AddRelativeRotation(const FQuat& DeltaRotation, bool bSweep=false, FHitResult* OutSweepHitResult=nullptr, ETeleportType Teleport = ETeleportType::None);
@@ -292,7 +292,7 @@ public:
 
 	
 	/** Adds a delta to the rotation of this component in its local reference frame */
-	UFUNCTION(BlueprintCallable, Category="Utilities|Transformation", meta=(DisplayName="AddLocalRotation", AdvancedDisplay="bSweep,SweepHitResult"))
+	UFUNCTION(BlueprintCallable, Category="Utilities|Transformation", meta=(DisplayName="AddLocalRotation", AdvancedDisplay="bSweep,SweepHitResult,bTeleport"))
 	void K2_AddLocalRotation(FRotator DeltaRotation, bool bSweep, FHitResult& SweepHitResult, bool bTeleport);
 	void AddLocalRotation(FRotator DeltaRotation, bool bSweep=false, FHitResult* OutSweepHitResult=nullptr, ETeleportType Teleport = ETeleportType::None);
 	void AddLocalRotation(const FQuat& DeltaRotation, bool bSweep=false, FHitResult* OutSweepHitResult=nullptr, ETeleportType Teleport = ETeleportType::None);
@@ -311,7 +311,7 @@ public:
 
 
 	/** Set the relative rotation of this component to put it at the supplied orientation in world space. */
-	UFUNCTION(BlueprintCallable, Category="Utilities|Transformation", meta=(DisplayName="SetWorldRotation", AdvancedDisplay="bSweep,SweepHitResult"))
+	UFUNCTION(BlueprintCallable, Category="Utilities|Transformation", meta=(DisplayName="SetWorldRotation", AdvancedDisplay="bSweep,SweepHitResult,bTeleport"))
 	void K2_SetWorldRotation(FRotator NewRotation, bool bSweep, FHitResult& SweepHitResult, bool bTeleport);
 	void SetWorldRotation(FRotator NewRotation, bool bSweep=false, FHitResult* OutSweepHitResult=nullptr, ETeleportType Teleport = ETeleportType::None);
 	void SetWorldRotation(const FQuat& NewRotation, bool bSweep=false, FHitResult* OutSweepHitResult=nullptr, ETeleportType Teleport = ETeleportType::None);
@@ -334,7 +334,7 @@ public:
 
 
 	/** Adds a delta to the rotation of this component in world space. */
-	UFUNCTION(BlueprintCallable, Category="Utilities|Transformation", meta=(DisplayName="AddWorldRotation", AdvancedDisplay="bSweep,SweepHitResult"))
+	UFUNCTION(BlueprintCallable, Category="Utilities|Transformation", meta=(DisplayName="AddWorldRotation", AdvancedDisplay="bSweep,SweepHitResult,bTeleport"))
 	void K2_AddWorldRotation(FRotator DeltaRotation, bool bSweep, FHitResult& SweepHitResult, bool bTeleport);
 	void AddWorldRotation(FRotator DeltaRotation, bool bSweep=false, FHitResult* OutSweepHitResult=nullptr, ETeleportType Teleport = ETeleportType::None);
 	void AddWorldRotation(const FQuat& DeltaRotation, bool bSweep=false, FHitResult* OutSweepHitResult=nullptr, ETeleportType Teleport = ETeleportType::None);
diff --git a/Engine/Source/Runtime/Engine/Classes/GameFramework/Actor.h b/Engine/Source/Runtime/Engine/Classes/GameFramework/Actor.h
--- a/Engine/Source/Runtime/Engine/Classes/GameFramework/Actor.h
+++ b/Engine/Source/Runtime/Engine/Classes/GameFramework/Actor.h
@@ -774,7 +774,7 @@ public:
 	 * @param  bSweep				Whether to sweep to the target rotation (not currently supported).
 	 * @param  SweepHitResult		The hit result from the move if swept.
 	 */
-	UFUNCTION(BlueprintCallable, Category="Utilities|Transformation", meta=(DisplayName="AddActorWorldRotation", AdvancedDisplay="bSweep,SweepHitResult"))
+	UFUNCTION(BlueprintCallable, Category="Utilities|Transformation", meta=(DisplayName="AddActorWorldRotation", AdvancedDisplay="bSweep,SweepHitResult,bTeleport"))
 	void K2_AddActorWorldRotation(FRotator DeltaRotation, bool bSweep, FHitResult& SweepHitResult, bool bTeleport);
 	void AddActorWorldRotation(FRotator DeltaRotation, bool bSweep=false, FHitResult* OutSweepHitResult=nullptr, ETeleportType Teleport = ETeleportType::None);
 	void AddActorWorldRotation(const FQuat& DeltaRotation, bool bSweep=false, FHitResult* OutSweepHitResult=nullptr, ETeleportType Teleport = ETeleportType::None);
@@ -802,7 +802,7 @@ public:
 
 
 	/** Adds a delta to the rotation of this component in its local reference frame */
-	UFUNCTION(BlueprintCallable, Category="Utilities|Transformation", meta=(DisplayName="AddActorLocalRotation", AdvancedDisplay="bSweep,SweepHitResult"))
+	UFUNCTION(BlueprintCallable, Category="Utilities|Transformation", meta=(DisplayName="AddActorLocalRotation", AdvancedDisplay="bSweep,SweepHitResult,bTeleport"))
 	void K2_AddActorLocalRotation(FRotator DeltaRotation, bool bSweep, FHitResult& SweepHitResult, bool bTeleport);
 	void AddActorLocalRotation(FRotator DeltaRotation, bool bSweep=false, FHitResult* OutSweepHitResult=nullptr, ETeleportType Teleport = ETeleportType::None);
 	void AddActorLocalRotation(const FQuat& DeltaRotation, bool bSweep=false, FHitResult* OutSweepHitResult=nullptr, ETeleportType Teleport = ETeleportType::None);
@@ -828,7 +828,7 @@ public:
 	 * @param NewRelativeRotation		New relative rotation to set the actor's RootComponent to
 	 * @param bSweep					Should we sweep to the destination rotation. If true, will stop short of the target if blocked by something
 	 */
-	UFUNCTION(BlueprintCallable, Category="Utilities|Transformation", meta=(DisplayName="SetActorRelativeRotation", AdvancedDisplay="bSweep,SweepHitResult"))
+	UFUNCTION(BlueprintCallable, Category="Utilities|Transformation", meta=(DisplayName="SetActorRelativeRotation", AdvancedDisplay="bSweep,SweepHitResult,bTeleport"))
 	void K2_SetActorRelativeRotation(FRotator NewRelativeRotation, bool bSweep, FHitResult& SweepHitResult, bool bTeleport);
 	void SetActorRelativeRotation(FRotator NewRelativeRotation, bool bSweep=false, FHitResult* OutSweepHitResult=nullptr, ETeleportType Teleport = ETeleportType::None);
 	void SetActorRelativeRotation(const FQuat& NewRelativeRotation, bool bSweep=false, FHitResult* OutSweepHitResult=nullptr, ETeleportType Teleport = ETeleportType::None);
