#vr Setting the default for frame buffering to 1 frame

diff --git a/Engine/Source/Runtime/Engine/Public/StereoRendering.h b/Engine/Source/Runtime/Engine/Public/StereoRendering.h
--- a/Engine/Source/Runtime/Engine/Public/StereoRendering.h
+++ b/Engine/Source/Runtime/Engine/Public/StereoRendering.h
@@ -131,7 +131,7 @@ public:
 	/**
 	 * Returns number of required buffered frames.
 	 */
-	virtual uint32 GetNumberOfBufferedFrames() const { return 3;  }
+	virtual uint32 GetNumberOfBufferedFrames() const { return 1;  }
 
 	/**
 	 * Allocates a render target texture. 
