Fixed interto one_shot_reverse behaviour

diff --git a/Engine/Source/Runtime/Engine/Private/Components/InterpToMovementComponent.cpp b/Engine/Source/Runtime/Engine/Private/Components/InterpToMovementComponent.cpp
--- a/Engine/Source/Runtime/Engine/Private/Components/InterpToMovementComponent.cpp
+++ b/Engine/Source/Runtime/Engine/Private/Components/InterpToMovementComponent.cpp
@@ -293,15 +293,22 @@ void UInterpToMovementComponent::HandleImpact(const FHitResult& Hit, float Time,
 		switch(BehaviourType )
 		{
 		case EInterpToBehaviourType::OneShot:
-		case EInterpToBehaviourType::OneShot_Reverse:
 			OnInterpToStop.Broadcast(Hit, Time);
-			// If one shot we are done. If One shot reverse and we are already in reverse we are done too.
-			if ((BehaviourType == EInterpToBehaviourType::OneShot) || (CurrentDirection == -1.0f))
+			bStopped = true;
+			StopSimulating(Hit);
+			return;
+		case EInterpToBehaviourType::OneShot_Reverse:
+			if( CurrentDirection == -1.0f)
 			{
+				OnInterpToStop.Broadcast(Hit, Time);
 				bStopped = true;
 				StopSimulating(Hit);
 				return;
 			}
+			else
+			{
+				ReverseDirection(Hit, Time, true);
+			}
 			break;
 		case EInterpToBehaviourType::Loop_Reset:
 			{
