Exposed whole UAnimSequence class from engine (added ENGINE_API instead of MinimalAPI) - it's helpful, to have access to it.

diff --git a/Engine/Source/Runtime/Engine/Classes/Animation/AnimSequence.h b/Engine/Source/Runtime/Engine/Classes/Animation/AnimSequence.h
--- a/Engine/Source/Runtime/Engine/Classes/Animation/AnimSequence.h
+++ b/Engine/Source/Runtime/Engine/Classes/Animation/AnimSequence.h
@@ -67,7 +67,7 @@ enum EAdditiveBasePoseType
  * reduced to 1 key that is constant over the entire sequence.
  */
 USTRUCT()
-struct FRawAnimSequenceTrack
+struct ENGINE_API FRawAnimSequenceTrack
 {
 	GENERATED_USTRUCT_BODY()
 
@@ -102,7 +102,7 @@ struct FRawAnimSequenceTrack
 // I wonder in the future, we change all compressed to be inside as well, so they all stay together
 // When remove tracks, it should be handled together 
 USTRUCT()
-struct FAnimSequenceTrackContainer
+struct ENGINE_API FAnimSequenceTrackContainer
 {
 	GENERATED_USTRUCT_BODY()
 
@@ -138,7 +138,7 @@ struct FAnimSequenceTrackContainer
 
 // @note We have a plan to support skeletal hierarchy. When that happens, we'd like to keep skeleton indexing.
 USTRUCT()
-struct FTrackToSkeletonMap
+struct ENGINE_API FTrackToSkeletonMap
 {
 	GENERATED_USTRUCT_BODY()
 
@@ -166,7 +166,7 @@ struct FTrackToSkeletonMap
  * Keyframe position data for one track.  Pos(i) occurs at Time(i).  Pos.Num() always equals Time.Num().
  */
 USTRUCT()
-struct FTranslationTrack
+struct ENGINE_API FTranslationTrack
 {
 	GENERATED_USTRUCT_BODY()
 
@@ -181,7 +181,7 @@ struct FTranslationTrack
  * Keyframe rotation data for one track.  Rot(i) occurs at Time(i).  Rot.Num() always equals Time.Num().
  */
 USTRUCT()
-struct FRotationTrack
+struct ENGINE_API FRotationTrack
 {
 	GENERATED_USTRUCT_BODY()
 
@@ -196,7 +196,7 @@ struct FRotationTrack
  * Keyframe scale data for one track.  Scale(i) occurs at Time(i).  Rot.Num() always equals Time.Num().
  */
 USTRUCT()
-struct FScaleTrack
+struct ENGINE_API FScaleTrack
 {
 	GENERATED_USTRUCT_BODY()
 
@@ -214,7 +214,7 @@ struct FScaleTrack
  * CurveWeights: List of weights for each frame
  */
 USTRUCT()
-struct FCurveTrack
+struct ENGINE_API FCurveTrack
 {
 	GENERATED_USTRUCT_BODY()
 
@@ -225,7 +225,7 @@ struct FCurveTrack
 	TArray<float> CurveWeights;
 
 	/** Returns true if valid curve weight exists in the array*/
-	ENGINE_API bool IsValidCurveTrack();
+	bool IsValidCurveTrack();
 	
 	/** This is very simple cut to 1 key method if all is same since I see so many redundant same value in every frame 
 	 *  Eventually this can get more complicated 
@@ -234,7 +234,7 @@ struct FCurveTrack
 };
 
 USTRUCT()
-struct FCompressedTrack
+struct ENGINE_API FCompressedTrack
 {
 	GENERATED_USTRUCT_BODY()
 
@@ -266,7 +266,7 @@ struct FCompressedTrack
 };
 
 USTRUCT()
-struct FCompressedOffsetData
+struct ENGINE_API FCompressedOffsetData
 {
 	GENERATED_USTRUCT_BODY()
 
@@ -325,8 +325,8 @@ struct FCompressedOffsetData
 	}
 };
 
-UCLASS(config=Engine, hidecategories=(UObject, Length), MinimalAPI, BlueprintType)
-class UAnimSequence : public UAnimSequenceBase
+UCLASS(config=Engine, hidecategories=(UObject, Length), BlueprintType)
+class ENGINE_API UAnimSequence : public UAnimSequenceBase
 {
 	GENERATED_UCLASS_BODY()
 
@@ -527,30 +527,30 @@ class UAnimSequence : public UAnimSequenceBase
 
 public:
 	// Begin UObject interface
-	ENGINE_API virtual void Serialize(FArchive& Ar) override;
-	ENGINE_API virtual void PostLoad() override;
+	virtual void Serialize(FArchive& Ar) override;
+	virtual void PostLoad() override;
 	virtual void PreSave() override;
 #if WITH_EDITOR
-	ENGINE_API virtual void PostEditChangeProperty(FPropertyChangedEvent& PropertyChangedEvent) override;
+	virtual void PostEditChangeProperty(FPropertyChangedEvent& PropertyChangedEvent) override;
 	virtual void PostDuplicate(bool bDuplicateForPIE) override;
 #endif // WITH_EDITOR
-	ENGINE_API virtual void BeginDestroy() override;
-	ENGINE_API virtual SIZE_T GetResourceSize(EResourceSizeMode::Type Mode) override;
+	virtual void BeginDestroy() override;
+	virtual SIZE_T GetResourceSize(EResourceSizeMode::Type Mode) override;
 	virtual void GetAssetRegistryTags(TArray<FAssetRegistryTag>& OutTags) const override;
 	// End of UObject interface
 
 	// Begin UAnimationAsset interface
-	ENGINE_API virtual bool IsValidAdditive() const override;
+	virtual bool IsValidAdditive() const override;
 #if WITH_EDITOR
-	ENGINE_API virtual bool GetAllAnimationSequencesReferred(TArray<UAnimSequence*>& AnimationSequences) override;
-	ENGINE_API virtual void ReplaceReferredAnimations(const TMap<UAnimSequence*, UAnimSequence*>& ReplacementMap) override;
-	ENGINE_API virtual int32 GetNumberOfFrames() const override { return NumFrames; }
+	virtual bool GetAllAnimationSequencesReferred(TArray<UAnimSequence*>& AnimationSequences) override;
+	virtual void ReplaceReferredAnimations(const TMap<UAnimSequence*, UAnimSequence*>& ReplacementMap) override;
+	virtual int32 GetNumberOfFrames() const override { return NumFrames; }
 #endif
 	// End of UAnimationAsset interface
 
 	// Begin UAnimSequenceBase interface
-	ENGINE_API virtual void OnAssetPlayerTickedInternal(FAnimAssetTickContext &Context, const float PreviousTime, const float MoveDelta, const FAnimTickRecord &Instance, class UAnimInstance* InstanceOwner) const override;
-	ENGINE_API virtual bool HasRootMotion() const override { return bEnableRootMotion; }
+	virtual void OnAssetPlayerTickedInternal(FAnimAssetTickContext &Context, const float PreviousTime, const float MoveDelta, const FAnimTickRecord &Instance, class UAnimInstance* InstanceOwner) const override;
+	virtual bool HasRootMotion() const override { return bEnableRootMotion; }
 	// End UAnimSequenceBase interface
 
 	// Extract Root Motion transform from the animation
@@ -582,7 +582,7 @@ public:
 	* @param	RequiredBones		Array of Desired Tracks
 	* @param	ExtractionContext	Extraction Context (position, looping, root motion, etc.)
 	*/
-	ENGINE_API void GetBonePose(FCompactPose& OutPose, FBlendedCurve& OutCurve, const FAnimExtractContext& ExtractionContext) const;
+	void GetBonePose(FCompactPose& OutPose, FBlendedCurve& OutCurve, const FAnimExtractContext& ExtractionContext) const;
 
 private:
 	/** 
@@ -613,7 +613,7 @@ public:
 	* @param	OutCurve			[out] Curves to fill	
 	* @param	ExtractionContext	Extraction Context (position, looping, root motion, etc.)
 	*/
-	ENGINE_API void GetBonePose_Additive(FCompactPose& OutPose, FBlendedCurve& OutCurve, const FAnimExtractContext& ExtractionContext) const;
+	void GetBonePose_Additive(FCompactPose& OutPose, FBlendedCurve& OutCurve, const FAnimExtractContext& ExtractionContext) const;
 
 	/**
 	* Get Bone Transform of the base (reference) pose of the additive animation for the Time given, relative to Parent for all RequiredBones
@@ -622,7 +622,7 @@ public:
 	* @param	OutCurve			[out] Curves to fill	
 	* @param	ExtractionContext	Extraction Context (position, looping, root motion, etc.)
 	*/
-	ENGINE_API void GetAdditiveBasePose(FCompactPose& OutPose, FBlendedCurve& OutCurve, const FAnimExtractContext& ExtractionContext) const;
+	void GetAdditiveBasePose(FCompactPose& OutPose, FBlendedCurve& OutCurve, const FAnimExtractContext& ExtractionContext) const;
 
 	/**
 	 * Get Bone Transform of the Time given, relative to Parent for the Track Given
@@ -632,7 +632,7 @@ public:
 	 * @param	Time			Time on track to interpolate to.
 	 * @param	bUseRawData		If true, use raw animation data instead of compressed data.
 	 */
-	ENGINE_API void GetBoneTransform(FTransform& OutAtom, int32 TrackIndex, float Time, bool bUseRawData) const;
+	void GetBoneTransform(FTransform& OutAtom, int32 TrackIndex, float Time, bool bUseRawData) const;
 
 	/**
 	 * Extract Bone Transform of the Time given, from InRawAnimationData
@@ -642,7 +642,7 @@ public:
 	 * @param	TrackIndex			Index of track to interpolate.
 	 * @param	Time				Time on track to interpolate to.
 	 */
-	ENGINE_API void ExtractBoneTransform(const TArray<struct FRawAnimSequenceTrack> & InRawAnimationData, FTransform& OutAtom, int32 TrackIndex, float Time) const;
+	void ExtractBoneTransform(const TArray<struct FRawAnimSequenceTrack> & InRawAnimationData, FTransform& OutAtom, int32 TrackIndex, float Time) const;
 
 	// End Transform related functions 
 
@@ -650,7 +650,7 @@ public:
 	/**
 	 * @return		The approximate size of raw animation data.
 	 */
-	ENGINE_API int32 GetApproxRawSize() const;
+	int32 GetApproxRawSize() const;
 
 	/**
 	 * @return		The approximate size of key-reduced animation data.
@@ -660,7 +660,7 @@ public:
 	/**
 	 * @return		The approximate size of compressed animation data.
 	 */
-	ENGINE_API int32 GetApproxCompressedSize() const;
+	int32 GetApproxCompressedSize() const;
 
 	/** 
 	 *  Utility function for lossless compression of a FRawAnimSequenceTrack 
@@ -674,7 +674,7 @@ public:
 	 * go down to a single frame, the time is stripped out as well.
 	 * @return true if keys were removed.
 	 */
-	ENGINE_API bool CompressRawAnimData(float MaxPosDiff, float MaxAngleDiff);
+	bool CompressRawAnimData(float MaxPosDiff, float MaxAngleDiff);
 
 	/**
 	 * Removes trivial frames -- frames of tracks when position or orientation is constant
@@ -682,7 +682,7 @@ public:
 	 * go down to a single frame, the time is stripped out as well.
 	 * @return true if keys were removed.
 	 */
-	ENGINE_API bool CompressRawAnimData();
+	bool CompressRawAnimData();
 	// End Memory related functions
 
 	// Begin Utility functions
@@ -697,7 +697,7 @@ public:
 	}
 
 	/** Clears any data in the AnimSequence, so it can be recycled when importing a new animation with same name over it. */
-	ENGINE_API void RecycleAnimSequence();
+	void RecycleAnimSequence();
 
 	/** 
 	 * Utility function to copy all UAnimSequence properties from Source to Destination.
@@ -708,12 +708,12 @@ public:
 	/** 
 	 * Copy AnimNotifies from one UAnimSequence to another.
 	 */
-	ENGINE_API static bool CopyNotifies(UAnimSequence* SourceAnimSeq, UAnimSequence* DestAnimSeq);
+	static bool CopyNotifies(UAnimSequence* SourceAnimSeq, UAnimSequence* DestAnimSeq);
 
 	/**
 	 * Flip Rotation's W For NonRoot items, and compress it again if SkelMesh exists
 	 */
-	ENGINE_API void FlipRotationWForNonRoot(USkeletalMesh * SkelMesh);
+	void FlipRotationWForNonRoot(USkeletalMesh * SkelMesh);
 
 	/** 
 	 * Return number of tracks in this animation
@@ -725,43 +725,43 @@ public:
 	/**
 	 * After imported or any other change is made, call this to apply post process
 	 */
-	ENGINE_API void PostProcessSequence();
+	void PostProcessSequence();
 	/** 
 	 * Insert extra frame of the first frame at the end of the frame so that it improves the interpolation when it loops
 	 * This increases framecount + time, so that it requires recompression
 	 */
-	ENGINE_API bool AddLoopingInterpolation();
+	bool AddLoopingInterpolation();
 
 	/** 
 	 * Bake Transform Curves.TransformCurves to RawAnimation after making a back up of current RawAnimation
 	 */
-	ENGINE_API void BakeTrackCurvesToRawAnimation();
+	void BakeTrackCurvesToRawAnimation();
 
 	/**
 	 * Add Key to Transform Curves
 	 */
-	ENGINE_API void AddKeyToSequence(float Time, const FName& BoneName, const FTransform& AdditiveTransform);
+	void AddKeyToSequence(float Time, const FName& BoneName, const FTransform& AdditiveTransform);
 	/**
 	 * Return true if it needs to re-bake
 	 */
-	ENGINE_API bool DoesNeedRebake() const;
+	bool DoesNeedRebake() const;
 	/**
 	 * Return true if it contains transform curves
 	 */
-	ENGINE_API bool DoesContainTransformCurves() const;
+	bool DoesContainTransformCurves() const;
 
 	/**
 	 * Create Animation Sequence from Reference Pose of the Mesh
 	 */
-	ENGINE_API bool CreateAnimation(class USkeletalMesh * Mesh);
+	bool CreateAnimation(class USkeletalMesh * Mesh);
 	/**
 	 * Create Animation Sequence from the Mesh Component's current bone trasnform
 	 */
-	ENGINE_API bool CreateAnimation(class USkeletalMeshComponent * MeshComponent);
+	bool CreateAnimation(class USkeletalMeshComponent * MeshComponent);
 	/**
 	 * Create Animation Sequence from the given animation
 	 */
-	ENGINE_API bool CreateAnimation(class UAnimSequence * Sequence);
+	bool CreateAnimation(class UAnimSequence * Sequence);
 
 	/**
 	 * Crops the raw anim data either from Start to CurrentTime or CurrentTime to End depending on
@@ -771,7 +771,7 @@ public:
 	 * @param	bFromStart		whether marker is begin or end marker
 	 * @return					true if the operation was successful.
 	 */
-	ENGINE_API bool CropRawAnimData( float CurrentTime, bool bFromStart );
+	bool CropRawAnimData( float CurrentTime, bool bFromStart );
 
 		
 	/**
@@ -783,12 +783,12 @@ public:
 	 * @param	CopyFrame		A frame that we copy from (0-based)
 	 * @return					true if the operation was successful.
 	 */
-	ENGINE_API bool InsertFramesToRawAnimData( int32 StartFrame, int32 EndFrame, int32 CopyFrame);
+	bool InsertFramesToRawAnimData( int32 StartFrame, int32 EndFrame, int32 CopyFrame);
 
 	/** 
 	 * Add validation check to see if it's being ready to play or not
 	 */
-	ENGINE_API virtual bool IsValidToPlay() const override;
+	virtual bool IsValidToPlay() const override;
 #endif
 
 private:
