UE-16083 Select node not passing by reference.
- New VM instruction: EX_SwithValue
- New KismetCompiledStatement: KCST_SwitchValue
- Functions from UKismetArrayLibrary don't have UArrayProperty input param anymore
- removed FBlueprintCompiledStatement::ArrayCoersionTermMap
- New Engine.ini option: [Kismet]bUseSelectRef - currently it's disabled
- FKCHandler_SelectRef - new implementation of Select node.

diff --git a/Engine/Config/BaseEngine.ini b/Engine/Config/BaseEngine.ini
--- a/Engine/Config/BaseEngine.ini
+++ b/Engine/Config/BaseEngine.ini
@@ -1649,6 +1649,7 @@ bDeferDependencyLoads=true
 bForceDisableCookedDependencyDeferring=false
 bExecutionAfterReturn=false
 bCanSuppressAccessViolation=false
+bUseSelectRef=false
 
 [/Script/Engine.Blueprint]
 bRecompileOnLoad=true
diff --git a/Engine/Source/Developer/BlueprintCompilerCppBackend/Private/BlueprintCompilerCppBackend.cpp b/Engine/Source/Developer/BlueprintCompilerCppBackend/Private/BlueprintCompilerCppBackend.cpp
--- a/Engine/Source/Developer/BlueprintCompilerCppBackend/Private/BlueprintCompilerCppBackend.cpp
+++ b/Engine/Source/Developer/BlueprintCompilerCppBackend/Private/BlueprintCompilerCppBackend.cpp
@@ -357,12 +357,6 @@ void FBlueprintCompilerCppBackend::EmitCallDelegateStatment(FKismetFunctionConte
 			FBPTerminal* Term = Statement.RHS[NumParams];
 			ensure(Term != NULL);
 
-			// See if this is a hidden array param term, which needs to be fixed up with the final generated UArrayProperty
-			if (FBPTerminal** ArrayParmTerm = Statement.ArrayCoersionTermMap.Find(Term))
-			{
-				Term->ObjectLiteral = (*ArrayParmTerm)->AssociatedVarProperty;
-			}
-
 			if ((Statement.TargetLabel != NULL) && (Statement.UbergraphCallIndex == NumParams))
 			{
 				// The target label will only ever be set on a call function when calling into the Ubergraph or
@@ -466,12 +460,6 @@ void FBlueprintCompilerCppBackend::EmitCallStatment(FKismetFunctionContext& Func
 				FBPTerminal* Term = Statement.RHS[NumParams];
 				ensure(Term != NULL);
 
-				// See if this is a hidden array param term, which needs to be fixed up with the final generated UArrayProperty
-				if (FBPTerminal** ArrayParmTerm = Statement.ArrayCoersionTermMap.Find(Term))
-				{
-					Term->ObjectLiteral = (*ArrayParmTerm)->AssociatedVarProperty;
-				}
-
 				if ((Statement.TargetLabel != NULL) && (Statement.UbergraphCallIndex == NumParams))
 				{
 					// The target label will only ever be set on a call function when calling into the Ubergraph or
diff --git a/Engine/Source/Editor/BlueprintGraph/Classes/K2Node_Select.h b/Engine/Source/Editor/BlueprintGraph/Classes/K2Node_Select.h
--- a/Engine/Source/Editor/BlueprintGraph/Classes/K2Node_Select.h
+++ b/Engine/Source/Editor/BlueprintGraph/Classes/K2Node_Select.h
@@ -58,6 +58,7 @@ class UK2Node_Select : public UK2Node, public INodeDependingOnEnumInterface
 	virtual void GetMenuActions(FBlueprintActionDatabaseRegistrar& ActionRegistrar) const override;
 	virtual FText GetMenuCategory() const override;
 	virtual int32 GetNodeRefreshPriority() const override { return EBaseNodeRefreshPriority::Low_UsesDependentWildcard; }
+	virtual void ExpandNode(class FKismetCompilerContext& CompilerContext, UEdGraph* SourceGraph) override;
 	// End UK2Node interface
 
 	// INodeDependingOnEnumInterface
diff --git a/Engine/Source/Editor/BlueprintGraph/Private/CallFunctionHandler.cpp b/Engine/Source/Editor/BlueprintGraph/Private/CallFunctionHandler.cpp
--- a/Engine/Source/Editor/BlueprintGraph/Private/CallFunctionHandler.cpp
+++ b/Engine/Source/Editor/BlueprintGraph/Private/CallFunctionHandler.cpp
@@ -420,37 +420,6 @@ void FKCHandler_CallFunction::CreateFunctionCallStatement(FKismetFunctionContext
 					}
 				}
 
-				// Check to see if we need to fix up any terms for array property coersion
-				if( UK2Node_CallArrayFunction* ArrayNode = Cast<UK2Node_CallArrayFunction>(Node) )
-				{
-					TArray< FArrayPropertyPinCombo > ArrayPinInfo;
-					ArrayNode->GetArrayPins(ArrayPinInfo);
-
-					for(auto Iter = ArrayPinInfo.CreateConstIterator(); Iter; ++Iter)
-					{
-						UEdGraphPin* ArrayPropPin = Iter->ArrayPropPin;
-						UEdGraphPin* ArrayTargetPin = Iter->ArrayPin;
-						if( ArrayPropPin && ArrayTargetPin )
-						{
-							// Find the array property literal term, used for specifying the type of TargetArray
-							for( int32 i = 0; i < Context.Literals.Num(); i++ )
-							{
-								FBPTerminal& Term = Context.Literals[i];
-								if( Term.Source == ArrayPropPin )
-								{
-									// Now, map the array property literal term to the TargetArray term.  The AssociatedVarProperty will later be filled in as the array prop's object literal in the backend
-									UEdGraphPin* Net = FEdGraphUtilities::GetNetFromPin(ArrayTargetPin);
-									FBPTerminal** TargetTerm = Context.NetMap.Find(Net);
-									if( TargetTerm )
-									{
-										Statement.ArrayCoersionTermMap.Add(&Term, *TargetTerm);
-									}
-								}
-							}
-						}
-					}
-				}
-
 				AdditionalCompiledStatementHandling(Context, Node, Statement);
 			}
 
diff --git a/Engine/Source/Editor/BlueprintGraph/Private/K2Node_Select.cpp b/Engine/Source/Editor/BlueprintGraph/Private/K2Node_Select.cpp
--- a/Engine/Source/Editor/BlueprintGraph/Private/K2Node_Select.cpp
+++ b/Engine/Source/Editor/BlueprintGraph/Private/K2Node_Select.cpp
@@ -13,6 +13,127 @@
 //////////////////////////////////////////////////////////////////////////
 // FKCHandler_Select
 
+class FKCHandler_SelectRef : public FNodeHandlingFunctor
+{
+protected:
+	TMap<UEdGraphNode*, FBPTerminal*> DefaultTermMap;
+
+public:
+	FKCHandler_SelectRef(FKismetCompilerContext& InCompilerContext)
+		: FNodeHandlingFunctor(InCompilerContext)
+	{
+	}
+
+	virtual void RegisterNets(FKismetFunctionContext& Context, UEdGraphNode* Node) override
+	{
+		UK2Node_Select* SelectNode = Cast<UK2Node_Select>(Node);
+		UEdGraphPin* ReturnPin = SelectNode ? SelectNode->GetReturnValuePin() : nullptr;
+		if (!ReturnPin)
+		{
+			Context.MessageLog.Error(*LOCTEXT("Error_NoReturnPin", "No return pin in @@").ToString(), Node);
+			return;
+		}
+
+		// return inline term
+		if (Context.NetMap.Find(ReturnPin))
+		{
+			Context.MessageLog.Error(*LOCTEXT("Error_ReturnTermAlreadyRegistered", "ICE: Return term is already registered @@").ToString(), Node);
+			return;
+		}
+
+		{
+			FBPTerminal* Term = new (Context.InlineGeneratedValues) FBPTerminal();
+			Term->CopyFromPin(ReturnPin, Context.NetNameMap->MakeValidName(ReturnPin));
+			Context.NetMap.Add(ReturnPin, Term);
+		}
+
+		//Register Default term
+		{
+			TArray<UEdGraphPin*> OptionPins;
+			SelectNode->GetOptionPins(OptionPins);
+			if (!OptionPins.Num())
+			{
+				Context.MessageLog.Error(*LOCTEXT("Error_NoOptionPin", "No option pin in @@").ToString(), Node);
+				return;
+			}
+
+			FString DefaultTermName = Context.NetNameMap->MakeValidName(Node) + TEXT("_Default");
+			FBPTerminal* DefaultTerm = Context.CreateLocalTerminalFromPinAutoChooseScope(OptionPins[0], DefaultTermName);
+			DefaultTermMap.Add(Node, DefaultTerm);
+		}
+
+		FNodeHandlingFunctor::RegisterNets(Context, Node);
+	}
+
+	virtual void Compile(FKismetFunctionContext& Context, UEdGraphNode* Node) override
+	{
+		UK2Node_Select* SelectNode = CastChecked<UK2Node_Select>(Node);
+		FBPTerminal* DefaultTerm = nullptr;
+		FBPTerminal* ReturnTerm = nullptr;
+		FBPTerminal* IndexTerm = nullptr;
+
+		{
+			UEdGraphPin* IndexPin = SelectNode->GetIndexPin();
+			UEdGraphPin* IndexPinNet = IndexPin ? FEdGraphUtilities::GetNetFromPin(IndexPin) : nullptr;
+			FBPTerminal** IndexTermPtr = IndexPinNet ? Context.NetMap.Find(IndexPinNet) : nullptr;
+			IndexTerm = IndexTermPtr ? *IndexTermPtr : nullptr;
+
+			UEdGraphPin* ReturnPin = SelectNode->GetReturnValuePin();
+			UEdGraphPin* ReturnPinNet = ReturnPin ? FEdGraphUtilities::GetNetFromPin(ReturnPin) : nullptr;
+			FBPTerminal** ReturnTermPtr = ReturnPinNet ? Context.NetMap.Find(ReturnPinNet) : nullptr;
+			ReturnTerm = ReturnTermPtr ? *ReturnTermPtr : nullptr;
+
+			FBPTerminal** DefaultTermPtr = DefaultTermMap.Find(SelectNode);
+			DefaultTerm = DefaultTermPtr ? *DefaultTermPtr : nullptr;
+			
+			if (!ReturnTerm || !IndexTerm || !DefaultTerm)
+			{
+				Context.MessageLog.Error(*LOCTEXT("Error_InvalidSelect", "ICE: invalid select node @@").ToString(), Node);
+				return;
+			}
+		}
+
+		FBlueprintCompiledStatement* SelectStatement = new FBlueprintCompiledStatement();
+		SelectStatement->Type = EKismetCompiledStatementType::KCST_SwitchValue;
+		Context.AllGeneratedStatements.Add(SelectStatement);
+		ReturnTerm->InlineGeneratedParameter = SelectStatement;
+		SelectStatement->RHS.Add(IndexTerm);
+
+		TArray<UEdGraphPin*> OptionPins;
+		SelectNode->GetOptionPins(OptionPins);
+		for (int32 OptionIdx = 0; OptionIdx < OptionPins.Num(); OptionIdx++)
+		{
+			{
+				FBPTerminal* LiteralTerm = Context.CreateLocalTerminal(ETerminalSpecification::TS_Literal);
+				LiteralTerm->Type = IndexTerm->Type;
+				LiteralTerm->bIsLiteral = true;
+				const UEnum* NodeEnum = SelectNode->GetEnum();
+				LiteralTerm->Name = NodeEnum ? OptionPins[OptionIdx]->PinName : FString::Printf(TEXT("%d"), OptionIdx);
+
+				if (!CompilerContext.GetSchema()->DefaultValueSimpleValidation(LiteralTerm->Type, LiteralTerm->Name, LiteralTerm->Name, nullptr, FText()))
+				{
+					Context.MessageLog.Error(*FString::Printf(*LOCTEXT("Error_InvalidOptionValue", "Invalid option value '%s' in @@").ToString(), *LiteralTerm->Name), Node);
+					return;
+				}
+				SelectStatement->RHS.Add(LiteralTerm);
+			}
+			{
+				UEdGraphPin* NetPin = OptionPins[OptionIdx] ? FEdGraphUtilities::GetNetFromPin(OptionPins[OptionIdx]) : nullptr;
+				FBPTerminal** ValueTermPtr = NetPin ? Context.NetMap.Find(NetPin) : nullptr;
+				FBPTerminal* ValueTerm = ValueTermPtr ? *ValueTermPtr : nullptr;
+				if (!ensure(ValueTerm))
+				{
+					Context.MessageLog.Error(*LOCTEXT("Error_NoTermFound", "No term registered for pin @@").ToString(), NetPin);
+					return;
+				}
+				SelectStatement->RHS.Add(ValueTerm);
+			}
+		}
+
+		SelectStatement->RHS.Add(DefaultTerm);
+	}
+};
+
 class FKCHandler_Select : public FNodeHandlingFunctor
 {
 protected:
@@ -864,7 +985,10 @@ bool UK2Node_Select::IsConnectionDisallowed(const UEdGraphPin* MyPin, const UEdG
 
 FNodeHandlingFunctor* UK2Node_Select::CreateNodeHandler(FKismetCompilerContext& CompilerContext) const
 {
-	return new FKCHandler_Select(CompilerContext);
+	static const FBoolConfigValueHelper UseSelectRef(TEXT("Kismet"), TEXT("bUseSelectRef"), GEngineIni);
+	return UseSelectRef
+		? static_cast<FNodeHandlingFunctor*>(new FKCHandler_SelectRef(CompilerContext))
+		: static_cast<FNodeHandlingFunctor*>(new FKCHandler_Select(CompilerContext));
 }
 
 void UK2Node_Select::GetMenuActions(FBlueprintActionDatabaseRegistrar& ActionRegistrar) const
@@ -892,4 +1016,50 @@ FText UK2Node_Select::GetMenuCategory() const
 	return FEditorCategoryUtils::GetCommonCategory(FCommonEditorCategory::Utilities);
 }
 
+void UK2Node_Select::ExpandNode(class FKismetCompilerContext& CompilerContext, UEdGraph* SourceGraph)
+{
+	Super::ExpandNode(CompilerContext, SourceGraph);
+
+	static const FBoolConfigValueHelper UseSelectRef(TEXT("Kismet"), TEXT("bUseSelectRef"), GEngineIni);
+	if (!UseSelectRef)
+	{
+		return;
+	}
+
+	bool bSuccess = false;
+	auto Schema = CompilerContext.GetSchema();
+	for (auto Pin : Pins)
+	{
+		const bool bValidAutoRefPin = Pin && !Schema->IsMetaPin(*Pin) && (Pin->Direction == EGPD_Input) && !Pin->LinkedTo.Num();
+		if (!bValidAutoRefPin)
+		{
+			continue;
+		}
+
+		const bool bHasDefaultValue = !Pin->DefaultValue.IsEmpty() || Pin->DefaultObject || !Pin->DefaultTextValue.IsEmpty();
+
+		//default values can be reset when the pin is connected
+		const auto DefaultValue = Pin->DefaultValue;
+		const auto DefaultObject = Pin->DefaultObject;
+		const auto DefaultTextValue = Pin->DefaultTextValue;
+		const auto AutogeneratedDefaultValue = Pin->AutogeneratedDefaultValue;
+
+		auto ValuePin = UK2Node_CallFunction::InnerHandleAutoCreateRef(this, Pin, CompilerContext, SourceGraph, bHasDefaultValue);
+		if (ValuePin)
+		{
+			if (!DefaultObject && DefaultTextValue.IsEmpty() && (DefaultValue == AutogeneratedDefaultValue))
+			{
+				// Use the latest code to set default value
+				Schema->SetPinDefaultValueBasedOnType(ValuePin);
+			}
+			else
+			{
+				ValuePin->DefaultValue = DefaultValue;
+				ValuePin->DefaultObject = DefaultObject;
+				ValuePin->DefaultTextValue = DefaultTextValue;
+			}
+		}
+	}
+}
+
 #undef LOCTEXT_NAMESPACE
diff --git a/Engine/Source/Editor/KismetCompiler/Private/KismetCompilerVMBackend.cpp b/Engine/Source/Editor/KismetCompiler/Private/KismetCompilerVMBackend.cpp
--- a/Engine/Source/Editor/KismetCompiler/Private/KismetCompilerVMBackend.cpp
+++ b/Engine/Source/Editor/KismetCompiler/Private/KismetCompilerVMBackend.cpp
@@ -211,6 +211,10 @@ private:
 	bool bIsUbergraph;
 
 	FBlueprintCompiledStatement& ReturnStatement;
+
+	FKismetCompilerContext* CurrentCompilerContext;
+	FKismetFunctionContext* CurrentFunctionContext;
+
 protected:
 	/**
 	 * This class is designed to be used like so to emit a bytecode context expression:
@@ -319,6 +323,8 @@ public:
 		, UbergraphStatementLabelMap(InUbergraphStatementLabelMap)
 		, bIsUbergraph(bInIsUbergraph)
 		, ReturnStatement(InReturnStatement)
+		, CurrentCompilerContext(nullptr)
+		, CurrentFunctionContext(nullptr)
 	{
 		VectorStruct = TBaseStructure<FVector>::Get();
 		RotatorStruct = TBaseStructure<FRotator>::Get();
@@ -699,7 +705,7 @@ public:
 					// SetArray instruction will be called with empty parameter list.
 					Writer << EX_SetArray;
 					FBPTerminal* ArrayTerm = Statement.RHS[NumParams];
-					ensure(ArrayTerm && !ArrayTerm->bIsLiteral && !Statement.ArrayCoersionTermMap.Find(ArrayTerm));
+					ensure(ArrayTerm && !ArrayTerm->bIsLiteral);
 					EmitTerm(ArrayTerm, Param);
 					Writer << EX_EndArray;
 				}
@@ -808,12 +814,6 @@ public:
 				FBPTerminal* Term = Statement.RHS[NumParams];
 				check(Term != NULL);
 
-				// See if this is a hidden array param term, which needs to be fixed up with the final generated UArrayProperty
-				if( FBPTerminal** ArrayParmTerm = Statement.ArrayCoersionTermMap.Find(Term) )
-				{
-					Term->ObjectLiteral = (*ArrayParmTerm)->AssociatedVarProperty;
-				}
-
 				// Latent function handling:  Need to emit a fixup request into the FLatentInfo struct
 				static const FName NAME_LatentInfo = TEXT("LatentInfo");
  				if (bIsUbergraph && FuncParamProperty->GetName() == FunctionToCall->GetMetaData("LatentInfo"))
@@ -822,18 +822,8 @@ public:
  				}
 				else
 				{
-					if (Term->InlineGeneratedParameter)
-					{
-						ensure(!Term->InlineGeneratedParameter->bIsJumpTarget);
-						ensure(Term->InlineGeneratedParameter->Type == KCST_CallFunction);
-
-						GenerateCodeForStatement(CompilerContext, FunctionContext, *Term->InlineGeneratedParameter, SourceNode);
-					}
-					else
-					{
-						// Emit parameter term normally
-						EmitTerm(Term, FuncParamProperty);
-					}
+					// Emit parameter term normally
+					EmitTerm(Term, FuncParamProperty);
 				}
 
 
@@ -866,12 +856,6 @@ public:
 			FBPTerminal* Term = Statement.RHS[NumParams];
 			check(Term != NULL);
 
-			// See if this is a hidden array param term, which needs to be fixed up with the final generated UArrayProperty
-			if( FBPTerminal** ArrayParmTerm = Statement.ArrayCoersionTermMap.Find(Term) )
-			{
-				Term->ObjectLiteral = (*ArrayParmTerm)->AssociatedVarProperty;
-			}
-
 			// Emit parameter term normally
 			EmitTerm(Term, FuncParamProperty);
 
@@ -884,7 +868,19 @@ public:
 
 	void EmitTerm(FBPTerminal* Term, UProperty* CoerceProperty = NULL, FBPTerminal* RValueTerm = NULL)
 	{
-		if (Term->Context == NULL)
+		if (Term->InlineGeneratedParameter)
+		{
+			ensure(!Term->InlineGeneratedParameter->bIsJumpTarget);
+			auto TermSourceAsNode = Cast<UEdGraphNode>(Term->Source);
+			auto TermSourceAsPin = Cast<UEdGraphPin>(Term->Source);
+			UEdGraphNode* SourceNode = TermSourceAsNode ? TermSourceAsNode
+				: (TermSourceAsPin ? TermSourceAsPin->GetOwningNodeUnchecked() : nullptr);
+			if (ensure(CurrentCompilerContext && CurrentFunctionContext))
+			{
+				GenerateCodeForStatement(*CurrentCompilerContext, *CurrentFunctionContext, *Term->InlineGeneratedParameter, SourceNode);
+			}
+		}
+		else if (Term->Context == NULL)
 		{
 			EmitTermExpr(Term, CoerceProperty);
 		}
@@ -919,14 +915,14 @@ public:
 
 	void EmitDestinationExpression(FBPTerminal* DestinationExpression)
 	{
-		check(DestinationExpression->AssociatedVarProperty != NULL);
-
-		const bool bIsDelegate = Cast<UDelegateProperty>(DestinationExpression->AssociatedVarProperty) != NULL;
-		const bool bIsMulticastDelegate = Cast<UMulticastDelegateProperty>(DestinationExpression->AssociatedVarProperty) != NULL;
-		const bool bIsBoolean = Cast<UBoolProperty>(DestinationExpression->AssociatedVarProperty) != NULL;
-		const bool bIsObj = Cast<UObjectPropertyBase>(DestinationExpression->AssociatedVarProperty) != NULL;
-		const bool bIsAsset = Cast<UAssetObjectProperty>(DestinationExpression->AssociatedVarProperty) != NULL;
-		const bool bIsWeakObjPtr = Cast<UWeakObjectProperty>(DestinationExpression->AssociatedVarProperty) != NULL;
+		check(Schema && DestinationExpression && !DestinationExpression->Type.PinCategory.IsEmpty());
+
+		const bool bIsDelegate = Schema->PC_Delegate == DestinationExpression->Type.PinCategory;
+		const bool bIsMulticastDelegate = Schema->PC_MCDelegate == DestinationExpression->Type.PinCategory;
+		const bool bIsBoolean = Schema->PC_Boolean == DestinationExpression->Type.PinCategory;
+		const bool bIsObj = (Schema->PC_Object == DestinationExpression->Type.PinCategory) || (Schema->PC_Class == DestinationExpression->Type.PinCategory);
+		const bool bIsAsset = Schema->PC_Asset == DestinationExpression->Type.PinCategory;
+		const bool bIsWeakObjPtr = DestinationExpression->Type.bIsWeakPointer;
 		if (bIsMulticastDelegate)
 		{
 			Writer << EX_LetMulticastDelegate;
@@ -974,7 +970,7 @@ public:
 		FBPTerminal* DestinationExpression = Statement.LHS;
 		FBPTerminal* SourceExpression = Statement.RHS[0];
 
-		Writer << Ex_LetValueOnPersistentFrame;
+		Writer << EX_LetValueOnPersistentFrame;
 		check(ClassBeingBuilt && ClassBeingBuilt->UberGraphFunction);
 		Writer << DestinationExpression->AssociatedVarProperty;
 
@@ -1254,6 +1250,43 @@ public:
 		}
 	}
 
+	void EmitSwitchValue(FBlueprintCompiledStatement& Statement)
+	{
+		if ((Statement.RHS.Num() < 4) || (1 == (Statement.RHS.Num() % 2)))
+		{
+			// Error
+			ensure(false);
+		}
+
+		Writer << EX_SwithValue;
+		// number of cases (without default)
+		const int32 TermsPerCase = 2;
+		uint16 NumCases = ((Statement.RHS.Num() - 2) / TermsPerCase);
+		Writer << NumCases;
+		// end goto index
+		CodeSkipSizeType PatchUpNeededAtOffset = Writer.EmitPlaceholderSkip();
+		// index term
+		const int32 TermsBeforeCases = 1;
+		EmitTerm(Statement.RHS[0]);
+
+		UProperty* VirtualIndexProperty = Statement.RHS[0]->AssociatedVarProperty;
+		check(VirtualIndexProperty);
+
+		for (uint16 TermIndex = TermsBeforeCases; TermIndex < (NumCases * TermsPerCase); ++TermIndex)
+		{
+			EmitTerm(Statement.RHS[TermIndex], VirtualIndexProperty); // it's a literal value
+			++TermIndex;
+			CodeSkipSizeType PatchOffsetToNextCase = Writer.EmitPlaceholderSkip();
+			EmitTerm(Statement.RHS[TermIndex]);  // it's not a literal
+			Writer.CommitSkip(PatchOffsetToNextCase, Writer.ScriptBuffer.Num());
+		}
+
+		// default term
+		EmitTerm(Statement.RHS[TermsBeforeCases + NumCases*TermsPerCase]);
+
+		Writer.CommitSkip(PatchUpNeededAtOffset, Writer.ScriptBuffer.Num());
+	}
+
 	void PushReturnAddress(FBlueprintCompiledStatement& ReturnTarget)
 	{
 		Writer << EX_PushExecutionFlow;
@@ -1273,6 +1306,9 @@ public:
 
 	void GenerateCodeForStatement(FKismetCompilerContext& CompilerContext, FKismetFunctionContext& FunctionContext, FBlueprintCompiledStatement& Statement, UEdGraphNode* SourceNode)
 	{
+		TGuardValue<FKismetCompilerContext*> CompilerContextGuard(CurrentCompilerContext, &CompilerContext);
+		TGuardValue<FKismetFunctionContext*> FunctionContextGuard(CurrentFunctionContext, &FunctionContext);
+
 		// Record the start of this statement in the bytecode if it's needed as a target label
 		if (Statement.bIsJumpTarget)
 		{
@@ -1402,6 +1438,9 @@ public:
 		case KCST_Return:
 			EmitReturn(FunctionContext);
 			break;
+		case KCST_SwitchValue:
+			EmitSwitchValue(Statement);
+			break;
 		default:
 			UE_LOG(LogK2Compiler, Warning, TEXT("VM backend encountered unsupported statement type %d"), (int32)Statement.Type);
 		}
diff --git a/Engine/Source/Editor/KismetCompiler/Public/BlueprintCompiledStatement.h b/Engine/Source/Editor/KismetCompiler/Public/BlueprintCompiledStatement.h
--- a/Engine/Source/Editor/KismetCompiler/Public/BlueprintCompiledStatement.h
+++ b/Engine/Source/Editor/KismetCompiler/Public/BlueprintCompiledStatement.h
@@ -36,6 +36,7 @@ enum EKismetCompiledStatementType
 	KCST_CastInterfaceToObj = 26, // Cast<TargetClass>(TargetInterface)
 	KCST_GotoReturn = 27,	// goto ReturnLabel
 	KCST_GotoReturnIfNot = 28, // [if (!TargetObject->TargetProperty)] goto TargetLabel
+	KCST_SwitchValue = 29,
 };
 
 //@TODO: Too rigid / icky design
@@ -69,9 +70,6 @@ struct FBlueprintCompiledStatement
 	// The index of the argument to replace (only used when KCST_CallFunction has a non-NULL TargetLabel)
 	int32 UbergraphCallIndex;
 
-	// Mapping of automatic array-property FKismetTerms to the FBPTerminal used for array type specification
-	TMap<struct FBPTerminal*, struct FBPTerminal*> ArrayCoersionTermMap;
-
 	// Destination of assignment statement or result from function call
 	struct FBPTerminal* LHS;
 
diff --git a/Engine/Source/Editor/UnrealEd/Private/ScriptDisassembler.cpp b/Engine/Source/Editor/UnrealEd/Private/ScriptDisassembler.cpp
--- a/Engine/Source/Editor/UnrealEd/Private/ScriptDisassembler.cpp
+++ b/Engine/Source/Editor/UnrealEd/Private/ScriptDisassembler.cpp
@@ -291,7 +291,7 @@ void FKismetBytecodeDisassembler::ProcessCommon(int32& ScriptIndex, EExprToken O
 			DropIndent();
 			break;
 		}
-	case Ex_LetValueOnPersistentFrame:
+	case EX_LetValueOnPersistentFrame:
 		{
 			Ar.Logf(TEXT("%s $%X: LetValueOnPersistentFrame"), *Indents, (int32)Opcode);
 			AddIndent();
@@ -815,6 +815,32 @@ void FKismetBytecodeDisassembler::ProcessCommon(int32& ScriptIndex, EExprToken O
 			Ar.Logf(TEXT("%s $%X: .. debug site .."), *Indents, (int32)Opcode);
 			break;
 		}
+	case EX_SwithValue:
+		{
+			const auto NumCases = ReadWORD(ScriptIndex);
+			const auto AfterSkip = ReadSkipCount(ScriptIndex);
+
+			Ar.Logf(TEXT("%s $%X: Switch Value %d cases, end in 0x%X"), *Indents, (int32)Opcode, NumCases, AfterSkip);
+			AddIndent();
+			Ar.Logf(TEXT("%s Index:"), *Indents);
+			SerializeExpr(ScriptIndex);
+
+			for (uint16 CaseIndex = 0; CaseIndex < NumCases; ++CaseIndex)
+			{
+				Ar.Logf(TEXT("%s [%d] Case Index (label: 0x%X):"), *Indents, CaseIndex, ScriptIndex);
+				SerializeExpr(ScriptIndex);	// case index value term
+				const auto OffsetToNextCase = ReadSkipCount(ScriptIndex);
+				Ar.Logf(TEXT("%s [%d] Offset to the next case: 0x%X"), *Indents, CaseIndex, OffsetToNextCase);
+				Ar.Logf(TEXT("%s [%d] Case Result:"), *Indents, CaseIndex);
+				SerializeExpr(ScriptIndex);	// case term
+			}
+
+			Ar.Logf(TEXT("%s Default result (label: 0x%X):"), *Indents, ScriptIndex);
+			SerializeExpr(ScriptIndex);
+			Ar.Logf(TEXT("%s (label: 0x%X)"), *Indents, ScriptIndex);
+			DropIndent();
+			break;
+		}
 	default:
 		{
 			// This should never occur.
diff --git a/Engine/Source/Runtime/CoreUObject/Private/UObject/ScriptCore.cpp b/Engine/Source/Runtime/CoreUObject/Private/UObject/ScriptCore.cpp
--- a/Engine/Source/Runtime/CoreUObject/Private/UObject/ScriptCore.cpp
+++ b/Engine/Source/Runtime/CoreUObject/Private/UObject/ScriptCore.cpp
@@ -1461,7 +1461,63 @@ void UObject::execLetValueOnPersistentFrame(FFrame& Stack, RESULT_DECL)
 	checkf(false, TEXT("execLetValueOnPersistentFrame: UberGraphPersistentFrame is not supported by current build!"));
 #endif
 }
-IMPLEMENT_VM_FUNCTION(Ex_LetValueOnPersistentFrame, execLetValueOnPersistentFrame);
+IMPLEMENT_VM_FUNCTION(EX_LetValueOnPersistentFrame, execLetValueOnPersistentFrame);
+
+void UObject::execSwithValue(FFrame& Stack, RESULT_DECL)
+{
+	const int32 NumCases = Stack.ReadWord();
+	const CodeSkipSizeType OffsetToEnd = Stack.ReadCodeSkipCount();
+
+	Stack.MostRecentProperty = nullptr;
+	Stack.MostRecentPropertyAddress = nullptr;
+	Stack.Step(Stack.Object, nullptr);
+
+	UProperty* IndexProperty = Stack.MostRecentProperty;
+	checkSlow(IndexProperty);
+
+	uint8* IndexAdress = Stack.MostRecentPropertyAddress;
+	if (!ensure(IndexAdress))
+	{
+		FBlueprintExceptionInfo ExceptionInfo(EBlueprintExceptionType::NonFatalError, TEXT("Switch - Unknown index"));
+		FBlueprintCoreDelegates::ThrowScriptException(this, Stack, ExceptionInfo);
+	}
+
+	bool bProperCaseUsed = false;
+	{
+		auto LocalTempIndexMem = (uint8*)FMemory_Alloca(IndexProperty->GetSize());
+		IndexProperty->InitializeValue(LocalTempIndexMem);
+		for (int32 CaseIndex = 0; CaseIndex < NumCases; ++CaseIndex)
+		{
+			Stack.Step(Stack.Object, LocalTempIndexMem); // case index value
+			const CodeSkipSizeType OffsetToNextCase = Stack.ReadCodeSkipCount();
+
+			if (IndexAdress && IndexProperty->Identical(IndexAdress, LocalTempIndexMem))
+			{
+				Stack.Step(Stack.Object, RESULT_PARAM);
+				bProperCaseUsed = true;
+				break;
+			}
+
+			// skip to the next case
+			Stack.Code = &Stack.Node->Script[OffsetToNextCase];
+		}
+		IndexProperty->DestroyValue(LocalTempIndexMem);
+	}
+
+	if (bProperCaseUsed)
+	{
+		Stack.Code = &Stack.Node->Script[OffsetToEnd];
+	}
+	else
+	{
+		FBlueprintExceptionInfo ExceptionInfo(EBlueprintExceptionType::NonFatalError, TEXT("Switch - Out of bounds index"));
+		FBlueprintCoreDelegates::ThrowScriptException(this, Stack, ExceptionInfo);
+
+		// get default value
+		Stack.Step(Stack.Object, RESULT_PARAM);
+	}
+}
+IMPLEMENT_VM_FUNCTION(EX_SwithValue, execSwithValue);
 
 void UObject::execLet(FFrame& Stack, RESULT_DECL)
 {
diff --git a/Engine/Source/Runtime/CoreUObject/Public/UObject/Script.h b/Engine/Source/Runtime/CoreUObject/Public/UObject/Script.h
--- a/Engine/Source/Runtime/CoreUObject/Public/UObject/Script.h
+++ b/Engine/Source/Runtime/CoreUObject/Public/UObject/Script.h
@@ -235,11 +235,12 @@ enum EExprToken
 	EX_BindDelegate			= 0x61, // bind object and name to delegate
 	EX_RemoveMulticastDelegate = 0x62, // Remove a delegate from a multicast delegate's targets
 	EX_CallMulticastDelegate = 0x63, // Call multicast delegate
-	Ex_LetValueOnPersistentFrame = 0x64,
+	EX_LetValueOnPersistentFrame = 0x64,
 	EX_ArrayConst			= 0x65,
 	EX_EndArrayConst		= 0x66,
 	EX_AssetConst			= 0x67,
 	EX_CallMath				= 0x68, // static pure function from on local call space
+	EX_SwithValue			= 0x69,
 	EX_Max					= 0x100,
 };
 
diff --git a/Engine/Source/Runtime/CoreUObject/Public/UObject/ScriptSerialization.h b/Engine/Source/Runtime/CoreUObject/Public/UObject/ScriptSerialization.h
--- a/Engine/Source/Runtime/CoreUObject/Public/UObject/ScriptSerialization.h
+++ b/Engine/Source/Runtime/CoreUObject/Public/UObject/ScriptSerialization.h
@@ -141,7 +141,7 @@
 			SerializeExpr( iCode, Ar ); // Assignment expr.
 			break;
 		}
-		case Ex_LetValueOnPersistentFrame:
+		case EX_LetValueOnPersistentFrame:
 		{
 			XFER_PROP_POINTER;			// Destination property.
 			SerializeExpr(iCode, Ar);	// Assignment expr.
@@ -420,6 +420,23 @@
 			SerializeExpr( iCode, Ar ); 
 			break;
 		}
+		case EX_SwithValue:
+		{
+			XFER(uint16); // number of cases, without default one
+			const uint16 NumCases = *(uint16*)(&Script[iCode - sizeof(uint16)]);
+			XFER(CodeSkipSizeType); // Code offset, go to it, when done.
+			SerializeExpr(iCode, Ar);	//index term
+
+			for (uint16 CaseIndex = 0; CaseIndex < NumCases; ++CaseIndex)
+			{
+				SerializeExpr(iCode, Ar);	// case index value term
+				XFER(CodeSkipSizeType);		// offset to the next case
+				SerializeExpr(iCode, Ar);	// case term
+			}
+
+			SerializeExpr(iCode, Ar);	//default term
+			break;
+		}
 		default:
 		{
 			// This should never occur.
diff --git a/Engine/Source/Runtime/CoreUObject/Public/UObject/UObject.h b/Engine/Source/Runtime/CoreUObject/Public/UObject/UObject.h
--- a/Engine/Source/Runtime/CoreUObject/Public/UObject/UObject.h
+++ b/Engine/Source/Runtime/CoreUObject/Public/UObject/UObject.h
@@ -1144,6 +1144,8 @@ public:
 
 	DECLARE_FUNCTION(execCallMathFunction);
 
+	DECLARE_FUNCTION(execSwithValue);
+
 	// -- K2 support functions
 	struct Object_eventExecuteUbergraph_Parms
 	{
diff --git a/Engine/Source/Runtime/Engine/Classes/Kismet/KismetArrayLibrary.h b/Engine/Source/Runtime/Engine/Classes/Kismet/KismetArrayLibrary.h
--- a/Engine/Source/Runtime/Engine/Classes/Kismet/KismetArrayLibrary.h
+++ b/Engine/Source/Runtime/Engine/Classes/Kismet/KismetArrayLibrary.h
@@ -14,8 +14,8 @@ class ENGINE_API UKismetArrayLibrary : public UBlueprintFunctionLibrary
 	 *@param	NewItem			The item to add to the array
 	 *@return	The index of the newly added item
 	*/
-	UFUNCTION(BlueprintCallable, CustomThunk, meta=(DisplayName = "Add", CompactNodeTitle = "ADD", ArrayParm = "TargetArray|ArrayProperty", ArrayTypeDependentParams = "NewItem", AutoCreateRefTerm = "NewItem"), Category="Utilities|Array")
-	static int32 Array_Add(const TArray<int32>& TargetArray, const UArrayProperty* ArrayProperty, const int32& NewItem);
+	UFUNCTION(BlueprintCallable, CustomThunk, meta=(DisplayName = "Add", CompactNodeTitle = "ADD", ArrayParm = "TargetArray", ArrayTypeDependentParams = "NewItem", AutoCreateRefTerm = "NewItem"), Category="Utilities|Array")
+	static int32 Array_Add(const TArray<int32>& TargetArray, const int32& NewItem);
 
 	/**
 	*Add item to array (unique)
@@ -24,16 +24,16 @@ class ENGINE_API UKismetArrayLibrary : public UBlueprintFunctionLibrary
 	*@param		NewItem			The item to add to the array
 	*@return	The index of the newly added item, or INDEX_NONE if the item is already present in the array
 	*/
-	UFUNCTION(BlueprintCallable, CustomThunk, meta = (DisplayName = "Add Unique", CompactNodeTitle = "ADDUNIQUE", ArrayParm = "TargetArray|ArrayProperty", ArrayTypeDependentParams = "NewItem", AutoCreateRefTerm = "NewItem"), Category = "Utilities|Array")
-	static int32 Array_AddUnique(const TArray<int32>& TargetArray, const UArrayProperty* ArrayProperty, const int32& NewItem);
+	UFUNCTION(BlueprintCallable, CustomThunk, meta = (DisplayName = "Add Unique", CompactNodeTitle = "ADDUNIQUE", ArrayParm = "TargetArray", ArrayTypeDependentParams = "NewItem", AutoCreateRefTerm = "NewItem"), Category = "Utilities|Array")
+	static int32 Array_AddUnique(const TArray<int32>& TargetArray, const int32& NewItem);
 
 	/** 
 	 * Shuffle (randomize) the elements of an array
 	 *
 	 *@param	TargetArray		The array to shuffle
 	*/
-	UFUNCTION(BlueprintCallable, CustomThunk, meta=(DisplayName = "Shuffle", CompactNodeTitle = "SHUFFLE", ArrayParm = "TargetArray|ArrayProperty"), Category="Utilities|Array")
-	static void Array_Shuffle(const TArray<int32>& TargetArray, const UArrayProperty* ArrayProperty);
+	UFUNCTION(BlueprintCallable, CustomThunk, meta=(DisplayName = "Shuffle", CompactNodeTitle = "SHUFFLE", ArrayParm = "TargetArray"), Category="Utilities|Array")
+	static void Array_Shuffle(const TArray<int32>& TargetArray);
 
 	/** 
 	 *Append an array to another array
@@ -41,8 +41,8 @@ class ENGINE_API UKismetArrayLibrary : public UBlueprintFunctionLibrary
 	 *@param	TargetArray		The array to add the source array to
 	 *@param	SourceArray		The array to add to the target array
 	*/
-	UFUNCTION(BlueprintCallable, CustomThunk, meta=(DisplayName = "Append Array", CompactNodeTitle = "APPEND", ArrayParm = "TargetArray|TargetArrayProperty,SourceArray|SourceArrayProperty", ArrayTypeDependentParams = "SourceArray"), Category="Utilities|Array")
-	static void Array_Append(const TArray<int32>& TargetArray, const UArrayProperty* TargetArrayProperty, const TArray<int32>& SourceArray, const UArrayProperty* SourceArrayProperty);
+	UFUNCTION(BlueprintCallable, CustomThunk, meta=(DisplayName = "Append Array", CompactNodeTitle = "APPEND", ArrayParm = "TargetArray,SourceArray", ArrayTypeDependentParams = "SourceArray"), Category="Utilities|Array")
+	static void Array_Append(const TArray<int32>& TargetArray, const TArray<int32>& SourceArray);
 
 	/* 
 	 *Insert item at the given index into the array.
@@ -51,8 +51,8 @@ class ENGINE_API UKismetArrayLibrary : public UBlueprintFunctionLibrary
 	 *@param	NewItem			The item to insert into the array
 	 *@param	Index			The index at which to insert the item into the array
 	*/
-	UFUNCTION(BlueprintCallable, CustomThunk, meta=(DisplayName = "Insert", CompactNodeTitle = "INSERT", ArrayParm = "TargetArray|ArrayProperty", ArrayTypeDependentParams = "NewItem", AutoCreateRefTerm = "NewItem"), Category="Utilities|Array")
-	static void Array_Insert(const TArray<int32>& TargetArray, const UArrayProperty* ArrayProperty, const int32& NewItem, int32 Index);
+	UFUNCTION(BlueprintCallable, CustomThunk, meta=(DisplayName = "Insert", CompactNodeTitle = "INSERT", ArrayParm = "TargetArray", ArrayTypeDependentParams = "NewItem", AutoCreateRefTerm = "NewItem"), Category="Utilities|Array")
+	static void Array_Insert(const TArray<int32>& TargetArray, const int32& NewItem, int32 Index);
 
 
 	/* 
@@ -61,8 +61,8 @@ class ENGINE_API UKismetArrayLibrary : public UBlueprintFunctionLibrary
 	 *@param	TargetArray		The array to remove from
 	 *@param	IndexToRemove	The index into the array to remove from
 	*/
-	UFUNCTION(BlueprintCallable, CustomThunk, meta=(DisplayName = "Remove Index", CompactNodeTitle = "REMOVE INDEX", ArrayParm = "TargetArray|ArrayProperty"), Category="Utilities|Array")
-	static void Array_Remove(const TArray<int32>& TargetArray, const UArrayProperty* ArrayProperty, int32 IndexToRemove);
+	UFUNCTION(BlueprintCallable, CustomThunk, meta=(DisplayName = "Remove Index", CompactNodeTitle = "REMOVE INDEX", ArrayParm = "TargetArray"), Category="Utilities|Array")
+	static void Array_Remove(const TArray<int32>& TargetArray, int32 IndexToRemove);
 
 	/* 
 	 *Remove all instances of item from array.
@@ -71,16 +71,16 @@ class ENGINE_API UKismetArrayLibrary : public UBlueprintFunctionLibrary
 	 *@param	Item			The item to remove from the array
 	 *@return	True if one or more items were removed
 	*/
-	UFUNCTION(BlueprintCallable, CustomThunk, meta=(DisplayName = "Remove Item", CompactNodeTitle = "REMOVE", ArrayParm = "TargetArray|ArrayProperty", ArrayTypeDependentParams = "Item", AutoCreateRefTerm = "Item"), Category="Utilities|Array")
-	static bool Array_RemoveItem(const TArray<int32>& TargetArray, const UArrayProperty* ArrayProperty, const int32 &Item);
+	UFUNCTION(BlueprintCallable, CustomThunk, meta=(DisplayName = "Remove Item", CompactNodeTitle = "REMOVE", ArrayParm = "TargetArray", ArrayTypeDependentParams = "Item", AutoCreateRefTerm = "Item"), Category="Utilities|Array")
+	static bool Array_RemoveItem(const TArray<int32>& TargetArray, const int32 &Item);
 
 	/* 
 	 *Clear an array, removes all content
 	 *
 	 *@param	TargetArray		The array to clear
 	*/
-	UFUNCTION(BlueprintCallable, CustomThunk, meta=(DisplayName = "Clear", CompactNodeTitle = "CLEAR", ArrayParm = "TargetArray|ArrayProperty"), Category="Utilities|Array")
-	static void Array_Clear(const TArray<int32>& TargetArray, const UArrayProperty* ArrayProperty);
+	UFUNCTION(BlueprintCallable, CustomThunk, meta=(DisplayName = "Clear", CompactNodeTitle = "CLEAR", ArrayParm = "TargetArray"), Category="Utilities|Array")
+	static void Array_Clear(const TArray<int32>& TargetArray);
 
 	/* 
 	 *Resize Array to specified size. 
@@ -88,8 +88,8 @@ class ENGINE_API UKismetArrayLibrary : public UBlueprintFunctionLibrary
 	 *@param	TargetArray		The array to resize
 	 *@param	Size			The new size of the array
 	*/
-	UFUNCTION(BlueprintCallable, CustomThunk, meta=(DisplayName = "Resize", CompactNodeTitle = "RESIZE", ArrayParm = "TargetArray|ArrayProperty"), Category="Utilities|Array")
-	static void Array_Resize(const TArray<int32>& TargetArray, const UArrayProperty* ArrayProperty, int32 Size);
+	UFUNCTION(BlueprintCallable, CustomThunk, meta=(DisplayName = "Resize", CompactNodeTitle = "RESIZE", ArrayParm = "TargetArray"), Category="Utilities|Array")
+	static void Array_Resize(const TArray<int32>& TargetArray, int32 Size);
 
 	/* 
 	 *Get the number of items in an array
@@ -97,8 +97,8 @@ class ENGINE_API UKismetArrayLibrary : public UBlueprintFunctionLibrary
 	 *@param	TargetArray		The array to get the length of
 	 *@return	The length of the array
 	*/
-	UFUNCTION(BlueprintPure, CustomThunk, meta=(DisplayName = "Length", CompactNodeTitle = "LENGTH", ArrayParm = "TargetArray|ArrayProperty", Keywords = "num size count"), Category="Utilities|Array")
-	static int32 Array_Length(const TArray<int32>& TargetArray, const UArrayProperty* ArrayProperty);
+	UFUNCTION(BlueprintPure, CustomThunk, meta=(DisplayName = "Length", CompactNodeTitle = "LENGTH", ArrayParm = "TargetArray", Keywords = "num size count"), Category="Utilities|Array")
+	static int32 Array_Length(const TArray<int32>& TargetArray);
 
 
 	/* 
@@ -107,8 +107,8 @@ class ENGINE_API UKismetArrayLibrary : public UBlueprintFunctionLibrary
 	 *@param	TargetArray		The array to perform the operation on
 	 *@return	The last valid index of the array
 	*/
-	UFUNCTION(BlueprintPure, CustomThunk, meta=(DisplayName = "Last Index", CompactNodeTitle = "LAST INDEX", ArrayParm = "TargetArray|ArrayProperty"), Category="Utilities|Array")
-	static int32 Array_LastIndex(const TArray<int32>& TargetArray, const UArrayProperty* ArrayProperty);
+	UFUNCTION(BlueprintPure, CustomThunk, meta=(DisplayName = "Last Index", CompactNodeTitle = "LAST INDEX", ArrayParm = "TargetArray"), Category="Utilities|Array")
+	static int32 Array_LastIndex(const TArray<int32>& TargetArray);
 
 	/* 
 	 *Given an array and an index, returns the item found at that index
@@ -117,8 +117,8 @@ class ENGINE_API UKismetArrayLibrary : public UBlueprintFunctionLibrary
 	 *@param	Index			The index in the array to get an item from
 	 *@return	The item stored at the index
 	*/
-	UFUNCTION(BlueprintPure, CustomThunk, meta=(DisplayName = "Get", CompactNodeTitle = "GET", ArrayParm = "TargetArray|ArrayProperty", ArrayTypeDependentParams = "Item"), Category="Utilities|Array")
-	static void Array_Get(const TArray<int32>& TargetArray, const UArrayProperty* ArrayProperty, int32 Index, int32& Item);
+	UFUNCTION(BlueprintPure, CustomThunk, meta=(DisplayName = "Get", CompactNodeTitle = "GET", ArrayParm = "TargetArray", ArrayTypeDependentParams = "Item"), Category="Utilities|Array")
+	static void Array_Get(const TArray<int32>& TargetArray, int32 Index, int32& Item);
 
 	/* 
 	 *Given an array and an index, assigns the item to that array element
@@ -128,8 +128,8 @@ class ENGINE_API UKismetArrayLibrary : public UBlueprintFunctionLibrary
 	 *@param	Item			The item to assign to the index of the array
 	 *@param	bSizeToFit		If true, the array will expand if Index is greater than the current size of the array
 	*/
-	UFUNCTION(BlueprintCallable, CustomThunk, meta=(DisplayName = "Set Array Elem", ArrayParm = "TargetArray|ArrayProperty", ArrayTypeDependentParams = "Item", AutoCreateRefTerm = "Item"), Category="Utilities|Array")
-	static void Array_Set(const TArray<int32>& TargetArray, const UArrayProperty* ArrayProperty, int32 Index, const int32& Item, bool bSizeToFit);
+	UFUNCTION(BlueprintCallable, CustomThunk, meta=(DisplayName = "Set Array Elem", ArrayParm = "TargetArray", ArrayTypeDependentParams = "Item", AutoCreateRefTerm = "Item"), Category="Utilities|Array")
+	static void Array_Set(const TArray<int32>& TargetArray, int32 Index, const int32& Item, bool bSizeToFit);
 
 	/*  
 	 *Finds the index of the first instance of the item within the array
@@ -138,8 +138,8 @@ class ENGINE_API UKismetArrayLibrary : public UBlueprintFunctionLibrary
 	 *@param	ItemToFind		The item to look for
 	 *@return	The index the item was found at, or -1 if not found
 	*/
-	UFUNCTION(BlueprintPure, CustomThunk, meta=(DisplayName = "Find Item", CompactNodeTitle = "FIND", ArrayParm = "TargetArray|ArrayProperty", ArrayTypeDependentParams = "ItemToFind", AutoCreateRefTerm = "ItemToFind"), Category="Utilities|Array")
-	static int32 Array_Find(const TArray<int32>& TargetArray, const UArrayProperty* ArrayProperty, const int32& ItemToFind);
+	UFUNCTION(BlueprintPure, CustomThunk, meta=(DisplayName = "Find Item", CompactNodeTitle = "FIND", ArrayParm = "TargetArray", ArrayTypeDependentParams = "ItemToFind", AutoCreateRefTerm = "ItemToFind"), Category="Utilities|Array")
+	static int32 Array_Find(const TArray<int32>& TargetArray, const int32& ItemToFind);
 
 	/*  
 	 *Returns true if the array contains the given item
@@ -148,8 +148,8 @@ class ENGINE_API UKismetArrayLibrary : public UBlueprintFunctionLibrary
 	 *@param	ItemToFind		The item to look for
 	 *@return	True if the item was found within the array
 	*/
-	UFUNCTION(BlueprintPure, CustomThunk, meta=(DisplayName = "Contains Item", CompactNodeTitle = "CONTAINS", ArrayParm = "TargetArray|ArrayProperty", ArrayTypeDependentParams = "ItemToFind", AutoCreateRefTerm = "ItemToFind"), Category="Utilities|Array")
-	static bool Array_Contains(const TArray<int32>& TargetArray, const UArrayProperty* ArrayProperty, const int32& ItemToFind);
+	UFUNCTION(BlueprintPure, CustomThunk, meta=(DisplayName = "Contains Item", CompactNodeTitle = "CONTAINS", ArrayParm = "TargetArray", ArrayTypeDependentParams = "ItemToFind", AutoCreateRefTerm = "ItemToFind"), Category="Utilities|Array")
+	static bool Array_Contains(const TArray<int32>& TargetArray, const int32& ItemToFind);
 
 	/*  
 	 *Filter an array based on a Class derived from Actor.  
@@ -197,10 +197,10 @@ public:
 
 	DECLARE_FUNCTION(execArray_Add)
 	{
+		Stack.MostRecentProperty = nullptr;
  		Stack.StepCompiledIn<UArrayProperty>(NULL);
 		void* ArrayAddr = Stack.MostRecentPropertyAddress;
- 		
-		P_GET_OBJECT(UArrayProperty, ArrayProperty);
+		UArrayProperty* ArrayProperty = CastChecked<UArrayProperty>(Stack.MostRecentProperty);
  
  		// Since NewItem isn't really an int, step the stack manually
  		const UProperty* InnerProp = ArrayProperty->Inner;
@@ -221,10 +221,10 @@ public:
 
 	DECLARE_FUNCTION(execArray_AddUnique)
 	{
+		Stack.MostRecentProperty = nullptr;
 		Stack.StepCompiledIn<UArrayProperty>(NULL);
 		void* ArrayAddr = Stack.MostRecentPropertyAddress;
-
-		P_GET_OBJECT(UArrayProperty, ArrayProperty);
+		UArrayProperty* ArrayProperty = CastChecked<UArrayProperty>(Stack.MostRecentProperty);
 
 		// Since NewItem isn't really an int, step the stack manually
 		const UProperty* InnerProp = ArrayProperty->Inner;
@@ -245,10 +245,11 @@ public:
 
 	DECLARE_FUNCTION(execArray_Shuffle)
 	{
+		Stack.MostRecentProperty = nullptr;
 		Stack.StepCompiledIn<UArrayProperty>(NULL);
 		void* ArrayAddr = Stack.MostRecentPropertyAddress;
+		UArrayProperty* ArrayProperty = CastChecked<UArrayProperty>(Stack.MostRecentProperty);
 
-		P_GET_OBJECT(UArrayProperty, ArrayProperty);
 		P_FINISH;
 
 		GenericArray_Shuffle(ArrayAddr, ArrayProperty);
@@ -257,16 +258,16 @@ public:
 	DECLARE_FUNCTION(execArray_Append)
 	{
 		// Retrieve the target array
+		Stack.MostRecentProperty = nullptr;
 		Stack.StepCompiledIn<UArrayProperty>(NULL);
 		void* TargetArrayAddr = Stack.MostRecentPropertyAddress;
-
-		P_GET_OBJECT(UArrayProperty, TargetArrayProperty);
+		UArrayProperty* TargetArrayProperty = CastChecked<UArrayProperty>(Stack.MostRecentProperty);
 
 		// Retrieve the source array
+		Stack.MostRecentProperty = nullptr;
 		Stack.StepCompiledIn<UArrayProperty>(NULL);
 		void* SourceArrayAddr = Stack.MostRecentPropertyAddress;
-
-		P_GET_OBJECT(UArrayProperty, SourceArrayProperty);
+		UArrayProperty* SourceArrayProperty = CastChecked<UArrayProperty>(Stack.MostRecentProperty);
 
 		P_FINISH;
 
@@ -275,10 +276,10 @@ public:
 
 	DECLARE_FUNCTION(execArray_Insert)
 	{
+		Stack.MostRecentProperty = nullptr;
 		Stack.StepCompiledIn<UArrayProperty>(NULL);
 		void* ArrayAddr = Stack.MostRecentPropertyAddress;
-
-		P_GET_OBJECT(UArrayProperty, ArrayProperty);
+		UArrayProperty* ArrayProperty = CastChecked<UArrayProperty>(Stack.MostRecentProperty);
 
 		// Since NewItem isn't really an int, step the stack manually
 		const UProperty* InnerProp = ArrayProperty->Inner;
@@ -300,10 +301,11 @@ public:
 
 	DECLARE_FUNCTION(execArray_Remove)
 	{
+		Stack.MostRecentProperty = nullptr;
 		Stack.StepCompiledIn<UArrayProperty>(NULL);
 		void* ArrayAddr = Stack.MostRecentPropertyAddress;
+		UArrayProperty* ArrayProperty = CastChecked<UArrayProperty>(Stack.MostRecentProperty);
 
-		P_GET_OBJECT(UArrayProperty, ArrayProperty);
 		P_GET_PROPERTY(UIntProperty, Index);
 		P_FINISH;
 
@@ -312,10 +314,10 @@ public:
 
 	DECLARE_FUNCTION(execArray_RemoveItem)
 	{
+		Stack.MostRecentProperty = nullptr;
 		Stack.StepCompiledIn<UArrayProperty>(NULL);
 		void* ArrayAddr = Stack.MostRecentPropertyAddress;
-
-		P_GET_OBJECT(UArrayProperty, ArrayProperty);
+		UArrayProperty* ArrayProperty = CastChecked<UArrayProperty>(Stack.MostRecentProperty);
 
 		// Since Item isn't really an int, step the stack manually
 		const UProperty* InnerProp = ArrayProperty->Inner;
@@ -340,10 +342,11 @@ public:
 	
 	DECLARE_FUNCTION(execArray_Clear)
 	{
+		Stack.MostRecentProperty = nullptr;
 		Stack.StepCompiledIn<UArrayProperty>(NULL);
 		void* ArrayAddr = Stack.MostRecentPropertyAddress;
+		UArrayProperty* ArrayProperty = CastChecked<UArrayProperty>(Stack.MostRecentProperty);
 
-		P_GET_OBJECT(UArrayProperty, ArrayProperty);
 		P_FINISH;
 
 		GenericArray_Clear(ArrayAddr, ArrayProperty);
@@ -351,10 +354,11 @@ public:
 
 	DECLARE_FUNCTION(execArray_Resize)
 	{
+		Stack.MostRecentProperty = nullptr;
 		Stack.StepCompiledIn<UArrayProperty>(NULL);
 		void* ArrayAddr = Stack.MostRecentPropertyAddress;
+		UArrayProperty* ArrayProperty = CastChecked<UArrayProperty>(Stack.MostRecentProperty);
 
-		P_GET_OBJECT(UArrayProperty, ArrayProperty);
 		P_GET_PROPERTY(UIntProperty, Size);
 		P_FINISH;
 
@@ -363,10 +367,11 @@ public:
 
 	DECLARE_FUNCTION(execArray_Length)
 	{
+		Stack.MostRecentProperty = nullptr;
 		Stack.StepCompiledIn<UArrayProperty>(NULL);
 		void* ArrayAddr = Stack.MostRecentPropertyAddress;
+		UArrayProperty* ArrayProperty = CastChecked<UArrayProperty>(Stack.MostRecentProperty);
 
-		P_GET_OBJECT(UArrayProperty, ArrayProperty);
 		P_FINISH;
 
 		*(int32*)RESULT_PARAM = GenericArray_Length(ArrayAddr, ArrayProperty);
@@ -374,10 +379,11 @@ public:
 
 	DECLARE_FUNCTION(execArray_LastIndex)
 	{
+		Stack.MostRecentProperty = nullptr;
 		Stack.StepCompiledIn<UArrayProperty>(NULL);
 		void* ArrayAddr = Stack.MostRecentPropertyAddress;
+		UArrayProperty* ArrayProperty = CastChecked<UArrayProperty>(Stack.MostRecentProperty);
 
-		P_GET_OBJECT(UArrayProperty, ArrayProperty);
 		P_FINISH;
 
 		*(int32*)RESULT_PARAM = GenericArray_LastIndex(ArrayAddr, ArrayProperty);
@@ -385,10 +391,11 @@ public:
 
 	DECLARE_FUNCTION(execArray_Get)
 	{
+		Stack.MostRecentProperty = nullptr;
 		Stack.StepCompiledIn<UArrayProperty>(NULL);
 		void* ArrayAddr = Stack.MostRecentPropertyAddress;
+		UArrayProperty* ArrayProperty = CastChecked<UArrayProperty>(Stack.MostRecentProperty);
 
-		P_GET_OBJECT(UArrayProperty, ArrayProperty);
 		P_GET_PROPERTY(UIntProperty, Index);
 
 		// Since Item isn't really an int, step the stack manually
@@ -410,10 +417,11 @@ public:
 
 	DECLARE_FUNCTION(execArray_Set)
 	{
+		Stack.MostRecentProperty = nullptr;
 		Stack.StepCompiledIn<UArrayProperty>(NULL);
 		void* ArrayAddr = Stack.MostRecentPropertyAddress;
+		UArrayProperty* ArrayProperty = CastChecked<UArrayProperty>(Stack.MostRecentProperty);
 
-		P_GET_OBJECT(UArrayProperty, ArrayProperty);
 		P_GET_PROPERTY(UIntProperty, Index);
 
 		// Since NewItem isn't really an int, step the stack manually
@@ -437,10 +445,10 @@ public:
 
 	DECLARE_FUNCTION(execArray_Find)
 	{
+		Stack.MostRecentProperty = nullptr;
 		Stack.StepCompiledIn<UArrayProperty>(NULL);
 		void* ArrayAddr = Stack.MostRecentPropertyAddress;
-
-		P_GET_OBJECT(UArrayProperty, ArrayProperty);
+		UArrayProperty* ArrayProperty = CastChecked<UArrayProperty>(Stack.MostRecentProperty);
 
 		// Since ItemToFind isn't really an int, step the stack manually
 		const UProperty* InnerProp = ArrayProperty->Inner;
@@ -466,10 +474,10 @@ public:
 	
 	DECLARE_FUNCTION(execArray_Contains)
 	{
+		Stack.MostRecentProperty = nullptr;
 		Stack.StepCompiledIn<UArrayProperty>(NULL);
 		void* ArrayAddr = Stack.MostRecentPropertyAddress;
-
-		P_GET_OBJECT(UArrayProperty, ArrayProperty);
+		UArrayProperty* ArrayProperty = CastChecked<UArrayProperty>(Stack.MostRecentProperty);
 
 		// Since ItemToFind isn't really an int, step the stack manually
 		const UProperty* InnerProp = ArrayProperty->Inner;
diff --git a/Engine/Source/Runtime/Engine/Private/KismetArrayLibrary.cpp b/Engine/Source/Runtime/Engine/Private/KismetArrayLibrary.cpp
--- a/Engine/Source/Runtime/Engine/Private/KismetArrayLibrary.cpp
+++ b/Engine/Source/Runtime/Engine/Private/KismetArrayLibrary.cpp
@@ -294,77 +294,77 @@ void UKismetArrayLibrary::GenericArray_HandleBool(const UProperty* Property, voi
 //////////////////////////////////////////////////////////////////////////
 // Stubs for the UFunctions declared as kismet callable.  These are never actually called...the CustomThunk code calls the appropriate native function with a void* reference to the array
 
-int32 UKismetArrayLibrary::Array_Add(const TArray<int32>& TargetArray, const UArrayProperty* ArrayProp, const int32& NewItem)
+int32 UKismetArrayLibrary::Array_Add(const TArray<int32>& TargetArray, const int32& NewItem)
 {
 	// We should never hit these!  They're stubs to avoid NoExport on the class.  Call the Generic* equivalent instead
 	check(0);
 	return 0;
 }
 
-void UKismetArrayLibrary::Array_Shuffle(const TArray<int32>& TargetArray, const UArrayProperty* ArrayProp)
+void UKismetArrayLibrary::Array_Shuffle(const TArray<int32>& TargetArray)
 {
 	// We should never hit these!  They're stubs to avoid NoExport on the class.  Call the Generic* equivalent instead
 	check(0);
 }
 
-void UKismetArrayLibrary::Array_Insert(const TArray<int32>& TargetArray, const UArrayProperty* ArrayProp, const int32& NewItem, int32 Index)
+void UKismetArrayLibrary::Array_Insert(const TArray<int32>& TargetArray, const int32& NewItem, int32 Index)
 {
 	// We should never hit these!  They're stubs to avoid NoExport on the class.  Call the Generic* equivalent instead
 	check(0);
 }
 
-void UKismetArrayLibrary::Array_Remove(const TArray<int32>& TargetArray, const UArrayProperty* ArrayProp, int32 IndexToRemove)
+void UKismetArrayLibrary::Array_Remove(const TArray<int32>& TargetArray, int32 IndexToRemove)
 {
 	// We should never hit these!  They're stubs to avoid NoExport on the class.  Call the Generic* equivalent instead
 	check(0);
 }
 
-bool UKismetArrayLibrary::Array_RemoveItem(const TArray<int32>& TargetArray, const UArrayProperty* ArrayProp, const int32& IndexToRemove)
+bool UKismetArrayLibrary::Array_RemoveItem(const TArray<int32>& TargetArray, const int32& IndexToRemove)
 {
 	// We should never hit these!  They're stubs to avoid NoExport on the class.  Call the Generic* equivalent instead
 	check(0);
 	return true;
 }
 
-void UKismetArrayLibrary::Array_Clear(const TArray<int32>& TargetArray, const UArrayProperty* ArrayProp)
+void UKismetArrayLibrary::Array_Clear(const TArray<int32>& TargetArray)
 {
 	// We should never hit these!  They're stubs to avoid NoExport on the class.  Call the Generic* equivalent instead
 	check(0);
 }
 
-static void Array_Resize(const TArray<int32>& TargetArray, const UArrayProperty* ArrayProperty, int32 Size)
+static void Array_Resize(const TArray<int32>& TargetArray, int32 Size)
 {
 	// We should never hit these!  They're stubs to avoid NoExport on the class.  Call the Generic* equivalent instead
 	check(0);
 }
 
-int32 UKismetArrayLibrary::Array_Length(const TArray<int32>& TargetArray, const UArrayProperty* ArrayProp)
+int32 UKismetArrayLibrary::Array_Length(const TArray<int32>& TargetArray)
 {
 	// We should never hit these!  They're stubs to avoid NoExport on the class.  Call the Generic* equivalent instead
 	check(0);
 	return 0;
 }
 
-int32 UKismetArrayLibrary::Array_LastIndex(const TArray<int32>& TargetArray, const UArrayProperty* ArrayProp)
+int32 UKismetArrayLibrary::Array_LastIndex(const TArray<int32>& TargetArray)
 {
 	// We should never hit these!  They're stubs to avoid NoExport on the class.  Call the Generic* equivalent instead
 	check(0);
 	return 0;
 }
 
-void UKismetArrayLibrary::Array_Get(const TArray<int32>& TargetArray, const UArrayProperty* ArrayProp, int32 Index, int32& Item)
+void UKismetArrayLibrary::Array_Get(const TArray<int32>& TargetArray, int32 Index, int32& Item)
 {
 	// We should never hit these!  They're stubs to avoid NoExport on the class.  Call the Generic* equivalent instead
 	check(0);
 }
 
-void UKismetArrayLibrary::Array_Set(const TArray<int32>& TargetArray, const UArrayProperty* ArrayProp, int32 Index, const int32& NewItem, bool bSizeToFit)
+void UKismetArrayLibrary::Array_Set(const TArray<int32>& TargetArray, int32 Index, const int32& NewItem, bool bSizeToFit)
 {
 	// We should never hit these!  They're stubs to avoid NoExport on the class.  Call the Generic* equivalent instead
 	check(0);
 }
 
-int32 UKismetArrayLibrary::Array_Find(const TArray<int32>& TargetArray, const UArrayProperty* ArrayProperty, const int32& ItemToFind)
+int32 UKismetArrayLibrary::Array_Find(const TArray<int32>& TargetArray, const int32& ItemToFind)
 {
 	// We should never hit these!  They're stubs to avoid NoExport on the class.  Call the Generic* equivalent instead
 	check(0);
diff --git a/Engine/Source/Runtime/Engine/Public/GeneratedCodeHelpers.h b/Engine/Source/Runtime/Engine/Public/GeneratedCodeHelpers.h
--- a/Engine/Source/Runtime/Engine/Public/GeneratedCodeHelpers.h
+++ b/Engine/Source/Runtime/Engine/Public/GeneratedCodeHelpers.h
@@ -20,93 +20,116 @@ struct FCustomThunkTemplates
 	//Replacements for CustomThunk functions from UKismetArrayLibrary
 
 	template<typename T>
-	static int32 Array_Add(TArray<T>& TargetArray, const UArrayProperty* ArrayProperty, const T& NewItem)
+	static int32 Array_Add(TArray<T>& TargetArray, const T& NewItem)
 	{
-		return UKismetArrayLibrary::GenericArray_Add(&TargetArray, ArrayProperty, &NewItem);
+		return TargetArray.Add(NewItem);
 	}
 
 	template<typename T>
-	static int32 Array_AddUnique(TArray<T>& TargetArray, const UArrayProperty* ArrayProperty, const T& NewItem)
+	static int32 Array_AddUnique(TArray<T>& TargetArray, const T& NewItem)
 	{
-		return UKismetArrayLibrary::GenericArray_AddUnique(&TargetArray, ArrayProperty, &NewItem);
+		return TargetArray.AddUnique(NewItem);
 	}
 
 	template<typename T>
 	static void Array_Shuffle(TArray<T>& TargetArray, const UArrayProperty* ArrayProperty)
 	{
-		UKismetArrayLibrary::GenericArray_Shuffle(&TargetArray, ArrayProperty);
+		int32 LastIndex = TargetArray.Num() - 1;
+		for (int32 i = 0; i < LastIndex; ++i)
+		{
+			int32 Index = FMath::RandRange(0, LastIndex);
+			if (i != Index)
+			{
+				TargetArray.Swap(i, Index);
+			}
+		}
 	}
 
 	template<typename T>
-	static void Array_Append(TArray<T>& TargetArray, const UArrayProperty* TargetArrayProperty, const TArray<T>& SourceArray, const UArrayProperty* SourceArrayProperty)
+	static void Array_Append(TArray<T>& TargetArray, const TArray<T>& SourceArray)
 	{
-		UKismetArrayLibrary::GenericArray_Append(&TargetArray, TargetArrayProperty, &SourceArray, SourceArrayProperty);
+		TargetArray.Append(SourceArray);
 	}
 
 	template<typename T>
-	static void Array_Insert(TArray<T>& TargetArray, const UArrayProperty* ArrayProperty, const T& NewItem, int32 Index)
+	static void Array_Insert(TArray<T>& TargetArray, const T& NewItem, int32 Index)
 	{
-		UKismetArrayLibrary::GenericArray_Insert(&TargetArray, ArrayProperty, &NewItem, Index);
+		TargetArray.Insert(NewItem, Index);
 	}
 
 	template<typename T>
-	static void Array_Remove(TArray<T>& TargetArray, const UArrayProperty* ArrayProperty, int32 IndexToRemove)
+	static void Array_Remove(TArray<T>& TargetArray, int32 IndexToRemove)
 	{
-		UKismetArrayLibrary::GenericArray_Remove(&TargetArray, ArrayProperty, IndexToRemove);
+		if (ensure(TargetArray.IsValidIndex(IndexToRemove)))
+		{
+			TargetArray.RemoveAt(IndexToRemove);
+		}
 	}
 
 	template<typename T>
-	static int32 Array_Find(const TArray<T>& TargetArray, const UArrayProperty* ArrayProperty, const T& ItemToFind)
+	static int32 Array_Find(const TArray<T>& TargetArray, const T& ItemToFind)
 	{
 		return TargetArray.Find(ItemToFind);
 	}
 
 	template<typename T>
-	static bool Array_RemoveItem(TArray<T>& TargetArray, const UArrayProperty* ArrayProperty, const T& Item)
+	static bool Array_RemoveItem(TArray<T>& TargetArray, const T& Item)
 	{
 		return TargetArray.Remove(Item);
 	}
 
 	template<typename T>
-	static void Array_Clear(TArray<T>& TargetArray, const UArrayProperty* ArrayProperty)
+	static void Array_Clear(TArray<T>& TargetArray)
 	{
-		UKismetArrayLibrary::GenericArray_Clear(&TargetArray, ArrayProperty);
+		TargetArray.Empty();
 	}
 
 	template<typename T>
-	static void Array_Resize(TArray<T>& TargetArray, const UArrayProperty* ArrayProperty, int32 Size)
+	static void Array_Resize(TArray<T>& TargetArray, int32 Size)
 	{
-		UKismetArrayLibrary::GenericArray_Resize(&TargetArray, ArrayProperty, Size);
+		TargetArray.SetNum(Size);
 	}
 
 	template<typename T>
-	static int32 Array_Length(const TArray<T>& TargetArray, const UArrayProperty* ArrayProperty)
+	static int32 Array_Length(const TArray<T>& TargetArray)
 	{
-		return UKismetArrayLibrary::GenericArray_Length(&TargetArray, ArrayProperty);
+		return TargetArray.Num();
 	}
 
 	template<typename T>
-	static int32 Array_LastIndex(const TArray<T>& TargetArray, const UArrayProperty* ArrayProperty)
+	static int32 Array_LastIndex(const TArray<T>& TargetArray)
 	{
-		return UKismetArrayLibrary::GenericArray_Length(&TargetArray, ArrayProperty);
+		return TargetArray.Num() - 1;
 	}
 
 	template<typename T>
-	static void Array_Get(TArray<T>& TargetArray, const UArrayProperty* ArrayProperty, int32 Index, T& Item)
+	static void Array_Get(TArray<T>& TargetArray, int32 Index, T& Item)
 	{
-		UKismetArrayLibrary::GenericArray_Get(&TargetArray, ArrayProperty, Index, &Item);
+		Item = TargetArray[Index];
 	}
 
 	template<typename T>
-	static void Array_Set(TArray<T>& TargetArray, const UArrayProperty* ArrayProperty, int32 Index, const T& Item, bool bSizeToFit)
+	static void Array_Set(TArray<T>& TargetArray, int32 Index, const T& Item, bool bSizeToFit)
 	{
-		UKismetArrayLibrary::GenericArray_Set(&TargetArray, ArrayProperty, Index, &Item, bSizeToFit);
+		if (ensure(Index > 0))
+		{
+			const bool bIsValidIndex = ;
+			if (!TargetArray.IsValidIndex(Index) && bSizeToFit)
+			{
+				TargetArray.SetNum(Index + 1);
+			}
+
+			if (TargetArray.IsValidIndex(Index))
+			{
+				TargetArray[Index] = Item;
+			}
+		}
 	}
 
 	template<typename T>
-	static bool Array_Contains(const TArray<T>& TargetArray, const UArrayProperty* ArrayProperty, const T& ItemToFind)
+	static bool Array_Contains(const TArray<T>& TargetArray, const T& ItemToFind)
 	{
-		return Array_Find(TargetArray, ArrayProperty, ItemToFind) > 0;
+		return TargetArray.Contains(ItemToFind);
 	}
 
 	template<typename T>
