Fixed cloth wind for skeletal meshes applying strong wind even when no sources were present

diff --git a/Engine/Source/Runtime/Engine/Private/SkeletalMeshComponentPhysics.cpp b/Engine/Source/Runtime/Engine/Private/SkeletalMeshComponentPhysics.cpp
--- a/Engine/Source/Runtime/Engine/Private/SkeletalMeshComponentPhysics.cpp
+++ b/Engine/Source/Runtime/Engine/Private/SkeletalMeshComponentPhysics.cpp
@@ -2291,7 +2291,7 @@ void USkeletalMeshComponent::ApplyWindForCloth(FClothingActor& ClothingActor)
 
 			physx::PxVec3 WindVelocity(WindDirection.X, WindDirection.Y, WindDirection.Z);
 
-			WindVelocity *= WindUnitAmout;
+			WindVelocity *= WindUnitAmout * WindSpeed;
 			float WindAdaption = rand()%20 * 0.1f; // make range from 0 to 2
 
 			NxParameterized::Interface* ActorDesc = ApexClothingActor->getActorDesc();
