PR #1254: exposed FMath::FixedTurn to Blueprints (Contributed by [REDACTED_PERSON])

diff --git a/Engine/Source/Runtime/Engine/Classes/Kismet/KismetMathLibrary.h b/Engine/Source/Runtime/Engine/Classes/Kismet/KismetMathLibrary.h
--- a/Engine/Source/Runtime/Engine/Classes/Kismet/KismetMathLibrary.h
+++ b/Engine/Source/Runtime/Engine/Classes/Kismet/KismetMathLibrary.h
@@ -535,6 +535,17 @@ class ENGINE_API UKismetMathLibrary : public UBlueprintFunctionLibrary
 	UFUNCTION(BlueprintPure, Category = "Math|Float")
 	static float MakePulsatingValue(float InCurrentTime, float InPulsesPerSecond = 1.0f, float InPhase = 0.0f);
 
+	/** 
+	 * Returns a new rotation component value
+	 *
+	 * @param InCurrent is the current rotation value
+	 * @param InDesired is the desired rotation value
+	 * @param  is the rotation amount to apply
+	 *
+	 * @return a new rotation component value clamped in the range (-360,360)
+	 */
+	UFUNCTION(BlueprintPure, Category="Math|Float")
+	static float FixedTurn(float InCurrent, float InDesired, float InDeltaRate);
 
 	//
 	// Vector functions.
diff --git a/Engine/Source/Runtime/Engine/Private/KismetMathLibrary.cpp b/Engine/Source/Runtime/Engine/Private/KismetMathLibrary.cpp
--- a/Engine/Source/Runtime/Engine/Private/KismetMathLibrary.cpp
+++ b/Engine/Source/Runtime/Engine/Private/KismetMathLibrary.cpp
@@ -584,6 +584,11 @@ float UKismetMathLibrary::MakePulsatingValue(float InCurrentTime, float InPulses
 	return FMath::MakePulsatingValue((double)InCurrentTime, InPulsesPerSecond, InPhase);
 }
 
+float UKismetMathLibrary::FixedTurn(float InCurrent, float InDesired, float InDeltaRate)
+{
+	return FMath::FixedTurn(InCurrent, InDesired, InDeltaRate);
+}
+
 float UKismetMathLibrary::RandomFloat()
 {
 	return FMath::FRand();
