UE-11209 - Alpha Chanel Visualization Changes in Texture Editor with sRGB Setting

Using Gamma of 1.0 when rendering alpha only.

diff --git a/Engine/Source/Runtime/Engine/Private/BatchedElements.cpp b/Engine/Source/Runtime/Engine/Private/BatchedElements.cpp
--- a/Engine/Source/Runtime/Engine/Private/BatchedElements.cpp
+++ b/Engine/Source/Runtime/Engine/Private/BatchedElements.cpp
@@ -463,7 +463,7 @@ void FBatchedElements::PrepareShaders(
 	if( BatchedElementParameters != NULL )
 	{
 		// Use the vertex/pixel shader that we were given
-		BatchedElementParameters->BindShaders(RHICmdList, FeatureLevel, Transform, Gamma, ColorWeights, Texture);
+		BatchedElementParameters->BindShaders(RHICmdList, FeatureLevel, Transform, GammaToUse, ColorWeights, Texture);
 	}
 	else
 	{
