Random BodyInstance.cpp fix. VS2015 really didnt like derefencing the array and using it in the condition statement.

Resulted in an unexpected eof issue.

diff --git a/Engine/Source/Runtime/Engine/Private/PhysicsEngine/BodyInstance.cpp b/Engine/Source/Runtime/Engine/Private/PhysicsEngine/BodyInstance.cpp
--- a/Engine/Source/Runtime/Engine/Private/PhysicsEngine/BodyInstance.cpp
+++ b/Engine/Source/Runtime/Engine/Private/PhysicsEngine/BodyInstance.cpp
@@ -3783,7 +3783,8 @@ bool FBodyInstance::LineTrace(struct FHitResult& OutHit, const FVector& Start, c
 					{
 						for (int HitIndex = 0; HitIndex < NumHits; HitIndex++)
 						{
-							if (PHits[HitIndex].distance < BestHitDistance)
+							PxRaycastHit& Hit = PHits[HitIndex];
+							if (Hit.distance < BestHitDistance)
 							{
 								BestHitDistance = PHits[HitIndex].distance;
 								BestHit = PHits[HitIndex];
