Added Map Check warning after loading a level in the editor that uses a material that needs one of its Usage Flags set. UE-15628

diff --git a/Engine/Source/Runtime/CoreUObject/Private/UObject/UObjectBaseUtility.cpp b/Engine/Source/Runtime/CoreUObject/Private/UObject/UObjectBaseUtility.cpp
--- a/Engine/Source/Runtime/CoreUObject/Private/UObject/UObjectBaseUtility.cpp
+++ b/Engine/Source/Runtime/CoreUObject/Private/UObject/UObjectBaseUtility.cpp
@@ -127,7 +127,7 @@ UPackage* UObjectBaseUtility::GetOutermost() const
 /** 
  * Finds the outermost package and marks it dirty
  */
-void UObjectBaseUtility::MarkPackageDirty() const
+bool UObjectBaseUtility::MarkPackageDirty() const
 {
 	// since transient objects will never be saved into a package, there is no need to mark a package dirty
 	// if we're transient
@@ -157,9 +157,17 @@ void UObjectBaseUtility::MarkPackageDirty() const
 
 				// Always call PackageMarkedDirtyEvent, even when the package is already dirty
 				Package->PackageMarkedDirtyEvent.Broadcast(Package, bIsDirty);
+
+				return true;
+			}
+			else
+			{
+				// notify the caller that the request to mark the package as dirty was suppressed
+				return false;
 			}
 		}
 	}
+	return true;
 }
 
 /**
diff --git a/Engine/Source/Runtime/CoreUObject/Public/UObject/UObjectBaseUtility.h b/Engine/Source/Runtime/CoreUObject/Public/UObject/UObjectBaseUtility.h
--- a/Engine/Source/Runtime/CoreUObject/Public/UObject/UObjectBaseUtility.h
+++ b/Engine/Source/Runtime/CoreUObject/Public/UObject/UObjectBaseUtility.h
@@ -222,9 +222,12 @@ public:
 	UPackage* GetOutermost() const;
 
 	/** 
-	 * Finds the outermost package and marks it dirty
+	 * Finds the outermost package and marks it dirty. 
+	 * The editor suppresses this behavior during load as it is against policy to dirty packages simply by loading them.
+	 *
+	 * @return false if the request to mark the package dirty was suppressed by the editor and true otherwise.
 	 */
-	void MarkPackageDirty() const;
+	bool MarkPackageDirty() const;
 
 	/**
 	* Determines whether this object is a template object
diff --git a/Engine/Source/Runtime/Engine/Classes/Materials/Material.h b/Engine/Source/Runtime/Engine/Classes/Materials/Material.h
--- a/Engine/Source/Runtime/Engine/Classes/Materials/Material.h
+++ b/Engine/Source/Runtime/Engine/Classes/Materials/Material.h
@@ -882,6 +882,14 @@ private:
 	/** to share code for PostLoad() and PostEditChangeProperty(), and UMaterialInstance::InitResources(), needs to be refactored */
 	void PropagateDataToMaterialProxy();
 
+#if WITH_EDITOR
+	/** Marks the material's package dirty in order to make a material usage change set during map load persistent. 
+	  * This couldn't be done during map load as loading cannot mark packages dirty. Invoked manually by the user 
+	  * from the Map Check message log.
+	  */
+	void FixupMaterialUsageAfterLoad();
+#endif
+
 public:
 
 	/** @return the name of the given usage flag. */
diff --git a/Engine/Source/Runtime/Engine/Private/Materials/Material.cpp b/Engine/Source/Runtime/Engine/Private/Materials/Material.cpp
--- a/Engine/Source/Runtime/Engine/Private/Materials/Material.cpp
+++ b/Engine/Source/Runtime/Engine/Private/Materials/Material.cpp
@@ -32,6 +32,8 @@
 #include "ShaderCompiler.h"
 #include "Materials/MaterialParameterCollection.h"
 #if WITH_EDITOR
+#include "MessageLog.h"
+#include "UObjectToken.h"
 #include "UnrealEd.h"
 #include "SlateBasics.h"  // For AddNotification
 #include "SNotificationList.h"
@@ -1024,7 +1026,20 @@ bool UMaterial::SetMaterialUsage(bool &bNeedsRecompile, EMaterialUsage Usage, co
 
 			// Mark the package dirty so that hopefully it will be saved with the new usage flag.
 			// This is important because the only way an artist can fix an infinite 'compile on load' scenario is by saving with the new usage flag
-			MarkPackageDirty();
+			if (!MarkPackageDirty())
+			{
+#if WITH_EDITOR
+				// The package could not be marked as dirty as we're loading content in the editor. Add a Map Check error to notify the user.
+				FFormatNamedArguments Arguments;
+				Arguments.Add(TEXT("Material"), FText::FromString(*GetPathName()));
+				Arguments.Add(TEXT("Usage"), FText::FromString(*GetUsageName(Usage)));
+				FMessageLog("MapCheck").Warning()
+					->AddToken(FUObjectToken::Create(this))
+					->AddToken(FTextToken::Create(FText::Format(LOCTEXT("MapCheck_SetMaterialUsage", "Material {Material} was missing the usage flag {Usage}. If the material asset is not re-saved, it may not render correctly when run outside the editor."), Arguments)))
+					->AddToken(FActionToken::Create(LOCTEXT("MapCheck_FixMaterialUsage", "Fix"), LOCTEXT("MapCheck_FixMaterialUsage_Desc", "Click to set the usage flag correctly and mark the asset file as needing to be saved."), FOnActionTokenExecuted::CreateUObject(this, &UMaterial::FixupMaterialUsageAfterLoad), true));
+				FMessageLog("MapCheck").Open(EMessageSeverity::Warning);
+#endif
+			}
 		}
 		else
 		{
@@ -1061,6 +1076,14 @@ bool UMaterial::SetMaterialUsage(bool &bNeedsRecompile, EMaterialUsage Usage, co
 	return true;
 }
 
+#if WITH_EDITOR
+void UMaterial::FixupMaterialUsageAfterLoad()
+{
+	// All we need to do here is mark the package dirty as the usage itself was set on load.
+	MarkPackageDirty();
+}
+#endif
+
 void UMaterial::GetAllVectorParameterNames(TArray<FName> &OutParameterNames, TArray<FGuid> &OutParameterIds) const
 {
 	OutParameterNames.Empty();
