Tidy up of some logic.

diff --git a/Engine/Source/Runtime/Engine/Private/UnrealEngine.cpp b/Engine/Source/Runtime/Engine/Private/UnrealEngine.cpp
--- a/Engine/Source/Runtime/Engine/Private/UnrealEngine.cpp
+++ b/Engine/Source/Runtime/Engine/Private/UnrealEngine.cpp
@@ -10586,34 +10586,21 @@ bool UEngine::CommitMapChange( FWorldContext &Context )
 		// if there are pending streaming changes replicated from the server, apply them immediately
 		if (Context.PendingLevelStreamingStatusUpdates.Num() > 0)
 		{
-			for (int32 i = 0; i < Context.PendingLevelStreamingStatusUpdates.Num(); i++)
+			for (const FLevelStreamingStatus& PendingUpdate : Context.PendingLevelStreamingStatusUpdates)
 			{
-				ULevelStreaming* LevelStreamingObject = NULL;
-				for (int32 j = 0; j < Context.World()->StreamingLevels.Num(); j++)
-				{
-					if (Context.World()->StreamingLevels[j] != NULL && Context.World()->StreamingLevels[j]->GetWorldAssetPackageFName() == Context.PendingLevelStreamingStatusUpdates[i].PackageName)
-					{
-						LevelStreamingObject = Context.World()->StreamingLevels[j];
-						if (LevelStreamingObject != NULL)
-						{
-							LevelStreamingObject->bShouldBeLoaded	= Context.PendingLevelStreamingStatusUpdates[i].bShouldBeLoaded;
-							LevelStreamingObject->bShouldBeVisible	= Context.PendingLevelStreamingStatusUpdates[i].bShouldBeVisible;
-							LevelStreamingObject->LevelLODIndex		= Context.PendingLevelStreamingStatusUpdates[i].LODIndex;
-						}
-						else
-						{
-							check(LevelStreamingObject);
-							UE_LOG(LogStreaming, Log, TEXT("Unable to handle streaming object %s"),*LevelStreamingObject->GetName());
-						}
+				ULevelStreaming** Found = Context.World()->StreamingLevels.FindByPredicate([&](ULevelStreaming* Level){
+					return Level && Level->GetWorldAssetPackageFName() == PendingUpdate.PackageName;
+				});
 
-						// break out of object iterator if we found a match
-						break;
-					}
+				if (Found)
+				{
+					(*Found)->bShouldBeLoaded  = PendingUpdate.bShouldBeLoaded;
+					(*Found)->bShouldBeVisible = PendingUpdate.bShouldBeVisible;
+					(*Found)->LevelLODIndex    = PendingUpdate.LODIndex;
 				}
-
-				if (LevelStreamingObject == NULL)
+				else
 				{
-					UE_LOG(LogStreaming, Log, TEXT("Unable to find streaming object %s"), *Context.PendingLevelStreamingStatusUpdates[i].PackageName.ToString());
+					UE_LOG(LogStreaming, Log, TEXT("Unable to find streaming object %s"), *PendingUpdate.PackageName.ToString());
 				}
 			}
 
