Cvar to turn cloth on and off (p.ClothPhysics)

diff --git a/Engine/Source/Runtime/Engine/Private/PhysicsEngine/PhysScene.cpp b/Engine/Source/Runtime/Engine/Private/PhysicsEngine/PhysScene.cpp
--- a/Engine/Source/Runtime/Engine/Private/PhysicsEngine/PhysScene.cpp
+++ b/Engine/Source/Runtime/Engine/Private/PhysicsEngine/PhysScene.cpp
@@ -1037,26 +1037,32 @@ void FPhysScene::StartFrame()
 	SyncDeltaSeconds = DeltaSeconds;
 }
 
+TAutoConsoleVariable<int32> CVarEnableClothPhysics(TEXT("p.ClothPhysics"), 1, TEXT("If 1, physics cloth will be used for simulation."));
+
 void FPhysScene::StartCloth()
 {
 	FGraphEventArray FinishPrerequisites;
-	TickPhysScene(PST_Cloth, PhysicsSubsceneCompletion[PST_Cloth]);
+	if(CVarEnableClothPhysics.GetValueOnGameThread())
 	{
-		if (PhysicsSubsceneCompletion[PST_Cloth].GetReference())
+		TickPhysScene(PST_Cloth, PhysicsSubsceneCompletion[PST_Cloth]);
 		{
-			DECLARE_CYCLE_STAT(TEXT("FDelegateGraphTask.ProcessPhysScene_Cloth"),
+			if (PhysicsSubsceneCompletion[PST_Cloth].GetReference())
+			{
+				DECLARE_CYCLE_STAT(TEXT("FDelegateGraphTask.ProcessPhysScene_Cloth"),
 				STAT_FDelegateGraphTask_ProcessPhysScene_Cloth,
-				STATGROUP_TaskGraphTasks);
+					STATGROUP_TaskGraphTasks);
 
-			new (FinishPrerequisites)FGraphEventRef(
-				FDelegateGraphTask::CreateAndDispatchWhenReady(
+				new (FinishPrerequisites)FGraphEventRef(
+					FDelegateGraphTask::CreateAndDispatchWhenReady(
 					FDelegateGraphTask::FDelegate::CreateRaw(this, &FPhysScene::SceneCompletionTask, PST_Cloth),
 					GET_STATID(STAT_FDelegateGraphTask_ProcessPhysScene_Cloth), PhysicsSubsceneCompletion[PST_Cloth],
 					ENamedThreads::GameThread, ENamedThreads::GameThread
-				)
-			);
+					)
+					);
+			}
 		}
 	}
+	
 
 	//If the async scene is lagged we start it here to make sure any cloth in the async scene is using the results of the previous simulation.
 	if (FrameLagAsync() && bAsyncSceneEnabled)
diff --git a/Engine/Source/Runtime/Engine/Private/SkeletalMeshComponentPhysics.cpp b/Engine/Source/Runtime/Engine/Private/SkeletalMeshComponentPhysics.cpp
--- a/Engine/Source/Runtime/Engine/Private/SkeletalMeshComponentPhysics.cpp
+++ b/Engine/Source/Runtime/Engine/Private/SkeletalMeshComponentPhysics.cpp
@@ -3488,8 +3488,15 @@ void USkeletalMeshComponent::UpdateClothMorphTarget()
 	}
 }
 
+extern TAutoConsoleVariable<int32> CVarEnableClothPhysics;
+
 void USkeletalMeshComponent::UpdateClothState(float DeltaTime)
 {
+	if (CVarEnableClothPhysics.GetValueOnGameThread() == 0)
+	{
+		return;
+	}
+
 	// if turned on bClothMorphTarget option
 	if (bClothMorphTarget)
 	{
@@ -3697,6 +3704,11 @@ public:
 
 void USkeletalMeshComponent::PerformTickClothing(float DeltaTime)
 {
+	if (CVarEnableClothPhysics.GetValueOnGameThread() == 0)
+	{
+		return;
+	}
+
 #if WITH_APEX_CLOTHING
 	// animated but bone transforms were not updated because it was not rendered
 	if(PoseTickedThisFrame() && !bRecentlyRendered)
@@ -3731,6 +3743,11 @@ void USkeletalMeshComponent::GetUpdateClothSimulationData(TArray<FClothSimulData
 {
 #if WITH_APEX_CLOTHING
 
+	if (CVarEnableClothPhysics.GetValueOnGameThread() == 0)
+	{
+		return;
+	}
+
 	int32 NumClothingActors = ClothingActors.Num();
 
 	if(NumClothingActors == 0 || bDisableClothSimulation)
