Fix for race condition from APEX. Sometimes setup events can fire from APEX before world ticks start. Over multiple world transitions this can cause inconsistencies in clothing resources leading to crashes.

diff --git a/Engine/Source/Runtime/Engine/Private/World.cpp b/Engine/Source/Runtime/Engine/Private/World.cpp
--- a/Engine/Source/Runtime/Engine/Private/World.cpp
+++ b/Engine/Source/Runtime/Engine/Private/World.cpp
@@ -489,6 +489,11 @@ void UWorld::FinishDestroy()
 		{
 			delete PhysicsScene;
 			PhysicsScene = NULL;
+
+			if(GPhysCommandHandler)
+			{
+				GPhysCommandHandler->Flush();
+			}
 		}
 
 		if (Scene)
