Fix post eval event to happen when bone transform is finalized, so that users can get the latest transform if they want to.

diff --git a/Engine/Source/Runtime/Engine/Classes/Components/SkeletalMeshComponent.h b/Engine/Source/Runtime/Engine/Classes/Components/SkeletalMeshComponent.h
--- a/Engine/Source/Runtime/Engine/Classes/Components/SkeletalMeshComponent.h
+++ b/Engine/Source/Runtime/Engine/Classes/Components/SkeletalMeshComponent.h
@@ -916,7 +916,7 @@ public:
 
 	virtual bool IsPlayingRootMotion() override;
 	virtual bool IsPlayingRootMotionFromEverything() override;
-
+	virtual void FinalizeBoneTransform() override;
 	// End USkinnedMeshComponent interface
 	/** 
 	 *	Iterate over each joint in the physics for this mesh, setting its AngularPositionTarget based on the animation information.
diff --git a/Engine/Source/Runtime/Engine/Classes/Components/SkinnedMeshComponent.h b/Engine/Source/Runtime/Engine/Classes/Components/SkinnedMeshComponent.h
--- a/Engine/Source/Runtime/Engine/Classes/Components/SkinnedMeshComponent.h
+++ b/Engine/Source/Runtime/Engine/Classes/Components/SkinnedMeshComponent.h
@@ -549,6 +549,12 @@ public:
 	 */
 	virtual bool UpdateLODStatus();
 
+	/**
+	 * Finalize bone transform of this current tick
+	 * After this function, any query to bone transform should be latest of the data
+	 */
+	virtual void FinalizeBoneTransform();
+
 	/** Initialize the LOD entries for the component */
 	void InitLODInfos();
 
@@ -574,13 +580,13 @@ public:
 	/** Get the number of space bases */
 	int32 GetNumSpaceBases() const { return GetSpaceBases().Num(); }
 
-	/** Flip the editable space base buffer */
-	void FlipEditableSpaceBases();
-
 	void SetSpaceBaseDoubleBuffering(bool bInDoubleBufferedBlendSpaces);
 
 protected:
 
+	/** Flip the editable space base buffer */
+	void FlipEditableSpaceBases();
+
 	/** Track whether we still need to flip to recently modified buffer */
 	bool bNeedToFlipSpaceBaseBuffers;
 
diff --git a/Engine/Source/Runtime/Engine/Private/Components/DestructibleComponent.cpp b/Engine/Source/Runtime/Engine/Private/Components/DestructibleComponent.cpp
--- a/Engine/Source/Runtime/Engine/Private/Components/DestructibleComponent.cpp
+++ b/Engine/Source/Runtime/Engine/Private/Components/DestructibleComponent.cpp
@@ -1000,7 +1000,7 @@ void UDestructibleComponent::SetChunksWorldTM(const TArray<FUpdateChunksInfo>& U
 	MarkRenderDynamicDataDirty();
 
 	//Update bone visibilty and flip the editable space base buffer
-	FlipEditableSpaceBases();
+	FinalizeBoneTransform();
 }
 
 void UDestructibleComponent::SetChunkWorldRT( int32 ChunkIndex, const FQuat& WorldRotation, const FVector& WorldTranslation )
diff --git a/Engine/Source/Runtime/Engine/Private/Components/PoseableMeshComponent.cpp b/Engine/Source/Runtime/Engine/Private/Components/PoseableMeshComponent.cpp
--- a/Engine/Source/Runtime/Engine/Private/Components/PoseableMeshComponent.cpp
+++ b/Engine/Source/Runtime/Engine/Private/Components/PoseableMeshComponent.cpp
@@ -35,7 +35,7 @@ bool UPoseableMeshComponent::AllocateTransformData()
 		}
 
 		FillSpaceBases();
-		FlipEditableSpaceBases();
+		FinalizeBoneTransform();
 
 		return true;
 	}
@@ -63,7 +63,7 @@ void UPoseableMeshComponent::RefreshBoneTransforms(FActorComponentTickFunction*
 
 	// We need the mesh space bone transforms now for renderer to get delta from ref pose:
 	FillSpaceBases();
-	FlipEditableSpaceBases();
+	FinalizeBoneTransform();
 
 	MarkRenderDynamicDataDirty();
 }
diff --git a/Engine/Source/Runtime/Engine/Private/Components/SkeletalMeshComponent.cpp b/Engine/Source/Runtime/Engine/Private/Components/SkeletalMeshComponent.cpp
--- a/Engine/Source/Runtime/Engine/Private/Components/SkeletalMeshComponent.cpp
+++ b/Engine/Source/Runtime/Engine/Private/Components/SkeletalMeshComponent.cpp
@@ -1179,7 +1179,7 @@ void USkeletalMeshComponent::RefreshBoneTransforms(FActorComponentTickFunction*
 	if (TickFunction == NULL)
 	{
 		//Since we aren't doing this through the tick system, assume we want the buffer flipped now
-		FlipEditableSpaceBases();
+		FinalizeBoneTransform();
 	}
 }
 
@@ -1189,6 +1189,13 @@ void USkeletalMeshComponent::PostAnimEvaluation(FAnimationEvaluationContext& Eva
 	AnimEvaluationContext.Clear();
 
 	SCOPE_CYCLE_COUNTER(STAT_PostAnimEvaluation);
+	
+	if(AnimScriptInstance)
+	{
+		// curve update happens first
+		AnimScriptInstance->UpdateCurves(EvaluatedCurve);
+	}
+	
 	if (EvaluationContext.bDuplicateToCacheBones)
 	{
 		CachedSpaceBases = GetEditableSpaceBases();
@@ -1237,7 +1244,7 @@ void USkeletalMeshComponent::PostAnimEvaluation(FAnimationEvaluationContext& Eva
 		SCOPE_CYCLE_COUNTER(STAT_UpdateLocalToWorldAndOverlaps);
 
 		// Updated last good bone positions
-		FlipEditableSpaceBases();
+		FinalizeBoneTransform();
 
 		// New bone positions need to be sent to render thread
 		UpdateComponentToWorld();
@@ -1246,12 +1253,6 @@ void USkeletalMeshComponent::PostAnimEvaluation(FAnimationEvaluationContext& Eva
 		UpdateOverlaps();
 	}
 
-	if(AnimScriptInstance)
-	{
-		// curve update happens first
-		AnimScriptInstance->UpdateCurves(EvaluatedCurve);
-		AnimScriptInstance->PostEvaluateAnimation();
-	}
 
 	MarkRenderDynamicDataDirty();
 }
@@ -2347,3 +2348,13 @@ bool USkeletalMeshComponent::DoCustomNavigableGeometryExport(FNavigableGeometryE
 	// skip fallback export of body setup data
 	return false;
 }
+
+void USkeletalMeshComponent::FinalizeBoneTransform() 
+{
+	Super::FinalizeBoneTransform();
+
+	if (AnimScriptInstance)
+	{
+		AnimScriptInstance->PostEvaluateAnimation();
+	}
+}
\ No newline at end of file
diff --git a/Engine/Source/Runtime/Engine/Private/Components/SkinnedMeshComponent.cpp b/Engine/Source/Runtime/Engine/Private/Components/SkinnedMeshComponent.cpp
--- a/Engine/Source/Runtime/Engine/Private/Components/SkinnedMeshComponent.cpp
+++ b/Engine/Source/Runtime/Engine/Private/Components/SkinnedMeshComponent.cpp
@@ -1890,6 +1890,11 @@ TArray<FActiveVertexAnim> USkinnedMeshComponent::UpdateActiveVertexAnims(const U
 	return OutVertexAnims;
 }
 
+void USkinnedMeshComponent::FinalizeBoneTransform()
+{
+	FlipEditableSpaceBases();
+}
+
 void USkinnedMeshComponent::FlipEditableSpaceBases()
 {
 	if (bNeedToFlipSpaceBaseBuffers)
diff --git a/Engine/Source/Runtime/Engine/Private/PhysicsEngine/PhysAnim.cpp b/Engine/Source/Runtime/Engine/Private/PhysicsEngine/PhysAnim.cpp
--- a/Engine/Source/Runtime/Engine/Private/PhysicsEngine/PhysAnim.cpp
+++ b/Engine/Source/Runtime/Engine/Private/PhysicsEngine/PhysAnim.cpp
@@ -429,7 +429,7 @@ void USkeletalMeshComponent::CompleteParallelBlendPhysics()
 		// New bone positions need to be sent to render thread
 		MarkRenderDynamicDataDirty();
 
-	FlipEditableSpaceBases();
+	FinalizeBoneTransform();
 
 	ParallelAnimationEvaluationTask.SafeRelease();
 	ParallelBlendPhysicsCompletionTask.SafeRelease();
diff --git a/Engine/Source/Runtime/Engine/Private/SkeletalMesh.cpp b/Engine/Source/Runtime/Engine/Private/SkeletalMesh.cpp
--- a/Engine/Source/Runtime/Engine/Private/SkeletalMesh.cpp
+++ b/Engine/Source/Runtime/Engine/Private/SkeletalMesh.cpp
@@ -3803,7 +3803,7 @@ void ASkeletalMeshActor::PreviewSetAnimPosition(FName SlotName, int32 ChannelInd
 		SkeletalMeshComponent->RefreshSlaveComponents();
 		SkeletalMeshComponent->UpdateComponentToWorld();
 	}
-	SkeletalMeshComponent->FlipEditableSpaceBases();
+	SkeletalMeshComponent->FinalizeBoneTransform();
 }
 
 
diff --git a/Engine/Source/Runtime/Engine/Private/SkeletalMeshComponentPhysics.cpp b/Engine/Source/Runtime/Engine/Private/SkeletalMeshComponentPhysics.cpp
--- a/Engine/Source/Runtime/Engine/Private/SkeletalMeshComponentPhysics.cpp
+++ b/Engine/Source/Runtime/Engine/Private/SkeletalMeshComponentPhysics.cpp
@@ -3122,7 +3122,7 @@ void USkeletalMeshComponent::PreClothTick(float DeltaTime)
 	// if physics is disabled on dedicated server, no reason to be here. 
 	if (!bEnablePhysicsOnDedicatedServer && IsRunningDedicatedServer())
 	{
-		FlipEditableSpaceBases();
+		FinalizeBoneTransform();
 		return;
 	}
 
