UE4 - fix async render thread updates and particle threading UE-18189

diff --git a/Engine/Source/Runtime/Engine/Classes/Particles/ParticleSystemComponent.h b/Engine/Source/Runtime/Engine/Classes/Particles/ParticleSystemComponent.h
--- a/Engine/Source/Runtime/Engine/Classes/Particles/ParticleSystemComponent.h
+++ b/Engine/Source/Runtime/Engine/Classes/Particles/ParticleSystemComponent.h
@@ -453,11 +453,12 @@ private:
 	TArray<struct FParticleSysParam> AsyncInstanceParameters;
 	/** Is AsyncComponentToWorld etc valid? */
 	bool bAsyncDataCopyIsValid;
+	bool bParallelRenderThreadUpdate;
 public:
 
 	const FTransform& GetAsyncComponentToWorld()
 	{
-		if (!IsInGameThread())
+		if (!bParallelRenderThreadUpdate && !IsInGameThread())
 		{
 			check(bAsyncDataCopyIsValid); 
 			return AsyncComponentToWorld;
@@ -467,7 +468,7 @@ public:
 
 	const TArray<struct FParticleSysParam>& GetAsyncInstanceParameters()
 	{
-		if (!IsInGameThread())
+		if (!bParallelRenderThreadUpdate && !IsInGameThread())
 		{
 			check(bAsyncDataCopyIsValid); 
 			return AsyncInstanceParameters;
diff --git a/Engine/Source/Runtime/Engine/Private/Particles/ParticleComponents.cpp b/Engine/Source/Runtime/Engine/Private/Particles/ParticleComponents.cpp
--- a/Engine/Source/Runtime/Engine/Private/Particles/ParticleComponents.cpp
+++ b/Engine/Source/Runtime/Engine/Private/Particles/ParticleComponents.cpp
@@ -2916,6 +2916,11 @@ void UParticleSystemComponent::SendRenderDynamicData_Concurrent()
 	ForceAsyncWorkCompletion(ENSURE_AND_STALL);
 	Super::SendRenderDynamicData_Concurrent();
 
+	check(!bAsyncDataCopyIsValid);
+	check(!bParallelRenderThreadUpdate);
+	bParallelRenderThreadUpdate = true;
+
+
 	FParticleSystemSceneProxy* PSysSceneProxy = (FParticleSystemSceneProxy*)SceneProxy;
 	if (PSysSceneProxy != NULL)
 	{
@@ -2935,6 +2940,7 @@ void UParticleSystemComponent::SendRenderDynamicData_Concurrent()
 			}
 		}
 	}
+	bParallelRenderThreadUpdate = false;
 }
 
 void UParticleSystemComponent::DestroyRenderState_Concurrent()
@@ -3734,6 +3740,7 @@ void UParticleSystemComponent::TickComponent(float DeltaTime, enum ELevelTick Ti
 	{
 		QUICK_SCOPE_CYCLE_COUNTER(STAT_UParticleSystemComponent_QueueTasks);
 		bAsyncDataCopyIsValid = true;
+		check(!bParallelRenderThreadUpdate);
 		AsyncComponentToWorld = ComponentToWorld;
 		AsyncInstanceParameters = InstanceParameters;
 
