Github pull request #827 - skip spawning particle system components on dedicated server.

diff --git a/Engine/Source/Runtime/Engine/Private/GameplayStatics.cpp b/Engine/Source/Runtime/Engine/Private/GameplayStatics.cpp
--- a/Engine/Source/Runtime/Engine/Private/GameplayStatics.cpp
+++ b/Engine/Source/Runtime/Engine/Private/GameplayStatics.cpp
@@ -672,7 +672,7 @@ UParticleSystemComponent* UGameplayStatics::SpawnEmitterAttached(UParticleSystem
 		{
 			UE_LOG(LogScript, Warning, TEXT("UGameplayStatics::SpawnEmitterAttached: NULL AttachComponent specified!"));
 		}
-		else
+		else if (AttachToComponent->GetWorld() != NULL && AttachToComponent->GetWorld()->GetNetMode() != NM_DedicatedServer)
 		{
 			PSC = CreateParticleSystem(EmitterTemplate, AttachToComponent->GetWorld(), AttachToComponent->GetOwner(), bAutoDestroy);
 
