Move physx into its own module.

diff --git a/Engine/Source/Runtime/Engine/Private/PhysicsEngine/PhysXSupport.h b/Engine/Source/Runtime/Engine/Private/PhysicsEngine/PhysXSupport.h
--- a/Engine/Source/Runtime/Engine/Private/PhysicsEngine/PhysXSupport.h
+++ b/Engine/Source/Runtime/Engine/Private/PhysicsEngine/PhysXSupport.h
@@ -6,15 +6,9 @@
 
 #pragma once
 
-#if WITH_PHYSX
-
-#include "PhysXIncludes.h"
-#include "EngineLogs.h"
+#include "PhysXPublic.h"
 
-// Whether or not to use the PhysX scene lock
-#ifndef USE_SCENE_LOCK
-#define USE_SCENE_LOCK			1
-#endif
+#if WITH_PHYSX
 
 // Whether to track PhysX memory allocations
 #ifndef PHYSX_MEMORY_VALIDATION
@@ -33,388 +27,11 @@
 
 #define PHYSX_MEMORY_STAT_ONLY (0)
 
-#if USE_SCENE_LOCK
-
-/** Scoped scene read lock - we use this instead of PxSceneReadLock because it handles NULL scene */
-class FPhysXSceneReadLock
-{
-public:
-	
-	FPhysXSceneReadLock(PxScene* PInScene, const char* filename, PxU32 lineno)
-		: PScene(PInScene)
-	{
-		SCOPE_CYCLE_COUNTER(STAT_PhysSceneReadLock);
-		if(PScene)
-		{
-			PScene->lockRead(filename, lineno);
-		}
-	}
-
-	~FPhysXSceneReadLock()
-	{
-		if(PScene)
-		{
-			PScene->unlockRead();
-		}
-	}
-
-private:
-	PxScene* PScene;
-};
-
-#if WITH_APEX
-/** Scoped scene read lock - we use this instead of PxSceneReadLock because it handles NULL scene */
-class FApexSceneReadLock
-{
-public:
-
-	FApexSceneReadLock(NxApexScene* PInScene, const char* filename, PxU32 lineno)
-		: PScene(PInScene)
-	{
-		SCOPE_CYCLE_COUNTER(STAT_PhysSceneReadLock);
-		if (PScene)
-		{
-			PScene->lockRead(filename, lineno);
-		}
-	}
-
-	~FApexSceneReadLock()
-	{
-		if (PScene)
-		{
-			PScene->unlockRead();
-		}
-	}
-
-private:
-	NxApexScene* PScene;
-};
-#endif
-
-/** Scoped scene write lock - we use this instead of PxSceneReadLock because it handles NULL scene */
-class FPhysXSceneWriteLock
-{
-public:
-	FPhysXSceneWriteLock(PxScene* PInScene, const char* filename, PxU32 lineno)
-		: PScene(PInScene)
-	{
-		SCOPE_CYCLE_COUNTER(STAT_PhysSceneWriteLock);
-		if(PScene)
-		{
-			PScene->lockWrite(filename, lineno);
-		}
-	}
-
-	~FPhysXSceneWriteLock()
-	{
-		if(PScene)
-		{
-			PScene->unlockWrite();
-		}
-	}
-
-private:
-	PxScene* PScene;
-};
-
-#if WITH_APEX
-/** Scoped scene write lock - we use this instead of PxSceneReadLock because it handles NULL scene */
-class FApexSceneWriteLock
-{
-public:
-	FApexSceneWriteLock(NxApexScene* PInScene, const char* filename, PxU32 lineno)
-		: PScene(PInScene)
-	{
-		SCOPE_CYCLE_COUNTER(STAT_PhysSceneWriteLock);
-		if (PScene)
-		{
-			PScene->lockWrite(filename, lineno);
-		}
-	}
-
-	~FApexSceneWriteLock()
-	{
-		if (PScene)
-		{
-			PScene->unlockWrite();
-		}
-	}
-
-private:
-	NxApexScene* PScene;
-};
-#endif
-
-#define SCOPED_SCENE_READ_LOCK( _scene ) FPhysXSceneReadLock PREPROCESSOR_JOIN(_rlock,__LINE__)(_scene, __FILE__, __LINE__)
-#define SCOPED_SCENE_WRITE_LOCK( _scene ) FPhysXSceneWriteLock PREPROCESSOR_JOIN(_wlock,__LINE__)(_scene, __FILE__, __LINE__)
-#if WITH_APEX
-#define SCOPED_APEX_SCENE_READ_LOCK( _scene ) FApexSceneReadLock PREPROCESSOR_JOIN(_rlock,__LINE__)(_scene, __FILE__, __LINE__)
-#define SCOPED_APEX_SCENE_WRITE_LOCK( _scene ) FApexSceneWriteLock PREPROCESSOR_JOIN(_wlock,__LINE__)(_scene, __FILE__, __LINE__)
-#endif
-
-#define SCENE_LOCK_READ( _scene )		{ SCOPE_CYCLE_COUNTER(STAT_PhysSceneReadLock); if((_scene) != NULL) { (_scene)->lockRead(__FILE__, __LINE__); } }
-#define SCENE_UNLOCK_READ( _scene )		{ if((_scene) != NULL) { (_scene)->unlockRead(); } }
-#define SCENE_LOCK_WRITE( _scene )		{ SCOPE_CYCLE_COUNTER(STAT_PhysSceneWriteLock); if((_scene) != NULL) { (_scene)->lockWrite(__FILE__, __LINE__); } }
-#define SCENE_UNLOCK_WRITE( _scene )	{ if((_scene) != NULL) { (_scene)->unlockWrite(); } }
-#else
-#define SCOPED_SCENE_READ_LOCK_INDEXED( _scene, _index )
-#define SCOPED_SCENE_READ_LOCK( _scene )
-#define SCOPED_SCENE_WRITE_LOCK_INDEXED( _scene, _index )
-#define SCOPED_SCENE_WRITE_LOCK( _scene )
-#define SCENE_LOCK_READ( _scene )
-#define SCENE_UNLOCK_READ( _scene )
-#define SCENE_LOCK_WRITE( _scene )
-#define SCENE_UNLOCK_WRITE( _scene )
-
-#if WITH_APEX
-#define SCOPED_APEX_SCENE_READ_LOCK( _scene )
-#define SCOPED_APEX_SCENE_WRITE_LOCK( _scene )
-#endif
-#endif
-
-/** Get a pointer to the PxScene from an SceneIndex (will be NULL if scene already shut down) */
-PxScene* GetPhysXSceneFromIndex(int32 InSceneIndex);
-
 #if WITH_APEX
 /** Get a pointer to the NxApexScene from an SceneIndex (will be NULL if scene already shut down) */
 NxApexScene* GetApexSceneFromIndex(int32 InSceneIndex);
 #endif
 
-template <bool NeedsLock>
-struct FPhysXSupport
-{
-	/** Obtains the appropriate PhysX scene lock for READING and executes the passed in lambda.
-	 *  If SceneType < 0, the Sync actor is used, otherwise the async.
-	 *  Note: The lambda is only executed if the physx actor requested is non-null.
-	 *  returns true if the requested actor is non-null
-	 */
-	template <typename LambdaType>
-	static bool ExecuteOnPxRigidActorReadOnly(const FBodyInstance* BI, const LambdaType& Func, int32 SceneType = -1)
-	{
-		if (const PxRigidActor* PRigidActor = BI->GetPxRigidActor_AssumesLocked(SceneType))
-		{
-			const int32 SceneIndex = (PRigidActor == BI->RigidActorSync ? BI->SceneIndexSync : BI->SceneIndexAsync);
-			PxScene* PScene = GetPhysXSceneFromIndex(SceneIndex);
-			if (NeedsLock)
-			{
-				SCENE_LOCK_READ(PScene);
-			}
-
-			Func(PRigidActor);
-
-			if (NeedsLock)
-			{
-				SCENE_UNLOCK_READ(PScene);
-			}
-
-			return true;
-		}
-
-		return false;
-	}
-
-	/** Obtains the appropriate PhysX scene lock for READING and executes the passed in lambda.
-	 *  Note: The lambda is only executed if the physx actor is a non-null RigidBody
-	 *  returns true if found a non-null RigidBody
-	 */
-	template <typename LambdaType>
-	static bool ExecuteOnPxRigidBodyReadOnly(const FBodyInstance* BI, const LambdaType& Func)
-	{
-		bool bSuccess = false;
-		if (const physx::PxRigidActor* RigidActor = BI->GetPxRigidActor_AssumesLocked())
-		{
-			const int32 SceneIndex = (RigidActor == BI->RigidActorSync ? BI->SceneIndexSync : BI->SceneIndexAsync);
-			PxScene* PScene = GetPhysXSceneFromIndex(SceneIndex);
-			if (NeedsLock)
-			{
-				SCENE_LOCK_READ(PScene);
-			}
-
-			if (const physx::PxRigidBody* PRigidBody = RigidActor->isRigidBody())
-			{
-				Func(PRigidBody);
-				bSuccess = true;
-			}
-
-			if (NeedsLock)
-			{
-				SCENE_UNLOCK_READ(PScene);
-			}
-		}
-
-		return bSuccess;
-	}
-
-	/** Obtains the appropriate PhysX scene lock for WRITING and executes the passed in lambda.
-	 *  Note: The lambda is only executed if the physx actor is a non-null RigidBody
-	 *  returns true if found a non-null RigidBody.
-	 */
-	template <typename LambdaType>
-	static bool ExecuteOnPxRigidBodyReadWrite(const FBodyInstance* BI, const LambdaType& Func)
-	{
-		bool bSuccess = false;
-		if (physx::PxRigidActor* RigidActor = BI->GetPxRigidActor_AssumesLocked())
-		{
-			const int32 SceneIndex = (RigidActor == BI->RigidActorSync ? BI->SceneIndexSync : BI->SceneIndexAsync);
-			PxScene* PScene = GetPhysXSceneFromIndex(SceneIndex);
-			if(NeedsLock)
-			{
-				SCENE_LOCK_WRITE(PScene);
-			}
-			
-			if (physx::PxRigidBody* PRigidBody = RigidActor->isRigidBody())
-			{
-				Func(PRigidBody);
-				bSuccess = true;
-			}
-
-			if(NeedsLock)
-			{
-				SCENE_UNLOCK_WRITE(PScene);
-			}
-		}
-
-		return bSuccess;
-	}
-
-	/** Obtains the appropriate PhysX scene lock for READING and executes the passed in lambda.
-	 *  Note: The lambda is only executed if the physx actor is a non-null RigidDynamic
-	 *  returns true if found a non-null RigidDynamic.
-	 */
-	template <typename LambdaType>
-	static bool ExecuteOnPxRigidDynamicReadOnly(const FBodyInstance* BI, const LambdaType& Func)
-	{
-		bool bSuccess = false;
-		if (physx::PxRigidActor* RigidActor = BI->GetPxRigidActor_AssumesLocked())
-		{
-			const int32 SceneIndex = (RigidActor == BI->RigidActorSync ? BI->SceneIndexSync : BI->SceneIndexAsync);
-			PxScene* PScene = GetPhysXSceneFromIndex(SceneIndex);
-			if (NeedsLock)
-			{
-				SCENE_LOCK_READ(PScene);
-			}
-
-			if (physx::PxRigidDynamic* PRigidDynamic = RigidActor->isRigidDynamic())
-			{
-				Func(PRigidDynamic);
-				bSuccess = true;
-			}
-
-			if (NeedsLock)
-			{
-				SCENE_UNLOCK_READ(PScene);
-			}
-		}
-
-		return bSuccess;
-	}
-
-	/** Obtains the appropriate PhysX scene lock for WRITING and executes the passed in lambda.
-	 *  Note: The lambda is only executed if the physx actor is a non-null RigidDynamic
-	 *  returns true if found a non-null RigidDynamic.
-	 */
-	template <typename LambdaType>
-	static bool ExecuteOnPxRigidDynamicReadWrite(const FBodyInstance* BI, const LambdaType& Func)
-	{
-		bool bSuccess = false;
-		if (physx::PxRigidActor* RigidActor = BI->GetPxRigidActor_AssumesLocked())
-		{
-			const int32 SceneIndex = (RigidActor == BI->RigidActorSync ? BI->SceneIndexSync : BI->SceneIndexAsync);
-			PxScene* PScene = GetPhysXSceneFromIndex(SceneIndex);
-			if (NeedsLock)
-			{
-				SCENE_LOCK_WRITE(PScene);
-			}
-
-			if (physx::PxRigidDynamic* PRigidDynamic = RigidActor->isRigidDynamic())
-			{
-				Func(PRigidDynamic);
-				bSuccess = true;
-			}
-
-			if (NeedsLock)
-			{
-				SCENE_UNLOCK_WRITE(PScene);
-			}
-		}
-
-		return bSuccess;
-	}
-};
-
-// Utility functions for obtaining locks and executing lambda. This indirection is needed for vs2012 but should be inlined
-template <typename LambdaType> bool ExecuteOnPxRigidActorReadOnly(const FBodyInstance* BI, const LambdaType& Func, int32 SceneType = -1){ return FPhysXSupport<true>::ExecuteOnPxRigidActorReadOnly(BI, Func, SceneType); }
-template <typename LambdaType> bool ExecuteOnPxRigidBodyReadOnly(const FBodyInstance* BI, const LambdaType& Func) { return FPhysXSupport<true>::ExecuteOnPxRigidBodyReadOnly(BI, Func); }
-template <typename LambdaType> bool ExecuteOnPxRigidBodyReadWrite(const FBodyInstance* BI, const LambdaType& Func){ return FPhysXSupport<true>::ExecuteOnPxRigidBodyReadWrite(BI, Func); }
-template <typename LambdaType> bool ExecuteOnPxRigidDynamicReadOnly(const FBodyInstance* BI, const LambdaType& Func){ return FPhysXSupport<true>::ExecuteOnPxRigidDynamicReadOnly(BI, Func); }
-template <typename LambdaType> bool ExecuteOnPxRigidDynamicReadWrite(const FBodyInstance* BI, const LambdaType& Func){ return FPhysXSupport<true>::ExecuteOnPxRigidDynamicReadWrite(BI, Func); }
-
-//////// BASIC TYPE CONVERSION
-
-/** Convert Unreal FMatrix to PhysX PxTransform */
-ENGINE_API PxTransform UMatrix2PTransform(const FMatrix& UTM);
-/** Convert Unreal FTransform to PhysX PxTransform */
-ENGINE_API PxTransform U2PTransform(const FTransform& UTransform);
-/** Convert Unreal FVector to PhysX PxVec3 */
-ENGINE_API PxVec3 U2PVector(const FVector& UVec);
-/** Convert Unreal FQuat to PhysX PxTransform */
-ENGINE_API PxQuat U2PQuat(const FQuat& UQuat);
-/** Convert Unreal FMatrix to PhysX PxMat44 */
-ENGINE_API PxMat44 U2PMatrix(const FMatrix& UTM);
-/** Convert Unreal FPlane to PhysX plane def */
-ENGINE_API PxPlane U2PPlane(const FPlane& Plane);
-/** Convert PhysX PxTransform to Unreal PxTransform */
-ENGINE_API FTransform P2UTransform(const PxTransform& PTM);
-/** Convert PhysX PxVec3 to Unreal FVector */
-ENGINE_API FVector P2UVector(const PxVec3& PVec);
-/** Convert PhysX PxQuat to Unreal FQuat */
-ENGINE_API FQuat P2UQuat(const PxQuat& PQuat);
-/** Convert PhysX plane def to Unreal FPlane */
-ENGINE_API FPlane P2UPlane(const PxReal P[4]);
-ENGINE_API FPlane P2UPlane(const PxPlane& Plane);
-/** Convert PhysX PxMat44 to Unreal FMatrix */
-ENGINE_API FMatrix P2UMatrix(const PxMat44& PMat);
-/** Convert PhysX PxTransform to Unreal FMatrix */
-ENGINE_API FMatrix PTransform2UMatrix(const PxTransform& PTM);
-
-// inlines
-
-ENGINE_API FORCEINLINE_DEBUGGABLE PxVec3 U2PVector(const FVector& UVec)
-{
-	return PxVec3(UVec.X, UVec.Y, UVec.Z);
-}
-
-ENGINE_API FORCEINLINE_DEBUGGABLE PxQuat U2PQuat(const FQuat& UQuat)
-{
-	return PxQuat( UQuat.X, UQuat.Y, UQuat.Z, UQuat.W );
-}
-
-ENGINE_API FORCEINLINE_DEBUGGABLE PxPlane U2PPlane(const FPlane& Plane)
-{
-	return PxPlane(Plane.X, Plane.Y, Plane.Z, -Plane.W);
-}
-
-ENGINE_API FORCEINLINE_DEBUGGABLE FVector P2UVector(const PxVec3& PVec)
-{
-	return FVector(PVec.x, PVec.y, PVec.z);
-}
-
-ENGINE_API FORCEINLINE_DEBUGGABLE FQuat P2UQuat(const PxQuat& PQuat)
-{
-	return FQuat(PQuat.x, PQuat.y, PQuat.z, PQuat.w);
-}
-
-ENGINE_API FORCEINLINE_DEBUGGABLE FPlane P2UPlane(const PxReal P[4])
-{
-	return FPlane(P[0], P[1], P[2], -P[3]);
-}
-
-ENGINE_API FORCEINLINE_DEBUGGABLE FPlane P2UPlane(const PxPlane& Plane)
-{
-	return FPlane(Plane.n.x, Plane.n.y, Plane.n.z, -Plane.d);
-}
-
-
 //////// GEOM CONVERSION
 // we need this helper struct since PhysX needs geoms to be on the stack
 struct UCollision2PGeom
diff --git a/Engine/Source/Runtime/Engine/Public/PhysXPublic.h b/Engine/Source/Runtime/Engine/Public/PhysXPublic.h
new file mode 100644
--- /dev/null
+++ b/Engine/Source/Runtime/Engine/Public/PhysXPublic.h
@@ -0,0 +1,395 @@
+// Copyright 1998-2015 Epic Games, Inc. All Rights Reserved.
+
+/*=============================================================================
+	PhysXSupport.h: PhysX support
+=============================================================================*/
+
+#pragma once
+
+#if WITH_PHYSX
+
+#include "PhysXIncludes.h"
+#include "EngineLogs.h"
+
+// Whether or not to use the PhysX scene lock
+#ifndef USE_SCENE_LOCK
+#define USE_SCENE_LOCK			1
+#endif
+
+#if USE_SCENE_LOCK
+
+/** Scoped scene read lock - we use this instead of PxSceneReadLock because it handles NULL scene */
+class FPhysXSceneReadLock
+{
+public:
+	
+	FPhysXSceneReadLock(PxScene* PInScene, const char* filename, PxU32 lineno)
+		: PScene(PInScene)
+	{
+		SCOPE_CYCLE_COUNTER(STAT_PhysSceneReadLock);
+		if(PScene)
+		{
+			PScene->lockRead(filename, lineno);
+		}
+	}
+
+	~FPhysXSceneReadLock()
+	{
+		if(PScene)
+		{
+			PScene->unlockRead();
+		}
+	}
+
+private:
+	PxScene* PScene;
+};
+
+#if WITH_APEX
+/** Scoped scene read lock - we use this instead of PxSceneReadLock because it handles NULL scene */
+class FApexSceneReadLock
+{
+public:
+
+	FApexSceneReadLock(NxApexScene* PInScene, const char* filename, PxU32 lineno)
+		: PScene(PInScene)
+	{
+		SCOPE_CYCLE_COUNTER(STAT_PhysSceneReadLock);
+		if (PScene)
+		{
+			PScene->lockRead(filename, lineno);
+		}
+	}
+
+	~FApexSceneReadLock()
+	{
+		if (PScene)
+		{
+			PScene->unlockRead();
+		}
+	}
+
+private:
+	NxApexScene* PScene;
+};
+#endif
+
+/** Scoped scene write lock - we use this instead of PxSceneReadLock because it handles NULL scene */
+class FPhysXSceneWriteLock
+{
+public:
+	FPhysXSceneWriteLock(PxScene* PInScene, const char* filename, PxU32 lineno)
+		: PScene(PInScene)
+	{
+		SCOPE_CYCLE_COUNTER(STAT_PhysSceneWriteLock);
+		if(PScene)
+		{
+			PScene->lockWrite(filename, lineno);
+		}
+	}
+
+	~FPhysXSceneWriteLock()
+	{
+		if(PScene)
+		{
+			PScene->unlockWrite();
+		}
+	}
+
+private:
+	PxScene* PScene;
+};
+
+#if WITH_APEX
+/** Scoped scene write lock - we use this instead of PxSceneReadLock because it handles NULL scene */
+class FApexSceneWriteLock
+{
+public:
+	FApexSceneWriteLock(NxApexScene* PInScene, const char* filename, PxU32 lineno)
+		: PScene(PInScene)
+	{
+		SCOPE_CYCLE_COUNTER(STAT_PhysSceneWriteLock);
+		if (PScene)
+		{
+			PScene->lockWrite(filename, lineno);
+		}
+	}
+
+	~FApexSceneWriteLock()
+	{
+		if (PScene)
+		{
+			PScene->unlockWrite();
+		}
+	}
+
+private:
+	NxApexScene* PScene;
+};
+#endif
+
+#define SCOPED_SCENE_READ_LOCK( _scene ) FPhysXSceneReadLock PREPROCESSOR_JOIN(_rlock,__LINE__)(_scene, __FILE__, __LINE__)
+#define SCOPED_SCENE_WRITE_LOCK( _scene ) FPhysXSceneWriteLock PREPROCESSOR_JOIN(_wlock,__LINE__)(_scene, __FILE__, __LINE__)
+#if WITH_APEX
+#define SCOPED_APEX_SCENE_READ_LOCK( _scene ) FApexSceneReadLock PREPROCESSOR_JOIN(_rlock,__LINE__)(_scene, __FILE__, __LINE__)
+#define SCOPED_APEX_SCENE_WRITE_LOCK( _scene ) FApexSceneWriteLock PREPROCESSOR_JOIN(_wlock,__LINE__)(_scene, __FILE__, __LINE__)
+#endif
+
+#define SCENE_LOCK_READ( _scene )		{ SCOPE_CYCLE_COUNTER(STAT_PhysSceneReadLock); if((_scene) != NULL) { (_scene)->lockRead(__FILE__, __LINE__); } }
+#define SCENE_UNLOCK_READ( _scene )		{ if((_scene) != NULL) { (_scene)->unlockRead(); } }
+#define SCENE_LOCK_WRITE( _scene )		{ SCOPE_CYCLE_COUNTER(STAT_PhysSceneWriteLock); if((_scene) != NULL) { (_scene)->lockWrite(__FILE__, __LINE__); } }
+#define SCENE_UNLOCK_WRITE( _scene )	{ if((_scene) != NULL) { (_scene)->unlockWrite(); } }
+#else
+#define SCOPED_SCENE_READ_LOCK_INDEXED( _scene, _index )
+#define SCOPED_SCENE_READ_LOCK( _scene )
+#define SCOPED_SCENE_WRITE_LOCK_INDEXED( _scene, _index )
+#define SCOPED_SCENE_WRITE_LOCK( _scene )
+#define SCENE_LOCK_READ( _scene )
+#define SCENE_UNLOCK_READ( _scene )
+#define SCENE_LOCK_WRITE( _scene )
+#define SCENE_UNLOCK_WRITE( _scene )
+
+#if WITH_APEX
+#define SCOPED_APEX_SCENE_READ_LOCK( _scene )
+#define SCOPED_APEX_SCENE_WRITE_LOCK( _scene )
+#endif
+#endif
+
+/** Get a pointer to the PxScene from an SceneIndex (will be NULL if scene already shut down) */
+PxScene* GetPhysXSceneFromIndex(int32 InSceneIndex);
+
+template <bool NeedsLock>
+struct FPhysXSupport
+{
+	/** Obtains the appropriate PhysX scene lock for READING and executes the passed in lambda.
+	 *  If SceneType < 0, the Sync actor is used, otherwise the async.
+	 *  Note: The lambda is only executed if the physx actor requested is non-null.
+	 *  returns true if the requested actor is non-null
+	 */
+	template <typename LambdaType>
+	static bool ExecuteOnPxRigidActorReadOnly(const FBodyInstance* BI, const LambdaType& Func, int32 SceneType = -1)
+	{
+		if (const PxRigidActor* PRigidActor = BI->GetPxRigidActor_AssumesLocked(SceneType))
+		{
+			const int32 SceneIndex = (PRigidActor == BI->RigidActorSync ? BI->SceneIndexSync : BI->SceneIndexAsync);
+			PxScene* PScene = GetPhysXSceneFromIndex(SceneIndex);
+			if (NeedsLock)
+			{
+				SCENE_LOCK_READ(PScene);
+			}
+
+			Func(PRigidActor);
+
+			if (NeedsLock)
+			{
+				SCENE_UNLOCK_READ(PScene);
+			}
+
+			return true;
+		}
+
+		return false;
+	}
+
+	/** Obtains the appropriate PhysX scene lock for READING and executes the passed in lambda.
+	 *  Note: The lambda is only executed if the physx actor is a non-null RigidBody
+	 *  returns true if found a non-null RigidBody
+	 */
+	template <typename LambdaType>
+	static bool ExecuteOnPxRigidBodyReadOnly(const FBodyInstance* BI, const LambdaType& Func)
+	{
+		bool bSuccess = false;
+		if (const physx::PxRigidActor* RigidActor = BI->GetPxRigidActor_AssumesLocked())
+		{
+			const int32 SceneIndex = (RigidActor == BI->RigidActorSync ? BI->SceneIndexSync : BI->SceneIndexAsync);
+			PxScene* PScene = GetPhysXSceneFromIndex(SceneIndex);
+			if (NeedsLock)
+			{
+				SCENE_LOCK_READ(PScene);
+			}
+
+			if (const physx::PxRigidBody* PRigidBody = RigidActor->isRigidBody())
+			{
+				Func(PRigidBody);
+				bSuccess = true;
+			}
+
+			if (NeedsLock)
+			{
+				SCENE_UNLOCK_READ(PScene);
+			}
+		}
+
+		return bSuccess;
+	}
+
+	/** Obtains the appropriate PhysX scene lock for WRITING and executes the passed in lambda.
+	 *  Note: The lambda is only executed if the physx actor is a non-null RigidBody
+	 *  returns true if found a non-null RigidBody.
+	 */
+	template <typename LambdaType>
+	static bool ExecuteOnPxRigidBodyReadWrite(const FBodyInstance* BI, const LambdaType& Func)
+	{
+		bool bSuccess = false;
+		if (physx::PxRigidActor* RigidActor = BI->GetPxRigidActor_AssumesLocked())
+		{
+			const int32 SceneIndex = (RigidActor == BI->RigidActorSync ? BI->SceneIndexSync : BI->SceneIndexAsync);
+			PxScene* PScene = GetPhysXSceneFromIndex(SceneIndex);
+			if(NeedsLock)
+			{
+				SCENE_LOCK_WRITE(PScene);
+			}
+			
+			if (physx::PxRigidBody* PRigidBody = RigidActor->isRigidBody())
+			{
+				Func(PRigidBody);
+				bSuccess = true;
+			}
+
+			if(NeedsLock)
+			{
+				SCENE_UNLOCK_WRITE(PScene);
+			}
+		}
+
+		return bSuccess;
+	}
+
+	/** Obtains the appropriate PhysX scene lock for READING and executes the passed in lambda.
+	 *  Note: The lambda is only executed if the physx actor is a non-null RigidDynamic
+	 *  returns true if found a non-null RigidDynamic.
+	 */
+	template <typename LambdaType>
+	static bool ExecuteOnPxRigidDynamicReadOnly(const FBodyInstance* BI, const LambdaType& Func)
+	{
+		bool bSuccess = false;
+		if (physx::PxRigidActor* RigidActor = BI->GetPxRigidActor_AssumesLocked())
+		{
+			const int32 SceneIndex = (RigidActor == BI->RigidActorSync ? BI->SceneIndexSync : BI->SceneIndexAsync);
+			PxScene* PScene = GetPhysXSceneFromIndex(SceneIndex);
+			if (NeedsLock)
+			{
+				SCENE_LOCK_READ(PScene);
+			}
+
+			if (physx::PxRigidDynamic* PRigidDynamic = RigidActor->isRigidDynamic())
+			{
+				Func(PRigidDynamic);
+				bSuccess = true;
+			}
+
+			if (NeedsLock)
+			{
+				SCENE_UNLOCK_READ(PScene);
+			}
+		}
+
+		return bSuccess;
+	}
+
+	/** Obtains the appropriate PhysX scene lock for WRITING and executes the passed in lambda.
+	 *  Note: The lambda is only executed if the physx actor is a non-null RigidDynamic
+	 *  returns true if found a non-null RigidDynamic.
+	 */
+	template <typename LambdaType>
+	static bool ExecuteOnPxRigidDynamicReadWrite(const FBodyInstance* BI, const LambdaType& Func)
+	{
+		bool bSuccess = false;
+		if (physx::PxRigidActor* RigidActor = BI->GetPxRigidActor_AssumesLocked())
+		{
+			const int32 SceneIndex = (RigidActor == BI->RigidActorSync ? BI->SceneIndexSync : BI->SceneIndexAsync);
+			PxScene* PScene = GetPhysXSceneFromIndex(SceneIndex);
+			if (NeedsLock)
+			{
+				SCENE_LOCK_WRITE(PScene);
+			}
+
+			if (physx::PxRigidDynamic* PRigidDynamic = RigidActor->isRigidDynamic())
+			{
+				Func(PRigidDynamic);
+				bSuccess = true;
+			}
+
+			if (NeedsLock)
+			{
+				SCENE_UNLOCK_WRITE(PScene);
+			}
+		}
+
+		return bSuccess;
+	}
+};
+
+// Utility functions for obtaining locks and executing lambda. This indirection is needed for vs2012 but should be inlined
+template <typename LambdaType> bool ExecuteOnPxRigidActorReadOnly(const FBodyInstance* BI, const LambdaType& Func, int32 SceneType = -1){ return FPhysXSupport<true>::ExecuteOnPxRigidActorReadOnly(BI, Func, SceneType); }
+template <typename LambdaType> bool ExecuteOnPxRigidBodyReadOnly(const FBodyInstance* BI, const LambdaType& Func) { return FPhysXSupport<true>::ExecuteOnPxRigidBodyReadOnly(BI, Func); }
+template <typename LambdaType> bool ExecuteOnPxRigidBodyReadWrite(const FBodyInstance* BI, const LambdaType& Func){ return FPhysXSupport<true>::ExecuteOnPxRigidBodyReadWrite(BI, Func); }
+template <typename LambdaType> bool ExecuteOnPxRigidDynamicReadOnly(const FBodyInstance* BI, const LambdaType& Func){ return FPhysXSupport<true>::ExecuteOnPxRigidDynamicReadOnly(BI, Func); }
+template <typename LambdaType> bool ExecuteOnPxRigidDynamicReadWrite(const FBodyInstance* BI, const LambdaType& Func){ return FPhysXSupport<true>::ExecuteOnPxRigidDynamicReadWrite(BI, Func); }
+
+//////// BASIC TYPE CONVERSION
+
+/** Convert Unreal FMatrix to PhysX PxTransform */
+ENGINE_API PxTransform UMatrix2PTransform(const FMatrix& UTM);
+/** Convert Unreal FTransform to PhysX PxTransform */
+ENGINE_API PxTransform U2PTransform(const FTransform& UTransform);
+/** Convert Unreal FVector to PhysX PxVec3 */
+ENGINE_API PxVec3 U2PVector(const FVector& UVec);
+/** Convert Unreal FQuat to PhysX PxTransform */
+ENGINE_API PxQuat U2PQuat(const FQuat& UQuat);
+/** Convert Unreal FMatrix to PhysX PxMat44 */
+ENGINE_API PxMat44 U2PMatrix(const FMatrix& UTM);
+/** Convert Unreal FPlane to PhysX plane def */
+ENGINE_API PxPlane U2PPlane(const FPlane& Plane);
+/** Convert PhysX PxTransform to Unreal PxTransform */
+ENGINE_API FTransform P2UTransform(const PxTransform& PTM);
+/** Convert PhysX PxVec3 to Unreal FVector */
+ENGINE_API FVector P2UVector(const PxVec3& PVec);
+/** Convert PhysX PxQuat to Unreal FQuat */
+ENGINE_API FQuat P2UQuat(const PxQuat& PQuat);
+/** Convert PhysX plane def to Unreal FPlane */
+ENGINE_API FPlane P2UPlane(const PxReal P[4]);
+ENGINE_API FPlane P2UPlane(const PxPlane& Plane);
+/** Convert PhysX PxMat44 to Unreal FMatrix */
+ENGINE_API FMatrix P2UMatrix(const PxMat44& PMat);
+/** Convert PhysX PxTransform to Unreal FMatrix */
+ENGINE_API FMatrix PTransform2UMatrix(const PxTransform& PTM);
+
+// inlines
+
+ENGINE_API FORCEINLINE_DEBUGGABLE PxVec3 U2PVector(const FVector& UVec)
+{
+	return PxVec3(UVec.X, UVec.Y, UVec.Z);
+}
+
+ENGINE_API FORCEINLINE_DEBUGGABLE PxQuat U2PQuat(const FQuat& UQuat)
+{
+	return PxQuat( UQuat.X, UQuat.Y, UQuat.Z, UQuat.W );
+}
+
+ENGINE_API FORCEINLINE_DEBUGGABLE PxPlane U2PPlane(const FPlane& Plane)
+{
+	return PxPlane(Plane.X, Plane.Y, Plane.Z, -Plane.W);
+}
+
+ENGINE_API FORCEINLINE_DEBUGGABLE FVector P2UVector(const PxVec3& PVec)
+{
+	return FVector(PVec.x, PVec.y, PVec.z);
+}
+
+ENGINE_API FORCEINLINE_DEBUGGABLE FQuat P2UQuat(const PxQuat& PQuat)
+{
+	return FQuat(PQuat.x, PQuat.y, PQuat.z, PQuat.w);
+}
+
+ENGINE_API FORCEINLINE_DEBUGGABLE FPlane P2UPlane(const PxReal P[4])
+{
+	return FPlane(P[0], P[1], P[2], -P[3]);
+}
+
+ENGINE_API FORCEINLINE_DEBUGGABLE FPlane P2UPlane(const PxPlane& Plane)
+{
+	return FPlane(Plane.n.x, Plane.n.y, Plane.n.z, -Plane.d);
+}
+
+#endif
\ No newline at end of file
