PR #1346: Console ServerTravel command can't pass maps parameters (Contributed by [REDACTED_PERSON])

diff --git a/Engine/Source/Runtime/Engine/Private/UnrealEngine.cpp b/Engine/Source/Runtime/Engine/Private/UnrealEngine.cpp
--- a/Engine/Source/Runtime/Engine/Private/UnrealEngine.cpp
+++ b/Engine/Source/Runtime/Engine/Private/UnrealEngine.cpp
@@ -8452,16 +8452,26 @@ bool UEngine::HandleServerTravelCommand( const TCHAR* Cmd, FOutputDevice& Ar, UW
 {
 	if( InWorld->IsServer() )
 	{
-		FString MapName = Cmd;
+		FString URL(Cmd);
+		FString MapName, Options;
+
+		// Split the options off the map name before we validate it
+		if (URL.Split(TEXT("?"), &MapName, &Options, ESearchCase::CaseSensitive) == false)
+		{
+			MapName = URL;
+		}
+
 		if (MakeSureMapNameIsValid(MapName))
 		{
-			InWorld->ServerTravel(MapName);
+			// If there were options reconstitute the URL before sending in to the server travel call 
+			URL = (Options.IsEmpty() ? MapName : MapName + TEXT("?") + Options)
+			InWorld->ServerTravel(URL);
+			return true;
 		}
 		else
 		{
 			Ar.Logf(TEXT("ERROR: The map '%s' is either short package name or does not exist."), *MapName);
 		}
-		return true;
 	}
 	return false;
 }
