Compile fix

diff --git a/Engine/Source/Runtime/Engine/Private/UnrealEngine.cpp b/Engine/Source/Runtime/Engine/Private/UnrealEngine.cpp
--- a/Engine/Source/Runtime/Engine/Private/UnrealEngine.cpp
+++ b/Engine/Source/Runtime/Engine/Private/UnrealEngine.cpp
@@ -8464,7 +8464,7 @@ bool UEngine::HandleServerTravelCommand( const TCHAR* Cmd, FOutputDevice& Ar, UW
 		if (MakeSureMapNameIsValid(MapName))
 		{
 			// If there were options reconstitute the URL before sending in to the server travel call 
-			URL = (Options.IsEmpty() ? MapName : MapName + TEXT("?") + Options)
+			URL = (Options.IsEmpty() ? MapName : MapName + TEXT("?") + Options);
 			InWorld->ServerTravel(URL);
 			return true;
 		}
