Turn off deferred actors over a frame. This needs more profiling to be sure of the perf implications. For now we simply defer until insertion of static bodies is done.

diff --git a/Engine/Source/Runtime/Engine/Private/PhysicsEngine/BodyInstance.cpp b/Engine/Source/Runtime/Engine/Private/PhysicsEngine/BodyInstance.cpp
--- a/Engine/Source/Runtime/Engine/Private/PhysicsEngine/BodyInstance.cpp
+++ b/Engine/Source/Runtime/Engine/Private/PhysicsEngine/BodyInstance.cpp
@@ -1375,6 +1375,10 @@ struct FInitBodiesHelper
 
 			AddActorsToScene_PhysX_AssumesLocked(PSyncActors, PAsyncActors, PDynamicActors, bDynamicsUseAsync ? PAsyncScene : PSyncScene);
 		}
+
+		PhysScene->FlushDeferredActors();	//For now we do not actually defer over multiple frames. This needs better profiling to determine how useful it actually is.
+
+
 	}
 #endif
 
