PR #1358: int typo (Contributed by [REDACTED_PERSON])

diff --git a/Engine/Source/Runtime/Engine/Private/GameViewportClient.cpp b/Engine/Source/Runtime/Engine/Private/GameViewportClient.cpp
--- a/Engine/Source/Runtime/Engine/Private/GameViewportClient.cpp
+++ b/Engine/Source/Runtime/Engine/Private/GameViewportClient.cpp
@@ -886,7 +886,7 @@ void UGameViewportClient::Draw(FViewport* InViewport, FCanvas* SceneCanvas)
 			const bool bEnableStereo = GEngine->IsStereoscopic3D(InViewport);
 			int32 NumViews = bEnableStereo ? 2 : 1;
 
-			for (int i = 0; i < NumViews; ++i)
+			for (int32 i = 0; i < NumViews; ++i)
 			{
 				// Calculate the player's view information.
 				FVector		ViewLocation;
