Remove unneeded override

diff --git a/Engine/Source/Runtime/Engine/Classes/GameFramework/DefaultPhysicsVolume.h b/Engine/Source/Runtime/Engine/Classes/GameFramework/DefaultPhysicsVolume.h
--- a/Engine/Source/Runtime/Engine/Classes/GameFramework/DefaultPhysicsVolume.h
+++ b/Engine/Source/Runtime/Engine/Classes/GameFramework/DefaultPhysicsVolume.h
@@ -18,10 +18,6 @@ UCLASS(notplaceable, transient, MinimalAPI)
 class ADefaultPhysicsVolume : public APhysicsVolume
 {
 	GENERATED_UCLASS_BODY()
-
-	// Begin AActor interface
-	ENGINE_API virtual void Destroyed() override;
-	// End AActor interface
 };
 
 
diff --git a/Engine/Source/Runtime/Engine/Private/DefaultPhysicsVolume.cpp b/Engine/Source/Runtime/Engine/Private/DefaultPhysicsVolume.cpp
--- a/Engine/Source/Runtime/Engine/Private/DefaultPhysicsVolume.cpp
+++ b/Engine/Source/Runtime/Engine/Private/DefaultPhysicsVolume.cpp
@@ -16,9 +16,4 @@ ADefaultPhysicsVolume::ADefaultPhysicsVolume(const FObjectInitializer& ObjectIni
 	// update default values when world is restarted
 	TerminalVelocity = UPhysicsSettings::Get()->DefaultTerminalVelocity;
 	FluidFriction = UPhysicsSettings::Get()->DefaultFluidFriction;
-}
-
-void ADefaultPhysicsVolume::Destroyed()
-{
-	UE_LOG(LogVolume, Log, TEXT("%s destroyed!"), *GetName());
-}
+}
\ No newline at end of file
