Cleanup some WITH_EDITOR

diff --git a/Engine/Source/Runtime/Engine/Classes/Engine/Brush.h b/Engine/Source/Runtime/Engine/Classes/Engine/Brush.h
--- a/Engine/Source/Runtime/Engine/Classes/Engine/Brush.h
+++ b/Engine/Source/Runtime/Engine/Classes/Engine/Brush.h
@@ -139,10 +139,8 @@ public:
 public:
 	
 	// UObject interface.
-
-	virtual void PostLoad() override;
-
 #if WITH_EDITOR
+	virtual void PostLoad() override;
 	virtual void PostEditMove(bool bFinished) override;
 	virtual void PostEditChangeProperty(FPropertyChangedEvent& PropertyChangedEvent) override;
 #endif // WITH_EDITOR
@@ -162,11 +160,11 @@ public:
 public:
 	
 	// AActor interface
-	virtual void Destroyed() override;
-	virtual void PostRegisterAllComponents() override;
 	virtual bool IsLevelBoundsRelevant() const override;
 
 #if WITH_EDITOR
+	virtual void Destroyed() override;
+	virtual void PostRegisterAllComponents() override;
 	virtual void CheckForErrors() override;
 	virtual void SetIsTemporarilyHiddenInEditor( bool bIsHidden ) override;
 
@@ -215,10 +213,6 @@ public:
 	 * @param	InLevel The level that needs rebuilding
 	 */
 	static void SetNeedRebuild(ULevel* InLevel){if(InLevel){LevelsToRebuild.AddUnique(InLevel);}}
-#else
-	static bool NeedsRebuild(TArray< TWeakObjectPtr< ULevel > >* OutLevels = nullptr){return false;}
-	static void OnRebuildDone(){}
-	static void SetNeedRebuild(ULevel* InLevel){}
 #endif//WITH_EDITOR
 
 	/** @return true if this is a static brush */
diff --git a/Engine/Source/Runtime/Engine/Classes/Engine/CullDistanceVolume.h b/Engine/Source/Runtime/Engine/Classes/Engine/CullDistanceVolume.h
--- a/Engine/Source/Runtime/Engine/Classes/Engine/CullDistanceVolume.h
+++ b/Engine/Source/Runtime/Engine/Classes/Engine/CullDistanceVolume.h
@@ -61,14 +61,8 @@ public:
 #if WITH_EDITOR
 	virtual void PostEditChangeProperty(FPropertyChangedEvent& PropertyChangedEvent) override;
 	virtual void PostEditMove(bool bFinished) override;
-#endif // WITH_EDITOR
-
-public:
-
-	/** 
-	 * Override Destroyed so that we can re-calculate primitive draw distances after this volume has been deleted.
-	 */
 	virtual void Destroyed() override;
+#endif // WITH_EDITOR
 
 	/**
 	 * Returns whether the passed in primitive can be affected by cull distance volumes.
diff --git a/Engine/Source/Runtime/Engine/Private/Brush.cpp b/Engine/Source/Runtime/Engine/Private/Brush.cpp
--- a/Engine/Source/Runtime/Engine/Private/Brush.cpp
+++ b/Engine/Source/Runtime/Engine/Private/Brush.cpp
@@ -164,12 +164,11 @@ void ABrush::SetPrePivot( const FVector& InPrePivot )
 	}
 }
 
-
+#if WITH_EDITOR
 void ABrush::PostLoad()
 {
 	Super::PostLoad();
 
-#if WITH_EDITOR
 	if (BrushBuilder && BrushBuilder->GetOuter() != this)
 	{
 		BrushBuilder = DuplicateObject<UBrushBuilder>(BrushBuilder, this);
@@ -202,7 +201,6 @@ void ABrush::PostLoad()
 	{
 		UE_LOG(LogPhysics, Log, TEXT("%s does not have BrushBodySetup. No collision."), *GetName());
 	}
-#endif
 }
 
 void ABrush::Destroyed()
@@ -220,13 +218,12 @@ void ABrush::PostRegisterAllComponents()
 {
 	Super::PostRegisterAllComponents();
 
-#if WITH_EDITOR
 	if ( GIsEditor )
 	{
 		OnBrushRegistered.Broadcast(this);
 	}
-#endif
 }
+#endif
 
 bool ABrush::IsLevelBoundsRelevant() const
 {
diff --git a/Engine/Source/Runtime/Engine/Private/CullDistanceVolume.cpp b/Engine/Source/Runtime/Engine/Private/CullDistanceVolume.cpp
--- a/Engine/Source/Runtime/Engine/Private/CullDistanceVolume.cpp
+++ b/Engine/Source/Runtime/Engine/Private/CullDistanceVolume.cpp
@@ -16,19 +16,17 @@ ACullDistanceVolume::ACullDistanceVolume(const FObjectInitializer& ObjectInitial
 	bEnabled = true;
 }
 
+#if WITH_EDITOR
 void ACullDistanceVolume::Destroyed()
 {
 	Super::Destroyed();
 
-#if WITH_EDITOR
 	if (GetWorld())
 	{
 		GetWorld()->bDoDelayedUpdateCullDistanceVolumes = true;
 	}
-#endif
 }
 
-#if WITH_EDITOR
 void ACullDistanceVolume::PostEditChangeProperty(FPropertyChangedEvent& PropertyChangedEvent)
 {
 	Super::PostEditChangeProperty(PropertyChangedEvent);
diff --git a/Engine/Source/Runtime/Landscape/Private/Landscape.cpp b/Engine/Source/Runtime/Landscape/Private/Landscape.cpp
--- a/Engine/Source/Runtime/Landscape/Private/Landscape.cpp
+++ b/Engine/Source/Runtime/Landscape/Private/Landscape.cpp
@@ -1373,16 +1373,12 @@ void ALandscapeProxy::Destroyed()
 		SplineComponent->ModifySplines();
 	}
 
-#if WITH_EDITOR
 	TotalComponentsNeedingGrassMapRender -= NumComponentsNeedingGrassMapRender;
 	NumComponentsNeedingGrassMapRender = 0;
 	TotalTexturesToStreamForVisibleGrassMapRender -= NumTexturesToStreamForVisibleGrassMapRender;
 	NumTexturesToStreamForVisibleGrassMapRender = 0;
-#endif
 }
-#endif
 
-#if WITH_EDITOR
 void ALandscapeProxy::GetSharedProperties(ALandscapeProxy* Landscape)
 {
 	if (GIsEditor && Landscape)
