PR #1193: Do not ever attempt to create physics state after component passed the point when physics state could be destroyed. (Contributed by [REDACTED_PERSON])

diff --git a/Engine/Source/Runtime/Engine/Private/Components/PrimitiveComponent.cpp b/Engine/Source/Runtime/Engine/Private/Components/PrimitiveComponent.cpp
--- a/Engine/Source/Runtime/Engine/Private/Components/PrimitiveComponent.cpp
+++ b/Engine/Source/Runtime/Engine/Private/Components/PrimitiveComponent.cpp
@@ -893,6 +893,11 @@ bool UPrimitiveComponent::ShouldComponentAddToScene() const
 
 bool UPrimitiveComponent::ShouldCreatePhysicsState() const
 {
+	if (IsBeingDestroyed())
+	{
+		return false;
+	}
+
 	bool bShouldCreatePhysicsState = IsRegistered() && (bAlwaysCreatePhysicsState || IsCollisionEnabled());
 
 #if WITH_EDITOR
diff --git a/Engine/Source/Runtime/Engine/Private/InstancedStaticMesh.cpp b/Engine/Source/Runtime/Engine/Private/InstancedStaticMesh.cpp
--- a/Engine/Source/Runtime/Engine/Private/InstancedStaticMesh.cpp
+++ b/Engine/Source/Runtime/Engine/Private/InstancedStaticMesh.cpp
@@ -1483,7 +1483,7 @@ bool UInstancedStaticMeshComponent::UpdateInstanceTransform(int32 InstanceIndex,
 
 bool UInstancedStaticMeshComponent::ShouldCreatePhysicsState() const
 {
-	return IsRegistered() && (bAlwaysCreatePhysicsState || IsCollisionEnabled());
+	return IsRegistered() && !IsBeingDestroyed() && (bAlwaysCreatePhysicsState || IsCollisionEnabled());
 }
 
 void UInstancedStaticMeshComponent::ClearInstances()
diff --git a/Engine/Source/Runtime/Engine/Private/Vehicles/WheeledVehicleMovementComponent.cpp b/Engine/Source/Runtime/Engine/Private/Vehicles/WheeledVehicleMovementComponent.cpp
--- a/Engine/Source/Runtime/Engine/Private/Vehicles/WheeledVehicleMovementComponent.cpp
+++ b/Engine/Source/Runtime/Engine/Private/Vehicles/WheeledVehicleMovementComponent.cpp
@@ -658,6 +658,11 @@ void UWheeledVehicleMovementComponent::DestroyPhysicsState()
 
 bool UWheeledVehicleMovementComponent::ShouldCreatePhysicsState() const
 {
+	if (!IsRegistered() || IsBeingDestroyed())
+	{
+		return false;
+	}
+
 	// only create physx vehicle in game
 	if (GetWorld()->IsGameWorld())
 	{
