#ue4 - Fixed an issue with particle orbit not working when render in more than a single view.

-Moved PreviousOffset update code into the actual particle Update function called during Tick instead of on render per view.

diff --git a/Engine/Source/Runtime/Engine/Private/Particles/ParticleModules_Orbit.cpp b/Engine/Source/Runtime/Engine/Private/Particles/ParticleModules_Orbit.cpp
--- a/Engine/Source/Runtime/Engine/Private/Particles/ParticleModules_Orbit.cpp
+++ b/Engine/Source/Runtime/Engine/Private/Particles/ParticleModules_Orbit.cpp
@@ -124,6 +124,7 @@ void UParticleModuleOrbit::Spawn(FParticleEmitterInstance* Owner, int32 Offset,
 				LocalOffset = OffsetAmount.GetValue(Owner->EmitterTime, Owner->Component);
 			}
 			OrbitPayload.BaseOffset += LocalOffset;
+			OrbitPayload.PreviousOffset = OrbitPayload.Offset;
 			OrbitPayload.Offset += LocalOffset;
 		}
 
@@ -181,6 +182,7 @@ void UParticleModuleOrbit::Update(FParticleEmitterInstance* Owner, int32 Offset,
 
 			//@todo. Do we need to update the base offset here???
 //			OrbitPayload.BaseOffset += LocalOffset;
+			OrbitPayload.PreviousOffset = OrbitPayload.Offset;
 			OrbitPayload.Offset += LocalOffset;
 		}
 
diff --git a/Engine/Source/Runtime/Engine/Private/Particles/ParticleSystemRender.cpp b/Engine/Source/Runtime/Engine/Private/Particles/ParticleSystemRender.cpp
--- a/Engine/Source/Runtime/Engine/Private/Particles/ParticleSystemRender.cpp
+++ b/Engine/Source/Runtime/Engine/Private/Particles/ParticleSystemRender.cpp
@@ -531,9 +531,6 @@ void ApplyOrbitToPosition(
 	const FBaseParticle& Particle, 
 	const FDynamicSpriteEmitterReplayDataBase& Source, 
 	const FMatrix& InLocalToWorld,
-	FOrbitChainModuleInstancePayload*& LocalOrbitPayload,
-	FVector& OrbitOffset,
-	FVector& PrevOrbitOffset,
 	FVector& ParticlePosition,
 	FVector& ParticleOldPosition
 	)
@@ -543,18 +540,17 @@ void ApplyOrbitToPosition(
 		int32 CurrentOffset = Source.OrbitModuleOffset;
 		const uint8* ParticleBase = (const uint8*)&Particle;
 		PARTICLE_ELEMENT(FOrbitChainModuleInstancePayload, OrbitPayload);
-		OrbitOffset = OrbitPayload.Offset;
 
-		if (Source.bUseLocalSpace == false)
+		if (Source.bUseLocalSpace)
 		{
-			OrbitOffset = InLocalToWorld.TransformVector(OrbitOffset);
+			ParticlePosition += OrbitPayload.Offset;
+			ParticleOldPosition += OrbitPayload.PreviousOffset;
+		}
+		else
+		{
+			ParticlePosition += InLocalToWorld.TransformVector(OrbitPayload.Offset);
+			ParticleOldPosition += InLocalToWorld.TransformVector(OrbitPayload.PreviousOffset);
 		}
-		PrevOrbitOffset = OrbitPayload.PreviousOffset;
-
-		LocalOrbitPayload = &OrbitPayload;
-
-		ParticlePosition += OrbitOffset;
-		ParticleOldPosition += PrevOrbitOffset;
 	}
 }
 
@@ -590,8 +586,6 @@ bool FDynamicSpriteEmitterData::GetVertexAndIndexData(void* VertexData, void* Dy
 	FParticleSpriteVertex* FillVertex;
 	FParticleVertexDynamicParameter* DynFillVertex;
 
-	FVector OrbitOffset(0, 0, 0);
-	FVector PrevOrbitOffset(0.0f, 0.0f, 0.0f);
 	FVector4 DynamicParameterValue(1.0f,1.0f,1.0f,1.0f);
 	FVector ParticlePosition;
 	FVector ParticleOldPosition;
@@ -617,9 +611,7 @@ bool FDynamicSpriteEmitterData::GetVertexAndIndexData(void* VertexData, void* Dy
 		ParticlePosition = Particle.Location;
 		ParticleOldPosition = Particle.OldLocation;
 
-		FOrbitChainModuleInstancePayload* LocalOrbitPayload = NULL;
-
-		ApplyOrbitToPosition(Particle, Source, InLocalToWorld, LocalOrbitPayload, OrbitOffset, PrevOrbitOffset, ParticlePosition, ParticleOldPosition);
+		ApplyOrbitToPosition(Particle, Source, InLocalToWorld, ParticlePosition, ParticleOldPosition);
 
 		if (Source.CameraPayloadOffset != 0)
 		{
@@ -660,11 +652,6 @@ bool FDynamicSpriteEmitterData::GetVertexAndIndexData(void* VertexData, void* Dy
 		}
 
 		TempVert += VertexStride;
-
-		if (LocalOrbitPayload)
-		{
-			LocalOrbitPayload->PreviousOffset = OrbitOffset;
-		}
 	}
 
 	return true;
@@ -704,8 +691,6 @@ bool FDynamicSpriteEmitterData::GetVertexAndIndexDataNonInstanced(void* VertexDa
 	FParticleSpriteVertexNonInstanced* FillVertex;
 	FParticleVertexDynamicParameter* DynFillVertex;
 
-	FVector OrbitOffset(0, 0, 0);
-	FVector PrevOrbitOffset(0.0f, 0.0f, 0.0f);
 	FVector4 DynamicParameterValue(1.0f,1.0f,1.0f,1.0f);
 	FVector ParticlePosition;
 	FVector ParticleOldPosition;
@@ -731,9 +716,7 @@ bool FDynamicSpriteEmitterData::GetVertexAndIndexDataNonInstanced(void* VertexDa
 		ParticlePosition = Particle.Location;
 		ParticleOldPosition = Particle.OldLocation;
 
-		FOrbitChainModuleInstancePayload* LocalOrbitPayload = NULL;
-
-		ApplyOrbitToPosition(Particle, Source, InLocalToWorld, LocalOrbitPayload, OrbitOffset, PrevOrbitOffset, ParticlePosition, ParticleOldPosition);
+		ApplyOrbitToPosition(Particle, Source, InLocalToWorld, ParticlePosition, ParticleOldPosition);
 
 		if (Source.CameraPayloadOffset != 0)
 		{
@@ -797,11 +780,6 @@ bool FDynamicSpriteEmitterData::GetVertexAndIndexDataNonInstanced(void* VertexDa
 			TempDynamicParameterVert += VertexDynamicParameterStride;
 		}
 		TempVert += VertexStride;
-
-		if (LocalOrbitPayload)
-		{
-			LocalOrbitPayload->PreviousOffset = OrbitOffset;
-		}
 	}
 
 	return true;
@@ -866,12 +844,9 @@ void GatherParticleLightData(const FDynamicSpriteEmitterReplayDataBase& Source,
 				}
 
 				FVector ParticlePosition = Particle.Location;
-				FOrbitChainModuleInstancePayload* LocalOrbitPayload = NULL;
 				FVector Unused(0, 0, 0);
-				FVector Unused2(0, 0, 0);
-				FVector Unused3(0, 0, 0);
 
-				ApplyOrbitToPosition(Particle, Source, InLocalToWorld, LocalOrbitPayload, Unused, Unused2, ParticlePosition, Unused3);
+				ApplyOrbitToPosition(Particle, Source, InLocalToWorld, ParticlePosition, Unused);
 
 				FVector LightPosition = Source.bUseLocalSpace ? FVector(InLocalToWorld.TransformPosition(ParticlePosition)) : ParticlePosition;
 				
