UE-17938 Adding ability to set a volume headroom per platform to better achieve platform consistency with respect to volume.

- Added stubs per platform in their ini file that doesn't change volume, except for mac which this bug is about and for which I reduced the volume by -6 dB. This can be changed by designers as per needed for their game per platform. The final headroom should be set by designers.
- Removed the hard-coded headroom constant used in OpenAL recently added.

diff --git a/Engine/Config/Android/AndroidEngine.ini b/Engine/Config/Android/AndroidEngine.ini
--- a/Engine/Config/Android/AndroidEngine.ini
+++ b/Engine/Config/Android/AndroidEngine.ini
@@ -1,5 +1,7 @@
 [Audio]
 AudioDeviceModuleName=AndroidAudio
+; Defines a platform-specific volume headroom (in dB) for audio to provide better platform consistency with respect to volume levels.
+PlatformHeadroomDB=0
 
 [DeviceProfileManager]
 DeviceProfileSelectionModule="AndroidDeviceProfileSelector"
diff --git a/Engine/Config/HTML5/HTML5Engine.ini b/Engine/Config/HTML5/HTML5Engine.ini
--- a/Engine/Config/HTML5/HTML5Engine.ini
+++ b/Engine/Config/HTML5/HTML5Engine.ini
@@ -1,5 +1,7 @@
 [Audio]
 AudioDeviceModuleName=ALAudio
+; Defines a platform-specific volume headroom (in dB) for audio to provide better platform consistency with respect to volume levels.
+PlatformHeadroomDB=0
 
 [OnlineSubsystemMcp]
 benabled=false
diff --git a/Engine/Config/IOS/IOSEngine.ini b/Engine/Config/IOS/IOSEngine.ini
--- a/Engine/Config/IOS/IOSEngine.ini
+++ b/Engine/Config/IOS/IOSEngine.ini
@@ -1,5 +1,7 @@
 [Audio]
 AudioDeviceModuleName=IOSAudio
+; Defines a platform-specific volume headroom (in dB) for audio to provide better platform consistency with respect to volume levels.
+PlatformHeadroomDB=0
 
 [OnlineSubsystem]
 DefaultPlatformService=IOS
diff --git a/Engine/Config/Linux/LinuxEngine.ini b/Engine/Config/Linux/LinuxEngine.ini
--- a/Engine/Config/Linux/LinuxEngine.ini
+++ b/Engine/Config/Linux/LinuxEngine.ini
@@ -1,5 +1,7 @@
 [Audio]
 AudioDeviceModuleName=ALAudio
+; Defines a platform-specific volume headroom (in dB) for audio to provide better platform consistency with respect to volume levels.
+PlatformHeadroomDB=0
 
 [SystemSettings]
 r.setres=1280x720
diff --git a/Engine/Config/Mac/MacEngine.ini b/Engine/Config/Mac/MacEngine.ini
--- a/Engine/Config/Mac/MacEngine.ini
+++ b/Engine/Config/Mac/MacEngine.ini
@@ -1,5 +1,7 @@
 [Audio]
 AudioDeviceModuleName=CoreAudio
+; Defines a platform-specific volume headroom (in dB) for audio to provide better platform consistency with respect to volume levels.
+PlatformHeadroomDB=-6
 
 [DerivedDataBackendGraph]
 Shared=(Type=FileSystem, ReadOnly=false, Clean=false, Flush=false, DeleteUnused=true, UnusedFileAge=38, FoldersToClean=-1, Path=/Volumes/UE4DDC, EnvPathOverride=UE-SharedDataCachePath)
diff --git a/Engine/Config/Windows/WindowsEngine.ini b/Engine/Config/Windows/WindowsEngine.ini
--- a/Engine/Config/Windows/WindowsEngine.ini
+++ b/Engine/Config/Windows/WindowsEngine.ini
@@ -1,5 +1,7 @@
 [Audio]
 AudioDeviceModuleName=XAudio2
+; Defines a platform-specific volume headroom (in dB) for audio to provide better platform consistency with respect to volume levels.
+PlatformHeadroomDB=0
 
 [TextureStreaming]
 ; PoolSizeVRAMPercentage is how much percentage of GPU Dedicated VRAM should be used as a TexturePool cache for streaming textures (0 - unlimited streaming)
diff --git a/Engine/Source/Runtime/Engine/Private/AudioDevice.cpp b/Engine/Source/Runtime/Engine/Private/AudioDevice.cpp
--- a/Engine/Source/Runtime/Engine/Private/AudioDevice.cpp
+++ b/Engine/Source/Runtime/Engine/Private/AudioDevice.cpp
@@ -32,6 +32,7 @@ FAudioDevice::FAudioDevice()
 	, DefaultBaseSoundMix(nullptr)
 	, Effects(nullptr)
 	, CurrentAudioVolume(nullptr)
+	, PlatformAudioHeadroom(1.0f)
 	, HighestPriorityReverb(nullptr)
 	, SpatializationPlugin(nullptr)
 	, SpatializeProcessor(nullptr)
@@ -59,10 +60,18 @@ bool FAudioDevice::Init(int32 InMaxChannels)
 	// initialize config variables
 	MaxChannels = InMaxChannels;
 	verify(GConfig->GetInt(TEXT("Audio"), TEXT("CommonAudioPoolSize"), CommonAudioPoolSize, GEngineIni));
-	
+
 	// If this is true, skip the initial startup precache so we can do it later in the flow
 	GConfig->GetBool(TEXT("Audio"), TEXT("DeferStartupPrecache"), bDeferStartupPrecache, GEngineIni);
 
+	// Get an optional engine ini setting for platform headroom. 
+	float Headroom = 1.0f;
+	if (GConfig->GetFloat(TEXT("Audio"), TEXT("PlatformHeadroomDB"), Headroom, GEngineIni))
+	{
+		// Convert dB to linear volume
+		PlatformAudioHeadroom = FMath::Pow(10.0f, Headroom / 10.0f);
+	}
+
 	const FStringAssetReference DefaultBaseSoundMixName = GetDefault<UAudioSettings>()->DefaultBaseSoundMix;
 	if (DefaultBaseSoundMixName.IsValid())
 	{
diff --git a/Engine/Source/Runtime/Engine/Public/AudioDevice.h b/Engine/Source/Runtime/Engine/Public/AudioDevice.h
--- a/Engine/Source/Runtime/Engine/Public/AudioDevice.h
+++ b/Engine/Source/Runtime/Engine/Public/AudioDevice.h
@@ -784,6 +784,9 @@ public:
 	/** The volume the listener resides in */
 	const class AAudioVolume*						CurrentAudioVolume;
 
+	/** A volume headroom to apply to specific platforms to achieve beter platform consistency. */
+	float PlatformAudioHeadroom;
+
 	/** Reverb Effects activated without volumes */
 	TMap<FName, FActivatedReverb>					ActivatedReverbs;
 
diff --git a/Engine/Source/Runtime/IOS/IOSAudio/Private/IOSAudioSource.cpp b/Engine/Source/Runtime/IOS/IOSAudio/Private/IOSAudioSource.cpp
--- a/Engine/Source/Runtime/IOS/IOSAudio/Private/IOSAudioSource.cpp
+++ b/Engine/Source/Runtime/IOS/IOSAudio/Private/IOSAudioSource.cpp
@@ -257,7 +257,7 @@ void FIOSAudioSoundSource::Update(void)
 	}
 
 	// Apply global multiplier to disable sound when not the foreground app
-	Volume = FMath::Clamp(Volume, 0.0f, MAX_VOLUME) * FApp::GetVolumeMultiplier();
+	Volume = FMath::Clamp(Volume, 0.0f, MAX_VOLUME) * FApp::GetVolumeMultiplier() * AudioDevice->PlatformAudioHeadroom;
 
 	// Convert to dB
 	const AudioUnitParameterValue Gain = FMath::Clamp<float>(20.0f * log10(Volume), -100, 20.0f);
diff --git a/Engine/Source/Runtime/Linux/ALAudio/Private/ALAudioSource.cpp b/Engine/Source/Runtime/Linux/ALAudio/Private/ALAudioSource.cpp
--- a/Engine/Source/Runtime/Linux/ALAudio/Private/ALAudioSource.cpp
+++ b/Engine/Source/Runtime/Linux/ALAudio/Private/ALAudioSource.cpp
@@ -100,7 +100,7 @@ void FALSoundSource::Update( void )
 			Volume *= 1.25f;
 		}
 		Volume *= FApp::GetVolumeMultiplier();
-		Volume *= AUDIO_OPEN_AL_VOLUME_SCALE;
+		Volume *= AudioDevice->PlatformAudioHeadroom;
 		Volume = FMath::Clamp(Volume, 0.0f, MAX_VOLUME);
 	}
 
diff --git a/Engine/Source/Runtime/Linux/ALAudio/Public/ALAudioDevice.h b/Engine/Source/Runtime/Linux/ALAudio/Public/ALAudioDevice.h
--- a/Engine/Source/Runtime/Linux/ALAudio/Public/ALAudioDevice.h
+++ b/Engine/Source/Runtime/Linux/ALAudio/Public/ALAudioDevice.h
@@ -33,10 +33,6 @@ DECLARE_LOG_CATEGORY_EXTERN(LogALAudio, Log, All);
 // <=> 1 UU == 0.0127 m
 #define AUDIO_DISTANCE_FACTOR ( 0.0127f )
 
-// Audio volume is higher in OpenAL than XAudio2, we'll add a -3dB attenuation as headroom
-// to attempt to make it more on-par.
-#define AUDIO_OPEN_AL_VOLUME_SCALE (0.707f)
-
 /**
  * OpenAL implementation of FSoundBuffer, containing the wave data and format information.
  */
diff --git a/Engine/Source/Runtime/Mac/CoreAudio/Private/CoreAudioSource.cpp b/Engine/Source/Runtime/Mac/CoreAudio/Private/CoreAudioSource.cpp
--- a/Engine/Source/Runtime/Mac/CoreAudio/Private/CoreAudioSource.cpp
+++ b/Engine/Source/Runtime/Mac/CoreAudio/Private/CoreAudioSource.cpp
@@ -298,7 +298,9 @@ void FCoreAudioSoundSource::Update( void )
 	{
 		Volume = WaveInstance->GetActualVolume();
 	}
-	
+
+	Volume *= AudioDevice->PlatformAudioHeadroom;
+
 	if( Buffer->NumChannels < 3 )
 	{
 		float Azimuth = 0.0f;
diff --git a/Engine/Source/Runtime/Windows/XAudio2/Private/XAudio2Source.cpp b/Engine/Source/Runtime/Windows/XAudio2/Private/XAudio2Source.cpp
--- a/Engine/Source/Runtime/Windows/XAudio2/Private/XAudio2Source.cpp
+++ b/Engine/Source/Runtime/Windows/XAudio2/Private/XAudio2Source.cpp
@@ -572,7 +572,7 @@ void FXAudio2SoundSource::GetChannelVolumes( float ChannelVolumes[CHANNELOUT_COU
 			UE_LOG(LogXAudio2, Warning, TEXT("FXAudio2SoundSource contains unreasonble value %f in channel %d: %s"), ChannelVolumes[i], i, *Describe_Internal(true, false));
 		}
 
-		ChannelVolumes[i] = FMath::Clamp<float>(ChannelVolumes[i] * FApp::GetVolumeMultiplier(), 0.0f, MAX_VOLUME);
+		ChannelVolumes[i] = FMath::Clamp<float>(ChannelVolumes[i] * FApp::GetVolumeMultiplier() * AudioDevice->PlatformAudioHeadroom, 0.0f, MAX_VOLUME);
 	}
 }
 
