Make collision profile not abstract. This is needed for developer settings to show up.

diff --git a/Engine/Source/Runtime/Engine/Classes/Engine/CollisionProfile.h b/Engine/Source/Runtime/Engine/Classes/Engine/CollisionProfile.h
--- a/Engine/Source/Runtime/Engine/Classes/Engine/CollisionProfile.h
+++ b/Engine/Source/Runtime/Engine/Classes/Engine/CollisionProfile.h
@@ -138,7 +138,7 @@ struct ENGINE_API FCustomProfile
 /**
  * Set up and modify collision settings.
  */
-UCLASS(abstract, config=Engine, defaultconfig, MinimalAPI, meta=(DisplayName="Collision"))
+UCLASS(config=Engine, defaultconfig, MinimalAPI, meta=(DisplayName="Collision"))
 class UCollisionProfile : public UDeveloperSettings
 {
 	GENERATED_UCLASS_BODY()
