Fixed issue where editing a Material Parameter Collection results in marking an unnecessarily large number of materials as dirty.

diff --git a/Engine/Source/Runtime/Engine/Classes/Materials/Material.h b/Engine/Source/Runtime/Engine/Classes/Materials/Material.h
--- a/Engine/Source/Runtime/Engine/Classes/Materials/Material.h
+++ b/Engine/Source/Runtime/Engine/Classes/Materials/Material.h
@@ -1000,6 +1000,42 @@ public:
 		}
 	}
 
+	/** Get all expressions of the requested type, recursing through any function expressions in the material */
+	template<typename ExpressionType>
+	void GetAllExpressionsInMaterialAndFunctionsOfType(TArray<ExpressionType*>& OutExpressions) const
+	{
+		for (UMaterialExpression* Expression : Expressions)
+		{
+			ExpressionType* ExpressionOfType = Cast<ExpressionType>(Expression);
+			if (ExpressionOfType)
+			{
+				OutExpressions.Add(ExpressionOfType);
+			}
+
+			UMaterialExpressionMaterialFunctionCall* ExpressionFunctionCall = Cast<UMaterialExpressionMaterialFunctionCall>(Expression);
+			if (ExpressionFunctionCall && ExpressionFunctionCall->MaterialFunction)
+			{
+				TArray<UMaterialFunction*> Functions;
+				Functions.Add(ExpressionFunctionCall->MaterialFunction);
+
+				ExpressionFunctionCall->MaterialFunction->GetDependentFunctions(Functions);
+
+				// Handle nested functions
+				for (UMaterialFunction* Function : Functions)
+				{
+					for (UMaterialExpression* Expression : Function->FunctionExpressions)
+					{
+						ExpressionType* ExpressionOfType = Cast<ExpressionType>(Expression);
+						if (ExpressionOfType)
+						{
+							OutExpressions.Add(ExpressionOfType);
+						}
+					}
+				}
+			}
+		}
+	}
+
 	/** Determines whether each quality level has different nodes by inspecting the material's expressions. */
 	void GetQualityLevelNodeUsage(TArray<bool, TInlineAllocator<EMaterialQualityLevel::Num> >& QualityLevelsUsed);
 
diff --git a/Engine/Source/Runtime/Engine/Private/Materials/ParameterCollection.cpp b/Engine/Source/Runtime/Engine/Private/Materials/ParameterCollection.cpp
--- a/Engine/Source/Runtime/Engine/Private/Materials/ParameterCollection.cpp
+++ b/Engine/Source/Runtime/Engine/Private/Materials/ParameterCollection.cpp
@@ -4,6 +4,7 @@
 #include "ParameterCollection.h"
 #include "Materials/MaterialParameterCollectionInstance.h"
 #include "Materials/MaterialParameterCollection.h"
+#include "Materials/MaterialExpressionCollectionParameter.h"
 
 UMaterialParameterCollection::UMaterialParameterCollection(const FObjectInitializer& ObjectInitializer)
 	: Super(ObjectInitializer)
@@ -30,23 +31,23 @@ void UMaterialParameterCollection::PostLoad()
 
 #if WITH_EDITOR
 
-int32 PreviousScalarCount = 0;
-int32 PreviousVectorCount = 0;
+TArray<FCollectionScalarParameter> PreviousScalarParameters;
+TArray<FCollectionVectorParameter> PreviousVectorParameters;
 
 void UMaterialParameterCollection::PreEditChange(class FEditPropertyChain& PropertyAboutToChange)
 {
 	Super::PreEditChange(PropertyAboutToChange);
 
-	PreviousScalarCount = ScalarParameters.Num();
-	PreviousVectorCount = VectorParameters.Num();
+	PreviousScalarParameters = ScalarParameters;
+	PreviousVectorParameters = VectorParameters;
 }
 
 void UMaterialParameterCollection::PostEditChangeProperty(FPropertyChangedEvent& PropertyChangedEvent)
 {
 	// If the array counts have changed, an element has been added or removed, and we need to update the uniform buffer layout,
 	// Which also requires recompiling any referencing materials
-	if (ScalarParameters.Num() != PreviousScalarCount
-		|| VectorParameters.Num() != PreviousVectorCount)
+	if (ScalarParameters.Num() != PreviousScalarParameters.Num()
+		|| VectorParameters.Num() != PreviousVectorParameters.Num())
 	{
 		// Limit the count of parameters to fit within uniform buffer limits
 		const uint32 MaxScalarParameters = 1024;
@@ -76,6 +77,28 @@ void UMaterialParameterCollection::PostEditChangeProperty(FPropertyChangedEvent&
 			CurrentWorld->AddParameterCollectionInstance(this, false);
 		}
 
+		// Build set of changed parameter names
+		TSet<FName> ParameterNames;
+		for (const FCollectionVectorParameter& Param : PreviousVectorParameters)
+		{
+			ParameterNames.Add(Param.ParameterName);
+		}
+
+		for (const FCollectionScalarParameter& Param : PreviousScalarParameters)
+		{
+			ParameterNames.Add(Param.ParameterName);
+		}
+
+		for (const FCollectionVectorParameter& Param : VectorParameters)
+		{
+			ParameterNames.Remove(Param.ParameterName);
+		}
+
+		for (const FCollectionScalarParameter& Param : ScalarParameters)
+		{
+			ParameterNames.Remove(Param.ParameterName);
+		}
+
 		// Create a material update context so we can safely update materials using this parameter collection.
 		{
 			FMaterialUpdateContext UpdateContext;
@@ -95,12 +118,20 @@ void UMaterialParameterCollection::PostEditChangeProperty(FPropertyChangedEvent&
 				}
 				else
 				{
-					for (int32 FunctionIndex = 0; FunctionIndex < CurrentMaterial->MaterialParameterCollectionInfos.Num(); FunctionIndex++)
+					for (int32 FunctionIndex = 0; FunctionIndex < CurrentMaterial->MaterialParameterCollectionInfos.Num() && !bRecompile; FunctionIndex++)
 					{
 						if (CurrentMaterial->MaterialParameterCollectionInfos[FunctionIndex].ParameterCollection == this)
 						{
-							bRecompile = true;
-							break;
+							TArray<UMaterialExpressionCollectionParameter*> CollectionParameters;
+							CurrentMaterial->GetAllExpressionsInMaterialAndFunctionsOfType(CollectionParameters);
+							for (UMaterialExpressionCollectionParameter* CollectionParameter : CollectionParameters)
+							{
+								if (ParameterNames.Contains(CollectionParameter->ParameterName))
+								{
+									bRecompile = true;
+									break;
+								}
+							}
 						}
 					}
 				}
@@ -125,6 +156,9 @@ void UMaterialParameterCollection::PostEditChangeProperty(FPropertyChangedEvent&
 		CurrentWorld->UpdateParameterCollectionInstances(true);
 	}
 
+	PreviousScalarParameters.Empty();
+	PreviousVectorParameters.Empty();
+
 	Super::PostEditChangeProperty(PropertyChangedEvent);
 }
 
@@ -160,7 +194,7 @@ inline void CreateUniqueName(const TCHAR* InBaseName, TArray<T>& InExistingItems
 void UMaterialParameterCollection::PostEditChangeChainProperty(struct FPropertyChangedChainEvent& PropertyChangedEvent)
 {
 	// Auto-populate with useful parameter names
-	if (ScalarParameters.Num() > PreviousScalarCount)
+	if (ScalarParameters.Num() > PreviousScalarParameters.Num())
 	{
 		int32 NewArrayIndex = PropertyChangedEvent.GetArrayIndex("ScalarParameters");
 
@@ -170,7 +204,7 @@ void UMaterialParameterCollection::PostEditChangeChainProperty(struct FPropertyC
 		}
 	}
 
-	if (VectorParameters.Num() > PreviousVectorCount)
+	if (VectorParameters.Num() > PreviousVectorParameters.Num())
 	{
 		int32 NewArrayIndex = PropertyChangedEvent.GetArrayIndex("VectorParameters");
 
