Fixes to non-editor builds.

diff --git a/Engine/Source/Runtime/Engine/Private/Actor.cpp b/Engine/Source/Runtime/Engine/Private/Actor.cpp
--- a/Engine/Source/Runtime/Engine/Private/Actor.cpp
+++ b/Engine/Source/Runtime/Engine/Private/Actor.cpp
@@ -93,7 +93,9 @@ void AActor::InitializeDefaults()
 	bFindCameraComponentWhenViewTarget = true;
 	bAllowReceiveTickEventOnDedicatedServer = true;
 	bRelevantForNetworkReplays = true;
+#if WITH_EDITORONLY_DATA
 	PivotOffset = FVector::ZeroVector;
+#endif
 }
 
 void FActorTickFunction::ExecuteTick(float DeltaTime, enum ELevelTick TickType, ENamedThreads::Type CurrentThread, const FGraphEventRef& MyCompletionGraphEvent)
diff --git a/Engine/Source/Runtime/Engine/Private/Brush.cpp b/Engine/Source/Runtime/Engine/Private/Brush.cpp
--- a/Engine/Source/Runtime/Engine/Private/Brush.cpp
+++ b/Engine/Source/Runtime/Engine/Private/Brush.cpp
@@ -143,8 +143,6 @@ void ABrush::SetIsTemporarilyHiddenInEditor( bool bIsHidden )
 	}
 }
 
-#endif
-
 FVector ABrush::GetPrePivot() const
 {
 	return GetPivotOffset();
@@ -155,7 +153,6 @@ void ABrush::SetPrePivot( const FVector& InPrePivot )
 	SetPivotOffset(InPrePivot);
 }
 
-#if WITH_EDITOR
 void ABrush::PostLoad()
 {
 	Super::PostLoad();
diff --git a/Engine/Source/Runtime/Engine/Private/Components/BrushComponent.cpp b/Engine/Source/Runtime/Engine/Private/Components/BrushComponent.cpp
--- a/Engine/Source/Runtime/Engine/Private/Components/BrushComponent.cpp
+++ b/Engine/Source/Runtime/Engine/Private/Components/BrushComponent.cpp
@@ -545,11 +545,13 @@ void UBrushComponent::PostLoad()
 		BrushBodySetup->bGenerateMirroredCollision = false;
 	}
 
+#if WITH_EDITOR
 	AActor* Owner = GetOwner();
 	if (Owner)
 	{
 		Owner->SetPivotOffset(PrePivot_DEPRECATED);
 	}
+#endif
 }
 
 
