Fixed shadowed variable warning.

diff --git a/Engine/Source/Runtime/Engine/Classes/Materials/Material.h b/Engine/Source/Runtime/Engine/Classes/Materials/Material.h
--- a/Engine/Source/Runtime/Engine/Classes/Materials/Material.h
+++ b/Engine/Source/Runtime/Engine/Classes/Materials/Material.h
@@ -1023,12 +1023,12 @@ public:
 				// Handle nested functions
 				for (UMaterialFunction* Function : Functions)
 				{
-					for (UMaterialExpression* Expression : Function->FunctionExpressions)
+					for (UMaterialExpression* FunctionExpression : Function->FunctionExpressions)
 					{
-						ExpressionType* ExpressionOfType = Cast<ExpressionType>(Expression);
-						if (ExpressionOfType)
+						ExpressionType* FunctionExpressionOfType = Cast<ExpressionType>(FunctionExpression);
+						if (FunctionExpressionOfType)
 						{
-							OutExpressions.Add(ExpressionOfType);
+							OutExpressions.Add(FunctionExpressionOfType);
 						}
 					}
 				}
