Label worker thread anim stats as such

diff --git a/Engine/Source/Runtime/Engine/Classes/Animation/AnimInstance.h b/Engine/Source/Runtime/Engine/Classes/Animation/AnimInstance.h
--- a/Engine/Source/Runtime/Engine/Classes/Animation/AnimInstance.h
+++ b/Engine/Source/Runtime/Engine/Classes/Animation/AnimInstance.h
@@ -390,6 +390,9 @@ public:
 	UFUNCTION(BlueprintCallable, Category = "Animation")
 	virtual APawn* TryGetPawnOwner() const;
 
+	// Are we being evaluated on a worker thread
+	bool IsRunningParallelEvaluation() const;
+
 public:
 	/** Returns the owning actor of this AnimInstance */
 	UFUNCTION(BlueprintCallable, Category = "Animation")
diff --git a/Engine/Source/Runtime/Engine/Classes/Components/PoseableMeshComponent.h b/Engine/Source/Runtime/Engine/Classes/Components/PoseableMeshComponent.h
--- a/Engine/Source/Runtime/Engine/Classes/Components/PoseableMeshComponent.h
+++ b/Engine/Source/Runtime/Engine/Classes/Components/PoseableMeshComponent.h
@@ -49,6 +49,7 @@ class ENGINE_API UPoseableMeshComponent : public USkinnedMeshComponent
 	virtual bool AllocateTransformData() override;
 	// End USkinnedMeshComponent Interface
 
+	bool IsRunningParallelEvaluation() const { return false; }
 	/**
 	 * Take the LocalAtoms array (translation vector, rotation quaternion and scale vector) and update the array of component-space bone transformation matrices (SpaceBases).
 	 * It will work down hierarchy multiplying the component-space transform of the parent by the relative transform of the child.
diff --git a/Engine/Source/Runtime/Engine/Classes/Components/SkeletalMeshComponent.h b/Engine/Source/Runtime/Engine/Classes/Components/SkeletalMeshComponent.h
--- a/Engine/Source/Runtime/Engine/Classes/Components/SkeletalMeshComponent.h
+++ b/Engine/Source/Runtime/Engine/Classes/Components/SkeletalMeshComponent.h
@@ -1256,7 +1256,14 @@ public:
 		AnimEvaluationContext.Clear();
 	}
 
-	bool IsRunningParallelEvaluation(bool bBlockOnTask, bool bPerformPostAnimEvaluation);
+	// Returns whether we are currently trying to run a parallel animation evaluation task
+	bool IsRunningParallelEvaluation() const { return IsValidRef(ParallelAnimationEvaluationTask); }
+
+	// Management function for if we want to do an evaluation but may already be running one
+	// bBlockOnTask - if true and we are currently performing parallel eval we wait for it to finish
+	// bPerformPostAnimEvaluation - if true and we are currently performing parallel eval we call PostAnimEvaluation too
+	// return true if parallel task was running.
+	bool HandleExistingParallelEvaluationTask(bool bBlockOnTask, bool bPerformPostAnimEvaluation);
 
 	friend class FSkeletalMeshComponentDetails;
 
diff --git a/Engine/Source/Runtime/Engine/Private/Animation/AnimInstance.cpp b/Engine/Source/Runtime/Engine/Private/Animation/AnimInstance.cpp
--- a/Engine/Source/Runtime/Engine/Private/Animation/AnimInstance.cpp
+++ b/Engine/Source/Runtime/Engine/Private/Animation/AnimInstance.cpp
@@ -34,19 +34,13 @@ DEFINE_STAT(STAT_SkelCompUpdateTransform);
 DEFINE_STAT(STAT_UpdateRBBones);
 DEFINE_STAT(STAT_UpdateRBJoints);
 DEFINE_STAT(STAT_UpdateLocalToWorldAndOverlaps);
-DEFINE_STAT(STAT_SkelComposeTime);
 DEFINE_STAT(STAT_GetAnimationPose);
-DEFINE_STAT(STAT_AnimNativeEvaluatePoses);
 DEFINE_STAT(STAT_AnimTriggerAnimNotifies);
-DEFINE_STAT(STAT_AnimNativeBlendPoses);
-DEFINE_STAT(STAT_AnimGraphEvaluate);
-DEFINE_STAT(STAT_AnimBlendTime);
 DEFINE_STAT(STAT_RefreshBoneTransforms);
 DEFINE_STAT(STAT_InterpolateSkippedFrames);
 DEFINE_STAT(STAT_AnimTickTime);
 DEFINE_STAT(STAT_SkinnedMeshCompTick);
 DEFINE_STAT(STAT_TickUpdateRate);
-DEFINE_STAT(STAT_PerformAnimEvaluation);
 DEFINE_STAT(STAT_BlueprintUpdateAnimation);
 DEFINE_STAT(STAT_BlueprintPostEvaluateAnimation);
 DEFINE_STAT(STAT_NativeUpdateAnimation);
@@ -56,13 +50,21 @@ DEFINE_STAT(STAT_AnimInitTime);
 
 DEFINE_STAT(STAT_AnimStateMachineUpdate);
 DEFINE_STAT(STAT_AnimStateMachineFindTransition);
-DEFINE_STAT(STAT_AnimStateMachineEvaluate);
 
 DEFINE_STAT(STAT_SkinPerPolyVertices);
 DEFINE_STAT(STAT_UpdateTriMeshVertices);
 
 DEFINE_STAT(STAT_AnimGameThreadTime);
 
+
+#define DO_ANIMSTAT_PROCESSING(StatName) DEFINE_STAT(STAT_ ## StatName)
+#include "AnimMTStats.h"
+#undef DO_ANIMSTAT_PROCESSING
+
+#define DO_ANIMSTAT_PROCESSING(StatName) DEFINE_STAT(STAT_ ## StatName ## _WorkerThread)
+#include "AnimMTStats.h"
+#undef DO_ANIMSTAT_PROCESSING
+
 // Define AnimNotify
 DEFINE_LOG_CATEGORY(LogAnimNotify);
 
@@ -453,7 +455,7 @@ void UAnimInstance::UpdateAnimation(float DeltaSeconds)
 
 void UAnimInstance::EvaluateAnimation(FPoseContext& Output)
 {
-	SCOPE_CYCLE_COUNTER(STAT_AnimNativeEvaluatePoses);
+	ANIM_MT_SCOPE_CYCLE_COUNTER(EvaluateAnimInstance, IsRunningParallelEvaluation());
 
 	// If bone caches have been invalidated, have AnimNodes refresh those.
 	if( bBoneCachesInvalidated && RootNode )
@@ -470,7 +472,7 @@ void UAnimInstance::EvaluateAnimation(FPoseContext& Output)
 	{
 		if (RootNode != NULL)
 		{
-			SCOPE_CYCLE_COUNTER(STAT_AnimGraphEvaluate);
+			ANIM_MT_SCOPE_CYCLE_COUNTER(EvaluateAnimGraph, IsRunningParallelEvaluation());
 			IncrementGraphTraversalCounter();
 			RootNode->Evaluate(Output);
 		}
@@ -2909,4 +2911,14 @@ FAnimNode_AssetPlayerBase* UAnimInstance::GetRelevantAssetPlayerFromState(int32
 	return ResultPlayer;
 }
 
+bool UAnimInstance::IsRunningParallelEvaluation() const
+{
+	USkeletalMeshComponent* Comp = GetOwningComponent();
+	if (Comp)
+	{
+		return Comp->IsRunningParallelEvaluation();
+	}
+	return false;
+}
+
 #undef LOCTEXT_NAMESPACE 
diff --git a/Engine/Source/Runtime/Engine/Private/Animation/AnimNode_BlendListBase.cpp b/Engine/Source/Runtime/Engine/Private/Animation/AnimNode_BlendListBase.cpp
--- a/Engine/Source/Runtime/Engine/Private/Animation/AnimNode_BlendListBase.cpp
+++ b/Engine/Source/Runtime/Engine/Private/Animation/AnimNode_BlendListBase.cpp
@@ -142,7 +142,7 @@ void FAnimNode_BlendListBase::Update(const FAnimationUpdateContext& Context)
 
 void FAnimNode_BlendListBase::Evaluate(FPoseContext& Output)
 {
-	SCOPE_CYCLE_COUNTER(STAT_AnimNativeBlendPoses);
+	ANIM_MT_SCOPE_CYCLE_COUNTER(BlendPosesInGraph, Output.AnimInstance->IsRunningParallelEvaluation());
 
 	const int32 NumPoses = PosesToEvaluate.Num();
 
diff --git a/Engine/Source/Runtime/Engine/Private/Animation/AnimNode_LayeredBoneBlend.cpp b/Engine/Source/Runtime/Engine/Private/Animation/AnimNode_LayeredBoneBlend.cpp
--- a/Engine/Source/Runtime/Engine/Private/Animation/AnimNode_LayeredBoneBlend.cpp
+++ b/Engine/Source/Runtime/Engine/Private/Animation/AnimNode_LayeredBoneBlend.cpp
@@ -70,7 +70,7 @@ void FAnimNode_LayeredBoneBlend::Update(const FAnimationUpdateContext& Context)
 
 void FAnimNode_LayeredBoneBlend::Evaluate(FPoseContext& Output)
 {
-	SCOPE_CYCLE_COUNTER(STAT_AnimNativeBlendPoses);
+	ANIM_MT_SCOPE_CYCLE_COUNTER(BlendPosesInGraph, Output.AnimInstance->IsRunningParallelEvaluation());
 
 	const int NumPoses = BlendPoses.Num();
 	if (NumPoses == 0)
diff --git a/Engine/Source/Runtime/Engine/Private/Animation/AnimNode_StateMachine.cpp b/Engine/Source/Runtime/Engine/Private/Animation/AnimNode_StateMachine.cpp
--- a/Engine/Source/Runtime/Engine/Private/Animation/AnimNode_StateMachine.cpp
+++ b/Engine/Source/Runtime/Engine/Private/Animation/AnimNode_StateMachine.cpp
@@ -574,7 +574,7 @@ void FAnimNode_StateMachine::Evaluate(FPoseContext& Output)
 		return;
 	}
 
-	SCOPE_CYCLE_COUNTER(STAT_AnimStateMachineEvaluate);
+	ANIM_MT_SCOPE_CYCLE_COUNTER(EvaluateAnimStateMachine, Output.AnimInstance->IsRunningParallelEvaluation());
 	
 	if (ActiveTransitionArray.Num() > 0)
 	{
diff --git a/Engine/Source/Runtime/Engine/Private/Components/PoseableMeshComponent.cpp b/Engine/Source/Runtime/Engine/Private/Components/PoseableMeshComponent.cpp
--- a/Engine/Source/Runtime/Engine/Private/Components/PoseableMeshComponent.cpp
+++ b/Engine/Source/Runtime/Engine/Private/Components/PoseableMeshComponent.cpp
@@ -70,7 +70,7 @@ void UPoseableMeshComponent::RefreshBoneTransforms(FActorComponentTickFunction*
 
 void UPoseableMeshComponent::FillSpaceBases()
 {
-	SCOPE_CYCLE_COUNTER(STAT_SkelComposeTime);
+	ANIM_MT_SCOPE_CYCLE_COUNTER(FillSpaceBases, IsRunningParallelEvaluation());
 
 	if( !SkeletalMesh )
 	{
diff --git a/Engine/Source/Runtime/Engine/Private/Components/SkeletalMeshComponent.cpp b/Engine/Source/Runtime/Engine/Private/Components/SkeletalMeshComponent.cpp
--- a/Engine/Source/Runtime/Engine/Private/Components/SkeletalMeshComponent.cpp
+++ b/Engine/Source/Runtime/Engine/Private/Components/SkeletalMeshComponent.cpp
@@ -257,7 +257,7 @@ void USkeletalMeshComponent::InitAnim(bool bForceReinit)
 		// and it is safe to continue
 		const bool bBlockOnTask = true; // wait on evaluation task so it is safe to continue with Init
 		const bool bPerformPostAnimEvaluation = false; // Skip post evaluation, it would be wasted work
-		IsRunningParallelEvaluation(bBlockOnTask, bPerformPostAnimEvaluation);
+		HandleExistingParallelEvaluationTask(bBlockOnTask, bPerformPostAnimEvaluation);
 
 		bool bBlueprintMismatch = (AnimBlueprintGeneratedClass != NULL) && 
 			(AnimScriptInstance != NULL) && (AnimScriptInstance->GetClass() != AnimBlueprintGeneratedClass);
@@ -664,7 +664,7 @@ static void IntersectBoneIndexArrays(TArray<FBoneIndexType>& Output, const TArra
 
 void USkeletalMeshComponent::FillSpaceBases(const USkeletalMesh* InSkeletalMesh, const TArray<FTransform>& SourceAtoms, TArray<FTransform>& DestSpaceBases) const
 {
-	SCOPE_CYCLE_COUNTER(STAT_SkelComposeTime);
+	ANIM_MT_SCOPE_CYCLE_COUNTER(FillSpaceBases, IsRunningParallelEvaluation());
 
 	if( !InSkeletalMesh )
 	{
@@ -937,7 +937,7 @@ void USkeletalMeshComponent::RecalcRequiredBones(int32 LODIndex)
 
 void USkeletalMeshComponent::EvaluateAnimation(const USkeletalMesh* InSkeletalMesh, UAnimInstance* InAnimInstance, TArray<FTransform>& OutLocalAtoms, TArray<FActiveVertexAnim>& OutVertexAnims, FVector& OutRootBoneTranslation, FBlendedCurve& OutCurve) const
 {
-	SCOPE_CYCLE_COUNTER(STAT_AnimBlendTime);
+	ANIM_MT_SCOPE_CYCLE_COUNTER(SkeletalComponentAnimEvaluate, IsRunningParallelEvaluation());
 
 	if( !InSkeletalMesh )
 	{
@@ -1019,7 +1019,8 @@ void USkeletalMeshComponent::UpdateSlaveComponent()
 
 void USkeletalMeshComponent::PerformAnimationEvaluation(const USkeletalMesh* InSkeletalMesh, UAnimInstance* InAnimInstance, TArray<FTransform>& OutSpaceBases, TArray<FTransform>& OutLocalAtoms, TArray<FActiveVertexAnim>& OutVertexAnims, FVector& OutRootBoneTranslation, FBlendedCurve& OutCurve) const
 {
-	SCOPE_CYCLE_COUNTER(STAT_PerformAnimEvaluation);
+	ANIM_MT_SCOPE_CYCLE_COUNTER(PerformAnimEvaluation, IsRunningParallelEvaluation());
+
 	// Can't do anything without a SkeletalMesh
 	// Do nothing more if no bones in skeleton.
 	if (!InSkeletalMesh || OutSpaceBases.Num() == 0)
@@ -1097,7 +1098,7 @@ void USkeletalMeshComponent::RefreshBoneTransforms(FActorComponentTickFunction*
 															// will need to wait on an existing task.
 
 	const bool bPerformPostAnimEvaluation = true;
-	if (IsRunningParallelEvaluation(bBlockOnTask, bPerformPostAnimEvaluation))
+	if (HandleExistingParallelEvaluationTask(bBlockOnTask, bPerformPostAnimEvaluation))
 	{
 		return;
 	}
@@ -2219,7 +2220,7 @@ void USkeletalMeshComponent::RefreshActiveVertexAnims()
 	}
 }
 
-bool USkeletalMeshComponent::IsRunningParallelEvaluation(bool bBlockOnTask, bool bPerformPostAnimEvaluation)
+bool USkeletalMeshComponent::HandleExistingParallelEvaluationTask(bool bBlockOnTask, bool bPerformPostAnimEvaluation)
 {
 	if (IsValidRef(ParallelAnimationEvaluationTask)) // We are already processing eval on another thread
 	{
diff --git a/Engine/Source/Runtime/Engine/Private/PhysicsEngine/PhysAnim.cpp b/Engine/Source/Runtime/Engine/Private/PhysicsEngine/PhysAnim.cpp
--- a/Engine/Source/Runtime/Engine/Private/PhysicsEngine/PhysAnim.cpp
+++ b/Engine/Source/Runtime/Engine/Private/PhysicsEngine/PhysAnim.cpp
@@ -384,7 +384,7 @@ void USkeletalMeshComponent::BlendInPhysics()
 	// If we don't have or want any physics, we do nothing.
 	if( Bodies.Num() > 0 )
 	{
-		IsRunningParallelEvaluation(/*bBlockOnTask = */ true, /*bPerformPostAnimEvaluation =*/ true);
+		HandleExistingParallelEvaluationTask(/*bBlockOnTask = */ true, /*bPerformPostAnimEvaluation =*/ true);
 		// start parallel work
 		check(!IsValidRef(ParallelAnimationEvaluationTask));
 
diff --git a/Engine/Source/Runtime/Engine/Public/Animation/AnimMTStats.h b/Engine/Source/Runtime/Engine/Public/Animation/AnimMTStats.h
new file mode 100644
--- /dev/null
+++ b/Engine/Source/Runtime/Engine/Public/Animation/AnimMTStats.h
@@ -0,0 +1,10 @@
+// Copyright 1998-2015 Epic Games, Inc. All Rights Reserved.
+
+DO_ANIMSTAT_PROCESSING(PerformAnimEvaluation);
+DO_ANIMSTAT_PROCESSING(SkeletalComponentAnimEvaluate);
+DO_ANIMSTAT_PROCESSING(EvaluateAnimGraph);
+DO_ANIMSTAT_PROCESSING(EvaluateAnimInstance);
+DO_ANIMSTAT_PROCESSING(DecompressAnimationPose);
+DO_ANIMSTAT_PROCESSING(BlendPosesInGraph);
+DO_ANIMSTAT_PROCESSING(EvaluateAnimStateMachine);
+DO_ANIMSTAT_PROCESSING(FillSpaceBases);
\ No newline at end of file
diff --git a/Engine/Source/Runtime/Engine/Public/Animation/AnimStats.h b/Engine/Source/Runtime/Engine/Public/Animation/AnimStats.h
--- a/Engine/Source/Runtime/Engine/Public/Animation/AnimStats.h
+++ b/Engine/Source/Runtime/Engine/Public/Animation/AnimStats.h
@@ -13,17 +13,9 @@ DECLARE_CYCLE_STAT_EXTERN(TEXT("Anim Tick Time"), STAT_AnimTickTime, STATGROUP_A
 DECLARE_CYCLE_STAT_EXTERN(TEXT("StateMachine Update"), STAT_AnimStateMachineUpdate, STATGROUP_Anim, );
 DECLARE_CYCLE_STAT_EXTERN(TEXT("StateMachine Find Transition"), STAT_AnimStateMachineFindTransition, STATGROUP_Anim, );
 DECLARE_CYCLE_STAT_EXTERN(TEXT("RefreshBoneTransforms"), STAT_RefreshBoneTransforms, STATGROUP_Anim, );
-DECLARE_CYCLE_STAT_EXTERN(TEXT("Evaluate Anim"), STAT_AnimBlendTime, STATGROUP_Anim, );
-DECLARE_CYCLE_STAT_EXTERN(TEXT("Perform Anim Evaluation"), STAT_PerformAnimEvaluation, STATGROUP_Anim, );
 DECLARE_CYCLE_STAT_EXTERN(TEXT("Post Anim Evaluation"), STAT_PostAnimEvaluation, STATGROUP_Anim, );
-DECLARE_CYCLE_STAT_EXTERN(TEXT("GraphTime"), STAT_AnimGraphEvaluate, STATGROUP_Anim, );
-DECLARE_CYCLE_STAT_EXTERN(TEXT("Eval Pose"), STAT_AnimNativeEvaluatePoses, STATGROUP_Anim, );
 DECLARE_CYCLE_STAT_EXTERN(TEXT("Trigger Notifies"), STAT_AnimTriggerAnimNotifies, STATGROUP_Anim, );
 DECLARE_CYCLE_STAT_EXTERN(TEXT("Anim Decompression"), STAT_GetAnimationPose, STATGROUP_Anim, );
-DECLARE_CYCLE_STAT_EXTERN(TEXT("Blend Poses"), STAT_AnimNativeBlendPoses, STATGROUP_Anim, );
-DECLARE_CYCLE_STAT_EXTERN(TEXT("Copy Pose"), STAT_AnimNativeCopyPoses, STATGROUP_Anim, );
-DECLARE_CYCLE_STAT_EXTERN(TEXT("StateMachine Eval"), STAT_AnimStateMachineEvaluate, STATGROUP_Anim, );
-DECLARE_CYCLE_STAT_EXTERN(TEXT("FillSpaceBases"), STAT_SkelComposeTime, STATGROUP_Anim, );
 DECLARE_CYCLE_STAT_EXTERN(TEXT("InterpolateSkippedFrames"), STAT_InterpolateSkippedFrames, STATGROUP_Anim, );
 DECLARE_CYCLE_STAT_EXTERN(TEXT("UpdateKinematicBonesToAnim"), STAT_UpdateRBBones, STATGROUP_Anim, );
 DECLARE_CYCLE_STAT_EXTERN(TEXT("UpdateRBJointsMotors"), STAT_UpdateRBJoints, STATGROUP_Anim, );
@@ -39,4 +31,19 @@ DECLARE_CYCLE_STAT_EXTERN(TEXT("BlueprintUpdateAnimation"), STAT_BlueprintUpdate
 DECLARE_CYCLE_STAT_EXTERN(TEXT("NativeUpdateAnimation"), STAT_NativeUpdateAnimation, STATGROUP_Anim, );
 DECLARE_CYCLE_STAT_EXTERN(TEXT("BlueprintPostEvaluateAnimation"), STAT_BlueprintPostEvaluateAnimation, STATGROUP_Anim, );
 
+#define DO_ANIMSTAT_PROCESSING(StatName) DECLARE_CYCLE_STAT_EXTERN(TEXT(#StatName), STAT_ ## StatName, STATGROUP_Anim, )
+#include "AnimMTStats.h"
+#undef DO_ANIMSTAT_PROCESSING
 
+#define DO_ANIMSTAT_PROCESSING(StatName) DECLARE_CYCLE_STAT_EXTERN(TEXT(#StatName) TEXT("_WorkerThread"), STAT_ ## StatName ## _WorkerThread, STATGROUP_Anim, )
+#include "AnimMTStats.h"
+#undef DO_ANIMSTAT_PROCESSING
+
+#if STATS
+#define ANIM_MT_SCOPE_CYCLE_COUNTER(StatName, bIsMultithreaded) \
+	SCOPE_CYCLE_COUNTER_GUARD \
+	TStatId CycleCountID_##StatName = (bIsMultithreaded ? GET_STATID(STAT_ ## StatName ## _WorkerThread) : GET_STATID(STAT_ ## StatName)); \
+	FScopeCycleCounter CycleCount_##StatName(CycleCountID_##StatName);
+#else
+	#define ANIM_MT_SCOPE_CYCLE_COUNTER(StatName, bIsMultithreaded)
+#endif
\ No newline at end of file
