PR #1161: Fixed function "Set light color" for light component. (Contributed by [REDACTED_PERSON]). UE-15810

diff --git a/Engine/Source/Runtime/Engine/Classes/Components/LightComponent.h b/Engine/Source/Runtime/Engine/Classes/Components/LightComponent.h
--- a/Engine/Source/Runtime/Engine/Classes/Components/LightComponent.h
+++ b/Engine/Source/Runtime/Engine/Classes/Components/LightComponent.h
@@ -190,7 +190,7 @@ public:
 
 	/** Set color of the light */
 	UFUNCTION(BlueprintCallable, Category="Rendering|Components|Light")
-	void SetLightColor(FLinearColor NewLightColor);
+	void SetLightColor(FLinearColor NewLightColor, bool bSRGB = true);
 
 	UFUNCTION(BlueprintCallable, Category = "Rendering|Components|Light")
 	void SetTemperature(float NewTemperature);
diff --git a/Engine/Source/Runtime/Engine/Private/Components/LightComponent.cpp b/Engine/Source/Runtime/Engine/Private/Components/LightComponent.cpp
--- a/Engine/Source/Runtime/Engine/Private/Components/LightComponent.cpp
+++ b/Engine/Source/Runtime/Engine/Private/Components/LightComponent.cpp
@@ -720,9 +720,9 @@ void ULightComponent::SetIndirectLightingIntensity(float NewIntensity)
 }
 
 /** Set color of the light */
-void ULightComponent::SetLightColor(FLinearColor NewLightColor)
+void ULightComponent::SetLightColor(FLinearColor NewLightColor, bool bSRGB)
 {
-	FColor NewColor(NewLightColor.ToFColor(true));
+	FColor NewColor(NewLightColor.ToFColor(bSRGB));
 
 	// Can't set color on a static light
 	if (AreDynamicDataChangesAllowed()
