UE4 - additional error control to help track down UE-17227, a crash reporter bug involving the shader compiler transfer filename

diff --git a/Engine/Source/Runtime/Engine/Private/ShaderCompiler/ShaderCompiler.cpp b/Engine/Source/Runtime/Engine/Private/ShaderCompiler/ShaderCompiler.cpp
--- a/Engine/Source/Runtime/Engine/Private/ShaderCompiler/ShaderCompiler.cpp
+++ b/Engine/Source/Runtime/Engine/Private/ShaderCompiler/ShaderCompiler.cpp
@@ -109,7 +109,7 @@ namespace XGEConsoleVariables
 }
 
 // Serialize Queued Job information
-static void DoWriteTasks(const TArray<FShaderCompileJob*>& QueuedJobs, FArchive& TransferFile)
+static bool DoWriteTasks(const TArray<FShaderCompileJob*>& QueuedJobs, FArchive& TransferFile)
 {
 	int32 ShaderCompileWorkerInputVersion = 2;
 	TransferFile << ShaderCompileWorkerInputVersion;
@@ -122,7 +122,7 @@ static void DoWriteTasks(const TArray<FShaderCompileJob*>& QueuedJobs, FArchive&
 		TransferFile << QueuedJobs[JobIndex]->Input;
 	}
 
-	TransferFile.Close();
+	return TransferFile.Close();
 }
 
 // Process results from Worker Process
@@ -744,19 +744,29 @@ void FShaderCompileThreadRunnable::WriteNewTasks()
 					}
 					TransferFile = IFileManager::Get().CreateFileWriter(*TransferFileName, FILEWRITE_EvenIfReadOnly);
 					RetryCount++;
+					if (TransferFile == NULL)
+					{
+						UE_LOG(LogShaderCompilers, Warning, TEXT("Could not create the shader compiler transfer file '%s', retrying..."), *TransferFileName);
+					}
 				}
 				if (TransferFile == NULL)
 				{
-					TransferFile = IFileManager::Get().CreateFileWriter(*TransferFileName, FILEWRITE_EvenIfReadOnly | FILEWRITE_NoFail);
+					UE_LOG(LogShaderCompilers, Fatal, TEXT("Could not create the shader compiler transfer file '%s'."), *TransferFileName);
 				}
 				check(TransferFile);
 
-				DoWriteTasks(CurrentWorkerInfo.QueuedJobs, *TransferFile);
+				if (!DoWriteTasks(CurrentWorkerInfo.QueuedJobs, *TransferFile))
+				{
+					UE_LOG(LogShaderCompilers, Fatal, TEXT("Could not write the shader compiler transfer filename to '%s'."), *TransferFileName);
+				}
 				delete TransferFile;
 
 				// Change the transfer file name to proper one
 				FString ProperTransferFileName = WorkingDirectory / TEXT("WorkerInputOnly.in");
-				IFileManager::Get().Move(*ProperTransferFileName, *TransferFileName);
+				if (!IFileManager::Get().Move(*ProperTransferFileName, *TransferFileName))
+				{
+					UE_LOG(LogShaderCompilers, Fatal, TEXT("Could not rename the shader compiler transfer filename to '%s' from '%s'."), *ProperTransferFileName, *TransferFileName);
+				}
 			}
 		}
 	}
@@ -1145,6 +1155,14 @@ FShaderCompilingManager::FShaderCompilingManager() :
 	// With processes from other games, processes from the same game or threads in this same process.
 	// Use IFileManager to do path conversion to properly handle sandbox paths (outside of standard paths in particular).
 	ShaderBaseWorkingDirectory = FPlatformProcess::ShaderWorkingDir() / FString::FromInt(ProcessId) + TEXT("/");
+	if (!IFileManager::Get().DeleteDirectory(*ShaderBaseWorkingDirectory, false, true))
+	{
+		UE_LOG(LogShaderCompilers, Fatal, TEXT("Could not delete the shader compiler working directory '%s'."), *ShaderBaseWorkingDirectory);
+	}
+	else
+	{
+		UE_LOG(LogShaderCompilers, Log, TEXT("Cleaned the shader compiler working directory '%s'."), *ShaderBaseWorkingDirectory);
+	}
 	FString AbsoluteBaseDirectory = IFileManager::Get().ConvertToAbsolutePathForExternalAppForWrite(*ShaderBaseWorkingDirectory);
 	FPaths::NormalizeDirectoryName(AbsoluteBaseDirectory);
 	AbsoluteShaderBaseWorkingDirectory = AbsoluteBaseDirectory + TEXT("/");
