Clarify and (very slightly) optimize default FindPlayerStart implementation

diff --git a/Engine/Source/Runtime/Engine/Private/GameMode.cpp b/Engine/Source/Runtime/Engine/Private/GameMode.cpp
--- a/Engine/Source/Runtime/Engine/Private/GameMode.cpp
+++ b/Engine/Source/Runtime/Engine/Private/GameMode.cpp
@@ -358,10 +358,11 @@ AActor* AGameMode::FindPlayerStart_Implementation( AController* Player, const FS
 	// if incoming start is specified, then just use it
 	if( !IncomingName.IsEmpty() )
 	{
+		const FName IncomingPlayerStartTag = FName(*IncomingName);
 		for (TActorIterator<APlayerStart> It(World); It; ++It)
 		{
 			APlayerStart* Start = *It;
-			if (Start && Start->PlayerStartTag == FName(*IncomingName))
+			if (Start && Start->PlayerStartTag == IncomingPlayerStartTag)
 			{
 				return Start;
 			}
@@ -380,25 +381,9 @@ AActor* AGameMode::FindPlayerStart_Implementation( AController* Player, const FS
 		// no player start found
 		UE_LOG(LogGameMode, Log, TEXT("Warning - PATHS NOT DEFINED or NO PLAYERSTART with positive rating"));
 
-		// Search all loaded levels for possible player start object
-		for ( int32 LevelIndex = 0; LevelIndex < World->GetNumLevels(); ++LevelIndex )
-		{
-			ULevel* Level = World->GetLevel( LevelIndex );
-			for ( int32 ActorIndex = 0; ActorIndex < Level->Actors.Num(); ++ActorIndex )
-			{
-				AActor* NavObject = Cast<AActor>(Level->Actors[ActorIndex]);
-				if ( NavObject )
-				{
-					BestStart = NavObject;
-					break;
-				}
-			}
-
-			if (BestStart != NULL)
-			{
-				break;
-			}
-		}
+		// This is a bit odd, but there was a complex chunk of code that in the end always resulted in this, so we may as well just 
+		// short cut it down to this.  Basically we are saying spawn at 0,0,0 if we didn't find a proper player start
+		BestStart = World->GetWorldSettings();
 	}
 
 	return BestStart;
