PR #790: FQuantizedLightmapData::HasNonZeroData improvement (Contributed by [REDACTED_PERSON]). UE-8959

diff --git a/Engine/Source/Runtime/Engine/Private/LightMap.cpp b/Engine/Source/Runtime/Engine/Private/LightMap.cpp
--- a/Engine/Source/Runtime/Engine/Private/LightMap.cpp
+++ b/Engine/Source/Runtime/Engine/Private/LightMap.cpp
@@ -2189,7 +2189,10 @@ bool FQuantizedLightmapData::HasNonZeroData() const
 
 		if (LightmapSample.Coverage >= MinCoverageThreshold)
 		{
-			for (int32 CoefficentIndex = 0; CoefficentIndex < NUM_STORED_LIGHTMAP_COEF; CoefficentIndex++)
+			// Don't look at simple lightmap coefficients if we're not building them.
+			const int32 NumCoefficients = GEngine->bShouldGenerateLowQualityLightmaps ? NUM_STORED_LIGHTMAP_COEF : NUM_HQ_LIGHTMAP_COEF;
+
+			for (int32 CoefficentIndex = 0; CoefficentIndex < NumCoefficients; CoefficentIndex++)
 			{
 				if ((LightmapSample.Coefficients[CoefficentIndex][0] != 0) || (LightmapSample.Coefficients[CoefficentIndex][1] != 0) || (LightmapSample.Coefficients[CoefficentIndex][2] != 0))
 				{
