Fix sound mix failing to be pushed if it had just been popped but UpdateSoundClasses hadn't yet been executed
Mark some 0.0 doubles properly as floats

diff --git a/Engine/Source/Runtime/Engine/Private/AudioDevice.cpp b/Engine/Source/Runtime/Engine/Private/AudioDevice.cpp
--- a/Engine/Source/Runtime/Engine/Private/AudioDevice.cpp
+++ b/Engine/Source/Runtime/Engine/Private/AudioDevice.cpp
@@ -242,11 +242,11 @@ void FAudioDevice::ResetInterpolation( void )
 {
 	for (FListener& Listener : Listeners)
 	{
-		Listener.InteriorStartTime = 0.0;
-		Listener.InteriorEndTime = 0.0;
-		Listener.ExteriorEndTime = 0.0;
-		Listener.InteriorLPFEndTime = 0.0;
-		Listener.ExteriorLPFEndTime = 0.0;
+		Listener.InteriorStartTime = 0.f;
+		Listener.InteriorEndTime = 0.f;
+		Listener.ExteriorEndTime = 0.f;
+		Listener.InteriorLPFEndTime = 0.f;
+		Listener.ExteriorLPFEndTime = 0.f;
 
 		Listener.InteriorVolumeInterp = 0.0f;
 		Listener.InteriorLPFInterp = 0.0f;
@@ -1070,7 +1070,7 @@ bool FAudioDevice::ApplySoundMix( USoundMix* NewMix, FSoundMixState* SoundMixSta
 void FAudioDevice::UpdateSoundMix(class USoundMix* SoundMix, FSoundMixState* SoundMixState)
 {
 	// If this SoundMix will automatically end, add some more time
-	if (SoundMixState->FadeOutStartTime >= 0.0)
+	if (SoundMixState->FadeOutStartTime >= 0.f)
 	{
 		if (SoundMixState->CurrentState == ESoundMixState::FadingOut)
 		{
@@ -1156,7 +1156,7 @@ bool FAudioDevice::TryClearingSoundMix(USoundMix* SoundMix, FSoundMixState* Soun
 		if (SoundMixState->ActiveRefCount == 0 && SoundMixState->PassiveRefCount == 0 && SoundMixState->IsBaseSoundMix == false)
 		{
 			// do whatever is needed to remove influence of this SoundMix
-			if (SoundMix->FadeOutTime > 0.0)
+			if (SoundMix->FadeOutTime > 0.f)
 			{
 				if (SoundMixState->CurrentState == ESoundMixState::Inactive)
 				{
@@ -1319,7 +1319,7 @@ void FAudioDevice::UpdateSoundClassProperties()
 			SoundMixState->CurrentState = ESoundMixState::FadingIn;
 		}
 		else if( FApp::GetCurrentTime() >= SoundMixState->FadeInEndTime
-			&& ( SoundMixState->IsBaseSoundMix || SoundMixState->PassiveRefCount > 0 || SoundMixState->FadeOutStartTime < 0.0 || FApp::GetCurrentTime() < SoundMixState->FadeOutStartTime ) )
+			&& ( SoundMixState->IsBaseSoundMix || SoundMixState->PassiveRefCount > 0 || SoundMixState->FadeOutStartTime < 0.f || FApp::GetCurrentTime() < SoundMixState->FadeOutStartTime ) )
 		{
 			// .. ensure the full mix is applied between the end of the fade in time and the start of the fade out time
 			// or if SoundMix is the base or active via a passive push - ignores duration.
@@ -1339,7 +1339,7 @@ void FAudioDevice::UpdateSoundClassProperties()
 		}
 		else
 		{
-			check( SoundMixState->EndTime >= 0.0 && FApp::GetCurrentTime() >= SoundMixState->EndTime );
+			check( SoundMixState->EndTime >= 0.f && FApp::GetCurrentTime() >= SoundMixState->EndTime );
 			// Clear the effect of this SoundMix - may need to revisit for passive
 			SoundMixState->InterpValue = 0.0f;
 			SoundMixState->CurrentState = ESoundMixState::AwaitingRemoval;
@@ -1488,6 +1488,12 @@ void FAudioDevice::PushSoundMixModifier(USoundMix* SoundMix, bool bIsPassive)
 		}
 		else
 		{
+			// Make sure that even if UpdateSoundClasses hasn't been run since we indicated we wanted the
+			// mix to fade out that we flag is as such so the push is applied correctly
+			if (SoundMixState->FadeOutStartTime > 0.f)
+			{
+				SoundMixState->CurrentState = ESoundMixState::FadingOut;
+			}
 			UpdateSoundMix(SoundMix, SoundMixState);
 		}
 
@@ -1517,7 +1523,7 @@ void FAudioDevice::PopSoundMixModifier(USoundMix* SoundMix, bool bIsPassive)
 				if (SoundMixState->PassiveRefCount == 0)
 				{
 					// Check whether Fade out time was previously set and reset it to current time
-					if (SoundMixState->FadeOutStartTime >= 0.0 && FApp::GetCurrentTime() > SoundMixState->FadeOutStartTime)
+					if (SoundMixState->FadeOutStartTime >= 0.f && FApp::GetCurrentTime() > SoundMixState->FadeOutStartTime)
 					{
 						SoundMixState->FadeOutStartTime = FApp::GetCurrentTime();
 						SoundMixState->EndTime = SoundMixState->FadeOutStartTime + SoundMix->FadeOutTime;
