Fix marquee select when there are placed actors with skeletal mesh components

Merging 2612686

//depot/UE4-UT/...

to //depot/UE4/...

diff --git a/Engine/Source/Runtime/Engine/Private/Components/SkeletalMeshComponent.cpp b/Engine/Source/Runtime/Engine/Private/Components/SkeletalMeshComponent.cpp
--- a/Engine/Source/Runtime/Engine/Private/Components/SkeletalMeshComponent.cpp
+++ b/Engine/Source/Runtime/Engine/Private/Components/SkeletalMeshComponent.cpp
@@ -2012,7 +2012,10 @@ bool USkeletalMeshComponent::ComponentIsTouchingSelectionBox(const FBox& InSelBB
 
 		// If the selection box has to encompass all of the component and none of the component's verts failed the intersection test, this component
 		// is consider touching
-		return true;
+		if (bMustEncompassEntireComponent)
+		{
+			return true;
+		}
 	}
 
 	return false;
