Don't tick timeline components while deactivated

diff --git a/Engine/Source/Runtime/Engine/Classes/Components/TimelineComponent.h b/Engine/Source/Runtime/Engine/Classes/Components/TimelineComponent.h
--- a/Engine/Source/Runtime/Engine/Classes/Components/TimelineComponent.h
+++ b/Engine/Source/Runtime/Engine/Classes/Components/TimelineComponent.h
@@ -446,6 +446,8 @@ public:
 
 	// Begin ActorComponent interface.
 	virtual void TickComponent(float DeltaTime, enum ELevelTick TickType, FActorComponentTickFunction *ThisTickFunction) override;
+	virtual void Activate(bool bReset=false) override;
+	virtual void Deactivate() override;
 	// End ActorComponent interface.
 
 	/** Get the signature function for Timeline event functions */
diff --git a/Engine/Source/Runtime/Engine/Private/Timeline.cpp b/Engine/Source/Runtime/Engine/Private/Timeline.cpp
--- a/Engine/Source/Runtime/Engine/Private/Timeline.cpp
+++ b/Engine/Source/Runtime/Engine/Private/Timeline.cpp
@@ -563,6 +563,7 @@ UTimelineComponent::UTimelineComponent(const FObjectInitializer& ObjectInitializ
 	: Super(ObjectInitializer)
 {
 	PrimaryComponentTick.bCanEverTick = true;
+	PrimaryComponentTick.bStartWithTickEnabled = false;
 	PrimaryComponentTick.TickGroup = TG_PrePhysics;
 }
 
@@ -572,22 +573,35 @@ void UTimelineComponent::TickComponent(float DeltaTime, enum ELevelTick TickType
 
 	Super::TickComponent(DeltaTime, TickType, ThisTickFunction);
 
-	if (bIsActive)
-	{
-		TheTimeline.TickTimeline(DeltaTime);
+	TheTimeline.TickTimeline(DeltaTime);
 
-		if (!IsNetSimulating())
+	if (!IsNetSimulating())
+	{
+		// Do not deactivate if we are done, since bActive is a replicated property and we shouldn't have simulating
+		// clients touch replicated variables.
+		if (!TheTimeline.IsPlaying())
 		{
-			// Do not deactivate if we are done, since bActive is a replicated property and we shouldn't have simulating
-			// clients touch replicated variables.
-			if (!TheTimeline.IsPlaying())
-			{
-				Deactivate();
-			}
+			Deactivate();
 		}
 	}
 }
 
+void UTimelineComponent::Activate(bool bReset)
+{
+	Super::Activate(bReset);
+
+	PrimaryComponentTick.bStartWithTickEnabled = true;
+	PrimaryComponentTick.SetTickFunctionEnable(true);
+}
+
+void UTimelineComponent::Deactivate()
+{
+	Super::Deactivate();
+
+	PrimaryComponentTick.bStartWithTickEnabled = false;
+	PrimaryComponentTick.SetTickFunctionEnable(false);
+}
+
 void UTimelineComponent::Play()
 {
 	Activate();
