FRepMovement now defaults to more aggressive quantization for more efficient bandwidth usage.
The quantization level can be configured in the editor if the defaults cause artifacts.

diff --git a/Engine/Source/Runtime/Engine/Classes/Engine/EngineTypes.h b/Engine/Source/Runtime/Engine/Classes/Engine/EngineTypes.h
--- a/Engine/Source/Runtime/Engine/Classes/Engine/EngineTypes.h
+++ b/Engine/Source/Runtime/Engine/Classes/Engine/EngineTypes.h
@@ -2698,6 +2698,26 @@ namespace EEndPlayReason
 
 DECLARE_DYNAMIC_DELEGATE(FTimerDynamicDelegate);
 
+UENUM()
+enum class EVectorQuantization : uint8
+{
+	/** Each vector component will be rounded to the nearest whole number. Equivalent to the behavior of FVector_NetQuantize. */
+	RoundWholeNumber,
+	/** Each vector component will be rounded, preserving one decimal place. Equivalent to the behavior of FVector_NetQuantize10. */
+	RoundOneDecimal,
+	/** Each vector component will be rounded, preserving two decimal places. Equivalent to the behavior of FVector_NetQuantize100. */
+	RoundTwoDecimals
+};
+
+UENUM()
+enum class ERotatorQuantization : uint8
+{
+	/** The rotator will be compressed to 8 bits per component. */
+	ByteComponents,
+	/** The rotator will be compressed to 16 bits per component. */
+	ShortComponents
+};
+
 /** Replicated movement data of our RootComponent.
   * Struct used for efficient replication as velocity and location are generally replicated together (this saves a repindex) 
   * and velocity.Z is commonly zero (most position replications are for walking pawns). 
@@ -2707,34 +2727,56 @@ struct FRepMovement
 {
 	GENERATED_USTRUCT_BODY()
 
-	UPROPERTY()
-	FVector_NetQuantize100 LinearVelocity;
+	UPROPERTY(Transient)
+	FVector LinearVelocity;
 
-	UPROPERTY()
-	FVector_NetQuantize100 AngularVelocity;
+	UPROPERTY(Transient)
+	FVector AngularVelocity;
 	
-	UPROPERTY()
-	FVector_NetQuantize100 Location;
+	UPROPERTY(Transient)
+	FVector Location;
 
-	UPROPERTY()
+	UPROPERTY(Transient)
 	FRotator Rotation;
 
 	/** If set, RootComponent should be sleeping. */
-	UPROPERTY()
+	UPROPERTY(Transient)
 	uint8 bSimulatedPhysicSleep : 1;
 
 	/** If set, additional physic data (angular velocity) will be replicated. */
-	UPROPERTY()
+	UPROPERTY(Transient)
 	uint8 bRepPhysics : 1;
 
-	FRepMovement()
-		: LinearVelocity(ForceInit)
-		, AngularVelocity(ForceInit)
-		, Location(ForceInit)
-		, Rotation(ForceInit)
-		, bSimulatedPhysicSleep(false)
-		, bRepPhysics(false)
-	{}
+	/** Allows tuning the compression level for the replicated location and velocity vectors. You should only need to change this from the default if you see visual artifacts. */
+	UPROPERTY(EditDefaultsOnly, Category=Replication, AdvancedDisplay)
+	EVectorQuantization VectorQuantizationLevel;
+
+	/** Allows tuning the compression level for replicated rotation. You should only need to change this from the default if you see visual artifacts. */
+	UPROPERTY(EditDefaultsOnly, Category=Replication, AdvancedDisplay)
+	ERotatorQuantization RotationQuantizationLevel;
+
+	FRepMovement();
+
+	bool SerializeQuantizedVector(FArchive& Ar, FVector& Vector)
+	{
+		switch(VectorQuantizationLevel)
+		{
+			case EVectorQuantization::RoundTwoDecimals:
+			{
+				return SerializePackedVector<100, 30>(Vector, Ar);
+			}
+
+			case EVectorQuantization::RoundOneDecimal:
+			{
+				return SerializePackedVector<10, 24>(Vector, Ar);
+			}
+
+			default:
+			{
+				return SerializePackedVector<1, 20>(Vector, Ar);
+			}
+		}
+	}
 
 	bool NetSerialize(FArchive& Ar, class UPackageMap* Map, bool& bOutSuccess)
 	{
@@ -2746,21 +2788,30 @@ struct FRepMovement
 
 		bOutSuccess = true;
 
-		bool bOutSuccessLocal = true;
+		// update location, rotation, linear velocity
+		bOutSuccess &= SerializeQuantizedVector( Ar, Location );
+		
+		switch(RotationQuantizationLevel)
+		{
+			case ERotatorQuantization::ByteComponents:
+			{
+				Rotation.SerializeCompressed( Ar );
+				break;
+			}
 
-		// update location, linear velocity
-		Location.NetSerialize( Ar, Map, bOutSuccessLocal );
-		bOutSuccess &= bOutSuccessLocal;
-		Rotation.NetSerialize( Ar, Map, bOutSuccessLocal );
-		bOutSuccess &= bOutSuccessLocal;
-		LinearVelocity.NetSerialize( Ar, Map, bOutSuccessLocal );
-		bOutSuccess &= bOutSuccessLocal;
+			case ERotatorQuantization::ShortComponents:
+			{
+				Rotation.SerializeCompressedShort( Ar );
+				break;
+			}
+		}
+		
+		bOutSuccess &= SerializeQuantizedVector( Ar, LinearVelocity );
 
 		// update angular velocity if required
 		if ( bRepPhysics )
 		{
-			AngularVelocity.NetSerialize( Ar, Map, bOutSuccessLocal );
-			bOutSuccess &= bOutSuccessLocal;
+			bOutSuccess &= SerializeQuantizedVector( Ar, AngularVelocity );
 		}
 
 		return true;
diff --git a/Engine/Source/Runtime/Engine/Classes/GameFramework/Actor.h b/Engine/Source/Runtime/Engine/Classes/GameFramework/Actor.h
--- a/Engine/Source/Runtime/Engine/Classes/GameFramework/Actor.h
+++ b/Engine/Source/Runtime/Engine/Classes/GameFramework/Actor.h
@@ -246,7 +246,7 @@ public:
 	ENetRole GetRemoteRole() const;
 
 	/** Used for replication of our RootComponent's position and velocity */
-	UPROPERTY(Transient, ReplicatedUsing=OnRep_ReplicatedMovement)
+	UPROPERTY(EditDefaultsOnly, ReplicatedUsing=OnRep_ReplicatedMovement, Category=Replication, AdvancedDisplay)
 	struct FRepMovement ReplicatedMovement;
 
 	/** Used for replicating attachment of this actor's RootComponent to another actor. */
diff --git a/Engine/Source/Runtime/Engine/Private/EngineTypes.cpp b/Engine/Source/Runtime/Engine/Private/EngineTypes.cpp
--- a/Engine/Source/Runtime/Engine/Private/EngineTypes.cpp
+++ b/Engine/Source/Runtime/Engine/Private/EngineTypes.cpp
@@ -214,3 +214,14 @@ FString FHitResult::ToString() const
 		FaceIndex);
 }
 
+FRepMovement::FRepMovement()
+	: LinearVelocity(ForceInit)
+	, AngularVelocity(ForceInit)
+	, Location(ForceInit)
+	, Rotation(ForceInit)
+	, bSimulatedPhysicSleep(false)
+	, bRepPhysics(false)
+	, VectorQuantizationLevel(EVectorQuantization::RoundWholeNumber)
+	, RotationQuantizationLevel(ERotatorQuantization::ByteComponents)
+{
+}
