Don't force enable actor tick function when posessing if bStartWithTickEnabled is false

diff --git a/Engine/Source/Runtime/Engine/Private/DebugCameraController.cpp b/Engine/Source/Runtime/Engine/Private/DebugCameraController.cpp
--- a/Engine/Source/Runtime/Engine/Private/DebugCameraController.cpp
+++ b/Engine/Source/Runtime/Engine/Private/DebugCameraController.cpp
@@ -265,7 +265,10 @@ ASpectatorPawn* ADebugCameraController::SpawnSpectatorPawn()
 			{
 				SpawnedSpectator->PossessedBy(this);
 				SpawnedSpectator->PawnClientRestart();
-				SpawnedSpectator->SetActorTickEnabled(true);
+				if (SpawnedSpectator->PrimaryActorTick.bStartWithTickEnabled)
+				{
+					SpawnedSpectator->SetActorTickEnabled(true);
+				}
 
 				UE_LOG(LogPlayerController, Verbose, TEXT("Spawned spectator %s [server:%d]"), *GetNameSafe(SpawnedSpectator), GetNetMode() < NM_Client);
 			}
diff --git a/Engine/Source/Runtime/Engine/Private/PlayerController.cpp b/Engine/Source/Runtime/Engine/Private/PlayerController.cpp
--- a/Engine/Source/Runtime/Engine/Private/PlayerController.cpp
+++ b/Engine/Source/Runtime/Engine/Private/PlayerController.cpp
@@ -752,7 +752,7 @@ void APlayerController::Possess(APawn* PawnToPossess)
 		SetPawn(PawnToPossess);
 		check(GetPawn() != NULL);
 
-		if (GetPawn() != NULL)
+		if (GetPawn() && GetPawn()->PrimaryActorTick.bStartWithTickEnabled)
 		{
 			GetPawn()->SetActorTickEnabled(true);
 		}
@@ -3949,7 +3949,10 @@ ASpectatorPawn* APlayerController::SpawnSpectatorPawn()
 				SpawnedSpectator->SetReplicates(false); // Client-side only
 				SpawnedSpectator->PossessedBy(this);
 				SpawnedSpectator->PawnClientRestart();
-				SpawnedSpectator->SetActorTickEnabled(true);
+				if (SpawnedSpectator->PrimaryActorTick.bStartWithTickEnabled)
+				{
+					SpawnedSpectator->SetActorTickEnabled(true);
+				}
 
 				UE_LOG(LogPlayerController, Verbose, TEXT("Spawned spectator %s [server:%d]"), *GetNameSafe(SpawnedSpectator), GetNetMode() < NM_Client);
 			}
