Increase the default max angular velocity to 10 revolutions per second.

Improve the UI for mass override and angular override so that it uses the standard override UI.

diff --git a/Engine/Source/Runtime/Core/Public/UObject/ObjectVersion.h b/Engine/Source/Runtime/Core/Public/UObject/ObjectVersion.h
--- a/Engine/Source/Runtime/Core/Public/UObject/ObjectVersion.h
+++ b/Engine/Source/Runtime/Core/Public/UObject/ObjectVersion.h
@@ -556,6 +556,9 @@ enum EUnrealEngineObjectUE4Version
 	VER_UE4_MESH_PARTICLE_COLLISIONS_CONSIDER_PARTICLE_SIZE,
 	// Adjacency buffer building no longer automatically handled based on triangle count, user-controlled
 	VER_UE4_BUILD_MESH_ADJ_BUFFER_FLAG_EXPOSED,
+	// Change the default max angular velocity
+	VER_UE4_MAX_ANGULAR_VELOCITY_DEFAULT,
+
 	// -----<new versions can be added before this line>-------------------------------------------------
 	// - this needs to be the last line (see note below)
 	VER_UE4_AUTOMATIC_VERSION_PLUS_ONE,
diff --git a/Engine/Source/Runtime/Engine/Classes/PhysicsEngine/BodyInstance.h b/Engine/Source/Runtime/Engine/Classes/PhysicsEngine/BodyInstance.h
--- a/Engine/Source/Runtime/Engine/Classes/PhysicsEngine/BodyInstance.h
+++ b/Engine/Source/Runtime/Engine/Classes/PhysicsEngine/BodyInstance.h
@@ -208,7 +208,7 @@ public:
 	uint32 bSimulatePhysics:1;
 
 	/** If true, mass will not be automatically computed and you must set it directly */
-	UPROPERTY(EditAnywhere, BlueprintReadOnly, Category = Physics, meta=(DisplayName="Override"))
+	UPROPERTY(meta=(DisplayName="Override"))
 	uint32 bOverrideMass : 1;
 
 	/**Mass of the body in KG. By default we compute this based on physical material and mass scale.
@@ -279,6 +279,10 @@ public:
 	UPROPERTY(EditAnywhere, Category = Physics, meta = (DisplayName = "Z"))
 	uint32 bLockZRotation : 1;
 
+	/** Override the default max angular velocity */
+	UPROPERTY(meta = (editcondition = "bSimulatePhysics"))
+	uint32 bOverrideMaxAngularVelocity : 1;
+
 	/** Locks physical movement along specified axis.*/
 	UPROPERTY(EditAnywhere, Category = Physics, meta = (DisplayName = "Mode"))
 	TEnumAsByte<EDOFMode::Type> DOFMode;
@@ -296,7 +300,7 @@ public:
 	float MassScale;
 
 	/** The maximum angular velocity for this instance */
-	UPROPERTY(EditAnywhere, AdvancedDisplay, BlueprintReadOnly, Category = Physics, meta = (editcondition = "bSimulatePhysics"))
+	UPROPERTY(EditAnywhere, AdvancedDisplay, BlueprintReadOnly, Category = Physics, meta = (editcondition = "bOverrideMaxAngularVelocity"))
 	float MaxAngularVelocity;
 
 
@@ -330,7 +334,7 @@ protected:
 	uint32 bUseAsyncScene:1;
 
 	/** Whether this body instance has its own custom MaxDepenetrationVelocity*/
-	UPROPERTY(EditAnywhere, AdvancedDisplay, BlueprintReadOnly, Category = Physics)
+	UPROPERTY()
 	uint32 bOverrideMaxDepenetrationVelocity : 1;
 
 	/** The maximum velocity used to depenetrate this object*/
@@ -338,7 +342,7 @@ protected:
 	float MaxDepenetrationVelocity;
 
 	/** Whether this instance of the object has its own custom walkable slope override setting. */
-	UPROPERTY(EditAnywhere, AdvancedDisplay, BlueprintReadOnly, Category=Physics)
+	UPROPERTY()
 	uint32 bOverrideWalkableSlopeOnInstance : 1;
 
 	/**
@@ -716,7 +720,10 @@ public:
 	/** Set the angular velocity of this body */
 	void SetAngularVelocity(const FVector& NewAngVel, bool bAddToCurrent);
 	/** Set the maximum angular velocity of this body */
-	void SetMaxAngularVelocity(float NewMaxAngVel, bool bAddToCurrent);
+	void SetMaxAngularVelocity(float NewMaxAngVel, bool bAddToCurrent, bool bUpdateOverrideMaxAngularVelocity = true);
+	/** Get the maximum angular velocity of this body */
+	float GetMaxAngularVelocity() const;
+
 	/** Set the maximum depenetration velocity the physics simulation will introduce */
 	void SetMaxDepenetrationVelocity(float MaxVelocity);
 	/** Set whether we should get a notification about physics collisions */
diff --git a/Engine/Source/Runtime/Engine/Classes/PhysicsEngine/PhysicsSettings.h b/Engine/Source/Runtime/Engine/Classes/PhysicsEngine/PhysicsSettings.h
--- a/Engine/Source/Runtime/Engine/Classes/PhysicsEngine/PhysicsSettings.h
+++ b/Engine/Source/Runtime/Engine/Classes/PhysicsEngine/PhysicsSettings.h
@@ -151,6 +151,10 @@ class ENGINE_API UPhysicsSettings : public UDeveloperSettings
 	UPROPERTY(config, EditAnywhere, Category = Simulation)
 	TEnumAsByte<EFrictionCombineMode::Type> RestitutionCombineMode;
 
+	/** Max angular velocity that a simulated object can achieve.*/
+	UPROPERTY(config, EditAnywhere, Category = Simulation)
+	float MaxAngularVelocity;
+
 	/** Max velocity which may be used to depenetrate simulated physics objects. 0 means no maximum. */
 	UPROPERTY(config, EditAnywhere, Category = Simulation)
 	float MaxDepenetrationVelocity;
diff --git a/Engine/Source/Runtime/Engine/Private/PhysicsEngine/BodyInstance.cpp b/Engine/Source/Runtime/Engine/Private/PhysicsEngine/BodyInstance.cpp
--- a/Engine/Source/Runtime/Engine/Private/PhysicsEngine/BodyInstance.cpp
+++ b/Engine/Source/Runtime/Engine/Private/PhysicsEngine/BodyInstance.cpp
@@ -239,7 +239,6 @@ FBodyInstance::FBodyInstance()
 , CustomDOFPlaneNormal(FVector::ZeroVector)
 , COMNudge(ForceInit)
 , MassScale(1.f)
-, MaxAngularVelocity(400.f)
 , DOFConstraint(NULL)
 , WeldParent(NULL)
 , bUseAsyncScene(false)
@@ -266,6 +265,7 @@ FBodyInstance::FBodyInstance()
 , BodyInstancePtr(nullptr)
 #endif
 {
+	MaxAngularVelocity = UPhysicsSettings::Get()->MaxAngularVelocity;
 }
 
 FArchive& operator<<(FArchive& Ar,FBodyInstance& BodyInst)
@@ -276,6 +276,14 @@ FArchive& operator<<(FArchive& Ar,FBodyInstance& BodyInst)
 		Ar << BodyInst.PhysMaterialOverride;
 	}
 
+	if (Ar.IsLoading() && Ar.UE4Ver() < VER_UE4_MAX_ANGULAR_VELOCITY_DEFAULT)
+	{
+		if(BodyInst.MaxAngularVelocity != 400.f)
+		{
+			BodyInst.bOverrideMaxAngularVelocity = true;
+		}
+	}
+
 	return Ar;
 }
 
@@ -1530,7 +1538,7 @@ struct FInitBodiesHelper
 						// Update damping
 						Instance->UpdateDampingProperties();
 
-						Instance->SetMaxAngularVelocity(Instance->MaxAngularVelocity, false);
+						Instance->SetMaxAngularVelocity(Instance->GetMaxAngularVelocity(), false, false);
 
 						Instance->SetMaxDepenetrationVelocity(Instance->bOverrideMaxDepenetrationVelocity ? Instance->MaxDepenetrationVelocity : UPhysicsSettings::Get()->MaxDepenetrationVelocity);
 
@@ -3348,7 +3356,12 @@ void FBodyInstance::SetAngularVelocity(const FVector& NewAngVel, bool bAddToCurr
 #endif
 }
 
-void FBodyInstance::SetMaxAngularVelocity(float NewMaxAngVel, bool bAddToCurrent)
+float FBodyInstance::GetMaxAngularVelocity() const
+{
+	return bOverrideMaxAngularVelocity ? MaxAngularVelocity : UPhysicsSettings::Get()->MaxAngularVelocity;
+}
+
+void FBodyInstance::SetMaxAngularVelocity(float NewMaxAngVel, bool bAddToCurrent, bool bUpdateOverrideMaxAngularVelocity)
 {
 #if WITH_PHYSX
 	const bool bIsDynamic = ExecuteOnPxRigidDynamicReadWrite(this, [&](PxRigidDynamic* PRigidDynamic)
@@ -3374,6 +3387,12 @@ void FBodyInstance::SetMaxAngularVelocity(float NewMaxAngVel, bool bAddToCurrent
 	{
 		MaxAngularVelocity = NewMaxAngVel;	//doesn't really matter since we are not dynamic, but makes sense that we update this anyway
 	}
+
+	if(bUpdateOverrideMaxAngularVelocity)
+	{
+		bOverrideMaxAngularVelocity = true;
+	}
+	
 #endif
 
 	//@TODO: BOX2D: Implement SetMaxAngularVelocity
@@ -4450,7 +4469,7 @@ void FBodyInstance::InitDynamicProperties_AssumesLocked()
 		{
 			UpdateMassProperties();
 			UpdateDampingProperties();
-			SetMaxAngularVelocity(MaxAngularVelocity, false);
+			SetMaxAngularVelocity(GetMaxAngularVelocity(), false, false);
 			SetMaxDepenetrationVelocity(bOverrideMaxDepenetrationVelocity ? MaxDepenetrationVelocity : UPhysicsSettings::Get()->MaxDepenetrationVelocity);
 		}else
 		{
diff --git a/Engine/Source/Runtime/Engine/Private/PhysicsEngine/PhysicsSettings.cpp b/Engine/Source/Runtime/Engine/Private/PhysicsEngine/PhysicsSettings.cpp
--- a/Engine/Source/Runtime/Engine/Private/PhysicsEngine/PhysicsSettings.cpp
+++ b/Engine/Source/Runtime/Engine/Private/PhysicsEngine/PhysicsSettings.cpp
@@ -19,6 +19,7 @@ UPhysicsSettings::UPhysicsSettings(const FObjectInitializer& ObjectInitializer)
 	, bEnable2DPhysics(false)
 	, LockedAxis_DEPRECATED(ESettingsLockedAxis::Invalid)
 	, BounceThresholdVelocity(200.f)
+	, MaxAngularVelocity(3600)	//10 revolutions per second
 	, bSimulateSkeletalMeshOnDedicatedServer(true)
 	, bDefaultHasComplexCollision(true)
 	, MaxPhysicsDeltaTime(1.f / 30.f)
