BP: Added DateFromIsoString, DateFromString and TimespanFromString BP functions

Based on GitHub PR #1082 (https://github.com/EpicGames/UnrealEngine/pull/1082/files)

#github 1082

diff --git a/Engine/Source/Runtime/Engine/Classes/Kismet/KismetMathLibrary.h b/Engine/Source/Runtime/Engine/Classes/Kismet/KismetMathLibrary.h
--- a/Engine/Source/Runtime/Engine/Classes/Kismet/KismetMathLibrary.h
+++ b/Engine/Source/Runtime/Engine/Classes/Kismet/KismetMathLibrary.h
@@ -953,6 +953,14 @@ class ENGINE_API UKismetMathLibrary : public UBlueprintFunctionLibrary
 	UFUNCTION(BlueprintPure, meta=(DisplayName="UtcNow"), Category="Math|DateTime")
 	static FDateTime UtcNow( );
 
+	/* Converts a date string in ISO-8601 format to a DateTime object */
+	UFUNCTION(BlueprintPure, Category="Math|DateTime")
+	static bool DateTimeFromIsoString(FString IsoString, FDateTime& Result);
+
+	/* Converts a date string to a DateTime object */
+	UFUNCTION(BlueprintPure, Category="Math|DateTime")
+	static bool DateTimeFromString(FString DateTimeString, FDateTime& Result);
+
 	//
 	// Timespan functions.
 	//
@@ -1081,6 +1089,11 @@ class ENGINE_API UKismetMathLibrary : public UBlueprintFunctionLibrary
 	UFUNCTION(BlueprintPure, meta=(DisplayName="ZeroValue"), Category="Math|Timespan")
 	static FTimespan TimespanZeroValue( );
 
+	/* Converts a time span string to a Timespan object */
+	UFUNCTION(BlueprintPure, Category="Math|Timespan")
+	static bool TimespanFromString(FString TimespanString, FTimespan& Result);
+
+
 	// -- Begin K2 utilities
 
 	/** Converts a byte to a float */
diff --git a/Engine/Source/Runtime/Engine/Private/KismetMathLibrary.cpp b/Engine/Source/Runtime/Engine/Private/KismetMathLibrary.cpp
--- a/Engine/Source/Runtime/Engine/Private/KismetMathLibrary.cpp
+++ b/Engine/Source/Runtime/Engine/Private/KismetMathLibrary.cpp
@@ -1422,6 +1422,18 @@ FDateTime UKismetMathLibrary::UtcNow( )
 }
 
 
+bool UKismetMathLibrary::DateTimeFromIsoString(FString IsoString, FDateTime& Result)
+{
+	return FDateTime::ParseIso8601(*IsoString, Result);
+}
+
+
+bool UKismetMathLibrary::DateTimeFromString(FString DateTimeString, FDateTime& Result)
+{
+	return FDateTime::Parse(DateTimeString, Result);
+}
+
+
 /* Timespan functions
  *****************************************************************************/
 
@@ -1618,6 +1630,12 @@ FTimespan UKismetMathLibrary::TimespanZeroValue( )
 }
 
 
+bool UKismetMathLibrary::TimespanFromString(FString TimespanString, FTimespan& Result)
+{
+	return FTimespan::Parse(TimespanString, Result);
+}
+
+
 /* K2 Utilities
  *****************************************************************************/
 
