PR #1352: Fix "Access violation" crash on running game without Audio hardware (Contributed by [REDACTED_PERSON])

diff --git a/Engine/Source/Runtime/Engine/Private/AudioDevice.cpp b/Engine/Source/Runtime/Engine/Private/AudioDevice.cpp
--- a/Engine/Source/Runtime/Engine/Private/AudioDevice.cpp
+++ b/Engine/Source/Runtime/Engine/Private/AudioDevice.cpp
@@ -226,7 +226,10 @@ void FAudioDevice::AddReferencedObjects( FReferenceCollector& Collector )
 		Collector.AddReferencedObject(It.Key());
 	}
 
-	Effects->AddReferencedObjects(Collector);
+	if (Effects)
+	{
+		Effects->AddReferencedObjects(Collector);
+	}
 
 	for( int32 i = 0; i < ActiveSounds.Num(); ++i )
 	{
