Merging 2613006

//depot/UE4-UT/...

to //depot/UE4/...

diff --git a/Engine/Source/Runtime/Engine/Private/GameState.cpp b/Engine/Source/Runtime/Engine/Private/GameState.cpp
--- a/Engine/Source/Runtime/Engine/Private/GameState.cpp
+++ b/Engine/Source/Runtime/Engine/Private/GameState.cpp
@@ -55,7 +55,7 @@ void AGameState::DefaultTimer()
 		}
 	}
 
-	GetWorldTimerManager().SetTimer(TimerHandle_DefaultTimer, this, &AGameState::DefaultTimer, GetWorldSettings()->GetEffectiveTimeDilation(), true);
+	GetWorldTimerManager().SetTimer(TimerHandle_DefaultTimer, this, &AGameState::DefaultTimer, 1.0f / GetWorldSettings()->GetEffectiveTimeDilation(), true);
 }
 
 void AGameState::PostInitializeComponents()
@@ -63,7 +63,7 @@ void AGameState::PostInitializeComponents()
 	Super::PostInitializeComponents();
 
 	FTimerManager& TimerManager = GetWorldTimerManager();
-	TimerManager.SetTimer(TimerHandle_DefaultTimer, this, &AGameState::DefaultTimer, GetWorldSettings()->GetEffectiveTimeDilation(), true);
+	TimerManager.SetTimer(TimerHandle_DefaultTimer, this, &AGameState::DefaultTimer, 1.0f / GetWorldSettings()->GetEffectiveTimeDilation(), true);
 
 	UWorld* World = GetWorld();
 	World->GameState = this;
