Blueprints: Expose some Vector : Integer arithmetic methods to Blueprints (Contributed by [REDACTED_PERSON])
PR #1220

diff --git a/Engine/Source/Runtime/Engine/Classes/Kismet/KismetMathLibrary.h b/Engine/Source/Runtime/Engine/Classes/Kismet/KismetMathLibrary.h
--- a/Engine/Source/Runtime/Engine/Classes/Kismet/KismetMathLibrary.h
+++ b/Engine/Source/Runtime/Engine/Classes/Kismet/KismetMathLibrary.h
@@ -554,6 +554,10 @@ class ENGINE_API UKismetMathLibrary : public UBlueprintFunctionLibrary
 	/* Scales Vector A by B */
 	UFUNCTION(BlueprintPure, meta=(DisplayName = "vector * float", CompactNodeTitle = "*", Keywords = "* multiply"), Category="Math|Vector")
 	static FVector Multiply_VectorFloat(FVector A, float B);
+	
+	/* Scales Vector A by B */
+	UFUNCTION(BlueprintPure, meta=(DisplayName = "vector * int", CompactNodeTitle = "*", Keywords = "* multiply"), Category="Math|Vector")
+	static FVector Multiply_VectorInt(FVector A, int32 B);
 
 	/* Element-wise Vector multiplication (Result = {A.x*B.x, A.y*B.y, A.z*B.z}) */
 	UFUNCTION(BlueprintPure, meta=(DisplayName = "vector * vector", CompactNodeTitle = "*", Keywords = "* multiply", CommutativeAssociativeBinaryOperator = "true"), Category="Math|Vector")
@@ -563,6 +567,10 @@ class ENGINE_API UKismetMathLibrary : public UBlueprintFunctionLibrary
 	UFUNCTION(BlueprintPure, meta=(DisplayName = "vector / float", CompactNodeTitle = "/", Keywords = "/ divide division"), Category="Math|Vector")
 	static FVector Divide_VectorFloat(FVector A, float B = 1.f);
 
+	/* Vector divide by an integer */
+	UFUNCTION(BlueprintPure, meta=(DisplayName = "vector / int", CompactNodeTitle = "/", Keywords = "/ divide division"), Category="Math|Vector")
+	static FVector Divide_VectorInt(FVector A, int32 B = 1);
+	
 	/* Vector divide by vector */
 	UFUNCTION(BlueprintPure, meta=(DisplayName = "vector / vector", CompactNodeTitle = "/", Keywords = "/ divide division"), Category="Math|Vector")
 	static FVector Divide_VectorVector(FVector A, FVector B = FVector(1.f,1.f,1.f));
@@ -571,16 +579,26 @@ class ENGINE_API UKismetMathLibrary : public UBlueprintFunctionLibrary
 	UFUNCTION(BlueprintPure, meta=(DisplayName = "vector + vector", CompactNodeTitle = "+", Keywords = "+ add plus", CommutativeAssociativeBinaryOperator = "true"), Category="Math|Vector")
 	static FVector Add_VectorVector(FVector A, FVector B);
 
+	/** Adds a float to each component of a vector */
 	UFUNCTION(BlueprintPure, meta=(DisplayName = "vector + float", CompactNodeTitle = "+", Keywords = "+ add plus"), Category="Math|Vector")
 	static FVector Add_VectorFloat(FVector A, float B);
+	
+	/** Adds an integer to each component of a vector */
+	UFUNCTION(BlueprintPure, meta=(DisplayName = "vector + int", CompactNodeTitle = "+", Keywords = "+ add plus"), Category="Math|Vector")
+	static FVector Add_VectorInt(FVector A, int32 B);
 
 	/* Vector subtraction */
 	UFUNCTION(BlueprintPure, meta=(DisplayName = "vector - vector", CompactNodeTitle = "-", Keywords = "- subtract minus"), Category="Math|Vector")
 	static FVector Subtract_VectorVector(FVector A, FVector B);
 
+	/** Subtracts a float from each component of a vector */
 	UFUNCTION(BlueprintPure, meta=(DisplayName = "vector - float", CompactNodeTitle = "-", Keywords = "- subtract minus"), Category="Math|Vector")
 	static FVector Subtract_VectorFloat(FVector A, float B);
 
+	/** Subtracts an integer from each component of a vector */
+	UFUNCTION(BlueprintPure, meta=(DisplayName = "vector - int", CompactNodeTitle = "-", Keywords = "- subtract minus"), Category="Math|Vector")
+	static FVector Subtract_VectorInt(FVector A, int32 B);
+
 	/* Returns result of vector A rotated by the inverse of Rotator B */
 	UFUNCTION(BlueprintPure, meta=(DisplayName = "UnrotateVector"), Category="Math|Vector")
 	static FVector LessLess_VectorRotator(FVector A, FRotator B);
@@ -729,17 +747,21 @@ class ENGINE_API UKismetMathLibrary : public UBlueprintFunctionLibrary
 	// Rotator functions.
 	//
 
-	//Returns true if rotator A is equal to rotator B (A == B) within a specified error tolerance
+	// Returns true if rotator A is equal to rotator B (A == B) within a specified error tolerance
 	UFUNCTION(BlueprintPure, meta=(DisplayName = "Equal (Rotator)", CompactNodeTitle = "==", Keywords = "== equal"), Category="Math|Rotator")
 	static bool EqualEqual_RotatorRotator(FRotator A, FRotator B, float ErrorTolerance = 1.e-4f);
 
-	//Returns true if rotator A is not equal to rotator B (A != B) within a specified error tolerance
+	// Returns true if rotator A is not equal to rotator B (A != B) within a specified error tolerance
 	UFUNCTION(BlueprintPure, meta=(DisplayName = "Not Equal (Rotator)", CompactNodeTitle = "!=", Keywords = "!= not equal"), Category="Math|Rotator")
 	static bool NotEqual_RotatorRotator(FRotator A, FRotator B, float ErrorTolerance = 1.e-4f);
 
-	//Returns rotator representing rotator A scaled by B 
+	// Returns rotator representing rotator A scaled by B 
 	UFUNCTION(BlueprintPure, meta=(DisplayName = "ScaleRotator", CompactNodeTitle = "*", Keywords = "* multiply rotate rotation"), Category="Math|Rotator")
 	static FRotator Multiply_RotatorFloat(FRotator A, float B);
+	
+	// Returns rotator representing rotator A scaled by B 
+	UFUNCTION(BlueprintPure, meta=(DisplayName = "ScaleRotator (int)", CompactNodeTitle = "*", Keywords = "* multiply rotate rotation"), Category="Math|Rotator")
+	static FRotator Multiply_RotatorInt(FRotator A, int32 B);
 
 	/** Combine 2 rotations to give you the resulting rotation */
 	UFUNCTION(BlueprintPure, meta=(DisplayName = "CombineRotators", Keywords="rotate rotation add"), Category="Math|Rotator")
diff --git a/Engine/Source/Runtime/Engine/Private/KismetMathLibrary.cpp b/Engine/Source/Runtime/Engine/Private/KismetMathLibrary.cpp
--- a/Engine/Source/Runtime/Engine/Private/KismetMathLibrary.cpp
+++ b/Engine/Source/Runtime/Engine/Private/KismetMathLibrary.cpp
@@ -686,7 +686,12 @@ float UKismetMathLibrary::FInterpTo_Constant(float Current, float Target, float
 FVector UKismetMathLibrary::Multiply_VectorFloat(FVector A, float B)
 {
 	return A * B;
-}	
+}
+
+FVector UKismetMathLibrary::Multiply_VectorInt(FVector A, int32 B)
+{
+	return A * (float)B;
+}
 
 FVector UKismetMathLibrary::Multiply_VectorVector(FVector A, FVector B)
 {
@@ -702,8 +707,20 @@ FVector UKismetMathLibrary::Divide_VectorFloat(FVector A, float B)
 		return FVector::ZeroVector;
 	}
 
-	return A/B;
-}	
+	return A / B;
+}
+
+FVector UKismetMathLibrary::Divide_VectorInt(FVector A, int32 B)
+{
+	if (B == 0)
+	{
+		//@TODO: EXCEPTION: Throw script exception 
+		FFrame::KismetExecutionMessage(TEXT("Divide by zero: Divide_VectorInt"), ELogVerbosity::Warning);
+		return FVector::ZeroVector;
+	}
+
+	return A / (float)B;
+}
 
 FVector UKismetMathLibrary::Divide_VectorVector(FVector A, FVector B)
 {
@@ -728,6 +745,11 @@ FVector UKismetMathLibrary::Add_VectorFloat(FVector A, float B)
 	return A + B;
 }	
 
+FVector UKismetMathLibrary::Add_VectorInt(FVector A, int32 B)
+{
+	return A + (float)B;
+}
+
 FVector UKismetMathLibrary::Subtract_VectorVector(FVector A, FVector B)
 {
 	return A - B;
@@ -736,7 +758,12 @@ FVector UKismetMathLibrary::Subtract_VectorVector(FVector A, FVector B)
 FVector UKismetMathLibrary::Subtract_VectorFloat(FVector A, float B)
 {
 	return A - B;
-}	
+}
+
+FVector UKismetMathLibrary::Subtract_VectorInt(FVector A, int32 B)
+{
+	return A - (float)B;
+}
 
 FVector UKismetMathLibrary::LessLess_VectorRotator(FVector A, FRotator B)
 {
@@ -957,11 +984,16 @@ bool UKismetMathLibrary::EqualEqual_RotatorRotator(FRotator A, FRotator B, float
 bool UKismetMathLibrary::NotEqual_RotatorRotator(FRotator A, FRotator B, float ErrorTolerance)
 {
 	return !A.Equals(B, ErrorTolerance);
-}	
+}
 
 FRotator UKismetMathLibrary::Multiply_RotatorFloat(FRotator A, float B)
 {
 	return A * B;
+}
+
+FRotator UKismetMathLibrary::Multiply_RotatorInt(FRotator A, int32 B)
+{
+	return A * (float)B;
 }	
 
 FRotator UKismetMathLibrary::ComposeRotators(FRotator A, FRotator B)
