PR #1332 Improve error message: Instanced meshes don't yet support unique static lighting for each LOD, lighting on LOD 1+ may be incorrect (contributed by [REDACTED_PERSON])

diff --git a/Engine/Source/Runtime/Engine/Private/InstancedStaticMesh.cpp b/Engine/Source/Runtime/Engine/Private/InstancedStaticMesh.cpp
--- a/Engine/Source/Runtime/Engine/Private/InstancedStaticMesh.cpp
+++ b/Engine/Source/Runtime/Engine/Private/InstancedStaticMesh.cpp
@@ -1193,7 +1193,9 @@ void UInstancedStaticMeshComponent::GetStaticLightingInfo(FStaticLightingPrimiti
 		// Need to create per-LOD instance data to fix that
 		if (!bCanLODsShareStaticLighting)
 		{
-			UE_LOG(LogStaticMesh, Warning, TEXT("Instanced meshes don't yet support unique static lighting for each LOD, lighting on LOD 1+ may be incorrect"));
+			FMessageLog("LightingResults").Message(EMessageSeverity::Warning)
+				->AddToken(FUObjectToken::Create(this))
+				->AddToken(FTextToken::Create(NSLOCTEXT("InstancedStaticMesh", "UniqueStaticLightingForLODWarning", "Instanced meshes don't yet support unique static lighting for each LOD, lighting on LOD 1+ may be incorrect")));
 			bCanLODsShareStaticLighting = true;
 		}
 
