#ue4 - Tighten tolerance on rotation changes in SceneComponent, so we push full transform updates when smaller changes occur.

- FRotator versions of API functions check the rotator versus RelativeRotation when possible, to prevent FRotator<->FQuat conversions from triggering unnecessary updates.

diff --git a/Engine/Source/Runtime/Engine/Classes/Components/SceneComponent.h b/Engine/Source/Runtime/Engine/Classes/Components/SceneComponent.h
--- a/Engine/Source/Runtime/Engine/Classes/Components/SceneComponent.h
+++ b/Engine/Source/Runtime/Engine/Classes/Components/SceneComponent.h
@@ -69,6 +69,9 @@ enum EMoveComponentFlags
 	MOVECOMP_NeverIgnoreBlockingOverlaps= 0x0004,	// never ignore initial blocking overlaps during movement, which are usually ignored when moving out of an object. MOVECOMP_IgnoreBases is still respected.
 };
 
+#define SCENECOMPONENT_QUAT_TOLERANCE		(1.e-8f) // Comparison tolerance for checking if two FQuats are the same when moving SceneComponents.
+#define SCENECOMPONENT_ROTATOR_TOLERANCE	(1.e-4f) // Comparison tolerance for checking if two FRotators are the same when moving SceneComponents.
+
 FORCEINLINE EMoveComponentFlags operator|(EMoveComponentFlags Arg1,EMoveComponentFlags Arg2)	{ return EMoveComponentFlags(uint32(Arg1) | uint32(Arg2)); }
 FORCEINLINE EMoveComponentFlags operator&(EMoveComponentFlags Arg1,EMoveComponentFlags Arg2)	{ return EMoveComponentFlags(uint32(Arg1) & uint32(Arg2)); }
 FORCEINLINE void operator&=(EMoveComponentFlags& Dest,EMoveComponentFlags Arg)					{ Dest = EMoveComponentFlags(Dest & Arg); }
@@ -886,31 +889,16 @@ FORCEINLINE_DEBUGGABLE bool USceneComponent::MoveComponent(const FVector& Delta,
 	return MoveComponentImpl(Delta, NewRotation, bSweep, Hit, MoveFlags, Teleport);
 }
 
-FORCEINLINE_DEBUGGABLE bool USceneComponent::MoveComponent(const FVector& Delta, const FRotator& NewRotation, bool bSweep, FHitResult* Hit, EMoveComponentFlags MoveFlags, ETeleportType Teleport)
-{
-	return MoveComponentImpl(Delta, NewRotation.Quaternion(), bSweep, Hit, MoveFlags, Teleport);
-}
-
 FORCEINLINE_DEBUGGABLE void USceneComponent::SetRelativeLocation(FVector NewLocation, bool bSweep, FHitResult* OutSweepHitResult, ETeleportType Teleport)
 {
 	SetRelativeLocationAndRotation(NewLocation, RelativeRotationCache.RotatorToQuat(RelativeRotation), bSweep, OutSweepHitResult, Teleport);
 }
 
-FORCEINLINE_DEBUGGABLE void USceneComponent::SetRelativeRotation(FRotator NewRotation, bool bSweep, FHitResult* OutSweepHitResult, ETeleportType Teleport)
-{
-	SetRelativeLocationAndRotation(RelativeLocation, NewRotation, bSweep, OutSweepHitResult, Teleport);
-}
-
 FORCEINLINE_DEBUGGABLE void USceneComponent::SetRelativeRotation(const FQuat& NewRotation, bool bSweep, FHitResult* OutSweepHitResult, ETeleportType Teleport)
 {
 	SetRelativeLocationAndRotation(RelativeLocation, NewRotation, bSweep, OutSweepHitResult, Teleport);
 }
 
-FORCEINLINE_DEBUGGABLE void USceneComponent::SetRelativeLocationAndRotation(FVector NewLocation, FRotator NewRotation, bool bSweep, FHitResult* OutSweepHitResult, ETeleportType Teleport)
-{
-	SetRelativeLocationAndRotation(NewLocation, NewRotation.Quaternion(), bSweep, OutSweepHitResult, Teleport);
-}
-
 FORCEINLINE_DEBUGGABLE void USceneComponent::AddRelativeLocation(FVector DeltaLocation, bool bSweep, FHitResult* OutSweepHitResult, ETeleportType Teleport)
 {
 	SetRelativeLocationAndRotation(RelativeLocation + DeltaLocation, RelativeRotationCache.RotatorToQuat(RelativeRotation), bSweep, OutSweepHitResult, Teleport);
@@ -918,12 +906,7 @@ FORCEINLINE_DEBUGGABLE void USceneComponent::AddRelativeLocation(FVector DeltaLo
 
 FORCEINLINE_DEBUGGABLE void USceneComponent::AddRelativeRotation(FRotator DeltaRotation, bool bSweep, FHitResult* OutSweepHitResult, ETeleportType Teleport)
 {
-	SetRelativeLocationAndRotation(RelativeLocation, RelativeRotation + DeltaRotation, bSweep, OutSweepHitResult, Teleport);
-}
-
-FORCEINLINE_DEBUGGABLE void USceneComponent::SetWorldRotation(const FRotator NewRotation, bool bSweep, FHitResult* OutSweepHitResult, ETeleportType Teleport)
-{
-	SetWorldRotation(NewRotation.Quaternion(), bSweep, OutSweepHitResult, Teleport);
+	SetRelativeRotation(RelativeRotation + DeltaRotation, bSweep, OutSweepHitResult, Teleport);
 }
 
 //////////////////////////////////////////////////////////////////////////
diff --git a/Engine/Source/Runtime/Engine/Classes/Engine/EngineTypes.h b/Engine/Source/Runtime/Engine/Classes/Engine/EngineTypes.h
--- a/Engine/Source/Runtime/Engine/Classes/Engine/EngineTypes.h
+++ b/Engine/Source/Runtime/Engine/Classes/Engine/EngineTypes.h
@@ -1083,7 +1083,7 @@ struct ENGINE_API FRotationConversionCache
 	{
 		if (CachedRotator != InRotator)
 		{
-			CachedRotator = InRotator.GetNormalized();;
+			CachedRotator = InRotator.GetNormalized();
 			CachedQuat = CachedRotator.Quaternion();
 		}
 		return CachedQuat;
@@ -1141,6 +1141,18 @@ struct ENGINE_API FRotationConversionCache
 		return InNormalizedQuat.Rotator();
 	}
 
+	/** Return the cached Quat. */
+	FORCEINLINE_DEBUGGABLE FQuat GetCachedQuat() const
+	{
+		return CachedQuat;
+	}
+
+	/** Return the cached Rotator. */
+	FORCEINLINE_DEBUGGABLE FRotator GetCachedRotator() const
+	{
+		return CachedRotator;
+	}
+
 private:
 	mutable FQuat		CachedQuat;		// FQuat matching CachedRotator such that CachedQuat.Rotator() == CachedRotator.
 	mutable FRotator	CachedRotator;	// FRotator matching CachedQuat such that CachedRotator.Quaternion() == CachedQuat.
diff --git a/Engine/Source/Runtime/Engine/Private/Components/PrimitiveComponent.cpp b/Engine/Source/Runtime/Engine/Private/Components/PrimitiveComponent.cpp
--- a/Engine/Source/Runtime/Engine/Private/Components/PrimitiveComponent.cpp
+++ b/Engine/Source/Runtime/Engine/Private/Components/PrimitiveComponent.cpp
@@ -1399,7 +1399,7 @@ bool UPrimitiveComponent::MoveComponentImpl( const FVector& Delta, const FQuat&
 	if (DeltaSizeSq <= MinMovementDistSq)
 	{
 		// Skip if no vector or rotation.
-		if (NewRotationQuat.Equals(InitialRotationQuat, 1.e-6f))
+		if (NewRotationQuat.Equals(InitialRotationQuat, SCENECOMPONENT_QUAT_TOLERANCE))
 		{
 			// copy to optional output param
 			if (OutHit)
diff --git a/Engine/Source/Runtime/Engine/Private/Components/SceneComponent.cpp b/Engine/Source/Runtime/Engine/Private/Components/SceneComponent.cpp
--- a/Engine/Source/Runtime/Engine/Private/Components/SceneComponent.cpp
+++ b/Engine/Source/Runtime/Engine/Private/Components/SceneComponent.cpp
@@ -733,6 +733,31 @@ void USceneComponent::SetRelativeLocationAndRotation(FVector NewLocation, const
 	MoveComponent(DesiredDelta, DesiredWorldTransform.GetRotation(), bSweep, OutSweepHitResult, MOVECOMP_NoFlags, Teleport);
 }
 
+// The FRotator version. It could be a simple wrapper to the FQuat version but it tries to avoid FQuat conversion if possible because:
+// (a) conversions affect rotation equality tests so that SetRotation() calls with the same FRotator can cause unnecessary updates because we think they are different rotations after normalization.
+// (b) conversions are expensive.
+void USceneComponent::SetRelativeLocationAndRotation(FVector NewLocation, FRotator NewRotation, bool bSweep, FHitResult* OutSweepHitResult, ETeleportType Teleport)
+{
+	if (NewLocation != RelativeLocation)
+	{
+		// It's possible that NewRotation == RelativeRotation, so check the cache for a Rotator->Quat conversion.
+		SetRelativeLocationAndRotation(NewLocation, RelativeRotationCache.RotatorToQuat_ReadOnly(NewRotation), bSweep, OutSweepHitResult, Teleport);
+	}
+	else if (!NewRotation.Equals(RelativeRotation, SCENECOMPONENT_ROTATOR_TOLERANCE))
+	{
+		// We know the rotations are different, don't bother with the cache.
+		SetRelativeLocationAndRotation(NewLocation, NewRotation.Quaternion(), bSweep, OutSweepHitResult, Teleport);
+	}
+}
+
+void USceneComponent::SetRelativeRotation(FRotator NewRotation, bool bSweep, FHitResult* OutSweepHitResult, ETeleportType Teleport)
+{
+	if (!NewRotation.Equals(RelativeRotation, SCENECOMPONENT_ROTATOR_TOLERANCE))
+	{
+		// We know the rotations are different, don't bother with the cache.
+		SetRelativeLocationAndRotation(RelativeLocation, NewRotation.Quaternion(), bSweep, OutSweepHitResult, Teleport);
+	}
+}
 
 void USceneComponent::AddRelativeRotation(const FQuat& DeltaRotation, bool bSweep, FHitResult* OutSweepHitResult, ETeleportType Teleport)
 {
@@ -870,6 +895,19 @@ void USceneComponent::SetWorldRotation(const FQuat& NewRotation, bool bSweep, FH
 	SetRelativeRotation(NewRelRotation, bSweep, OutSweepHitResult, Teleport);
 }
 
+void USceneComponent::SetWorldRotation(FRotator NewRotation, bool bSweep, FHitResult* OutSweepHitResult, ETeleportType Teleport)
+{
+	if (AttachParent == nullptr)
+	{
+		// No parent, relative == world. Use FRotator version because it can check for rotation change without conversion issues.
+		SetRelativeRotation(NewRotation, bSweep, OutSweepHitResult, Teleport);
+	}
+	else
+	{
+		SetWorldRotation(NewRotation.Quaternion(), bSweep, OutSweepHitResult, Teleport);
+	}
+}
+
 
 void USceneComponent::SetWorldScale3D(FVector NewScale)
 {
@@ -918,7 +956,15 @@ void USceneComponent::SetWorldTransform(const FTransform& NewTransform, bool bSw
 
 void USceneComponent::SetWorldLocationAndRotation(FVector NewLocation, FRotator NewRotation, bool bSweep, FHitResult* OutSweepHitResult, ETeleportType Teleport)
 {
-	SetWorldLocationAndRotation(NewLocation, NewRotation.Quaternion(), bSweep, OutSweepHitResult, Teleport);
+	if (AttachParent == nullptr)
+	{
+		// No parent, relative == world. Use FRotator version because it can check for rotation change without conversion issues.
+		SetRelativeLocationAndRotation(NewLocation, NewRotation, bSweep, OutSweepHitResult, Teleport);
+	}
+	else
+	{
+		SetWorldLocationAndRotation(NewLocation, NewRotation.Quaternion(), bSweep, OutSweepHitResult, Teleport);
+	}
 }
 
 void USceneComponent::SetWorldLocationAndRotation(FVector NewLocation, const FQuat& NewRotation, bool bSweep, FHitResult* OutSweepHitResult, ETeleportType Teleport)
@@ -1843,11 +1889,11 @@ bool USceneComponent::InternalSetWorldLocationAndRotation(FVector NewLocation, c
 		}
 	}
 
-	if (!NewLocation.Equals(RelativeLocation) || !NewRotationQuat.Equals(RelativeRotationCache.RotatorToQuat_ReadOnly(RelativeRotation), 1e-6f))
+	if (!NewLocation.Equals(RelativeLocation) || !NewRotationQuat.Equals(RelativeRotationCache.RotatorToQuat_ReadOnly(RelativeRotation), SCENECOMPONENT_QUAT_TOLERANCE))
 	{
 		RelativeLocation = NewLocation;
-		RelativeRotation = RelativeRotationCache.QuatToRotator(NewRotationQuat); // Normalizes rotator, if this is a new rotation. Then we'll use it below.
-		UpdateComponentToWorldWithParent(AttachParent, bNoPhysics, RelativeRotationCache.RotatorToQuat(RelativeRotation), Teleport);
+		RelativeRotation = RelativeRotationCache.QuatToRotator(NewRotationQuat); // Normalizes quat, if this is a new rotation. Then we'll use it below.
+		UpdateComponentToWorldWithParent(AttachParent, bNoPhysics, RelativeRotationCache.GetCachedQuat(), Teleport);
 
 		PostUpdateNavigationData();
 		return true;
@@ -1907,6 +1953,29 @@ bool USceneComponent::CheckStaticMobilityAndWarn(const FText& ActionText) const
 	return false;
 }
 
+
+// FRotator version. This could be a simple wrapper to the FQuat version, but in the case of no significant change in location or rotation (as FRotator),
+// we avoid passing through to the FQuat version because conversion can generate a false negative for the rotation equality comparison done using a strict tolerance.
+bool USceneComponent::MoveComponent(const FVector& Delta, const FRotator& NewRotation, bool bSweep, FHitResult* Hit, EMoveComponentFlags MoveFlags, ETeleportType Teleport)
+{
+	if (AttachParent == nullptr)
+	{
+		if (Delta.IsZero() && NewRotation.Equals(RelativeRotation, SCENECOMPONENT_ROTATOR_TOLERANCE))
+		{
+			if (Hit)
+			{
+				*Hit = FHitResult();
+			}
+			return true;
+		}
+
+		return MoveComponentImpl(Delta, RelativeRotationCache.RotatorToQuat_ReadOnly(NewRotation), bSweep, Hit, MoveFlags, Teleport);
+	}
+
+	return MoveComponentImpl(Delta, NewRotation.Quaternion(), bSweep, Hit, MoveFlags, Teleport);
+}
+
+
 bool USceneComponent::MoveComponentImpl(const FVector& Delta, const FQuat& NewRotation, bool bSweep, FHitResult* OutHit, EMoveComponentFlags MoveFlags, ETeleportType Teleport)
 {
 	SCOPE_CYCLE_COUNTER(STAT_MoveComponentSceneComponentTime);
@@ -1933,7 +2002,7 @@ bool USceneComponent::MoveComponentImpl(const FVector& Delta, const FQuat& NewRo
 	if( Delta.IsZero() )
 	{
 		// Skip if no vector or rotation.
-		if (NewRotation.Equals(ComponentToWorld.GetRotation(), 1.e-6f))
+		if (NewRotation.Equals(ComponentToWorld.GetRotation(), SCENECOMPONENT_QUAT_TOLERANCE))
 		{
 			return true;
 		}
