#ue4 - Tighten tolerance on rotation changes in CharacterMovement when determing if the object we stand on has rotated. Picks up slower rotation rates at higher framerates.

diff --git a/Engine/Source/Runtime/Engine/Private/Components/CharacterMovementComponent.cpp b/Engine/Source/Runtime/Engine/Private/Components/CharacterMovementComponent.cpp
--- a/Engine/Source/Runtime/Engine/Private/Components/CharacterMovementComponent.cpp
+++ b/Engine/Source/Runtime/Engine/Private/Components/CharacterMovementComponent.cpp
@@ -1384,7 +1384,7 @@ void UCharacterMovementComponent::UpdateBasedMovement(float DeltaSeconds)
 	}
 
 	// Find change in rotation
-	const bool bRotationChanged = !OldBaseQuat.Equals(NewBaseQuat, 1e-6f);
+	const bool bRotationChanged = !OldBaseQuat.Equals(NewBaseQuat, 1e-8f);
 	if (bRotationChanged)
 	{
 		DeltaQuat = NewBaseQuat * OldBaseQuat.Inverse();
@@ -1417,7 +1417,7 @@ void UCharacterMovementComponent::UpdateBasedMovement(float DeltaSeconds)
 				CharacterOwner->FaceRotation(TargetRotator, 0.f);
 				FinalQuat = UpdatedComponent->GetComponentQuat();
 
-				if (PawnOldQuat.Equals(FinalQuat))
+				if (PawnOldQuat.Equals(FinalQuat, 1e-6f))
 				{
 					// Nothing changed. This means we probably are using another rotation mechanism (bOrientToMovement etc). We should still follow the base object.
 					// @todo: This assumes only Yaw is used, currently a valid assumption. This is the only reason FaceRotation() is used above really, aside from being a virtual hook.
