Fixed GitHub #1363 : With undefined #WITH_CLOTH_COLLISION_DETECTION the engine won't compile

diff --git a/Engine/Source/Runtime/Engine/Private/SkeletalMeshComponentPhysics.cpp b/Engine/Source/Runtime/Engine/Private/SkeletalMeshComponentPhysics.cpp
--- a/Engine/Source/Runtime/Engine/Private/SkeletalMeshComponentPhysics.cpp
+++ b/Engine/Source/Runtime/Engine/Private/SkeletalMeshComponentPhysics.cpp
@@ -2171,11 +2171,13 @@ bool USkeletalMeshComponent::CreateClothingActor(int32 AssetIndex, physx::apex::
 	ClothingActor->forcePhysicalLod(1); // 1 will be changed to "GetActivePhysicalLod()" later
 	ClothingActor->setFrozen(false);
 
+#if WITH_CLOTH_COLLISION_DETECTION
 	// process clothing collisions once for the case that this component doesn't move
 	if(bCollideWithEnvironment)
 	{
 		ProcessClothCollisionWithEnvironment();
 	}
+#endif // WITH_CLOTH_COLLISION_DETECTION
 
 	return true;
 }
