Add back EditAnywhere, stops editor crashing when generating a details panel for static meshes. This is temporary and needs to be fixed properly though

diff --git a/Engine/Source/Runtime/Engine/Classes/PhysicsEngine/BodyInstance.h b/Engine/Source/Runtime/Engine/Classes/PhysicsEngine/BodyInstance.h
--- a/Engine/Source/Runtime/Engine/Classes/PhysicsEngine/BodyInstance.h
+++ b/Engine/Source/Runtime/Engine/Classes/PhysicsEngine/BodyInstance.h
@@ -208,7 +208,7 @@ public:
 	uint32 bSimulatePhysics:1;
 
 	/** If true, mass will not be automatically computed and you must set it directly */
-	UPROPERTY(meta=(DisplayName="Override"))
+	UPROPERTY(EditAnywhere, meta=(DisplayName="Override"))
 	uint32 bOverrideMass : 1;
 
 	/**Mass of the body in KG. By default we compute this based on physical material and mass scale.
