UE-6816 Remove ShapeMaterial from ShapeComponent. Only did something on SphereComponent, and didn't work correctly there.

diff --git a/Engine/Source/Runtime/Engine/Classes/Components/ShapeComponent.h b/Engine/Source/Runtime/Engine/Classes/Components/ShapeComponent.h
--- a/Engine/Source/Runtime/Engine/Classes/Components/ShapeComponent.h
+++ b/Engine/Source/Runtime/Engine/Classes/Components/ShapeComponent.h
@@ -20,9 +20,6 @@ class ENGINE_API UShapeComponent : public UPrimitiveComponent
 	UPROPERTY(transient, duplicatetransient)
 	class UBodySetup* ShapeBodySetup;
 
-	UPROPERTY(EditAnywhere, AdvancedDisplay, BlueprintReadWrite, Category=Shape)
-	class UMaterialInterface* ShapeMaterial;
-
 	/** Only show this component if the actor is selected */
 	UPROPERTY()
 	uint32 bDrawOnlyIfSelected:1;
@@ -41,7 +38,6 @@ class ENGINE_API UShapeComponent : public UPrimitiveComponent
 
 	// Begin UPrimitiveComponent interface.
 	virtual FPrimitiveSceneProxy* CreateSceneProxy() override;
-	virtual void GetUsedMaterials( TArray<UMaterialInterface*>& OutMaterials ) const override; 
 	virtual class UBodySetup* GetBodySetup() override;
 	virtual bool DoCustomNavigableGeometryExport(FNavigableGeometryExport& GeomExport) const override;
 	virtual void GetNavigationData(FNavigationRelevantData& Data) const override;
diff --git a/Engine/Source/Runtime/Engine/Private/Components/ShapeComponent.cpp b/Engine/Source/Runtime/Engine/Private/Components/ShapeComponent.cpp
--- a/Engine/Source/Runtime/Engine/Private/Components/ShapeComponent.cpp
+++ b/Engine/Source/Runtime/Engine/Private/Components/ShapeComponent.cpp
@@ -50,11 +50,6 @@ UBodySetup* UShapeComponent::GetBodySetup()
 	return ShapeBodySetup;
 }
 
-void UShapeComponent::GetUsedMaterials( TArray<UMaterialInterface*>& OutMaterials ) const
-{
-	OutMaterials.Add( ShapeMaterial );
-}
-
 #if WITH_EDITOR
 void UShapeComponent::PostEditChangeProperty(FPropertyChangedEvent& PropertyChangedEvent)
 {
diff --git a/Engine/Source/Runtime/Engine/Private/Components/SphereComponent.cpp b/Engine/Source/Runtime/Engine/Private/Components/SphereComponent.cpp
--- a/Engine/Source/Runtime/Engine/Private/Components/SphereComponent.cpp
+++ b/Engine/Source/Runtime/Engine/Private/Components/SphereComponent.cpp
@@ -84,7 +84,6 @@ FPrimitiveSceneProxy* USphereComponent::CreateSceneProxy()
 			:	FPrimitiveSceneProxy(InComponent)
 			,	bDrawOnlyIfSelected( InComponent->bDrawOnlyIfSelected )
 			,	SphereColor(InComponent->ShapeColor)
-			,	ShapeMaterial(InComponent->ShapeMaterial)
 			,	SphereRadius(InComponent->SphereRadius)
 		{
 			bWillEverBeLit = false;
@@ -104,28 +103,21 @@ FPrimitiveSceneProxy* USphereComponent::CreateSceneProxy()
 					FPrimitiveDrawInterface* PDI = Collector.GetPDI(ViewIndex);
 
 					const FMatrix& LocalToWorld = GetLocalToWorld();
-					int32 SphereSides =  FMath::Clamp<int32>(SphereRadius/4.f, 16, 64);
-					if(ShapeMaterial && !View->Family->EngineShowFlags.Wireframe)
-					{
-						GetSphereMesh(LocalToWorld.GetOrigin(), FVector(SphereRadius), SphereSides, SphereSides/2, ShapeMaterial->GetRenderProxy(false), SDPG_World, false, ViewIndex, Collector);
-					}
-					else
-					{
-						const FLinearColor DrawSphereColor = GetViewSelectionColor(SphereColor, *View, IsSelected(), IsHovered(), false, IsIndividuallySelected() );
-
-						float AbsScaleX = LocalToWorld.GetScaledAxis(EAxis::X).Size();
-						float AbsScaleY = LocalToWorld.GetScaledAxis(EAxis::Y).Size();
-						float AbsScaleZ = LocalToWorld.GetScaledAxis(EAxis::Z).Size();
-						float MinAbsScale = FMath::Min3(AbsScaleX, AbsScaleY, AbsScaleZ);
-
-						FVector ScaledX = LocalToWorld.GetUnitAxis(EAxis::X) * MinAbsScale;
-						FVector ScaledY = LocalToWorld.GetUnitAxis(EAxis::Y) * MinAbsScale;
-						FVector ScaledZ = LocalToWorld.GetUnitAxis(EAxis::Z) * MinAbsScale;
-
-						DrawCircle(PDI, LocalToWorld.GetOrigin(), ScaledX, ScaledY, DrawSphereColor, SphereRadius, SphereSides, SDPG_World);
-						DrawCircle(PDI, LocalToWorld.GetOrigin(), ScaledX, ScaledZ, DrawSphereColor, SphereRadius, SphereSides, SDPG_World);
-						DrawCircle(PDI, LocalToWorld.GetOrigin(), ScaledY, ScaledZ, DrawSphereColor, SphereRadius, SphereSides, SDPG_World);
-					}
+					const FLinearColor DrawSphereColor = GetViewSelectionColor(SphereColor, *View, IsSelected(), IsHovered(), false, IsIndividuallySelected() );
+
+					float AbsScaleX = LocalToWorld.GetScaledAxis(EAxis::X).Size();
+					float AbsScaleY = LocalToWorld.GetScaledAxis(EAxis::Y).Size();
+					float AbsScaleZ = LocalToWorld.GetScaledAxis(EAxis::Z).Size();
+					float MinAbsScale = FMath::Min3(AbsScaleX, AbsScaleY, AbsScaleZ);
+
+					FVector ScaledX = LocalToWorld.GetUnitAxis(EAxis::X) * MinAbsScale;
+					FVector ScaledY = LocalToWorld.GetUnitAxis(EAxis::Y) * MinAbsScale;
+					FVector ScaledZ = LocalToWorld.GetUnitAxis(EAxis::Z) * MinAbsScale;
+
+					const int32 SphereSides = FMath::Clamp<int32>(SphereRadius / 4.f, 16, 64);
+					DrawCircle(PDI, LocalToWorld.GetOrigin(), ScaledX, ScaledY, DrawSphereColor, SphereRadius, SphereSides, SDPG_World);
+					DrawCircle(PDI, LocalToWorld.GetOrigin(), ScaledX, ScaledZ, DrawSphereColor, SphereRadius, SphereSides, SDPG_World);
+					DrawCircle(PDI, LocalToWorld.GetOrigin(), ScaledY, ScaledZ, DrawSphereColor, SphereRadius, SphereSides, SDPG_World);
 				}
 			}
 		}
@@ -149,7 +141,6 @@ FPrimitiveSceneProxy* USphereComponent::CreateSceneProxy()
 	private:
 		const uint32				bDrawOnlyIfSelected:1;
 		const FColor				SphereColor;
-		const UMaterialInterface*	ShapeMaterial;
 		const float					SphereRadius;
 	};
 
