Add category back to bOverrideMass, to stop UHT complaining

diff --git a/Engine/Source/Runtime/Engine/Classes/PhysicsEngine/BodyInstance.h b/Engine/Source/Runtime/Engine/Classes/PhysicsEngine/BodyInstance.h
--- a/Engine/Source/Runtime/Engine/Classes/PhysicsEngine/BodyInstance.h
+++ b/Engine/Source/Runtime/Engine/Classes/PhysicsEngine/BodyInstance.h
@@ -208,7 +208,7 @@ public:
 	uint32 bSimulatePhysics:1;
 
 	/** If true, mass will not be automatically computed and you must set it directly */
-	UPROPERTY(EditAnywhere, meta=(DisplayName="Override"))
+	UPROPERTY(EditAnywhere, Category = Physics, meta=(DisplayName="Override"))
 	uint32 bOverrideMass : 1;
 
 	/**Mass of the body in KG. By default we compute this based on physical material and mass scale.
