UE4 - merge cloth enable cvar fix from orion

diff --git a/Engine/Source/Runtime/Engine/Private/SkeletalMeshComponentPhysics.cpp b/Engine/Source/Runtime/Engine/Private/SkeletalMeshComponentPhysics.cpp
--- a/Engine/Source/Runtime/Engine/Private/SkeletalMeshComponentPhysics.cpp
+++ b/Engine/Source/Runtime/Engine/Private/SkeletalMeshComponentPhysics.cpp
@@ -3745,7 +3745,7 @@ void USkeletalMeshComponent::GetUpdateClothSimulationData(TArray<FClothSimulData
 {
 #if WITH_APEX_CLOTHING
 
-	if (CVarEnableClothPhysics.GetValueOnGameThread() == 0)
+	if (CVarEnableClothPhysics.GetValueOnAnyThread() == 0)
 	{
 		return;
 	}
