Fix for montage linkable elements losing linkage on load

diff --git a/Engine/Source/Runtime/Engine/Classes/Animation/AnimLinkableElement.cpp b/Engine/Source/Runtime/Engine/Classes/Animation/AnimLinkableElement.cpp
--- a/Engine/Source/Runtime/Engine/Classes/Animation/AnimLinkableElement.cpp
+++ b/Engine/Source/Runtime/Engine/Classes/Animation/AnimLinkableElement.cpp
@@ -376,3 +376,31 @@ void FAnimLinkableElement::Link(UAnimSequenceBase* AnimObject, float AbsTime, in
 		LinkSequence(AnimObject, AbsTime);
 	}
 }
+
+void FAnimLinkableElement::RefreshSegmentOnLoad()
+{
+	// We only perform this step if we have valid data from a previous link
+	if(LinkedMontage && SegmentIndex != INDEX_NONE && SlotIndex != INDEX_NONE)
+	{
+		if(LinkedMontage->SlotAnimTracks.IsValidIndex(SlotIndex))
+		{
+			FSlotAnimationTrack& Slot = LinkedMontage->SlotAnimTracks[SlotIndex];
+			if(Slot.AnimTrack.AnimSegments.IsValidIndex(SegmentIndex))
+			{
+				FAnimSegment& Segment = Slot.AnimTrack.AnimSegments[SegmentIndex];
+
+				if(Segment.AnimReference == LinkedSequence)
+				{
+					if(CachedLinkMethod == EAnimLinkMethod::Relative)
+					{
+						LinkValue = FMath::Clamp<float>(LinkValue, 0.0f, Segment.GetLength());
+					}
+
+					// Update timing
+					SegmentBeginTime = Segment.StartPos;
+					SegmentLength = Segment.GetLength();
+				}
+			}
+		}
+	}
+}
diff --git a/Engine/Source/Runtime/Engine/Classes/Animation/AnimLinkableElement.h b/Engine/Source/Runtime/Engine/Classes/Animation/AnimLinkableElement.h
--- a/Engine/Source/Runtime/Engine/Classes/Animation/AnimLinkableElement.h
+++ b/Engine/Source/Runtime/Engine/Classes/Animation/AnimLinkableElement.h
@@ -121,6 +121,11 @@ struct FAnimLinkableElement
 	/** Relinks this element if it is outside of its segments */
 	ENGINE_API void ConditionalRelink();
 
+	/** Refreshes the current segment data (Begin time, length etc.) and validate the link time
+	 *  Intended to update the internal state when segment lengths/times could have changed
+	 */
+	ENGINE_API void RefreshSegmentOnLoad();
+
 protected:
 
 	/** Gets the segment in the current montage in the current slot that is at the time of this element */
diff --git a/Engine/Source/Runtime/Engine/Private/Animation/AnimMontage.cpp b/Engine/Source/Runtime/Engine/Private/Animation/AnimMontage.cpp
--- a/Engine/Source/Runtime/Engine/Private/Animation/AnimMontage.cpp
+++ b/Engine/Source/Runtime/Engine/Private/Animation/AnimMontage.cpp
@@ -359,6 +359,7 @@ void UAnimMontage::PostLoad()
 		}
 		else
 		{
+			Composite.RefreshSegmentOnLoad();
 			Composite.LinkMontage(this, Composite.GetTime());
 		}
 	}
