Hid Particle option from the Beam Target and Source modules as it does not work.

diff --git a/Engine/Source/Runtime/Engine/Classes/Particles/Beam/ParticleModuleBeamBase.h b/Engine/Source/Runtime/Engine/Classes/Particles/Beam/ParticleModuleBeamBase.h
--- a/Engine/Source/Runtime/Engine/Classes/Particles/Beam/ParticleModuleBeamBase.h
+++ b/Engine/Source/Runtime/Engine/Classes/Particles/Beam/ParticleModuleBeamBase.h
@@ -29,7 +29,7 @@ enum Beam2SourceTargetMethod
 	/** Particle	- use the particles from a given emitter in the system.		
 	 *	The name of the emitter should be set in <Source/Target>Name.
 	 */
-	PEB2STM_Particle UMETA(DisplayName="Particle"),
+	PEB2STM_Particle UMETA(Hidden),
 	/** Actor		- use the actor as the source/target.
 	 *	The name of the actor should be set in <Source/Target>Name.
 	 */
