Fixed shadowed variable warning.

diff --git a/Engine/Source/Runtime/Engine/Private/Actor.cpp b/Engine/Source/Runtime/Engine/Private/Actor.cpp
--- a/Engine/Source/Runtime/Engine/Private/Actor.cpp
+++ b/Engine/Source/Runtime/Engine/Private/Actor.cpp
@@ -303,12 +303,12 @@ bool AActor::TeleportTo( const FVector& DestLocation, const FRotator& DestRotati
 		if (!bNoCheck && (ActorPrimComp->IsCollisionEnabled() || (bCollideWhenPlacing && (GetNetMode() != NM_Client))) )
 		{
 			// Apply the pivot offset to the desired location
-			FVector PivotOffset = GetRootComponent()->Bounds.Origin - PrevLocation;
-			NewLocation = NewLocation + PivotOffset;
+			FVector Offset = GetRootComponent()->Bounds.Origin - PrevLocation;
+			NewLocation = NewLocation + Offset;
 
 			// check if able to find an acceptable destination for this actor that doesn't embed it in world geometry
 			bTeleportSucceeded = GetWorld()->FindTeleportSpot(this, NewLocation, DestRotation);
-			NewLocation = NewLocation - PivotOffset;
+			NewLocation = NewLocation - Offset;
 		}
 
 		if ( bTeleportSucceeded )
