Don't use the template to generate new actor name if the template is a class default object

diff --git a/Engine/Source/Runtime/Engine/Private/LevelActor.cpp b/Engine/Source/Runtime/Engine/Private/LevelActor.cpp
--- a/Engine/Source/Runtime/Engine/Private/LevelActor.cpp
+++ b/Engine/Source/Runtime/Engine/Private/LevelActor.cpp
@@ -298,7 +298,7 @@ AActor* UWorld::SpawnActor( UClass* Class, FVector const* Location, FRotator con
 	{
 		Template = Class->GetDefaultObject<AActor>();
 	}
-	else if (NewActorName.IsNone())
+	else if (NewActorName.IsNone() && !Template->HasAnyFlags(RF_ClassDefaultObject))
 	{
 		NewActorName = MakeUniqueObjectName(LevelToSpawnIn, Template->GetClass(), *Template->GetFName().GetPlainNameString());
 	}
