Added bDrawShadow option to HUD DebugText so we can set shadow per debug line instead of it only being a per-HUD setting

diff --git a/Engine/Source/Runtime/Engine/Classes/GameFramework/DebugTextInfo.h b/Engine/Source/Runtime/Engine/Classes/GameFramework/DebugTextInfo.h
--- a/Engine/Source/Runtime/Engine/Classes/GameFramework/DebugTextInfo.h
+++ b/Engine/Source/Runtime/Engine/Classes/GameFramework/DebugTextInfo.h
@@ -54,6 +54,10 @@ struct FDebugTextInfo
 	UPROPERTY()
 	uint32 bKeepAttachedToActor:1;
 
+	/** Whether to draw a shadow for the text */
+	UPROPERTY()
+	uint32 bDrawShadow:1;
+
 	/** When we first spawn store off the original actor location for use with bKeepAttachedToActor */
 	UPROPERTY()
 	FVector OrigActorLocation;
@@ -75,6 +79,7 @@ struct FDebugTextInfo
 		, TextColor(ForceInit)
 		, bAbsoluteLocation(false)
 		, bKeepAttachedToActor(false)
+		, bDrawShadow(false)
 		, OrigActorLocation(ForceInit)
 		, Font(NULL)
 		, FontScale(1.0f)
diff --git a/Engine/Source/Runtime/Engine/Classes/GameFramework/HUD.h b/Engine/Source/Runtime/Engine/Classes/GameFramework/HUD.h
--- a/Engine/Source/Runtime/Engine/Classes/GameFramework/HUD.h
+++ b/Engine/Source/Runtime/Engine/Classes/GameFramework/HUD.h
@@ -150,9 +150,10 @@ public:
 	 * @param bKeepAttachedToActor 	if this is true the text will follow the actor, otherwise it will be drawn at the location when the call was made
 	 * @param InFont 				font to use
 	 * @param FontScale 			scale
+	 * @param bDrawShadow 			Draw shadow on this string
 	 */
 	UFUNCTION(reliable, client, SealedEvent)
-	void AddDebugText(const FString& DebugText, AActor* SrcActor = NULL, float Duration = 0, FVector Offset = FVector(ForceInit), FVector DesiredOffset = FVector(ForceInit), FColor TextColor = FColor(ForceInit), bool bSkipOverwriteCheck = false, bool bAbsoluteLocation = false, bool bKeepAttachedToActor = false, UFont* InFont = NULL, float FontScale = 1.0);
+	void AddDebugText(const FString& DebugText, AActor* SrcActor = NULL, float Duration = 0, FVector Offset = FVector(ForceInit), FVector DesiredOffset = FVector(ForceInit), FColor TextColor = FColor(ForceInit), bool bSkipOverwriteCheck = false, bool bAbsoluteLocation = false, bool bKeepAttachedToActor = false, UFont* InFont = NULL, float FontScale = 1.0, bool bDrawShadow = false);
 
 	/**
 	 * Remove all debug strings added via AddDebugText
diff --git a/Engine/Source/Runtime/Engine/Private/DrawDebugHelpers.cpp b/Engine/Source/Runtime/Engine/Private/DrawDebugHelpers.cpp
--- a/Engine/Source/Runtime/Engine/Private/DrawDebugHelpers.cpp
+++ b/Engine/Source/Runtime/Engine/Private/DrawDebugHelpers.cpp
@@ -655,7 +655,7 @@ void DrawDebugCone(const UWorld* InWorld, FVector const& Origin, FVector const&
 	}
 }
 
-void DrawDebugString(const UWorld* InWorld, FVector const& TextLocation, const FString& Text, class AActor* TestBaseActor, FColor const& TextColor, float Duration)
+void DrawDebugString(const UWorld* InWorld, FVector const& TextLocation, const FString& Text, class AActor* TestBaseActor, FColor const& TextColor, float Duration, bool bDrawShadow)
 {
 	// no debug line drawing on dedicated server
 	if (GEngine->GetNetMode(InWorld) != NM_DedicatedServer)
@@ -669,7 +669,7 @@ void DrawDebugString(const UWorld* InWorld, FVector const& TextLocation, const F
 			APlayerController* PlayerController = *Iterator;
 			if (PlayerController->MyHUD)
 			{
-				PlayerController->MyHUD->AddDebugText(Text, BaseAct, Duration, TextLocation, TextLocation, TextColor, true, (TestBaseActor==NULL));
+				PlayerController->MyHUD->AddDebugText(Text, BaseAct, Duration, TextLocation, TextLocation, TextColor, true, (TestBaseActor==NULL), false, nullptr, 1.0f, bDrawShadow);
 			}
 		}
 	}
diff --git a/Engine/Source/Runtime/Engine/Private/HUD.cpp b/Engine/Source/Runtime/Engine/Private/HUD.cpp
--- a/Engine/Source/Runtime/Engine/Private/HUD.cpp
+++ b/Engine/Source/Runtime/Engine/Private/HUD.cpp
@@ -421,10 +421,6 @@ void AHUD::DrawDebugTextList()
 		PlayerOwner->GetPlayerViewPoint(CameraLoc, CameraRot);
 
 		FCanvasTextItem TextItem( FVector2D::ZeroVector, FText::GetEmpty(), GEngine->GetSmallFont(), FLinearColor::White );
-		if (bEnableDebugTextShadow == true)
-		{
-			TextItem.EnableShadow(FLinearColor::Black);
-		}
 		for (int32 Idx = 0; Idx < DebugTextList.Num(); Idx++)
 		{
 			if (DebugTextList[Idx].SrcActor == NULL)
@@ -462,6 +458,15 @@ void AHUD::DrawDebugTextList()
 				}
 			}
 
+			if (bEnableDebugTextShadow || DebugTextList[Idx].bDrawShadow)
+			{
+				TextItem.EnableShadow(FLinearColor::Black);
+			}
+			else
+			{
+				TextItem.DisableShadow();
+			}
+				
 			// don't draw text behind the camera
 			if ( ((WorldTextLoc - CameraLoc) | CameraRot.Vector()) > 0.f )
 			{
@@ -496,7 +501,8 @@ void AHUD::AddDebugText_Implementation(const FString& DebugText,
 										 bool bAbsoluteLocation,
 										 bool bKeepAttachedToActor,
 										 UFont* InFont,
-										 float FontScale
+										 float FontScale,
+										 bool bDrawShadow
 										 )
 {
 	// set a default color
@@ -549,6 +555,7 @@ void AHUD::AddDebugText_Implementation(const FString& DebugText,
 			DebugTextList[Idx].OrigActorLocation = SrcActor->GetActorLocation();
 			DebugTextList[Idx].Font = InFont;
 			DebugTextList[Idx].FontScale = FontScale;
+			DebugTextList[Idx].bDrawShadow = bDrawShadow;
 		}
 	}
 }
diff --git a/Engine/Source/Runtime/Engine/Public/DrawDebugHelpers.h b/Engine/Source/Runtime/Engine/Public/DrawDebugHelpers.h
--- a/Engine/Source/Runtime/Engine/Public/DrawDebugHelpers.h
+++ b/Engine/Source/Runtime/Engine/Public/DrawDebugHelpers.h
@@ -32,7 +32,7 @@ ENGINE_API void DrawDebugCylinder(const UWorld* InWorld, FVector const& Start, F
 ENGINE_API void DrawDebugCone(const UWorld* InWorld, FVector const& Origin, FVector const& Direction, float Length, float AngleWidth, float AngleHeight, int32 NumSides, FColor const& Color, bool bPersistentLines=false, float LifeTime=-1.f, uint8 DepthPriority = 0);
 /** Used by gameplay when defining a cone by a vertical and horizontal dot products. */
 ENGINE_API void DrawDebugAltCone(const UWorld* InWorld, FVector const& Origin, FRotator const& Rotation, float Length, float AngleWidth, float AngleHeight, FColor const& DrawColor, bool bPersistentLines=false, float LifeTime=-1.f, uint8 DepthPriority=0);
-ENGINE_API void DrawDebugString(const UWorld* InWorld, FVector const& TextLocation, const FString& Text, class AActor* TestBaseActor = NULL, FColor const& TextColor = FColor::White, float Duration = -1.000000);
+ENGINE_API void DrawDebugString(const UWorld* InWorld, FVector const& TextLocation, const FString& Text, class AActor* TestBaseActor = NULL, FColor const& TextColor = FColor::White, float Duration = -1.000000, bool bDrawShadow = false);
 ENGINE_API void DrawDebugFrustum(const UWorld* InWorld, const FMatrix& FrustumToWorld, FColor const& Color, bool bPersistentLines = false, float LifeTime=-1.f, uint8 DepthPriority = 0);
 /** Draw a capsule using the LineBatcher */
 ENGINE_API void DrawDebugCapsule(const UWorld* InWorld, FVector const& Center, float HalfHeight, float Radius, const FQuat& Rotation, FColor const& Color, bool bPersistentLines = false, float LifeTime = -1.f, uint8 DepthPriority = 0, float Thickness = 0);
