Add p.VisualizeMovement cvar for showing in-world debug drawing of CharacterMovement state (arrows for velocity/accel show up over character collision capsules)
- virtual UCharacterMovementComponent::VisualizeMovement() can be overridden in child classes to add project-specific debug drawing
- Marked GetMovementName() as const so it can be used in debugging display

diff --git a/Engine/Source/Runtime/Engine/Classes/GameFramework/CharacterMovementComponent.h b/Engine/Source/Runtime/Engine/Classes/GameFramework/CharacterMovementComponent.h
--- a/Engine/Source/Runtime/Engine/Classes/GameFramework/CharacterMovementComponent.h
+++ b/Engine/Source/Runtime/Engine/Classes/GameFramework/CharacterMovementComponent.h
@@ -1301,7 +1301,7 @@ public:
 	virtual void SetUpdatedComponent(USceneComponent* NewUpdatedComponent) override;
 	
 	/** @Return MovementMode string */
-	virtual FString GetMovementName();
+	virtual FString GetMovementName() const;
 
 	/** 
 	 * Add impulse to character. Impulses are accumulated each tick and applied together
@@ -1339,6 +1339,11 @@ public:
 	 */
 	virtual void DisplayDebug(class UCanvas* Canvas, const FDebugDisplayInfo& DebugDisplay, float& YL, float& YPos);
 
+	/**
+	 * Draw in-world debug information for character movement (called with p.VisualizeMovement > 0).
+	 */
+	virtual void VisualizeMovement() const;
+
 	/** Check if swimming pawn just ran into edge of the pool and should jump out. */
 	virtual bool CheckWaterJump(FVector CheckPoint, FVector& WallNormal);
 
diff --git a/Engine/Source/Runtime/Engine/Private/Components/CharacterMovementComponent.cpp b/Engine/Source/Runtime/Engine/Private/Components/CharacterMovementComponent.cpp
--- a/Engine/Source/Runtime/Engine/Private/Components/CharacterMovementComponent.cpp
+++ b/Engine/Source/Runtime/Engine/Private/Components/CharacterMovementComponent.cpp
@@ -112,6 +112,13 @@ static TAutoConsoleVariable<float> CVarNetCorrectionLifetime(
 	TEXT("How long a visualized network correction persists.\n")
 	TEXT("Time in seconds each visualized network correction persists."),
 	ECVF_Cheat);
+
+static TAutoConsoleVariable<int32> CVarVisualizeMovement(
+	TEXT("p.VisualizeMovement"),
+	0,
+	TEXT("Whether to draw in-world debug information for character movement.\n")
+	TEXT("0: Disable, 1: Enable"),
+	ECVF_Cheat);
 #endif // !(UE_BUILD_SHIPPING || UE_BUILD_TEST)
 
 
@@ -979,6 +986,15 @@ void UCharacterMovementComponent::TickComponent(float DeltaTime, enum ELevelTick
 
 		ApplyRepulsionForce(DeltaTime);
 	}
+
+#if !(UE_BUILD_SHIPPING || UE_BUILD_TEST)
+	const bool bVisualizeMovement = CVarVisualizeMovement.GetValueOnGameThread() > 0;
+	if (bVisualizeMovement)
+	{
+		VisualizeMovement();
+	}
+#endif // !(UE_BUILD_SHIPPING || UE_BUILD_TEST)
+
 }
 
 void UCharacterMovementComponent::PreClothTick(float DeltaTime, FCharacterMovementComponentPreClothTickFunction& ThisTickFunction)
@@ -5481,7 +5497,7 @@ void UCharacterMovementComponent::ApplyImpactPhysicsForces(const FHitResult& Imp
 	}
 }
 
-FString UCharacterMovementComponent::GetMovementName()
+FString UCharacterMovementComponent::GetMovementName() const
 {
 	if( CharacterOwner )
 	{
@@ -5542,6 +5558,51 @@ void UCharacterMovementComponent::DisplayDebug(UCanvas* Canvas, const FDebugDisp
 	YPos += YL;
 }
 
+void UCharacterMovementComponent::VisualizeMovement() const
+{
+	if (CharacterOwner == nullptr)
+	{
+		return;
+	}
+
+#if !(UE_BUILD_SHIPPING || UE_BUILD_TEST)
+	const FVector TopOfCapsule = GetActorLocation() + FVector(0.f, 0.f, CharacterOwner->GetSimpleCollisionHalfHeight());
+
+	// Velocity
+	{
+		const FColor DebugColor = FColor::Green;
+		const FVector DebugLocation = TopOfCapsule;
+		DrawDebugDirectionalArrow(GetWorld(), DebugLocation, DebugLocation + Velocity, 
+			100.f, DebugColor, false, -1.f, (uint8)'\000', 10.f);
+
+		FString DebugText = FString::Printf(TEXT("Velocity: %s (Speed: %.2f)"), *Velocity.ToCompactString(), Velocity.Size());
+		DrawDebugString(GetWorld(), DebugLocation + FVector(0.f,0.f,5.f), DebugText, nullptr, DebugColor, 0.f, true);
+	}
+
+	// Acceleration
+	{
+		const FColor DebugColor = FColor::Yellow;
+		const float MaxAccelerationLineLength = 200.f;
+		const float CurrentMaxAccel = GetMaxAcceleration();
+		const float CurrentAccelAsPercentOfMaxAccel = CurrentMaxAccel > 0.f ? Acceleration.Size() / CurrentMaxAccel : 1.f;
+		const FVector DebugLocation = TopOfCapsule + FVector(0.f,0.f,15.f);
+		DrawDebugDirectionalArrow(GetWorld(), DebugLocation, 
+			DebugLocation + Acceleration.GetSafeNormal(SMALL_NUMBER) * CurrentAccelAsPercentOfMaxAccel * MaxAccelerationLineLength, 
+			25.f, DebugColor, false, -1.f, (uint8)'\000', 8.f);
+
+		FString DebugText = FString::Printf(TEXT("Acceleration: %s"), *Acceleration.ToCompactString());
+		DrawDebugString(GetWorld(), DebugLocation + FVector(0.f,0.f,5.f), DebugText, nullptr, DebugColor, 0.f, true);
+	}
+
+	// Movement Mode
+	{
+		const FColor DebugColor = FColor::Blue;
+		const FVector DebugLocation = TopOfCapsule + FVector(0.f,0.f,35.f);
+		FString DebugText = FString::Printf(TEXT("MovementMode: %s"), *GetMovementName());
+		DrawDebugString(GetWorld(), DebugLocation, DebugText, nullptr, DebugColor, 0.f, true);
+	}
+#endif // !(UE_BUILD_SHIPPING || UE_BUILD_TEST)
+}
 
 void UCharacterMovementComponent::ForceReplicationUpdate()
 {
