PR1284 - added GetPreferredLanguages function for use in blueprints
#ios

diff --git a/Engine/Source/Runtime/Core/Private/GenericPlatform/GenericPlatformMisc.cpp b/Engine/Source/Runtime/Core/Private/GenericPlatform/GenericPlatformMisc.cpp
--- a/Engine/Source/Runtime/Core/Private/GenericPlatform/GenericPlatformMisc.cpp
+++ b/Engine/Source/Runtime/Core/Private/GenericPlatform/GenericPlatformMisc.cpp
@@ -905,6 +905,11 @@ const TCHAR* FGenericPlatformMisc::GetEngineMode()
 		TEXT( "Game" );
 }
 
+TArray<FString> FGenericPlatformMisc::GetPreferredLanguages()
+{
+	return TArray<FString>();
+}
+
 void FGenericPlatformMisc::PlatformPreInit()
 {
 	FGenericCrashContext::Initialize();
diff --git a/Engine/Source/Runtime/Core/Private/IOS/IOSPlatformMisc.cpp b/Engine/Source/Runtime/Core/Private/IOS/IOSPlatformMisc.cpp
--- a/Engine/Source/Runtime/Core/Private/IOS/IOSPlatformMisc.cpp
+++ b/Engine/Source/Runtime/Core/Private/IOS/IOSPlatformMisc.cpp
@@ -740,6 +740,18 @@ TArray<uint8> FIOSPlatformMisc::GetSystemFontBytes()
 	return FontBytes;
 }
 
+TArray<FString> FIOSPlatformMisc::GetPreferredLanguages()
+{
+	TArray<FString> Results;
+
+	NSArray* Languages = [NSLocale preferredLanguages];
+	for (NSString* Language in Languages)
+	{
+		Results.Add(FString(Language));
+	}
+	return Results;
+}
+
 void FIOSPlatformMisc::RegisterForRemoteNotifications()
 {
 	UIApplication* application = [UIApplication sharedApplication];
diff --git a/Engine/Source/Runtime/Core/Public/GenericPlatform/GenericPlatformMisc.h b/Engine/Source/Runtime/Core/Public/GenericPlatform/GenericPlatformMisc.h
--- a/Engine/Source/Runtime/Core/Public/GenericPlatform/GenericPlatformMisc.h
+++ b/Engine/Source/Runtime/Core/Public/GenericPlatform/GenericPlatformMisc.h
@@ -826,6 +826,12 @@ public:
 	 */
 	static const TCHAR* GetEngineMode();
 
+	/**
+	 * Returns an array of the user's preferred languages in order of preference
+	 * @return An array of language IDs ordered from most preferred to least
+	 */
+	static TArray<FString> GetPreferredLanguages();
+
 	/**
 	 * Requests permission to send remote notifications to the user's device.
 	 */
diff --git a/Engine/Source/Runtime/Core/Public/IOS/IOSPlatformMisc.h b/Engine/Source/Runtime/Core/Public/IOS/IOSPlatformMisc.h
--- a/Engine/Source/Runtime/Core/Public/IOS/IOSPlatformMisc.h
+++ b/Engine/Source/Runtime/Core/Public/IOS/IOSPlatformMisc.h
@@ -94,6 +94,7 @@ struct CORE_API FIOSPlatformMisc : public FGenericPlatformMisc
 	static bool SetStoredValue(const FString& InStoreId, const FString& InSectionName, const FString& InKeyName, const FString& InValue);
 	static bool GetStoredValue(const FString& InStoreId, const FString& InSectionName, const FString& InKeyName, FString& OutValue);
 	static TArray<uint8> GetSystemFontBytes();
+	static TArray<FString> GetPreferredLanguages();
 
 	static void RegisterForRemoteNotifications();
 
diff --git a/Engine/Source/Runtime/Engine/Classes/Kismet/KismetSystemLibrary.h b/Engine/Source/Runtime/Engine/Classes/Kismet/KismetSystemLibrary.h
--- a/Engine/Source/Runtime/Engine/Classes/Kismet/KismetSystemLibrary.h
+++ b/Engine/Source/Runtime/Engine/Classes/Kismet/KismetSystemLibrary.h
@@ -1333,6 +1333,13 @@ class ENGINE_API UKismetSystemLibrary : public UBlueprintFunctionLibrary
 	UFUNCTION(BlueprintCallable, Category = "Utilities", meta = (HidePin = "WorldContextObject", DefaultToSelf = "WorldContextObject"))
 	static void SetSupressViewportTransitionMessage(UObject* WorldContextObject, bool bState);
 
+	/**
+	 * Returns an array of the user's preferred languages in order of preference
+	 * @return An array of language IDs ordered from most preferred to least
+	 */
+	UFUNCTION(BlueprintCallable, Category = "Utilities|Platform")
+		static TArray<FString> GetPreferredLanguages();
+
 	/**
 	 * Requests permission to send remote notifications to the user's device.
 	 * (iOS only)
diff --git a/Engine/Source/Runtime/Engine/Private/KismetSystemLibrary.cpp b/Engine/Source/Runtime/Engine/Private/KismetSystemLibrary.cpp
--- a/Engine/Source/Runtime/Engine/Private/KismetSystemLibrary.cpp
+++ b/Engine/Source/Runtime/Engine/Private/KismetSystemLibrary.cpp
@@ -2994,6 +2994,12 @@ void UKismetSystemLibrary::SetSupressViewportTransitionMessage(UObject* WorldCon
 		World->GetFirstLocalPlayerFromController()->ViewportClient->SetSuppressTransitionMessage(bState);
 	}
 }
+
+TArray<FString> UKismetSystemLibrary::GetPreferredLanguages()
+{
+	return FPlatformMisc::GetPreferredLanguages();
+}
+
 struct FLoadAssetActionBase : public FPendingLatentAction, public FGCObject
 {
 	// @TODO: it would be good to have static/global manager? 
