Fix crash when updating scale of a static object with async scene enbled.

diff --git a/Engine/Source/Runtime/Engine/Private/PhysicsEngine/BodyInstance.cpp b/Engine/Source/Runtime/Engine/Private/PhysicsEngine/BodyInstance.cpp
--- a/Engine/Source/Runtime/Engine/Private/PhysicsEngine/BodyInstance.cpp
+++ b/Engine/Source/Runtime/Engine/Private/PhysicsEngine/BodyInstance.cpp
@@ -645,6 +645,7 @@ PxShape* ClonePhysXShape_AssumesLocked(PxShape* PShape)
 	PNewShape->setRestOffset(PShape->getRestOffset());
 	PNewShape->setSimulationFilterData(PShape->getSimulationFilterData());
 	PNewShape->setQueryFilterData(PShape->getQueryFilterData());
+	PNewShape->userData = PShape->userData;
 
 	return PNewShape;
 }
