PR1287 - addition of GetLocalCurrecyCode and GetLocalCurrencySymbol
#ios

diff --git a/Engine/Source/Runtime/Core/Private/GenericPlatform/GenericPlatformMisc.cpp b/Engine/Source/Runtime/Core/Private/GenericPlatform/GenericPlatformMisc.cpp
--- a/Engine/Source/Runtime/Core/Private/GenericPlatform/GenericPlatformMisc.cpp
+++ b/Engine/Source/Runtime/Core/Private/GenericPlatform/GenericPlatformMisc.cpp
@@ -907,9 +907,22 @@ const TCHAR* FGenericPlatformMisc::GetEngineMode()
 
 TArray<FString> FGenericPlatformMisc::GetPreferredLanguages()
 {
+	// not implemented by default
 	return TArray<FString>();
 }
 
+FString FGenericPlatformMisc::GetLocalCurrencyCode()
+{
+	// not implemented by default
+	return FString();
+}
+
+FString FGenericPlatformMisc::GetLocalCurrencySymbol()
+{
+	// not implemented by default
+	return FString();
+}
+
 void FGenericPlatformMisc::PlatformPreInit()
 {
 	FGenericCrashContext::Initialize();
diff --git a/Engine/Source/Runtime/Core/Private/IOS/IOSPlatformMisc.cpp b/Engine/Source/Runtime/Core/Private/IOS/IOSPlatformMisc.cpp
--- a/Engine/Source/Runtime/Core/Private/IOS/IOSPlatformMisc.cpp
+++ b/Engine/Source/Runtime/Core/Private/IOS/IOSPlatformMisc.cpp
@@ -752,6 +752,16 @@ TArray<FString> FIOSPlatformMisc::GetPreferredLanguages()
 	return Results;
 }
 
+FString FIOSPlatformMisc::GetLocalCurrencyCode()
+{
+	return FString([[NSLocale currentLocale] objectForKey:NSLocaleCurrencyCode]);
+}
+
+FString FIOSPlatformMisc::GetLocalCurrencySymbol()
+{
+	return FString([[NSLocale currentLocale] objectForKey:NSLocaleCurrencySymbol]);
+}
+
 void FIOSPlatformMisc::RegisterForRemoteNotifications()
 {
 	UIApplication* application = [UIApplication sharedApplication];
diff --git a/Engine/Source/Runtime/Core/Public/GenericPlatform/GenericPlatformMisc.h b/Engine/Source/Runtime/Core/Public/GenericPlatform/GenericPlatformMisc.h
--- a/Engine/Source/Runtime/Core/Public/GenericPlatform/GenericPlatformMisc.h
+++ b/Engine/Source/Runtime/Core/Public/GenericPlatform/GenericPlatformMisc.h
@@ -832,6 +832,18 @@ public:
 	 */
 	static TArray<FString> GetPreferredLanguages();
 
+	/**
+	* Returns the currency code associated with the device's locale
+	* @return the currency code associated with the device's locale
+	*/
+	static FString GetLocalCurrencyCode();
+
+	/**
+	* Returns the currency symbol associated with the device's locale
+	* @return the currency symbol associated with the device's locale
+	*/
+	static FString GetLocalCurrencySymbol();
+
 	/**
 	 * Requests permission to send remote notifications to the user's device.
 	 */
diff --git a/Engine/Source/Runtime/Core/Public/IOS/IOSPlatformMisc.h b/Engine/Source/Runtime/Core/Public/IOS/IOSPlatformMisc.h
--- a/Engine/Source/Runtime/Core/Public/IOS/IOSPlatformMisc.h
+++ b/Engine/Source/Runtime/Core/Public/IOS/IOSPlatformMisc.h
@@ -95,6 +95,8 @@ struct CORE_API FIOSPlatformMisc : public FGenericPlatformMisc
 	static bool GetStoredValue(const FString& InStoreId, const FString& InSectionName, const FString& InKeyName, FString& OutValue);
 	static TArray<uint8> GetSystemFontBytes();
 	static TArray<FString> GetPreferredLanguages();
+	static FString GetLocalCurrencyCode();
+	static FString GetLocalCurrencySymbol();
 
 	static void RegisterForRemoteNotifications();
 
diff --git a/Engine/Source/Runtime/Engine/Classes/Kismet/KismetSystemLibrary.h b/Engine/Source/Runtime/Engine/Classes/Kismet/KismetSystemLibrary.h
--- a/Engine/Source/Runtime/Engine/Classes/Kismet/KismetSystemLibrary.h
+++ b/Engine/Source/Runtime/Engine/Classes/Kismet/KismetSystemLibrary.h
@@ -1338,7 +1338,21 @@ class ENGINE_API UKismetSystemLibrary : public UBlueprintFunctionLibrary
 	 * @return An array of language IDs ordered from most preferred to least
 	 */
 	UFUNCTION(BlueprintCallable, Category = "Utilities|Platform")
-		static TArray<FString> GetPreferredLanguages();
+	static TArray<FString> GetPreferredLanguages();
+
+	/**
+	* Returns the currency code associated with the device's locale
+	* @return the currency code associated with the device's locale
+	*/
+	UFUNCTION(BlueprintCallable, Category = "Utilities|Platform")
+	static FString GetLocalCurrencyCode();
+
+	/**
+	* Returns the currency symbol associated with the device's locale
+	* @return the currency symbol associated with the device's locale
+	*/
+	UFUNCTION(BlueprintCallable, Category = "Utilities|Platform")
+	static FString GetLocalCurrencySymbol();
 
 	/**
 	 * Requests permission to send remote notifications to the user's device.
diff --git a/Engine/Source/Runtime/Engine/Private/KismetSystemLibrary.cpp b/Engine/Source/Runtime/Engine/Private/KismetSystemLibrary.cpp
--- a/Engine/Source/Runtime/Engine/Private/KismetSystemLibrary.cpp
+++ b/Engine/Source/Runtime/Engine/Private/KismetSystemLibrary.cpp
@@ -3000,6 +3000,16 @@ TArray<FString> UKismetSystemLibrary::GetPreferredLanguages()
 	return FPlatformMisc::GetPreferredLanguages();
 }
 
+FString UKismetSystemLibrary::GetLocalCurrencyCode()
+{
+	return FPlatformMisc::GetLocalCurrencyCode();
+}
+
+FString UKismetSystemLibrary::GetLocalCurrencySymbol()
+{
+	return FPlatformMisc::GetLocalCurrencySymbol();
+}
+
 struct FLoadAssetActionBase : public FPendingLatentAction, public FGCObject
 {
 	// @TODO: it would be good to have static/global manager? 
