Gameplay - Fixing ProjectWorldLocationToScreenWithDistance to return true when it projects successfully.  Changing FSceneView::ProjectWorldToScreen to return a bool when projections fails as may happen if you're behind the half plane (points behind you no longer report projecting successfully).

diff --git a/Engine/Source/Runtime/Engine/Private/GameplayStatics.cpp b/Engine/Source/Runtime/Engine/Private/GameplayStatics.cpp
--- a/Engine/Source/Runtime/Engine/Private/GameplayStatics.cpp
+++ b/Engine/Source/Runtime/Engine/Private/GameplayStatics.cpp
@@ -1631,8 +1631,7 @@ bool UGameplayStatics::ProjectWorldToScreen(APlayerController const* Player, con
 		if (LP->GetProjectionData(LP->ViewportClient->Viewport, eSSP_FULL, /*out*/ ProjectionData))
 		{
 			FMatrix const ViewProjectionMatrix = ProjectionData.ComputeViewProjectionMatrix();
-			FSceneView::ProjectWorldToScreen(WorldPosition, ProjectionData.GetConstrainedViewRect(), ViewProjectionMatrix, ScreenPosition);
-			return true;
+			return FSceneView::ProjectWorldToScreen(WorldPosition, ProjectionData.GetConstrainedViewRect(), ViewProjectionMatrix, ScreenPosition);
 		}
 	}
 
diff --git a/Engine/Source/Runtime/Engine/Private/PlayerController.cpp b/Engine/Source/Runtime/Engine/Private/PlayerController.cpp
--- a/Engine/Source/Runtime/Engine/Private/PlayerController.cpp
+++ b/Engine/Source/Runtime/Engine/Private/PlayerController.cpp
@@ -1873,6 +1873,8 @@ bool APlayerController::ProjectWorldLocationToScreenWithDistance(FVector WorldLo
 			if (LP->GetProjectionData(LP->ViewportClient->Viewport, eSSP_FULL, /*out*/ ProjectionData))
 			{
 				ScreenLocation = FVector(ScreenLoc2D.X, ScreenLoc2D.Y, FVector::Dist(ProjectionData.ViewOrigin, WorldLocation));
+
+				return true;
 			}
 		}
 	}
diff --git a/Engine/Source/Runtime/Engine/Private/SceneView.cpp b/Engine/Source/Runtime/Engine/Private/SceneView.cpp
--- a/Engine/Source/Runtime/Engine/Private/SceneView.cpp
+++ b/Engine/Source/Runtime/Engine/Private/SceneView.cpp
@@ -775,28 +775,30 @@ void FSceneView::DeprojectScreenToWorld(const FVector2D& ScreenPos, const FIntRe
 	out_WorldDirection = RayDirWorldSpace;
 }
 
-void FSceneView::ProjectWorldToScreen(const FVector& WorldPosition, const FIntRect& ViewRect, const FMatrix& ViewProjectionMatrix, FVector2D& out_ScreenPos)
+bool FSceneView::ProjectWorldToScreen(const FVector& WorldPosition, const FIntRect& ViewRect, const FMatrix& ViewProjectionMatrix, FVector2D& out_ScreenPos)
 {
 	FPlane Result = ViewProjectionMatrix.TransformFVector4(FVector4(WorldPosition, 1.f));
-	if (Result.W == 0)
+	if ( Result.W > 0.0f )
 	{
-		Result.W = KINDA_SMALL_NUMBER;
-	}
+		// the result of this will be x and y coords in -1..1 projection space
+		const float RHW = 1.0f / Result.W;
+		FPlane PosInScreenSpace = FPlane(Result.X * RHW, Result.Y * RHW, Result.Z * RHW, Result.W);
 
-	// the result of this will be x and y coords in -1..1 projection space
-	const float RHW = 1.0f / Result.W;
-	FPlane PosInScreenSpace = FPlane(Result.X * RHW, Result.Y * RHW, Result.Z * RHW, Result.W);
+		// Move from projection space to normalized 0..1 UI space
+		const float NormalizedX = ( PosInScreenSpace.X / 2.f ) + 0.5f;
+		const float NormalizedY = 1.f - ( PosInScreenSpace.Y / 2.f ) - 0.5f;
 
-	// Move from projection space to normalized 0..1 UI space
-	const float NormalizedX = (PosInScreenSpace.X / 2.f) + 0.5f;
-	const float NormalizedY = 1.f - (PosInScreenSpace.Y / 2.f) - 0.5f;
+		FVector2D RayStartViewRectSpace(
+			(float)ViewRect.Min.X + ( NormalizedX * (float)ViewRect.Width() ),
+			(float)ViewRect.Min.Y + ( NormalizedY * (float)ViewRect.Height() )
+			);
 
-	FVector2D RayStartViewRectSpace(
-		(float)ViewRect.Min.X + (NormalizedX * (float)ViewRect.Width()),
-		(float)ViewRect.Min.Y + (NormalizedY * (float)ViewRect.Height())
-		);
+		out_ScreenPos = RayStartViewRectSpace;
 
-	out_ScreenPos = RayStartViewRectSpace;
+		return true;
+	}
+	
+	return false;
 }
 
 
diff --git a/Engine/Source/Runtime/Engine/Public/SceneView.h b/Engine/Source/Runtime/Engine/Public/SceneView.h
--- a/Engine/Source/Runtime/Engine/Public/SceneView.h
+++ b/Engine/Source/Runtime/Engine/Public/SceneView.h
@@ -662,7 +662,7 @@ public:
 	 * @param ViewProjectionMatrix - combined view projection matrix
 	 * @param out_ScreenPos (out) - screen coordinates in pixels
 	 */
-	static void ProjectWorldToScreen(const FVector& WorldPosition, const FIntRect& ViewRect, const FMatrix& ViewProjectionMatrix, FVector2D& out_ScreenPos);
+	static bool ProjectWorldToScreen(const FVector& WorldPosition, const FIntRect& ViewRect, const FMatrix& ViewProjectionMatrix, FVector2D& out_ScreenPos);
 
 	inline FVector GetViewRight() const { return ViewMatrices.ViewMatrix.GetColumn(0); }
 	inline FVector GetViewUp() const { return ViewMatrices.ViewMatrix.GetColumn(1); }
