#ue4 - Use faster Cast<> in UNavigationSystem::UpdateNavOctreeAfterMove().

diff --git a/Engine/Source/Runtime/Engine/Private/AI/Navigation/NavigationSystem.cpp b/Engine/Source/Runtime/Engine/Private/AI/Navigation/NavigationSystem.cpp
--- a/Engine/Source/Runtime/Engine/Private/AI/Navigation/NavigationSystem.cpp
+++ b/Engine/Source/Runtime/Engine/Private/AI/Navigation/NavigationSystem.cpp
@@ -2587,9 +2587,10 @@ void UNavigationSystem::UpdateNavOctreeAfterMove(USceneComponent* Comp)
 
 		for (int32 ComponentIndex = 0; ComponentIndex < Components.Num(); ComponentIndex++)
 		{
-			if (Components[ComponentIndex] && !Components[ComponentIndex]->IsA(USceneComponent::StaticClass()))
+			UActorComponent* const Component = Components[ComponentIndex];
+			if (Component && !Cast<USceneComponent>(Component))
 			{
-				UpdateNavOctree(Components[ComponentIndex]);
+				UpdateNavOctree(Component);
 			}
 		}
 
