#ue4 - Use faster SetLocationAndRotation() rather than separate calls for location and rotation.

diff --git a/Engine/Source/Runtime/Engine/Private/Brush.cpp b/Engine/Source/Runtime/Engine/Private/Brush.cpp
--- a/Engine/Source/Runtime/Engine/Private/Brush.cpp
+++ b/Engine/Source/Runtime/Engine/Private/Brush.cpp
@@ -86,8 +86,7 @@ void ABrush::CopyPosRotScaleFrom( ABrush* Other )
 	check(Other);
 	check(Other->BrushComponent);
 
-	SetActorLocation(Other->GetActorLocation(), false);
-	SetActorRotation(Other->GetActorRotation());
+	SetActorLocationAndRotation(Other->GetActorLocation(), Other->GetActorRotation(), false);
 	if( GetRootComponent() != NULL )
 	{
 		SetPivotOffset(Other->GetPivotOffset());
@@ -105,8 +104,7 @@ void ABrush::InitPosRotScale()
 {
 	check(BrushComponent);
 
-	SetActorLocation(FVector::ZeroVector, false);
-	SetActorRotation(FRotator::ZeroRotator);
+	SetActorLocationAndRotation(FVector::ZeroVector, FQuat::Identity, false);
 	SetPivotOffset( FVector::ZeroVector );
 }
 
