added comment

diff --git a/Engine/Source/Runtime/Engine/Public/SceneView.h b/Engine/Source/Runtime/Engine/Public/SceneView.h
--- a/Engine/Source/Runtime/Engine/Public/SceneView.h
+++ b/Engine/Source/Runtime/Engine/Public/SceneView.h
@@ -250,6 +250,7 @@ struct FViewMatrices
 		return GetInvProjMatrix() * GetInvViewMatrix();
 	}
 
+	// @return in radians (horizontal,vertical)
 	FVector2D GetHalfFieldOfViewPerAxis() const
 	{
 		const FMatrix ClipToView = GetInvProjNoAAMatrix();
