Add flag to notifies to stop them triggering on dedicated servers

diff --git a/Engine/Source/Editor/Persona/Private/AnimNotifyDetails.cpp b/Engine/Source/Editor/Persona/Private/AnimNotifyDetails.cpp
--- a/Engine/Source/Editor/Persona/Private/AnimNotifyDetails.cpp
+++ b/Engine/Source/Editor/Persona/Private/AnimNotifyDetails.cpp
@@ -123,6 +123,7 @@ void FAnimNotifyDetails::CustomizeDetails( IDetailLayoutBuilder& DetailBuilder )
 	TriggerFilterModeHandle = DetailBuilder.GetProperty(TEXT("Event.NotifyFilterType"));
 
 	FPropVisPair TriggerSettingNames[] = { { TEXT("Event.NotifyTriggerChance"), TAttribute<EVisibility>(EVisibility::Visible) }
+										 , { TEXT("Event.bTriggerOnDedicatedServer"), TAttribute<EVisibility>(EVisibility::Visible) }
 										 , { TEXT("Event.NotifyFilterType"), TAttribute<EVisibility>(EVisibility::Visible) }
 										 , { TEXT("Event.NotifyFilterLOD"), TAttribute<EVisibility>(this, &FAnimNotifyDetails::VisibilityForLODFilterMode) } };
 
diff --git a/Engine/Source/Runtime/Engine/Private/Animation/AnimInstance.cpp b/Engine/Source/Runtime/Engine/Private/Animation/AnimInstance.cpp
--- a/Engine/Source/Runtime/Engine/Private/Animation/AnimInstance.cpp
+++ b/Engine/Source/Runtime/Engine/Private/Animation/AnimInstance.cpp
@@ -1030,7 +1030,8 @@ void UAnimInstance::AddAnimNotifies(const TArray<const FAnimNotifyEvent*>& NewNo
 	for (const FAnimNotifyEvent* Notify : NewNotifies)
 	{
 		// only add if it is over TriggerWeightThreshold
-		if (Notify->TriggerWeightThreshold <= InstanceWeight && PassesFiltering(Notify) && PassesChanceOfTriggering(Notify))
+		const bool bPassesDedicatedServerCheck = Notify->bTriggerOnDedicatedServer || !IsRunningDedicatedServer();
+		if (bPassesDedicatedServerCheck && Notify->TriggerWeightThreshold <= InstanceWeight && PassesFiltering(Notify) && PassesChanceOfTriggering(Notify))
 		{
 			// Only add unique AnimNotifyState instances just once. We can get multiple triggers if looping over an animation.
 			// It is the same state, so just report it once.
diff --git a/Engine/Source/Runtime/Engine/Public/Animation/AnimTypes.h b/Engine/Source/Runtime/Engine/Public/Animation/AnimTypes.h
--- a/Engine/Source/Runtime/Engine/Public/Animation/AnimTypes.h
+++ b/Engine/Source/Runtime/Engine/Public/Animation/AnimTypes.h
@@ -103,6 +103,10 @@ struct FAnimNotifyEvent : public FAnimLinkableElement
 	UPROPERTY(EditAnywhere, BlueprintReadWrite, Category = AnimNotifyTriggerSettings, meta = (ClampMin = "0"))
 	int32 NotifyFilterLOD;
 
+	/** If disabled this notify will be skipped on dedicated servers */
+	UPROPERTY(EditAnywhere, BlueprintReadWrite, Category = AnimNotifyTriggerSettings)
+	bool bTriggerOnDedicatedServer;
+
 #if WITH_EDITORONLY_DATA
 	/** Color of Notify in editor */
 	UPROPERTY()
@@ -128,6 +132,7 @@ struct FAnimNotifyEvent : public FAnimLinkableElement
 		, NotifyTriggerChance(1.f)
 		, NotifyFilterType(ENotifyFilterType::NoFiltering)
 		, NotifyFilterLOD(0)
+		, bTriggerOnDedicatedServer(true)
 #if WITH_EDITORONLY_DATA
 		, TrackIndex(0)
 #endif // WITH_EDITORONLY_DATA
