Supplementary fix to deferred deletion crash in destructibles. Previous fix in CL#2609497 caused root chunks not to get activated.

diff --git a/Engine/Source/Runtime/Engine/Private/Components/DestructibleComponent.cpp b/Engine/Source/Runtime/Engine/Private/Components/DestructibleComponent.cpp
--- a/Engine/Source/Runtime/Engine/Private/Components/DestructibleComponent.cpp
+++ b/Engine/Source/Runtime/Engine/Private/Components/DestructibleComponent.cpp
@@ -364,6 +364,14 @@ void UDestructibleComponent::CreatePhysicsState()
 	PhysxChunkUserData.Reset(ChunkCount);
 	PhysxChunkUserData.AddZeroed(ChunkCount);
 
+	if(PhysScene->DeferredCommandHandler.HasPendingCommands())
+	{
+		// Lock and flush deferred command handler here to stop any currently pending deletions from affecting new actors.
+		// Only do this if we have any commands to flush to avoid the scene lock if possible
+		SCOPED_SCENE_WRITE_LOCK(PScene);
+		PhysScene->DeferredCommandHandler.Flush();
+	}
+	
 	// Create an APEX NxDestructibleActor from the Destructible asset and actor descriptor
 	ApexDestructibleActor = static_cast<NxDestructibleActor*>(TheDestructibleMesh->ApexDestructibleAsset->createApexActor(*ActorParams, *ApexScene));
 	check(ApexDestructibleActor);
@@ -377,21 +385,21 @@ void UDestructibleComponent::CreatePhysicsState()
 	ApexDestructibleActor->cacheModuleData();
 
 	// BRGTODO : Per-actor LOD setting
-//	ApexDestructibleActor->forcePhysicalLod( DestructibleActor->LOD );
+	//	ApexDestructibleActor->forcePhysicalLod( DestructibleActor->LOD );
 
 	// Start asleep if requested
 	PxRigidDynamic* PRootActor = ApexDestructibleActor->getChunkPhysXActor(0);
 
 
 	//  Put to sleep or wake up only if the component is physics-simulated
-	if (PRootActor != NULL && BodyInstance.bSimulatePhysics)
+	if(PRootActor != NULL && BodyInstance.bSimulatePhysics)
 	{
 		SCOPED_SCENE_WRITE_LOCK(PScene);	//Question, since apex is defer adding actors do we need to lock? Locking the async scene is expensive!
 
 		PRootActor->setActorFlag(PxActorFlag::eDISABLE_GRAVITY, !BodyInstance.bEnableGravity);
 
 		// Sleep/wake up as appropriate
-		if (!BodyInstance.bStartAwake)
+		if(!BodyInstance.bStartAwake)
 		{
 			//ApexDestructibleActor->setChunkPhysXActorAwakeState(0, false);	//TODO: broke during bad integration of apex. Will turn on once new libs are in
 		}
@@ -845,11 +853,6 @@ FTransform UDestructibleComponent::GetSocketTransform(FName InSocketName, ERelat
 #if WITH_APEX
 void UDestructibleComponent::Pair( int32 ChunkIndex, PxShape* PShape)
 {
-	if(ApexDestructibleActor == nullptr)	//since we do deferred deletion it's possible we've already meant to delete this so ignore any simulation callbacks
-	{
-		return;
-	}
-
 	FDestructibleChunkInfo* CI;
 	FPhysxUserData* UserData;
 
@@ -1469,6 +1472,11 @@ void UDestructibleComponent::SetCollisionResponseForAllActors(const FCollisionRe
 
 void UDestructibleComponent::SetCollisionResponseForShape(PxShape* Shape, int32 ChunkIdx)
 {
+	if(ApexDestructibleActor == nullptr) //since we do deferred deletion it's possible we've already meant to delete this so ignore any simulation callbacks
+	{
+		return;
+	}
+
 	// Get collision channel and response
 	PxFilterData PQueryFilterData, PSimFilterData;
 	uint8 MoveChannel = GetCollisionObjectType();
diff --git a/Engine/Source/Runtime/Engine/Private/PhysicsEngine/PhysCommandHandler.cpp b/Engine/Source/Runtime/Engine/Private/PhysicsEngine/PhysCommandHandler.cpp
--- a/Engine/Source/Runtime/Engine/Private/PhysicsEngine/PhysCommandHandler.cpp
+++ b/Engine/Source/Runtime/Engine/Private/PhysicsEngine/PhysCommandHandler.cpp
@@ -24,6 +24,11 @@ void FPhysCommandHandler::Flush()
 	PendingCommands.Empty();
 }
 
+bool FPhysCommandHandler::HasPendingCommands()
+{
+	return PendingCommands.Num() > 0;
+}
+
 void FPhysCommandHandler::ExecuteCommands()
 {
 	for (int32 i = 0; i < PendingCommands.Num(); ++i)
diff --git a/Engine/Source/Runtime/Engine/Public/PhysicsPublic.h b/Engine/Source/Runtime/Engine/Public/PhysicsPublic.h
--- a/Engine/Source/Runtime/Engine/Public/PhysicsPublic.h
+++ b/Engine/Source/Runtime/Engine/Public/PhysicsPublic.h
@@ -210,6 +210,7 @@ public:
 	
 	/** Executes pending commands and clears buffer **/
 	void ENGINE_API Flush();
+	bool ENGINE_API HasPendingCommands();
 
 #if WITH_APEX
 	/** enqueues a command to release destructible actor once apex has finished simulating */
@@ -252,6 +253,7 @@ private:
 	/** Array of commands waiting to execute once simulation is done */
 	TArray<FPhysPendingCommand> PendingCommands;
 };
+
 namespace SleepEvent
 {
 	enum Type
