Allow GameMode HUDClass to be set to none in the properties.

diff --git a/Engine/Source/Runtime/Engine/Classes/GameFramework/GameMode.h b/Engine/Source/Runtime/Engine/Classes/GameFramework/GameMode.h
--- a/Engine/Source/Runtime/Engine/Classes/GameFramework/GameMode.h
+++ b/Engine/Source/Runtime/Engine/Classes/GameFramework/GameMode.h
@@ -181,11 +181,11 @@ public:
 
 	/** The default pawn class used by players. */
 	UPROPERTY(EditAnywhere, BlueprintReadWrite, Category=Classes)
-	TSubclassOf<class APawn>  DefaultPawnClass;
+	TSubclassOf<class APawn> DefaultPawnClass;
 
 	/** HUD class this game uses. */
-	UPROPERTY(EditAnywhere, noclear, BlueprintReadWrite, Category=Classes, meta=(DisplayName="HUD Class"))
-	TSubclassOf<class AHUD>  HUDClass;    
+	UPROPERTY(EditAnywhere, BlueprintReadWrite, Category=Classes, meta=(DisplayName="HUD Class"))
+	TSubclassOf<class AHUD> HUDClass;    
 
 	/** Current number of spectators. */
 	UPROPERTY(BlueprintReadOnly, Category=GameMode)
