Prevent axis value getters from being used outside of their self context.

diff --git a/Engine/Source/Runtime/Engine/Classes/GameFramework/Actor.h b/Engine/Source/Runtime/Engine/Classes/GameFramework/Actor.h
--- a/Engine/Source/Runtime/Engine/Classes/GameFramework/Actor.h
+++ b/Engine/Source/Runtime/Engine/Classes/GameFramework/Actor.h
@@ -596,15 +596,15 @@ public:
 	virtual void DisableInput(class APlayerController* PlayerController);
 
 	/** Gets the value of the input axis if input is enabled for this actor. */
-	UFUNCTION(BlueprintCallable, meta=(BlueprintInternalUseOnly="true", BlueprintProtected = "true", HidePin="InputAxisName"))
+	UFUNCTION(BlueprintCallable, meta=(BlueprintInternalUseOnly="true", HideSelfPin="true", HidePin="InputAxisName"))
 	float GetInputAxisValue(const FName InputAxisName) const;
 
 	/** Gets the value of the input axis key if input is enabled for this actor. */
-	UFUNCTION(BlueprintCallable, meta=(BlueprintInternalUseOnly="true", BlueprintProtected = "true", HidePin="InputAxisKey"))
+	UFUNCTION(BlueprintCallable, meta=(BlueprintInternalUseOnly="true", HideSelfPin="true", HidePin="InputAxisKey"))
 	float GetInputAxisKeyValue(const FKey InputAxisKey) const;
 
 	/** Gets the value of the input axis key if input is enabled for this actor. */
-	UFUNCTION(BlueprintCallable, meta=(BlueprintInternalUseOnly="true", BlueprintProtected = "true", HidePin="InputAxisKey"))
+	UFUNCTION(BlueprintCallable, meta=(BlueprintInternalUseOnly="true", HideSelfPin="true", HidePin="InputAxisKey"))
 	FVector GetInputVectorAxisValue(const FKey InputAxisKey) const;
 
 	/** Returns the instigator for this actor, or NULL if there is none. */
