Made some useful CameraShake instance variables public and BP-accessible. Also Added UCameraAnimInst::SetScale().  (UE-15857)

diff --git a/Engine/Source/Runtime/Engine/Classes/Camera/CameraAnimInst.h b/Engine/Source/Runtime/Engine/Classes/Camera/CameraAnimInst.h
--- a/Engine/Source/Runtime/Engine/Classes/Camera/CameraAnimInst.h
+++ b/Engine/Source/Runtime/Engine/Classes/Camera/CameraAnimInst.h
@@ -136,7 +136,10 @@ public:
 	UFUNCTION(BlueprintCallable, Category=CameraAnimInst)
 	void SetDuration(float NewDuration);
 
-	
+	/** Changes the scale of the animation while playing.*/
+	UFUNCTION(BlueprintCallable, Category = CameraAnimInst)
+	void SetScale(float NewDuration);
+
 
 protected:
 	/** Returns InterpGroupInst subobject **/
diff --git a/Engine/Source/Runtime/Engine/Classes/Camera/CameraShake.h b/Engine/Source/Runtime/Engine/Classes/Camera/CameraShake.h
--- a/Engine/Source/Runtime/Engine/Classes/Camera/CameraShake.h
+++ b/Engine/Source/Runtime/Engine/Classes/Camera/CameraShake.h
@@ -182,10 +182,13 @@ protected:
 	// INSTANCE DATA
 	
 	/** True if this shake is currently blending in. */
-	uint32 bBlendingIn:1;
+	uint16 bBlendingIn:1;
 
 	/** True if this shake is currently blending out. */
-	uint32 bBlendingOut:1;
+	uint16 bBlendingOut:1;
+
+	/** What space to play the shake in before applying to the camera.  Affects both Anim and Oscillation shakes. */
+	ECameraAnimPlaySpace::Type PlaySpace;
 
 	/** How long this instance has been blending in. */
 	float CurrentBlendInTime;
@@ -196,9 +199,6 @@ protected:
 	UPROPERTY(transient, BlueprintReadOnly, Category = CameraShake)
 	class APlayerCameraManager* CameraOwner;
 
-	/** Time remaining for oscillation shakes. Less than 0.f means shake infinitely. */
-	float OscillatorTimeRemaining;
-
 	/** Current location sinusoidal offset. */
 	FVector LocSinOffset;
 
@@ -208,21 +208,23 @@ protected:
 	/** Current FOV sinusoidal offset. */
 	float FOVSinOffset;
 
+	/** Matrix defining the playspace, used when PlaySpace == CAPS_UserDefined */
+	FMatrix UserPlaySpaceMatrix;
+
+public:
 	/** Overall intensity scale for this shake instance. */
+	UPROPERTY(transient, BlueprintReadOnly, Category = CameraShake)
 	float ShakeScale;
 
-	/** What space to play the shake in before applying to the camera.  Affects both Anim and Oscillation shakes. */
-	ECameraAnimPlaySpace::Type PlaySpace;
+	/** Time remaining for oscillation shakes. Less than 0.f means shake infinitely. */
+	UPROPERTY(transient, BlueprintReadOnly, Category = CameraShake)
+	float OscillatorTimeRemaining;
 
 	/** The playing instance of the CameraAnim-based shake, if any. */
-	UPROPERTY()
+	UPROPERTY(transient, BlueprintReadOnly, Category = CameraShake)
 	class UCameraAnimInst* AnimInst;
 
-	/** Matrix defining the playspace, used when PlaySpace == CAPS_UserDefined */
-	FMatrix UserPlaySpaceMatrix;
-
 public:
-
 	// Blueprint API
 
 	/** Called every tick to let the shake modify the point of view */
@@ -247,7 +249,6 @@ public:
 	virtual void PlayShake(class APlayerCameraManager* Camera, float Scale, ECameraAnimPlaySpace::Type InPlaySpace, FRotator UserPlaySpaceRot = FRotator::ZeroRotator);
 	virtual bool IsFinished() const;
 	virtual void StopShake();
-
 };
 
 
diff --git a/Engine/Source/Runtime/Engine/Private/Camera/CameraAnimInst.cpp b/Engine/Source/Runtime/Engine/Private/Camera/CameraAnimInst.cpp
--- a/Engine/Source/Runtime/Engine/Private/Camera/CameraAnimInst.cpp
+++ b/Engine/Source/Runtime/Engine/Private/Camera/CameraAnimInst.cpp
@@ -157,6 +157,10 @@ void UCameraAnimInst::SetDuration(float NewDuration)
 	bFinished = false;
 }
 
+void UCameraAnimInst::SetScale(float NewScale)
+{
+	BasePlayScale = NewScale;
+}
 
 void UCameraAnimInst::Play(UCameraAnim* Anim, class AActor* CamActor, float InRate, float InScale, float InBlendInTime, float InBlendOutTime, bool bInLooping, bool bRandomStartTime, float Duration)
 {
