Increase maximum magnitude of quantized FRepMovement vectors to at least +/- ~10 million.
Separate quantization modes for location and velocity.

diff --git a/Engine/Source/Runtime/Engine/Classes/Engine/EngineTypes.h b/Engine/Source/Runtime/Engine/Classes/Engine/EngineTypes.h
--- a/Engine/Source/Runtime/Engine/Classes/Engine/EngineTypes.h
+++ b/Engine/Source/Runtime/Engine/Classes/Engine/EngineTypes.h
@@ -2713,11 +2713,11 @@ DECLARE_DYNAMIC_DELEGATE(FTimerDynamicDelegate);
 UENUM()
 enum class EVectorQuantization : uint8
 {
-	/** Each vector component will be rounded to the nearest whole number. Equivalent to the behavior of FVector_NetQuantize. */
+	/** Each vector component will be rounded to the nearest whole number. */
 	RoundWholeNumber,
-	/** Each vector component will be rounded, preserving one decimal place. Equivalent to the behavior of FVector_NetQuantize10. */
+	/** Each vector component will be rounded, preserving one decimal place. */
 	RoundOneDecimal,
-	/** Each vector component will be rounded, preserving two decimal places. Equivalent to the behavior of FVector_NetQuantize100. */
+	/** Each vector component will be rounded, preserving two decimal places. */
 	RoundTwoDecimals
 };
 
@@ -2759,9 +2759,13 @@ struct FRepMovement
 	UPROPERTY(Transient)
 	uint8 bRepPhysics : 1;
 
-	/** Allows tuning the compression level for the replicated location and velocity vectors. You should only need to change this from the default if you see visual artifacts. */
+	/** Allows tuning the compression level for the replicated location vector. You should only need to change this from the default if you see visual artifacts. */
 	UPROPERTY(EditDefaultsOnly, Category=Replication, AdvancedDisplay)
-	EVectorQuantization VectorQuantizationLevel;
+	EVectorQuantization LocationQuantizationLevel;
+
+	/** Allows tuning the compression level for the replicated velocity vectors. You should only need to change this from the default if you see visual artifacts. */
+	UPROPERTY(EditDefaultsOnly, Category=Replication, AdvancedDisplay)
+	EVectorQuantization VelocityQuantizationLevel;
 
 	/** Allows tuning the compression level for replicated rotation. You should only need to change this from the default if you see visual artifacts. */
 	UPROPERTY(EditDefaultsOnly, Category=Replication, AdvancedDisplay)
@@ -2769,9 +2773,15 @@ struct FRepMovement
 
 	FRepMovement();
 
-	bool SerializeQuantizedVector(FArchive& Ar, FVector& Vector)
+	bool SerializeQuantizedVector(FArchive& Ar, FVector& Vector, EVectorQuantization QuantizationLevel)
 	{
-		switch(VectorQuantizationLevel)
+		// Since FRepMovement used to use FVector_NetQuantize100, we're allowing enough bits per component
+		// regardless of the quantization level so that we can still support at least the same maximum magnitude
+		// (2^30 / 100, or ~10 million).
+		// This uses no inherent extra bandwidth since we're still using the same number of bits to store the
+		// bits-per-component value. Of course, larger magnitudes will still use more bandwidth,
+		// as has always been the case.
+		switch(QuantizationLevel)
 		{
 			case EVectorQuantization::RoundTwoDecimals:
 			{
@@ -2780,12 +2790,12 @@ struct FRepMovement
 
 			case EVectorQuantization::RoundOneDecimal:
 			{
-				return SerializePackedVector<10, 24>(Vector, Ar);
+				return SerializePackedVector<10, 27>(Vector, Ar);
 			}
 
 			default:
 			{
-				return SerializePackedVector<1, 20>(Vector, Ar);
+				return SerializePackedVector<1, 24>(Vector, Ar);
 			}
 		}
 	}
@@ -2801,7 +2811,7 @@ struct FRepMovement
 		bOutSuccess = true;
 
 		// update location, rotation, linear velocity
-		bOutSuccess &= SerializeQuantizedVector( Ar, Location );
+		bOutSuccess &= SerializeQuantizedVector( Ar, Location, LocationQuantizationLevel );
 		
 		switch(RotationQuantizationLevel)
 		{
@@ -2818,12 +2828,12 @@ struct FRepMovement
 			}
 		}
 		
-		bOutSuccess &= SerializeQuantizedVector( Ar, LinearVelocity );
+		bOutSuccess &= SerializeQuantizedVector( Ar, LinearVelocity, VelocityQuantizationLevel );
 
 		// update angular velocity if required
 		if ( bRepPhysics )
 		{
-			bOutSuccess &= SerializeQuantizedVector( Ar, AngularVelocity );
+			bOutSuccess &= SerializeQuantizedVector( Ar, AngularVelocity, VelocityQuantizationLevel );
 		}
 
 		return true;
diff --git a/Engine/Source/Runtime/Engine/Private/EngineTypes.cpp b/Engine/Source/Runtime/Engine/Private/EngineTypes.cpp
--- a/Engine/Source/Runtime/Engine/Private/EngineTypes.cpp
+++ b/Engine/Source/Runtime/Engine/Private/EngineTypes.cpp
@@ -221,7 +221,8 @@ FRepMovement::FRepMovement()
 	, Rotation(ForceInit)
 	, bSimulatedPhysicSleep(false)
 	, bRepPhysics(false)
-	, VectorQuantizationLevel(EVectorQuantization::RoundWholeNumber)
+	, LocationQuantizationLevel(EVectorQuantization::RoundWholeNumber)
+	, VelocityQuantizationLevel(EVectorQuantization::RoundWholeNumber)
 	, RotationQuantizationLevel(ERotatorQuantization::ByteComponents)
 {
 }
