Make SetTimer return handle
Expose all other timer functions to operate off handle
Deprecate all timer functions using name
Deprecate all timer functions, except Set, using delegate

diff --git a/Engine/Source/Runtime/Engine/Classes/Engine/EngineTypes.h b/Engine/Source/Runtime/Engine/Classes/Engine/EngineTypes.h
--- a/Engine/Source/Runtime/Engine/Classes/Engine/EngineTypes.h
+++ b/Engine/Source/Runtime/Engine/Classes/Engine/EngineTypes.h
@@ -2707,9 +2707,50 @@ namespace EEndPlayReason
 
 }
 
-
 DECLARE_DYNAMIC_DELEGATE(FTimerDynamicDelegate);
 
+// Unique handle that can be used to distinguish timers that have identical delegates.
+USTRUCT(BlueprintType)
+struct FTimerHandle
+{
+	GENERATED_BODY()
+
+	FTimerHandle()
+	: Handle(INDEX_NONE)
+	{
+	}
+
+	bool IsValid() const
+	{
+		return Handle != INDEX_NONE;
+	}
+
+	void Invalidate()
+	{
+		Handle = INDEX_NONE;
+	}
+
+	void MakeValid();
+
+	bool operator==(const FTimerHandle& Other) const
+	{
+		return Handle == Other.Handle;
+	}
+
+	bool operator!=(const FTimerHandle& Other) const
+	{
+		return Handle != Other.Handle;
+	}
+
+	FString ToString() const
+	{
+		return FString::Printf(TEXT("%d"), Handle);
+	}
+
+private:
+	int32 Handle;
+};
+
 UENUM()
 enum class EVectorQuantization : uint8
 {
diff --git a/Engine/Source/Runtime/Engine/Classes/Kismet/KismetSystemLibrary.h b/Engine/Source/Runtime/Engine/Classes/Kismet/KismetSystemLibrary.h
--- a/Engine/Source/Runtime/Engine/Classes/Kismet/KismetSystemLibrary.h
+++ b/Engine/Source/Runtime/Engine/Classes/Kismet/KismetSystemLibrary.h
@@ -300,67 +300,152 @@ class ENGINE_API UKismetSystemLibrary : public UBlueprintFunctionLibrary
 	 * Set a timer to execute delegate. Setting an existing timer will reset that timer with updated parameters.
 	 * @param Delegate		Delegate. Can be a K2 function or a Custom Event.
 	 * @param Time			How long to wait before executing the delegate, in seconds. Setting a timer to <= 0 seconds will clear it if it is set.
-	 * @param bLooping		true to keep executing the delegate every Time seconds, false to execute delegate only once.
+	 * @param bLooping		True to keep executing the delegate every Time seconds, false to execute delegate only once.
+	 * @return				The timer handle to pass to other timer functions to manipulate this timer.
 	 */
-	UFUNCTION(BlueprintCallable, meta=(DisplayName = "SetTimer_Delegate"), Category="Utilities|Time")
-	static void K2_SetTimerDelegate(FTimerDynamicDelegate Delegate, float Time, bool bLooping);
+	UFUNCTION(BlueprintCallable, meta=(DisplayName = "Set Timer by Delegate"), Category="Utilities|Time")
+	static FTimerHandle K2_SetTimerDelegate(FTimerDynamicDelegate Delegate, float Time, bool bLooping);
 
 	/**
 	 * Clears a set timer.
 	 * @param Delegate		Delegate. Can be a K2 function or a Custom Event.
 	 */
-	UFUNCTION(BlueprintCallable, meta=(DisplayName = "ClearTimer_Delegate", DefaultToSelf = "Object"), Category="Utilities|Time")
+	UFUNCTION(BlueprintCallable, meta=(DeprecatedFunction, DeprecationMessage = "Use Clear Timer by Handle", DisplayName = "Clear Timer by Delegate"), Category="Utilities|Time")
 	static void K2_ClearTimerDelegate(FTimerDynamicDelegate Delegate);
 
 	/**
 	 * Pauses a set timer at its current elapsed time.
 	 * @param Delegate		Delegate. Can be a K2 function or a Custom Event.
 	 */
-	UFUNCTION(BlueprintCallable, meta=(DisplayName = "PauseTimer_Delegate", DefaultToSelf = "Object"), Category="Utilities|Time")
+	UFUNCTION(BlueprintCallable, meta=(DeprecatedFunction, DeprecationMessage = "Use Pause Timer by Handle", DisplayName = "Pause Timer by Delegate"), Category="Utilities|Time")
 	static void K2_PauseTimerDelegate(FTimerDynamicDelegate Delegate);
 
 	/**
 	 * Resumes a paused timer from its current elapsed time.
 	 * @param Delegate		Delegate. Can be a K2 function or a Custom Event.
 	 */
-	UFUNCTION(BlueprintCallable, meta=(DisplayName = "UnPauseTimer_Delegate", DefaultToSelf = "Object"), Category="Utilities|Time")
+	UFUNCTION(BlueprintCallable, meta=(DeprecatedFunction, DeprecationMessage = "Use Unpause Timer by Handle", DisplayName = "Unpause Timer by Delegate"), Category="Utilities|Time")
 	static void K2_UnPauseTimerDelegate(FTimerDynamicDelegate Delegate);
 
 	/**
-	 * Returns true is a timer exists and is active for the given delegate, false otherwise.
+	 * Returns true if a timer exists and is active for the given delegate, false otherwise.
 	 * @param Delegate		Delegate. Can be a K2 function or a Custom Event.
+	 * @return				True if the timer exists and is active.
 	 */
-	UFUNCTION(BlueprintPure, meta=(DisplayName = "IsTimerActive_Delegate", DefaultToSelf = "Object"), Category="Utilities|Time")
+	UFUNCTION(BlueprintPure, meta=(DeprecatedFunction, DeprecationMessage = "Use Is Timer Active by Handle", DisplayName = "Is Timer Active by Delegate"), Category="Utilities|Time")
 	static bool K2_IsTimerActiveDelegate(FTimerDynamicDelegate Delegate);
 
 	/**
-	 * Returns true is a timer exists and is paused for the given delegate, false otherwise.
+	 * Returns true if a timer exists and is paused for the given delegate, false otherwise.
 	 * @param Delegate		Delegate. Can be a K2 function or a Custom Event.
+	 * @return				True if the timer exists and is paused.
 	 */
-	UFUNCTION(BlueprintPure, meta = (DisplayName = "IsTimerPaused_Delegate", DefaultToSelf = "Object"), Category = "Utilities|Time")
+	UFUNCTION(BlueprintPure, meta=(DeprecatedFunction, DeprecationMessage = "Use Is Timer Paused by Handle", DisplayName = "Is Timer Paused by Delegate"), Category = "Utilities|Time")
 	static bool K2_IsTimerPausedDelegate(FTimerDynamicDelegate Delegate);
 
 	/**
 	 * Returns true is a timer for the given delegate exists, false otherwise.
 	 * @param Delegate		Delegate. Can be a K2 function or a Custom Event.
+	 * @return				True if the timer exists.
 	 */
-	UFUNCTION(BlueprintPure, meta = (DisplayName = "TimerExists_Delegate", DefaultToSelf = "Object"), Category = "Utilities|Time")
+	UFUNCTION(BlueprintPure, meta=(DeprecatedFunction, DeprecationMessage = "Use Does Timer Exist by Handle", DisplayName = "Does Timer Exist by Delegate"), Category = "Utilities|Time")
 	static bool K2_TimerExistsDelegate(FTimerDynamicDelegate Delegate);
 	
 	/**
 	 * Returns elapsed time for the given delegate (time since current countdown iteration began).
 	 * @param Delegate		Delegate. Can be a K2 function or a Custom Event.
+	 * @return				How long has elapsed since the current iteration of the timer began.
 	 */
-	UFUNCTION(BlueprintPure, meta=(DisplayName = "GetTimerElapsedTime_Delegate", DefaultToSelf = "Object"), Category="Utilities|Time")
+	UFUNCTION(BlueprintPure, meta=(DeprecatedFunction, DeprecationMessage = "Use Get Timer Elapsed Time by Handle", DisplayName = "Get Timer Elapsed Time by Delegate"), Category="Utilities|Time")
 	static float K2_GetTimerElapsedTimeDelegate(FTimerDynamicDelegate Delegate);
 
 	/**
 	 * Returns time until the timer will next execute its delegate.
 	 * @param Delegate		Delegate. Can be a K2 function or a Custom Event.
+	 * @return				How long is remaining in the current iteration of the timer.
 	 */
-	UFUNCTION(BlueprintPure, meta=(DisplayName = "GetTimerRemainingTime_Delegate", DefaultToSelf = "Object"), Category="Utilities|Time")
+	UFUNCTION(BlueprintPure, meta=(DeprecatedFunction, DeprecationMessage = "Use Get Timer Remaining Time by Handle", DisplayName = "Get Timer Remaining Time by Delegate"), Category="Utilities|Time")
 	static float K2_GetTimerRemainingTimeDelegate(FTimerDynamicDelegate Delegate);
 
+	// --- Timer functions with handle input ----------
+
+	/**
+	 * Returns whether the timer handle is valid. This does not indicate that there is an active timer that this handle references, but rather that it once referenced a valid timer.
+	 * @param Handle		The handle of the timer to check validity of.
+	 * @return				Whether the timer handle is valid.
+	 */
+	UFUNCTION(BlueprintPure, meta=(DisplayName = "Is Valid"), Category="Utilities|Time")
+	static bool K2_IsValidTimerHandle(FTimerHandle Handle);
+
+	/**
+	 * Returns whether the timer handle is valid. This does not indicate that there is an active timer that this handle references, but rather that it once referenced a valid timer.
+	 * @param Handle		The handle of the timer to check validity of.
+	 * @return				Return the invalidated timer handle for convenience.
+	 */
+	UFUNCTION(BlueprintCallable, meta=(DisplayName = "Invalidate"), Category="Utilities|Time")
+	static FTimerHandle K2_InvalidateTimerHandle(UPARAM(ref) FTimerHandle& Handle);
+
+	/**
+	 * Clears a set timer.
+	 * @param Handle		The handle of the timer to clear.
+	 */
+	UFUNCTION(BlueprintCallable, meta=(DisplayName = "Clear Timer by Handle", WorldContext="WorldContextObject"), Category="Utilities|Time")
+	static void K2_ClearTimerHandle(UObject* WorldContextObject, FTimerHandle Handle);
+
+	/**
+	 * Pauses a set timer at its current elapsed time.
+	 * @param Handle		The handle of the timer to pause.
+	 */
+	UFUNCTION(BlueprintCallable, meta=(DisplayName = "Pause Timer by Handle", WorldContext="WorldContextObject"), Category="Utilities|Time")
+	static void K2_PauseTimerHandle(UObject* WorldContextObject, FTimerHandle Handle);
+
+	/**
+	 * Resumes a paused timer from its current elapsed time.
+	 * @param Handle		The handle of the timer to unpause.
+	 */
+	UFUNCTION(BlueprintCallable, meta=(DisplayName = "Unpause Timer by Handle", WorldContext="WorldContextObject"), Category="Utilities|Time")
+	static void K2_UnPauseTimerHandle(UObject* WorldContextObject, FTimerHandle Handle);
+
+	/**
+	 * Returns true if a timer exists and is active for the given handle, false otherwise.
+	 * @param Handle		The handle of the timer to check whether it is active.
+	 * @return				True if the timer exists and is active.
+	 */
+	UFUNCTION(BlueprintPure, meta=(DisplayName = "Is Timer Active by Handle", WorldContext="WorldContextObject"), Category="Utilities|Time")
+	static bool K2_IsTimerActiveHandle(UObject* WorldContextObject, FTimerHandle Handle);
+
+	/**
+	 * Returns true if a timer exists and is paused for the given handle, false otherwise.
+	 * @param Handle		The handle of the timer to check whether it is paused.
+	 * @return				True if the timer exists and is paused.
+	 */
+	UFUNCTION(BlueprintPure, meta=(DisplayName = "Is Timer Paused by Handle", WorldContext="WorldContextObject"), Category = "Utilities|Time")
+	static bool K2_IsTimerPausedHandle(UObject* WorldContextObject, FTimerHandle Handle);
+
+	/**
+	 * Returns true is a timer for the given handle exists, false otherwise.
+	 * @param Handle		The handle to check whether it exists.
+ 	 * @return				True if the timer exists.
+	 */
+	UFUNCTION(BlueprintPure, meta=(DisplayName = "Does Timer Exist by Handle", WorldContext="WorldContextObject"), Category = "Utilities|Time")
+	static bool K2_TimerExistsHandle(UObject* WorldContextObject, FTimerHandle Handle);
+	
+	/**
+	 * Returns elapsed time for the given handle (time since current countdown iteration began).
+	 * @param Handle		The handle of the timer to get the elapsed time of.
+	 * @return				How long has elapsed since the current iteration of the timer began.
+	 */
+	UFUNCTION(BlueprintPure, meta=(DisplayName = "Get Timer Elapsed Time by Handle", WorldContext="WorldContextObject"), Category="Utilities|Time")
+	static float K2_GetTimerElapsedTimeHandle(UObject* WorldContextObject, FTimerHandle Handle);
+
+	/**
+	 * Returns time until the timer will next execute its handle.
+	 * @param Handle		The handle of the timer to time remaining of.
+	 * @return				How long is remaining in the current iteration of the timer.
+	 */
+	UFUNCTION(BlueprintPure, meta=(DisplayName = "Get Timer Remaining Time by Handle", WorldContext="WorldContextObject"), Category="Utilities|Time")
+	static float K2_GetTimerRemainingTimeHandle(UObject* WorldContextObject, FTimerHandle Handle);
+
 	// --- Timer functions ------------------------------
 
 	/**
@@ -369,16 +454,17 @@ class ENGINE_API UKismetSystemLibrary : public UBlueprintFunctionLibrary
 	 * @param FunctionName	Delegate function name. Can be a K2 function or a Custom Event.
 	 * @param Time			How long to wait before executing the delegate, in seconds. Setting a timer to <= 0 seconds will clear it if it is set.
 	 * @param bLooping		true to keep executing the delegate every Time seconds, false to execute delegate only once.
+	 * @return				The timer handle to pass to other timer functions to manipulate this timer.
 	 */
-	UFUNCTION(BlueprintCallable, meta=(DisplayName = "SetTimer", DefaultToSelf = "Object"), Category="Utilities|Time")
-	static void K2_SetTimer(UObject* Object, FString FunctionName, float Time, bool bLooping);
+	UFUNCTION(BlueprintCallable, meta=(DeprecatedFunction, DeprecationMessage = "Use Set Timer by Delegate", DisplayName = "Set Timer by Function Name", DefaultToSelf = "Object"), Category="Utilities|Time")
+	static FTimerHandle K2_SetTimer(UObject* Object, FString FunctionName, float Time, bool bLooping);
 
 	/**
 	 * Clears a set timer.
 	 * @param Object		Object that implements the delegate function. Defaults to self (this blueprint)
 	 * @param FunctionName	Delegate function name. Can be a K2 function or a Custom Event.
 	 */
-	UFUNCTION(BlueprintCallable, meta=(DisplayName = "ClearTimer", DefaultToSelf = "Object"), Category="Utilities|Time")
+	UFUNCTION(BlueprintCallable, meta=(DeprecatedFunction, DeprecationMessage = "Use Clear Timer by Handle", DisplayName = "Clear Timer by Function Name", DefaultToSelf = "Object"), Category="Utilities|Time")
 	static void K2_ClearTimer(UObject* Object, FString FunctionName);
 
 	/**
@@ -386,7 +472,7 @@ class ENGINE_API UKismetSystemLibrary : public UBlueprintFunctionLibrary
 	 * @param Object		Object that implements the delegate function. Defaults to self (this blueprint)
 	 * @param FunctionName	Delegate function name. Can be a K2 function or a Custom Event.
 	 */
-	UFUNCTION(BlueprintCallable, meta=(DisplayName = "PauseTimer", DefaultToSelf = "Object"), Category="Utilities|Time")
+	UFUNCTION(BlueprintCallable, meta=(DeprecatedFunction, DeprecationMessage = "Use Pause Timer by Handle", DisplayName = "Pause Timer by Function Name", DefaultToSelf = "Object"), Category="Utilities|Time")
 	static void K2_PauseTimer(UObject* Object, FString FunctionName);
 
 	/**
@@ -394,47 +480,52 @@ class ENGINE_API UKismetSystemLibrary : public UBlueprintFunctionLibrary
 	 * @param Object		Object that implements the delegate function. Defaults to self (this blueprint)
 	 * @param FunctionName	Delegate function name. Can be a K2 function or a Custom Event.
 	 */
-	UFUNCTION(BlueprintCallable, meta=(DisplayName = "UnPauseTimer", DefaultToSelf = "Object"), Category="Utilities|Time")
+	UFUNCTION(BlueprintCallable, meta=(DeprecatedFunction, DeprecationMessage = "Use Unpause Timer by Handle", DisplayName = "Unpause Timer by Function Name", DefaultToSelf = "Object"), Category="Utilities|Time")
 	static void K2_UnPauseTimer(UObject* Object, FString FunctionName);
 
 	/**
-	 * Returns true is a timer exists and is active for the given delegate, false otherwise.
+	 * Returns true if a timer exists and is active for the given delegate, false otherwise.
 	 * @param Object		Object that implements the delegate function. Defaults to self (this blueprint)
 	 * @param FunctionName	Delegate function name. Can be a K2 function or a Custom Event.
+	 * @return				True if the timer exists and is active.
 	 */
-	UFUNCTION(BlueprintPure, meta=(DisplayName = "IsTimerActive", DefaultToSelf = "Object"), Category="Utilities|Time")
+	UFUNCTION(BlueprintPure, meta=(DeprecatedFunction, DeprecationMessage = "Use Is Timer Active by Handle", DisplayName = "Is Timer Active by Function Name", DefaultToSelf = "Object"), Category="Utilities|Time")
 	static bool K2_IsTimerActive(UObject* Object, FString FunctionName);
 
 	/**
-	* Returns true is a timer exists and is paused for the given delegate, false otherwise.
+	* Returns true if a timer exists and is paused for the given delegate, false otherwise.
 	* @param Object		Object that implements the delegate function. Defaults to self (this blueprint)
 	* @param FunctionName	Delegate function name. Can be a K2 function or a Custom Event.
+	* @return				True if the timer exists and is paused.
 	*/
-	UFUNCTION(BlueprintPure, meta = (DisplayName = "IsTimerPaused", DefaultToSelf = "Object"), Category = "Utilities|Time")
+	UFUNCTION(BlueprintPure, meta = (DeprecatedFunction, DeprecationMessage = "Use Is Timer Paused by Handle", DisplayName = "Is Timer Paused by Function Name", DefaultToSelf = "Object"), Category = "Utilities|Time")
 	static bool K2_IsTimerPaused(UObject* Object, FString FunctionName);
 
 	/**
 	* Returns true is a timer for the given delegate exists, false otherwise.
 	* @param Object		Object that implements the delegate function. Defaults to self (this blueprint)
 	* @param FunctionName	Delegate function name. Can be a K2 function or a Custom Event.
+	* @return				True if the timer exists.
 	*/
-	UFUNCTION(BlueprintPure, meta = (DisplayName = "TimerExists", DefaultToSelf = "Object"), Category = "Utilities|Time")
+	UFUNCTION(BlueprintPure, meta = (DeprecatedFunction, DeprecationMessage = "Use Timer Exists by Handle", DisplayName = "Does Timer Exist by Function Name", DefaultToSelf = "Object"), Category = "Utilities|Time")
 	static bool K2_TimerExists(UObject* Object, FString FunctionName);
 	
 	/**
 	 * Returns elapsed time for the given delegate (time since current countdown iteration began).
 	 * @param Object		Object that implements the delegate function. Defaults to self (this blueprint)
 	 * @param FunctionName	Delegate function name. Can be a K2 function or a Custom Event.
+	 * @return				How long has elapsed since the current iteration of the timer began.
 	 */
-	UFUNCTION(BlueprintPure, meta=(DisplayName = "GetTimerElapsedTime", DefaultToSelf = "Object"), Category="Utilities|Time")
+	UFUNCTION(BlueprintPure, meta=(DeprecatedFunction, DeprecationMessage = "Use Get Timer Elapsed Time by Handle", DisplayName = "Get Timer Elapsed Time by Function Name", DefaultToSelf = "Object"), Category="Utilities|Time")
 	static float K2_GetTimerElapsedTime(UObject* Object, FString FunctionName);
 
 	/**
 	 * Returns time until the timer will next execute its delegate.
 	 * @param Object		Object that implements the delegate function. Defaults to self (this blueprint)
 	 * @param FunctionName	Delegate function name. Can be a K2 function or a Custom Event.
+	 * @return				How long is remaining in the current iteration of the timer.
 	 */
-	UFUNCTION(BlueprintPure, meta=(DisplayName = "GetTimerRemainingTime", DefaultToSelf = "Object"), Category="Utilities|Time")
+	UFUNCTION(BlueprintPure, meta=(DeprecatedFunction, DeprecationMessage = "Use Get Timer Remaining Time by Handle", DisplayName = "Get Timer Remaining Time by Function Name", DefaultToSelf = "Object"), Category="Utilities|Time")
 	static float K2_GetTimerRemainingTime(UObject* Object, FString FunctionName);
 
 
diff --git a/Engine/Source/Runtime/Engine/Private/KismetSystemLibrary.cpp b/Engine/Source/Runtime/Engine/Private/KismetSystemLibrary.cpp
--- a/Engine/Source/Runtime/Engine/Private/KismetSystemLibrary.cpp
+++ b/Engine/Source/Runtime/Engine/Private/KismetSystemLibrary.cpp
@@ -285,8 +285,18 @@ void UKismetSystemLibrary::QuitGame(UObject* WorldContextObject, class APlayerCo
 	}
 }
 
+bool UKismetSystemLibrary::K2_IsValidTimerHandle(FTimerHandle TimerHandle)
+{
+	return TimerHandle.IsValid();
+}
 
-void UKismetSystemLibrary::K2_SetTimer(UObject* Object, FString FunctionName, float Time, bool bLooping)
+FTimerHandle UKismetSystemLibrary::K2_InvalidateTimerHandle(FTimerHandle& TimerHandle)
+{
+	TimerHandle.Invalidate();
+	return TimerHandle;
+}
+
+FTimerHandle UKismetSystemLibrary::K2_SetTimer(UObject* Object, FString FunctionName, float Time, bool bLooping)
 {
 	FName const FunctionFName(*FunctionName);
 
@@ -299,24 +309,25 @@ void UKismetSystemLibrary::K2_SetTimer(UObject* Object, FString FunctionName, fl
 			// FTimerDynamicDelegate expects zero parameters and will choke on execution if it tries
 			// to execute a mismatched function
 			UE_LOG(LogBlueprintUserMessages, Warning, TEXT("SetTimer passed a function (%s) that expects parameters."), *FunctionName);
-			return;
+			return FTimerHandle();
 		}
 	}
 
 	FTimerDynamicDelegate Delegate;
 	Delegate.BindUFunction(Object, FunctionFName);
-	K2_SetTimerDelegate(Delegate, Time, bLooping);
+	return K2_SetTimerDelegate(Delegate, Time, bLooping);
 }
 
-void UKismetSystemLibrary::K2_SetTimerDelegate(FTimerDynamicDelegate Delegate, float Time, bool bLooping)
+FTimerHandle UKismetSystemLibrary::K2_SetTimerDelegate(FTimerDynamicDelegate Delegate, float Time, bool bLooping)
 {
+	FTimerHandle Handle;
 	if (Delegate.IsBound())
 	{
 		const UWorld* const World = GEngine->GetWorldFromContextObject(Delegate.GetUObject());
 		if(World)
 		{
-			auto& TimerManager = World->GetTimerManager();
-			auto Handle = TimerManager.K2_FindDynamicTimerHandle(Delegate);
+			FTimerManager& TimerManager = World->GetTimerManager();
+			Handle = TimerManager.K2_FindDynamicTimerHandle(Delegate);
 			TimerManager.SetTimer(Handle, Delegate, Time, bLooping);
 		}
 	}
@@ -326,6 +337,8 @@ void UKismetSystemLibrary::K2_SetTimerDelegate(FTimerDynamicDelegate Delegate, f
 			TEXT("SetTimer passed a bad function (%s) or object (%s)"),
 			*Delegate.GetFunctionName().ToString(), *GetNameSafe(Delegate.GetUObject()));
 	}
+
+	return Handle;
 }
 
 void UKismetSystemLibrary::K2_ClearTimer(UObject* Object, FString FunctionName)
@@ -343,9 +356,9 @@ void UKismetSystemLibrary::K2_ClearTimerDelegate(FTimerDynamicDelegate Delegate)
 		UWorld* World = GEngine->GetWorldFromContextObject(Delegate.GetUObject());
 		if (World)
 		{
-			auto& TimerManager = World->GetTimerManager();
-			auto Handle = TimerManager.K2_FindDynamicTimerHandle(Delegate);
-			World->GetTimerManager().ClearTimer(Handle);
+			FTimerManager& TimerManager = World->GetTimerManager();
+			FTimerHandle Handle = TimerManager.K2_FindDynamicTimerHandle(Delegate);
+			TimerManager.ClearTimer(Handle);
 		}
 	}
 	else
@@ -356,6 +369,18 @@ void UKismetSystemLibrary::K2_ClearTimerDelegate(FTimerDynamicDelegate Delegate)
 	}
 }
 
+void UKismetSystemLibrary::K2_ClearTimerHandle(UObject* WorldContextObject, FTimerHandle Handle)
+{
+	if (Handle.IsValid())
+	{
+		UWorld* World = GEngine->GetWorldFromContextObject( WorldContextObject );
+		if (World)
+		{
+			World->GetTimerManager().ClearTimer(Handle);
+		}
+	}
+}
+
 void UKismetSystemLibrary::K2_PauseTimer(UObject* Object, FString FunctionName)
 {
 	FTimerDynamicDelegate Delegate;
@@ -371,9 +396,9 @@ void UKismetSystemLibrary::K2_PauseTimerDelegate(FTimerDynamicDelegate Delegate)
 		UWorld* World = GEngine->GetWorldFromContextObject(Delegate.GetUObject());
 		if(World)
 		{
-			auto& TimerManager = World->GetTimerManager();
-			auto Handle = TimerManager.K2_FindDynamicTimerHandle(Delegate);
-			World->GetTimerManager().PauseTimer(Handle);
+			FTimerManager& TimerManager = World->GetTimerManager();
+			FTimerHandle Handle = TimerManager.K2_FindDynamicTimerHandle(Delegate);
+			TimerManager.PauseTimer(Handle);
 		}
 	}
 	else
@@ -384,6 +409,18 @@ void UKismetSystemLibrary::K2_PauseTimerDelegate(FTimerDynamicDelegate Delegate)
 	}
 }
 
+void UKismetSystemLibrary::K2_PauseTimerHandle(UObject* WorldContextObject, FTimerHandle Handle)
+{
+	if (Handle.IsValid())
+	{
+		UWorld* World = GEngine->GetWorldFromContextObject( WorldContextObject );
+		if (World)
+		{
+			World->GetTimerManager().PauseTimer(Handle);
+		}
+	}
+}
+
 void UKismetSystemLibrary::K2_UnPauseTimer(UObject* Object, FString FunctionName)
 {
 	FTimerDynamicDelegate Delegate;
@@ -399,9 +436,9 @@ void UKismetSystemLibrary::K2_UnPauseTimerDelegate(FTimerDynamicDelegate Delegat
 		UWorld* World = GEngine->GetWorldFromContextObject(Delegate.GetUObject());
 		if(World)
 		{
-			auto& TimerManager = World->GetTimerManager();
-			auto Handle = TimerManager.K2_FindDynamicTimerHandle(Delegate);
-			World->GetTimerManager().UnPauseTimer(Handle);
+			FTimerManager& TimerManager = World->GetTimerManager();
+			FTimerHandle Handle = TimerManager.K2_FindDynamicTimerHandle(Delegate);
+			TimerManager.UnPauseTimer(Handle);
 		}
 	}
 	else
@@ -412,6 +449,18 @@ void UKismetSystemLibrary::K2_UnPauseTimerDelegate(FTimerDynamicDelegate Delegat
 	}
 }
 
+void UKismetSystemLibrary::K2_UnPauseTimerHandle(UObject* WorldContextObject, FTimerHandle Handle)
+{
+	if (Handle.IsValid())
+	{
+		UWorld* World = GEngine->GetWorldFromContextObject( WorldContextObject );
+		if (World)
+		{
+			World->GetTimerManager().UnPauseTimer(Handle);
+		}
+	}
+}
+
 bool UKismetSystemLibrary::K2_IsTimerActive(UObject* Object, FString FunctionName)
 {
 	FTimerDynamicDelegate Delegate;
@@ -422,18 +471,15 @@ bool UKismetSystemLibrary::K2_IsTimerActive(UObject* Object, FString FunctionNam
 
 bool UKismetSystemLibrary::K2_IsTimerActiveDelegate(FTimerDynamicDelegate Delegate)
 {
+	bool bIsActive = false;
 	if (Delegate.IsBound())
 	{
 		UWorld* World = GEngine->GetWorldFromContextObject(Delegate.GetUObject());
 		if(World)
 		{
-			auto& TimerManager = World->GetTimerManager();
-			auto Handle = TimerManager.K2_FindDynamicTimerHandle(Delegate);
-			return World->GetTimerManager().IsTimerActive(Handle);
-		}
-		else
-		{
-			return false;
+			FTimerManager& TimerManager = World->GetTimerManager();
+			FTimerHandle Handle = TimerManager.K2_FindDynamicTimerHandle(Delegate);
+			bIsActive = TimerManager.IsTimerActive(Handle);
 		}
 	}
 	else
@@ -441,8 +487,24 @@ bool UKismetSystemLibrary::K2_IsTimerActiveDelegate(FTimerDynamicDelegate Delega
 		UE_LOG(LogBlueprintUserMessages, Warning, 
 			TEXT("IsTimerActive passed a bad function (%s) or object (%s)"),
 			*Delegate.GetFunctionName().ToString(), *GetNameSafe(Delegate.GetUObject()));
-		return false;
 	}
+
+	return bIsActive;
+}
+
+bool UKismetSystemLibrary::K2_IsTimerActiveHandle(UObject* WorldContextObject, FTimerHandle Handle)
+{
+	bool bIsActive = false;
+	if (Handle.IsValid())
+	{
+		UWorld* World = GEngine->GetWorldFromContextObject( WorldContextObject );
+		if (World)
+		{
+			bIsActive = World->GetTimerManager().IsTimerActive(Handle);
+		}
+	}
+
+	return bIsActive;
 }
 
 bool UKismetSystemLibrary::K2_IsTimerPaused(UObject* Object, FString FunctionName)
@@ -455,18 +517,15 @@ bool UKismetSystemLibrary::K2_IsTimerPaused(UObject* Object, FString FunctionNam
 
 bool UKismetSystemLibrary::K2_IsTimerPausedDelegate(FTimerDynamicDelegate Delegate)
 {
+	bool bIsPaused = false;
 	if (Delegate.IsBound())
 	{
 		UWorld* World = GEngine->GetWorldFromContextObject(Delegate.GetUObject());
 		if(World)
 		{
-			auto& TimerManager = World->GetTimerManager();
-			auto Handle = TimerManager.K2_FindDynamicTimerHandle(Delegate);
-			return World->GetTimerManager().IsTimerPaused(Handle);
-		}
-		else
-		{
-			return false;
+			FTimerManager& TimerManager = World->GetTimerManager();
+			FTimerHandle Handle = TimerManager.K2_FindDynamicTimerHandle(Delegate);
+			bIsPaused = TimerManager.IsTimerPaused(Handle);
 		}
 	}
 	else
@@ -474,8 +533,23 @@ bool UKismetSystemLibrary::K2_IsTimerPausedDelegate(FTimerDynamicDelegate Delega
 		UE_LOG(LogBlueprintUserMessages, Warning, 
 			TEXT("IsTimerPaused passed a bad function (%s) or object (%s)"),
 			*Delegate.GetFunctionName().ToString(), *GetNameSafe(Delegate.GetUObject()));
-		return false;
 	}
+	return bIsPaused;
+}
+
+bool UKismetSystemLibrary::K2_IsTimerPausedHandle(UObject* WorldContextObject, FTimerHandle Handle)
+{
+	bool bIsPaused = false;
+	if (Handle.IsValid())
+	{
+		UWorld* World = GEngine->GetWorldFromContextObject( WorldContextObject );
+		if (World)
+		{
+			bIsPaused = World->GetTimerManager().IsTimerPaused(Handle);
+		}
+	}
+
+	return bIsPaused;
 }
 
 bool UKismetSystemLibrary::K2_TimerExists(UObject* Object, FString FunctionName)
@@ -488,18 +562,15 @@ bool UKismetSystemLibrary::K2_TimerExists(UObject* Object, FString FunctionName)
 
 bool UKismetSystemLibrary::K2_TimerExistsDelegate(FTimerDynamicDelegate Delegate)
 {
+	bool bTimerExists = false;
 	if (Delegate.IsBound())
 	{
 		UWorld* World = GEngine->GetWorldFromContextObject(Delegate.GetUObject());
 		if(World)
 		{
-			auto& TimerManager = World->GetTimerManager();
-			auto Handle = TimerManager.K2_FindDynamicTimerHandle(Delegate);
-			return World->GetTimerManager().TimerExists(Handle);
-		}
-		else
-		{
-			return false;
+			FTimerManager& TimerManager = World->GetTimerManager();
+			FTimerHandle Handle = TimerManager.K2_FindDynamicTimerHandle(Delegate);
+			bTimerExists = TimerManager.TimerExists(Handle);
 		}
 	}
 	else
@@ -507,8 +578,23 @@ bool UKismetSystemLibrary::K2_TimerExistsDelegate(FTimerDynamicDelegate Delegate
 		UE_LOG(LogBlueprintUserMessages, Warning,
 			TEXT("TimerExists passed a bad function (%s) or object (%s)"),
 			*Delegate.GetFunctionName().ToString(), *GetNameSafe(Delegate.GetUObject()));
-		return false;
 	}
+	return bTimerExists;
+}
+
+bool UKismetSystemLibrary::K2_TimerExistsHandle(UObject* WorldContextObject, FTimerHandle Handle)
+{
+	bool bTimerExists = false;
+	if (Handle.IsValid())
+	{
+		UWorld* World = GEngine->GetWorldFromContextObject( WorldContextObject );
+		if (World)
+		{
+			bTimerExists = World->GetTimerManager().TimerExists(Handle);
+		}
+	}
+
+	return bTimerExists;
 }
 
 float UKismetSystemLibrary::K2_GetTimerElapsedTime(UObject* Object, FString FunctionName)
@@ -521,18 +607,15 @@ float UKismetSystemLibrary::K2_GetTimerElapsedTime(UObject* Object, FString Func
 
 float UKismetSystemLibrary::K2_GetTimerElapsedTimeDelegate(FTimerDynamicDelegate Delegate)
 {
+	float ElapsedTime = 0.f;
 	if (Delegate.IsBound())
 	{
 		UWorld* World = GEngine->GetWorldFromContextObject(Delegate.GetUObject());
 		if(World)
 		{
-			auto& TimerManager = World->GetTimerManager();
-			auto Handle = TimerManager.K2_FindDynamicTimerHandle(Delegate);
-			return World->GetTimerManager().GetTimerElapsed(Handle);
-		}
-		else
-		{
-			return false;
+			FTimerManager& TimerManager = World->GetTimerManager();
+			FTimerHandle Handle = TimerManager.K2_FindDynamicTimerHandle(Delegate);
+			ElapsedTime = TimerManager.GetTimerElapsed(Handle);
 		}
 	}
 	else
@@ -540,8 +623,23 @@ float UKismetSystemLibrary::K2_GetTimerElapsedTimeDelegate(FTimerDynamicDelegate
 		UE_LOG(LogBlueprintUserMessages, Warning, 
 			TEXT("GetTimerElapsedTime passed a bad function (%s) or object (%s)"), 
 			*Delegate.GetFunctionName().ToString(), *GetNameSafe(Delegate.GetUObject()));
-		return 0.0f;
 	}
+	return ElapsedTime;
+}
+
+float UKismetSystemLibrary::K2_GetTimerElapsedTimeHandle(UObject* WorldContextObject, FTimerHandle Handle)
+{
+	float ElapsedTime = 0.f;
+	if (Handle.IsValid())
+	{
+		UWorld* World = GEngine->GetWorldFromContextObject( WorldContextObject );
+		if (World)
+		{
+			ElapsedTime = World->GetTimerManager().GetTimerElapsed(Handle);
+		}
+	}
+
+	return ElapsedTime;
 }
 
 float UKismetSystemLibrary::K2_GetTimerRemainingTime(UObject* Object, FString FunctionName)
@@ -554,18 +652,15 @@ float UKismetSystemLibrary::K2_GetTimerRemainingTime(UObject* Object, FString Fu
 
 float UKismetSystemLibrary::K2_GetTimerRemainingTimeDelegate(FTimerDynamicDelegate Delegate)
 {
+	float RemainingTime = 0.f;
 	if (Delegate.IsBound())
 	{
 		UWorld* World = GEngine->GetWorldFromContextObject(Delegate.GetUObject());
 		if(World)
 		{
-			auto& TimerManager = World->GetTimerManager();
-			auto Handle = TimerManager.K2_FindDynamicTimerHandle(Delegate);
-			return World->GetTimerManager().GetTimerRemaining(Handle);
-		}
-		else
-		{
-			return false;
+			FTimerManager& TimerManager = World->GetTimerManager();
+			FTimerHandle Handle = TimerManager.K2_FindDynamicTimerHandle(Delegate);
+			RemainingTime = TimerManager.GetTimerRemaining(Handle);
 		}
 	}
 	else
@@ -573,8 +668,23 @@ float UKismetSystemLibrary::K2_GetTimerRemainingTimeDelegate(FTimerDynamicDelega
 		UE_LOG(LogBlueprintUserMessages, Warning, 
 			TEXT("GetTimerRemainingTime passed a bad function (%s) or object (%s)"), 
 			*Delegate.GetFunctionName().ToString(), *GetNameSafe(Delegate.GetUObject()));
-		return 0.0f;
 	}
+	return RemainingTime;
+}
+
+float UKismetSystemLibrary::K2_GetTimerRemainingTimeHandle(UObject* WorldContextObject, FTimerHandle Handle)
+{
+	float RemainingTime = 0.f;
+	if (Handle.IsValid())
+	{
+		UWorld* World = GEngine->GetWorldFromContextObject( WorldContextObject );
+		if (World)
+		{
+			RemainingTime = World->GetTimerManager().GetTimerRemaining(Handle);
+		}
+	}
+
+	return RemainingTime;
 }
 
 void UKismetSystemLibrary::SetIntPropertyByName(UObject* Object, FName PropertyName, int32 Value)
@@ -641,7 +751,7 @@ void UKismetSystemLibrary::SetClassPropertyByName(UObject* Object, FName Propert
 {
 	if (Object && *Value)
 	{
-		auto ClassProp = FindField<UClassProperty>(Object->GetClass(), PropertyName);
+		UClassProperty* ClassProp = FindField<UClassProperty>(Object->GetClass(), PropertyName);
 		if (ClassProp != NULL && Value->IsChildOf(ClassProp->MetaClass)) // check it's the right type
 		{
 			ClassProp->SetObjectPropertyValue_InContainer(Object, *Value);
@@ -677,8 +787,8 @@ void UKismetSystemLibrary::SetAssetPropertyByName(UObject* Object, FName Propert
 {
 	if (Object != NULL)
 	{
-		auto ObjectProp = FindField<UAssetObjectProperty>(Object->GetClass(), PropertyName);
-		auto AssetPtr = (const FAssetPtr*)(&Value);
+		UAssetObjectProperty* ObjectProp = FindField<UAssetObjectProperty>(Object->GetClass(), PropertyName);
+		const FAssetPtr* AssetPtr = (const FAssetPtr*)(&Value);
 		ObjectProp->SetPropertyValue_InContainer(Object, *AssetPtr);
 	}
 }
@@ -687,8 +797,8 @@ void UKismetSystemLibrary::SetAssetClassPropertyByName(UObject* Object, FName Pr
 {
 	if (Object != NULL)
 	{
-		auto ObjectProp = FindField<UAssetClassProperty>(Object->GetClass(), PropertyName);
-		auto AssetPtr = (const FAssetPtr*)(&Value);
+		UAssetClassProperty* ObjectProp = FindField<UAssetClassProperty>(Object->GetClass(), PropertyName);
+		const FAssetPtr* AssetPtr = (const FAssetPtr*)(&Value);
 		ObjectProp->SetPropertyValue_InContainer(Object, *AssetPtr);
 	}
 }
@@ -3085,7 +3195,7 @@ void UKismetSystemLibrary::LoadAsset(UObject* WorldContextObject, const TAssetPt
 		FLatentActionManager& LatentManager = World->GetLatentActionManager();
 		if (LatentManager.FindExistingAction<FLoadAssetAction>(LatentInfo.CallbackTarget, LatentInfo.UUID) == nullptr)
 		{
-			auto NewAction = new FLoadAssetAction(Asset.ToStringReference(), OnLoaded, LatentInfo);
+			FLoadAssetAction* NewAction = new FLoadAssetAction(Asset.ToStringReference(), OnLoaded, LatentInfo);
 			LatentManager.AddNewAction(LatentInfo.CallbackTarget, LatentInfo.UUID, NewAction);
 		}
 	}
@@ -3105,7 +3215,7 @@ void UKismetSystemLibrary::LoadAssetClass(UObject* WorldContextObject, const TAs
 
 		virtual void OnLoaded() override
 		{
-			auto LoadedObject = Cast<UClass>(AssetReference.ResolveObject());
+			UClass* LoadedObject = Cast<UClass>(AssetReference.ResolveObject());
 			OnLoadedCallback.ExecuteIfBound(LoadedObject);
 		}
 	};
@@ -3116,7 +3226,7 @@ void UKismetSystemLibrary::LoadAssetClass(UObject* WorldContextObject, const TAs
 		FLatentActionManager& LatentManager = World->GetLatentActionManager();
 		if (LatentManager.FindExistingAction<FLoadAssetClassAction>(LatentInfo.CallbackTarget, LatentInfo.UUID) == nullptr)
 		{
-			auto NewAction = new FLoadAssetClassAction(AssetClass.ToStringReference(), OnLoaded, LatentInfo);
+			FLoadAssetClassAction* NewAction = new FLoadAssetClassAction(AssetClass.ToStringReference(), OnLoaded, LatentInfo);
 			LatentManager.AddNewAction(LatentInfo.CallbackTarget, LatentInfo.UUID, NewAction);
 		}
 	}
diff --git a/Engine/Source/Runtime/Engine/Public/TimerManager.h b/Engine/Source/Runtime/Engine/Public/TimerManager.h
--- a/Engine/Source/Runtime/Engine/Public/TimerManager.h
+++ b/Engine/Source/Runtime/Engine/Public/TimerManager.h
@@ -69,46 +69,6 @@ struct FTimerUnifiedDelegate
 	}
 };
 
-// Unique handle that can be used to distinguish timers that have identical delegates.
-struct FTimerHandle
-{
-	FTimerHandle()
-	: Handle(INDEX_NONE)
-	{
-
-	}
-
-	bool IsValid() const
-	{
-		return Handle != INDEX_NONE;
-	}
-
-	void Invalidate()
-	{
-		Handle = INDEX_NONE;
-	}
-
-	void MakeValid();
-
-	bool operator==(const FTimerHandle& Other) const
-	{
-		return Handle == Other.Handle;
-	}
-
-	bool operator!=(const FTimerHandle& Other) const
-	{
-		return Handle != Other.Handle;
-	}
-
-	FString ToString() const
-	{
-		return FString::Printf(TEXT("%d"), Handle);
-	}
-
-private:
-	int32 Handle;
-};
-
 enum class ETimerStatus : uint8
 {
 	Pending,
