[UE-16784] Simple bp context menu optimizations.

diff --git a/Engine/Source/Editor/GraphEditor/Private/GraphActionNode.cpp b/Engine/Source/Editor/GraphEditor/Private/GraphActionNode.cpp
--- a/Engine/Source/Editor/GraphEditor/Private/GraphActionNode.cpp
+++ b/Engine/Source/Editor/GraphEditor/Private/GraphActionNode.cpp
@@ -308,7 +308,7 @@ bool FGraphActionNode::IsCategoryNode() const
 //------------------------------------------------------------------------------
 bool FGraphActionNode::IsActionNode() const
 {
-	return HasValidAction();
+	return Actions.Num() != 0;
 }
 
 //------------------------------------------------------------------------------
@@ -361,15 +361,14 @@ bool FGraphActionNode::HasValidAction() const
 //------------------------------------------------------------------------------
 TSharedPtr<FEdGraphSchemaAction> FGraphActionNode::GetPrimaryAction() const
 {
-	TSharedPtr<FEdGraphSchemaAction> PrimaryAction;
-	for (TSharedPtr<FEdGraphSchemaAction> NodeAction : Actions)
+	for (const TSharedPtr<FEdGraphSchemaAction>& NodeAction : Actions)
 	{
 		if (NodeAction.IsValid())
 		{
-			PrimaryAction = NodeAction;
+			return NodeAction;
 		}
 	}
-	return PrimaryAction;
+	return TSharedPtr<FEdGraphSchemaAction>();
 }
 
 //------------------------------------------------------------------------------
diff --git a/Engine/Source/Runtime/Engine/Classes/EdGraph/EdGraphSchema.h b/Engine/Source/Runtime/Engine/Classes/EdGraph/EdGraphSchema.h
--- a/Engine/Source/Runtime/Engine/Classes/EdGraph/EdGraphSchema.h
+++ b/Engine/Source/Runtime/Engine/Classes/EdGraph/EdGraphSchema.h
@@ -298,9 +298,6 @@ public:
 		ActionGroup( const TArray< TSharedPtr<FEdGraphSchemaAction> >& InActions, FString const& RootCategory = TEXT("") );
 
 		/**
-		 * Concatenates RootCategory with the first action's category (RootCategory
-		 * coming first, as a prefix, and the splits the category hierarchy apart 
-		 * into separate entries.
 		 * 
 		 * @param  HierarchyOut	A list of the category tiers that this action should be listed under.
 		 */
@@ -318,8 +315,17 @@ public:
 		TArray< TSharedPtr<FEdGraphSchemaAction> > Actions;
 
 	private:
+		/**
+		 *Concatenates RootCategory with the first action's category (RootCategory
+		 * coming first, as a prefix, and the splits the category hierarchy apart
+		 * into separate entries.
+		 */
+		void InitCategoryChain();
+
 		/** The category to list this entry under (could be left empty, as it gets concatenated with the first sub-action's category) */
 		FString RootCategory;
+		/** The chain of categories */
+		TArray<FString> CategoryChain;
 	};
 private:
 
diff --git a/Engine/Source/Runtime/Engine/Private/EdGraph/EdGraphSchema.cpp b/Engine/Source/Runtime/Engine/Private/EdGraph/EdGraphSchema.cpp
--- a/Engine/Source/Runtime/Engine/Private/EdGraph/EdGraphSchema.cpp
+++ b/Engine/Source/Runtime/Engine/Private/EdGraph/EdGraphSchema.cpp
@@ -55,19 +55,39 @@ FGraphActionListBuilderBase::ActionGroup::ActionGroup( TSharedPtr<FEdGraphSchema
 	: RootCategory(CategoryPrefix)
 {
 	Actions.Add( InAction );
+	InitCategoryChain();
 }
 
 FGraphActionListBuilderBase::ActionGroup::ActionGroup( const TArray< TSharedPtr<FEdGraphSchemaAction> >& InActions, FString const& CategoryPrefix/* = TEXT("") */ )
 	: RootCategory(CategoryPrefix)
 {
 	Actions = InActions;
+	InitCategoryChain();
 }
 
 void FGraphActionListBuilderBase::ActionGroup::GetCategoryChain(TArray<FString>& HierarchyOut) const
 {
 #if WITH_EDITOR
+	HierarchyOut = CategoryChain;
+#endif
+}
+
+void FGraphActionListBuilderBase::ActionGroup::PerformAction( class UEdGraph* ParentGraph, TArray<UEdGraphPin*>& FromPins, const FVector2D Location )
+{
+	for ( int32 ActionIndex = 0; ActionIndex < Actions.Num(); ActionIndex++ )
+	{
+		TSharedPtr<FEdGraphSchemaAction> CurrentAction = Actions[ActionIndex];
+		if ( CurrentAction.IsValid() )
+		{
+			CurrentAction->PerformAction( ParentGraph, FromPins, Location );
+		}
+	}
+}
+
+void FGraphActionListBuilderBase::ActionGroup::InitCategoryChain()
+{
 	static FString const CategoryDelim("|");
-	FEditorCategoryUtils::GetCategoryDisplayString(RootCategory).ParseIntoArray(HierarchyOut, *CategoryDelim, true);
+	FEditorCategoryUtils::GetCategoryDisplayString(RootCategory).ParseIntoArray(CategoryChain, *CategoryDelim, true);
 
 	if (Actions.Num() > 0)
 	{
@@ -76,26 +96,13 @@ void FGraphActionListBuilderBase::ActionGroup::GetCategoryChain(TArray<FString>&
 		FString SubCategory = FEditorCategoryUtils::GetCategoryDisplayString(Actions[0]->Category.ToString());
 		SubCategory.ParseIntoArray(SubCategoryChain, *CategoryDelim, true);
 
-		HierarchyOut.Append(SubCategoryChain);
+		CategoryChain.Append(SubCategoryChain);
 	}
 
-	for (FString& Category : HierarchyOut)
+	for (FString& Category : CategoryChain)
 	{
 		Category.Trim();
 	}
-#endif
-}
-
-void FGraphActionListBuilderBase::ActionGroup::PerformAction( class UEdGraph* ParentGraph, TArray<UEdGraphPin*>& FromPins, const FVector2D Location )
-{
-	for ( int32 ActionIndex = 0; ActionIndex < Actions.Num(); ActionIndex++ )
-	{
-		TSharedPtr<FEdGraphSchemaAction> CurrentAction = Actions[ActionIndex];
-		if ( CurrentAction.IsValid() )
-		{
-			CurrentAction->PerformAction( ParentGraph, FromPins, Location );
-		}
-	}
 }
 
 /////////////////////////////////////////////////////
