Blueprints: Added a display name to the input pin of BreakRotator to make it look more like a built in break node

diff --git a/Engine/Source/Runtime/Engine/Classes/Kismet/KismetMathLibrary.h b/Engine/Source/Runtime/Engine/Classes/Kismet/KismetMathLibrary.h
--- a/Engine/Source/Runtime/Engine/Classes/Kismet/KismetMathLibrary.h
+++ b/Engine/Source/Runtime/Engine/Classes/Kismet/KismetMathLibrary.h
@@ -1269,7 +1269,8 @@ class ENGINE_API UKismetMathLibrary : public UBlueprintFunctionLibrary
 
 	/** Breaks apart a rotator into {Roll, Pitch, Yaw} angles in degrees */
 	UFUNCTION(BlueprintPure, Category = "Math|Rotator", meta = (Keywords = "rotation rotate rotator breakrotator", NativeBreakFunc))
-	static void BreakRotator(FRotator InRot,
+	static void BreakRotator(
+		UPARAM(DisplayName="Rotation") FRotator InRot,
 		UPARAM(DisplayName="X (Roll)") float& Roll,
 		UPARAM(DisplayName="Y (Pitch)") float& Pitch,
 		UPARAM(DisplayName="Z (Yaw)") float& Yaw);
