#ANIM: Support Append frames in the begin/end

- Also fixed notifies/curves editing

https://jira.ol.epicgames.net/browse/UE-6558

diff --git a/Engine/Source/Editor/KismetWidgets/Private/SScrubControlPanel.cpp b/Engine/Source/Editor/KismetWidgets/Private/SScrubControlPanel.cpp
--- a/Engine/Source/Editor/KismetWidgets/Private/SScrubControlPanel.cpp
+++ b/Engine/Source/Editor/KismetWidgets/Private/SScrubControlPanel.cpp
@@ -57,6 +57,7 @@ void SScrubControlPanel::Construct( const SScrubControlPanel::FArguments& InArgs
 				.OnSetInputViewRange(InArgs._OnSetInputViewRange)
 				.OnCropAnimSequence(InArgs._OnCropAnimSequence)
 				.OnAddAnimSequence(InArgs._OnAddAnimSequence)
+				.OnAppendAnimSequence(InArgs._OnAppendAnimSequence)
 				.OnReZeroAnimSequence(InArgs._OnReZeroAnimSequence)
 				.bAllowZoom(InArgs._bAllowZoom)
 				.DraggableBars(InArgs._DraggableBars)
diff --git a/Engine/Source/Editor/KismetWidgets/Private/SScrubWidget.cpp b/Engine/Source/Editor/KismetWidgets/Private/SScrubWidget.cpp
--- a/Engine/Source/Editor/KismetWidgets/Private/SScrubWidget.cpp
+++ b/Engine/Source/Editor/KismetWidgets/Private/SScrubWidget.cpp
@@ -2,7 +2,7 @@
 
 #include "KismetWidgetsPrivatePCH.h"
 #include "SScrubWidget.h"
-
+#include "STextEntryPopup.h"
 
 #define LOCTEXT_NAMESPACE "ScrubWidget"
 
@@ -41,6 +41,7 @@ void SScrubWidget::Construct( const SScrubWidget::FArguments& InArgs )
 	OnSetInputViewRange = InArgs._OnSetInputViewRange;
 	OnCropAnimSequence = InArgs._OnCropAnimSequence;
 	OnAddAnimSequence = InArgs._OnAddAnimSequence;
+	OnAppendAnimSequence = InArgs._OnAppendAnimSequence;
 	OnReZeroAnimSequence = InArgs._OnReZeroAnimSequence;
 
 	DraggableBars = InArgs._DraggableBars;
@@ -437,7 +438,7 @@ void SScrubWidget::CreateContextMenu(float CurrentFrameTime, const FPointerEvent
 					//Corrected frame time based on selected frame number
 					float CorrectedFrameTime = CurrentFrameFraction * SequenceLength.Get();
 
-					Action = FUIAction(FExecuteAction::CreateSP(this, &SScrubWidget::OnSequenceCropped, true, CorrectedFrameTime));
+					Action = FUIAction(FExecuteAction::CreateSP(this, &SScrubWidget::OnSequenceCroppedCalled, true, CorrectedFrameTime));
 					Label = FText::Format(LOCTEXT("RemoveTillFrame", "Remove frame 0 to frame {0}"), FText::AsNumber(CurrentFrameNumber));
 					MenuBuilder.AddMenuEntry(Label, LOCTEXT("RemoveBefore_ToolTip", "Remove sequence before current position"), FSlateIcon(), Action);
 				}
@@ -450,7 +451,7 @@ void SScrubWidget::CreateContextMenu(float CurrentFrameTime, const FPointerEvent
 				{
 					float NextFrameFraction = float(NextFrameNumber) / (float)NumOfKeys.Get();
 					float NextFrameTime = NextFrameFraction * SequenceLength.Get();
-					Action = FUIAction(FExecuteAction::CreateSP(this, &SScrubWidget::OnSequenceCropped, false, NextFrameTime));
+					Action = FUIAction(FExecuteAction::CreateSP(this, &SScrubWidget::OnSequenceCroppedCalled, false, NextFrameTime));
 					Label = FText::Format(LOCTEXT("RemoveFromFrame", "Remove from frame {0} to frame {1}"), FText::AsNumber(NextFrameNumber), FText::AsNumber(NumOfKeys.Get()));
 					MenuBuilder.AddMenuEntry(Label, LOCTEXT("RemoveAfter_ToolTip", "Remove sequence after current position"), FSlateIcon(), Action);
 				}
@@ -458,22 +459,37 @@ void SScrubWidget::CreateContextMenu(float CurrentFrameTime, const FPointerEvent
 
 			if (OnAddAnimSequence.IsBound())
 			{
+				MenuBuilder.AddMenuSeparator();
+
 				//Corrected frame time based on selected frame number
 				float CorrectedFrameTime = CurrentFrameFraction * SequenceLength.Get();
 
-				Action = FUIAction(FExecuteAction::CreateSP(this, &SScrubWidget::OnSequenceAdded, true, CurrentFrameNumber));
+				Action = FUIAction(FExecuteAction::CreateSP(this, &SScrubWidget::OnSequenceAddedCalled, true, CurrentFrameNumber));
 				Label = FText::Format(LOCTEXT("InsertBeforeCurrentFrame", "Insert frame before {0}"), FText::AsNumber(CurrentFrameNumber));
 				MenuBuilder.AddMenuEntry(Label, LOCTEXT("InsertBefore_ToolTip", "Insert a frame before current position"), FSlateIcon(), Action);
 
-				Action = FUIAction(FExecuteAction::CreateSP(this, &SScrubWidget::OnSequenceAdded, false, CurrentFrameNumber));
+				Action = FUIAction(FExecuteAction::CreateSP(this, &SScrubWidget::OnSequenceAddedCalled, false, CurrentFrameNumber));
 				Label = FText::Format(LOCTEXT("InsertAfterCurrentFrame", "Insert frame after {0}"), FText::AsNumber(CurrentFrameNumber));
 				MenuBuilder.AddMenuEntry(Label, LOCTEXT("InsertAfter_ToolTip", "Insert a frame after current position"), FSlateIcon(), Action);
 			}
 
+			if(OnAppendAnimSequence.IsBound())
+			{
+				MenuBuilder.AddMenuSeparator();
+
+				//Corrected frame time based on selected frame number
+				Action = FUIAction(FExecuteAction::CreateSP(this, &SScrubWidget::OnShowPopupOfAppendAnimation, true));
+				MenuBuilder.AddMenuEntry(LOCTEXT("AppendBegin", "Append in the beginning"), LOCTEXT("AppendBegin_ToolTip", "Append in the beginning"), FSlateIcon(), Action);
+
+				Action = FUIAction(FExecuteAction::CreateSP(this, &SScrubWidget::OnShowPopupOfAppendAnimation, false));
+				MenuBuilder.AddMenuEntry(LOCTEXT("AppendEnd", "Append at the end"), LOCTEXT("AppendEnd_ToolTip", "Append at the end"), FSlateIcon(), Action);
+			}
+
 			if (OnReZeroAnimSequence.IsBound())
 			{
+				MenuBuilder.AddMenuSeparator();
 				//Menu - "ReZero"
-				Action = FUIAction(FExecuteAction::CreateSP(this, &SScrubWidget::OnReZero));
+				Action = FUIAction(FExecuteAction::CreateSP(this, &SScrubWidget::OnReZeroCalled));
 				Label = FText::Format(LOCTEXT("ReZeroAtFrame", "ReZero at frame {0}"), FText::AsNumber(CurrentFrameNumber));
 				MenuBuilder.AddMenuEntry(Label, LOCTEXT("ReZeroAtFrame_ToolTip", "Resets the root track of the frame to (0, 0, 0), and apply the difference to all root transform of the sequence. It moves whole sequence to the amount of current root transform. "), FSlateIcon(), Action);
 			}
@@ -486,7 +502,7 @@ void SScrubWidget::CreateContextMenu(float CurrentFrameTime, const FPointerEvent
 	}
 }
 
-void SScrubWidget::OnSequenceCropped( bool bFromStart, float CurrentFrameTime )
+void SScrubWidget::OnSequenceCroppedCalled( bool bFromStart, float CurrentFrameTime )
 {
 	OnCropAnimSequence.ExecuteIfBound( bFromStart, CurrentFrameTime );
 
@@ -494,20 +510,53 @@ void SScrubWidget::OnSequenceCropped( bool bFromStart, float CurrentFrameTime )
 	OnSetInputViewRange.ExecuteIfBound( ViewInputMin.Get(), ViewInputMax.Get() );
 }
 
-void SScrubWidget::OnSequenceAdded(bool bBefore, int32 CurrentFrameNumber)
+void SScrubWidget::OnSequenceAddedCalled(bool bBefore, int32 CurrentFrameNumber)
 {
 	OnAddAnimSequence.ExecuteIfBound(bBefore, CurrentFrameNumber);
 
 	//Update scrubs new length to be new Sequence Length
 	// @Todo fixme: this whole thing needs to change to "Refresh" 
-	// - including the OnSequenceCropped
+	// - including the OnSequenceCroppedCalled
 	OnSetInputViewRange.ExecuteIfBound(ViewInputMin.Get(), SequenceLength.Get());
 }
 
-void SScrubWidget::OnReZero()
+void SScrubWidget::OnReZeroCalled()
 {
 	OnReZeroAnimSequence.ExecuteIfBound();
 }
 
+void SScrubWidget::OnShowPopupOfAppendAnimation(bool bBegin)
+{
+	TSharedRef<STextEntryPopup> TextEntry =
+		SNew(STextEntryPopup)
+		.Label(LOCTEXT("AppendAnim_AskNumFrames", "Number of Frames to Append"))
+		.OnTextCommitted(this, &SScrubWidget::OnSequenceAppendedCalled, bBegin);
+
+	// Show dialog to enter new track name
+	FSlateApplication::Get().PushMenu(
+		SharedThis(this),
+		FWidgetPath(),
+		TextEntry,
+		FSlateApplication::Get().GetCursorPos(),
+		FPopupTransitionEffect(FPopupTransitionEffect::TypeInPopup)
+		);
+}
+
+void SScrubWidget::OnSequenceAppendedCalled(const FText & InNewGroupText, ETextCommit::Type CommitInfo, bool bBegin)
+{
+	// just a concern
+	const static int32 MaxFrame = 1000;
+
+	// handle only onEnter. This is a big thing to apply when implicit focus change or any other event
+	if (CommitInfo == ETextCommit::OnEnter)
+	{
+		int32 NumFrames = FCString::Atoi(*InNewGroupText.ToString());
+		if (NumFrames > 0 && NumFrames < MaxFrame)
+		{
+			OnAppendAnimSequence.ExecuteIfBound(bBegin, NumFrames);
+			FSlateApplication::Get().DismissAllMenus();
+		}
+	}
+}
 
 #undef LOCTEXT_NAMESPACE
diff --git a/Engine/Source/Editor/KismetWidgets/Public/SScrubControlPanel.h b/Engine/Source/Editor/KismetWidgets/Public/SScrubControlPanel.h
--- a/Engine/Source/Editor/KismetWidgets/Public/SScrubControlPanel.h
+++ b/Engine/Source/Editor/KismetWidgets/Public/SScrubControlPanel.h
@@ -63,6 +63,8 @@ public:
 		SLATE_EVENT( FOnCropAnimSequence, OnCropAnimSequence )
 		/** Called when an frane is added before/after a selected frame */
 		SLATE_EVENT(FOnAddAnimSequence, OnAddAnimSequence)
+		/** Called when an frame is appended in the beginning or at the end */
+		SLATE_EVENT(FOnAppendAnimSequence, OnAppendAnimSequence)
 		/** Called to zero out selected frame's translation from origin */
 		SLATE_EVENT( FSimpleDelegate, OnReZeroAnimSequence )
 		SLATE_ATTRIBUTE( bool, IsRealtimeStreamingMode )
diff --git a/Engine/Source/Editor/KismetWidgets/Public/SScrubWidget.h b/Engine/Source/Editor/KismetWidgets/Public/SScrubWidget.h
--- a/Engine/Source/Editor/KismetWidgets/Public/SScrubWidget.h
+++ b/Engine/Source/Editor/KismetWidgets/Public/SScrubWidget.h
@@ -13,6 +13,7 @@
  */
 DECLARE_DELEGATE_TwoParams( FOnCropAnimSequence, bool, float )
 DECLARE_DELEGATE_TwoParams( FOnAddAnimSequence, bool, int32 )
+DECLARE_DELEGATE_TwoParams( FOnAppendAnimSequence, bool, int32 )
 DECLARE_DELEGATE_TwoParams( FOnScrubBarDrag, int32, float)
 
 class KISMETWIDGETS_API SScrubWidget : public SCompoundWidget
@@ -52,6 +53,8 @@ public:
 		SLATE_EVENT( FOnCropAnimSequence, OnCropAnimSequence )
 		/** Called when an frane is added before/after a selected frame */
 		SLATE_EVENT( FOnAddAnimSequence, OnAddAnimSequence )
+		/** Called when an frame is appended in the beginning or at the end */
+		SLATE_EVENT(FOnAppendAnimSequence, OnAppendAnimSequence)
 		/** Called to zero out selected frame's translation from origin */
 		SLATE_EVENT( FSimpleDelegate, OnReZeroAnimSequence )
 		/** Optional, additional values to draw on the timeline **/
@@ -94,13 +97,19 @@ private:
 	void CreateContextMenu( float CurrentFrameTime, const FPointerEvent& MouseEvent );
 
 	/** Function to crop animation sequence before/after selected frame */
-	void OnSequenceCropped( bool bFromStart, float CurrentFrameTime );
+	void OnSequenceCroppedCalled( bool bFromStart, float CurrentFrameTime );
 
 	/** Function to crop animation sequence before/after selected frame */
-	void OnSequenceAdded(bool bBefore, int32 CurrentFrameNumber);
+	void OnSequenceAddedCalled(bool bBefore, int32 CurrentFrameNumber);
+
+	/** Function to append frames in the beginning or at the end*/
+	void OnSequenceAppendedCalled(const FText & InNewGroupText, ETextCommit::Type CommitInfo, bool bBegin);
+
+	/** Function to ask how many frames to append */
+	void OnShowPopupOfAppendAnimation(bool bBegin);
 
 	/** Function to zero out translation of the selected frame */
-	void OnReZero();
+	void OnReZeroCalled();
 
 	TAttribute<float> ValueAttribute;
 	FOnFloatValueChanged OnValueChanged;
@@ -109,10 +118,11 @@ private:
 
 	TAttribute<float> ViewInputMin;
 	TAttribute<float> ViewInputMax;
-	FOnSetInputViewRange OnSetInputViewRange;
-	FOnCropAnimSequence OnCropAnimSequence;
-	FOnAddAnimSequence	OnAddAnimSequence;
-	FSimpleDelegate		OnReZeroAnimSequence;
+	FOnSetInputViewRange	OnSetInputViewRange;
+	FOnCropAnimSequence		OnCropAnimSequence;
+	FOnAddAnimSequence		OnAddAnimSequence;
+	FOnAppendAnimSequence	OnAppendAnimSequence;
+	FSimpleDelegate			OnReZeroAnimSequence;
 
 	/** Dragagble bars are generic lines drawn on the scrub widget that can be dragged with the mouse. This is very bare bones and just represents drawing/moving float values */
 	TAttribute<TArray<float>>	DraggableBars;
diff --git a/Engine/Source/Editor/Persona/Private/SAnimCurvePanel.cpp b/Engine/Source/Editor/Persona/Private/SAnimCurvePanel.cpp
--- a/Engine/Source/Editor/Persona/Private/SAnimCurvePanel.cpp
+++ b/Engine/Source/Editor/Persona/Private/SAnimCurvePanel.cpp
@@ -179,6 +179,9 @@ public:
 
 	// Bound to attribute for curve name, uses curve interface to request from skeleton
 	FText GetCurveName(USkeleton::AnimCurveUID Uid) const;
+
+	float GetLength() const { return PanelPtr.Pin()->GetLength(); }
+	TOptional<float> GetOptionalLength() const { return GetLength(); }
 };
 
 //////////////////////////////////////////////////////////////////////////
@@ -222,9 +225,8 @@ void SCurveEdTrack::Construct(const FArguments& InArgs)
 			.ViewMinInput(InArgs._ViewInputMin)
 			.ViewMaxInput(InArgs._ViewInputMax)
 			.DataMinInput(0.f)
-			.DataMaxInput(Sequence->SequenceLength)
-			// @fixme fix this to delegate
-			.TimelineLength(Sequence->SequenceLength)
+			.DataMaxInput(this, &SCurveEdTrack::GetOptionalLength)
+			.TimelineLength(this, &SCurveEdTrack::GetLength)
 			.NumberOfKeys(NumberOfKeys)
 			.DesiredSize(this, &SCurveEdTrack::GetDesiredSize)
 			.OnSetInputViewRange(InArgs._OnSetInputViewRange)
diff --git a/Engine/Source/Editor/Persona/Private/SAnimCurvePanel.h b/Engine/Source/Editor/Persona/Private/SAnimCurvePanel.h
--- a/Engine/Source/Editor/Persona/Private/SAnimCurvePanel.h
+++ b/Engine/Source/Editor/Persona/Private/SAnimCurvePanel.h
@@ -146,6 +146,11 @@ public:
 	*/
 	void UpdatePanel();
 
+	/**
+	 * Get Sequence Length 
+	 */
+	float GetLength() const { return Sequence->SequenceLength; }
+
 private:
 	TWeakPtr<FPersona> WeakPersona;
 	TSharedPtr<SSplitter> PanelSlot;
diff --git a/Engine/Source/Editor/Persona/Private/SAnimTrackCurvePanel.cpp b/Engine/Source/Editor/Persona/Private/SAnimTrackCurvePanel.cpp
--- a/Engine/Source/Editor/Persona/Private/SAnimTrackCurvePanel.cpp
+++ b/Engine/Source/Editor/Persona/Private/SAnimTrackCurvePanel.cpp
@@ -253,6 +253,9 @@ public:
 
 	// Bound to attribute for curve name, uses curve interface to request from skeleton
 	FText GetCurveName(USkeleton::AnimCurveUID Uid, ETransformCurve::Type Type) const;
+
+	float GetLength() const { return PanelPtr.Pin()->GetLength(); }
+	TOptional<float> GetOptionalLength() const { return GetLength(); }
 };
 
 //////////////////////////////////////////////////////////////////////////
@@ -299,9 +302,8 @@ void STransformCurveEdTrack::Construct(const FArguments& InArgs)
 					.ViewMinInput(InArgs._ViewInputMin)
 					.ViewMaxInput(InArgs._ViewInputMax)
 					.DataMinInput(0.f)
-					.DataMaxInput(Sequence->SequenceLength)
-					// @fixme fix this to delegate
-					.TimelineLength(Sequence->SequenceLength)
+					.DataMaxInput(this, &STransformCurveEdTrack::GetOptionalLength)
+					.TimelineLength(this, &STransformCurveEdTrack::GetLength)
 					.NumberOfKeys(NumberOfKeys)
 					.DesiredSize(this, &STransformCurveEdTrack::GetDesiredSize)
 					.OnSetInputViewRange(InArgs._OnSetInputViewRange)
@@ -364,9 +366,8 @@ void STransformCurveEdTrack::Construct(const FArguments& InArgs)
 					.ViewMinInput(InArgs._ViewInputMin)
 					.ViewMaxInput(InArgs._ViewInputMax)
 					.DataMinInput(0.f)
-					.DataMaxInput(Sequence->SequenceLength)
-					// @fixme fix this to delegate
-					.TimelineLength(Sequence->SequenceLength)
+					.DataMaxInput(this, &STransformCurveEdTrack::GetOptionalLength)
+					.TimelineLength(this, &STransformCurveEdTrack::GetLength)
 					.NumberOfKeys(NumberOfKeys)
 					.DesiredSize(this, &STransformCurveEdTrack::GetDesiredSize)
 					.OnSetInputViewRange(InArgs._OnSetInputViewRange)
@@ -414,9 +415,8 @@ void STransformCurveEdTrack::Construct(const FArguments& InArgs)
 					.ViewMinInput(InArgs._ViewInputMin)
 					.ViewMaxInput(InArgs._ViewInputMax)
 					.DataMinInput(0.f)
-					.DataMaxInput(Sequence->SequenceLength)
-					// @fixme fix this to delegate
-					.TimelineLength(Sequence->SequenceLength)
+					.DataMaxInput(this, &STransformCurveEdTrack::GetOptionalLength)
+					.TimelineLength(this, &STransformCurveEdTrack::GetLength)
 					.NumberOfKeys(NumberOfKeys)
 					.DesiredSize(this, &STransformCurveEdTrack::GetDesiredSize)
 					.OnSetInputViewRange(InArgs._OnSetInputViewRange)
diff --git a/Engine/Source/Editor/Persona/Private/SAnimTrackCurvePanel.h b/Engine/Source/Editor/Persona/Private/SAnimTrackCurvePanel.h
--- a/Engine/Source/Editor/Persona/Private/SAnimTrackCurvePanel.h
+++ b/Engine/Source/Editor/Persona/Private/SAnimTrackCurvePanel.h
@@ -75,6 +75,9 @@ public:
 	/** Get Context Menu Per Track **/
 	TSharedRef<SWidget> CreateCurveContextMenu(USkeleton::AnimCurveUID CurveUid) const;
 
+	/** Get Length of Sequence */
+	float GetLength() { return (Sequence)? Sequence->SequenceLength : 0.f; }
+
 private:
 	TWeakPtr<FPersona> WeakPersona;
 	TSharedPtr<SSplitter> PanelSlot;
diff --git a/Engine/Source/Editor/Persona/Private/SAnimationScrubPanel.cpp b/Engine/Source/Editor/Persona/Private/SAnimationScrubPanel.cpp
--- a/Engine/Source/Editor/Persona/Private/SAnimationScrubPanel.cpp
+++ b/Engine/Source/Editor/Persona/Private/SAnimationScrubPanel.cpp
@@ -16,6 +16,7 @@ void SAnimationScrubPanel::Construct( const SAnimationScrubPanel::FArguments& In
 	bSliderBeingDragged = false;
 	PersonaPtr = InArgs._Persona;
 	LockedSequence = InArgs._LockedSequence;
+	OnSetInputViewRange = InArgs._OnSetInputViewRange;
 
 	// register delegate for anim change notification
 	PersonaPtr.Pin()->RegisterOnAnimChanged( FPersona::FOnAnimChanged::CreateSP(this, &SAnimationScrubPanel::AnimChanged) );
@@ -53,6 +54,7 @@ void SAnimationScrubPanel::Construct( const SAnimationScrubPanel::FArguments& In
 			.OnSetInputViewRange(InArgs._OnSetInputViewRange)
 			.OnCropAnimSequence( this, &SAnimationScrubPanel::OnCropAnimSequence )
 			.OnAddAnimSequence( this, &SAnimationScrubPanel::OnInsertAnimSequence )
+			.OnAppendAnimSequence(this, &SAnimationScrubPanel::OnAppendAnimSequence )
 			.OnReZeroAnimSequence( this, &SAnimationScrubPanel::OnReZeroAnimSequence )
 			.bAllowZoom(InArgs._bAllowZoom)
 			.IsRealtimeStreamingMode(this, &SAnimationScrubPanel::IsRealtimeStreamingMode)
@@ -443,35 +445,6 @@ void SAnimationScrubPanel::OnCropAnimSequence( bool bFromStart, float CurrentTim
 				//Call modify to restore anim sequence current state
 				AnimSequence->Modify();
 
-				//Adjust anim notify time based on current time and bFromStart flag. 
-				for( int32 i=0; i<AnimSequence->Notifies.Num(); i++ )
-				{
-					FAnimNotifyEvent& Notify = AnimSequence->Notifies[i];
-					float NotifyTime = Notify.GetTime();
-					if( bFromStart )
-					{
-						//Cropping from start to current time
-						if( NotifyTime <= CurrentTime )
-						{
-							Notify.SetTime(0.0f);
-							Notify.TriggerTimeOffset = GetTriggerTimeOffsetForType(EAnimEventTriggerOffsets::OffsetAfter);
-						}
-						else
-						{
-							Notify.SetTime(NotifyTime - CurrentTime);
-						}
-					}
-					else
-					{
-						//Cropping from current time to the end.
-						if( NotifyTime >= CurrentTime )
-						{
-							Notify.SetTime(CurrentTime);
-							Notify.TriggerTimeOffset = GetTriggerTimeOffsetForType(EAnimEventTriggerOffsets::OffsetBefore);
-						}
-					}
-				}
-
 				// Crop the raw anim data.
 				AnimSequence->CropRawAnimData( CurrentTime, bFromStart );
 				// Recompress animation from Raw.
@@ -479,37 +452,66 @@ void SAnimationScrubPanel::OnCropAnimSequence( bool bFromStart, float CurrentTim
 
 				//Resetting slider position to the first frame
 				PreviewInstance->SetPosition( 0.0f, false );
+
+				OnSetInputViewRange.ExecuteIfBound(0, AnimSequence->SequenceLength);
 			}
 		}
 	}
 }
 
+void SAnimationScrubPanel::OnAppendAnimSequence( bool bFromStart, int32 NumOfFrames )
+{
+	UAnimSingleNodeInstance* PreviewInstance = GetPreviewInstance();
+	if(PreviewInstance && PreviewInstance->CurrentAsset)
+	{
+		UAnimSequence* AnimSequence = Cast<UAnimSequence>(PreviewInstance->CurrentAsset);
+		if(AnimSequence)
+		{
+			const FScopedTransaction Transaction(LOCTEXT("InsertAnimSequence", "Insert Animation Sequence"));
+
+			//Call modify to restore slider position
+			PreviewInstance->Modify();
+
+			//Call modify to restore anim sequence current state
+			AnimSequence->Modify();
+
+			// Crop the raw anim data.
+			int32 StartFrame = (bFromStart)? 0 : AnimSequence->NumFrames - 1;
+			int32 EndFrame = StartFrame + NumOfFrames;
+			int32 CopyFrame = StartFrame;
+			AnimSequence->InsertFramesToRawAnimData(StartFrame, EndFrame, CopyFrame);
+			// Recompress animation from Raw.
+			FAnimationUtils::CompressAnimSequence(AnimSequence, false, false);
+
+			OnSetInputViewRange.ExecuteIfBound(0, AnimSequence->SequenceLength);
+		}
+	}
+}
+
 void SAnimationScrubPanel::OnInsertAnimSequence( bool bBefore, int32 CurrentFrame )
 {
 	UAnimSingleNodeInstance* PreviewInstance = GetPreviewInstance();
-	if(PreviewInstance)
+	if(PreviewInstance && PreviewInstance->CurrentAsset)
 	{
-		float Length = PreviewInstance->GetLength();
-		if(PreviewInstance->CurrentAsset)
+		UAnimSequence* AnimSequence = Cast<UAnimSequence>(PreviewInstance->CurrentAsset);
+		if(AnimSequence)
 		{
-			UAnimSequence* AnimSequence = Cast<UAnimSequence>(PreviewInstance->CurrentAsset);
-			if(AnimSequence)
-			{
-				const FScopedTransaction Transaction(LOCTEXT("InsertAnimSequence", "Insert Animation Sequence"));
+			const FScopedTransaction Transaction(LOCTEXT("InsertAnimSequence", "Insert Animation Sequence"));
 
-				//Call modify to restore slider position
-				PreviewInstance->Modify();
+			//Call modify to restore slider position
+			PreviewInstance->Modify();
 
-				//Call modify to restore anim sequence current state
-				AnimSequence->Modify();
+			//Call modify to restore anim sequence current state
+			AnimSequence->Modify();
 
-				// Crop the raw anim data.
-				int32 StartFrame = (bBefore)? CurrentFrame : CurrentFrame + 1;
-				int32 EndFrame = StartFrame + 1;
-				AnimSequence->InsertFramesToRawAnimData(StartFrame, EndFrame, CurrentFrame);
-				// Recompress animation from Raw.
-				FAnimationUtils::CompressAnimSequence(AnimSequence, false, false);
-			}
+			// Crop the raw anim data.
+			int32 StartFrame = (bBefore)? CurrentFrame : CurrentFrame + 1;
+			int32 EndFrame = StartFrame + 1;
+			AnimSequence->InsertFramesToRawAnimData(StartFrame, EndFrame, CurrentFrame);
+			// Recompress animation from Raw.
+			FAnimationUtils::CompressAnimSequence(AnimSequence, false, false);
+
+			OnSetInputViewRange.ExecuteIfBound(0, AnimSequence->SequenceLength);
 		}
 	}
 }
diff --git a/Engine/Source/Editor/Persona/Private/SAnimationScrubPanel.h b/Engine/Source/Editor/Persona/Private/SAnimationScrubPanel.h
--- a/Engine/Source/Editor/Persona/Private/SAnimationScrubPanel.h
+++ b/Engine/Source/Editor/Persona/Private/SAnimationScrubPanel.h
@@ -46,6 +46,7 @@ public:
 	void ReplaceLockedSequence(class UAnimSequenceBase * NewLockedSequence);
 protected:
 	TWeakPtr<FPersona> PersonaPtr;
+	FOnSetInputViewRange OnSetInputViewRange;
 
 	bool bSliderBeingDragged;
 
@@ -66,6 +67,7 @@ protected:
 
 	/** Function to crop animation sequence before/after selected frame */
 	void OnCropAnimSequence( bool bFromStart, float CurrentTime );
+	void OnAppendAnimSequence( bool bFromStart, int32 NumOfFrames );
 	void OnInsertAnimSequence( bool bBefore, int32 CurrentFrame );
 
 	/** Function to zero out translation of the selected frame */
diff --git a/Engine/Source/Runtime/Engine/Classes/Animation/AnimSequence.h b/Engine/Source/Runtime/Engine/Classes/Animation/AnimSequence.h
--- a/Engine/Source/Runtime/Engine/Classes/Animation/AnimSequence.h
+++ b/Engine/Source/Runtime/Engine/Classes/Animation/AnimSequence.h
@@ -833,7 +833,11 @@ private:
 	 */
 	int32 InsertTrack(const FName& BoneName);
 
-	void ResizeSequence(float NewLength, int32 NewNumFrames);
+	/**
+	 * Utility function to resize the sequence
+	 * It rearranges curve data + notifies
+	 */
+	void ResizeSequence(float NewLength, int32 NewNumFrames, bool bInsert, int32 StartFrame/*inclusive */, int32 EndFrame/*inclusive*/);
 
 #endif
 
diff --git a/Engine/Source/Runtime/Engine/Classes/Curves/CurveBase.h b/Engine/Source/Runtime/Engine/Classes/Curves/CurveBase.h
--- a/Engine/Source/Runtime/Engine/Classes/Curves/CurveBase.h
+++ b/Engine/Source/Runtime/Engine/Classes/Curves/CurveBase.h
@@ -384,7 +384,7 @@ public:
 	void AutoSetTangents(float Tension = 0.f);
 
 	/** Resize curve length to the [MinTimeRange, MaxTimeRange] */
-	void ResizeTimeRange(float NewMinTimeRange, float NewMaxTimeRange);
+	void ReadjustTimeRange(float NewMinTimeRange, float NewMaxTimeRange, bool bInsert/* whether insert or remove*/, float OldStartTime, float OldEndTime);
 
 	/** Determine if two RichCurves are the same */
 	bool operator == (const FRichCurve& Curve) const;
diff --git a/Engine/Source/Runtime/Engine/Private/Animation/AnimCurveTypes.cpp b/Engine/Source/Runtime/Engine/Private/Animation/AnimCurveTypes.cpp
--- a/Engine/Source/Runtime/Engine/Private/Animation/AnimCurveTypes.cpp
+++ b/Engine/Source/Runtime/Engine/Private/Animation/AnimCurveTypes.cpp
@@ -61,9 +61,9 @@ void FFloatCurve::UpdateOrAddKey(float NewKey, float CurrentTime)
 	FloatCurve.UpdateOrAddKey(CurrentTime, NewKey);
 }
 
-void FFloatCurve::Resize(float NewLength)
+void FFloatCurve::Resize(float NewLength, bool bInsert/* whether insert or remove*/, float OldStartTime, float OldEndTime)
 {
-	FloatCurve.ResizeTimeRange(0, NewLength);
+	FloatCurve.ReadjustTimeRange(0, NewLength, bInsert, OldStartTime, OldEndTime);
 }
 ////////////////////////////////////////////////////
 //  FVectorCurve
@@ -92,11 +92,11 @@ void FVectorCurve::UpdateOrAddKey(const FVector& NewKey, float CurrentTime)
 	FloatCurves[Z].UpdateOrAddKey(CurrentTime, NewKey.Z);
 }
 
-void FVectorCurve::Resize(float NewLength)
+void FVectorCurve::Resize(float NewLength, bool bInsert/* whether insert or remove*/, float OldStartTime, float OldEndTime)
 {
-	FloatCurves[X].ResizeTimeRange(0, NewLength);
-	FloatCurves[Y].ResizeTimeRange(0, NewLength);
-	FloatCurves[Z].ResizeTimeRange(0, NewLength);
+	FloatCurves[X].ReadjustTimeRange(0, NewLength, bInsert, OldStartTime, OldEndTime);
+	FloatCurves[Y].ReadjustTimeRange(0, NewLength, bInsert, OldStartTime, OldEndTime);
+	FloatCurves[Z].ReadjustTimeRange(0, NewLength, bInsert, OldStartTime, OldEndTime);
 }
 ////////////////////////////////////////////////////
 //  FTransformCurve
@@ -145,11 +145,11 @@ void FTransformCurve::UpdateOrAddKey(const FTransform& NewKey, float CurrentTime
 	ScaleCurve.UpdateOrAddKey(NewKey.GetScale3D(), CurrentTime);
 }
 
-void FTransformCurve::Resize(float NewLength)
+void FTransformCurve::Resize(float NewLength, bool bInsert/* whether insert or remove*/, float OldStartTime, float OldEndTime)
 {
-	TranslationCurve.Resize(NewLength);
-	RotationCurve.Resize(NewLength);
-	ScaleCurve.Resize(NewLength);
+	TranslationCurve.Resize(NewLength, bInsert, OldStartTime, OldEndTime);
+	RotationCurve.Resize(NewLength, bInsert, OldStartTime, OldEndTime);
+	ScaleCurve.Resize(NewLength, bInsert, OldStartTime, OldEndTime);
 }
 
 /////////////////////////////////////////////////////
@@ -247,22 +247,22 @@ bool FRawCurveTracks::AddCurveData(USkeleton::AnimCurveUID Uid, int32 CurveFlags
 	}
 }
 
-void FRawCurveTracks::Resize(float TotalLength)
+void FRawCurveTracks::Resize(float TotalLength, bool bInsert/* whether insert or remove*/, float OldStartTime, float OldEndTime)
 {
 	for (auto& Curve: FloatCurves)
 	{
-		Curve.Resize(TotalLength);
+		Curve.Resize(TotalLength, bInsert, OldStartTime, OldEndTime);
 	}
 
 #if WITH_EDITORONLY_DATA
 	for(auto& Curve: VectorCurves)
 	{
-		Curve.Resize(TotalLength);
+		Curve.Resize(TotalLength, bInsert, OldStartTime, OldEndTime);
 	}
 
 	for(auto& Curve: TransformCurves)
 	{
-		Curve.Resize(TotalLength);
+		Curve.Resize(TotalLength, bInsert, OldStartTime, OldEndTime);
 	}
 #endif
 }
diff --git a/Engine/Source/Runtime/Engine/Private/Animation/AnimSequence.cpp b/Engine/Source/Runtime/Engine/Private/Animation/AnimSequence.cpp
--- a/Engine/Source/Runtime/Engine/Private/Animation/AnimSequence.cpp
+++ b/Engine/Source/Runtime/Engine/Private/Animation/AnimSequence.cpp
@@ -1274,14 +1274,106 @@ static int32 CropRawTrack(FRawAnimSequenceTrack& RawTrack, int32 StartKey, int32
 	return FMath::Max<int32>( RawTrack.PosKeys.Num(), FMath::Max<int32>(RawTrack.RotKeys.Num(), RawTrack.ScaleKeys.Num()) );
 }
 
-void UAnimSequence::ResizeSequence(float NewLength, int32 NewNumFrames)
+void UAnimSequence::ResizeSequence(float NewLength, int32 NewNumFrames, bool bInsert, int32 StartFrame/*inclusive */, int32 EndFrame/*inclusive*/)
 {
+	check (NewNumFrames > 0);
+	check (StartFrame < EndFrame);
+
+	int32 OldNumFrames = NumFrames;
+	float OldSequenceLength = SequenceLength;
+
+	// verify condition
 	NumFrames = NewNumFrames;
 	// Update sequence length to match new number of frames.
 	SequenceLength = NewLength;
 
+	float Interval = OldSequenceLength / OldNumFrames;
+	ensure (Interval == SequenceLength/NumFrames);
+
+	float OldStartTime = StartFrame * Interval;
+	float OldEndTime = EndFrame * Interval;
+	float Duration = OldEndTime - OldStartTime;
+
+	// re-locate notifies
+	for (auto& Notify: Notifies)
+	{
+		float CurrentTime = Notify.GetTime();
+		float NewDuration = 0.f;
+		if (bInsert)
+		{
+			// if state, make sure to adjust end time
+			if(Notify.NotifyStateClass)
+			{
+				float NotifyDuration = Notify.GetDuration();
+				float NotifyEnd = CurrentTime + NotifyDuration;
+				if(NotifyEnd >= OldStartTime)
+				{
+					NewDuration = NotifyDuration + Duration;
+				}
+				else
+				{
+					NewDuration = NotifyDuration;
+				}
+			}
+
+			// when insert, we only care about start time
+			// if it's later than start time
+			if (CurrentTime >= OldStartTime)
+			{
+				CurrentTime += Duration;
+			}
+		}
+		else
+		{
+			// if state, make sure to adjust end time
+			if(Notify.NotifyStateClass)
+			{
+				float NotifyDuration = Notify.GetDuration();
+				float NotifyEnd = CurrentTime + NotifyDuration;
+				NewDuration = NotifyDuration;
+				if(NotifyEnd >= OldStartTime && NotifyEnd <= OldEndTime)
+				{
+					// small number @todo see if there is define for this
+					NewDuration = 0.1f;
+				}
+				else if (NotifyEnd > OldEndTime)
+				{
+					NewDuration = NotifyEnd-Duration-CurrentTime;
+				}
+				else
+				{
+					NewDuration = NotifyDuration;
+				}
+
+				NewDuration = FMath::Max(NewDuration, 0.1f);
+			}
+
+			if (CurrentTime >= OldStartTime && CurrentTime <= OldEndTime)
+			{
+				CurrentTime = OldStartTime;
+			}
+			else if (CurrentTime > OldEndTime)
+			{
+				CurrentTime -= Duration;
+			}
+		}
+
+		float ClampedCurrentTime = FMath::Clamp(CurrentTime, 0.f, SequenceLength);
+		Notify.LinkSequence(this, ClampedCurrentTime);
+		Notify.SetDuration(NewDuration);
+
+		if (ClampedCurrentTime == 0.f)
+		{
+			Notify.TriggerTimeOffset = GetTriggerTimeOffsetForType(EAnimEventTriggerOffsets::OffsetAfter);
+		}
+		else if (ClampedCurrentTime == SequenceLength)
+		{
+			Notify.TriggerTimeOffset = GetTriggerTimeOffsetForType(EAnimEventTriggerOffsets::OffsetBefore);
+		}
+	}
+
 	// resize curves
-	RawCurveData.Resize(NewLength);
+	RawCurveData.Resize(NewLength, bInsert, OldStartTime, OldEndTime);
 }
 
 bool UAnimSequence::InsertFramesToRawAnimData( int32 StartFrame, int32 EndFrame, int32 CopyFrame)
@@ -1327,7 +1419,7 @@ bool UAnimSequence::InsertFramesToRawAnimData( int32 StartFrame, int32 EndFrame,
 		float const FrameTime = SequenceLength / ((float)NumFrames);
 
 		int32 NewNumFrames = NumFrames + NumFramesToInsert;
-		ResizeSequence((float)NewNumFrames * FrameTime, NewNumFrames);
+		ResizeSequence((float)NewNumFrames * FrameTime, NewNumFrames, true, StartFrame, EndFrame);
 
 		UE_LOG(LogAnimation, Log, TEXT("\tSequenceLength: %f, NumFrames: %d"), SequenceLength, NumFrames);
 
@@ -1376,14 +1468,14 @@ bool UAnimSequence::CropRawAnimData( float CurrentTime, bool bFromStart )
 	int32 const NumKeys = bFromStart ? KeyIndex : TotalNumOfFrames - KeyIndex ;
 
 	// Recalculate NumFrames
-	NumFrames = TotalNumOfFrames - NumKeys;
+	int32 NewNumFrames = TotalNumOfFrames - NumKeys;
 
 	UE_LOG(LogAnimation, Log, TEXT("UAnimSequence::CropRawAnimData %s - CurrentTime: %f, bFromStart: %d, TotalNumOfFrames: %d, KeyIndex: %d, StartKey: %d, NumKeys: %d"), *GetName(), CurrentTime, bFromStart, TotalNumOfFrames, KeyIndex, StartKey, NumKeys);
 
 	// Iterate over tracks removing keys from each one.
 	for(int32 i=0; i<RawAnimationData.Num(); i++)
 	{
-		// Update NumFrames below to reflect actual number of keys while we crop the anim data
+		// Update NewNumFrames below to reflect actual number of keys while we crop the anim data
 		CropRawTrack(RawAnimationData[i], StartKey, NumKeys, TotalNumOfFrames);
 	}
 
@@ -1391,66 +1483,12 @@ bool UAnimSequence::CropRawAnimData( float CurrentTime, bool bFromStart )
 	for(int32 i=0; i<RawAnimationData.Num(); i++)
 	{
 		FRawAnimSequenceTrack& RawTrack = RawAnimationData[i];
-		check(RawTrack.PosKeys.Num() == 1 || RawTrack.PosKeys.Num() == NumFrames);
-		check(RawTrack.RotKeys.Num() == 1 || RawTrack.RotKeys.Num() == NumFrames);
-	}
-
-	// Crop curve data
-	{
-		float CroppedStartTime = (bFromStart)? (NumKeys + StartKey)*FrameTime : 0.f;
-		float CroppedEndTime = (bFromStart)? SequenceLength : (StartKey-1)*FrameTime;
-		
-		// make sure it doesn't go negative
-		CroppedStartTime = FMath::Max<float>(CroppedStartTime, 0.f);
-		CroppedEndTime = FMath::Max<float>(CroppedEndTime, 0.f);
-
-		for ( auto CurveIter = RawCurveData.FloatCurves.CreateIterator(); CurveIter; ++CurveIter )
-		{
-			FFloatCurve& Curve = *CurveIter;
-			// fix up curve before deleting
-			// add these two values to keep the previous curve shape
-			// it's possible I'm adding same values in the same location
-			if (bFromStart)
-			{
-				// start from beginning, just add at the start time
-				float EvaluationTime = CroppedStartTime; 
-				float Value = Curve.FloatCurve.Eval(EvaluationTime);
-				Curve.FloatCurve.AddKey(EvaluationTime, Value);
-			}
-			else
-			{
-				float EvaluationTime = CroppedEndTime; 
-				float Value = Curve.FloatCurve.Eval(EvaluationTime);
-				Curve.FloatCurve.AddKey(EvaluationTime, Value);
-			}
-
-			// now delete anything outside of range
-			TArray<FKeyHandle> KeysToDelete;
-			for (auto It(Curve.FloatCurve.GetKeyHandleIterator()); It; ++It)
-			{
-				FKeyHandle KeyHandle = It.Key();
-				float KeyTime = Curve.FloatCurve.GetKeyTime(KeyHandle);
-				// if outside of range, just delete it. 
-				if ( KeyTime<CroppedStartTime || KeyTime>CroppedEndTime)
-				{
-					KeysToDelete.Add(KeyHandle);
-				}
-				// if starting from beginning, 
-				// fix this up to apply delta
-				else if (bFromStart)
-				{
-					Curve.FloatCurve.SetKeyTime(KeyHandle, KeyTime - CroppedStartTime);
-				}
-			}
-			for (int32 i = 0; i < KeysToDelete.Num(); ++i)
-			{
-				Curve.FloatCurve.DeleteKey(KeysToDelete[i]);
-			}
-		}
+		check(RawTrack.PosKeys.Num() == 1 || RawTrack.PosKeys.Num() == NewNumFrames);
+		check(RawTrack.RotKeys.Num() == 1 || RawTrack.RotKeys.Num() == NewNumFrames);
 	}
 
 	// Update sequence length to match new number of frames.
-	ResizeSequence((float)NumFrames * FrameTime, NumFrames);
+	ResizeSequence((float)NewNumFrames * FrameTime, NewNumFrames, false, StartKey, StartKey+NumKeys);
 
 	UE_LOG(LogAnimation, Log, TEXT("\tSequenceLength: %f, NumFrames: %d"), SequenceLength, NumFrames);
 
diff --git a/Engine/Source/Runtime/Engine/Private/CurveBase.cpp b/Engine/Source/Runtime/Engine/Private/CurveBase.cpp
--- a/Engine/Source/Runtime/Engine/Private/CurveBase.cpp
+++ b/Engine/Source/Runtime/Engine/Private/CurveBase.cpp
@@ -683,8 +683,88 @@ void FRichCurve::AutoSetTangents(float Tension)
 	}
 }
 
-void FRichCurve::ResizeTimeRange(float NewMinTimeRange, float NewMaxTimeRange)
+void FRichCurve::ReadjustTimeRange(float NewMinTimeRange, float NewMaxTimeRange, bool bInsert/* whether insert or remove*/, float OldStartTime, float OldEndTime)
 {
+	// first readjust modified time keys
+	float ModifiedDuration = OldEndTime - OldStartTime;
+	if (bInsert)
+	{
+		for(int32 KeyIndex=0; KeyIndex<Keys.Num(); ++KeyIndex)
+		{
+			float& CurrentTime = Keys[KeyIndex].Time;
+			if (CurrentTime >= OldStartTime)
+			{
+				CurrentTime += ModifiedDuration;
+			}
+		}
+	}
+	else
+	{
+		// since we only allow one key at a given time, we will just cache the value that needs to be saved
+		// this is the key to be replaced when this section is gone
+		bool bAddNewKey = false; 
+		float NewValue = 0.f;
+		TArray<int32> KeysToDelete;
+
+		for(int32 KeyIndex=0; KeyIndex<Keys.Num(); ++KeyIndex)
+		{
+			float& CurrentTime = Keys[KeyIndex].Time;
+			// if this key exists between range of deleted
+			// we'll evaluate the value at the "OldStartTime"
+			// and re-add key, so that it keeps the previous value at the
+			// start time
+			// But that means if there are multiple keys, 
+			// since we don't want multiple values in the same time
+			// the last one will override the value
+			if( CurrentTime >= OldStartTime && CurrentTime <= OldEndTime)
+			{
+				// get new value and add new key on one of OldStartTime, OldEndTime;
+				// this is a bit complicated problem since we don't know if OldStartTime or OldEndTime is preferred. 
+				// generall we use OldEndTime unless OldStartTime == 0.f
+				// which means it's cut in the beginning. Otherwise it will always use the end time. 
+				bAddNewKey = true;
+				if (OldStartTime != 0.f)
+				{
+					NewValue = Eval(OldStartTime);
+				}
+				else
+				{
+					NewValue = Eval(OldEndTime);
+				}
+				// remove this key, but later because it might change eval result
+				KeysToDelete.Add(KeyIndex);
+			}
+			else if (CurrentTime > OldEndTime)
+			{
+				CurrentTime -= ModifiedDuration;
+			}
+		}
+
+		if (bAddNewKey)
+		{
+			for (auto KeyIndex : KeysToDelete)
+			{
+				const FKeyHandle* KeyHandle = KeyHandlesToIndices.FindKey(KeyIndex);
+				if(KeyHandle)
+				{
+					DeleteKey(*KeyHandle);
+				}
+			}
+
+			UpdateOrAddKey(OldStartTime, NewValue);
+		}
+	}
+
+	// now remove all redundant key
+	TArray<FRichCurveKey> NewKeys;
+	Exchange(NewKeys, Keys);
+
+	for(int32 KeyIndex=0; KeyIndex<NewKeys.Num(); ++KeyIndex)
+	{
+		UpdateOrAddKey(NewKeys[KeyIndex].Time, NewKeys[KeyIndex].Value);
+	}
+
+	// now cull out all out of range 
 	float MinTime, MaxTime;
 	GetTimeRange(MinTime, MaxTime);
 
@@ -694,7 +774,7 @@ void FRichCurve::ResizeTimeRange(float NewMinTimeRange, float NewMaxTimeRange)
 	if (MinTime < NewMinTimeRange)
 	{
 		float NewValue = Eval(NewMinTimeRange);
-		AddKey(NewMinTimeRange, NewValue);
+		UpdateOrAddKey(NewMinTimeRange, NewValue);
 
 		bNeedToDeleteKey = true;
 	}
@@ -703,7 +783,7 @@ void FRichCurve::ResizeTimeRange(float NewMinTimeRange, float NewMaxTimeRange)
 	if(MaxTime > NewMaxTimeRange)
 	{
 		float NewValue = Eval(NewMaxTimeRange);
-		AddKey(NewMaxTimeRange, NewValue);
+		UpdateOrAddKey(NewMaxTimeRange, NewValue);
 
 		bNeedToDeleteKey = true;
 	}
diff --git a/Engine/Source/Runtime/Engine/Public/Animation/AnimCurveTypes.h b/Engine/Source/Runtime/Engine/Public/Animation/AnimCurveTypes.h
--- a/Engine/Source/Runtime/Engine/Public/Animation/AnimCurveTypes.h
+++ b/Engine/Source/Runtime/Engine/Public/Animation/AnimCurveTypes.h
@@ -131,7 +131,7 @@ struct FFloatCurve : public FAnimCurveBase
 	void CopyCurve(FFloatCurve& SourceCurve);
 	float Evaluate(float CurrentTime) const;
 	void UpdateOrAddKey(float NewKey, float CurrentTime);
-	void Resize(float NewLength);
+	void Resize(float NewLength, bool bInsert/* whether insert or remove*/, float OldStartTime, float OldEndTime);
 };
 
 USTRUCT()
@@ -172,7 +172,7 @@ struct FVectorCurve : public FAnimCurveBase
 	FVector Evaluate(float CurrentTime, float BlendWeight) const;
 	void UpdateOrAddKey(const FVector& NewKey, float CurrentTime);
 	bool DoesContainKey() const { return (FloatCurves[0].GetNumKeys() > 0 || FloatCurves[1].GetNumKeys() > 0 || FloatCurves[2].GetNumKeys() > 0);}
-	void Resize(float NewLength);
+	void Resize(float NewLength, bool bInsert/* whether insert or remove*/, float OldStartTime, float OldEndTime);
 };
 
 USTRUCT()
@@ -214,7 +214,7 @@ struct FTransformCurve: public FAnimCurveBase
 	void CopyCurve(FTransformCurve& SourceCurve);
 	FTransform Evaluate(float CurrentTime, float BlendWeight) const;
 	ENGINE_API void UpdateOrAddKey(const FTransform& NewKey, float CurrentTime);
-	void Resize(float NewLength);
+	void Resize(float NewLength, bool bInsert/* whether insert or remove*/, float OldStartTime, float OldEndTime);
 };
 
 /**
@@ -405,7 +405,7 @@ struct FRawCurveTracks
 	/*
 	 * resize curve length. If longer, it doesn't do any. If shorter, remove previous keys and add new key to the end of the frame. 
 	 */
-	void Resize(float TotalLength);
+	void Resize(float TotalLength, bool bInsert/* whether insert or remove*/, float OldStartTime, float OldEndTime);
 	/** 
 	 * Clear all keys
 	 */
