Fix crash due to incorrect current time when switching a sequence player nodes animation to one that is shorter than the current one

diff --git a/Engine/Source/Runtime/Engine/Private/Animation/AnimNode_SequencePlayer.cpp b/Engine/Source/Runtime/Engine/Private/Animation/AnimNode_SequencePlayer.cpp
--- a/Engine/Source/Runtime/Engine/Private/Animation/AnimNode_SequencePlayer.cpp
+++ b/Engine/Source/Runtime/Engine/Private/Animation/AnimNode_SequencePlayer.cpp
@@ -36,6 +36,10 @@ void FAnimNode_SequencePlayer::UpdateAssetPlayer(const FAnimationUpdateContext&
 		FAnimGroupInstance* SyncGroup;
 		FAnimTickRecord& TickRecord = Context.AnimInstance->CreateUninitializedTickRecord(GroupIndex, /*out*/ SyncGroup);
 
+		if (InternalTimeAccumulator > Sequence->SequenceLength)
+		{
+			InternalTimeAccumulator = 0.f;
+		}
 		Context.AnimInstance->MakeSequenceTickRecord(TickRecord, Sequence, bLoopAnimation, PlayRate, FinalBlendWeight, /*inout*/ InternalTimeAccumulator);
 
 		// Update the sync group if it exists
