Fix non-editor builds

diff --git a/Engine/Source/Runtime/Engine/Private/EdGraph/EdGraphSchema.cpp b/Engine/Source/Runtime/Engine/Private/EdGraph/EdGraphSchema.cpp
--- a/Engine/Source/Runtime/Engine/Private/EdGraph/EdGraphSchema.cpp
+++ b/Engine/Source/Runtime/Engine/Private/EdGraph/EdGraphSchema.cpp
@@ -86,6 +86,7 @@ void FGraphActionListBuilderBase::ActionGroup::PerformAction( class UEdGraph* Pa
 
 void FGraphActionListBuilderBase::ActionGroup::InitCategoryChain()
 {
+#if WITH_EDITOR
 	static FString const CategoryDelim("|");
 	FEditorCategoryUtils::GetCategoryDisplayString(RootCategory).ParseIntoArray(CategoryChain, *CategoryDelim, true);
 
@@ -103,6 +104,7 @@ void FGraphActionListBuilderBase::ActionGroup::InitCategoryChain()
 	{
 		Category.Trim();
 	}
+#endif
 }
 
 /////////////////////////////////////////////////////
