Fix skeletal mesh overlap not working randomly. The math for the shape position was inverted :(

diff --git a/Engine/Source/Runtime/Engine/Private/PhysicsEngine/BodyInstance.cpp b/Engine/Source/Runtime/Engine/Private/PhysicsEngine/BodyInstance.cpp
--- a/Engine/Source/Runtime/Engine/Private/PhysicsEngine/BodyInstance.cpp
+++ b/Engine/Source/Runtime/Engine/Private/PhysicsEngine/BodyInstance.cpp
@@ -4123,8 +4123,8 @@ bool FBodyInstance::OverlapMulti(TArray<struct FOverlapResult>& InOutOverlaps, c
 	if (pWorldToComponent)
 	{
 		const FTransform RootTM = WeldParent ? WeldParent->GetUnrealWorldTransform() : GetUnrealWorldTransform();
-		const FTransform ShapeSpaceToComponentSpace = RootTM * (*pWorldToComponent);
-		BodyInstanceSpaceToTestSpace = ShapeSpaceToComponentSpace * ComponentSpaceToTestSpace;
+		const FTransform LocalOffset = (*pWorldToComponent) * RootTM;
+		BodyInstanceSpaceToTestSpace = ComponentSpaceToTestSpace * LocalOffset;
 	}
 	else
 	{
