UE-9920: AGameSession::ApproveLogin shouldn't localize

diff --git a/Engine/Source/Runtime/Engine/Private/GameSession.cpp b/Engine/Source/Runtime/Engine/Private/GameSession.cpp
--- a/Engine/Source/Runtime/Engine/Private/GameSession.cpp
+++ b/Engine/Source/Runtime/Engine/Private/GameSession.cpp
@@ -187,7 +187,7 @@ FString AGameSession::ApproveLogin(const FString& Options)
 
 	if (AtCapacity(SpectatorOnly == 1))
 	{
-		return NSLOCTEXT("NetworkErrors", "ServerAtCapacity", "Server full.").ToString();
+		return TEXT( "Server full." );
 	}
 
 	int32 SplitscreenCount = 0;
@@ -195,7 +195,8 @@ FString AGameSession::ApproveLogin(const FString& Options)
 
 	if (SplitscreenCount > MaxSplitscreensPerConnection)
 	{
-		return FText::Format(NSLOCTEXT("NetworkErrors", "ServerAtCapacitySS", "A maximum of '{0}' splitscreen players are allowed"), FText::AsNumber(MaxSplitscreensPerConnection)).ToString();
+		UE_LOG(LogGameSession, Warning, TEXT("ApproveLogin: A maximum of %i splitscreen players are allowed"), MaxSplitscreensPerConnection);
+		return TEXT("Maximum splitscreen players");
 	}
 
 	return TEXT("");
