Fixed black-screen for ES 3.1 when not using MobileHDR (e.g. Tappy Chicken)

diff --git a/Engine/Source/Runtime/Engine/Private/SceneView.cpp b/Engine/Source/Runtime/Engine/Private/SceneView.cpp
--- a/Engine/Source/Runtime/Engine/Private/SceneView.cpp
+++ b/Engine/Source/Runtime/Engine/Private/SceneView.cpp
@@ -452,7 +452,8 @@ FSceneView::FSceneView(const FSceneViewInitOptions& InitOptions)
 
 	// OpenGL Gamma space output in GLSL flips Y when rendering directly to the back buffer (so not needed on PC, as we never render directly into the back buffer)
 	auto ShaderPlatform = GShaderPlatformForFeatureLevel[FeatureLevel];
-	static bool bPlatformRequiresReverseCulling = (IsOpenGLPlatform(ShaderPlatform) && !IsPCPlatform(ShaderPlatform));
+	bool bUsingForwardRenderer = (Family && Family->Scene) ? !Family->Scene->ShouldUseDeferredRenderer() : false;
+	bool bPlatformRequiresReverseCulling = (IsOpenGLPlatform(ShaderPlatform) && bUsingForwardRenderer && !IsPCPlatform(ShaderPlatform));
 	static auto* MobileHDRCvar = IConsoleManager::Get().FindTConsoleVariableDataInt(TEXT("r.MobileHDR"));
 	check(MobileHDRCvar);
 	bReverseCulling = (bPlatformRequiresReverseCulling && MobileHDRCvar->GetValueOnAnyThread() == 0) ? !bReverseCulling : bReverseCulling;
