Don't create bodies for InstancedStaticMeshComponent instances with 0 scale.
Add Mobility setting to FoliageType

diff --git a/Engine/Source/Runtime/Engine/Private/InstancedStaticMesh.cpp b/Engine/Source/Runtime/Engine/Private/InstancedStaticMesh.cpp
--- a/Engine/Source/Runtime/Engine/Private/InstancedStaticMesh.cpp
+++ b/Engine/Source/Runtime/Engine/Private/InstancedStaticMesh.cpp
@@ -990,30 +990,38 @@ void UInstancedStaticMeshComponent::CreateAllInstanceBodies()
     
 	    for (int32 i = 0; i < NumBodies; ++i)
 	    {
-		    InstanceBodies[i] = new FBodyInstance;
-		    FBodyInstance* Instance = InstanceBodies[i];
-			const FTransform InstanceTM = FTransform(PerInstanceSMData[i].Transform) * ComponentToWorld;   
-		    Instance->CopyBodyInstancePropertiesFrom(&BodyInstance);
-		    Instance->InstanceBodyIndex = i; // Set body index 
-		    Instance->bAutoWeld = false;
-    
-		    // make sure we never enable bSimulatePhysics for ISMComps
-		    Instance->bSimulatePhysics = false;
-
-			if(Mobility == EComponentMobility::Movable)
+			const FTransform InstanceTM = FTransform(PerInstanceSMData[i].Transform) * ComponentToWorld;
+			if (InstanceTM.GetScale3D().IsNearlyZero())
 			{
-				Instance->InitBody(BodySetup, InstanceTM, this, PhysScene);
-			}else
+				InstanceBodies[i] = nullptr;
+			}
+			else
 			{
-				Transforms.Add(InstanceTM);
-#if WITH_PHYSX
-				Instance->RigidActorSyncId = i + 1;
+				InstanceBodies[i] = new FBodyInstance;
+				FBodyInstance* Instance = InstanceBodies[i];
+				Instance->CopyBodyInstancePropertiesFrom(&BodyInstance);
+				Instance->InstanceBodyIndex = i; // Set body index 
+				Instance->bAutoWeld = false;
 
-				if (GetWorld()->GetPhysicsScene()->HasAsyncScene())
+				// make sure we never enable bSimulatePhysics for ISMComps
+				Instance->bSimulatePhysics = false;
+
+				if (Mobility == EComponentMobility::Movable)
 				{
-					Instance->RigidActorAsyncId = Instance->RigidActorSyncId + NumBodies;
+					Instance->InitBody(BodySetup, InstanceTM, this, PhysScene);
 				}
+				else
+				{
+					Transforms.Add(InstanceTM);
+#if WITH_PHYSX
+					Instance->RigidActorSyncId = i + 1;
+
+					if (GetWorld()->GetPhysicsScene()->HasAsyncScene())
+					{
+						Instance->RigidActorAsyncId = Instance->RigidActorSyncId + NumBodies;
+					}
 #endif
+				}
 			}
 	    }
 
@@ -1041,9 +1049,11 @@ void UInstancedStaticMeshComponent::ClearAllInstanceBodies()
 	QUICK_SCOPE_CYCLE_COUNTER(STAT_UInstancedStaticMeshComponent_ClearAllInstanceBodies);
 	for (int32 i = 0; i < InstanceBodies.Num(); i++)
 	{
-		check(InstanceBodies[i]);
-		InstanceBodies[i]->TermBody();
-		delete InstanceBodies[i];
+		if (InstanceBodies[i])
+		{
+			InstanceBodies[i]->TermBody();
+			delete InstanceBodies[i];
+		}
 	}
 
 	InstanceBodies.Empty();
@@ -1056,7 +1066,10 @@ void UInstancedStaticMeshComponent::CreatePhysicsState()
 
 	FPhysScene* PhysScene = GetWorld()->GetPhysicsScene();
 
-	if (!PhysScene) { return; }
+	if (!PhysScene)
+	{
+		return;
+	}
 
 #if WITH_PHYSX
 	check(Aggregates.Num() == 0);
@@ -1463,11 +1476,34 @@ bool UInstancedStaticMeshComponent::UpdateInstanceTransform(int32 InstanceIndex,
 	{
 		// Physics uses world transform of the instance
 		FTransform WorldTransform = bWorldSpace ? NewInstanceTransform : (LocalTransform * ComponentToWorld);
-		FBodyInstance* InstanceBodyInstance = InstanceBodies[InstanceIndex];
+		FBodyInstance*& InstanceBodyInstance = InstanceBodies[InstanceIndex];
 #if WITH_PHYSX
-		// Update transform.
-		InstanceBodyInstance->SetBodyTransform(WorldTransform, ETeleportType::None);
-		InstanceBodyInstance->UpdateBodyScale(WorldTransform.GetScale3D());
+		if (NewInstanceTransform.GetScale3D().IsNearlyZero())
+		{
+			if (InstanceBodyInstance)
+			{
+				// delete BodyInstance
+				InstanceBodyInstance->TermBody();
+				delete InstanceBodyInstance;
+				InstanceBodyInstance = nullptr;
+			}
+		}
+		else
+		{
+			if (InstanceBodyInstance)
+			{
+				// Update existing BodyInstance
+				InstanceBodyInstance->SetBodyTransform(WorldTransform, ETeleportType::None);
+				InstanceBodyInstance->UpdateBodyScale(WorldTransform.GetScale3D());
+			}
+			else
+			{
+				// create new BodyInstance
+				InstanceBodyInstance = new FBodyInstance();
+				InitInstanceBody(InstanceIndex, InstanceBodyInstance);
+			}
+
+		}
 #endif //WITH_PHYSX
 	}
 
@@ -1574,10 +1610,17 @@ void UInstancedStaticMeshComponent::SetupNewInstanceData(FInstancedStaticMeshIns
 			}
 		}
 
-		FBodyInstance* NewBodyInstance = new FBodyInstance();
-		int32 BodyIndex = InstanceBodies.Insert(NewBodyInstance, InInstanceIndex);
-		check(InInstanceIndex == BodyIndex);
-		InitInstanceBody(BodyIndex, NewBodyInstance);
+		if (InInstanceTransform.GetScale3D().IsNearlyZero())
+		{
+			InstanceBodies.Insert(nullptr, InInstanceIndex);
+		}
+		else
+		{
+			FBodyInstance* NewBodyInstance = new FBodyInstance();
+			int32 BodyIndex = InstanceBodies.Insert(NewBodyInstance, InInstanceIndex);
+			check(InInstanceIndex == BodyIndex);
+			InitInstanceBody(BodyIndex, NewBodyInstance);
+		}
 	}
 }
 
diff --git a/Engine/Source/Runtime/Engine/Private/PhysicsEngine/BodyInstance.cpp b/Engine/Source/Runtime/Engine/Private/PhysicsEngine/BodyInstance.cpp
--- a/Engine/Source/Runtime/Engine/Private/PhysicsEngine/BodyInstance.cpp
+++ b/Engine/Source/Runtime/Engine/Private/PhysicsEngine/BodyInstance.cpp
@@ -1157,6 +1157,10 @@ struct FInitBodiesHelper
 		for (int32 BodyIdx = 0; BodyIdx < NumBodies; ++BodyIdx)
 		{
 			FBodyInstance* Instance = Bodies[BodyIdx];
+			if (!Instance)
+			{
+				continue;
+			}
 			const FTransform& Transform = Transforms[BodyIdx];
 
 			FBodyInstance::ValidateTransform(Transform, DebugName, BodySetup);
@@ -1355,7 +1359,10 @@ struct FInitBodiesHelper
 			for (int32 BodyIdx = 0, NumBodies = Bodies.Num(); BodyIdx < NumBodies; ++BodyIdx)
 			{
 				FBodyInstance* Instance = Bodies[BodyIdx];
-				Instance->InitDynamicProperties_AssumesLocked();
+				if (Instance)
+				{
+					Instance->InitDynamicProperties_AssumesLocked();
+				}
 			}
 		}
 	}
@@ -1404,6 +1411,11 @@ struct FInitBodiesHelper
 				for (int32 BodyIdx = 0; BodyIdx < NumBodies; ++BodyIdx)
 				{
 					FBodyInstance* Instance = Bodies[BodyIdx];
+					if (!Instance)
+					{
+						continue;
+					}
+
 					const FTransform& Transform = Transforms[BodyIdx];
 
 					const b2Vec2 Scale2D = FPhysicsIntegration2D::ConvertUnrealVectorToBox(Instance->Scale3D);
@@ -4048,10 +4060,13 @@ bool FBodyInstance::OverlapTestForBodiesImpl(const FVector& Pos, const FQuat& Ro
 			{
 				for (const FBodyInstance* BodyInstance : Bodies)
 				{
-					bHaveOverlap = BodyInstance->OverlapPhysX_AssumesLocked(*PGeom, PShapeGlobalPose);
-					if (bHaveOverlap)
+					if (BodyInstance)
 					{
-						return;
+						bHaveOverlap = BodyInstance->OverlapPhysX_AssumesLocked(*PGeom, PShapeGlobalPose);
+						if (bHaveOverlap)
+						{
+							return;
+						}
 					}
 				}
 			}
diff --git a/Engine/Source/Runtime/Foliage/Private/InstancedFoliage.cpp b/Engine/Source/Runtime/Foliage/Private/InstancedFoliage.cpp
--- a/Engine/Source/Runtime/Foliage/Private/InstancedFoliage.cpp
+++ b/Engine/Source/Runtime/Foliage/Private/InstancedFoliage.cpp
@@ -50,6 +50,8 @@ struct FFoliageCustomVersion
 		FoliageTypeProceduralScaleAndShade = 8,
 		// Added FoliageHISMC and blueprint support
 		FoliageHISMCBlueprints = 9,
+		// Added Mobility setting to UFoliageType
+		AddedMobility = 10,
 		// -----<new versions can be added above this line>-------------------------------------------------
 		VersionPlusOne,
 		LatestVersion = VersionPlusOne - 1
@@ -290,7 +292,7 @@ UFoliageType::UFoliageType(const FObjectInitializer& ObjectInitializer)
 	VertexColorMask = FOLIAGEVERTEXCOLORMASK_Disabled;
 	VertexColorMaskThreshold = 0.5f;
 
-	bEnableStaticLighting = true;
+	Mobility = EComponentMobility::Static;
 	CastShadow = true;
 	bCastDynamicShadow = true;
 	bCastStaticShadow = true;
@@ -361,6 +363,11 @@ void UFoliageType::Serialize(FArchive& Ar)
 		LandscapeLayer_DEPRECATED = NAME_None;
 	}
 
+	if (Ar.IsLoading() && GetLinkerCustomVersion(FFoliageCustomVersion::GUID) < FFoliageCustomVersion::AddedMobility)
+	{
+		Mobility = bEnableStaticLighting_DEPRECATED ? EComponentMobility::Static : EComponentMobility::Movable;
+	}
+
 #if WITH_EDITORONLY_DATA
 	if (Ar.IsLoading())
 	{
@@ -703,10 +710,9 @@ void FFoliageMeshInfo::UpdateComponentSettings(const UFoliageType* InSettings)
 			bNeedsMarkRenderStateDirty = true;
 		}
 
-		EComponentMobility::Type NewMobility = FoliageType->bEnableStaticLighting ? EComponentMobility::Static : EComponentMobility::Movable;
-		if (Component->Mobility != NewMobility)
+		if (Component->Mobility != FoliageType->Mobility)
 		{
-			Component->Mobility = NewMobility;
+			Component->SetMobility(FoliageType->Mobility);
 			bNeedsMarkRenderStateDirty = true;
 			bNeedsInvalidateLightingCache = true;
 		}
diff --git a/Engine/Source/Runtime/Foliage/Public/FoliageType.h b/Engine/Source/Runtime/Foliage/Public/FoliageType.h
--- a/Engine/Source/Runtime/Foliage/Public/FoliageType.h
+++ b/Engine/Source/Runtime/Foliage/Public/FoliageType.h
@@ -176,6 +176,10 @@ public:
 public:
 	// INSTANCE SETTINGS
 
+	/** Mobility property to apply to foliage components */
+	UPROPERTY(Category = InstanceSettings, EditAnywhere, BlueprintReadOnly)
+	TEnumAsByte<EComponentMobility::Type> Mobility;
+
 	/**
 	 * The distance where instances will begin to fade out if using a PerInstanceFadeAmount material node. 0 disables.
 	 * When the entire cluster is beyond this distance, the cluster is completely culled and not rendered at all.
@@ -183,9 +187,9 @@ public:
 	UPROPERTY(EditAnywhere, Category=InstanceSettings, meta=(UIMin=0))
 	FInt32Interval CullDistance;
 
-	/** Controls whether the foliage should take part in static lighting/shadowing. If false, mesh will not receive or cast static lighting or shadows, regardless of other settings. */
-	UPROPERTY(EditAnywhere, BlueprintReadOnly, Category=InstanceSettings)
-	uint32 bEnableStaticLighting : 1;
+	/** Deprecated. Now use the Mobility setting to control static/dynamic lighting */
+	UPROPERTY()
+	uint32 bEnableStaticLighting_DEPRECATED : 1;
 
 	/** Controls whether the foliage should cast a shadow or not. */
 	UPROPERTY(EditAnywhere, BlueprintReadOnly, Category=InstanceSettings)
