Don't compile UWorld::Listen code into client builds

Unshelved from pending changelist '2582794':

PR #1228: Require WITH_SERVER_CODE for UWorld::Listen (Contributed by [REDACTED_PERSON])

diff --git a/Engine/Source/Runtime/Engine/Private/World.cpp b/Engine/Source/Runtime/Engine/Private/World.cpp
--- a/Engine/Source/Runtime/Engine/Private/World.cpp
+++ b/Engine/Source/Runtime/Engine/Private/World.cpp
@@ -4027,6 +4027,7 @@ void UWorld::NotifyControlMessage(UNetConnection* Connection, uint8 MessageType,
 
 bool UWorld::Listen( FURL& InURL )
 {
+#if WITH_SERVER_CODE
 	if( NetDriver )
 	{
 		GEngine->BroadcastNetworkFailure(this, NetDriver, ENetworkFailure::NetDriverAlreadyExists);
@@ -4063,6 +4064,9 @@ bool UWorld::Listen( FURL& InURL )
 
 	NextSwitchCountdown = NetDriver->ServerTravelPause;
 	return true;
+#else
+	return false;
+#endif // WITH_SERVER_CODE
 }
 
 bool UWorld::IsClient()
