Fix crash counting memory

diff --git a/Engine/Source/Runtime/Engine/Private/UnrealEngine.cpp b/Engine/Source/Runtime/Engine/Private/UnrealEngine.cpp
--- a/Engine/Source/Runtime/Engine/Private/UnrealEngine.cpp
+++ b/Engine/Source/Runtime/Engine/Private/UnrealEngine.cpp
@@ -1879,8 +1879,10 @@ void UEngine::Serialize(FArchive& Ar)
 	if (Ar.IsCountingMemory())
 	{
 		// Only use the main audio device when counting memory
-		FAudioDevice* AudioDevice = GetMainAudioDevice();
-		AudioDevice->CountBytes(Ar);
+		if (FAudioDevice* AudioDevice = GetMainAudioDevice())
+		{
+			AudioDevice->CountBytes(Ar);
+		}
 	}
 }
 
