Back out changelist 2616676

diff --git a/Engine/Source/Runtime/Engine/Private/World.cpp b/Engine/Source/Runtime/Engine/Private/World.cpp
--- a/Engine/Source/Runtime/Engine/Private/World.cpp
+++ b/Engine/Source/Runtime/Engine/Private/World.cpp
@@ -4027,7 +4027,6 @@ void UWorld::NotifyControlMessage(UNetConnection* Connection, uint8 MessageType,
 
 bool UWorld::Listen( FURL& InURL )
 {
-#if WITH_SERVER_CODE
 	if( NetDriver )
 	{
 		GEngine->BroadcastNetworkFailure(this, NetDriver, ENetworkFailure::NetDriverAlreadyExists);
@@ -4064,9 +4063,6 @@ bool UWorld::Listen( FURL& InURL )
 
 	NextSwitchCountdown = NetDriver->ServerTravelPause;
 	return true;
-#else
-	return false;
-#endif // WITH_SERVER_CODE
 }
 
 bool UWorld::IsClient()
