Stop anim re-entrancy checks doing work when DO_CHECK is 0

diff --git a/Engine/Source/Runtime/Engine/Private/Animation/AnimNodeBase.cpp b/Engine/Source/Runtime/Engine/Private/Animation/AnimNodeBase.cpp
--- a/Engine/Source/Runtime/Engine/Private/Animation/AnimNodeBase.cpp
+++ b/Engine/Source/Runtime/Engine/Private/Animation/AnimNodeBase.cpp
@@ -26,9 +26,11 @@ void FPoseLinkBase::AttemptRelink(const FAnimationBaseContext& Context)
 
 void FPoseLinkBase::Initialize(const FAnimationInitializeContext& Context)
 {
+#if DO_CHECK
 	checkf( !bProcessed, TEXT( "Initialize already in progress, circular link for AnimInstance [%s] Blueprint [%s]" ), \
 		Context.AnimInstance ? *Context.AnimInstance->GetFullName() : TEXT( "None" ), Context.GetAnimBlueprintClass() ? *Context.GetAnimBlueprintClass()->GetFullName() : TEXT( "None" ) );
 	TGuardValue<bool> CircularGuard(bProcessed, true);
+#endif
 
 	AttemptRelink(Context);
 
@@ -41,9 +43,11 @@ void FPoseLinkBase::Initialize(const FAnimationInitializeContext& Context)
 
 void FPoseLinkBase::CacheBones(const FAnimationCacheBonesContext& Context) 
 {
-	checkf( !bProcessed, TEXT( "Initialize already in progress, circular link for AnimInstance [%s] Blueprint [%s]" ), \
+#if DO_CHECK
+	checkf( !bProcessed, TEXT( "CacheBones already in progress, circular link for AnimInstance [%s] Blueprint [%s]" ), \
 		Context.AnimInstance ? *Context.AnimInstance->GetFullName() : TEXT( "None" ), Context.GetAnimBlueprintClass() ? *Context.GetAnimBlueprintClass()->GetFullName() : TEXT( "None" ) );
 	TGuardValue<bool> CircularGuard(bProcessed, true);
+#endif
 
 	if (LinkedNode != NULL)
 	{
@@ -53,9 +57,11 @@ void FPoseLinkBase::CacheBones(const FAnimationCacheBonesContext& Context)
 
 void FPoseLinkBase::Update(const FAnimationUpdateContext& Context)
 {
+#if DO_CHECK
 	checkf( !bProcessed, TEXT( "Update already in progress, circular link for AnimInstance [%s] Blueprint [%s]" ), \
 		Context.AnimInstance ? *Context.AnimInstance->GetFullName() : TEXT( "None" ), Context.GetAnimBlueprintClass() ? *Context.GetAnimBlueprintClass()->GetFullName() : TEXT( "None" ) );
 	TGuardValue<bool> CircularGuard(bProcessed, true);
+#endif
 
 #if WITH_EDITOR
 	if (GIsEditor)
@@ -100,9 +106,11 @@ void FPoseLinkBase::GatherDebugData(FNodeDebugData& DebugData)
 
 void FPoseLink::Evaluate(FPoseContext& Output)
 {
+#if DO_CHECK
 	checkf( !bProcessed, TEXT( "Evaluate already in progress, circular link for AnimInstance [%s] Blueprint [%s]" ), \
 		Output.AnimInstance ? *Output.AnimInstance->GetFullName() : TEXT( "None" ), Output.GetAnimBlueprintClass() ? *Output.GetAnimBlueprintClass()->GetFullName() : TEXT( "None" ) );
 	TGuardValue<bool> CircularGuard(bProcessed, true);
+#endif
 
 #if WITH_EDITOR
 	if ((LinkedNode == NULL) && GIsEditor)
@@ -132,9 +140,11 @@ void FPoseLink::Evaluate(FPoseContext& Output)
 
 void FComponentSpacePoseLink::EvaluateComponentSpace(FComponentSpacePoseContext& Output)
 {
+#if DO_CHECK
 	checkf( !bProcessed, TEXT( "EvaluateComponentSpace already in progress, circular link for AnimInstance [%s] Blueprint [%s]" ), \
 		Output.AnimInstance ? *Output.AnimInstance->GetFullName() : TEXT( "None" ), Output.GetAnimBlueprintClass() ? *Output.GetAnimBlueprintClass()->GetFullName() : TEXT( "None" ) );
 	TGuardValue<bool> CircularGuard(bProcessed, true);
+#endif
 
 	if (LinkedNode != NULL)
 	{
