Better blueprint name and documentation for ActorComponent Is Being Destroyed

diff --git a/Engine/Source/Runtime/Engine/Classes/Components/ActorComponent.h b/Engine/Source/Runtime/Engine/Classes/Components/ActorComponent.h
--- a/Engine/Source/Runtime/Engine/Classes/Components/ActorComponent.h
+++ b/Engine/Source/Runtime/Engine/Classes/Components/ActorComponent.h
@@ -192,7 +192,10 @@ public:
 	bool HasBeenInitialized() const { return bHasBeenInitialized; }
 	bool HasBegunPlay() const { return bHasBegunPlay; }
 
-	UFUNCTION(BlueprintCallable, Category="Components")
+	/**
+	 * Returns whether the component is in the process of being destroyed.
+	 */
+	UFUNCTION(BlueprintCallable, Category="Components", meta=(DisplayName="Is Component Being Destroyed"))
 	bool IsBeingDestroyed() const
 	{
 		return bIsBeingDestroyed;
