Forgot to submit this file which routes wake/sleep events.

diff --git a/Engine/Source/Runtime/Engine/Private/PhysicsEngine/PhysXSupport.cpp b/Engine/Source/Runtime/Engine/Private/PhysicsEngine/PhysXSupport.cpp
--- a/Engine/Source/Runtime/Engine/Private/PhysicsEngine/PhysXSupport.cpp
+++ b/Engine/Source/Runtime/Engine/Private/PhysicsEngine/PhysXSupport.cpp
@@ -518,14 +518,20 @@ void FPhysXSimEventCallback::onConstraintBreak( PxConstraintInfo* constraints, P
 	}
 }
 
-void FPhysXSimEventCallback::onWake(PxActor** actors, PxU32 count)
+void FPhysXSimEventCallback::onWake(PxActor** Actors, PxU32 Count)
 {
-
+	for(PxU32 ActorIdx = 0; ActorIdx < Count; ++ActorIdx)
+	{
+		OwningScene->AddPendingSleepingEvent(Actors[ActorIdx], SleepEvent::SET_Wakeup, SceneType);
+	}
 }
 
-void FPhysXSimEventCallback::onSleep(PxActor** actors, PxU32 count)
+void FPhysXSimEventCallback::onSleep(PxActor** Actors, PxU32 Count)
 {
-
+	for (PxU32 ActorIdx = 0; ActorIdx < Count; ++ActorIdx)
+	{
+		OwningScene->AddPendingSleepingEvent(Actors[ActorIdx], SleepEvent::SET_Sleep, SceneType);
+	}
 }
 
 //////////////////////////////////////////////////////////////////////////
