Fix addForce crash - now that we have support from APEX we can just call it regardless of if it's in the scene.
Also fix awake flag being ignored.

diff --git a/Engine/Source/Runtime/Engine/Private/Components/DestructibleComponent.cpp b/Engine/Source/Runtime/Engine/Private/Components/DestructibleComponent.cpp
--- a/Engine/Source/Runtime/Engine/Private/Components/DestructibleComponent.cpp
+++ b/Engine/Source/Runtime/Engine/Private/Components/DestructibleComponent.cpp
@@ -401,7 +401,7 @@ void UDestructibleComponent::CreatePhysicsState()
 		// Sleep/wake up as appropriate
 		if(!BodyInstance.bStartAwake)
 		{
-			//ApexDestructibleActor->setChunkPhysXActorAwakeState(0, false);	//TODO: broke during bad integration of apex. Will turn on once new libs are in
+			ApexDestructibleActor->setChunkPhysXActorAwakeState(0, false);
 		}
 	}
 
@@ -462,10 +462,7 @@ void UDestructibleComponent::AddImpulse( FVector Impulse, FName BoneName /*= NAM
 	ExecuteOnPhysicsReadWrite([&]
 	{
 		const int32 ChunkIdx = BoneIdxToChunkIdx(GetBoneIndex(BoneName));
-		if(PxRigidDynamic* Actor = ApexDestructibleActor->getChunkPhysXActor(ChunkIdx))
-		{
-			Actor->addForce(U2PVector(Impulse), bVelChange ? PxForceMode::eVELOCITY_CHANGE : PxForceMode::eIMPULSE);
-		}
+		ApexDestructibleActor->addForce(ChunkIdx, U2PVector(Impulse),  bVelChange ? PxForceMode::eVELOCITY_CHANGE : PxForceMode::eIMPULSE);
 	});
 #endif
 }
@@ -476,10 +473,8 @@ void UDestructibleComponent::AddImpulseAtLocation( FVector Impulse, FVector Posi
 	ExecuteOnPhysicsReadWrite([&]
 	{
 		const int32 ChunkIdx = BoneIdxToChunkIdx(GetBoneIndex(BoneName));
-		if (PxRigidDynamic* Actor = ApexDestructibleActor->getChunkPhysXActor(ChunkIdx))
-		{
-			PxRigidBodyExt::addForceAtPos(*Actor, U2PVector(Impulse), U2PVector(Position), PxForceMode::eIMPULSE);
-		}
+		PxVec3 Location = U2PVector(Position);
+		ApexDestructibleActor->addForce(ChunkIdx, U2PVector(Impulse),  PxForceMode::eIMPULSE);
 	});
 #endif
 }
@@ -490,10 +485,7 @@ void UDestructibleComponent::AddForce( FVector Force, FName BoneName /*= NAME_No
 	ExecuteOnPhysicsReadWrite([&]
 	{
 		const int32 ChunkIdx = BoneIdxToChunkIdx(GetBoneIndex(BoneName));
-		if (PxRigidDynamic* Actor = ApexDestructibleActor->getChunkPhysXActor(ChunkIdx))
-		{
-			Actor->addForce(U2PVector(Force), bAccelChange ? PxForceMode::eACCELERATION : PxForceMode::eFORCE);
-		}
+		ApexDestructibleActor->addForce(ChunkIdx, U2PVector(Force), bAccelChange ? PxForceMode::eACCELERATION : PxForceMode::eFORCE);
 	});
 #endif
 }
@@ -504,10 +496,9 @@ void UDestructibleComponent::AddForceAtLocation( FVector Force, FVector Location
 	ExecuteOnPhysicsReadWrite([&]
 	{
 		int32 ChunkIdx = BoneIdxToChunkIdx(GetBoneIndex(BoneName));
-		if (PxRigidDynamic* Actor = ApexDestructibleActor->getChunkPhysXActor(ChunkIdx))
-		{
-			PxRigidBodyExt::addForceAtPos(*Actor, U2PVector(Force), U2PVector(Location), PxForceMode::eFORCE);
-		}
+		PxVec3 Position = U2PVector(Location);
+
+		ApexDestructibleActor->addForce(ChunkIdx, U2PVector(Force), PxForceMode::eFORCE, &Position);
 	});
 #endif
 }
