Get rid of some GWorlds

diff --git a/Engine/Source/Editor/UnrealEd/Private/PlayLevel.cpp b/Engine/Source/Editor/UnrealEd/Private/PlayLevel.cpp
--- a/Engine/Source/Editor/UnrealEd/Private/PlayLevel.cpp
+++ b/Engine/Source/Editor/UnrealEd/Private/PlayLevel.cpp
@@ -262,7 +262,7 @@ void UEditorEngine::EndPlayMap()
 			UWorld* TheWorld = UWorld::FindWorldInPackage(Object->GetOutermost());
 			if ( TheWorld )
 			{
-				StaticExec(GWorld, *FString::Printf(TEXT("OBJ REFS CLASS=WORLD NAME=%s"), *TheWorld->GetPathName()));
+				StaticExec(nullptr, *FString::Printf(TEXT("OBJ REFS CLASS=WORLD NAME=%s"), *TheWorld->GetPathName()));
 			}
 			else
 			{
diff --git a/Engine/Source/Runtime/Engine/Private/TickTaskManager.cpp b/Engine/Source/Runtime/Engine/Private/TickTaskManager.cpp
--- a/Engine/Source/Runtime/Engine/Private/TickTaskManager.cpp
+++ b/Engine/Source/Runtime/Engine/Private/TickTaskManager.cpp
@@ -1614,10 +1614,10 @@ static void RemoveTestTickFunctions(const TArray<FString>& Args)
 	}
 }
 
-static void AddTestTickFunctions(const TArray<FString>& Args)
+static void AddTestTickFunctions(const TArray<FString>& Args, UWorld* InWorld)
 {
 	RemoveTestTickFunctions(Args);
-	ULevel* Level = GWorld->GetCurrentLevel();
+	ULevel* Level = InWorld->GetCurrentLevel();
 	UE_LOG(LogConsoleResponse, Display, TEXT("Adding 1000 ticks in a cache coherent fashion."));
 
 
@@ -1628,10 +1628,10 @@ static void AddTestTickFunctions(const TArray<FString>& Args)
 	}
 }
 
-static void AddIndirectTestTickFunctions(const TArray<FString>& Args)
+static void AddIndirectTestTickFunctions(const TArray<FString>& Args, UWorld* InWorld)
 {
 	RemoveTestTickFunctions(Args);
-	ULevel* Level = GWorld->GetCurrentLevel();
+	ULevel* Level = InWorld->GetCurrentLevel();
 	UE_LOG(LogConsoleResponse, Display, TEXT("Adding 1000 ticks in a cache coherent fashion."));
 	TArray<FTestTickFunction*> Junk;
 	for (int32 Index = 0; Index < NumTestTickFunctions; Index++)
@@ -1656,16 +1656,16 @@ static FAutoConsoleCommand RemoveTestTickFunctionsCmd(
 	FConsoleCommandWithArgsDelegate::CreateStatic(&RemoveTestTickFunctions)
 	);
 
-static FAutoConsoleCommand AddTestTickFunctionsCmd(
+static FAutoConsoleCommandWithWorldAndArgs AddTestTickFunctionsCmd(
 	TEXT("tick.AddTestTickFunctions"),
 	TEXT("Add no-op ticks to test performance of ticking infrastructure."),
-	FConsoleCommandWithArgsDelegate::CreateStatic(&AddTestTickFunctions)
+	FConsoleCommandWithWorldAndArgsDelegate::CreateStatic(&AddTestTickFunctions)
 	);
 
-static FAutoConsoleCommand AddIndirectTestTickFunctionsCmd(
+static FAutoConsoleCommandWithWorldAndArgs AddIndirectTestTickFunctionsCmd(
 	TEXT("tick.AddIndirectTestTickFunctions"),
 	TEXT("Add no-op ticks to test performance of ticking infrastructure."),
-	FConsoleCommandWithArgsDelegate::CreateStatic(&AddIndirectTestTickFunctions)
+	FConsoleCommandWithWorldAndArgsDelegate::CreateStatic(&AddIndirectTestTickFunctions)
 	);
 
 
