UE-12121 - When material is set to "Use material attributes" the "tangent space normal" option does not function

Works now.

diff --git a/Engine/Source/Runtime/Engine/Private/Materials/MaterialShared.cpp b/Engine/Source/Runtime/Engine/Private/Materials/MaterialShared.cpp
--- a/Engine/Source/Runtime/Engine/Private/Materials/MaterialShared.cpp
+++ b/Engine/Source/Runtime/Engine/Private/Materials/MaterialShared.cpp
@@ -699,7 +699,7 @@ bool IsTranslucentBlendMode(EBlendMode BlendMode)
 }
 
 int32 FMaterialResource::GetMaterialDomain() const { return Material->MaterialDomain; }
-bool FMaterialResource::IsTangentSpaceNormal() const { return Material->bTangentSpaceNormal || Material->Normal.Expression == NULL; }
+bool FMaterialResource::IsTangentSpaceNormal() const { return Material->bTangentSpaceNormal || (Material->Normal.Expression == NULL && !Material->bUseMaterialAttributes); }
 bool FMaterialResource::ShouldInjectEmissiveIntoLPV() const { return Material->bUseEmissiveForDynamicAreaLighting; }
 bool FMaterialResource::ShouldBlockGI() const { return Material->bBlockGI; }
 bool FMaterialResource::ShouldGenerateSphericalParticleNormals() const { return Material->bGenerateSphericalParticleNormals; }
