#ue4 - Improve comments on actor set location/rotation/transform functions, especially regarding bSweep and bTeleport.

diff --git a/Engine/Source/Runtime/Engine/Classes/GameFramework/Actor.h b/Engine/Source/Runtime/Engine/Classes/GameFramework/Actor.h
--- a/Engine/Source/Runtime/Engine/Classes/GameFramework/Actor.h
+++ b/Engine/Source/Runtime/Engine/Classes/GameFramework/Actor.h
@@ -641,11 +641,16 @@ public:
 	FVector K2_GetActorLocation() const;
 
 	/** 
-	 *	Move the Actor to the specified location.
-	 *	@param NewLocation	The new location to move the Actor to.
-	 *	@param bSweep		Should we sweep to the destination location, stopping short of the target if blocked by something. Note:If the root component has no collision this will have no effect.
-	 *  @param SweepHitResult	The hit result from the move if swept.
-	 *	@return	Whether the location was successfully set (if not swept), or whether movement occurred at all (if swept).
+	 * Move the Actor to the specified location.
+	 * @param NewLocation	The new location to move the Actor to.
+	 * @param bSweep		Whether we sweep to the destination location, triggering overlaps along the way and stopping short of the target if blocked by something.
+	 *						Only the root component is swept and checked for blocking collision, child components move without sweeping. If collision is off, this has no effect.
+	 * @param bTeleport		Whether we teleport the physics state (if physics collision is enabled for this object).
+	 *						If true, physics velocity for this object is unchanged (so ragdoll parts are not affected by change in location).
+	 *						If false, physics velocity is updated based on the change in position (affecting ragdoll parts).
+	 *						If CCD is on and not teleporting, this will affect objects along the entire sweep volume.
+	 * @param SweepHitResult	The hit result from the move if swept.
+	 * @return	Whether the location was successfully set (if not swept), or whether movement occurred at all (if swept).
 	 */
 	UFUNCTION(BlueprintCallable, meta=(DisplayName = "SetActorLocation", Keywords="position"), Category="Utilities|Transformation")
 	bool K2_SetActorLocation(FVector NewLocation, bool bSweep, FHitResult& SweepHitResult, bool bTeleport);
@@ -684,8 +689,13 @@ public:
 	/** 
 	 * Move the actor instantly to the specified location. 
 	 * 
-	 * @param NewLocation		The new location to teleport the Actor to.
-	 * @param bSweep			Whether to sweep to the destination location, triggering overlaps along the way and stopping at the first blocking hit.
+	 * @param NewLocation	The new location to teleport the Actor to.
+	 * @param bSweep		Whether we sweep to the destination location, triggering overlaps along the way and stopping short of the target if blocked by something.
+	 *						Only the root component is swept and checked for blocking collision, child components move without sweeping. If collision is off, this has no effect.
+	 * @param bTeleport		Whether we teleport the physics state (if physics collision is enabled for this object).
+	 *						If true, physics velocity for this object is unchanged (so ragdoll parts are not affected by change in location).
+	 *						If false, physics velocity is updated based on the change in position (affecting ragdoll parts).
+	 *						If CCD is on and not teleporting, this will affect objects along the entire sweep volume.
 	 * @param OutSweepHitResult The hit result from the move if swept.
 	 * @return	Whether the location was successfully set if not swept, or whether movement occurred if swept.
 	 */
@@ -704,10 +714,15 @@ public:
 	/** 
 	 * Move the actor instantly to the specified location and rotation.
 	 * 
-	 * @param NewLocation			The new location to teleport the Actor to.
-	 * @param NewRotation			The new rotation for the Actor.
-	 * @param bSweep				Whether to sweep to the destination location, triggering overlaps along the way and stopping at the first blocking hit.
-	 * @param SweepHitResult		The hit result from the move if swept.
+	 * @param NewLocation		The new location to teleport the Actor to.
+	 * @param NewRotation		The new rotation for the Actor.
+	 * @param bSweep			Whether we sweep to the destination location, triggering overlaps along the way and stopping short of the target if blocked by something.
+	 *							Only the root component is swept and checked for blocking collision, child components move without sweeping. If collision is off, this has no effect.
+	 * @param bTeleport			Whether we teleport the physics state (if physics collision is enabled for this object).
+	 *							If true, physics velocity for this object is unchanged (so ragdoll parts are not affected by change in location).
+	 *							If false, physics velocity is updated based on the change in position (affecting ragdoll parts).
+	 *							If CCD is on and not teleporting, this will affect objects along the entire sweep volume.
+	 * @param SweepHitResult	The hit result from the move if swept.
 	 * @return	Whether the rotation was successfully set.
 	 */
 	UFUNCTION(BlueprintCallable, Category="Utilities|Transformation", meta=(DisplayName="SetActorLocationAndRotation"))
@@ -716,9 +731,14 @@ public:
 	/** 
 	 * Move the actor instantly to the specified location and rotation.
 	 * 
-	 * @param NewLocation			The new location to teleport the Actor to.
-	 * @param NewRotation			The new rotation for the Actor.
-	 * @param bSweep				Whether to sweep to the destination location, triggering overlaps along the way and stopping at the first blocking hit.
+	 * @param NewLocation		The new location to teleport the Actor to.
+	 * @param NewRotation		The new rotation for the Actor.
+	 * @param bSweep			Whether we sweep to the destination location, triggering overlaps along the way and stopping short of the target if blocked by something.
+	 *							Only the root component is swept and checked for blocking collision, child components move without sweeping. If collision is off, this has no effect.
+	 * @param bTeleport			Whether we teleport the physics state (if physics collision is enabled for this object).
+	 *							If true, physics velocity for this object is unchanged (so ragdoll parts are not affected by change in location).
+	 *							If false, physics velocity is updated based on the change in position (affecting ragdoll parts).
+	 *							If CCD is on and not teleporting, this will affect objects along the entire sweep volume.
 	 * @param OutSweepHitResult	The hit result from the move if swept.
 	 * @return	Whether the rotation was successfully set.
 	 */
@@ -756,9 +776,14 @@ public:
 	/**
 	 * Adds a delta to the location of this actor in world space.
 	 * 
-	 * @param  DeltaLocation		The change in location.
-	 * @param  bSweep				Whether to sweep to the destination location, triggering overlaps along the way and stopping at the first blocking hit.
-	 * @param  SweepHitResult		The hit result from the move if swept.
+	 * @param DeltaLocation		The change in location.
+	 * @param bSweep			Whether we sweep to the destination location, triggering overlaps along the way and stopping short of the target if blocked by something.
+	 *							Only the root component is swept and checked for blocking collision, child components move without sweeping. If collision is off, this has no effect.
+	 * @param bTeleport			Whether we teleport the physics state (if physics collision is enabled for this object).
+	 *							If true, physics velocity for this object is unchanged (so ragdoll parts are not affected by change in location).
+	 *							If false, physics velocity is updated based on the change in position (affecting ragdoll parts).
+	 *							If CCD is on and not teleporting, this will affect objects along the entire sweep volume.
+	 * @param SweepHitResult	The hit result from the move if swept.
 	 */
 	UFUNCTION(BlueprintCallable, Category="Utilities|Transformation", meta=(DisplayName="AddActorWorldOffset", Keywords="location position"))
 	void K2_AddActorWorldOffset(FVector DeltaLocation, bool bSweep, FHitResult& SweepHitResult, bool bTeleport);
@@ -766,9 +791,14 @@ public:
 	/**
 	 * Adds a delta to the location of this actor in world space.
 	 * 
-	 * @param  DeltaLocation		The change in location.
-	 * @param  bSweep				Whether to sweep to the destination location, triggering overlaps along the way and stopping at the first blocking hit.
-	 * @param  SweepHitResult		The hit result from the move if swept.
+	 * @param DeltaLocation		The change in location.
+	 * @param bSweep			Whether we sweep to the destination location, triggering overlaps along the way and stopping short of the target if blocked by something.
+	 *							Only the root component is swept and checked for blocking collision, child components move without sweeping. If collision is off, this has no effect.
+	 * @param Teleport			Whether we teleport the physics state (if physics collision is enabled for this object).
+	 *							If TeleportPhysics, physics velocity for this object is unchanged (so ragdoll parts are not affected by change in location).
+	 *							If None, physics velocity is updated based on the change in position (affecting ragdoll parts).
+	 *							If CCD is on and not teleporting, this will affect objects along the entire sweep volume.
+	 * @param SweepHitResult	The hit result from the move if swept.
 	 */
 	void AddActorWorldOffset(FVector DeltaLocation, bool bSweep=false, FHitResult* OutSweepHitResult=nullptr, ETeleportType Teleport = ETeleportType::None);
 
@@ -776,9 +806,13 @@ public:
 	/**
 	 * Adds a delta to the rotation of this actor in world space.
 	 * 
-	 * @param  DeltaRotation		The change in rotation.
-	 * @param  bSweep				Whether to sweep to the target rotation (not currently supported).
-	 * @param  SweepHitResult		The hit result from the move if swept.
+	 * @param DeltaRotation		The change in rotation.
+	 * @param bSweep			Whether to sweep to the target rotation (not currently supported for rotation).
+	 * @param bTeleport			Whether we teleport the physics state (if physics collision is enabled for this object).
+	 *							If true, physics velocity for this object is unchanged (so ragdoll parts are not affected by change in location).
+	 *							If false, physics velocity is updated based on the change in position (affecting ragdoll parts).
+	 *							If CCD is on and not teleporting, this will affect objects along the entire sweep volume.
+	 * @param SweepHitResult	The hit result from the move if swept.
 	 */
 	UFUNCTION(BlueprintCallable, Category="Utilities|Transformation", meta=(DisplayName="AddActorWorldRotation", AdvancedDisplay="bSweep,SweepHitResult,bTeleport"))
 	void K2_AddActorWorldRotation(FRotator DeltaRotation, bool bSweep, FHitResult& SweepHitResult, bool bTeleport);
@@ -794,36 +828,74 @@ public:
 
 	/** 
 	 * Set the Actors transform to the specified one.
-	 * @param bSweep		Whether to sweep to the destination location, stopping short of the target if blocked by something.
+	 * @param NewTransform		The new transform.
+	 * @param bSweep			Whether we sweep to the destination location, triggering overlaps along the way and stopping short of the target if blocked by something.
+	 *							Only the root component is swept and checked for blocking collision, child components move without sweeping. If collision is off, this has no effect.
+	 * @param bTeleport			Whether we teleport the physics state (if physics collision is enabled for this object).
+	 *							If true, physics velocity for this object is unchanged (so ragdoll parts are not affected by change in location).
+	 *							If false, physics velocity is updated based on the change in position (affecting ragdoll parts).
+	 *							If CCD is on and not teleporting, this will affect objects along the entire sweep volume.
 	 */
 	UFUNCTION(BlueprintCallable, Category="Utilities|Transformation", meta=(DisplayName="SetActorTransform"))
 	bool K2_SetActorTransform(const FTransform& NewTransform, bool bSweep, FHitResult& SweepHitResult, bool bTeleport);
 	bool SetActorTransform(const FTransform& NewTransform, bool bSweep=false, FHitResult* OutSweepHitResult=nullptr, ETeleportType Teleport = ETeleportType::None);
 
 
-	/** Adds a delta to the location of this component in its local reference frame */
+	/** 
+	 * Adds a delta to the location of this component in its local reference frame.
+	 * @param DelatLocation		The change in location in local space.
+	 * @param bSweep			Whether we sweep to the destination location, triggering overlaps along the way and stopping short of the target if blocked by something.
+	 *							Only the root component is swept and checked for blocking collision, child components move without sweeping. If collision is off, this has no effect.
+	 * @param bTeleport			Whether we teleport the physics state (if physics collision is enabled for this object).
+	 *							If true, physics velocity for this object is unchanged (so ragdoll parts are not affected by change in location).
+	 *							If false, physics velocity is updated based on the change in position (affecting ragdoll parts).
+	 *							If CCD is on and not teleporting, this will affect objects along the entire sweep volume.
+	 */
 	UFUNCTION(BlueprintCallable, Category="Utilities|Transformation", meta=(DisplayName="AddActorLocalOffset", Keywords="location position"))
 	void K2_AddActorLocalOffset(FVector DeltaLocation, bool bSweep, FHitResult& SweepHitResult, bool bTeleport);
 	void AddActorLocalOffset(FVector DeltaLocation, bool bSweep=false, FHitResult* OutSweepHitResult=nullptr, ETeleportType Teleport = ETeleportType::None);
 
 
-	/** Adds a delta to the rotation of this component in its local reference frame */
+	/**
+	 * Adds a delta to the rotation of this component in its local reference frame
+	 * @param DeltaRotation		The change in rotation in local space.
+	 * @param bSweep			Whether we sweep to the destination location, triggering overlaps along the way and stopping short of the target if blocked by something.
+	 *							Only the root component is swept and checked for blocking collision, child components move without sweeping. If collision is off, this has no effect.
+	 * @param bTeleport			Whether we teleport the physics state (if physics collision is enabled for this object).
+	 *							If true, physics velocity for this object is unchanged (so ragdoll parts are not affected by change in location).
+	 *							If false, physics velocity is updated based on the change in position (affecting ragdoll parts).
+	 *							If CCD is on and not teleporting, this will affect objects along the entire sweep volume.
+	 */
 	UFUNCTION(BlueprintCallable, Category="Utilities|Transformation", meta=(DisplayName="AddActorLocalRotation", AdvancedDisplay="bSweep,SweepHitResult,bTeleport"))
 	void K2_AddActorLocalRotation(FRotator DeltaRotation, bool bSweep, FHitResult& SweepHitResult, bool bTeleport);
 	void AddActorLocalRotation(FRotator DeltaRotation, bool bSweep=false, FHitResult* OutSweepHitResult=nullptr, ETeleportType Teleport = ETeleportType::None);
 	void AddActorLocalRotation(const FQuat& DeltaRotation, bool bSweep=false, FHitResult* OutSweepHitResult=nullptr, ETeleportType Teleport = ETeleportType::None);
 
 
-	/** Adds a delta to the transform of this component in its local reference frame */
+	/**
+	 * Adds a delta to the transform of this component in its local reference frame
+	 * @param NewTransform		The change in transform in local space.
+	 * @param bSweep			Whether we sweep to the destination location, triggering overlaps along the way and stopping short of the target if blocked by something.
+	 *							Only the root component is swept and checked for blocking collision, child components move without sweeping. If collision is off, this has no effect.
+	 * @param bTeleport			Whether we teleport the physics state (if physics collision is enabled for this object).
+	 *							If true, physics velocity for this object is unchanged (so ragdoll parts are not affected by change in location).
+	 *							If false, physics velocity is updated based on the change in position (affecting ragdoll parts).
+	 *							If CCD is on and not teleporting, this will affect objects along the entire sweep volume.
+	 */
 	UFUNCTION(BlueprintCallable, Category="Utilities|Transformation", meta=(DisplayName="AddActorLocalTransform"))
 	void K2_AddActorLocalTransform(const FTransform& NewTransform, bool bSweep, FHitResult& SweepHitResult, bool bTeleport);
 	void AddActorLocalTransform(const FTransform& NewTransform, bool bSweep=false, FHitResult* OutSweepHitResult=nullptr, ETeleportType Teleport = ETeleportType::None);
 
 
 	/**
-	 * Set the actor's RootComponent to the specified relative location
-	 * @param NewRelativeLocation	New relative location to set the actor's RootComponent to
-	 * @param bSweep				Should we sweep to the destination location. If true, will stop short of the target if blocked by something
+	 * Set the actor's RootComponent to the specified relative location.
+	 * @param NewRelativeLocation	New relative location of the actor's root component
+	 * @param bSweep				Whether we sweep to the destination location, triggering overlaps along the way and stopping short of the target if blocked by something.
+	 *								Only the root component is swept and checked for blocking collision, child components move without sweeping. If collision is off, this has no effect.
+	 * @param bTeleport				Whether we teleport the physics state (if physics collision is enabled for this object).
+	 *								If true, physics velocity for this object is unchanged (so ragdoll parts are not affected by change in location).
+	 *								If false, physics velocity is updated based on the change in position (affecting ragdoll parts).
+	 *								If CCD is on and not teleporting, this will affect objects along the entire sweep volume.
 	 */
 	UFUNCTION(BlueprintCallable, Category="Utilities|Transformation", meta=(DisplayName="SetActorRelativeLocation"))
 	void K2_SetActorRelativeLocation(FVector NewRelativeLocation, bool bSweep, FHitResult& SweepHitResult, bool bTeleport);
@@ -831,8 +903,13 @@ public:
 
 	/**
 	 * Set the actor's RootComponent to the specified relative rotation
-	 * @param NewRelativeRotation		New relative rotation to set the actor's RootComponent to
-	 * @param bSweep					Should we sweep to the destination rotation. If true, will stop short of the target if blocked by something
+	 * @param NewRelativeRotation	New relative rotation of the actor's root component
+	 * @param bSweep				Whether we sweep to the destination location, triggering overlaps along the way and stopping short of the target if blocked by something.
+	 *								Only the root component is swept and checked for blocking collision, child components move without sweeping. If collision is off, this has no effect.
+	 * @param bTeleport				Whether we teleport the physics state (if physics collision is enabled for this object).
+	 *								If true, physics velocity for this object is unchanged (so ragdoll parts are not affected by change in location).
+	 *								If false, physics velocity is updated based on the change in position (affecting ragdoll parts).
+	 *								If CCD is on and not teleporting, this will affect objects along the entire sweep volume.
 	 */
 	UFUNCTION(BlueprintCallable, Category="Utilities|Transformation", meta=(DisplayName="SetActorRelativeRotation", AdvancedDisplay="bSweep,SweepHitResult,bTeleport"))
 	void K2_SetActorRelativeRotation(FRotator NewRelativeRotation, bool bSweep, FHitResult& SweepHitResult, bool bTeleport);
@@ -841,8 +918,13 @@ public:
 
 	/**
 	 * Set the actor's RootComponent to the specified relative transform
-	 * @param NewRelativeTransform		New relative transform to set the actor's RootComponent to
-	 * @param bSweep					Should we sweep to the destination transform. If true, will stop short of the target if blocked by something
+	 * @param NewRelativeTransform		New relative transform of the actor's root component
+	 * @param bSweep			Whether we sweep to the destination location, triggering overlaps along the way and stopping short of the target if blocked by something.
+	 *							Only the root component is swept and checked for blocking collision, child components move without sweeping. If collision is off, this has no effect.
+	 * @param bTeleport			Whether we teleport the physics state (if physics collision is enabled for this object).
+	 *							If true, physics velocity for this object is unchanged (so ragdoll parts are not affected by change in location).
+	 *							If false, physics velocity is updated based on the change in position (affecting ragdoll parts).
+	 *							If CCD is on and not teleporting, this will affect objects along the entire sweep volume.
 	 */
 	UFUNCTION(BlueprintCallable, Category="Utilities|Transformation", meta=(DisplayName="SetActorRelativeTransform"))
 	void K2_SetActorRelativeTransform(const FTransform& NewRelativeTransform, bool bSweep, FHitResult& SweepHitResult, bool bTeleport);
