Implements DD{X,Y}(UVs) inputs on the TextureSample material node.

diff --git a/Engine/Documentation/Source/Engine/Rendering/Materials/ExpressionReference/Textures/TextureExpressions.INT.udn b/Engine/Documentation/Source/Engine/Rendering/Materials/ExpressionReference/Textures/TextureExpressions.INT.udn
--- a/Engine/Documentation/Source/Engine/Rendering/Materials/ExpressionReference/Textures/TextureExpressions.INT.udn
+++ b/Engine/Documentation/Source/Engine/Rendering/Materials/ExpressionReference/Textures/TextureExpressions.INT.udn
@@ -70,9 +70,11 @@ The TextureSample expression outputs the color value(s) from a texture. This tex
 | [REGION:tablesection]Properties[/REGION] ||
 | **Texture** | Specifies the texture sampled by the expression. To set the texture, first select a texture in the **Content Browser**.  Then, select the Texture property in the expression's property window and click the 'Use Current Selection' button. |
 | **Sampler Type** | The type of data that will be sampled and output from the node. |
-| **MipValueMode** | Applies a noise value to the texture that affects the look and performance. |
+| **MipValueMode** | Selects how to customize the sample's mip-level or derivatives from the default hardware computed. Affects the look and performance. |
 | [REGION:tablesection]Inputs[/REGION] ||
 | **UVs** | Takes in UV texture coordinates to use for the texture. If no values are input to the UVs, the texture coordinates of the mesh the material is applied to are used. If the TextureSample represents a cubemap texture, the UVs coordinate must be a three-channel value rather than just a two-channel value. |
+| **DDX(UVs)** | Takes the DDX of UVs input for anisotropic filtering. By default, the hardware automatically do it, and faster. But there is some cases where you may need to explicitly to set it avoiding the DDX material expresion to avoid artifacts like the 2x2 pixels block artifact of the deferred decal. Only available with **MipValueMode** = **Derivative**. |
+| **DDY(UVs)** | Takes the DDY of UVs input for anisotropic filtering as same as DDX(UVs) takes the DDX of the UVs input. |
 | [REGION:tablesection]Outputs[/REGION] ||
 | **RGB** | Outputs the three-channel RGB vector value of the color. |
 | **R** | Outputs the red channel of the color. |
@@ -83,6 +85,9 @@ The TextureSample expression outputs the color value(s) from a texture. This tex
 
 ![TextureSampleExample.png](TextureSampleExample.png)
 
+For a better understanding of the DDX(UVs) and DDY(UVs) inputs, here is what you could do. But you need to keep in mind that this has the exact same output as with **MipValueMode** = **None**, but is slower. The DDX(UVs) and DDY(UVs) should only be used when you actually needs to overide the default hardware computation because computing those values by a different way than using the DDX/DDY material expression.
+
+![SampleTextureDerivative.png](SampleTextureDerivative.png)
 
 [INCLUDE:Engine/Rendering/Materials/ExpressionReference/Particles#ParticleSubUV]
 
diff --git a/Engine/Shaders/Common.usf b/Engine/Shaders/Common.usf
--- a/Engine/Shaders/Common.usf
+++ b/Engine/Shaders/Common.usf
@@ -238,6 +238,10 @@ MaterialFloat4 Texture2DSampleBias(Texture2D Tex, SamplerState Sampler, float2 U
 	return Tex.SampleBias(Sampler, UV, MipBias);
 #endif
 }
+MaterialFloat4 Texture2DSampleGrad(Texture2D Tex, SamplerState Sampler, float2 UV, MaterialFloat2 DDX, MaterialFloat2 DDY)
+{
+	return Tex.SampleGrad(Sampler, UV, DDX, DDY);
+}
 MaterialFloat4 Texture3DSampleLevel(Texture3D Tex, SamplerState Sampler, float3 UV, MaterialFloat Mip)
 {
 	return Tex.SampleLevel(Sampler, UV, Mip);
@@ -254,6 +258,10 @@ MaterialFloat4 TextureCubeSampleBias(TextureCube Tex, SamplerState Sampler, floa
 	return Tex.SampleBias(Sampler, UV, MipBias);
 #endif
 }
+MaterialFloat4 TextureCubeSampleGrad(TextureCube Tex, SamplerState Sampler, float3 UV, MaterialFloat3 DDX, MaterialFloat3 DDY)
+{
+	return Tex.SampleGrad(Sampler, UV, DDX, DDY);
+}
 
 #if FEATURE_LEVEL >= FEATURE_LEVEL_SM5 // Cubemap arrays are not supported in SM4 feature level
 	MaterialFloat4 TextureCubeArraySampleLevel(TextureCubeArray Tex, SamplerState Sampler, float3 UV, float ArrayIndex, MaterialFloat Mip)
diff --git a/Engine/Source/Editor/MaterialEditor/Private/MaterialEditor.cpp b/Engine/Source/Editor/MaterialEditor/Private/MaterialEditor.cpp
--- a/Engine/Source/Editor/MaterialEditor/Private/MaterialEditor.cpp
+++ b/Engine/Source/Editor/MaterialEditor/Private/MaterialEditor.cpp
@@ -2146,6 +2146,8 @@ void FMaterialEditor::OnConvertObjects()
 							NewTextureExpr->IsDefaultMeshpaintTexture = TextureSampleExpression->IsDefaultMeshpaintTexture;
 							NewTextureExpr->TextureObject = TextureSampleExpression->TextureObject;
 							NewTextureExpr->MipValue = TextureSampleExpression->MipValue;
+							NewTextureExpr->CoordinatesDX = TextureSampleExpression->CoordinatesDX;
+							NewTextureExpr->CoordinatesDY = TextureSampleExpression->CoordinatesDY;
 							NewTextureExpr->MipValueMode = TextureSampleExpression->MipValueMode;
 							NewGraphNode->ReconstructNode();
 						}
diff --git a/Engine/Source/Runtime/Engine/Classes/Materials/MaterialExpressionTextureSample.h b/Engine/Source/Runtime/Engine/Classes/Materials/MaterialExpressionTextureSample.h
--- a/Engine/Source/Runtime/Engine/Classes/Materials/MaterialExpressionTextureSample.h
+++ b/Engine/Source/Runtime/Engine/Classes/Materials/MaterialExpressionTextureSample.h
@@ -10,9 +10,18 @@
 UENUM()
 enum ETextureMipValueMode
 {
+	/* Use hardware computed sample's mip level with automatic anisotropic filtering support. */
 	TMVM_None UMETA(DisplayName="None (use computed mip level)"),
+
+	/* Explicitly compute the sample's mip level. Disables anisotropic filtering. */
 	TMVM_MipLevel UMETA(DisplayName="MipLevel (absolute, 0 is full resolution)"),
+	
+	/* Bias the hardware computed sample's mip level. Disables anisotropic filtering. */
 	TMVM_MipBias UMETA(DisplayName="MipBias (relative to the computed mip level)"),
+	
+	/* Explicitly compute the sample's DDX and DDY for anisotropic filtering. */
+	TMVM_Derivative UMETA(DisplayName="Derivative (explicit derivative to compute mip level)"),
+
 	TMVM_MAX,
 };
 
@@ -34,6 +43,14 @@ class ENGINE_API UMaterialExpressionTextureSample : public UMaterialExpressionTe
 	/** Meaning depends on MipValueMode, a single unit is one mip level  */
 	UPROPERTY(meta = (RequiredInput = "false", ToolTip = "Defaults to 'ConstMipValue' if not specified"))
 	FExpressionInput MipValue;
+	
+	/** Enabled only if MipValueMode == TMVM_Derivative */
+	UPROPERTY(meta = (RequiredInput = "true", ToolTip = "Coordinates derivative over the X axis"))
+	FExpressionInput CoordinatesDX;
+	
+	/** Enabled only if MipValueMode == TMVM_Derivative */
+	UPROPERTY(meta = (RequiredInput = "true", ToolTip = "Coordinates derivative over the Y axis"))
+	FExpressionInput CoordinatesDY;
 
 	/** Defines how the MipValue property is applied to the texture lookup */
 	UPROPERTY(EditAnywhere, Category=MaterialExpressionTextureSample, meta=(DisplayName = "MipValueMode"))
@@ -79,4 +96,10 @@ class ENGINE_API UMaterialExpressionTextureSample : public UMaterialExpressionTe
 	// End UMaterialExpression Interface
 
 	void UpdateTextureResource(class UTexture* InTexture);
-};
\ No newline at end of file
+	
+	int32 CompileMipValue0(class FMaterialCompiler* Compiler);
+	int32 CompileMipValue1(class FMaterialCompiler* Compiler);
+};
+
+
+
diff --git a/Engine/Source/Runtime/Engine/Private/Materials/HLSLMaterialTranslator.h b/Engine/Source/Runtime/Engine/Private/Materials/HLSLMaterialTranslator.h
--- a/Engine/Source/Runtime/Engine/Private/Materials/HLSLMaterialTranslator.h
+++ b/Engine/Source/Runtime/Engine/Private/Materials/HLSLMaterialTranslator.h
@@ -2288,7 +2288,8 @@ protected:
 		int32 TextureIndex,
 		int32 CoordinateIndex,
 		EMaterialSamplerType SamplerType,
-		int32 MipValueIndex=INDEX_NONE,
+		int32 MipValue0Index=INDEX_NONE,
+		int32 MipValue1Index=INDEX_NONE,
 		ETextureMipValueMode MipValueMode=TMVM_None,
 		ESamplerSourceMode SamplerSource=SSM_FromTextureAsset
 		) override
@@ -2318,10 +2319,13 @@ protected:
 			return INDEX_NONE;
 		}
 
-		FString MipValueCode =
-			(MipValueIndex != INDEX_NONE && MipValueMode != TMVM_None)
-			? CoerceParameter(MipValueIndex, MCT_Float1)
-			: TEXT("0.0f");
+		FString MipValue0Code = TEXT("0.0f");
+		FString MipValue1Code = TEXT("0.0f");
+
+		if (MipValue0Index != INDEX_NONE && (MipValueMode == TMVM_MipBias || MipValueMode == TMVM_MipLevel))
+		{
+			MipValue0Code = CoerceParameter(MipValue0Index, MCT_Float1);
+		}
 
 		// if we are not in the PS we need a mip level
 		if(ShaderFrequency != SF_Pixel)
@@ -2350,6 +2354,8 @@ protected:
 			(TextureType == MCT_TextureCube)
 			? TEXT("TextureCubeSample")
 			: TEXT("Texture2DSample");
+		
+		EMaterialValueType UVsType = (TextureType == MCT_TextureCube) ? MCT_Float3 : MCT_Float2;
 	
 		if(MipValueMode == TMVM_None)
 		{
@@ -2371,6 +2377,26 @@ protected:
 		{
 			SampleCode += TEXT("Bias(%s,%sSampler,%s,%s)");
 		}
+		else if(MipValueMode == TMVM_Derivative)
+		{
+			if (MipValue0Index == INDEX_NONE)
+			{
+				return Errorf(TEXT("Missing DDX(UVs) parameter"));
+			}
+			else if (MipValue1Index == INDEX_NONE)
+			{
+				return Errorf(TEXT("Missing DDY(UVs) parameter"));
+			}
+
+			SampleCode += TEXT("Grad(%s,%sSampler,%s,%s,%s)");
+
+			MipValue0Code = CoerceParameter(MipValue0Index, UVsType);
+			MipValue1Code = CoerceParameter(MipValue1Index, UVsType);
+		}
+		else
+		{
+			check(0);
+		}
 
 		switch( SamplerType )
 		{
@@ -2417,10 +2443,7 @@ protected:
 			? CoerceParameter(TextureIndex, MCT_TextureCube)
 			: CoerceParameter(TextureIndex, MCT_Texture2D);
 
-		FString UVs =
-			(TextureType == MCT_TextureCube)
-			? CoerceParameter(CoordinateIndex, MCT_Float3)
-			: CoerceParameter(CoordinateIndex, MCT_Float2);
+		FString UVs = CoerceParameter(CoordinateIndex, UVsType);
 
 		return AddCodeChunk(
 			MCT_Float4,
@@ -2428,7 +2451,8 @@ protected:
 			*TextureName,
 			*TextureName,
 			*UVs,
-			*MipValueCode
+			*MipValue0Code,
+			*MipValue1Code
 			);
 	}
 
diff --git a/Engine/Source/Runtime/Engine/Private/Materials/MaterialExpressions.cpp b/Engine/Source/Runtime/Engine/Private/Materials/MaterialExpressions.cpp
--- a/Engine/Source/Runtime/Engine/Private/Materials/MaterialExpressions.cpp
+++ b/Engine/Source/Runtime/Engine/Private/Materials/MaterialExpressions.cpp
@@ -1211,7 +1211,12 @@ FExpressionInput* UMaterialExpressionTextureSample::GetInput(int32 InputIndex)
 		IF_INPUT_RETURN(TextureObject);
 	}
 
-	if(MipValueMode != TMVM_None)
+	if(MipValueMode == TMVM_Derivative)
+	{
+		IF_INPUT_RETURN(CoordinatesDX);
+		IF_INPUT_RETURN(CoordinatesDY);
+	}
+	else if(MipValueMode != TMVM_None)
 	{
 		IF_INPUT_RETURN(MipValue);
 	}
@@ -1240,6 +1245,11 @@ FString UMaterialExpressionTextureSample::GetInputName(int32 InputIndex) const
 	{
 		IF_INPUT_RETURN(MipValue, TEXT("MipBias"));
 	}
+	else if(MipValueMode == TMVM_Derivative)
+	{
+		IF_INPUT_RETURN(CoordinatesDX, TEXT("DDX(UVs)"));
+		IF_INPUT_RETURN(CoordinatesDY, TEXT("DDY(UVs)"));
+	}
 
 	return TEXT("");
 }
@@ -1341,7 +1351,8 @@ int32 UMaterialExpressionTextureSample::Compile(class FMaterialCompiler* Compile
 				TextureCodeIndex,
 				Coordinates.Expression ? Coordinates.Compile(Compiler) : Compiler->TextureCoordinate(ConstCoordinate, false, false),
 				(EMaterialSamplerType)EffectiveSamplerType,
-				MipValue.Expression ? MipValue.Compile(Compiler) : Compiler->Constant(ConstMipValue),
+				CompileMipValue0(Compiler),
+				CompileMipValue1(Compiler),
 				MipValueMode,
 				SamplerSource
 				);
@@ -1403,6 +1414,11 @@ uint32 UMaterialExpressionTextureSample::GetInputType(int32 InputIndex)
 	{
 		IF_INPUT_RETURN(MipValue, MCT_Float);
 	}
+	else if(MipValueMode == TMVM_Derivative)
+	{
+		IF_INPUT_RETURN(CoordinatesDX, MCT_Float);
+		IF_INPUT_RETURN(CoordinatesDY, MCT_Float);
+	}
 
 	return MCT_Unknown;
 }
@@ -1420,6 +1436,36 @@ void UMaterialExpressionTextureSample::GetConnectorToolTip(int32 InputIndex, int
 }
 #endif
 
+int32 UMaterialExpressionTextureSample::CompileMipValue0(class FMaterialCompiler* Compiler)
+{
+	if (MipValueMode == TMVM_Derivative)
+	{
+		if (CoordinatesDX.Expression)
+		{
+			return CoordinatesDX.Compile(Compiler);
+		}
+	}
+	else if (MipValue.Expression)
+	{
+		return MipValue.Compile(Compiler);
+	}
+	else
+	{
+		return Compiler->Constant(ConstMipValue);
+	}
+
+	return INDEX_NONE;
+}
+
+int32 UMaterialExpressionTextureSample::CompileMipValue1(class FMaterialCompiler* Compiler)
+{
+	if (MipValueMode == TMVM_Derivative && CoordinatesDY.Expression)
+	{
+		return CoordinatesDY.Compile(Compiler);
+	}
+
+	return INDEX_NONE;
+}
 
 
 UMaterialExpressionAdd::UMaterialExpressionAdd(const FObjectInitializer& ObjectInitializer)
@@ -1490,11 +1536,15 @@ int32 UMaterialExpressionTextureSampleParameter::Compile(class FMaterialCompiler
 		return UMaterialExpressionTextureSample::Compile(Compiler, OutputIndex, MultiplexIndex);
 	}
 
+	int32 MipValue0Index = CompileMipValue0(Compiler);
+	int32 MipValue1Index = CompileMipValue1(Compiler);
+
 	return Compiler->TextureSample(
 					Compiler->TextureParameter(ParameterName, Texture, SamplerSource),
 					Coordinates.Expression ? Coordinates.Compile(Compiler) : Compiler->TextureCoordinate(ConstCoordinate, false, false),
 					(EMaterialSamplerType)SamplerType,
-					MipValue.Expression ? MipValue.Compile(Compiler) : Compiler->Constant(ConstMipValue),
+					MipValue0Index,
+					MipValue1Index,
 					MipValueMode,
 					SamplerSource
 					);
diff --git a/Engine/Source/Runtime/Engine/Public/MaterialCompiler.h b/Engine/Source/Runtime/Engine/Public/MaterialCompiler.h
--- a/Engine/Source/Runtime/Engine/Public/MaterialCompiler.h
+++ b/Engine/Source/Runtime/Engine/Public/MaterialCompiler.h
@@ -101,7 +101,7 @@ public:
 	virtual int32 If(int32 A,int32 B,int32 AGreaterThanB,int32 AEqualsB,int32 ALessThanB,int32 Threshold) = 0;
 
 	virtual int32 TextureCoordinate(uint32 CoordinateIndex, bool UnMirrorU, bool UnMirrorV) = 0;
-	virtual int32 TextureSample(int32 Texture,int32 Coordinate,enum EMaterialSamplerType SamplerType,int32 MipValueIndex=INDEX_NONE,ETextureMipValueMode MipValueMode=TMVM_None,ESamplerSourceMode SamplerSource=SSM_FromTextureAsset) = 0;
+	virtual int32 TextureSample(int32 Texture,int32 Coordinate,enum EMaterialSamplerType SamplerType,int32 MipValue0Index=INDEX_NONE,int32 MipValue1Index=INDEX_NONE,ETextureMipValueMode MipValueMode=TMVM_None,ESamplerSourceMode SamplerSource=SSM_FromTextureAsset) = 0;
 
 	virtual int32 TextureDecalMipmapLevel(int32 TextureSizeInput) = 0;
 
@@ -261,8 +261,8 @@ public:
 
 	virtual int32 If(int32 A,int32 B,int32 AGreaterThanB,int32 AEqualsB,int32 ALessThanB,int32 Threshold) override { return Compiler->If(A,B,AGreaterThanB,AEqualsB,ALessThanB,Threshold); }
 
-	virtual int32 TextureSample(int32 InTexture,int32 Coordinate,EMaterialSamplerType SamplerType,int32 MipValueIndex=INDEX_NONE,ETextureMipValueMode MipValueMode=TMVM_None,ESamplerSourceMode SamplerSource=SSM_FromTextureAsset) override 
-		{ return Compiler->TextureSample(InTexture,Coordinate,SamplerType,MipValueIndex,MipValueMode,SamplerSource); }
+	virtual int32 TextureSample(int32 InTexture,int32 Coordinate,enum EMaterialSamplerType SamplerType,int32 MipValue0Index,int32 MipValue1Index,ETextureMipValueMode MipValueMode,ESamplerSourceMode SamplerSource) override 
+		{ return Compiler->TextureSample(InTexture,Coordinate,SamplerType,MipValue0Index,MipValue1Index,MipValueMode,SamplerSource); }
 
 	virtual int32 TextureCoordinate(uint32 CoordinateIndex, bool UnMirrorU, bool UnMirrorV) override { return Compiler->TextureCoordinate(CoordinateIndex, UnMirrorU, UnMirrorV); }
 
