Fixed APlayerController::ReceivedSpectatorClass() to take a class of the correct type. (UE-18286)

diff --git a/Engine/Source/Runtime/Engine/Classes/GameFramework/PlayerController.h b/Engine/Source/Runtime/Engine/Classes/GameFramework/PlayerController.h
--- a/Engine/Source/Runtime/Engine/Classes/GameFramework/PlayerController.h
+++ b/Engine/Source/Runtime/Engine/Classes/GameFramework/PlayerController.h
@@ -1499,7 +1499,7 @@ public:
 	class APawn* GetPawnOrSpectator() const;
 
 	/** Called to notify the controller that the spectator class has been received. */
-	virtual void ReceivedSpectatorClass(TSubclassOf<class AGameMode> SpectatorClass);
+	virtual void ReceivedSpectatorClass(TSubclassOf<class ASpectatorPawn> SpectatorClass);
 
 	/** Returns the location the PlayerController is focused on.
 	 *  If there is a possessed Pawn, returns the Pawn's location.
diff --git a/Engine/Source/Runtime/Engine/Private/GameState.cpp b/Engine/Source/Runtime/Engine/Private/GameState.cpp
--- a/Engine/Source/Runtime/Engine/Private/GameState.cpp
+++ b/Engine/Source/Runtime/Engine/Private/GameState.cpp
@@ -113,7 +113,7 @@ void AGameState::ReceivedSpectatorClass()
 	{
 		if (It->PlayerController)
 		{
-			It->PlayerController->ReceivedSpectatorClass(GameModeClass);
+			It->PlayerController->ReceivedSpectatorClass(SpectatorClass);
 		}
 	}
 }
diff --git a/Engine/Source/Runtime/Engine/Private/PlayerController.cpp b/Engine/Source/Runtime/Engine/Private/PlayerController.cpp
--- a/Engine/Source/Runtime/Engine/Private/PlayerController.cpp
+++ b/Engine/Source/Runtime/Engine/Private/PlayerController.cpp
@@ -3564,7 +3564,7 @@ void APlayerController::ReceivedGameModeClass(TSubclassOf<AGameMode> GameModeCla
 {
 }
 
-void APlayerController::ReceivedSpectatorClass(TSubclassOf<AGameMode> SpectatorClass)
+void APlayerController::ReceivedSpectatorClass(TSubclassOf<ASpectatorPawn> SpectatorClass)
 {
 	if (IsInState(NAME_Spectating))
 	{
