#ue4 - Improve comments on SceneComponent set location/rotation/transform functions, especially regarding bSweep and bTeleport.

diff --git a/Engine/Source/Runtime/Engine/Classes/Components/SceneComponent.h b/Engine/Source/Runtime/Engine/Classes/Components/SceneComponent.h
--- a/Engine/Source/Runtime/Engine/Classes/Components/SceneComponent.h
+++ b/Engine/Source/Runtime/Engine/Classes/Components/SceneComponent.h
@@ -156,7 +156,7 @@ public:
 	UPROPERTY(EditAnywhere, BlueprintReadWrite, AdvancedDisplay, Category=Physics)
 	uint32 bShouldUpdatePhysicsVolume:1;
 
-	/** If true, a change in the bounds of this component will call trigger a streaming data rebuild */
+	/** If true, a change in the bounds of the component will call trigger a streaming data rebuild */
 	UPROPERTY()
 	uint32 bBoundsChangeTriggersStreamingDataRebuild:1;
 
@@ -179,10 +179,10 @@ private:
 	TWeakObjectPtr<class APhysicsVolume> PhysicsVolume;
 
 public:
-	/** Current bounds of this component */
+	/** Current bounds of the component */
 	FBoxSphereBounds Bounds;
 
-	/** Current transform of this component, relative to the world */
+	/** Current transform of the component, relative to the world */
 	FTransform ComponentToWorld;
 
 private:
@@ -191,11 +191,11 @@ private:
 
 public:
 
-	/** Location of this component relative to its parent */
+	/** Location of the component relative to its parent */
 	UPROPERTY(EditAnywhere, BlueprintReadOnly, ReplicatedUsing=OnRep_Transform, Category = Transform)
 	FVector RelativeLocation;
 
-	/** Rotation of this component relative to its parent */
+	/** Rotation of the component relative to its parent */
 	UPROPERTY(EditAnywhere, BlueprintReadOnly, ReplicatedUsing=OnRep_Transform, Category=Transform)
 	FRotator RelativeRotation;
 
@@ -205,7 +205,7 @@ private:
 
 public:
 	/**
-	*	Non-uniform scaling of this component relative to its parent.
+	*	Non-uniform scaling of the component relative to its parent.
 	*	Note that scaling is always applied in local space (no shearing etc)
 	*/
 	UPROPERTY(BlueprintReadOnly, ReplicatedUsing=OnRep_Transform, interp, Category=Transform)
@@ -241,108 +241,241 @@ private:
 public:
 
 	/**
-	 * Velocity of this component.
+	 * Velocity of the component.
 	 * @see GetComponentVelocity()
 	 */
 	UPROPERTY()
 	FVector ComponentVelocity;
 
 public:
-	/** Set the location of this component relative to its parent */
+	/**
+	 * Set the location of the component relative to its parent
+	 * @param NewLocation		New location of the component relative to its parent.		
+	 * @param SweepHitResult	Hit result from any impact if sweep is true.
+	 * @param bSweep			Whether we sweep to the destination location, triggering overlaps along the way and stopping short of the target if blocked by something.
+	 *							Only the root component is swept and checked for blocking collision, child components move without sweeping. If collision is off, this has no effect.
+	 * @param bTeleport			Whether we teleport the physics state (if physics collision is enabled for this object).
+	 *							If true, physics velocity for this object is unchanged (so ragdoll parts are not affected by change in location).
+	 *							If false, physics velocity is updated based on the change in position (affecting ragdoll parts).
+	 *							If CCD is on and not teleporting, this will affect objects along the entire sweep volume.
+	 */
 	UFUNCTION(BlueprintCallable, Category="Utilities|Transformation", meta=(DisplayName="SetRelativeLocation"))
 	void K2_SetRelativeLocation(FVector NewLocation, bool bSweep, FHitResult& SweepHitResult, bool bTeleport);
 	void SetRelativeLocation(FVector NewLocation, bool bSweep=false, FHitResult* OutSweepHitResult=nullptr, ETeleportType Teleport = ETeleportType::None);
 
 
-	/** Set the rotation of this component relative to its parent */
+	/**
+	 * Set the rotation of the component relative to its parent
+	 * @param NewRotation		New rotation of the component relative to its parent
+	 * @param SweepHitResult	Hit result from any impact if sweep is true.
+	 * @param bSweep			Whether we sweep to the destination (currently not supported for rotation).
+	 * @param bTeleport			Whether we teleport the physics state (if physics collision is enabled for this object).
+	 *							If true, physics velocity for this object is unchanged (so ragdoll parts are not affected by change in location).
+	 *							If false, physics velocity is updated based on the change in position (affecting ragdoll parts).
+	 */
 	UFUNCTION(BlueprintCallable, Category="Utilities|Transformation", meta=(DisplayName="SetRelativeRotation", AdvancedDisplay="bSweep,SweepHitResult,bTeleport"))
 	void K2_SetRelativeRotation(FRotator NewRotation, bool bSweep, FHitResult& SweepHitResult, bool bTeleport);
 	void SetRelativeRotation(FRotator NewRotation, bool bSweep=false, FHitResult* OutSweepHitResult=nullptr, ETeleportType Teleport = ETeleportType::None);
 	void SetRelativeRotation(const FQuat& NewRotation, bool bSweep=false, FHitResult* OutSweepHitResult=nullptr, ETeleportType Teleport = ETeleportType::None);
 
-	/** Set the transform of this component relative to its parent */
+	/**
+	 * Set the transform of the component relative to its parent
+	 * @param NewTransform		New transform of the component relative to its parent.
+	 * @param SweepHitResult	Hit result from any impact if sweep is true.
+	 * @param bSweep			Whether we sweep to the destination (currently not supported for rotation).
+	 * @param bTeleport			Whether we teleport the physics state (if physics collision is enabled for this object).
+	 *							If true, physics velocity for this object is unchanged (so ragdoll parts are not affected by change in location).
+	 *							If false, physics velocity is updated based on the change in position (affecting ragdoll parts).
+	 */
 	UFUNCTION(BlueprintCallable, Category="Utilities|Transformation", meta=(DisplayName="SetRelativeTransform"))
 	void K2_SetRelativeTransform(const FTransform& NewTransform, bool bSweep, FHitResult& SweepHitResult, bool bTeleport);
 	void SetRelativeTransform(const FTransform& NewTransform, bool bSweep=false, FHitResult* OutSweepHitResult=nullptr, ETeleportType Teleport = ETeleportType::None);
 
-	/** Returns the transform of this component relative to its parent */
+	/** Returns the transform of the component relative to its parent */
 	UFUNCTION(BlueprintCallable, Category="Utilities|Transformation")
 	FTransform GetRelativeTransform() const;
 
-	/** Set the transform of this component relative to its parent */
+	/** Reset the transform of the component relative to its parent. Sets relative location to zero, relative rotation to no rotation, and Scale to 1. */
 	UFUNCTION(BlueprintCallable, Category="Utilities|Transformation")
 	void ResetRelativeTransform();
 
-	/** Set the non-uniform of this component relative to its parent */
+	/** Set the non-uniform scale of the component relative to its parent */
 	UFUNCTION(BlueprintCallable, Category="Utilities|Transformation")
 	virtual void SetRelativeScale3D(FVector NewScale3D);
 
-	/** Adds a delta to the translation of this component relative to its parent */
+	/**
+	 * Adds a delta to the translation of the component relative to its parent
+	 * @param DeltaLocation		Change in location of the component relative to its parent
+	 * @param SweepHitResult	Hit result from any impact if sweep is true.
+	 * @param bSweep			Whether we sweep to the destination location, triggering overlaps along the way and stopping short of the target if blocked by something.
+	 *							Only the root component is swept and checked for blocking collision, child components move without sweeping. If collision is off, this has no effect.
+	 * @param bTeleport			Whether we teleport the physics state (if physics collision is enabled for this object).
+	 *							If true, physics velocity for this object is unchanged (so ragdoll parts are not affected by change in location).
+	 *							If false, physics velocity is updated based on the change in position (affecting ragdoll parts).
+	 *							If CCD is on and not teleporting, this will affect objects along the entire sweep volume.
+	 */
 	UFUNCTION(BlueprintCallable, Category="Utilities|Transformation", meta=(DisplayName="AddRelativeLocation"))
 	void K2_AddRelativeLocation(FVector DeltaLocation, bool bSweep, FHitResult& SweepHitResult, bool bTeleport);
 	void AddRelativeLocation(FVector DeltaLocation, bool bSweep=false, FHitResult* OutSweepHitResult=nullptr, ETeleportType Teleport = ETeleportType::None);
 
-	/** Adds a delta the rotation of this component relative to its parent */
+	/**
+	 * Adds a delta the rotation of the component relative to its parent
+	 * @param DeltaRotation		Change in rotation of the component relative to is parent.
+	 * @param SweepHitResult	Hit result from any impact if sweep is true.
+	 * @param bSweep			Whether we sweep to the destination (currently not supported for rotation).
+	 * @param bTeleport			Whether we teleport the physics state (if physics collision is enabled for this object).
+	 *							If true, physics velocity for this object is unchanged (so ragdoll parts are not affected by change in location).
+	 *							If false, physics velocity is updated based on the change in position (affecting ragdoll parts).
+	 */
 	UFUNCTION(BlueprintCallable, Category="Utilities|Transformation", meta=(DisplayName="AddRelativeRotation", AdvancedDisplay="bSweep,SweepHitResult,bTeleport"))
 	void K2_AddRelativeRotation(FRotator DeltaRotation, bool bSweep, FHitResult& SweepHitResult, bool bTeleport);
 	void AddRelativeRotation(FRotator DeltaRotation, bool bSweep=false, FHitResult* OutSweepHitResult=nullptr, ETeleportType Teleport = ETeleportType::None);
 	void AddRelativeRotation(const FQuat& DeltaRotation, bool bSweep=false, FHitResult* OutSweepHitResult=nullptr, ETeleportType Teleport = ETeleportType::None);
 
-	/** Adds a delta to the location of this component in its local reference frame */
+	/**
+	 * Adds a delta to the location of the component in its local reference frame
+	 * @param DeltaLocation		Change in location of the component in its local reference frame.
+	 * @param SweepHitResult	Hit result from any impact if sweep is true.
+	 * @param bSweep			Whether we sweep to the destination location, triggering overlaps along the way and stopping short of the target if blocked by something.
+	 *							Only the root component is swept and checked for blocking collision, child components move without sweeping. If collision is off, this has no effect.
+	 * @param bTeleport			Whether we teleport the physics state (if physics collision is enabled for this object).
+	 *							If true, physics velocity for this object is unchanged (so ragdoll parts are not affected by change in location).
+	 *							If false, physics velocity is updated based on the change in position (affecting ragdoll parts).
+	 *							If CCD is on and not teleporting, this will affect objects along the entire sweep volume.
+	 */
 	UFUNCTION(BlueprintCallable, Category="Utilities|Transformation", meta=(DisplayName="AddLocalOffset", Keywords="location position"))
 	void K2_AddLocalOffset(FVector DeltaLocation, bool bSweep, FHitResult& SweepHitResult, bool bTeleport);
 	void AddLocalOffset(FVector DeltaLocation, bool bSweep=false, FHitResult* OutSweepHitResult=nullptr, ETeleportType Teleport = ETeleportType::None);
 
 	
-	/** Adds a delta to the rotation of this component in its local reference frame */
+	/**
+	 * Adds a delta to the rotation of the component in its local reference frame
+	 * @param DeltaRotation		Change in rotation of the component in its local reference frame.
+	 * @param SweepHitResult	Hit result from any impact if sweep is true.
+	 * @param bSweep			Whether we sweep to the destination (currently not supported for rotation).
+	 * @param bTeleport			Whether we teleport the physics state (if physics collision is enabled for this object).
+	 *							If true, physics velocity for this object is unchanged (so ragdoll parts are not affected by change in location).
+	 *							If false, physics velocity is updated based on the change in position (affecting ragdoll parts).
+	 */
 	UFUNCTION(BlueprintCallable, Category="Utilities|Transformation", meta=(DisplayName="AddLocalRotation", AdvancedDisplay="bSweep,SweepHitResult,bTeleport"))
 	void K2_AddLocalRotation(FRotator DeltaRotation, bool bSweep, FHitResult& SweepHitResult, bool bTeleport);
 	void AddLocalRotation(FRotator DeltaRotation, bool bSweep=false, FHitResult* OutSweepHitResult=nullptr, ETeleportType Teleport = ETeleportType::None);
 	void AddLocalRotation(const FQuat& DeltaRotation, bool bSweep=false, FHitResult* OutSweepHitResult=nullptr, ETeleportType Teleport = ETeleportType::None);
 
 
-	/** Adds a delta to the transform of this component in its local reference frame. Scale is unchanged. */
+	/**
+	 * Adds a delta to the transform of the component in its local reference frame. Scale is unchanged.
+	 * @param DeltaTransform	Change in transform of the component in its local reference frame. Scale is unchanged.
+	 * @param SweepHitResult	Hit result from any impact if sweep is true.
+	 * @param bSweep			Whether we sweep to the destination location, triggering overlaps along the way and stopping short of the target if blocked by something.
+	 *							Only the root component is swept and checked for blocking collision, child components move without sweeping. If collision is off, this has no effect.
+	 * @param bTeleport			Whether we teleport the physics state (if physics collision is enabled for this object).
+	 *							If true, physics velocity for this object is unchanged (so ragdoll parts are not affected by change in location).
+	 *							If false, physics velocity is updated based on the change in position (affecting ragdoll parts).
+	 *							If CCD is on and not teleporting, this will affect objects along the entire sweep volume.
+	 */
 	UFUNCTION(BlueprintCallable, Category="Utilities|Transformation", meta=(DisplayName="AddLocalTransform"))
 	void K2_AddLocalTransform(const FTransform& DeltaTransform, bool bSweep, FHitResult& SweepHitResult, bool bTeleport);
 	void AddLocalTransform(const FTransform& DeltaTransform, bool bSweep=false, FHitResult* OutSweepHitResult=nullptr, ETeleportType Teleport = ETeleportType::None);
 
 
-	/** Set the relative location of this component to put it at the supplied location in world space. */
+	/**
+	 * Put this component at the specified location in world space. Updates relative location to achieve the final world location.
+	 * @param NewLocation		New location in world space for the component.
+	 * @param SweepHitResult	Hit result from any impact if sweep is true.
+	 * @param bSweep			Whether we sweep to the destination location, triggering overlaps along the way and stopping short of the target if blocked by something.
+	 *							Only the root component is swept and checked for blocking collision, child components move without sweeping. If collision is off, this has no effect.
+	 * @param bTeleport			Whether we teleport the physics state (if physics collision is enabled for this object).
+	 *							If true, physics velocity for this object is unchanged (so ragdoll parts are not affected by change in location).
+	 *							If false, physics velocity is updated based on the change in position (affecting ragdoll parts).
+	 *							If CCD is on and not teleporting, this will affect objects along the entire sweep volume.
+	 */
 	UFUNCTION(BlueprintCallable, Category="Utilities|Transformation", meta=(DisplayName="SetWorldLocation"))
 	void K2_SetWorldLocation(FVector NewLocation, bool bSweep, FHitResult& SweepHitResult, bool bTeleport);
 	void SetWorldLocation(FVector NewLocation, bool bSweep=false, FHitResult* OutSweepHitResult=nullptr, ETeleportType Teleport = ETeleportType::None);
 
 
-	/** Set the relative rotation of this component to put it at the supplied orientation in world space. */
+	/*
+	 * Put this component at the specified rotation in world space. Updates relative rotation to achieve the final world rotation.
+	 * @param NewRotation		New rotation in world space for the component.
+	 * @param SweepHitResult	Hit result from any impact if sweep is true.
+	 * @param bSweep			Whether we sweep to the destination (currently not supported for rotation).
+	 * @param bTeleport			Whether we teleport the physics state (if physics collision is enabled for this object).
+	 *							If true, physics velocity for this object is unchanged (so ragdoll parts are not affected by change in location).
+	 *							If false, physics velocity is updated based on the change in position (affecting ragdoll parts).
+	 *							If CCD is on and not teleporting, this will affect objects along the entire sweep volume.
+	 */
 	UFUNCTION(BlueprintCallable, Category="Utilities|Transformation", meta=(DisplayName="SetWorldRotation", AdvancedDisplay="bSweep,SweepHitResult,bTeleport"))
 	void K2_SetWorldRotation(FRotator NewRotation, bool bSweep, FHitResult& SweepHitResult, bool bTeleport);
 	void SetWorldRotation(FRotator NewRotation, bool bSweep=false, FHitResult* OutSweepHitResult=nullptr, ETeleportType Teleport = ETeleportType::None);
 	void SetWorldRotation(const FQuat& NewRotation, bool bSweep=false, FHitResult* OutSweepHitResult=nullptr, ETeleportType Teleport = ETeleportType::None);
 
 
-	/** Set the relative scale of this component to put it at the supplied scale in world space. */
+	/**
+	 * Set the relative scale of the component to put it at the supplied scale in world space.
+	 * @param NewScale		New scale in world space for this component.
+	 */
 	UFUNCTION(BlueprintCallable, Category="Utilities|Transformation")
 	void SetWorldScale3D(FVector NewScale);
 
-	/** Set the transform of this component in world space. */
+	/**
+	 * Set the transform of the component in world space.
+	 * @param NewTransform		New transform in world space for the component.
+	 * @param SweepHitResult	Hit result from any impact if sweep is true.
+	 * @param bSweep			Whether we sweep to the destination location, triggering overlaps along the way and stopping short of the target if blocked by something.
+	 *							Only the root component is swept and checked for blocking collision, child components move without sweeping. If collision is off, this has no effect.
+	 * @param bTeleport			Whether we teleport the physics state (if physics collision is enabled for this object).
+	 *							If true, physics velocity for this object is unchanged (so ragdoll parts are not affected by change in location).
+	 *							If false, physics velocity is updated based on the change in position (affecting ragdoll parts).
+	 *							If CCD is on and not teleporting, this will affect objects along the entire sweep volume.
+	 */
 	UFUNCTION(BlueprintCallable, Category="Utilities|Transformation", meta=(DisplayName="SetWorldTransform"))
 	void K2_SetWorldTransform(const FTransform& NewTransform, bool bSweep, FHitResult& SweepHitResult, bool bTeleport);
 	void SetWorldTransform(const FTransform& NewTransform, bool bSweep=false, FHitResult* OutSweepHitResult=nullptr, ETeleportType Teleport = ETeleportType::None);
 
 
-	/** Adds a delta to the location of this component in world space. */
+	/**
+	 * Adds a delta to the location of the component in world space.
+	 * @param DeltaLocation		Change in location in world space for the component.
+	 * @param SweepHitResult	Hit result from any impact if sweep is true.
+	 * @param bSweep			Whether we sweep to the destination location, triggering overlaps along the way and stopping short of the target if blocked by something.
+	 *							Only the root component is swept and checked for blocking collision, child components move without sweeping. If collision is off, this has no effect.
+	 * @param bTeleport			Whether we teleport the physics state (if physics collision is enabled for this object).
+	 *							If true, physics velocity for this object is unchanged (so ragdoll parts are not affected by change in location).
+	 *							If false, physics velocity is updated based on the change in position (affecting ragdoll parts).
+	 *							If CCD is on and not teleporting, this will affect objects along the entire sweep volume.
+	 */
 	UFUNCTION(BlueprintCallable, Category="Utilities|Transformation", meta=(DisplayName="AddWorldOffset", Keywords="location position"))
 	void K2_AddWorldOffset(FVector DeltaLocation, bool bSweep, FHitResult& SweepHitResult, bool bTeleport);
 	void AddWorldOffset(FVector DeltaLocation, bool bSweep=false, FHitResult* OutSweepHitResult=nullptr, ETeleportType Teleport = ETeleportType::None);
 
 
-	/** Adds a delta to the rotation of this component in world space. */
+	/**
+	 * Adds a delta to the rotation of the component in world space.
+	 * @param DeltaRotation		Change in rotation in world space for the component.
+	 * @param SweepHitResult	Hit result from any impact if sweep is true.
+	 * @param bSweep			Whether we sweep to the destination (currently not supported for rotation).
+	 * @param bTeleport			Whether we teleport the physics state (if physics collision is enabled for this object).
+	 *							If true, physics velocity for this object is unchanged (so ragdoll parts are not affected by change in location).
+	 *							If false, physics velocity is updated based on the change in position (affecting ragdoll parts).
+	 *							If CCD is on and not teleporting, this will affect objects along the entire sweep volume.
+	 */
 	UFUNCTION(BlueprintCallable, Category="Utilities|Transformation", meta=(DisplayName="AddWorldRotation", AdvancedDisplay="bSweep,SweepHitResult,bTeleport"))
 	void K2_AddWorldRotation(FRotator DeltaRotation, bool bSweep, FHitResult& SweepHitResult, bool bTeleport);
 	void AddWorldRotation(FRotator DeltaRotation, bool bSweep=false, FHitResult* OutSweepHitResult=nullptr, ETeleportType Teleport = ETeleportType::None);
 	void AddWorldRotation(const FQuat& DeltaRotation, bool bSweep=false, FHitResult* OutSweepHitResult=nullptr, ETeleportType Teleport = ETeleportType::None);
 
-	/** Adds a delta to the transform of this component in world space. Scale is unchanged. */
+	/**
+	 * Adds a delta to the transform of the component in world space. Scale is unchanged.
+	 * @param DeltaTransform	Change in transform in world space for the component. Scale is unchanged.
+	 * @param SweepHitResult	Hit result from any impact if sweep is true.
+	 * @param bSweep			Whether we sweep to the destination location, triggering overlaps along the way and stopping short of the target if blocked by something.
+	 *							Only the root component is swept and checked for blocking collision, child components move without sweeping. If collision is off, this has no effect.
+	 * @param bTeleport			Whether we teleport the physics state (if physics collision is enabled for this object).
+	 *							If true, physics velocity for this object is unchanged (so ragdoll parts are not affected by change in location).
+	 *							If false, physics velocity is updated based on the change in position (affecting ragdoll parts).
+	 *							If CCD is on and not teleporting, this will affect objects along the entire sweep volume.
+	 */
 	UFUNCTION(BlueprintCallable, Category="Utilities|Transformation", meta=(DisplayName="AddWorldTransform"))
 	void K2_AddWorldTransform(const FTransform& DeltaTransform, bool bSweep, FHitResult& SweepHitResult, bool bTeleport);
 	void AddWorldTransform(const FTransform& DeltaTransform, bool bSweep=false, FHitResult* OutSweepHitResult=nullptr, ETeleportType Teleport = ETeleportType::None);
@@ -409,24 +542,24 @@ public:
 
 	/** 
 	 *   Attach this component to another scene component, optionally at a named socket. It is valid to call this on components whether or not they have been Registered.
-	 *   @param bMaintainWorldTransform	If true, update the relative location/rotation of this component to keep its world position the same
+	 *   @param bMaintainWorldTransform	If true, update the relative location/rotation of the component to keep its world position the same
 	 */
 	void AttachTo(class USceneComponent* InParent, FName InSocketName = NAME_None, EAttachLocation::Type AttachType = EAttachLocation::KeepRelativeOffset, bool bWeldSimulatedBodies = false);
 
 	/**
 	*   Attach this component to another scene component, optionally at a named socket. It is valid to call this on components whether or not they have been Registered.
-	*   @param bMaintainWorldTransform	If true, update the relative location/rotation of this component to keep its world position the same
+	*   @param bMaintainWorldTransform	If true, update the relative location/rotation of the component to keep its world position the same
 	*/
 	UFUNCTION(BlueprintCallable, Category = "Utilities|Transformation", meta = (DisplayName = "AttachTo", AttachType = "KeepRelativeOffset"))
 	void K2_AttachTo(class USceneComponent* InParent, FName InSocketName = NAME_None, EAttachLocation::Type AttachType = EAttachLocation::KeepRelativeOffset, bool bWeldSimulatedBodies = true);
 
-	/** Zeroes out the relative transform of this component, and calls AttachTo(). Useful for attaching directly to a scene component or socket location  */
+	/** Zeroes out the relative transform of the component, and calls AttachTo(). Useful for attaching directly to a scene component or socket location  */
 	UFUNCTION(BlueprintCallable, meta=(DeprecatedFunction, DeprecationMessage = "Use AttachTo with EAttachLocation::SnapToTarget option instead"), Category="Utilities|Transformation")
 	void SnapTo(class USceneComponent* InParent, FName InSocketName = NAME_None);
 
 	/** 
 	 *	Detach this component from whatever it is attached to. Automatically unwelds components that are welded together (See WeldTo)
-	 *   @param bMaintainWorldTransform	If true, update the relative location/rotation of this component to keep its world position the same *	
+	 *   @param bMaintainWorldTransform	If true, update the relative location/rotation of the component to keep its world position the same *	
 	 */
 	UFUNCTION(BlueprintCallable, Category="Utilities|Transformation")
 	virtual void DetachFromParent(bool bMaintainWorldPosition = false, bool bCallModify = true);
@@ -490,8 +623,8 @@ public:
 	virtual void QuerySupportedSockets(TArray<FComponentSocketDescription>& OutSockets) const;
 
 	/** 
-	 * Get velocity of this component: either ComponentVelocity, or the velocity of the physics body if simulating physics.
-	 * @return Velocity of this component
+	 * Get velocity of the component: either ComponentVelocity, or the velocity of the physics body if simulating physics.
+	 * @return Velocity of the component
 	 */
 	UFUNCTION(BlueprintCallable, Category="Utilities|Transformation")
 	virtual FVector GetComponentVelocity() const;
@@ -588,7 +721,12 @@ public:
 	 * @note The overload taking rotation as an FQuat is slightly faster than the version using FRotator (which will be converted to an FQuat)..
 	 * @param Delta			The desired location change in world space.
 	 * @param NewRotation	The new desired rotation in world space.
-	 * @param bSweep		True to do a swept move, which will stop at a blocking collision.
+	 * @param bSweep		Whether we sweep to the destination location, triggering overlaps along the way and stopping short of the target if blocked by something.
+	 *						Only the root component is swept and checked for blocking collision, child components move without sweeping. If collision is off, this has no effect.
+	 * @param Teleport		Whether we teleport the physics state (if physics collision is enabled for this object).
+	 *						If TeleportPhysics, physics velocity for this object is unchanged (so ragdoll parts are not affected by change in location).
+	 *						If None, physics velocity is updated based on the change in position (affecting ragdoll parts).
+	 *						If CCD is on and not teleporting, this will affect objects along the entire swept volume.
 	 * @param Hit			Optional output describing the blocking hit that stopped the move, if any.
 	 * @param MoveFlags		Flags controlling behavior of the move. @see EMoveComponentFlags
 	 * @param Teleport      Determines whether to teleport the physics body or not. Teleporting will maintain constant velocity and avoid collisions along the path
@@ -689,16 +827,16 @@ public:
 	/** Update transforms of any components attached to this one. */
 	void UpdateChildTransforms(bool bSkipPhysicsMove = false, ETeleportType Teleport = ETeleportType::None);
 
-	/** Calculate the bounds of this component. Default behavior is a bounding box/sphere of zero size. */
+	/** Calculate the bounds of the component. Default behavior is a bounding box/sphere of zero size. */
 	virtual FBoxSphereBounds CalcBounds(const FTransform& LocalToWorld) const;
 
 	/**
-	 * Calculate the axis-aligned bounding cylinder of this component (radius in X-Y, half-height along Z axis).
+	 * Calculate the axis-aligned bounding cylinder of the component (radius in X-Y, half-height along Z axis).
 	 * Default behavior is just a cylinder around the box of the cached BoxSphereBounds.
 	 */
 	virtual void CalcBoundingCylinder(float& CylinderRadius, float& CylinderHalfHeight) const;
 
-	/** Update the Bounds of this component.*/
+	/** Update the Bounds of the component.*/
 	virtual void UpdateBounds();
 
 	/** If true, bounds should be used when placing component/actor in level, and spawning may fail */
@@ -767,7 +905,18 @@ protected:
 
 	
 public:
-	/** Set the location and rotation of this component relative to its parent */
+	/**
+	 * Set the location and rotation of the component relative to its parent
+	 * @param NewLocation		New location of the component relative to its parent.
+	 * @param NewRotation		New rotation of the component relative to its parent.
+	 * @param SweepHitResult	Hit result from any impact if sweep is true.
+	 * @param bSweep			Whether we sweep to the destination location, triggering overlaps along the way and stopping short of the target if blocked by something.
+	 *							Only the root component is swept and checked for blocking collision, child components move without sweeping. If collision is off, this has no effect.
+	 * @param bTeleport			Whether we teleport the physics state (if physics collision is enabled for this object).
+	 *							If true, physics velocity for this object is unchanged (so ragdoll parts are not affected by change in location).
+	 *							If false, physics velocity is updated based on the change in position (affecting ragdoll parts).
+	 *							If CCD is on and not teleporting, this will affect objects along the entire sweep volume.
+	 */
 	UFUNCTION(BlueprintCallable, Category="Utilities|Transformation", meta=(DisplayName="SetRelativeLocationAndRotation"))
 	void K2_SetRelativeLocationAndRotation(FVector NewLocation, FRotator NewRotation, bool bSweep, FHitResult& SweepHitResult, bool bTeleport);
 	void SetRelativeLocationAndRotation(FVector NewLocation, FRotator NewRotation, bool bSweep=false, FHitResult* OutSweepHitResult=nullptr, ETeleportType Teleport = ETeleportType::None);
@@ -777,12 +926,23 @@ public:
 	UFUNCTION(BlueprintCallable, Category="Utilities|Transformation")
 	void SetAbsolute(bool bNewAbsoluteLocation = false, bool bNewAbsoluteRotation = false, bool bNewAbsoluteScale = false);
 
-	/** Set the relative location and FRotator rotation of this component to put it at the supplied pose in world space. */
+	/**
+	 * Set the relative location and rotation of the component to put it at the supplied pose in world space.
+	 * @param NewLocation		New location in world space for the component.
+	 * @param NewRotation		New rotation in world space for the component.
+	 * @param SweepHitResult	Hit result from any impact if sweep is true.
+	 * @param bSweep			Whether we sweep to the destination location, triggering overlaps along the way and stopping short of the target if blocked by something.
+	 *							Only the root component is swept and checked for blocking collision, child components move without sweeping. If collision is off, this has no effect.
+	 * @param bTeleport			Whether we teleport the physics state (if physics collision is enabled for this object).
+	 *							If true, physics velocity for this object is unchanged (so ragdoll parts are not affected by change in location).
+	 *							If false, physics velocity is updated based on the change in position (affecting ragdoll parts).
+	 *							If CCD is on and not teleporting, this will affect objects along the entire sweep volume.
+	 */
 	UFUNCTION(BlueprintCallable, Category="Utilities|Transformation", meta=(DisplayName="SetWorldLocationAndRotation"))
 	void K2_SetWorldLocationAndRotation(FVector NewLocation, FRotator NewRotation, bool bSweep, FHitResult& SweepHitResult, bool bTeleport);
 	void SetWorldLocationAndRotation(FVector NewLocation, FRotator NewRotation, bool bSweep=false, FHitResult* OutSweepHitResult=nullptr, ETeleportType Teleport = ETeleportType::None);
 
-	/** Set the relative location and FQuat rotation of this component to put it at the supplied pose in world space. */
+	/** Set the relative location and FQuat rotation of the component to put it at the supplied pose in world space. */
 	void SetWorldLocationAndRotation(FVector NewLocation, const FQuat& NewRotation, bool bSweep=false, FHitResult* OutSweepHitResult=nullptr, ETeleportType Teleport = ETeleportType::None);
 
 	/** Special version of SetWorldLocationAndRotation that does not affect physics. */
diff --git a/Engine/Source/Runtime/Engine/Classes/GameFramework/Actor.h b/Engine/Source/Runtime/Engine/Classes/GameFramework/Actor.h
--- a/Engine/Source/Runtime/Engine/Classes/GameFramework/Actor.h
+++ b/Engine/Source/Runtime/Engine/Classes/GameFramework/Actor.h
@@ -648,7 +648,7 @@ public:
 	 * @param bTeleport		Whether we teleport the physics state (if physics collision is enabled for this object).
 	 *						If true, physics velocity for this object is unchanged (so ragdoll parts are not affected by change in location).
 	 *						If false, physics velocity is updated based on the change in position (affecting ragdoll parts).
-	 *						If CCD is on and not teleporting, this will affect objects along the entire sweep volume.
+	 *						If CCD is on and not teleporting, this will affect objects along the entire swept volume.
 	 * @param SweepHitResult	The hit result from the move if swept.
 	 * @return	Whether the location was successfully set (if not swept), or whether movement occurred at all (if swept).
 	 */
@@ -695,7 +695,7 @@ public:
 	 * @param bTeleport		Whether we teleport the physics state (if physics collision is enabled for this object).
 	 *						If true, physics velocity for this object is unchanged (so ragdoll parts are not affected by change in location).
 	 *						If false, physics velocity is updated based on the change in position (affecting ragdoll parts).
-	 *						If CCD is on and not teleporting, this will affect objects along the entire sweep volume.
+	 *						If CCD is on and not teleporting, this will affect objects along the entire swept volume.
 	 * @param OutSweepHitResult The hit result from the move if swept.
 	 * @return	Whether the location was successfully set if not swept, or whether movement occurred if swept.
 	 */
@@ -721,7 +721,7 @@ public:
 	 * @param bTeleport			Whether we teleport the physics state (if physics collision is enabled for this object).
 	 *							If true, physics velocity for this object is unchanged (so ragdoll parts are not affected by change in location).
 	 *							If false, physics velocity is updated based on the change in position (affecting ragdoll parts).
-	 *							If CCD is on and not teleporting, this will affect objects along the entire sweep volume.
+	 *							If CCD is on and not teleporting, this will affect objects along the entire swept volume.
 	 * @param SweepHitResult	The hit result from the move if swept.
 	 * @return	Whether the rotation was successfully set.
 	 */
@@ -738,7 +738,7 @@ public:
 	 * @param bTeleport			Whether we teleport the physics state (if physics collision is enabled for this object).
 	 *							If true, physics velocity for this object is unchanged (so ragdoll parts are not affected by change in location).
 	 *							If false, physics velocity is updated based on the change in position (affecting ragdoll parts).
-	 *							If CCD is on and not teleporting, this will affect objects along the entire sweep volume.
+	 *							If CCD is on and not teleporting, this will affect objects along the entire swept volume.
 	 * @param OutSweepHitResult	The hit result from the move if swept.
 	 * @return	Whether the rotation was successfully set.
 	 */
@@ -782,7 +782,7 @@ public:
 	 * @param bTeleport			Whether we teleport the physics state (if physics collision is enabled for this object).
 	 *							If true, physics velocity for this object is unchanged (so ragdoll parts are not affected by change in location).
 	 *							If false, physics velocity is updated based on the change in position (affecting ragdoll parts).
-	 *							If CCD is on and not teleporting, this will affect objects along the entire sweep volume.
+	 *							If CCD is on and not teleporting, this will affect objects along the entire swept volume.
 	 * @param SweepHitResult	The hit result from the move if swept.
 	 */
 	UFUNCTION(BlueprintCallable, Category="Utilities|Transformation", meta=(DisplayName="AddActorWorldOffset", Keywords="location position"))
@@ -797,7 +797,7 @@ public:
 	 * @param Teleport			Whether we teleport the physics state (if physics collision is enabled for this object).
 	 *							If TeleportPhysics, physics velocity for this object is unchanged (so ragdoll parts are not affected by change in location).
 	 *							If None, physics velocity is updated based on the change in position (affecting ragdoll parts).
-	 *							If CCD is on and not teleporting, this will affect objects along the entire sweep volume.
+	 *							If CCD is on and not teleporting, this will affect objects along the entire swept volume.
 	 * @param SweepHitResult	The hit result from the move if swept.
 	 */
 	void AddActorWorldOffset(FVector DeltaLocation, bool bSweep=false, FHitResult* OutSweepHitResult=nullptr, ETeleportType Teleport = ETeleportType::None);
@@ -811,7 +811,7 @@ public:
 	 * @param bTeleport			Whether we teleport the physics state (if physics collision is enabled for this object).
 	 *							If true, physics velocity for this object is unchanged (so ragdoll parts are not affected by change in location).
 	 *							If false, physics velocity is updated based on the change in position (affecting ragdoll parts).
-	 *							If CCD is on and not teleporting, this will affect objects along the entire sweep volume.
+	 *							If CCD is on and not teleporting, this will affect objects along the entire swept volume.
 	 * @param SweepHitResult	The hit result from the move if swept.
 	 */
 	UFUNCTION(BlueprintCallable, Category="Utilities|Transformation", meta=(DisplayName="AddActorWorldRotation", AdvancedDisplay="bSweep,SweepHitResult,bTeleport"))
@@ -834,7 +834,7 @@ public:
 	 * @param bTeleport			Whether we teleport the physics state (if physics collision is enabled for this object).
 	 *							If true, physics velocity for this object is unchanged (so ragdoll parts are not affected by change in location).
 	 *							If false, physics velocity is updated based on the change in position (affecting ragdoll parts).
-	 *							If CCD is on and not teleporting, this will affect objects along the entire sweep volume.
+	 *							If CCD is on and not teleporting, this will affect objects along the entire swept volume.
 	 */
 	UFUNCTION(BlueprintCallable, Category="Utilities|Transformation", meta=(DisplayName="SetActorTransform"))
 	bool K2_SetActorTransform(const FTransform& NewTransform, bool bSweep, FHitResult& SweepHitResult, bool bTeleport);
@@ -849,7 +849,7 @@ public:
 	 * @param bTeleport			Whether we teleport the physics state (if physics collision is enabled for this object).
 	 *							If true, physics velocity for this object is unchanged (so ragdoll parts are not affected by change in location).
 	 *							If false, physics velocity is updated based on the change in position (affecting ragdoll parts).
-	 *							If CCD is on and not teleporting, this will affect objects along the entire sweep volume.
+	 *							If CCD is on and not teleporting, this will affect objects along the entire swept volume.
 	 */
 	UFUNCTION(BlueprintCallable, Category="Utilities|Transformation", meta=(DisplayName="AddActorLocalOffset", Keywords="location position"))
 	void K2_AddActorLocalOffset(FVector DeltaLocation, bool bSweep, FHitResult& SweepHitResult, bool bTeleport);
@@ -864,7 +864,7 @@ public:
 	 * @param bTeleport			Whether we teleport the physics state (if physics collision is enabled for this object).
 	 *							If true, physics velocity for this object is unchanged (so ragdoll parts are not affected by change in location).
 	 *							If false, physics velocity is updated based on the change in position (affecting ragdoll parts).
-	 *							If CCD is on and not teleporting, this will affect objects along the entire sweep volume.
+	 *							If CCD is on and not teleporting, this will affect objects along the entire swept volume.
 	 */
 	UFUNCTION(BlueprintCallable, Category="Utilities|Transformation", meta=(DisplayName="AddActorLocalRotation", AdvancedDisplay="bSweep,SweepHitResult,bTeleport"))
 	void K2_AddActorLocalRotation(FRotator DeltaRotation, bool bSweep, FHitResult& SweepHitResult, bool bTeleport);
@@ -880,7 +880,7 @@ public:
 	 * @param bTeleport			Whether we teleport the physics state (if physics collision is enabled for this object).
 	 *							If true, physics velocity for this object is unchanged (so ragdoll parts are not affected by change in location).
 	 *							If false, physics velocity is updated based on the change in position (affecting ragdoll parts).
-	 *							If CCD is on and not teleporting, this will affect objects along the entire sweep volume.
+	 *							If CCD is on and not teleporting, this will affect objects along the entire swept volume.
 	 */
 	UFUNCTION(BlueprintCallable, Category="Utilities|Transformation", meta=(DisplayName="AddActorLocalTransform"))
 	void K2_AddActorLocalTransform(const FTransform& NewTransform, bool bSweep, FHitResult& SweepHitResult, bool bTeleport);
@@ -895,7 +895,7 @@ public:
 	 * @param bTeleport				Whether we teleport the physics state (if physics collision is enabled for this object).
 	 *								If true, physics velocity for this object is unchanged (so ragdoll parts are not affected by change in location).
 	 *								If false, physics velocity is updated based on the change in position (affecting ragdoll parts).
-	 *								If CCD is on and not teleporting, this will affect objects along the entire sweep volume.
+	 *								If CCD is on and not teleporting, this will affect objects along the entire swept volume.
 	 */
 	UFUNCTION(BlueprintCallable, Category="Utilities|Transformation", meta=(DisplayName="SetActorRelativeLocation"))
 	void K2_SetActorRelativeLocation(FVector NewRelativeLocation, bool bSweep, FHitResult& SweepHitResult, bool bTeleport);
@@ -909,7 +909,7 @@ public:
 	 * @param bTeleport				Whether we teleport the physics state (if physics collision is enabled for this object).
 	 *								If true, physics velocity for this object is unchanged (so ragdoll parts are not affected by change in location).
 	 *								If false, physics velocity is updated based on the change in position (affecting ragdoll parts).
-	 *								If CCD is on and not teleporting, this will affect objects along the entire sweep volume.
+	 *								If CCD is on and not teleporting, this will affect objects along the entire swept volume.
 	 */
 	UFUNCTION(BlueprintCallable, Category="Utilities|Transformation", meta=(DisplayName="SetActorRelativeRotation", AdvancedDisplay="bSweep,SweepHitResult,bTeleport"))
 	void K2_SetActorRelativeRotation(FRotator NewRelativeRotation, bool bSweep, FHitResult& SweepHitResult, bool bTeleport);
@@ -924,7 +924,7 @@ public:
 	 * @param bTeleport			Whether we teleport the physics state (if physics collision is enabled for this object).
 	 *							If true, physics velocity for this object is unchanged (so ragdoll parts are not affected by change in location).
 	 *							If false, physics velocity is updated based on the change in position (affecting ragdoll parts).
-	 *							If CCD is on and not teleporting, this will affect objects along the entire sweep volume.
+	 *							If CCD is on and not teleporting, this will affect objects along the entire swept volume.
 	 */
 	UFUNCTION(BlueprintCallable, Category="Utilities|Transformation", meta=(DisplayName="SetActorRelativeTransform"))
 	void K2_SetActorRelativeTransform(const FTransform& NewRelativeTransform, bool bSweep, FHitResult& SweepHitResult, bool bTeleport);
