#ue4 - Interpolate simulated position faster when proxy is coming to a stop.

diff --git a/Engine/Source/Runtime/Engine/Private/Components/CharacterMovementComponent.cpp b/Engine/Source/Runtime/Engine/Private/Components/CharacterMovementComponent.cpp
--- a/Engine/Source/Runtime/Engine/Private/Components/CharacterMovementComponent.cpp
+++ b/Engine/Source/Runtime/Engine/Private/Components/CharacterMovementComponent.cpp
@@ -5767,11 +5767,13 @@ void UCharacterMovementComponent::SmoothClientPosition(float DeltaSeconds)
 		}
 		else
 		{
-			// smooth interpolation of mesh translation to avoid popping of other client pawns unless under a low tick rate
-			if (DeltaSeconds < ClientData->SmoothNetUpdateTime)
+			// Smooth interpolation of mesh translation to avoid popping of other client pawns unless under a low tick rate.
+			// Faster interpolation if stopped.
+			const float SmoothLocationTime = Velocity.IsZero() ? 0.5f*ClientData->SmoothNetUpdateTime : ClientData->SmoothNetUpdateTime;
+			if (DeltaSeconds < SmoothLocationTime)
 			{
 				// Slowly decay translation offset
-				ClientData->MeshTranslationOffset = (ClientData->MeshTranslationOffset * (1.f - DeltaSeconds / ClientData->SmoothNetUpdateTime));
+				ClientData->MeshTranslationOffset = (ClientData->MeshTranslationOffset * (1.f - DeltaSeconds / SmoothLocationTime));
 			}
 			else
 			{
