Removed unnecessary inverse transformations in USceneComponent::ApplyWorldOffset

diff --git a/Engine/Source/Runtime/Engine/Private/Components/SceneComponent.cpp b/Engine/Source/Runtime/Engine/Private/Components/SceneComponent.cpp
--- a/Engine/Source/Runtime/Engine/Private/Components/SceneComponent.cpp
+++ b/Engine/Source/Runtime/Engine/Private/Components/SceneComponent.cpp
@@ -2161,22 +2161,14 @@ void USceneComponent::ApplyWorldOffset(const FVector& InOffset, bool bWorldShift
 	Bounds.Origin+= InOffset;
 
 	// Update component location
-	FVector NewLocation = GetComponentLocation() + InOffset;
-	// If attached to something, transform into local space
-	if (AttachParent != NULL)
+	if (AttachParent == nullptr || bAbsoluteLocation)
 	{
-		FTransform const ParentToWorld = AttachParent->GetSocketTransform(AttachSocketName);
-		if (!bAbsoluteLocation)
-		{
-			NewLocation = ParentToWorld.InverseTransformPosition(GetComponentLocation());
-		}
-	}
-	
-	RelativeLocation = NewLocation;
+		RelativeLocation = GetComponentLocation() + InOffset;
 		
-	// Calculate the new ComponentToWorld transform
-	const FTransform RelativeTransform(RelativeRotationCache.RotatorToQuat(RelativeRotation), RelativeLocation, RelativeScale3D);
-	ComponentToWorld = CalcNewComponentToWorld(RelativeTransform);
+		// Calculate the new ComponentToWorld transform
+		const FTransform RelativeTransform(RelativeRotationCache.RotatorToQuat(RelativeRotation), RelativeLocation, RelativeScale3D);
+		ComponentToWorld = CalcNewComponentToWorld(RelativeTransform);
+	}
 
 	// Physics move is skipped if physics state is not created or physics scene supports origin shifting
 	// We still need to send transform to physics scene to "transform back" actors which should ignore origin shifting
