Sanitize BodyInstance array in ISMC instead of checking for null inside BodyInstance.cpp

diff --git a/Engine/Source/Runtime/Engine/Private/InstancedStaticMesh.cpp b/Engine/Source/Runtime/Engine/Private/InstancedStaticMesh.cpp
--- a/Engine/Source/Runtime/Engine/Private/InstancedStaticMesh.cpp
+++ b/Engine/Source/Runtime/Engine/Private/InstancedStaticMesh.cpp
@@ -984,8 +984,12 @@ void UInstancedStaticMeshComponent::CreateAllInstanceBodies()
 	    const int32 NumBodies = PerInstanceSMData.Num();
 		check(InstanceBodies.Num() == 0);
 		InstanceBodies.SetNumUninitialized(NumBodies);
-    
-	    TArray<FTransform> Transforms;
+
+		// Sanitized array does not contain any nulls
+		TArray<FBodyInstance*> InstanceBodiesSanitized;
+		InstanceBodiesSanitized.Reserve(NumBodies);
+
+		TArray<FTransform> Transforms;
 	    Transforms.Reserve(NumBodies);
     
 	    for (int32 i = 0; i < NumBodies; ++i)
@@ -997,8 +1001,10 @@ void UInstancedStaticMeshComponent::CreateAllInstanceBodies()
 			}
 			else
 			{
-				InstanceBodies[i] = new FBodyInstance;
-				FBodyInstance* Instance = InstanceBodies[i];
+				FBodyInstance* Instance = new FBodyInstance;
+
+				InstanceBodiesSanitized.Add(Instance);
+				InstanceBodies[i] = Instance;
 				Instance->CopyBodyInstancePropertiesFrom(&BodyInstance);
 				Instance->InstanceBodyIndex = i; // Set body index 
 				Instance->bAutoWeld = false;
@@ -1025,7 +1031,7 @@ void UInstancedStaticMeshComponent::CreateAllInstanceBodies()
 			}
 	    }
 
-		if (NumBodies > 0 && Mobility == EComponentMobility::Static)
+		if (InstanceBodiesSanitized.Num() > 0 && Mobility == EComponentMobility::Static)
 		{
 			TArray<UBodySetup*> BodySetups;
 			TArray<UPhysicalMaterial*> PhysicalMaterials;
@@ -1036,10 +1042,10 @@ void UInstancedStaticMeshComponent::CreateAllInstanceBodies()
 			PhysicalMaterials.Add(FBodyInstance::GetSimplePhysicalMaterial(&BodyInstance, WeakSelfPtr, TWeakObjectPtr<UBodySetup>(BodySetup)));
 
 			PhysicsSerializer->CreatePhysicsData(BodySetups, PhysicalMaterials);
-			FBodyInstance::InitStaticBodies(InstanceBodies, Transforms, BodySetup, this, GetWorld()->GetPhysicsScene(), PhysicsSerializer);
+			FBodyInstance::InitStaticBodies(InstanceBodiesSanitized, Transforms, BodySetup, this, GetWorld()->GetPhysicsScene(), PhysicsSerializer);
 
-			//Serialize physics data for fast path cooking
-			PhysicsSerializer->SerializePhysics(InstanceBodies, BodySetups, PhysicalMaterials);
+			// Serialize physics data for fast path cooking
+			PhysicsSerializer->SerializePhysics(InstanceBodiesSanitized, BodySetups, PhysicalMaterials);
 		}
 	}
 }
diff --git a/Engine/Source/Runtime/Engine/Private/PhysicsEngine/BodyInstance.cpp b/Engine/Source/Runtime/Engine/Private/PhysicsEngine/BodyInstance.cpp
--- a/Engine/Source/Runtime/Engine/Private/PhysicsEngine/BodyInstance.cpp
+++ b/Engine/Source/Runtime/Engine/Private/PhysicsEngine/BodyInstance.cpp
@@ -1158,10 +1158,6 @@ struct FInitBodiesHelper
 		for (int32 BodyIdx = 0; BodyIdx < NumBodies; ++BodyIdx)
 		{
 			FBodyInstance* Instance = Bodies[BodyIdx];
-			if (!Instance)
-			{
-				continue;
-			}
 			const FTransform& Transform = Transforms[BodyIdx];
 
 			FBodyInstance::ValidateTransform(Transform, DebugName, BodySetup);
@@ -1360,10 +1356,7 @@ struct FInitBodiesHelper
 			for (int32 BodyIdx = 0, NumBodies = Bodies.Num(); BodyIdx < NumBodies; ++BodyIdx)
 			{
 				FBodyInstance* Instance = Bodies[BodyIdx];
-				if (Instance)
-				{
-					Instance->InitDynamicProperties_AssumesLocked();
-				}
+				Instance->InitDynamicProperties_AssumesLocked();
 			}
 		}
 	}
@@ -1412,11 +1405,6 @@ struct FInitBodiesHelper
 				for (int32 BodyIdx = 0; BodyIdx < NumBodies; ++BodyIdx)
 				{
 					FBodyInstance* Instance = Bodies[BodyIdx];
-					if (!Instance)
-					{
-						continue;
-					}
-
 					const FTransform& Transform = Transforms[BodyIdx];
 
 					const b2Vec2 Scale2D = FPhysicsIntegration2D::ConvertUnrealVectorToBox(Instance->Scale3D);
@@ -4072,13 +4060,10 @@ bool FBodyInstance::OverlapTestForBodiesImpl(const FVector& Pos, const FQuat& Ro
 			{
 				for (const FBodyInstance* BodyInstance : Bodies)
 				{
-					if (BodyInstance)
+					bHaveOverlap = BodyInstance->OverlapPhysX_AssumesLocked(*PGeom, PShapeGlobalPose);
+					if (bHaveOverlap)
 					{
-						bHaveOverlap = BodyInstance->OverlapPhysX_AssumesLocked(*PGeom, PShapeGlobalPose);
-						if (bHaveOverlap)
-						{
-							return;
-						}
+						return;
 					}
 				}
 			}
