Bulk foliage instance removal Blueprint function
Fixed up handling of PhysX bodies during instance removal

diff --git a/Engine/Source/Runtime/Engine/Classes/Components/HierarchicalInstancedStaticMeshComponent.h b/Engine/Source/Runtime/Engine/Classes/Components/HierarchicalInstancedStaticMeshComponent.h
--- a/Engine/Source/Runtime/Engine/Classes/Components/HierarchicalInstancedStaticMeshComponent.h
+++ b/Engine/Source/Runtime/Engine/Classes/Components/HierarchicalInstancedStaticMeshComponent.h
@@ -132,6 +132,10 @@ public:
 	virtual void ClearInstances() override;
 	virtual TArray<int32> GetInstancesOverlappingSphere(const FVector& Center, float Radius, bool bSphereInWorldSpace = true) const override;
 
+	/** Removes all the instances with indices specified in the InstancesToRemove array. Returns true on success. */
+	UFUNCTION(BlueprintCallable, Category = "Components|InstancedStaticMesh")
+	bool RemoveInstances(const TArray<int32>& InstancesToRemove);
+
 	/** Get the number of instances that overlap a given sphere */
 	int32 GetOverlappingSphereCount(const FSphere& Sphere) const;
 	/** Get the number of instances that overlap a given box */
@@ -163,6 +167,9 @@ public:
 	int32 DesiredInstancesPerLeaf();
 
 protected:
+	/** Removes a single instance without extra work such as rebuilding the tree or marking render state dirty. */
+	void RemoveInstanceInternal(int32 InstanceIndex);
+
 	/** Gets and approximate number of verts for each LOD to generate heuristics **/
 	int32 GetVertsForLOD(int32 LODIndex);
 	/** Average number of instances per leaf **/
diff --git a/Engine/Source/Runtime/Engine/Private/HierarchicalInstancedStaticMesh.cpp b/Engine/Source/Runtime/Engine/Private/HierarchicalInstancedStaticMesh.cpp
--- a/Engine/Source/Runtime/Engine/Private/HierarchicalInstancedStaticMesh.cpp
+++ b/Engine/Source/Runtime/Engine/Private/HierarchicalInstancedStaticMesh.cpp
@@ -1597,18 +1597,13 @@ SIZE_T UHierarchicalInstancedStaticMeshComponent::GetResourceSize( EResourceSize
 	return ResSize;
 }
 
-bool UHierarchicalInstancedStaticMeshComponent::RemoveInstance(int32 InstanceIndex)
+void UHierarchicalInstancedStaticMeshComponent::RemoveInstanceInternal(int32 InstanceIndex)
 {
-	if (!PerInstanceSMData.IsValidIndex(InstanceIndex))
-	{
-		return false;
-	}
-
 	PartialNavigationUpdate(InstanceIndex);
 
 	// Save the render index
 	RemovedInstances.Add(InstanceReorderTable[InstanceIndex]);
-	
+
 	// Remove the instance
 	PerInstanceSMData.RemoveAtSwap(InstanceIndex);
 	InstanceReorderTable.RemoveAtSwap(InstanceIndex);
@@ -1619,31 +1614,83 @@ bool UHierarchicalInstancedStaticMeshComponent::RemoveInstance(int32 InstanceInd
 	}
 #endif
 
-	if (IsAsyncBuilding())
+	// update the physics state
+	if (bPhysicsStateCreated)
 	{
-		// invalidate the results of the current async build as it's too slow to fix up deletes
-		bConcurrentRemoval = true;
+		// Clean up physics for removed instance
+		if (InstanceBodies[InstanceIndex])
+		{
+			InstanceBodies[InstanceIndex]->TermBody();
+			delete InstanceBodies[InstanceIndex];
+		}
+
+		int32 LastInstanceIndex = PerInstanceSMData.Num();
+
+		if (InstanceIndex == LastInstanceIndex)
+		{
+			// If we removed the last instance in the array we just need to remove it from the InstanceBodies array too.
+			InstanceBodies.RemoveAt(InstanceIndex);
+		}
+		else
+		{
+			if (InstanceBodies[LastInstanceIndex])
+			{
+				// term physics for swapped instance
+				InstanceBodies[LastInstanceIndex]->TermBody();
+			}
+
+			// swap in the last instance body if we have one
+			InstanceBodies.RemoveAtSwap(InstanceIndex);
+
+			// recreate physics for the instance we swapped in the removed item's place
+			if (InstanceBodies[InstanceIndex])
+			{
+				InitInstanceBody(InstanceIndex, InstanceBodies[InstanceIndex]);
+			}
+		}
 	}
-	else
+}
+
+bool UHierarchicalInstancedStaticMeshComponent::RemoveInstances(const TArray<int32>& InstancesToRemove)
+{
+	TArray<int32> SortedInstancesToRemove = InstancesToRemove;
+
+	// Sort so RemoveAtSwaps don't alter the indices of items still to remove
+	SortedInstancesToRemove.Sort(TGreater<int32>());
+
+	if (!PerInstanceSMData.IsValidIndex(SortedInstancesToRemove[0]) || !PerInstanceSMData.IsValidIndex(SortedInstancesToRemove.Last()))
 	{
-		BuildTreeAsync();
+		return false;
 	}
 
+	for (int32 Index : SortedInstancesToRemove)
+	{
+		RemoveInstanceInternal(Index);
+	}
 
-	// update the physics state
-	if (bPhysicsStateCreated)
+	ReleasePerInstanceRenderData();
+	MarkRenderStateDirty();
+
+	return true;
+}
+
+bool UHierarchicalInstancedStaticMeshComponent::RemoveInstance(int32 InstanceIndex)
+{
+	if (!PerInstanceSMData.IsValidIndex(InstanceIndex))
 	{
-		int32 OldLastIndex = PerInstanceSMData.Num();
+		return false;
+	}
 
-		// always shut down physics for the last instance
-		InstanceBodies[OldLastIndex]->TermBody();
+	RemoveInstanceInternal(InstanceIndex);
 
-		// recreate physics for the instance we swapped in the removed item's place
-		if (InstanceIndex != PerInstanceSMData.Num() && PerInstanceSMData.Num())
-		{
-			InstanceBodies[InstanceIndex]->TermBody();
-			InitInstanceBody(InstanceIndex, InstanceBodies[InstanceIndex]);
-		}
+	if (IsAsyncBuilding())
+	{
+		// invalidate the results of the current async build as it's too slow to fix up deletes
+		bConcurrentRemoval = true;
+	}
+	else
+	{
+		BuildTreeAsync();
 	}
 
 	ReleasePerInstanceRenderData();
diff --git a/Engine/Source/Runtime/Foliage/Private/InstancedFoliage.cpp b/Engine/Source/Runtime/Foliage/Private/InstancedFoliage.cpp
--- a/Engine/Source/Runtime/Foliage/Private/InstancedFoliage.cpp
+++ b/Engine/Source/Runtime/Foliage/Private/InstancedFoliage.cpp
@@ -2710,17 +2710,17 @@ void UFoliageInstancedStaticMeshComponent::ReceiveComponentDamage(float DamageAm
 				FVector LocalOrigin = GetComponentToWorld().Inverse().TransformPosition(RadialDamageEvent->Origin);
 				float Scale = GetComponentScale().X; // assume component (not instances) is uniformly scaled
 				
-				TArray<FFoliageInstanceDamage> DamageInstances;
-				DamageInstances.Empty(Instances.Num());
+				TArray<float> Damages;
+				Damages.Empty(Instances.Num());
 				
 				for (int32 InstanceIndex : Instances)
 				{
 					// Find distance in local space and then scale; quicker than transforming each instance to world space.
 					float DistanceFromOrigin = (PerInstanceSMData[InstanceIndex].Transform.GetOrigin() - LocalOrigin).Size() * Scale;
-					new(DamageInstances) FFoliageInstanceDamage(InstanceIndex, RadialDamageEvent->Params.GetDamageScale(DistanceFromOrigin));
+					Damages.Add(RadialDamageEvent->Params.GetDamageScale(DistanceFromOrigin));
 				}
 
-				OnInstanceTakeRadialDamage.Broadcast(DamageInstances, EventInstigator, RadialDamageEvent->Origin, MaxRadius, DamageTypeCDO, DamageCauser);
+				OnInstanceTakeRadialDamage.Broadcast(Instances, Damages, EventInstigator, RadialDamageEvent->Origin, MaxRadius, DamageTypeCDO, DamageCauser);
 			}
 		}
 	}
diff --git a/Engine/Source/Runtime/Foliage/Public/FoliageInstancedStaticMeshComponent.h b/Engine/Source/Runtime/Foliage/Public/FoliageInstancedStaticMeshComponent.h
--- a/Engine/Source/Runtime/Foliage/Public/FoliageInstancedStaticMeshComponent.h
+++ b/Engine/Source/Runtime/Foliage/Public/FoliageInstancedStaticMeshComponent.h
@@ -5,32 +5,10 @@
 #include "Components/HierarchicalInstancedStaticMeshComponent.h"
 #include "FoliageInstancedStaticMeshComponent.generated.h"
 
-USTRUCT(BlueprintType)
-struct FFoliageInstanceDamage
-{
-	GENERATED_USTRUCT_BODY()
-
-	FFoliageInstanceDamage()
-	: InstanceIndex(0)
-	, DamageAmount(0.f)
-	{}
-
-	FFoliageInstanceDamage(int32 InInstanceIndex, float InDamageAmount)
-	:	InstanceIndex(InInstanceIndex)
-	,	DamageAmount(InDamageAmount)
-	{}
-
-	UPROPERTY(EditAnywhere, BlueprintReadOnly, Category = FoliageInstanceDamage)
-	int32 InstanceIndex;
-
-	UPROPERTY(EditAnywhere, BlueprintReadOnly, Category = FoliageInstanceDamage)
-	float DamageAmount;
-};
-
 /**
 */
 DECLARE_DYNAMIC_MULTICAST_DELEGATE_SevenParams(FInstancePointDamageSignature, int32, InstanceIndex, float, Damage, class AController*, InstigatedBy, FVector, HitLocation, FVector, ShotFromDirection, const class UDamageType*, DamageType, AActor*, DamageCauser);
-DECLARE_DYNAMIC_MULTICAST_DELEGATE_SixParams(FInstanceRadialDamageSignature, const TArray<FFoliageInstanceDamage>&, DamageInstances, class AController*, InstigatedBy, FVector, Origin, float, MaxRadius, const class UDamageType*, DamageType, AActor*, DamageCauser);
+DECLARE_DYNAMIC_MULTICAST_DELEGATE_SevenParams(FInstanceRadialDamageSignature, const TArray<int32>&, Instances, const TArray<float>&, Damages, class AController*, InstigatedBy, FVector, Origin, float, MaxRadius, const class UDamageType*, DamageType, AActor*, DamageCauser);
 
 UCLASS(ClassGroup = Foliage, Blueprintable)
 class FOLIAGE_API UFoliageInstancedStaticMeshComponent : public UHierarchicalInstancedStaticMeshComponent
