Fix crash when calling RemoveInstances with 0 instances

diff --git a/Engine/Source/Runtime/Engine/Private/HierarchicalInstancedStaticMesh.cpp b/Engine/Source/Runtime/Engine/Private/HierarchicalInstancedStaticMesh.cpp
--- a/Engine/Source/Runtime/Engine/Private/HierarchicalInstancedStaticMesh.cpp
+++ b/Engine/Source/Runtime/Engine/Private/HierarchicalInstancedStaticMesh.cpp
@@ -1653,6 +1653,11 @@ void UHierarchicalInstancedStaticMeshComponent::RemoveInstanceInternal(int32 Ins
 
 bool UHierarchicalInstancedStaticMeshComponent::RemoveInstances(const TArray<int32>& InstancesToRemove)
 {
+	if (InstancesToRemove.Num() == 0)
+	{
+		return true;
+	}
+
 	TArray<int32> SortedInstancesToRemove = InstancesToRemove;
 
 	// Sort so RemoveAtSwaps don't alter the indices of items still to remove
