Build outage fixed, two EDITOR_ONLY defines missing

diff --git a/Engine/Source/Runtime/Engine/Classes/Engine/LODActor.h b/Engine/Source/Runtime/Engine/Classes/Engine/LODActor.h
--- a/Engine/Source/Runtime/Engine/Classes/Engine/LODActor.h
+++ b/Engine/Source/Runtime/Engine/Classes/Engine/LODActor.h
@@ -77,6 +77,8 @@ public:
 	/** Returns StaticMeshComponent subobject **/
 	class UStaticMeshComponent* GetStaticMeshComponent() const { return StaticMeshComponent; }
 
+#if WITH_EDITORONLY_DATA
 	/** Returns StaticMeshComponent subobject **/
 	class UDrawSphereComponent* GetDrawSphereComponent() const { return DrawSphereComponent; }
+#endif // WITH_EDITORONLY_DATA
 }; 
diff --git a/Engine/Source/Runtime/Engine/Private/World.cpp b/Engine/Source/Runtime/Engine/Private/World.cpp
--- a/Engine/Source/Runtime/Engine/Private/World.cpp
+++ b/Engine/Source/Runtime/Engine/Private/World.cpp
@@ -517,10 +517,12 @@ void UWorld::FinishDestroy()
 		delete TimerManager;
 	}
 
+#if WITH_EDITOR
 	if (HierarchicalLODBuilder)
 	{
 		delete HierarchicalLODBuilder;
 	}
+#endif // WITH_EDITOR
 
 	// Remove the PKG_ContainsMap flag from packages that no longer contain a world
 	{
