ADDED UE-9189 Implement some polar co-ord funcs for BPs
FIXED typo in comment

diff --git a/Engine/Source/Runtime/Core/Public/Math/UnrealMathUtility.h b/Engine/Source/Runtime/Core/Public/Math/UnrealMathUtility.h
--- a/Engine/Source/Runtime/Core/Public/Math/UnrealMathUtility.h
+++ b/Engine/Source/Runtime/Core/Public/Math/UnrealMathUtility.h
@@ -138,7 +138,7 @@ struct FMath : public FPlatformMath
 	 * Given a direction vector and a surface normal, returns the vector reflected across the surface normal.
 	 * Produces a result like shining a laser at a mirror!
 	 *
-	 * @param Direction Direction vector the ray is comming from.
+	 * @param Direction Direction vector the ray is coming from.
 	 * @param SurfaceNormal A normal of the surface the ray should be reflected on.
 	 *
 	 * @returns Reflected vector.
@@ -151,7 +151,7 @@ struct FMath : public FPlatformMath
 	template< class U > 
 	static FORCEINLINE bool IsWithin(const U& TestValue, const U& MinValue, const U& MaxValue)
 	{
-		return ((TestValue>=MinValue) && (TestValue < MaxValue));
+		return ((TestValue >= MinValue) && (TestValue < MaxValue));
 	}
 
 	/** Checks if value is within a range, inclusive on MaxValue) */
diff --git a/Engine/Source/Runtime/Engine/Classes/Kismet/KismetMathLibrary.h b/Engine/Source/Runtime/Engine/Classes/Kismet/KismetMathLibrary.h
--- a/Engine/Source/Runtime/Engine/Classes/Kismet/KismetMathLibrary.h
+++ b/Engine/Source/Runtime/Engine/Classes/Kismet/KismetMathLibrary.h
@@ -683,6 +683,15 @@ class ENGINE_API UKismetMathLibrary : public UBlueprintFunctionLibrary
 	UFUNCTION(BlueprintPure, Category = "Math|Random")
 	static FVector RandomUnitVectorInCone(FVector ConeDir, float ConeHalfAngle);
 
+	/**
+	* RandomUnitVectorWithYawAndPitch
+	*
+	* @param MaxYaw - The Yaw-angle of the cone (from ConeDir to horizontal-edge), in degrees.
+	* @param MaxPitch - The Pitch-angle of the cone (from ConeDir to vertical-edge), in degrees.	
+	*/
+	UFUNCTION(BlueprintPure, Category = "Math|Random", meta = (Keywords = "RandomVector"))
+	static FVector RandomUnitVectorInConeWithYawAndPitch(FVector ConeDir, float MaxYawInDegrees, float MaxPitchInDegrees);
+
 	// Mirrors a vector by a normal
 	UFUNCTION(BlueprintPure, Category="Math|Vector")
 	static FVector MirrorVectorByNormal(FVector InVect, FVector InNormal);
@@ -1214,6 +1223,14 @@ class ENGINE_API UKismetMathLibrary : public UBlueprintFunctionLibrary
 	UFUNCTION(BlueprintPure, Category="Math|Vector", meta=(Keywords="rotation rotate"))
 	static FVector GetUpVector(FRotator InRot);
 
+	/** Creates a directional vector from rotation values {Pitch, Yaw} supplied in degrees with specified Length*/	
+	UFUNCTION(BlueprintPure, Category = "Math|Vector", meta = (Keywords = "rotation rotate"))
+	static FVector CreateVectorFromYawPitch(float Yaw, float Pitch, float Length = 1.0f );
+
+	/** Breaks a vector apart into Yaw, Pitch rotation values given in degrees. (non-clamped) */
+	UFUNCTION(BlueprintPure, Category = "Math|Vector2D", meta = (NativeBreakFunc))
+	static void GetYawPitchFromVector(FVector InVec, float& Yaw, float& Pitch);
+
 	/** Makes a rotator {Roll, Pitch, Yaw} from rotation values supplied in degrees */
 	UFUNCTION(BlueprintPure, Category="Math|Rotator", meta=(Keywords="construct build rotation rotate rotator makerotator", NativeMakeFunc))
 	static FRotator MakeRotator(
diff --git a/Engine/Source/Runtime/Engine/Private/KismetMathLibrary.cpp b/Engine/Source/Runtime/Engine/Private/KismetMathLibrary.cpp
--- a/Engine/Source/Runtime/Engine/Private/KismetMathLibrary.cpp
+++ b/Engine/Source/Runtime/Engine/Private/KismetMathLibrary.cpp
@@ -886,6 +886,11 @@ FVector UKismetMathLibrary::RandomUnitVectorInCone(FVector ConeDir, float ConeHa
 	return FMath::VRandCone(ConeDir, ConeHalfAngle);
 }
 
+FVector UKismetMathLibrary::RandomUnitVectorInConeWithYawAndPitch(FVector ConeDir, float MaxYawInDegrees, float MaxPitchInDegrees)
+{
+	return FMath::VRandCone(ConeDir, DegreesToRadians(MaxYawInDegrees), DegreesToRadians(MaxPitchInDegrees));
+}
+
 FVector UKismetMathLibrary::MirrorVectorByNormal(FVector A, FVector B)
 {
 	return FMath::GetReflectionVector(A, B);
@@ -1804,6 +1809,27 @@ FVector UKismetMathLibrary::GetUpVector(FRotator InRot)
 	return FRotationMatrix(InRot).GetScaledAxis( EAxis::Z );
 }
 
+FVector UKismetMathLibrary::CreateVectorFromYawPitch(float Yaw, float Pitch, float Length /*= 1.0f */)
+{
+	// FRotator::Vector() behaviour 
+	float CP, SP, CY, SY;
+	FMath::SinCos(&SP, &CP, FMath::DegreesToRadians(Pitch));
+	FMath::SinCos(&SY, &CY, FMath::DegreesToRadians(Yaw));
+	FVector V = FVector(CP*CY, CP*SY, SP) * Length;
+		
+	return V;
+}
+
+void UKismetMathLibrary::GetYawPitchFromVector(FVector InVec, float& Yaw, float& Pitch)
+{
+	FVector NormalizedVector = InVec.GetSafeNormal();
+	// Find yaw.
+	Yaw = FMath::Atan2(NormalizedVector.Y, NormalizedVector.X) * 180.f / PI;
+
+	// Find pitch.
+	Pitch = FMath::Atan2(NormalizedVector.Z, FMath::Sqrt(NormalizedVector.X*NormalizedVector.X + NormalizedVector.Y*NormalizedVector.Y)) * 180.f / PI;	
+}
+
 FRotator UKismetMathLibrary::MakeRotator(float Roll, float Pitch, float Yaw)
 {
 	return FRotator(Pitch,Yaw,Roll);
