PR #1054: [New Info for every Trace! FHitResult::Distance added from PhysX Data] Now we can easily get the Distance between the start of the trace and the impact in FHitResult itself, returned from PhysX!  Rama (Contributed by [REDACTED_PERSON])
Update ADebugCameraHUD to use new info

diff --git a/Engine/Source/Runtime/Engine/Classes/Engine/EngineTypes.h b/Engine/Source/Runtime/Engine/Classes/Engine/EngineTypes.h
--- a/Engine/Source/Runtime/Engine/Classes/Engine/EngineTypes.h
+++ b/Engine/Source/Runtime/Engine/Classes/Engine/EngineTypes.h
@@ -1572,6 +1572,10 @@ struct ENGINE_API FHitResult
 	 */
 	UPROPERTY()
 	float Time;
+	 
+	/** The distance from the TraceStart to the ImpactPoint in world space. This value is 0 if there was an initial overlap (trace started inside another colliding object). */
+	UPROPERTY()
+	float Distance; 
 	
 	/**
 	 * The location in world space where the moving shape would end up against the impacted object, if there is a hit. Equal to the point of impact for line tests.
@@ -1674,7 +1678,7 @@ struct ENGINE_API FHitResult
 
 	/** Ctor for easily creating "fake" hits from limited data. */
 	FHitResult(class AActor* InActor, class UPrimitiveComponent* InComponent, FVector const& HitLoc, FVector const& HitNorm);
-
+ 
 	void Reset(float InTime = 1.f, bool bPreserveTraceData = true)
 	{
 		const FVector SavedTraceStart = TraceStart;
diff --git a/Engine/Source/Runtime/Engine/Private/Collision/CollisionConversions.cpp b/Engine/Source/Runtime/Engine/Private/Collision/CollisionConversions.cpp
--- a/Engine/Source/Runtime/Engine/Private/Collision/CollisionConversions.cpp
+++ b/Engine/Source/Runtime/Engine/Private/Collision/CollisionConversions.cpp
@@ -425,6 +425,7 @@ void ConvertQueryImpactHit(const UWorld* World, const PxLocationHit& PHit, FHitR
 	// calculate the hit time
 	const float HitTime = PHit.distance/CheckLength;
 	OutResult.Time = HitTime;
+	OutResult.Distance = PHit.distance;
 
 	// figure out where the the "safe" location for this shape is by moving from the startLoc toward the ImpactPoint
 	const FVector TraceStartToEnd = EndLoc - StartLoc;
@@ -812,6 +813,7 @@ static bool ConvertOverlappedShapeToImpactHit(const UWorld* World, const PxLocat
 	// Time of zero because initially overlapping
 	OutResult.bStartPenetrating = true;
 	OutResult.Time = 0.f;
+	OutResult.Distance = 0.f;
 
 	// Return start location as 'safe location'
 	OutResult.Location = P2UVector(QueryTM.p);
diff --git a/Engine/Source/Runtime/Engine/Private/DebugCameraHUD.cpp b/Engine/Source/Runtime/Engine/Private/DebugCameraHUD.cpp
--- a/Engine/Source/Runtime/Engine/Private/DebugCameraHUD.cpp
+++ b/Engine/Source/Runtime/Engine/Private/DebugCameraHUD.cpp
@@ -117,7 +117,7 @@ void ADebugCameraHUD::PostRender()
 				yl += Y;
 				Canvas->DrawText(RenderFont, FString::Printf(TEXT("HitLoc:%s HitNorm:%s"), *Hit.Location.ToString(), *Hit.Normal.ToString() ), X, yl, 1.f, 1.f, FontRenderInfo);
 				yl += Y;
-				Canvas->DrawText(RenderFont, FString::Printf(TEXT("HitDist: %f"), (CamLoc - Hit.Location).Size()), X, yl, 1.f, 1.f, FontRenderInfo);
+				Canvas->DrawText(RenderFont, FString::Printf(TEXT("HitDist: %f"), Hit.Distance), X, yl, 1.f, 1.f, FontRenderInfo);
 				yl += Y;
 				Canvas->DrawText(RenderFont, FString::Printf(TEXT("HitActor: '%s'"), HitActor ? *HitActor->GetFName().ToString() : TEXT("<NULL>")), X, yl, 1.f, 1.f, FontRenderInfo);
 				yl += Y;
