UE-18412 Improve tooltip on Max Substep Delta

diff --git a/Engine/Source/Runtime/Engine/Classes/PhysicsEngine/PhysicsSettings.h b/Engine/Source/Runtime/Engine/Classes/PhysicsEngine/PhysicsSettings.h
--- a/Engine/Source/Runtime/Engine/Classes/PhysicsEngine/PhysicsSettings.h
+++ b/Engine/Source/Runtime/Engine/Classes/PhysicsEngine/PhysicsSettings.h
@@ -185,7 +185,7 @@ class ENGINE_API UPhysicsSettings : public UDeveloperSettings
 	UPROPERTY(config, EditAnywhere, Category = Framerate)
 	bool bSubsteppingAsync;
 
-	/** Max delta time for an individual substep simulation. */
+	/** Max delta time (in seconds) for an individual simulation substep. */
 	UPROPERTY(config, EditAnywhere, meta = (ClampMin = "0.0013", UIMin = "0.0013", ClampMax = "1.0", UIMax = "1.0", editcondition = "bSubStepping"), Category=Framerate)
 	float MaxSubstepDeltaTime;
 
