PR #1373: Fixed FBodyInstance::IsInstanceAwake (Contributed by [REDACTED_PERSON])

diff --git a/Engine/Source/Runtime/Engine/Private/PhysicsEngine/BodyInstance.cpp b/Engine/Source/Runtime/Engine/Private/PhysicsEngine/BodyInstance.cpp
--- a/Engine/Source/Runtime/Engine/Private/PhysicsEngine/BodyInstance.cpp
+++ b/Engine/Source/Runtime/Engine/Private/PhysicsEngine/BodyInstance.cpp
@@ -3231,21 +3231,22 @@ void FBodyInstance::UpdateDampingProperties()
 
 bool FBodyInstance::IsInstanceAwake() const
 {
+	bool bIsSleeping = false;
 #if WITH_PHYSX
 	ExecuteOnPxRigidDynamicReadOnly(this, [&](const PxRigidDynamic* PRigidDynamic)
 	{
-		return !PRigidDynamic->isSleeping();
+		bIsSleeping = PRigidDynamic->isSleeping();
 	});
 #endif
 
 #if WITH_BOX2D
 	if (BodyInstancePtr)
 	{
-		return BodyInstancePtr->IsAwake();
+		bIsSleeping = !BodyInstancePtr->IsAwake();
 	}
 #endif
 
-	return false;
+	return !bIsSleeping;
 }
 
 void FBodyInstance::WakeInstance()
