% fix.

diff --git a/Engine/Source/Runtime/Engine/Private/World.cpp b/Engine/Source/Runtime/Engine/Private/World.cpp
--- a/Engine/Source/Runtime/Engine/Private/World.cpp
+++ b/Engine/Source/Runtime/Engine/Private/World.cpp
@@ -5287,7 +5287,7 @@ void UWorld::ServerTravel(const FString& FURL, bool bAbsolute, bool bShouldSkipG
 {
 	if (FURL.Contains(TEXT("%")) )
 	{
-		UE_LOG(LogWorld, Log, TEXT("FURL %s Contains illegal character '%'."), *FURL);
+		UE_LOG(LogWorld, Log, TEXT("FURL %s Contains illegal character '%%'."), *FURL);
 		return;
 	}
 
