#vr Fix for IsHMDEnabled function library

diff --git a/Engine/Source/Runtime/Engine/Private/HeadMountedDisplayFunctionLibrary.cpp b/Engine/Source/Runtime/Engine/Private/HeadMountedDisplayFunctionLibrary.cpp
--- a/Engine/Source/Runtime/Engine/Private/HeadMountedDisplayFunctionLibrary.cpp
+++ b/Engine/Source/Runtime/Engine/Private/HeadMountedDisplayFunctionLibrary.cpp
@@ -13,7 +13,7 @@ UHeadMountedDisplayFunctionLibrary::UHeadMountedDisplayFunctionLibrary(const FOb
 
 bool UHeadMountedDisplayFunctionLibrary::IsHeadMountedDisplayEnabled()
 {
-	return GEngine->HMDDevice.IsValid() && GEngine->HMDDevice->IsHeadTrackingAllowed();
+	return GEngine->HMDDevice.IsValid() && GEngine->HMDDevice->IsHMDEnabled();
 }
 
 bool UHeadMountedDisplayFunctionLibrary::EnableHMD(bool bEnable)
