Enable substepping on non apex builds

diff --git a/Engine/Source/Runtime/Engine/Private/PhysicsEngine/PhysScene.cpp b/Engine/Source/Runtime/Engine/Private/PhysicsEngine/PhysScene.cpp
--- a/Engine/Source/Runtime/Engine/Private/PhysicsEngine/PhysScene.cpp
+++ b/Engine/Source/Runtime/Engine/Private/PhysicsEngine/PhysScene.cpp
@@ -627,34 +627,36 @@ void FPhysScene::TickPhysScene(uint32 SceneType, FGraphEventRef& InOutCompletion
 #if !WITH_APEX
 	PxScene* PScene = GetPhysXScene(SceneType);
 	if (PScene && (UseDelta > 0.f))
-	{
-		PhysXCompletionTask* Task = new PhysXCompletionTask(InOutCompletionEvent, PScene->getTaskManager());
-		PScene->lockWrite();
-		PScene->simulate(AveragedFrameTime[SceneType], Task);
-		PScene->unlockWrite();
-		Task->removeReference();
-		bTaskOutstanding = true;
-	}
-#else	//	#if !WITH_APEX
-	// The APEX scene calls the simulate function for the PhysX scene, so we only call ApexScene->simulate().
+#else
 	NxApexScene* ApexScene = GetApexScene(SceneType);
-	if(ApexScene && UseDelta > 0.f)
+	if (ApexScene && UseDelta > 0.f)
+#endif
 	{
 #if WITH_SUBSTEPPING
 		if (IsSubstepping(SceneType)) //we don't bother sub-stepping cloth
 		{
 			bTaskOutstanding = SubstepSimulation(SceneType, InOutCompletionEvent);
-		}else
+		}
+		else
 #endif
 		{
+#if !WITH_APEX
+			PhysXCompletionTask* Task = new PhysXCompletionTask(InOutCompletionEvent, PScene->getTaskManager());
+			PScene->lockWrite();
+			PScene->simulate(AveragedFrameTime[SceneType], Task);
+			PScene->unlockWrite();
+			Task->removeReference();
+			bTaskOutstanding = true;
+#else
 			PhysXCompletionTask* Task = new PhysXCompletionTask(InOutCompletionEvent, ApexScene->getTaskManager());
 			ApexScene->simulate(AveragedFrameTime[SceneType], true, Task);
 			Task->removeReference();
 			bTaskOutstanding = true;
+#endif
 		}
 	}
-#endif	//	#if !WITH_APEX
-#endif // WITH_PHYSX
+#endif
+
 	if (!bTaskOutstanding)
 	{
 		InOutCompletionEvent->DispatchSubsequents(); // nothing to do, so nothing to wait for
diff --git a/Engine/Source/Runtime/Engine/Public/EngineDefines.h b/Engine/Source/Runtime/Engine/Public/EngineDefines.h
--- a/Engine/Source/Runtime/Engine/Public/EngineDefines.h
+++ b/Engine/Source/Runtime/Engine/Public/EngineDefines.h
@@ -33,11 +33,11 @@
 	#define WITH_PHYSICS_COOKING (WITH_EDITOR || WITH_APEX)		//APEX currently relies on cooking even at runtime
 #endif
 
-#if WITH_APEX
 #ifndef WITH_SUBSTEPPING
-	#define WITH_SUBSTEPPING 1
+#define WITH_SUBSTEPPING 1
 #endif
 
+#if WITH_APEX
 #ifndef WITH_APEX_CLOTHING
 	#define WITH_APEX_CLOTHING	(1 && WITH_APEX)
 #endif // WITH_APEX_CLOTHING
