Fix LOD Actor's editoronly component check. This is not valid during -game.

diff --git a/Engine/Source/Runtime/Engine/Private/LODActor.cpp b/Engine/Source/Runtime/Engine/Private/LODActor.cpp
--- a/Engine/Source/Runtime/Engine/Private/LODActor.cpp
+++ b/Engine/Source/Runtime/Engine/Private/LODActor.cpp
@@ -35,7 +35,10 @@ ALODActor::ALODActor(const FObjectInitializer& ObjectInitializer)
 	
 #if WITH_EDITORONLY_DATA	
 	DrawSphereComponent = CreateEditorOnlyDefaultSubobject<UDrawSphereComponent>(TEXT("VisualizeComponent0"));
-	DrawSphereComponent->SetSphereRadius(0.0f);
+	if (DrawSphereComponent)
+	{
+		DrawSphereComponent->SetSphereRadius(0.0f);
+	}
 #endif // WITH_EDITORONLY_DATA
 }
 
