fixed check() (crash) when using Blendables with a weight of 0

diff --git a/Engine/Source/Runtime/Engine/Classes/Engine/BlendableInterface.h b/Engine/Source/Runtime/Engine/Classes/Engine/BlendableInterface.h
--- a/Engine/Source/Runtime/Engine/Classes/Engine/BlendableInterface.h
+++ b/Engine/Source/Runtime/Engine/Classes/Engine/BlendableInterface.h
@@ -45,7 +45,7 @@ class IBlendableInterface
 {
 	GENERATED_IINTERFACE_BODY()
 		
-	/** @param Weight 0..1, 1=fully take the values from this object, crash if outside the valid range. */
+	/** @param Weight 0..1, excluding 0, 1=fully take the values from this object, crash if outside the valid range. */
 	virtual void OverrideBlendableSettings(class FSceneView& View, float Weight) const = 0;
 };
 
diff --git a/Engine/Source/Runtime/Engine/Private/Materials/Material.cpp b/Engine/Source/Runtime/Engine/Private/Materials/Material.cpp
--- a/Engine/Source/Runtime/Engine/Private/Materials/Material.cpp
+++ b/Engine/Source/Runtime/Engine/Private/Materials/Material.cpp
@@ -1128,7 +1128,7 @@ extern FPostProcessMaterialNode* IteratePostProcessMaterialNodes(const FFinalPos
 
 void UMaterialInterface::OverrideBlendableSettings(class FSceneView& View, float Weight) const
 {
-	check(Weight >= 0.0f && Weight <= 1.0f);
+	check(Weight > 0.0f && Weight <= 1.0f);
 
 	FFinalPostProcessSettings& Dest = View.FinalPostProcessSettings;
 
diff --git a/Engine/Source/Runtime/Engine/Private/Materials/MaterialInstance.cpp b/Engine/Source/Runtime/Engine/Private/Materials/MaterialInstance.cpp
--- a/Engine/Source/Runtime/Engine/Private/Materials/MaterialInstance.cpp
+++ b/Engine/Source/Runtime/Engine/Private/Materials/MaterialInstance.cpp
@@ -2303,7 +2303,7 @@ FPostProcessMaterialNode* IteratePostProcessMaterialNodes(const FFinalPostProces
 
 void UMaterialInstance::OverrideBlendableSettings(class FSceneView& View, float Weight) const
 {
-	check(Weight >= 0.0f && Weight <= 1.0f);
+	check(Weight > 0.0f && Weight <= 1.0f);
 
 	FFinalPostProcessSettings& Dest = View.FinalPostProcessSettings;
 
diff --git a/Engine/Source/Runtime/Engine/Private/SceneView.cpp b/Engine/Source/Runtime/Engine/Private/SceneView.cpp
--- a/Engine/Source/Runtime/Engine/Private/SceneView.cpp
+++ b/Engine/Source/Runtime/Engine/Private/SceneView.cpp
@@ -1024,7 +1024,6 @@ void FSceneView::OverridePostProcessSettings(const FPostProcessSettings& Src, fl
 
 		for(uint32 i = 0; i < Count; ++i)
 		{
-			float LocalWeight = Src.WeightedBlendables.Array[i].Weight;
 			UObject* Object = Src.WeightedBlendables.Array[i].Object;
 
 			if(!Object)
@@ -1039,7 +1038,12 @@ void FSceneView::OverridePostProcessSettings(const FPostProcessSettings& Src, fl
 				continue;
 			}
 
-			BlendableInterface->OverrideBlendableSettings(*this, LocalWeight * Weight);
+			float LocalWeight = Src.WeightedBlendables.Array[i].Weight * Weight;
+
+			if(LocalWeight > 0.0f)
+			{
+				BlendableInterface->OverrideBlendableSettings(*this, LocalWeight);
+			}
 		}
 	}
 }
