Added new SpawnActor that accepts Transform, so that it can apply scale properly

- It is used by [REDACTED_PERSON]

diff --git a/Engine/Source/Runtime/Engine/Classes/Engine/World.h b/Engine/Source/Runtime/Engine/Classes/Engine/World.h
--- a/Engine/Source/Runtime/Engine/Classes/Engine/World.h
+++ b/Engine/Source/Runtime/Engine/Classes/Engine/World.h
@@ -2779,7 +2779,27 @@ public:
 	 */
 	void RemoveActor( AActor* Actor, bool bShouldModifyLevel );
 
+	/**
+	 * Spawn Actors with given transform and SpawnParameters
+	 * 
+	 * @param	Class					Class to Spawn
+	 * @param	Location				Location To Spawn
+	 * @param	Rotation				Rotation To Spawn
+	 * @param	SpawmParameters			Spawn Parameters
+	 *
+	 * @return	Actor that just spawned
+	 */
 	AActor* SpawnActor( UClass* InClass, FVector const* Location=NULL, FRotator const* Rotation=NULL, const FActorSpawnParameters& SpawnParameters = FActorSpawnParameters() );
+	/**
+	 * Spawn Actors with given transform and SpawnParameters
+	 * 
+	 * @param	Class					Class to Spawn
+	 * @param	Transform				World Transform to spawn on
+	 * @param	SpawmParameters			Spawm Parameters
+	 *
+	 * @return	Actor that just spawned
+	 */
+	AActor* SpawnActor( UClass* Class, FTransform const* Transform, const FActorSpawnParameters& SpawnParameters = FActorSpawnParameters());
 
 	/** Templated version of SpawnActor that allows you to specify a class type via the template type */
 	template< class T >
@@ -2811,6 +2831,16 @@ public:
 	{
 		return CastChecked<T>(SpawnActor(Class, &Location, &Rotation, SpawnParameters),ECastCheckedType::NullAllowed);
 	}
+	/** 
+	
+	 *  Templated version of SpawnActor that allows you to specify whole Transform
+	 *  class type via parameter while the return type is a parent class of that type 
+	 */
+	template< class T >
+	T* SpawnActor(UClass* Class, FTransform const& Transform,const FActorSpawnParameters& SpawnParameters = FActorSpawnParameters())
+	{
+		return CastChecked<T>(SpawnActor(Class, &Transform, SpawnParameters), ECastCheckedType::NullAllowed);
+	}
 	
 	/**
 	* Spawns given class and returns class T pointer, forcibly sets world position. WILL NOT run Construction Script of Blueprints 
@@ -2820,8 +2850,7 @@ public:
 	template< class T >
 	T* SpawnActorDeferred(
 		UClass* Class,
-		FVector const& Location,
-		FRotator const& Rotation,
+		FTransform const& Transform,
 		AActor* Owner=NULL,
 		APawn* Instigator=NULL,
 		bool bNoCollisionFail=false
@@ -2836,7 +2865,7 @@ public:
 		SpawnInfo.Owner = Owner;
 		SpawnInfo.Instigator = Instigator;
 		SpawnInfo.bDeferConstruction = true;
-		return (Class != NULL) ? Cast<T>(SpawnActor(Class, &Location, &Rotation, SpawnInfo )) : NULL;
+		return (Class != NULL) ? Cast<T>(SpawnActor(Class, &Transform, SpawnInfo )) : NULL;
 	}
 
 	/** 
diff --git a/Engine/Source/Runtime/Engine/Classes/GameFramework/Actor.h b/Engine/Source/Runtime/Engine/Classes/GameFramework/Actor.h
--- a/Engine/Source/Runtime/Engine/Classes/GameFramework/Actor.h
+++ b/Engine/Source/Runtime/Engine/Classes/GameFramework/Actor.h
@@ -1854,7 +1854,7 @@ public:
 	virtual bool IsRelevancyOwnerFor(const AActor* ReplicatedActor, const AActor* ActorOwner, const AActor* ConnectionActor) const;
 
 	/** Called after the actor is spawned in the world.  Responsible for setting up actor for play. */
-	void PostSpawnInitialize(FVector const& SpawnLocation, FRotator const& SpawnRotation, AActor* InOwner, APawn* InInstigator, bool bRemoteOwned, bool bNoFail, bool bDeferConstruction);
+	void PostSpawnInitialize(FTransform const& SpawnTransform, AActor* InOwner, APawn* InInstigator, bool bRemoteOwned, bool bNoFail, bool bDeferConstruction);
 
 	/** Called to finish the spawning process, generally in the case of deferred spawning */
 	void FinishSpawning(const FTransform& Transform, bool bIsDefaultTransform = false);
diff --git a/Engine/Source/Runtime/Engine/Private/Actor.cpp b/Engine/Source/Runtime/Engine/Private/Actor.cpp
--- a/Engine/Source/Runtime/Engine/Private/Actor.cpp
+++ b/Engine/Source/Runtime/Engine/Private/Actor.cpp
@@ -2389,7 +2389,7 @@ static USceneComponent* FixupNativeActorComponents(AActor* Actor)
 	return SceneRootComponent;
 }
 
-void AActor::PostSpawnInitialize(FVector const& SpawnLocation, FRotator const& SpawnRotation, AActor* InOwner, APawn* InInstigator, bool bRemoteOwned, bool bNoFail, bool bDeferConstruction)
+void AActor::PostSpawnInitialize(FTransform const& SpawnTransform, AActor* InOwner, APawn* InInstigator, bool bRemoteOwned, bool bNoFail, bool bDeferConstruction)
 {
 	// General flow here is like so
 	// - Actor sets up the basics.
@@ -2415,7 +2415,7 @@ void AActor::PostSpawnInitialize(FVector const& SpawnLocation, FRotator const& S
 	// Set the actor's location and rotation.
 	if (SceneRootComponent != NULL)
 	{
-		SceneRootComponent->SetWorldLocationAndRotation(SpawnLocation, SpawnRotation);
+		SceneRootComponent->SetWorldTransform(SpawnTransform);
 	}
 
 	// Call OnComponentCreated on all default (native) components
@@ -2452,8 +2452,7 @@ void AActor::PostSpawnInitialize(FVector const& SpawnLocation, FRotator const& S
 	if (!bDeferConstruction)
 	{
 		// Preserve original root component scale
-		const FVector SpawnScale = GetRootComponent() ? GetRootComponent()->RelativeScale3D : FVector(1.0f, 1.0f, 1.0f);
-		FinishSpawning(FTransform(SpawnRotation, SpawnLocation, SpawnScale), true);
+		FinishSpawning(SpawnTransform, true);
 	}
 }
 
diff --git a/Engine/Source/Runtime/Engine/Private/GameplayStatics.cpp b/Engine/Source/Runtime/Engine/Private/GameplayStatics.cpp
--- a/Engine/Source/Runtime/Engine/Private/GameplayStatics.cpp
+++ b/Engine/Source/Runtime/Engine/Private/GameplayStatics.cpp
@@ -379,9 +379,6 @@ AActor* UGameplayStatics::BeginSpawningActorFromBlueprint(UObject* WorldContextO
 
 AActor* UGameplayStatics::BeginSpawningActorFromClass(UObject* WorldContextObject, TSubclassOf<AActor> ActorClass, const FTransform& SpawnTransform, bool bNoCollisionFail, AActor* Owner)
 {
-	const FVector SpawnLoc(SpawnTransform.GetTranslation());
-	const FRotator SpawnRot(SpawnTransform.GetRotation());
-
 	AActor* NewActor = NULL;
 
 	UClass* Class = *ActorClass;
@@ -403,7 +400,7 @@ AActor* UGameplayStatics::BeginSpawningActorFromClass(UObject* WorldContextObjec
 		UWorld* World = GEngine->GetWorldFromContextObject(WorldContextObject);
 		if (World)
 		{
-			NewActor = World->SpawnActorDeferred<AActor>(Class, SpawnLoc, SpawnRot, Owner, AutoInstigator, bNoCollisionFail);
+			NewActor = World->SpawnActorDeferred<AActor>(Class, SpawnTransform, Owner, AutoInstigator, bNoCollisionFail);
 		}
 		else
 		{
diff --git a/Engine/Source/Runtime/Engine/Private/LevelActor.cpp b/Engine/Source/Runtime/Engine/Private/LevelActor.cpp
--- a/Engine/Source/Runtime/Engine/Private/LevelActor.cpp
+++ b/Engine/Source/Runtime/Engine/Private/LevelActor.cpp
@@ -228,6 +228,21 @@ void LineCheckTracker::CaptureLineCheck(int32 LineCheckFlags, const FVector* Ext
 #endif	//PERF_SHOW_MULTI_PAWN_SPAWN_FRAMES
 
 AActor* UWorld::SpawnActor( UClass* Class, FVector const* Location, FRotator const* Rotation, const FActorSpawnParameters& SpawnParameters )
+{
+	FTransform Transform;
+	if (Location)
+	{
+		Transform.SetLocation(*Location);
+	}
+	if (Rotation)
+	{
+		Transform.SetRotation(FQuat(*Rotation));
+	}
+
+	return SpawnActor(Class, &Transform, SpawnParameters);
+}
+
+AActor* UWorld::SpawnActor( UClass* Class, FTransform const* Transform, const FActorSpawnParameters& SpawnParameters )
 {
 	SCOPE_CYCLE_COUNTER(STAT_SpawnActorTime);
 	check( CurrentLevel ); 	
@@ -272,16 +287,10 @@ AActor* UWorld::SpawnActor( UClass* Class, FVector const* Location, FRotator con
 		UE_LOG(LogSpawn, Warning, TEXT("SpawnActor failed because we are in the process of tearing down the world"));
 		return NULL;
 	}
-	else if (Location && Location->ContainsNaN())
-	{
-		UE_LOG(LogSpawn, Warning, TEXT("SpawnActor failed because the given location (%s) is invalid"), *Location->ToString());
-		return NULL;
-	}
-	else if (Rotation && Rotation->ContainsNaN())
+	else if (Transform && Transform->ContainsNaN())
 	{
-		UE_LOG(LogSpawn, Warning, TEXT("SpawnActor failed because the given rotation (%s) is invalid"), *Rotation->ToString());
+		UE_LOG(LogSpawn, Warning, TEXT("SpawnActor failed because the given transform (%s) is invalid"), *Transform->ToString());
 		return NULL;
-		
 	}
 
 	ULevel* LevelToSpawnIn = SpawnParameters.OverrideLevel;
@@ -311,8 +320,9 @@ AActor* UWorld::SpawnActor( UClass* Class, FVector const* Location, FRotator con
 		return NULL;
 	}
 
-	FVector NewLocation = Location ? *Location : (Template->GetRootComponent() ? Template->GetRootComponent()->RelativeLocation : FVector::ZeroVector);
-	FRotator NewRotation = Rotation ? *Rotation :	(Template->GetRootComponent() ? Template->GetRootComponent()->RelativeRotation : FRotator::ZeroRotator);
+	FVector NewLocation = Transform ? Transform->GetLocation() : (Template->GetRootComponent() ? Template->GetRootComponent()->RelativeLocation : FVector::ZeroVector);
+	FRotator NewRotation = Transform ? Transform->GetRotation().Rotator() :	(Template->GetRootComponent() ? Template->GetRootComponent()->RelativeRotation : FRotator::ZeroRotator);
+	FVector NewScale = Transform ? Transform->GetScale3D() :	(Template->GetRootComponent() ? Template->GetRootComponent()->RelativeScale3D : FVector(1.f) );
 
 	// If we can fail, and we have a root component, and either the component should collide or bCollideWhenPlacing is true, test
 	if( !SpawnParameters.bNoCollisionFail && Template->GetRootComponent() && Template->bCollideWhenPlacing && Template->GetRootComponent()->ShouldCollideWhenPlacing() )
@@ -349,7 +359,7 @@ AActor* UWorld::SpawnActor( UClass* Class, FVector const* Location, FRotator con
 	}
 #endif
 
-	Actor->PostSpawnInitialize(NewLocation, NewRotation, SpawnParameters.Owner, SpawnParameters.Instigator, SpawnParameters.bRemoteOwned, SpawnParameters.bNoFail, SpawnParameters.bDeferConstruction);
+	Actor->PostSpawnInitialize(FTransform(NewRotation, NewLocation, NewScale), SpawnParameters.Owner, SpawnParameters.Instigator, SpawnParameters.bRemoteOwned, SpawnParameters.bNoFail, SpawnParameters.bDeferConstruction);
 
 	if (Actor->IsPendingKill() && !SpawnParameters.bNoFail)
 	{
diff --git a/Engine/Source/Runtime/GameplayAbilities/Private/Abilities/Tasks/AbilityTask_SpawnActor.cpp b/Engine/Source/Runtime/GameplayAbilities/Private/Abilities/Tasks/AbilityTask_SpawnActor.cpp
--- a/Engine/Source/Runtime/GameplayAbilities/Private/Abilities/Tasks/AbilityTask_SpawnActor.cpp
+++ b/Engine/Source/Runtime/GameplayAbilities/Private/Abilities/Tasks/AbilityTask_SpawnActor.cpp
@@ -26,7 +26,7 @@ bool UAbilityTask_SpawnActor::BeginSpawningActor(UObject* WorldContextObject, FG
 	if (Ability.IsValid() && Ability.Get()->GetCurrentActorInfo()->IsNetAuthority())
 	{
 		UWorld* World = GEngine->GetWorldFromContextObject(WorldContextObject);
-		SpawnedActor = World->SpawnActorDeferred<AActor>(InClass, FVector::ZeroVector, FRotator::ZeroRotator, NULL, NULL, true);
+		SpawnedActor = World->SpawnActorDeferred<AActor>(InClass, FTransform::Identity, NULL, NULL, true);
 	}
 	
 	if (SpawnedActor == nullptr)
diff --git a/Engine/Source/Runtime/GameplayAbilities/Private/Abilities/Tasks/AbilityTask_VisualizeTargeting.cpp b/Engine/Source/Runtime/GameplayAbilities/Private/Abilities/Tasks/AbilityTask_VisualizeTargeting.cpp
--- a/Engine/Source/Runtime/GameplayAbilities/Private/Abilities/Tasks/AbilityTask_VisualizeTargeting.cpp
+++ b/Engine/Source/Runtime/GameplayAbilities/Private/Abilities/Tasks/AbilityTask_VisualizeTargeting.cpp
@@ -73,7 +73,7 @@ bool UAbilityTask_VisualizeTargeting::BeginSpawningActor(UObject* WorldContextOb
 			if (Class != NULL)
 			{
 				UWorld* World = GEngine->GetWorldFromContextObject(WorldContextObject);
-				SpawnedActor = World->SpawnActorDeferred<AGameplayAbilityTargetActor>(Class, FVector::ZeroVector, FRotator::ZeroRotator, NULL, NULL, true);
+				SpawnedActor = World->SpawnActorDeferred<AGameplayAbilityTargetActor>(Class, FTransform::Identity, NULL, NULL, true);
 			}
 
 			if (SpawnedActor)
diff --git a/Engine/Source/Runtime/GameplayAbilities/Private/Abilities/Tasks/AbilityTask_WaitTargetData.cpp b/Engine/Source/Runtime/GameplayAbilities/Private/Abilities/Tasks/AbilityTask_WaitTargetData.cpp
--- a/Engine/Source/Runtime/GameplayAbilities/Private/Abilities/Tasks/AbilityTask_WaitTargetData.cpp
+++ b/Engine/Source/Runtime/GameplayAbilities/Private/Abilities/Tasks/AbilityTask_WaitTargetData.cpp
@@ -82,7 +82,7 @@ bool UAbilityTask_WaitTargetData::BeginSpawningActor(UObject* WorldContextObject
 			if (Class != NULL)
 			{
 				UWorld* World = GEngine->GetWorldFromContextObject(WorldContextObject);
-				SpawnedActor = World->SpawnActorDeferred<AGameplayAbilityTargetActor>(Class, FVector::ZeroVector, FRotator::ZeroRotator, nullptr, nullptr, true);
+				SpawnedActor = World->SpawnActorDeferred<AGameplayAbilityTargetActor>(Class, FTransform::Identity, nullptr, nullptr, true);
 			}
 
 			if (SpawnedActor)
diff --git a/Engine/Source/Runtime/GameplayTasks/Private/Tasks/GameplayTask_SpawnActor.cpp b/Engine/Source/Runtime/GameplayTasks/Private/Tasks/GameplayTask_SpawnActor.cpp
--- a/Engine/Source/Runtime/GameplayTasks/Private/Tasks/GameplayTask_SpawnActor.cpp
+++ b/Engine/Source/Runtime/GameplayTasks/Private/Tasks/GameplayTask_SpawnActor.cpp
@@ -41,7 +41,7 @@ UGameplayTask_SpawnActor* UGameplayTask_SpawnActor::SpawnActor(TScriptInterface<
 bool UGameplayTask_SpawnActor::BeginSpawningActor(UObject* WorldContextObject, AActor*& SpawnedActor)
 {
 	UWorld* World = GEngine->GetWorldFromContextObject(WorldContextObject);
-	SpawnedActor = World->SpawnActorDeferred<AActor>(ClassToSpawn, CachedSpawnLocation, CachedSpawnRotation, NULL, NULL, true);
+	SpawnedActor = World->SpawnActorDeferred<AActor>(ClassToSpawn, FTransform(CachedSpawnRotation, CachedSpawnLocation), NULL, NULL, true);
 	
 	if (SpawnedActor == nullptr)
 	{
