Fix crash when calling set simulate physics on destructible in first frame. Added some sleep behavior for when you call false. This is debatable, but consistent with the existing flag.

diff --git a/Engine/Source/Runtime/Engine/Classes/Components/DestructibleComponent.h b/Engine/Source/Runtime/Engine/Classes/Components/DestructibleComponent.h
--- a/Engine/Source/Runtime/Engine/Classes/Components/DestructibleComponent.h
+++ b/Engine/Source/Runtime/Engine/Classes/Components/DestructibleComponent.h
@@ -136,6 +136,7 @@ public:
 	virtual void SetEnableGravity(bool bGravityEnabled) override;
 
 	virtual void WakeRigidBody(FName BoneName /* = NAME_None */) override;
+	virtual void SetSimulatePhysics(bool bSimulate) override;
 
 	virtual void SetMaterial(int32 ElementIndex, UMaterialInterface* Material) override;
 	// End UPrimitiveComponent interface.
diff --git a/Engine/Source/Runtime/Engine/Private/Components/DestructibleComponent.cpp b/Engine/Source/Runtime/Engine/Private/Components/DestructibleComponent.cpp
--- a/Engine/Source/Runtime/Engine/Private/Components/DestructibleComponent.cpp
+++ b/Engine/Source/Runtime/Engine/Private/Components/DestructibleComponent.cpp
@@ -1311,7 +1311,36 @@ void UDestructibleComponent::WakeRigidBody(FName BoneName /* = NAME_None */)
 	ExecuteOnPhysicsReadWrite([&]
 	{
 		const int32 ChunkIdx = BoneIdxToChunkIdx(GetBoneIndex(BoneName));
-		//ApexDestructibleActor->setChunkPhysXActorAwakeState(ChunkIdx, true); //TODO: broke during bad integration of apex. Will turn on once new libs are in
+		ApexDestructibleActor->setChunkPhysXActorAwakeState(ChunkIdx, true);
+	});
+#endif
+}
+
+void UDestructibleComponent::SetSimulatePhysics(bool bSimulate)
+{
+#if WITH_APEX
+	ExecuteOnPhysicsReadWrite([&]
+	{
+		if(bSimulate)
+		{
+			ApexDestructibleActor->setDynamic();
+		}else
+		{
+			PxRigidDynamic** PActorBuffer = NULL;
+			PxU32 PActorCount = 0;
+			if (ApexDestructibleActor->acquirePhysXActorBuffer(PActorBuffer, PActorCount))
+			{
+				for(uint32 ActorIdx = 0; ActorIdx < PActorCount; ++ActorIdx)
+				{
+					PxRigidDynamic* PActor = PActorBuffer[ActorIdx];
+					if(FDestructibleChunkInfo* ChunkInfo = FPhysxUserData::Get<FDestructibleChunkInfo>(PActor->userData))
+					{
+						ApexDestructibleActor->setChunkPhysXActorAwakeState(ChunkInfo->ChunkIndex, false);
+					}
+				}
+			}
+		}
+		
 	});
 #endif
 }
