- Added UEngine::bForceDisableFrameRateSmoothing to enable external systems (e.g. HMDs) to force smoothing off without overwriting the desired base engine setting
- varous HMD modules now setting bForceDisableFrameRateSmoothing appropriately as stereo rendering is enabled/disabled

diff --git a/Engine/Plugins/Runtime/GearVR/Source/GearVR/Private/GearVR.cpp b/Engine/Plugins/Runtime/GearVR/Source/GearVR/Private/GearVR.cpp
--- a/Engine/Plugins/Runtime/GearVR/Source/GearVR/Private/GearVR.cpp
+++ b/Engine/Plugins/Runtime/GearVR/Source/GearVR/Private/GearVR.cpp
@@ -390,6 +390,8 @@ bool FGearVR::EnablePositionalTracking(bool enable)
 
 bool FGearVR::DoEnableStereo(bool bStereo, bool bApplyToHmd)
 {
+	// Uncap fps to enable FPS higher than 62
+	GEngine->bForceDisableFrameRateSmoothing = bStereo;
 	return true;
 }
 
@@ -670,9 +672,6 @@ void FGearVR::Startup()
 	UpdateHmdRenderInfo();
 	UpdateStereoRenderingParams();
 
-	// Uncap fps to enable FPS higher than 62
-	GEngine->bSmoothFrameRate = false;
-
 	pGearVRBridge = new FGearVRCustomPresent(MinimumVsyncs);
 
 	LoadFromIni();
diff --git a/Engine/Plugins/Runtime/GearVR/Source/GearVR/Private/HeadMountedDisplayCommon.cpp b/Engine/Plugins/Runtime/GearVR/Source/GearVR/Private/HeadMountedDisplayCommon.cpp
--- a/Engine/Plugins/Runtime/GearVR/Source/GearVR/Private/HeadMountedDisplayCommon.cpp
+++ b/Engine/Plugins/Runtime/GearVR/Source/GearVR/Private/HeadMountedDisplayCommon.cpp
@@ -839,7 +839,7 @@ bool FHeadMountedDisplay::Exec(UWorld* InWorld, const TCHAR* Cmd, FOutputDevice&
 	}
 	else if (FParse::Command(&Cmd, TEXT("UNCAPFPS")))
 	{
-		GEngine->bSmoothFrameRate = false;
+		GEngine->bForceDisableFrameRateSmoothing = true;
 		return true;
 	}
 	else if (FParse::Command(&Cmd, TEXT("HMDVERSION")))
diff --git a/Engine/Plugins/Runtime/OculusRift/Source/OculusRift/Private/0.5/OculusRiftHMD_05.cpp b/Engine/Plugins/Runtime/OculusRift/Source/OculusRift/Private/0.5/OculusRiftHMD_05.cpp
--- a/Engine/Plugins/Runtime/OculusRift/Source/OculusRift/Private/0.5/OculusRiftHMD_05.cpp
+++ b/Engine/Plugins/Runtime/OculusRift/Source/OculusRift/Private/0.5/OculusRiftHMD_05.cpp
@@ -732,6 +732,9 @@ bool FOculusRiftHMD::DoEnableStereo(bool bStereo, bool bApplyToHmd)
 		return Settings->Flags.bStereoEnabled;
 	}
 
+	// Uncap fps to enable FPS higher than 62
+	GEngine->bForceDisableFrameRateSmoothing = bStereo;
+
 	TSharedPtr<SWindow> Window;
 	if (SceneVP)
 	{
@@ -1246,9 +1249,6 @@ void FOculusRiftHMD::Startup()
 		return;
 	}
 
-	// Uncap fps to enable FPS higher than 62
-	GEngine->bSmoothFrameRate = false;
-
 	SaveSystemValues();
 
 #ifdef OVR_SDK_RENDERING
diff --git a/Engine/Plugins/Runtime/OculusRift/Source/OculusRift/Private/HeadMountedDisplayCommon.cpp b/Engine/Plugins/Runtime/OculusRift/Source/OculusRift/Private/HeadMountedDisplayCommon.cpp
--- a/Engine/Plugins/Runtime/OculusRift/Source/OculusRift/Private/HeadMountedDisplayCommon.cpp
+++ b/Engine/Plugins/Runtime/OculusRift/Source/OculusRift/Private/HeadMountedDisplayCommon.cpp
@@ -862,7 +862,7 @@ bool FHeadMountedDisplay::Exec(UWorld* InWorld, const TCHAR* Cmd, FOutputDevice&
 	}
 	else if (FParse::Command(&Cmd, TEXT("UNCAPFPS")))
 	{
-		GEngine->bSmoothFrameRate = false;
+		GEngine->bForceDisableFrameRateSmoothing = true;
 		return true;
 	}
 	else if (FParse::Command(&Cmd, TEXT("HMDVERSION")))
diff --git a/Engine/Plugins/Runtime/OculusRift/Source/OculusRift/Private/OculusRiftHMD.cpp b/Engine/Plugins/Runtime/OculusRift/Source/OculusRift/Private/OculusRiftHMD.cpp
--- a/Engine/Plugins/Runtime/OculusRift/Source/OculusRift/Private/OculusRiftHMD.cpp
+++ b/Engine/Plugins/Runtime/OculusRift/Source/OculusRift/Private/OculusRiftHMD.cpp
@@ -711,6 +711,9 @@ bool FOculusRiftHMD::DoEnableStereo(bool bStereo, bool bApplyToHmd)
 		}
 	}
 
+	// Uncap fps to enable FPS higher than 62
+	GEngine->bForceDisableFrameRateSmoothing = bStereo;
+
 	bool wasFullscreenAllowed = IsFullscreenAllowed();
 	if (OnOculusStateChange(stereoToBeEnabled))
 	{
@@ -1137,9 +1140,6 @@ void FOculusRiftHMD::Startup()
 		return;
 	}
 
-	// Uncap fps to enable FPS higher than 62
-	GEngine->bSmoothFrameRate = false;
-
 #if defined(OVR_D3D_VERSION) && (OVR_D3D_VERSION == 11)
 	if (IsPCPlatform(GMaxRHIShaderPlatform) && !IsOpenGLPlatform(GMaxRHIShaderPlatform))
 	{
diff --git a/Engine/Plugins/Runtime/Steam/SteamVR/Source/SteamVR/Private/SteamVRHMD.cpp b/Engine/Plugins/Runtime/Steam/SteamVR/Source/SteamVR/Private/SteamVRHMD.cpp
--- a/Engine/Plugins/Runtime/Steam/SteamVR/Source/SteamVR/Private/SteamVRHMD.cpp
+++ b/Engine/Plugins/Runtime/Steam/SteamVR/Source/SteamVR/Private/SteamVRHMD.cpp
@@ -510,7 +510,7 @@ bool FSteamVRHMD::Exec( UWorld* InWorld, const TCHAR* Cmd, FOutputDevice& Ar )
 	}
 	else if (FParse::Command(&Cmd, TEXT("UNCAPFPS")))
 	{
-		GEngine->bSmoothFrameRate = false;
+		GEngine->bForceDisableFrameRateSmoothing = true;
 		return true;
 	}
 
@@ -602,12 +602,15 @@ bool FSteamVRHMD::IsStereoEnabled() const
 	return bStereoEnabled && bHmdEnabled;
 }
 
-bool FSteamVRHMD::EnableStereo(bool stereo)
+bool FSteamVRHMD::EnableStereo(bool bStereo)
 {
-	bStereoEnabled = (IsHMDEnabled()) ? stereo : false;
+	bStereoEnabled = (IsHMDEnabled()) ? bStereo : false;
 
-	FSystemResolution::RequestResolutionChange(1280, 720, (stereo) ? EWindowMode::WindowedMirror : EWindowMode::Windowed);
+	FSystemResolution::RequestResolutionChange(1280, 720, (bStereo) ? EWindowMode::WindowedMirror : EWindowMode::Windowed);
 
+	// Uncap fps to enable FPS higher than 62
+	GEngine->bForceDisableFrameRateSmoothing = bStereo;
+	
 	return bStereoEnabled;
 }
 
@@ -935,9 +938,6 @@ void FSteamVRHMD::Startup()
 		static IConsoleVariable* CFCFVar = IConsoleManager::Get().FindConsoleVariable(TEXT("r.finishcurrentframe"));
 		CFCFVar->Set(false);
 
-		// Uncap fps to enable FPS higher than 62
-		GEngine->bSmoothFrameRate = false;
-
 		// Grab the chaperone
 		vr::HmdError ChaperoneErr = vr::HmdError_None;
 		//VRChaperone = (vr::IVRChaperone*)vr::VR_GetGenericInterface(vr::IVRChaperone_Version, &ChaperoneErr);
diff --git a/Engine/Source/Runtime/Engine/Classes/Engine/Engine.h b/Engine/Source/Runtime/Engine/Classes/Engine/Engine.h
--- a/Engine/Source/Runtime/Engine/Classes/Engine/Engine.h
+++ b/Engine/Source/Runtime/Engine/Classes/Engine/Engine.h
@@ -1148,18 +1148,21 @@ public:
 	UPROPERTY(config)
 	uint32 bEnableEditorPSysRealtimeLOD:1;
 
+	/** Hook for external systems to transiently and forcibly disable framerate smoothing without stomping the original setting. */
+	uint32 bForceDisableFrameRateSmoothing : 1;
+
+	/** Whether to enable framerate smoothing. */
+	UPROPERTY(config, EditAnywhere, Category=Framerate, meta=(EditCondition="!bUseFixedFrameRate"))
+	uint32 bSmoothFrameRate:1;
+
 	/** Whether to use a fixed framerate. */
 	UPROPERTY(config, EditAnywhere, Category = Framerate)
 	uint32 bUseFixedFrameRate : 1;
-	
+
 	/** The fixed framerate to use. */
-	UPROPERTY(config, EditAnywhere, Category = Framerate, meta=(EditCondition="bUseFixedFrameRate"))
+	UPROPERTY(config, EditAnywhere, Category = Framerate, meta = (EditCondition = "bUseFixedFrameRate"))
 	float FixedFrameRate;
 
-	/** Whether to enable framerate smoothing.																		*/
-	UPROPERTY(config, EditAnywhere, Category=Framerate, meta=(EditCondition="!bUseFixedFrameRate"))
-	uint32 bSmoothFrameRate:1;
-
 	/** Range of framerates in which smoothing will kick in */
 	UPROPERTY(config, EditAnywhere, Category=Framerate, meta=(UIMin=0, UIMax=200, EditCondition="!bUseFixedFrameRate"))
 	FFloatRange SmoothedFrameRateRange;
@@ -1733,7 +1736,7 @@ public:
 	virtual void UpdateRunningAverageDeltaTime(float DeltaTime, bool bAllowFrameRateSmoothing = true);
 
 	/** Whether we're allowed to do frame rate smoothing */
-	bool IsAllowedFramerateSmoothing(bool bAllowFrameRateSmoothing) const;
+	virtual bool IsAllowedFramerateSmoothing() const;
 
 	/**
 	 * Pauses / un-pauses the game-play when focus of the game's window gets lost / gained.
diff --git a/Engine/Source/Runtime/Engine/Private/UnrealEngine.cpp b/Engine/Source/Runtime/Engine/Private/UnrealEngine.cpp
--- a/Engine/Source/Runtime/Engine/Private/UnrealEngine.cpp
+++ b/Engine/Source/Runtime/Engine/Private/UnrealEngine.cpp
@@ -6473,35 +6473,35 @@ void UEngine::InitializeRunningAverageDeltaTime()
 	RunningAverageDeltaTime = 1 / 100.f;
 }
 
-bool UEngine::IsAllowedFramerateSmoothing(bool bAllowFrameRateSmoothing) const
+bool UEngine::IsAllowedFramerateSmoothing() const
 {
-	return FPlatformProperties::AllowsFramerateSmoothing() && bSmoothFrameRate && bAllowFrameRateSmoothing && !IsRunningDedicatedServer();
+	return FPlatformProperties::AllowsFramerateSmoothing() && bSmoothFrameRate && !bForceDisableFrameRateSmoothing && !IsRunningDedicatedServer();
 }
 
 /** Compute tick rate limitor. */
 void UEngine::UpdateRunningAverageDeltaTime(float DeltaTime, bool bAllowFrameRateSmoothing)
 {
-	if (IsAllowedFramerateSmoothing(bAllowFrameRateSmoothing))
+	if (bAllowFrameRateSmoothing && IsAllowedFramerateSmoothing())
 	{
 		// Smooth the framerate if wanted. The code uses a simplistic running average. Other approaches, like reserving
 		// a percentage of time, ended up creating negative feedback loops in conjunction with GPU load and were abandonend.
-			if( DeltaTime < 0.0f )
-			{
+		if( DeltaTime < 0.0f )
+		{
 #if PLATFORM_ANDROID
-				UE_LOG(LogEngine, Warning, TEXT("Detected negative delta time - ignoring"));
-				DeltaTime = 0.01;
+			UE_LOG(LogEngine, Warning, TEXT("Detected negative delta time - ignoring"));
+			DeltaTime = 0.01;
 #elif (UE_BUILD_SHIPPING && WITH_EDITOR)
-				// End users don't have access to the secure parts of UDN. The localized string points to the release notes,
-				// which should include a link to the AMD CPU drivers download site.
-				UE_LOG(LogEngine, Fatal, TEXT("%s"), TEXT("CPU time drift detected! Please consult release notes on how to address this."));
+			// End users don't have access to the secure parts of UDN. The localized string points to the release notes,
+			// which should include a link to the AMD CPU drivers download site.
+			UE_LOG(LogEngine, Fatal, TEXT("%s"), TEXT("CPU time drift detected! Please consult release notes on how to address this."));
 #else
-				// Send developers to the support list thread.
-				UE_LOG(LogEngine, Fatal, TEXT("Negative delta time! Please see https://udn.epicgames.com/lists/showpost.php?list=ue3bugs&id=4364"));
+			// Send developers to the support list thread.
+			UE_LOG(LogEngine, Fatal, TEXT("Negative delta time! Please see https://udn.epicgames.com/lists/showpost.php?list=ue3bugs&id=4364"));
 #endif
-			}
+		}
 
-			// Keep track of running average over 300 frames, clamping at min of 5 FPS for individual delta times.
-			RunningAverageDeltaTime = FMath::Lerp<float>( RunningAverageDeltaTime, FMath::Min<float>( DeltaTime, 0.2f ), 1 / 300.f );
+		// Keep track of running average over 300 frames, clamping at min of 5 FPS for individual delta times.
+		RunningAverageDeltaTime = FMath::Lerp<float>( RunningAverageDeltaTime, FMath::Min<float>( DeltaTime, 0.2f ), 1 / 300.f );
 	}
 }
 
@@ -6511,21 +6511,21 @@ float UEngine::GetMaxTickRate(float DeltaTime, bool bAllowFrameRateSmoothing) co
 {
 	float MaxTickRate = 0;
 
-	if (IsAllowedFramerateSmoothing(bAllowFrameRateSmoothing))
+	if (bAllowFrameRateSmoothing && IsAllowedFramerateSmoothing())
 	{
-			// Work in FPS domain as that is what the function will return.
-			MaxTickRate = 1.f / RunningAverageDeltaTime;
+		// Work in FPS domain as that is what the function will return.
+		MaxTickRate = 1.f / RunningAverageDeltaTime;
 
-			// Clamp FPS into ini defined min/ max range.
-			if (SmoothedFrameRateRange.HasLowerBound())
-			{
-				MaxTickRate = FMath::Max( MaxTickRate, SmoothedFrameRateRange.GetLowerBoundValue() );
-			}
-			if (SmoothedFrameRateRange.HasUpperBound())
-			{
-				MaxTickRate = FMath::Min( MaxTickRate, SmoothedFrameRateRange.GetUpperBoundValue() );
-			}
+		// Clamp FPS into ini defined min/ max range.
+		if (SmoothedFrameRateRange.HasLowerBound())
+		{
+			MaxTickRate = FMath::Max( MaxTickRate, SmoothedFrameRateRange.GetLowerBoundValue() );
 		}
+		if (SmoothedFrameRateRange.HasUpperBound())
+		{
+			MaxTickRate = FMath::Min( MaxTickRate, SmoothedFrameRateRange.GetUpperBoundValue() );
+		}
+	}
 
 	if (CVarCauseHitches.GetValueOnGameThread())
 	{
