Strip PIE prefix when starting seamless travel

diff --git a/Engine/Source/Runtime/Engine/Private/World.cpp b/Engine/Source/Runtime/Engine/Private/World.cpp
--- a/Engine/Source/Runtime/Engine/Private/World.cpp
+++ b/Engine/Source/Runtime/Engine/Private/World.cpp
@@ -4329,9 +4329,10 @@ bool FSeamlessTravelHandler::StartTravel(UWorld* InCurrentWorld, const FURL& InU
 	else
 	{
 		UE_LOG(LogWorld, Log, TEXT("SeamlessTravel to: %s"), *InURL.Map);
-		if (!FPackageName::DoesPackageExist(InURL.Map, InGuid.IsValid() ? &InGuid : NULL))
+		FString MapName = UWorld::RemovePIEPrefix(InURL.Map);
+		if (!FPackageName::DoesPackageExist(MapName, InGuid.IsValid() ? &InGuid : NULL))
 		{
-			UE_LOG(LogWorld, Error, TEXT("Unable to travel to '%s' - file not found"), *InURL.Map);
+			UE_LOG(LogWorld, Error, TEXT("Unable to travel to '%s' - file not found"), *MapName);
 			return false;
 			// @todo: might have to handle this more gracefully to handle downloading (might also need to send GUID and check it here!)
 		}
