Fixed issues with debug commands to add/remove players (UE-16426)
- fixed crash doing SSSwapControllers with 0 players
- fixed GameInstance execs not firing with 0 players

diff --git a/Engine/Source/Runtime/Engine/Private/GameViewportClient.cpp b/Engine/Source/Runtime/Engine/Private/GameViewportClient.cpp
--- a/Engine/Source/Runtime/Engine/Private/GameViewportClient.cpp
+++ b/Engine/Source/Runtime/Engine/Private/GameViewportClient.cpp
@@ -1426,13 +1426,20 @@ void UGameViewportClient::PeekNetworkFailureMessages(UWorld *InWorld, UNetDriver
 void UGameViewportClient::SSSwapControllers()
 {
 #if !UE_BUILD_SHIPPING
-	const int32 TmpControllerID = GetOuterUEngine()->GetFirstGamePlayer(this)->GetControllerId();
+	UEngine* const Engine = GetOuterUEngine();
 
-	for (int32 Idx=0; Idx<GetOuterUEngine()->GetNumGamePlayers(this)-1; ++Idx)
+	int32 const NumPlayers = Engine ? Engine->GetNumGamePlayers(this) : 0;
+	if (NumPlayers > 1)
 	{
-		GetOuterUEngine()->GetGamePlayer(this, Idx)->SetControllerId(GetOuterUEngine()->GetGamePlayer(this, Idx+1)->GetControllerId());
+		ULocalPlayer* const LP = Engine ? Engine->GetFirstGamePlayer(this) : nullptr;
+		const int32 TmpControllerID = LP ? LP->GetControllerId() : 0;
+
+		for (int32 Idx = 0; Idx<NumPlayers-1; ++Idx)
+		{
+			Engine->GetGamePlayer(this, Idx)->SetControllerId(Engine->GetGamePlayer(this, Idx + 1)->GetControllerId());
+		}
+		Engine->GetGamePlayer(this, NumPlayers-1)->SetControllerId(TmpControllerID);
 	}
-	GetOuterUEngine()->GetGamePlayer(this, GetOuterUEngine()->GetNumGamePlayers(this)-1)->SetControllerId(TmpControllerID);
 #endif
 }
 
@@ -2094,7 +2101,7 @@ bool UGameViewportClient::Exec( UWorld* InWorld, const TCHAR* Cmd,FOutputDevice&
 	{
 		return true;
 	}
-	else if ( GameInstance && GameInstance->Exec(InWorld, Cmd, Ar) )
+	else if ( GameInstance && (GameInstance->Exec(InWorld, Cmd, Ar) || GameInstance->ProcessConsoleExec(Cmd, Ar, nullptr)) )
 	{
 		return true;
 	}
diff --git a/Engine/Source/Runtime/Engine/Private/Player.cpp b/Engine/Source/Runtime/Engine/Private/Player.cpp
--- a/Engine/Source/Runtime/Engine/Private/Player.cpp
+++ b/Engine/Source/Runtime/Engine/Private/Player.cpp
@@ -103,10 +103,6 @@ bool UPlayer::Exec( UWorld* InWorld, const TCHAR* Cmd,FOutputDevice& Ar)
 		{
 			return true;
 		}
-		else if( World->GetGameInstance() && World->GetGameInstance()->ProcessConsoleExec(Cmd, Ar, PCPawn) )
-		{
-			return true;
-		}
 		else if( World->GetAuthGameMode() && World->GetAuthGameMode()->ProcessConsoleExec(Cmd,Ar,PCPawn) )
 		{
 			return true;
