UE4 - Remove named pipes from SCW

diff --git a/Engine/Config/BaseEngine.ini b/Engine/Config/BaseEngine.ini
--- a/Engine/Config/BaseEngine.ini
+++ b/Engine/Config/BaseEngine.ini
@@ -1615,14 +1615,6 @@ bPromptToRetryFailedShaderCompiles=True
 bLogJobCompletionTimes=False
 ; Only using 10ms of game thread time per frame to process async shader maps
 ProcessGameThreadTargetTime=.01
-; Use named pipes as opposed to file for communicating to worker processes
-bUseNamedPipes=False
-; Use async IO (overlapped) on named pipes
-bUseNamedPipesAsync=True
-; Trigger one worker process for each pipe job in sequence
-bSingleJobPerNamedPipeProcess=False
-; Reuse processes and pipes
-bReuseNamedPipeAndProcess=True
 
 [DevOptions.Debug]
 ShowSelectedLightmap=False
diff --git a/Engine/Source/Programs/ShaderCompileWorker/Private/ShaderCompileWorker.cpp b/Engine/Source/Programs/ShaderCompileWorker/Private/ShaderCompileWorker.cpp
--- a/Engine/Source/Programs/ShaderCompileWorker/Private/ShaderCompileWorker.cpp
+++ b/Engine/Source/Programs/ShaderCompileWorker/Private/ShaderCompileWorker.cpp
@@ -88,32 +88,12 @@ static void ProcessCompilationJob(const FShaderCompilerInput& Input,FShaderCompi
 	Compiler->CompileShader(Input.ShaderFormat, Input, Output, WorkingDirectory);
 }
 
-#if PLATFORM_SUPPORTS_NAMED_PIPES
-static void VerifyResult(bool bResult, const TCHAR* InMessage = TEXT(""))
-{
-#if PLATFORM_WINDOWS
-	if (!bResult)
-	{
-		TCHAR ErrorBuffer[ MAX_SPRINTF ];
-		uint32 Error = GetLastError();
-		const TCHAR* Buffer = FWindowsPlatformMisc::GetSystemErrorMessage(ErrorBuffer, MAX_SPRINTF, Error);
-		FString Message = FString::Printf(TEXT("FAILED (%s) with GetLastError() %d: %s!\n"), InMessage, Error, ErrorBuffer);
-		FPlatformMisc::LowLevelOutputDebugStringf(*Message);
-		UE_LOG(LogShaders, Fatal, TEXT("%s"), *Message);
-	}
-#endif // PLATFORM_WINDOWS
-	verify(bResult != 0);
-}
-#endif
-
 class FWorkLoop
 {
 public:
 	enum ECommunicationMode
 	{
 		ThroughFile,
-		ThroughNamedPipeOnce,
-		ThroughNamedPipe,
 	};
 	FWorkLoop(const TCHAR* ParentProcessIdText,const TCHAR* InWorkingDirectory,const TCHAR* InInputFilename,const TCHAR* InOutputFilename, ECommunicationMode InCommunicationMode)
 	:	ParentProcessId(FCString::Atoi(ParentProcessIdText))
@@ -124,9 +104,6 @@ public:
 	,	InputFilePath(InCommunicationMode == ThroughFile ? (FString(InWorkingDirectory) + InInputFilename) : InInputFilename)
 	,	OutputFilePath(InCommunicationMode == ThroughFile ? (FString(InWorkingDirectory) + InOutputFilename) : InOutputFilename)
 	{
-#if PLATFORM_SUPPORTS_NAMED_PIPES
-		LastConnectionTime = FPlatformTime::Seconds();
-#endif
 	}
 
 	void Loop()
@@ -165,25 +142,6 @@ public:
 			// Change the output file name to requested one
 			IFileManager::Get().Move(*OutputFilePath, *TempFilePath);
 
-#if PLATFORM_SUPPORTS_NAMED_PIPES
-			if (CommunicationMode == ThroughNamedPipeOnce || CommunicationMode == ThroughNamedPipe)
-			{
-				VerifyResult(Pipe.WriteInt32(TransferBufferOut.Num()), TEXT("Writing Transfer Size"));
-				VerifyResult(Pipe.WriteBytes(TransferBufferOut.Num(), TransferBufferOut.GetData()), TEXT("Writing Transfer Buffer"));
-//FPlatformMisc::LowLevelOutputDebugStringf(TEXT("*** Closing pipe...\n"));
-				Pipe.Destroy();
-
-				if (CommunicationMode == ThroughNamedPipeOnce)
-				{
-					// Give up CPU time while we are waiting
-					FPlatformProcess::Sleep(0.02f);
-					break;
-				}
-
-				LastConnectionTime = FPlatformTime::Seconds();
-			}
-#endif	// PLATFORM_SUPPORTS_NAMED_PIPES
-
 			if (GShaderCompileUseXGE)
 			{
 				// To signal compilation completion, create a zero length file in the working directory.
@@ -214,18 +172,6 @@ private:
 	const FString OutputFilePath;
 	FString TempFilePath;
 
-#if PLATFORM_SUPPORTS_NAMED_PIPES
-	FPlatformNamedPipe Pipe;
-	TArray<uint8> TransferBufferIn;
-	TArray<uint8> TransferBufferOut;
-	double LastConnectionTime;
-#endif	// PLATFORM_SUPPORTS_NAMED_PIPES
-
-	bool IsUsingNamedPipes() const
-	{
-		return (CommunicationMode == ThroughNamedPipeOnce || CommunicationMode == ThroughNamedPipe);
-	}
-
 	/** Opens an input file, trying multiple times if necessary. */
 	FArchive* OpenInputFile()
 	{
@@ -238,34 +184,6 @@ private:
 			{
 				InputFile = IFileManager::Get().CreateFileReader(*InputFilePath,FILEREAD_Silent);
 			}
-			else
-			{
-#if PLATFORM_SUPPORTS_NAMED_PIPES
-				check(IsUsingNamedPipes()); //UE_LOG(LogShaders, Log, TEXT("Opening Pipe %s\n"), *InputFilePath);
-//FPlatformMisc::LowLevelOutputDebugStringf(TEXT("*** Trying to open pipe %s\n"), *InputFilePath);
-				if (Pipe.Create(InputFilePath, false, false))
-				{
-//FPlatformMisc::LowLevelOutputDebugStringf(TEXT("\tOpened!!!\n"));
-					// Read the total number of bytes
-					int32 TransferSize = 0;
-					VerifyResult(Pipe.ReadInt32(TransferSize));
-
-					// Prealloc and read the full buffer
-					TransferBufferIn.Empty(TransferSize);
-					TransferBufferIn.AddUninitialized(TransferSize);	//UE_LOG(LogShaders, Log, TEXT("Reading Buffer\n"));
-					VerifyResult(Pipe.ReadBytes(TransferSize, TransferBufferIn.GetData()));
-
-					return new FMemoryReader(TransferBufferIn);
-				}
-
-				double DeltaTime = FPlatformTime::Seconds();
-				if (DeltaTime - LastConnectionTime > 20.0f)
-				{
-					// If can't connect for more than 20 seconds, let's exit
-					FPlatformMisc::RequestExit(false);
-				}
-#endif	// PLATFORM_SUPPORTS_NAMED_PIPES
-			}
 
 			if(!InputFile && !bFirstOpenTry)
 			{
@@ -285,7 +203,7 @@ private:
 		FArchive& InputFile = *InputFilePtr;
 		int32 InputVersion;
 		InputFile << InputVersion;
-		check(ShaderCompileWorkerInputVersion == InputVersion);
+		checkf(ShaderCompileWorkerInputVersion == InputVersion, TEXT("ShaderCompilerWorker expecting input version %d, got %d instead! Forgot to build ShaderCompilerWorker?"), ShaderCompileWorkerInputVersion, InputVersion);
 
 		InputFile << NumBatches;
 
@@ -364,16 +282,6 @@ private:
 				UE_LOG(LogShaders, Fatal,TEXT("Couldn't save output file %s"), *TempFilePath);
 			}
 		}
-		else
-		{
-#if PLATFORM_SUPPORTS_NAMED_PIPES
-			check(IsUsingNamedPipes());
-
-			// Output Transfer Buffer...
-			TransferBufferOut.Empty(0);
-			OutputFilePtr = new FMemoryWriter(TransferBufferOut);
-#endif
-		}
 
 		return OutputFilePtr;
 	}
@@ -405,21 +313,11 @@ private:
 	/** Called in the idle loop, checks for conditions under which the helper should exit */
 	void CheckExitConditions()
 	{
-#if PLATFORM_SUPPORTS_NAMED_PIPES
-		if (CommunicationMode == ThroughNamedPipeOnce)
+		if (!InputFilename.Contains(TEXT("Only")))
 		{
-			UE_LOG(LogShaders, Log, TEXT("PipeOnce: exiting after one job."));
+			UE_LOG(LogShaders, Log, TEXT("InputFilename did not contain 'Only', exiting after one job."));
 			FPlatformMisc::RequestExit(false);
 		}
-		else if (CommunicationMode == ThroughFile)
-#endif
-		{
-			if (!InputFilename.Contains(TEXT("Only")))
-			{
-				UE_LOG(LogShaders, Log, TEXT("InputFilename did not contain 'Only', exiting after one job."));
-				FPlatformMisc::RequestExit(false);
-			}
-		}
 
 #if PLATFORM_MAC || PLATFORM_LINUX
 		if (!FPlatformMisc::IsDebuggerPresent() && ParentProcessId > 0)
@@ -455,10 +353,7 @@ private:
 				// If we couldn't open the process then it is no longer running, exit
 				if (ParentProcessHandle == nullptr)
 				{
-					if (!IsUsingNamedPipes())
-					{
-						checkf(IFileManager::Get().FileSize(*FilePath) == INDEX_NONE, TEXT("Exiting due to OpenProcess(ParentProcessId) failing and the input file is present!"));
-					}
+					checkf(IFileManager::Get().FileSize(*FilePath) == INDEX_NONE, TEXT("Exiting due to OpenProcess(ParentProcessId) failing and the input file is present!"));
 					UE_LOG(LogShaders, Log, TEXT("Couldn't OpenProcess, Parent process no longer running, exiting"));
 					FPlatformMisc::RequestExit(false);
 				}
@@ -470,10 +365,7 @@ private:
 					uint32 WaitResult = WaitForSingleObject(ParentProcessHandle, 0);
 					if (WaitResult != WAIT_TIMEOUT)
 					{
-						if (!IsUsingNamedPipes())
-						{
-							checkf(IFileManager::Get().FileSize(*FilePath) == INDEX_NONE, TEXT("Exiting due to WaitForSingleObject(ParentProcessHandle) signaling and the input file is present!"));
-						}
+						checkf(IFileManager::Get().FileSize(*FilePath) == INDEX_NONE, TEXT("Exiting due to WaitForSingleObject(ParentProcessHandle) signaling and the input file is present!"));
 						UE_LOG(LogShaders, Log, TEXT("WaitForSingleObject signaled, Parent process no longer running, exiting"));
 						FPlatformMisc::RequestExit(false);
 					}
@@ -529,19 +421,7 @@ int32 GuardedMain(int32 argc, TCHAR* argv[])
 
 	LastCompileTime = FPlatformTime::Seconds();
 
-	FString InCommunicating = (argc > 6) ? argv[6] : FString();
-#if PLATFORM_SUPPORTS_NAMED_PIPES
-	const bool bThroughFile = GShaderCompileUseXGE || (InCommunicating == FString(TEXT("-communicatethroughfile")));
-	const bool bThroughNamedPipe = (InCommunicating == FString(TEXT("-communicatethroughnamedpipe")));
-	const bool bThroughNamedPipeOnce = (InCommunicating == FString(TEXT("-communicatethroughnamedpipeonce")));
-#else
-	const bool bThroughFile = true;
-	const bool bThroughNamedPipe = false;
-	const bool bThroughNamedPipeOnce = false;
-#endif
-	check((int32)bThroughFile + (int32)bThroughNamedPipe + (int32)bThroughNamedPipeOnce == 1);
-
-	FWorkLoop::ECommunicationMode Mode = bThroughFile ? FWorkLoop::ThroughFile : (bThroughNamedPipeOnce ? FWorkLoop::ThroughNamedPipeOnce : FWorkLoop::ThroughNamedPipe);
+	FWorkLoop::ECommunicationMode Mode = FWorkLoop::ThroughFile;
 	FWorkLoop WorkLoop(argv[2], argv[1], argv[4], argv[5], Mode);
 
 	WorkLoop.Loop();
diff --git a/Engine/Source/Runtime/Engine/Private/ShaderCompiler/ShaderCompiler.cpp b/Engine/Source/Runtime/Engine/Private/ShaderCompiler/ShaderCompiler.cpp
--- a/Engine/Source/Runtime/Engine/Private/ShaderCompiler/ShaderCompiler.cpp
+++ b/Engine/Source/Runtime/Engine/Private/ShaderCompiler/ShaderCompiler.cpp
@@ -128,9 +128,10 @@ static bool DoWriteTasks(const TArray<FShaderCompileJob*>& QueuedJobs, FArchive&
 // Process results from Worker Process
 static void DoReadTaskResults(const TArray<FShaderCompileJob*>& QueuedJobs, FArchive& OutputFile)
 {
+	const int32 OutputVersion = 1;
 	int32 ShaderCompileWorkerOutputVersion;
 	OutputFile << ShaderCompileWorkerOutputVersion;
-	check(ShaderCompileWorkerOutputVersion == 1);
+	checkf(ShaderCompileWorkerOutputVersion == OutputVersion, TEXT("Expecting ShaderCompilerWorker output version %d, got %d instead! Forgot to build ShaderCompilerWorker?"), ShaderCompileWorkerOutputVersion, OutputVersion);
 
 	int32 ErrorCode;
 	OutputFile << ErrorCode;
@@ -178,198 +179,6 @@ static void DoReadTaskResults(const TArray<FShaderCompileJob*>& QueuedJobs, FArc
 	}
 }
 
-#if PLATFORM_SUPPORTS_NAMED_PIPES
-struct FShaderPipeConfig
-{
-	bool	bUseNamedPipes;
-	bool	bUseNamedPipesAsync;
-	bool	bSingleJobPerNamedPipeProcess;
-	bool	bReuseNamedPipeAndProcess;
-	int32	PipeGuid;
-
-	FShaderPipeConfig() :
-		bUseNamedPipes(true),
-		bUseNamedPipesAsync(true),
-		bSingleJobPerNamedPipeProcess(false),
-		bReuseNamedPipeAndProcess(true),
-		PipeGuid(0)
-	{
-	}
-
-	void ReadFromConfigIni()
-	{
-		verify(GConfig->GetBool( TEXT("DevOptions.Shaders"), TEXT("bUseNamedPipes"), bUseNamedPipes, GEngineIni ));
-		verify(GConfig->GetBool( TEXT("DevOptions.Shaders"), TEXT("bUseNamedPipesAsync"), bUseNamedPipesAsync, GEngineIni));
-		verify(GConfig->GetBool( TEXT("DevOptions.Shaders"), TEXT("bSingleJobPerNamedPipeProcess"), bSingleJobPerNamedPipeProcess, GEngineIni));
-		verify(GConfig->GetBool( TEXT("DevOptions.Shaders"), TEXT("bReuseNamedPipeAndProcess"), bReuseNamedPipeAndProcess, GEngineIni));
-	}
-};
-
-FShaderPipeConfig			GShaderPipeConfig;
-
-struct FPipeWorkerInfo
-{
-	FString PipeName;
-	FPlatformNamedPipe NamedPipe;
-
-	// Holds the serialized data for queued jobs to send to the worker process
-	TArray<uint8> WorkJobBuffer;
-
-	enum EState
-	{
-		State_Idle,
-		State_Connecting,
-		State_SendingJobData,
-		State_ReceivingResultSize,
-		State_ReceivingResults,
-	};
-
-	EState State;
-
-	// Holds the size of the response from the worker process
-	int32 ResultsTransferSize;
-
-	// Holds the response from the worker process
-	TArray<uint8> ResultsBuffer;
-
-	FPipeWorkerInfo() :
-		State(State_Idle),
-		ResultsTransferSize(0)
-	{
-	}
-
-	// Updates the state based off async communication with the pipe
-	bool UpdateResultsState()
-	{
-		bool bAgain = false;
-		do
-		{
-			bAgain = false;
-			if (!NamedPipe.UpdateAsyncStatus())
-			{
-				return false;
-			}
-
-			switch (State)
-			{
-				case State_Idle:
-					verify(NamedPipe.OpenConnection());
-					State = State_Connecting;
-					bAgain = true;
-					break;
-
-				case State_Connecting:
-					if (NamedPipe.IsReadyForRW())
-					{
-						if (NamedPipe.WriteBytes(WorkJobBuffer.Num(), WorkJobBuffer.GetData()))
-						{
-							State = State_SendingJobData;
-						}
-						else if (NamedPipe.HasFailed())
-						{
-							return false;
-						}
-						bAgain = true;
-					}
-					break;
-
-				case State_SendingJobData:
-					if (NamedPipe.IsReadyForRW())
-					{
-						// Read the total number of bytes from the response
-						if (NamedPipe.ReadInt32(ResultsTransferSize))
-						{
-							State = State_ReceivingResultSize;
-						}
-						else if (NamedPipe.HasFailed())
-						{
-							return false;
-						}
-						bAgain = true;
-					}
-					break;
-
-				case State_ReceivingResultSize:
-					if (NamedPipe.IsReadyForRW())
-					{
-						// Read the response
-						ResultsBuffer.Empty(ResultsTransferSize);
-						ResultsBuffer.AddUninitialized(ResultsTransferSize);
-						if (NamedPipe.ReadBytes(ResultsTransferSize, ResultsBuffer.GetData()))
-						{
-							State = State_ReceivingResults;
-						}
-						else if (NamedPipe.HasFailed())
-						{
-							return false;
-						}
-						bAgain = true;
-					}
-					break;
-
-				case State_ReceivingResults:
-					if (NamedPipe.IsReadyForRW())
-					{
-						State = State_Idle;
-						return true;
-					}
-					break;
-
-				default:
-					// Unknown state!
-					check(0);
-					bAgain = false;
-					break;
-			}
-		}
-		while (bAgain);
-
-		return false;
-	}
-
-	void CreatePipe(uint32 WorkerIndex, uint32 ProcessId, bool bAllocPipeName)
-	{
-		if (bAllocPipeName)
-		{
-			++GShaderPipeConfig.PipeGuid;
-			PipeName = FString::Printf(TEXT("\\\\.\\Pipe\\ShaderCompiler_%d_%d_%d"), ProcessId, WorkerIndex, GShaderPipeConfig.PipeGuid);
-//FPlatformMisc::LowLevelOutputDebugStringf(TEXT("*** New pipe %s\n"), *PipeName);
-		}
-
-//FPlatformMisc::LowLevelOutputDebugStringf(TEXT("*** Creating pipe %s\n"), *PipeName);
-		if (!NamedPipe.Create(PipeName, true, GShaderPipeConfig.bUseNamedPipesAsync))
-		{
-			UE_LOG(LogShaderCompilers, Fatal, TEXT("ShaderCompileWorker %d couldn't create pipe %s (GetLastError %d)"), WorkerIndex, *PipeName, GetLastError());
-		}
-	}
-
-	void DestroyPipe() 
-	{
-//FPlatformMisc::LowLevelOutputDebugStringf(TEXT("*** Destroying Pipe %s\n"), *NamedPipe.GetName());
-		NamedPipe.Destroy();
-		State = State_Idle;
-	}
-
-	void WriteTasksForPipe( TArray<FShaderCompileJob*>& QueuedJobs ) 
-	{
-		// Make the data we'll transfer through the pipe (but don't send it yet!)
-		WorkJobBuffer.Empty(0);
-		FMemoryWriter TransferWriter(WorkJobBuffer);
-		{
-			TArray<uint8> Buffer;
-			FMemoryWriter BufferWriter(Buffer);
-
-			DoWriteTasks(QueuedJobs, BufferWriter);
-
-			int32 BufferSize = Buffer.Num();
-			TransferWriter << BufferSize;
-			TransferWriter.Serialize(Buffer.GetData(), Buffer.Num());
-		}
-		TransferWriter.Close();
-	}
-};
-#endif	// PLATFORM_SUPPORTS_NAMED_PIPES
-
 /** Information tracked for each shader compile worker process instance. */
 struct FShaderCompileWorkerInfo
 {
@@ -385,13 +194,6 @@ struct FShaderCompileWorkerInfo
 	/** Tracks whether all tasks issued to the worker have been received. */
 	bool bComplete;
 
-#if PLATFORM_SUPPORTS_NAMED_PIPES
-	/** Named pipe */
-	FPipeWorkerInfo PipeWorker;
-
-	bool bWorkerForPipeWasLaunched;
-#endif
-
 	/** Time at which the worker started the most recent batch of tasks. */
 	double StartTime;
 
@@ -402,9 +204,6 @@ struct FShaderCompileWorkerInfo
 		bIssuedTasksToWorker(false),		
 		bLaunchedWorker(false),
 		bComplete(false),
-#if PLATFORM_SUPPORTS_NAMED_PIPES
-		bWorkerForPipeWasLaunched(false),
-#endif
 		StartTime(0)
 	{
 	}
@@ -418,98 +217,6 @@ struct FShaderCompileWorkerInfo
 			FPlatformProcess::CloseProc(WorkerProcess);
 		}
 	}
-
-	void CreatePipeAndNewTask(uint32 WorkerIndex, uint32 ProcessId)
-	{
-#if PLATFORM_SUPPORTS_NAMED_PIPES
-		check(GShaderPipeConfig.bUseNamedPipes);
-		if (QueuedJobs.Num() > 0)
-		{
-			// Open the pipe; figure out if the worker is still listening, which means we can recycle the pipe
-			bool bAllocNameForPipe = (PipeWorker.PipeName.Len() == 0);
-			bAllocNameForPipe |= !GShaderPipeConfig.bReuseNamedPipeAndProcess;
-			bAllocNameForPipe |= !bWorkerForPipeWasLaunched;
-
-			if (bWorkerForPipeWasLaunched)
-			{
-				// Make sure it's alive
-				if (!FShaderCompilingManager::IsShaderCompilerWorkerRunning(WorkerProcess))
-				{
-					bAllocNameForPipe = true;
-					bWorkerForPipeWasLaunched = false;
-				}
-			}
-
-			if (bAllocNameForPipe)
-			{
-				// Request a new worker
-				bWorkerForPipeWasLaunched = false;
-			}
-			PipeWorker.CreatePipe(WorkerIndex, ProcessId, bAllocNameForPipe);
-			PipeWorker.WriteTasksForPipe(QueuedJobs);
-		}
-#else
-		check(0);
-#endif	// PLATFORM_SUPPORTS_NAMED_PIPES
-	}
-
-	void ReadAndWaitResultsFromPipe()
-	{
-#if PLATFORM_SUPPORTS_NAMED_PIPES
-		check(GShaderPipeConfig.bUseNamedPipes);
-		if (QueuedJobs.Num() == 0 || !bWorkerForPipeWasLaunched)
-		{
-			return;
-		}
-
-		check(PipeWorker.NamedPipe.IsCreated());
-
-		bool bError = false;
-		do
-		{
-			PipeWorker.NamedPipe.BlockForAsyncIO();
-			if (PipeWorker.NamedPipe.HasFailed())
-			{
-				bError = true;
-				break;
-			}
-		}
-		while (!PipeWorker.UpdateResultsState());
-
-		if (!bError)
-		{
-			FMemoryReader ResultReader(PipeWorker.ResultsBuffer);
-			DoReadTaskResults(QueuedJobs, ResultReader);
-			bComplete = true;
-		}
-
-		// Close pipe
-		PipeWorker.DestroyPipe();
-#else
-		check(0);
-#endif	// PLATFORM_SUPPORTS_NAMED_PIPES
-	}
-
-	void ReadResultsFromPipe()
-	{
-#if PLATFORM_SUPPORTS_NAMED_PIPES
-		check(PipeWorker.NamedPipe.IsCreated());
-
-		PipeWorker.NamedPipe.BlockForAsyncIO();
-		if (PipeWorker.UpdateResultsState())
-		{
-			FMemoryReader ResultReader(PipeWorker.ResultsBuffer);
-			DoReadTaskResults(QueuedJobs, ResultReader);
-			bComplete = true;
-		}
-
-		// Close pipe
-		PipeWorker.DestroyPipe();
-#else
-		check(0);
-#endif	// PLATFORM_SUPPORTS_NAMED_PIPES
-	}
-
 };
 
 FShaderCompileThreadRunnableBase::FShaderCompileThreadRunnableBase(FShaderCompilingManager* InManager)
@@ -724,106 +431,45 @@ void FShaderCompileThreadRunnable::WriteNewTasks()
 			// 'Only' indicates that the worker should keep checking for more tasks after this one
 			FArchive* TransferFile = NULL;
 
-#if PLATFORM_SUPPORTS_NAMED_PIPES
-			if (GShaderPipeConfig.bUseNamedPipes && !GShaderPipeConfig.bSingleJobPerNamedPipeProcess)
+			int32 RetryCount = 0;
+			// Retry over the next two seconds if we can't write out the input file
+			// Anti-virus and indexing applications can interfere and cause this write to fail
+			//@todo - switch to shared memory or some other method without these unpredictable hazards
+			while (TransferFile == NULL && RetryCount < 2000)
 			{
-				CurrentWorkerInfo.CreatePipeAndNewTask(WorkerIndex, GShaderCompilingManager->ProcessId);
-			}
-			else
-#endif // PLATFORM_SUPPORTS_NAMED_PIPES
-			{
-				int32 RetryCount = 0;
-				// Retry over the next two seconds if we can't write out the input file
-				// Anti-virus and indexing applications can interfere and cause this write to fail
-				//@todo - switch to shared memory or some other method without these unpredictable hazards
-				while (TransferFile == NULL && RetryCount < 2000)
+				if (RetryCount > 0)
 				{
-					if (RetryCount > 0)
-					{
-						FPlatformProcess::Sleep(0.01f);
-					}
-					TransferFile = IFileManager::Get().CreateFileWriter(*TransferFileName, FILEWRITE_EvenIfReadOnly);
-					RetryCount++;
-					if (TransferFile == NULL)
-					{
-						UE_LOG(LogShaderCompilers, Warning, TEXT("Could not create the shader compiler transfer file '%s', retrying..."), *TransferFileName);
-					}
+					FPlatformProcess::Sleep(0.01f);
 				}
+				TransferFile = IFileManager::Get().CreateFileWriter(*TransferFileName, FILEWRITE_EvenIfReadOnly);
+				RetryCount++;
 				if (TransferFile == NULL)
 				{
-					UE_LOG(LogShaderCompilers, Fatal, TEXT("Could not create the shader compiler transfer file '%s'."), *TransferFileName);
-				}
-				check(TransferFile);
-
-				if (!DoWriteTasks(CurrentWorkerInfo.QueuedJobs, *TransferFile))
-				{
-					UE_LOG(LogShaderCompilers, Fatal, TEXT("Could not write the shader compiler transfer filename to '%s'."), *TransferFileName);
-				}
-				delete TransferFile;
-
-				// Change the transfer file name to proper one
-				FString ProperTransferFileName = WorkingDirectory / TEXT("WorkerInputOnly.in");
-				if (!IFileManager::Get().Move(*ProperTransferFileName, *TransferFileName))
-				{
-					UE_LOG(LogShaderCompilers, Fatal, TEXT("Could not rename the shader compiler transfer filename to '%s' from '%s'."), *ProperTransferFileName, *TransferFileName);
+					UE_LOG(LogShaderCompilers, Warning, TEXT("Could not create the shader compiler transfer file '%s', retrying..."), *TransferFileName);
 				}
 			}
-		}
-	}
-}
-
-void FShaderCompileThreadRunnable::LaunchWorkerIfNeeded(FShaderCompileWorkerInfo& CurrentWorkerInfo, uint32 WorkerIndex)
-{
-#if PLATFORM_SUPPORTS_NAMED_PIPES
-	if (CurrentWorkerInfo.QueuedJobs.Num() == 0)
-	{
-		return;
-	}
-
-	if (CurrentWorkerInfo.bWorkerForPipeWasLaunched)
-	{
-		// Check that the worker didn't fatal error or there was an error in the pipe
-		const double CurrentTime = FPlatformTime::Seconds();
-
-		// Limit how often we check for workers running since IsApplicationRunning eats up some CPU time on Windows
-		const bool bCheckForWorkerRunning = (CurrentTime - LastCheckForWorkersTime > 0.1f);
-		if (bCheckForWorkerRunning)
-		{
-			LastCheckForWorkersTime = CurrentTime;
-		}
-
-		if (bCheckForWorkerRunning && !FShaderCompilingManager::IsShaderCompilerWorkerRunning(CurrentWorkerInfo.WorkerProcess))
-		{
-			// Worker died, so clear this pipe and make a new one
-			CurrentWorkerInfo.PipeWorker.DestroyPipe();
-			CurrentWorkerInfo.bWorkerForPipeWasLaunched = false;
+			if (TransferFile == NULL)
+			{
+				UE_LOG(LogShaderCompilers, Fatal, TEXT("Could not create the shader compiler transfer file '%s'."), *TransferFileName);
+			}
+			check(TransferFile);
 
-			// clean up the proc handle
-			FPlatformProcess::CloseProc(CurrentWorkerInfo.WorkerProcess);
-			CurrentWorkerInfo.WorkerProcess = FProcHandle();
+			if (!DoWriteTasks(CurrentWorkerInfo.QueuedJobs, *TransferFile))
+			{
+				UE_LOG(LogShaderCompilers, Fatal, TEXT("Could not write the shader compiler transfer filename to '%s'."), *TransferFileName);
+			}
+			delete TransferFile;
 
-			check(GShaderPipeConfig.bUseNamedPipes && !GShaderPipeConfig.bSingleJobPerNamedPipeProcess);
-			CurrentWorkerInfo.CreatePipeAndNewTask(WorkerIndex, GShaderCompilingManager->ProcessId);
+			// Change the transfer file name to proper one
+			FString ProperTransferFileName = WorkingDirectory / TEXT("WorkerInputOnly.in");
+			if (!IFileManager::Get().Move(*ProperTransferFileName, *TransferFileName))
+			{
+				UE_LOG(LogShaderCompilers, Fatal, TEXT("Could not rename the shader compiler transfer filename to '%s' from '%s'."), *ProperTransferFileName, *TransferFileName);
+			}
 		}
 	}
-
-	if (!CurrentWorkerInfo.bWorkerForPipeWasLaunched)
-	{
-		const FString WorkingDirectory = Manager->ShaderBaseWorkingDirectory + FString::FromInt(WorkerIndex) + TEXT("/");
-
-		// Store the Id with this thread so that we will know not to launch it again
-		const FString& PipeName = CurrentWorkerInfo.PipeWorker.NamedPipe.GetName();
-		// make sure we don't overwrite a running process
-		check(!CurrentWorkerInfo.WorkerProcess.IsValid());
-		CurrentWorkerInfo.WorkerProcess = Manager->LaunchWorker(WorkingDirectory, Manager->ProcessId, WorkerIndex, PipeName, PipeName, true, !GShaderPipeConfig.bReuseNamedPipeAndProcess);		
-		CurrentWorkerInfo.bLaunchedWorker = true;
-		CurrentWorkerInfo.bWorkerForPipeWasLaunched = true;
-	}
-#endif
 }
 
-
-
 bool FShaderCompileThreadRunnable::LaunchWorkersIfNeeded()
 {
 	const double CurrentTime = FPlatformTime::Seconds();
@@ -851,57 +497,48 @@ bool FShaderCompileThreadRunnable::LaunchWorkersIfNeeded()
 			continue;
 		}
 
-#if PLATFORM_SUPPORTS_NAMED_PIPES
-		if (GShaderPipeConfig.bUseNamedPipes && !GShaderPipeConfig.bSingleJobPerNamedPipeProcess)
-		{
-			LaunchWorkerIfNeeded(CurrentWorkerInfo, WorkerIndex);
-		}
-		else
-#endif	// PLATFORM_SUPPORTS_NAMED_PIPES
+		if (!CurrentWorkerInfo.WorkerProcess.IsValid() || (bCheckForWorkerRunning && !FShaderCompilingManager::IsShaderCompilerWorkerRunning(CurrentWorkerInfo.WorkerProcess)))
 		{
-			if (!CurrentWorkerInfo.WorkerProcess.IsValid() || (bCheckForWorkerRunning && !FShaderCompilingManager::IsShaderCompilerWorkerRunning(CurrentWorkerInfo.WorkerProcess)))
+			// @TODO: dubious design - worker should not be launched unless we know there's more work to do.
+			bool bLaunchAgain = true;
+
+			// Detect when the worker has exited due to fatal error
+			// bLaunchedWorker check here is necessary to distinguish between 'process isn't running because it crashed' and 'process isn't running because it exited cleanly and the outputfile was already consumed'
+			if (CurrentWorkerInfo.WorkerProcess.IsValid())
 			{
-				// @TODO: dubious design - worker should not be launched unless we know there's more work to do.
-				bool bLaunchAgain = true;
+				// shader compiler exited one way or another, so clear out the stale PID.
+				FPlatformProcess::CloseProc(CurrentWorkerInfo.WorkerProcess);
+				CurrentWorkerInfo.WorkerProcess = FProcHandle();
 
-				// Detect when the worker has exited due to fatal error
-				// bLaunchedWorker check here is necessary to distinguish between 'process isn't running because it crashed' and 'process isn't running because it exited cleanly and the outputfile was already consumed'
-				if (CurrentWorkerInfo.WorkerProcess.IsValid())
+				if (CurrentWorkerInfo.bLaunchedWorker)
 				{
-					// shader compiler exited one way or another, so clear out the stale PID.
-					FPlatformProcess::CloseProc(CurrentWorkerInfo.WorkerProcess);
-					CurrentWorkerInfo.WorkerProcess = FProcHandle();
+					const FString WorkingDirectory = Manager->AbsoluteShaderBaseWorkingDirectory + FString::FromInt(WorkerIndex) + TEXT("/");
+					const FString OutputFileNameAndPath = WorkingDirectory + TEXT("WorkerOutputOnly.out");
 
-					if (CurrentWorkerInfo.bLaunchedWorker)
+					if (FPlatformFileManager::Get().GetPlatformFile().FileExists(*OutputFileNameAndPath))
 					{
-						const FString WorkingDirectory = Manager->AbsoluteShaderBaseWorkingDirectory + FString::FromInt(WorkerIndex) + TEXT("/");
-						const FString OutputFileNameAndPath = WorkingDirectory + TEXT("WorkerOutputOnly.out");
-
-						if (FPlatformFileManager::Get().GetPlatformFile().FileExists(*OutputFileNameAndPath))
-						{
-							// If the worker is no longer running but it successfully wrote out the output, no need to assert
-							bLaunchAgain = false;
-						}
-						else
-						{
-							UE_LOG(LogShaderCompilers, Error, TEXT("ShaderCompileWorker terminated unexpectedly!  Falling back to directly compiling which will be very slow.  Thread %u."), WorkerIndex);
+						// If the worker is no longer running but it successfully wrote out the output, no need to assert
+						bLaunchAgain = false;
+					}
+					else
+					{
+						UE_LOG(LogShaderCompilers, Error, TEXT("ShaderCompileWorker terminated unexpectedly!  Falling back to directly compiling which will be very slow.  Thread %u."), WorkerIndex);
 
-							bAbandonWorkers = true;
-							break;
-						}
+						bAbandonWorkers = true;
+						break;
 					}
 				}
+			}
 
-				if (bLaunchAgain)
-				{
-					const FString WorkingDirectory = Manager->ShaderBaseWorkingDirectory + FString::FromInt(WorkerIndex) + TEXT("/");
-					FString InputFileName(TEXT("WorkerInputOnly.in"));
-					FString OutputFileName(TEXT("WorkerOutputOnly.out"));
+			if (bLaunchAgain)
+			{
+				const FString WorkingDirectory = Manager->ShaderBaseWorkingDirectory + FString::FromInt(WorkerIndex) + TEXT("/");
+				FString InputFileName(TEXT("WorkerInputOnly.in"));
+				FString OutputFileName(TEXT("WorkerOutputOnly.out"));
 
-					// Store the handle with this thread so that we will know not to launch it again
-					CurrentWorkerInfo.WorkerProcess = Manager->LaunchWorker(WorkingDirectory, Manager->ProcessId, WorkerIndex, InputFileName, OutputFileName, false, false);					
-					CurrentWorkerInfo.bLaunchedWorker = true;
-				}
+				// Store the handle with this thread so that we will know not to launch it again
+				CurrentWorkerInfo.WorkerProcess = Manager->LaunchWorker(WorkingDirectory, Manager->ProcessId, WorkerIndex, InputFileName, OutputFileName);
+				CurrentWorkerInfo.bLaunchedWorker = true;
 			}
 		}
 	}
@@ -918,60 +555,40 @@ void FShaderCompileThreadRunnable::ReadAvailableResults()
 		// Check for available result files
 		if (CurrentWorkerInfo.QueuedJobs.Num() > 0)
 		{
-#if PLATFORM_SUPPORTS_NAMED_PIPES
-			if (GShaderPipeConfig.bUseNamedPipes && !GShaderPipeConfig.bSingleJobPerNamedPipeProcess)
-			{
-				if (!CurrentWorkerInfo.bWorkerForPipeWasLaunched || !CurrentWorkerInfo.PipeWorker.NamedPipe.UpdateAsyncStatus())
-				{
-					continue;
-				}
+			// Distributed compiles always use the same directory
+			const FString WorkingDirectory = Manager->AbsoluteShaderBaseWorkingDirectory + FString::FromInt(WorkerIndex) + TEXT("/");
+			// 'Only' indicates to the worker that it should log and continue checking for the input file after the first one is processed
+			const TCHAR* InputFileName = TEXT("WorkerInputOnly.in");
+			const FString OutputFileNameAndPath = WorkingDirectory + TEXT("WorkerOutputOnly.out");
 
-				if (CurrentWorkerInfo.PipeWorker.UpdateResultsState())
-				{
-					FMemoryReader ResultReader(CurrentWorkerInfo.PipeWorker.ResultsBuffer);
-					DoReadTaskResults(CurrentWorkerInfo.QueuedJobs, ResultReader);
-					CurrentWorkerInfo.bComplete = true;
-					CurrentWorkerInfo.PipeWorker.DestroyPipe();
-				}
-			}
-			else
-#endif // PLATFORM_SUPPORTS_NAMED_PIPES
+			// In the common case the output file will not exist, so check for existence before opening
+			// This is only a win if FileExists is faster than CreateFileReader, which it is on Windows
+			if (FPlatformFileManager::Get().GetPlatformFile().FileExists(*OutputFileNameAndPath))
 			{
-				// Distributed compiles always use the same directory
-				const FString WorkingDirectory = Manager->AbsoluteShaderBaseWorkingDirectory + FString::FromInt(WorkerIndex) + TEXT("/");
-				// 'Only' indicates to the worker that it should log and continue checking for the input file after the first one is processed
-				const TCHAR* InputFileName = TEXT("WorkerInputOnly.in");
-				const FString OutputFileNameAndPath = WorkingDirectory + TEXT("WorkerOutputOnly.out");
-
-				// In the common case the output file will not exist, so check for existence before opening
-				// This is only a win if FileExists is faster than CreateFileReader, which it is on Windows
-				if (FPlatformFileManager::Get().GetPlatformFile().FileExists(*OutputFileNameAndPath))
-				{
-					FArchive* OutputFilePtr = IFileManager::Get().CreateFileReader(*OutputFileNameAndPath, FILEREAD_Silent);
+				FArchive* OutputFilePtr = IFileManager::Get().CreateFileReader(*OutputFileNameAndPath, FILEREAD_Silent);
 
-					if (OutputFilePtr)
-					{
-						FArchive& OutputFile = *OutputFilePtr;
-						check(!CurrentWorkerInfo.bComplete);
+				if (OutputFilePtr)
+				{
+					FArchive& OutputFile = *OutputFilePtr;
+					check(!CurrentWorkerInfo.bComplete);
 						DoReadTaskResults(CurrentWorkerInfo.QueuedJobs, OutputFile);
 
-						// Close the output file.
-						delete OutputFilePtr;
-
-						// Delete the output file now that we have consumed it, to avoid reading stale data on the next compile loop.
-						bool bDeletedOutput = IFileManager::Get().Delete(*OutputFileNameAndPath, true, true);
-						int32 RetryCount = 0;
-						// Retry over the next two seconds if we couldn't delete it
-						while (!bDeletedOutput && RetryCount < 200)
-						{
-							FPlatformProcess::Sleep(0.01f);
-							bDeletedOutput = IFileManager::Get().Delete(*OutputFileNameAndPath, true, true);
-							RetryCount++;
-						}
-						checkf(bDeletedOutput, TEXT("Failed to delete %s!"), *OutputFileNameAndPath);
+					// Close the output file.
+					delete OutputFilePtr;
 
-						CurrentWorkerInfo.bComplete = true;
+					// Delete the output file now that we have consumed it, to avoid reading stale data on the next compile loop.
+					bool bDeletedOutput = IFileManager::Get().Delete(*OutputFileNameAndPath, true, true);
+					int32 RetryCount = 0;
+					// Retry over the next two seconds if we couldn't delete it
+					while (!bDeletedOutput && RetryCount < 200)
+					{
+						FPlatformProcess::Sleep(0.01f);
+						bDeletedOutput = IFileManager::Get().Delete(*OutputFileNameAndPath, true, true);
+						RetryCount++;
 					}
+					checkf(bDeletedOutput, TEXT("Failed to delete %s!"), *OutputFileNameAndPath);
+
+					CurrentWorkerInfo.bComplete = true;
 				}
 			}
 		}
@@ -1043,38 +660,23 @@ int32 FShaderCompileThreadRunnable::CompilingLoop()
 
 	if (Manager->bAllowCompilingThroughWorkers)
 	{
-#if PLATFORM_SUPPORTS_NAMED_PIPES
-		if (GShaderPipeConfig.bUseNamedPipes && GShaderPipeConfig.bSingleJobPerNamedPipeProcess && GShaderCompilingManager)
+		// Write out the files which are input to the shader compile workers
+		WriteNewTasks();
+
+		// Launch shader compile workers if they are not already running
+		// Workers can time out when idle so they may need to be relaunched
+		bool bAbandonWorkers = LaunchWorkersIfNeeded();
+
+		if (bAbandonWorkers)
 		{
-			for (int32 WorkerIndex = 0; WorkerIndex < WorkerInfos.Num(); WorkerIndex++)
-			{
-				FShaderCompileWorkerInfo& CurrentWorkerInfo = *WorkerInfos[WorkerIndex];
-				CurrentWorkerInfo.CreatePipeAndNewTask(WorkerIndex, GShaderCompilingManager->ProcessId);
-				LaunchWorkerIfNeeded(CurrentWorkerInfo, WorkerIndex);
-				CurrentWorkerInfo.ReadAndWaitResultsFromPipe();
-			}
+			// Fall back to local compiles if the SCW crashed.
+			// This is nasty but needed to work around issues where message passing through files to SCW is unreliable on random PCs
+			Manager->bAllowCompilingThroughWorkers = false;
 		}
 		else
-#endif	// PLATFORM_SUPPORTS_NAMED_PIPES
 		{
-			// Write out the files which are input to the shader compile workers
-			WriteNewTasks();
-
-			// Launch shader compile workers if they are not already running
-			// Workers can time out when idle so they may need to be relaunched
-			bool bAbandonWorkers = LaunchWorkersIfNeeded();
-
-			if (bAbandonWorkers)
-			{
-				// Fall back to local compiles if the SCW crashed.
-				// This is nasty but needed to work around issues where message passing through files to SCW is unreliable on random PCs
-				Manager->bAllowCompilingThroughWorkers = false;
-			}
-			else
-			{
-				// Read files which are outputs from the shader compile workers
-				ReadAvailableResults();
-			}
+			// Read files which are outputs from the shader compile workers
+			ReadAvailableResults();
 		}
 	}
 	else
@@ -1135,10 +737,6 @@ FShaderCompilingManager::FShaderCompilingManager() :
 	verify(GConfig->GetBool( TEXT("DevOptions.Shaders"), TEXT("bPromptToRetryFailedShaderCompiles"), bPromptToRetryFailedShaderCompiles, GEngineIni ));
 	verify(GConfig->GetBool( TEXT("DevOptions.Shaders"), TEXT("bLogJobCompletionTimes"), bLogJobCompletionTimes, GEngineIni ));
 
-#if PLATFORM_SUPPORTS_NAMED_PIPES
-	GShaderPipeConfig.ReadFromConfigIni();
-#endif
-
 	GRetryShaderCompilation = bPromptToRetryFailedShaderCompiles;
 
 	verify(GConfig->GetFloat( TEXT("DevOptions.Shaders"), TEXT("ProcessGameThreadTargetTime"), ProcessGameThreadTargetTime, GEngineIni ));
@@ -1250,7 +848,7 @@ void FShaderCompilingManager::AddJobs(TArray<FShaderCompileJob*>& NewJobs, bool
 }
 
 /** Launches the worker, returns the launched process handle. */
-FProcHandle FShaderCompilingManager::LaunchWorker(const FString& WorkingDirectory, uint32 InProcessId, uint32 ThreadId, const FString& WorkerInputFile, const FString& WorkerOutputFile, bool bUseNamedPipes, bool bSingleConnectionPipe)
+FProcHandle FShaderCompilingManager::LaunchWorker(const FString& WorkingDirectory, uint32 InProcessId, uint32 ThreadId, const FString& WorkerInputFile, const FString& WorkerOutputFile)
 {
 	// Setup the parameters that the worker application needs
 	// Surround the working directory with double quotes because it may contain a space 
@@ -1260,14 +858,7 @@ FProcHandle FShaderCompilingManager::LaunchWorker(const FString& WorkingDirector
 	FString WorkerAbsoluteDirectory = IFileManager::Get().ConvertToAbsolutePathForExternalAppForWrite(*WorkingDirectory);
 	FPaths::NormalizeDirectoryName(WorkerAbsoluteDirectory);
 	FString WorkerParameters = FString(TEXT("\"")) + WorkerAbsoluteDirectory + TEXT("/\" ") + FString::FromInt(InProcessId) + TEXT(" ") + FString::FromInt(ThreadId) + TEXT(" ") + WorkerInputFile + TEXT(" ") + WorkerOutputFile;
-	if (bUseNamedPipes)
-	{
-		WorkerParameters += FString(bSingleConnectionPipe ? TEXT(" -communicatethroughnamedpipeonce ") : TEXT(" -communicatethroughnamedpipe "));
-	}
-	else
-	{
-		WorkerParameters += FString(TEXT(" -communicatethroughfile "));
-	}
+	WorkerParameters += FString(TEXT(" -communicatethroughfile "));
 	if ( GIsBuildMachine )
 	{
 		WorkerParameters += FString(TEXT(" -buildmachine "));
diff --git a/Engine/Source/Runtime/Engine/Public/ShaderCompiler.h b/Engine/Source/Runtime/Engine/Public/ShaderCompiler.h
--- a/Engine/Source/Runtime/Engine/Public/ShaderCompiler.h
+++ b/Engine/Source/Runtime/Engine/Public/ShaderCompiler.h
@@ -131,9 +131,6 @@ private:
 
 	/** Main work loop. */
 	virtual int32 CompilingLoop() override;
-
-	/** Used when compiling through workers, launches worker processes if needed. */
-	void LaunchWorkerIfNeeded(FShaderCompileWorkerInfo & CurrentWorkerInfo, uint32 WorkerIndex);
 };
 
 class FShaderCompileXGEThreadRunnable : public FShaderCompileThreadRunnableBase
@@ -326,7 +323,7 @@ private:
 	double WorkersBusyTime;
 
 	/** Launches the worker, returns the launched process handle. */
-	FProcHandle LaunchWorker(const FString& WorkingDirectory, uint32 ProcessId, uint32 ThreadId, const FString& WorkerInputFile, const FString& WorkerOutputFile, bool bUseNamedPipes, bool bSingleConnectionPipe);
+	FProcHandle LaunchWorker(const FString& WorkingDirectory, uint32 ProcessId, uint32 ThreadId, const FString& WorkerInputFile, const FString& WorkerOutputFile);
 
 	/** Blocks on completion of the given shader maps. */
 	void BlockOnShaderMapCompletion(const TArray<int32>& ShaderMapIdsToFinishCompiling, TMap<int32, FShaderMapFinalizeResults>& CompiledShaderMaps);
