Reviving old SpawnActorDeferred

diff --git a/Engine/Source/Runtime/Engine/Classes/Engine/World.h b/Engine/Source/Runtime/Engine/Classes/Engine/World.h
--- a/Engine/Source/Runtime/Engine/Classes/Engine/World.h
+++ b/Engine/Source/Runtime/Engine/Classes/Engine/World.h
@@ -2848,6 +2848,33 @@ public:
 	* manually by calling UGameplayStatics::FinishSpawningActor (see AActor::OnConstruction).
 	*/
 	template< class T >
+	T* SpawnActorDeferred(
+		UClass* Class,
+		FVector const& Location,
+		FRotator const& Rotation,
+		AActor* Owner=NULL,
+		APawn* Instigator=NULL,
+		bool bNoCollisionFail=false
+		)
+	{
+		if( Owner )
+		{
+			check(this==Owner->GetWorld());
+		}
+		FActorSpawnParameters SpawnInfo;
+		SpawnInfo.bNoCollisionFail = bNoCollisionFail;
+		SpawnInfo.Owner = Owner;
+		SpawnInfo.Instigator = Instigator;
+		SpawnInfo.bDeferConstruction = true;
+		return (Class != NULL) ? Cast<T>(SpawnActor(Class, &Location, &Rotation, SpawnInfo )) : NULL;
+	}
+
+	/**
+	* Spawns given class and returns class T pointer, forcibly sets world transform(note this allows scale as well). WILL NOT run Construction Script of Blueprints 
+	* to give caller an opportunity to set parameters beforehand.  Caller is responsible for invoking construction
+	* manually by calling UGameplayStatics::FinishSpawningActor (see AActor::OnConstruction).
+	*/
+	template< class T >
 	T* SpawnActorDeferred(
 		UClass* Class,
 		FTransform const& Transform,
