fix for potential crash when adding navigation data to existing node in nav octree
#ue4

diff --git a/Engine/Source/Runtime/Engine/Private/AI/Navigation/NavigationSystem.cpp b/Engine/Source/Runtime/Engine/Private/AI/Navigation/NavigationSystem.cpp
--- a/Engine/Source/Runtime/Engine/Private/AI/Navigation/NavigationSystem.cpp
+++ b/Engine/Source/Runtime/Engine/Private/AI/Navigation/NavigationSystem.cpp
@@ -2366,7 +2366,7 @@ void UNavigationSystem::AddElementToNavOctree(const FNavigationDirtyElement& Dir
 		}
 
 		const FOctreeElementId* UseParentId = ParentId ? ParentId : GetObjectsNavOctreeId(NavigationParent);
-		if (UseParentId)
+		if (UseParentId && NavOctree->IsValidElementId(*UseParentId))
 		{
 			UE_LOG(LogNavOctree, Log, TEXT("ADD %s to %s"), *GetNameSafe(ElementOwner), *GetNameSafe(NavigationParent));
 			NavOctree->AppendToNode(*UseParentId, DirtyElement.NavInterface, ElementBounds, GeneratedData);
