Fix for incorrectly weighted states in a custom blend

diff --git a/Engine/Source/Runtime/Engine/Private/Animation/AnimNode_StateMachine.cpp b/Engine/Source/Runtime/Engine/Private/Animation/AnimNode_StateMachine.cpp
--- a/Engine/Source/Runtime/Engine/Private/Animation/AnimNode_StateMachine.cpp
+++ b/Engine/Source/Runtime/Engine/Private/Animation/AnimNode_StateMachine.cpp
@@ -545,7 +545,8 @@ void FAnimNode_StateMachine::UpdateTransitionStates(const FAnimationUpdateContex
 						{
 							const bool bUsePreviousState = (Evaluator->DataSource == EEvaluatorDataSource::EDS_SourcePose);
 							const int32 EffectiveStateIndex = bUsePreviousState ? Transition.PreviousState : Transition.NextState;
-							UpdateState(EffectiveStateIndex, Context);
+							FAnimationUpdateContext ContextToUse = Context.FractionalWeight(bUsePreviousState ? (1.0f - Transition.Alpha) : Transition.Alpha);
+							UpdateState(EffectiveStateIndex, ContextToUse);
 						}
 					}
 				}
