Fixed a bug with navmesh rendering component creation when there was already one saved with the level #UE4

UE-17650

diff --git a/Engine/Source/Runtime/Engine/Private/AI/Navigation/RecastNavMesh.cpp b/Engine/Source/Runtime/Engine/Private/AI/Navigation/RecastNavMesh.cpp
--- a/Engine/Source/Runtime/Engine/Private/AI/Navigation/RecastNavMesh.cpp
+++ b/Engine/Source/Runtime/Engine/Private/AI/Navigation/RecastNavMesh.cpp
@@ -441,7 +441,7 @@ UPrimitiveComponent* ARecastNavMesh::ConstructRenderingComponentImpl()
 
 	return NULL;
 #else
-	return NewObject<UNavMeshRenderingComponent>(this, TEXT("NavMeshRenderer"), RF_Transient);
+	return NewObject<UNavMeshRenderingComponent>(this, TEXT("NavRenderingComp"), RF_Transient);
 #endif // DO_NAVMESH_DEBUG_DRAWING_PER_TILE
 }
 
