Added safety checks for Navigation Data validity when iterating over NavigationSystem.NavDataSet #UE4

UE-18362

diff --git a/Engine/Source/Runtime/Engine/Private/AI/Navigation/NavigationSystem.cpp b/Engine/Source/Runtime/Engine/Private/AI/Navigation/NavigationSystem.cpp
--- a/Engine/Source/Runtime/Engine/Private/AI/Navigation/NavigationSystem.cpp
+++ b/Engine/Source/Runtime/Engine/Private/AI/Navigation/NavigationSystem.cpp
@@ -658,7 +658,10 @@ void UNavigationSystem::OnWorldInitDone(FNavigationSystemRunMode Mode)
 			{
 				for (ANavigationData* NavData : NavDataSet)
 				{
-					NavData->OnStreamingLevelAdded(Level, World);
+					if (NavData)
+					{
+						NavData->OnStreamingLevelAdded(Level, World);
+					}
 				}
 			}
 		}
@@ -671,7 +674,10 @@ void UNavigationSystem::OnWorldInitDone(FNavigationSystemRunMode Mode)
 			// for all the actors in sublevels
 			for (ANavigationData* NavData : NavDataSet)
 			{
-				NavData->MarkAsNeedingUpdate();
+				if (NavData)
+				{
+					NavData->MarkAsNeedingUpdate();
+				}
 			}
 		}
 	}
@@ -3329,7 +3335,10 @@ void UNavigationSystem::OnLevelAddedToWorld(ULevel* InLevel, UWorld* InWorld)
 		{
 			for (ANavigationData* NavData : NavDataSet)
 			{
-				NavData->OnStreamingLevelAdded(InLevel, InWorld);
+				if (NavData)
+				{
+					NavData->OnStreamingLevelAdded(InLevel, InWorld);
+				}
 			}
 		}
 	}
@@ -3345,7 +3354,10 @@ void UNavigationSystem::OnLevelRemovedFromWorld(ULevel* InLevel, UWorld* InWorld
 		{
 			for (ANavigationData* NavData : NavDataSet)
 			{
-				NavData->OnStreamingLevelRemoved(InLevel, InWorld);
+				if (NavData)
+				{
+					NavData->OnStreamingLevelRemoved(InLevel, InWorld);
+				}
 			}
 		}
 	}
