Fixed crash when renaming a folder in the world outliner

AActor::SetFolderPath had gained too much responsibility (recursion, detachment) which has now been devolved to the calling code. Any code that wishes to recursively set a folder path should now use AActor::SetFolderPath_Recursively and manage detachment themselves, if so desired.

Added traversal algortihms for actors in FActorEditorUtils.

This addresses UE-18115.

diff --git a/Engine/Source/Editor/SceneOutliner/Private/ActorTreeItem.cpp b/Engine/Source/Editor/SceneOutliner/Private/ActorTreeItem.cpp
--- a/Engine/Source/Editor/SceneOutliner/Private/ActorTreeItem.cpp
+++ b/Engine/Source/Editor/SceneOutliner/Private/ActorTreeItem.cpp
@@ -186,7 +186,8 @@ void FActorDropTarget::DetachActorFromParent(AActor* ChildActor)
 		AActor* OldParent = RootComp->AttachParent->GetOwner();
 		OldParent->Modify();
 		RootComp->DetachFromParent(true);
-		ChildActor->SetFolderPath(OldParent->GetFolderPath());
+		
+		ChildActor->SetFolderPath_Recursively(OldParent->GetFolderPath());
 	}
 }
 
diff --git a/Engine/Source/Editor/SceneOutliner/Private/FolderTreeItem.cpp b/Engine/Source/Editor/SceneOutliner/Private/FolderTreeItem.cpp
--- a/Engine/Source/Editor/SceneOutliner/Private/FolderTreeItem.cpp
+++ b/Engine/Source/Editor/SceneOutliner/Private/FolderTreeItem.cpp
@@ -122,14 +122,42 @@ void FFolderDropTarget::OnDrop(FDragDropPayload& DraggedObjects, UWorld& World,
 		}
 	}
 
+	// Set the folder path on all the dragged actors, and detach any that need to be moved
 	if (DraggedObjects.Actors)
 	{
+		TSet<AActor*> ParentActors;
+		TSet<AActor*> ChildActors;
+
 		for (auto WeakActor : DraggedObjects.Actors.GetValue())
 		{
 			AActor* Actor = WeakActor.Get();
 			if (Actor)
 			{
+				// First mark this object as a parent, then set its children's path
+				ParentActors.Add(Actor);
 				Actor->SetFolderPath(DestinationPath);
+
+				FActorEditorUtils::TraverseActorTree_ParentFirst(Actor, [&](AActor* InActor){
+					ChildActors.Add(InActor);
+					InActor->SetFolderPath(DestinationPath);
+					return true;
+				}, false);
+			}
+		}
+
+		// Detach parent actors
+		for (AActor* Parent : ParentActors)
+		{
+			auto* RootComp = Parent->GetRootComponent();
+
+			// We don't detach if it's a child of another that's been dragged
+			if (RootComp && RootComp->AttachParent && !ChildActors.Contains(Parent))
+			{
+				if (AActor* OldParentActor = RootComp->AttachParent->GetOwner())
+				{
+					OldParentActor->Modify();
+				}
+				RootComp->DetachFromParent(true);
 			}
 		}
 	}
@@ -219,8 +247,8 @@ void FFolderTreeItem::Delete()
 		{
 			if (AActor* Actor = ActorItem.Actor.Get())
 			{
-				Actor->SetFolderPath(ParentPath);
-			}	
+				Actor->SetFolderPath_Recursively(ParentPath);
+			}
 		}
 		virtual void Visit(FFolderTreeItem& FolderItem) const override
 		{
diff --git a/Engine/Source/Editor/UnrealEd/Private/EditorActorFolders.cpp b/Engine/Source/Editor/UnrealEd/Private/EditorActorFolders.cpp
--- a/Engine/Source/Editor/UnrealEd/Private/EditorActorFolders.cpp
+++ b/Engine/Source/Editor/UnrealEd/Private/EditorActorFolders.cpp
@@ -335,7 +335,7 @@ void FActorFolders::SetSelectedFolderPath(FName Path) const
 			}
 		}
 
-		Actor->SetFolderPath(Path);
+		Actor->SetFolderPath_Recursively(Path);
 	}
 }
 
@@ -426,7 +426,8 @@ bool FActorFolders::RenameFolderInWorld(UWorld& World, FName OldPath, FName NewP
 		if (Actor->GetFolderPath() == OldPath || PathIsChildOf(OldActorPath.ToString(), OldPathString))
 		{
 			RenamedFolders.Add(OldActorPath);
-			ActorIt->SetFolderPath(OldPathToNewPath(OldPathString, NewPathString, OldActorPath.ToString()), false);
+
+			ActorIt->SetFolderPath_Recursively(OldPathToNewPath(OldPathString, NewPathString, OldActorPath.ToString()));
 		}
 	}
 
diff --git a/Engine/Source/Editor/UnrealEd/Private/EditorEngine.cpp b/Engine/Source/Editor/UnrealEd/Private/EditorEngine.cpp
--- a/Engine/Source/Editor/UnrealEd/Private/EditorEngine.cpp
+++ b/Engine/Source/Editor/UnrealEd/Private/EditorEngine.cpp
@@ -3494,7 +3494,7 @@ bool UEditorEngine::DetachSelectedActors()
 			OldParentActor->Modify();
 			RootComp->DetachFromParent(true);
 			bDetachOccurred = true;
-			Actor->SetFolderPath(OldParentActor->GetFolderPath());
+			Actor->SetFolderPath_Recursively(OldParentActor->GetFolderPath());
 		}
 	}
 	return bDetachOccurred;
diff --git a/Engine/Source/Runtime/Engine/Classes/GameFramework/Actor.h b/Engine/Source/Runtime/Engine/Classes/GameFramework/Actor.h
--- a/Engine/Source/Runtime/Engine/Classes/GameFramework/Actor.h
+++ b/Engine/Source/Runtime/Engine/Classes/GameFramework/Actor.h
@@ -1622,10 +1622,15 @@ public:
 
 	/**
 	 * Assigns a new folder to this actor. Actor folder paths are only available in development builds.
-	 * @param	NewFolderPath	   The new folder to assign to the actor.
-	 * @param   bDetachFromParent  whether to detach this actor from its parent when setting the folder path
+	 * @param	NewFolderPath		The new folder to assign to the actor.
 	 */
-	void SetFolderPath(const FName& NewFolderPath, bool bDetachFromParent = true);
+	void SetFolderPath(const FName& NewFolderPath);
+
+	/**
+	 * Assigns a new folder to this actor and any attached children. Actor folder paths are only available in development builds.
+	 * @param	NewFolderPath		The new folder to assign to the actors.
+	 */
+	void SetFolderPath_Recursively(const FName& NewFolderPath);
 
 	/**
 	 * Used by the "Sync to Content Browser" right-click menu option in the editor.
diff --git a/Engine/Source/Runtime/Engine/Private/ActorEditor.cpp b/Engine/Source/Runtime/Engine/Private/ActorEditor.cpp
--- a/Engine/Source/Runtime/Engine/Private/ActorEditor.cpp
+++ b/Engine/Source/Runtime/Engine/Private/ActorEditor.cpp
@@ -683,50 +683,30 @@ const FName& AActor::GetFolderPath() const
 	return FolderPath;
 }
 
-void AActor::SetFolderPath(const FName& NewFolderPath, bool bDetachFromParent)
+void AActor::SetFolderPath(const FName& NewFolderPath)
 {
-	// Detach the actor if it is attached
-	USceneComponent* RootComp = GetRootComponent();
-	const bool bIsAttached = RootComp  && RootComp->AttachParent;
-
-	if (NewFolderPath == FolderPath && !bIsAttached)
+	if (NewFolderPath != FolderPath)
 	{
-		return;
-	}
-
-	Modify();
-
-	FName OldPath = FolderPath;
-	FolderPath = NewFolderPath;
-	
-	// Detach the actor if it is attached
-	if (RootComp && bIsAttached && bDetachFromParent)
-	{
-		AActor* OldParentActor = RootComp->AttachParent->GetOwner();
-		OldParentActor->Modify();
-
-		RootComp->DetachFromParent(true);
-	}
+		Modify();
 
-	if (GEngine)
-	{
-		GEngine->BroadcastLevelActorFolderChanged(this, OldPath);
-	}
+		FName OldPath = FolderPath;
+		FolderPath = NewFolderPath;
 
-	//recursively change folder path for children (but do not detach so they remain as child)
-	if(RootComp)
-	{
-		for(auto ChildSceneComponent : RootComp->AttachChildren)
+		if (GEngine)
 		{
-			AActor* ChildActor = ChildSceneComponent ? ChildSceneComponent->GetOwner() : nullptr;
-			if(ChildActor)
-			{
-				ChildActor->SetFolderPath(NewFolderPath, false);
-			}
+			GEngine->BroadcastLevelActorFolderChanged(this, OldPath);
 		}
 	}
 }
 
+void AActor::SetFolderPath_Recursively(const FName& NewFolderPath)
+{
+	FActorEditorUtils::TraverseActorTree_ParentFirst(this, [&](AActor* InActor){
+		InActor->SetFolderPath(NewFolderPath);
+		return true;
+	});
+}
+
 void AActor::CheckForDeprecated()
 {
 	if ( GetClass()->HasAnyClassFlags(CLASS_Deprecated) )
diff --git a/Engine/Source/Runtime/Engine/Private/ActorEditorUtils.cpp b/Engine/Source/Runtime/Engine/Private/ActorEditorUtils.cpp
--- a/Engine/Source/Runtime/Engine/Private/ActorEditorUtils.cpp
+++ b/Engine/Source/Runtime/Engine/Private/ActorEditorUtils.cpp
@@ -61,6 +61,49 @@ namespace FActorEditorUtils
 			}
 		}
 	}
+
+	bool TraverseActorTree_ParentFirst(AActor* InActor, const TFunctionRef<bool(AActor*)>& InPredicate, bool bIncludeThisActor)
+	{
+		if (!InActor)
+		{
+			return true;
+		}
+
+		if (bIncludeThisActor && !InPredicate(InActor))
+		{
+			return false;
+		}
+
+		for (auto* ChildComponent : InActor->GetRootComponent()->AttachChildren)
+		{
+			AActor* ChildActor = ChildComponent ? ChildComponent->GetOwner() : nullptr;
+			if (ChildActor && ChildActor != InActor && !TraverseActorTree_ParentFirst(ChildActor, InPredicate))
+			{
+				return false;
+			}
+		}
+		
+		return true;
+	}
+
+	bool TraverseActorTree_ChildFirst(AActor* InActor, const TFunctionRef<bool(AActor*)>& InPredicate, bool bIncludeThisActor)
+	{
+		if (!InActor)
+		{
+			return true;
+		}
+
+		for (auto* ChildComponent : InActor->GetRootComponent()->AttachChildren)
+		{
+			AActor* ChildActor = ChildComponent ? ChildComponent->GetOwner() : nullptr;
+			if (ChildActor && ChildActor != InActor && !TraverseActorTree_ChildFirst(ChildActor, InPredicate))
+			{
+				return false;
+			}
+		}
+
+		return !bIncludeThisActor || InPredicate(InActor);
+	}
 }
 
 
diff --git a/Engine/Source/Runtime/Engine/Public/ActorEditorUtils.h b/Engine/Source/Runtime/Engine/Public/ActorEditorUtils.h
--- a/Engine/Source/Runtime/Engine/Public/ActorEditorUtils.h
+++ b/Engine/Source/Runtime/Engine/Public/ActorEditorUtils.h
@@ -30,6 +30,26 @@ namespace FActorEditorUtils
 	 * @param OutEditableComponents The populated list of editable components on InActor
 	 */
 	ENGINE_API void GetEditableComponents( const AActor* InActor, TArray<UActorComponent*>& OutEditableComponents );
+
+	/**
+	 * Traverse an actor sub-tree (defined by attached children), parent first.
+	 * 
+	 * @param InActor				The Actor to traverse. InPredicate is only invoked for this actor where bIncludeThisActor is true
+	 * @param InPredicate			Predicate to call for each actor of the traversal. Return true to continue, false to terminate the traversal
+	 * @param bIncludeThisActor		True to include this actor in the traversal, false otherwise
+	 * @return true on expected completion, false where the client prematurely terminated the traversal
+	 */
+	ENGINE_API bool TraverseActorTree_ParentFirst(AActor* InActor, const TFunctionRef<bool(AActor*)>& InPredicate, bool bIncludeThisActor = true);
+
+	/**
+	 * Traverse an actor sub-tree (defined by attached children), child first.
+	 * 
+	 * @param InActor				The Actor to traverse. InPredicate is only invoked for this actor where bIncludeThisActor is true
+	 * @param InPredicate			Predicate to call for each actor of the traversal. Return true to continue, false to terminate the traversal
+	 * @param bIncludeThisActor		True to include this actor in the traversal, false otherwise
+	 * @return true on expected completion, false where the client prematurely terminated the traversal
+	 */
+	ENGINE_API bool TraverseActorTree_ChildFirst(AActor* InActor, const TFunctionRef<bool(AActor*)>& InPredicate, bool bIncludeThisActor = true);
 };
 
 #endif// __ActorEditorUtils_h__
