Fixed FloatingPawnMovement component's issue when controlling AI pawn following a path #UE4

Basically the speed clamping wasn't working correctly

diff --git a/Engine/Source/Runtime/Engine/Private/FloatingPawnMovement.cpp b/Engine/Source/Runtime/Engine/Private/FloatingPawnMovement.cpp
--- a/Engine/Source/Runtime/Engine/Private/FloatingPawnMovement.cpp
+++ b/Engine/Source/Runtime/Engine/Private/FloatingPawnMovement.cpp
@@ -39,7 +39,7 @@ void UFloatingPawnMovement::TickComponent(float DeltaTime, enum ELevelTick TickT
 		}
 		// if it's not player controller, but we do have a controller, then it's AI
 		// (that's not following a path) and we need to limit the speed
-		else if (IsExceedingMaxSpeed(MaxSpeed) == false)
+		else if (IsExceedingMaxSpeed(MaxSpeed) == true)
 		{
 			Velocity = Velocity.GetUnsafeNormal() * MaxSpeed;
 		}
