If TickInterval is set to 0 move function out of cooldown linked list and in to all enabled tick functions set

diff --git a/Engine/Source/Runtime/Engine/Private/TickTaskManager.cpp b/Engine/Source/Runtime/Engine/Private/TickTaskManager.cpp
--- a/Engine/Source/Runtime/Engine/Private/TickTaskManager.cpp
+++ b/Engine/Source/Runtime/Engine/Private/TickTaskManager.cpp
@@ -511,10 +511,18 @@ public:
 				CumulativeCooldown += TickFunction->RelativeTickCooldown;
 
 				TickFunction->TickState = FTickFunction::ETickState::Enabled;
-				// we store a bit in here if this came from an interval queue 
-				AllTickFunctions.Add((FTickFunction*)(UPTRINT(TickFunction) | 1));
 
-				TickFunctionsToReschedule.Add(FTickScheduleDetails(TickFunction, TickFunction->TickInterval - (Context.DeltaSeconds - CumulativeCooldown))); // Give credit for any overrun
+				if (TickFunction->TickInterval > 0.f)
+				{
+					// we store a bit in here if this came from an interval queue 
+					AllTickFunctions.Add((FTickFunction*)(UPTRINT(TickFunction) | 1));
+
+					TickFunctionsToReschedule.Add(FTickScheduleDetails(TickFunction, TickFunction->TickInterval - (Context.DeltaSeconds - CumulativeCooldown))); // Give credit for any overrun
+				}
+				else
+				{
+					AllEnabledTickFunctions.Add(TickFunction);
+				}
 
 				AllCoolingDownTickFunctions.Head = TickFunction->Next;
 				TickFunction = TickFunction->Next;
@@ -547,7 +555,7 @@ public:
 		TickFunctionsToReschedule.Add(FTickScheduleDetails(TickFunction, TickFunction->TickInterval));
 	}
 
-	/* Puts a TickFunction in to the cooldown state*/
+	/* Schedules tick functions for cooldown */
 	void ScheduleTickFunctionCooldowns()
 	{
 		if (TickFunctionsToReschedule.Num() > 0)
@@ -565,33 +573,41 @@ public:
 			FTickFunction* ComparisonTickFunction = AllCoolingDownTickFunctions.Head;
 			while (ComparisonTickFunction && RescheduleIndex < TickFunctionsToReschedule.Num())
 			{
-				const float CooldownTime = TickFunctionsToReschedule[RescheduleIndex].Cooldown;
-				if ((CumulativeCooldown + ComparisonTickFunction->RelativeTickCooldown) > CooldownTime)
+				FTickFunction* TickFunction = TickFunctionsToReschedule[RescheduleIndex].TickFunction;
+				if (TickFunction->TickInterval > 0.f)
 				{
-					FTickFunction* TickFunction = TickFunctionsToReschedule[RescheduleIndex].TickFunction;
-					TickFunction->TickState = FTickFunction::ETickState::CoolingDown;
-					TickFunction->RelativeTickCooldown = CooldownTime - CumulativeCooldown;
-
-					if (PrevComparisonTickFunction)
+					const float CooldownTime = TickFunctionsToReschedule[RescheduleIndex].Cooldown;
+					if ((CumulativeCooldown + ComparisonTickFunction->RelativeTickCooldown) > CooldownTime)
 					{
-						PrevComparisonTickFunction->Next = TickFunction;
+						TickFunction->TickState = FTickFunction::ETickState::CoolingDown;
+						TickFunction->RelativeTickCooldown = CooldownTime - CumulativeCooldown;
+
+						if (PrevComparisonTickFunction)
+						{
+							PrevComparisonTickFunction->Next = TickFunction;
+						}
+						else
+						{
+							check(ComparisonTickFunction == AllCoolingDownTickFunctions.Head);
+							AllCoolingDownTickFunctions.Head = TickFunction;
+						}
+						TickFunction->Next = ComparisonTickFunction;
+						PrevComparisonTickFunction = TickFunction;
+						ComparisonTickFunction->RelativeTickCooldown -= TickFunction->RelativeTickCooldown;
+						CumulativeCooldown += TickFunction->RelativeTickCooldown;
+						++RescheduleIndex;
 					}
 					else
 					{
-						check(ComparisonTickFunction == AllCoolingDownTickFunctions.Head);
-						AllCoolingDownTickFunctions.Head = TickFunction;
+						CumulativeCooldown += ComparisonTickFunction->RelativeTickCooldown;
+						PrevComparisonTickFunction = ComparisonTickFunction;
+						ComparisonTickFunction = ComparisonTickFunction->Next;
 					}
-					TickFunction->Next = ComparisonTickFunction;
-					PrevComparisonTickFunction = TickFunction;
-					ComparisonTickFunction->RelativeTickCooldown -= TickFunction->RelativeTickCooldown;
-					CumulativeCooldown += TickFunction->RelativeTickCooldown;
-					++RescheduleIndex;
 				}
 				else
 				{
-					CumulativeCooldown += ComparisonTickFunction->RelativeTickCooldown;
-					PrevComparisonTickFunction = ComparisonTickFunction;
-					ComparisonTickFunction = ComparisonTickFunction->Next;
+					AllEnabledTickFunctions.Add(TickFunction);
+					++RescheduleIndex;
 				}
 			}
 			for ( ; RescheduleIndex < TickFunctionsToReschedule.Num(); ++RescheduleIndex)
@@ -643,7 +659,14 @@ public:
 			{
 				CumulativeCooldown += TickFunction->RelativeTickCooldown;
 				TickFunction->QueueTickFunction(TTS, Context);
-				TickFunctionsToReschedule.Add(FTickScheduleDetails(TickFunction, TickFunction->TickInterval - (Context.DeltaSeconds - CumulativeCooldown))); // Give credit for any overrun
+				if (TickFunction->TickInterval > 0.f)
+				{
+					TickFunctionsToReschedule.Add(FTickScheduleDetails(TickFunction, TickFunction->TickInterval - (Context.DeltaSeconds - CumulativeCooldown))); // Give credit for any overrun
+				}
+				else
+				{
+					AllEnabledTickFunctions.Add(TickFunction);
+				}
 				AllCoolingDownTickFunctions.Head = TickFunction->Next;
 			}
 			else
@@ -699,12 +722,21 @@ public:
 			{
 				if (CumulativeCooldown < InContext.DeltaSeconds)
 				{
+					TickFunction->TickState = FTickFunction::ETickState::Enabled;
+
 					TickFunction->TickVisitedGFrameCounter = GFrameCounter;
 					TickFunction->TickQueuedGFrameCounter = GFrameCounter;
 					TickFunction->ExecuteTick(InContext.DeltaSeconds, InContext.TickType, ENamedThreads::GameThread, FGraphEventRef());
 					TickFunction->TaskPointer = nullptr; // this is stale, clear it out now
 
-					TickFunctionsToReschedule.Add(FTickScheduleDetails(TickFunction, TickFunction->TickInterval - (InContext.DeltaSeconds - CumulativeCooldown))); // Give credit for any overrun
+					if (TickFunction->TickInterval > 0.f)
+					{
+						TickFunctionsToReschedule.Add(FTickScheduleDetails(TickFunction, TickFunction->TickInterval - (InContext.DeltaSeconds - CumulativeCooldown))); // Give credit for any overrun
+					}
+					else
+					{
+						AllEnabledTickFunctions.Add(TickFunction);
+					}
 				}
 				else
 				{
@@ -765,6 +797,7 @@ public:
 	{
 		return AllEnabledTickFunctions.Contains(TickFunction) || AllDisabledTickFunctions.Contains(TickFunction) || AllCoolingDownTickFunctions.Contains(TickFunction);
 	}
+
 	/** Add the tick function to the master list **/
 	void AddTickFunction(FTickFunction* TickFunction)
 	{
