[UE-18152] Renaming a component and then adding a new component of the same type will no longer result in broken defaults editing.

change summary:
- modified USCS_Node::NameWasModified() to ensure that the component template name stays in sync with the new variable name. this in turn ensures that if a new SCS node is added with the same name as the original name of the existing node, the existing node's template will no longer be recycled and incorrectly shared with the new node.
- modified USCS_Node::PostLoad() to ensure that component template names are in sync w/ the variable name on load; if not, then a new instance is created with the proper name in order to ensure that the template association is unique after loading.

diff --git a/Engine/Source/Runtime/Engine/Private/SCS_Node.cpp b/Engine/Source/Runtime/Engine/Private/SCS_Node.cpp
--- a/Engine/Source/Runtime/Engine/Private/SCS_Node.cpp
+++ b/Engine/Source/Runtime/Engine/Private/SCS_Node.cpp
@@ -239,6 +239,12 @@ FName USCS_Node::GetVariableName() const
 
 void USCS_Node::NameWasModified()
 {
+	if(ComponentTemplate != nullptr)
+	{
+		// Ensure that the template name stays in sync with the variable name; otherwise, new SCS nodes for the same component type will recycle the subobject rather than create a new instance.
+		ComponentTemplate->Rename(*(VariableName.ToString() + TEXT("_GEN_VARIABLE")), nullptr, REN_DontCreateRedirectors);
+	}
+
 	OnNameChangedExternal.ExecuteIfBound(VariableName);
 }
 
@@ -401,6 +407,16 @@ void USCS_Node::PostLoad()
 {
 	Super::PostLoad();
 
+	if(ComponentTemplate != nullptr)
+	{
+		// Fix up component templates not in sync w/ the variable name on load; if not, duplicate it to a new instance w/ a corrected name in order to ensure that it is uniquely bound to this node only.
+		const FString ComponentTemplateName = VariableName.ToString() + TEXT("_GEN_VARIABLE");
+		if(ComponentTemplateName != ComponentTemplate->GetName())
+		{
+			ComponentTemplate = static_cast<UActorComponent*>(StaticDuplicateObject(ComponentTemplate, ComponentTemplate->GetOuter(), *ComponentTemplateName));
+		}
+	}
+
 	ValidateGuid();
 }
 
