Added SkeletalMesh Resource Size

diff --git a/Engine/Source/Runtime/Engine/Classes/Animation/VertexAnim/MorphTarget.h b/Engine/Source/Runtime/Engine/Classes/Animation/VertexAnim/MorphTarget.h
--- a/Engine/Source/Runtime/Engine/Classes/Animation/VertexAnim/MorphTarget.h
+++ b/Engine/Source/Runtime/Engine/Classes/Animation/VertexAnim/MorphTarget.h
@@ -53,6 +53,8 @@ struct FMorphTargetLODModel
 	TArray<FVertexAnimDelta> Vertices;
 	/** number of original verts in the base mesh */
 	int32 NumBaseMeshVerts;
+	/** Get Resource Size */
+	SIZE_T GetResourceSize() const;
 
 	/** pipe operator */
 	friend FArchive& operator<<( FArchive& Ar, FMorphTargetLODModel& M )
diff --git a/Engine/Source/Runtime/Engine/Classes/Engine/SkeletalMesh.h b/Engine/Source/Runtime/Engine/Classes/Engine/SkeletalMesh.h
--- a/Engine/Source/Runtime/Engine/Classes/Engine/SkeletalMesh.h
+++ b/Engine/Source/Runtime/Engine/Classes/Engine/SkeletalMesh.h
@@ -489,6 +489,9 @@ struct FClothingAssetData
 
 	// serialization
 	friend FArchive& operator<<(FArchive& Ar, FClothingAssetData& A);
+
+	// get resource size
+	SIZE_T GetResourceSize() const;
 };
 
 // Material interface for USkeletalMesh - contains a material and a shadow casting flag
@@ -649,7 +652,7 @@ public:
 	TMap<FName, UMorphTarget*> MorphTargetIndexMap;
 
 	/** Reference skeleton precomputed bases. */
-	TArray<FMatrix> RefBasesInvMatrix;    // @todo: wasteful ?!
+	TArray<FMatrix> RefBasesInvMatrix;    
 
 #if WITH_EDITORONLY_DATA
 	/** The section currently selected in the Editor. */
diff --git a/Engine/Source/Runtime/Engine/Classes/PhysicsEngine/PhysicsAsset.h b/Engine/Source/Runtime/Engine/Classes/PhysicsEngine/PhysicsAsset.h
--- a/Engine/Source/Runtime/Engine/Classes/PhysicsEngine/PhysicsAsset.h
+++ b/Engine/Source/Runtime/Engine/Classes/PhysicsEngine/PhysicsAsset.h
@@ -65,6 +65,7 @@ public:
 	virtual void PostLoad() override;
 	virtual FString GetDesc() override;
 	virtual void GetAssetRegistryTags(TArray<FAssetRegistryTag>& OutTags) const override;
+	virtual SIZE_T GetResourceSize(EResourceSizeMode::Type Mode) override;
 #if WITH_EDITOR
 	virtual void PostEditUndo() override;
 #endif
diff --git a/Engine/Source/Runtime/Engine/Private/Animation/ReferenceSkeleton.cpp b/Engine/Source/Runtime/Engine/Private/Animation/ReferenceSkeleton.cpp
--- a/Engine/Source/Runtime/Engine/Private/Animation/ReferenceSkeleton.cpp
+++ b/Engine/Source/Runtime/Engine/Private/Animation/ReferenceSkeleton.cpp
@@ -119,6 +119,17 @@ void FReferenceSkeleton::RebuildNameToIndexMap()
 	checkSlow(NameToIndexMap.Num() == NumBones);
 }
 
+SIZE_T FReferenceSkeleton::GetDataSize() const
+{
+	SIZE_T ResourceSize = 0;
+
+	ResourceSize += RefBoneInfo.GetAllocatedSize();
+	ResourceSize += RefBonePose.GetAllocatedSize();
+	ResourceSize += NameToIndexMap.GetAllocatedSize();
+
+	return ResourceSize;
+}
+
 FArchive & operator<<(FArchive & Ar, FReferenceSkeleton & F)
 {
 	Ar << F.RefBoneInfo;
@@ -139,4 +150,5 @@ FArchive & operator<<(FArchive & Ar, FReferenceSkeleton & F)
 	}
 
 	return Ar;
-}
\ No newline at end of file
+}
+
diff --git a/Engine/Source/Runtime/Engine/Private/MorphMesh.cpp b/Engine/Source/Runtime/Engine/Private/MorphMesh.cpp
--- a/Engine/Source/Runtime/Engine/Private/MorphMesh.cpp
+++ b/Engine/Source/Runtime/Engine/Private/MorphMesh.cpp
@@ -34,8 +34,18 @@ void UMorphTarget::Serialize( FArchive& Ar )
 
 SIZE_T UMorphTarget::GetResourceSize(EResourceSizeMode::Type Mode)
 {
-	return 0;
-}
+	SIZE_T ResourceSize = 0;
 
+	for (const auto& LODModel : MorphLODModels)
+	{
+		ResourceSize += LODModel.GetResourceSize();
+	}
+	return ResourceSize;
+}
 
 
+//////////////////////////////////////////////////////////////////////
+SIZE_T FMorphTargetLODModel::GetResourceSize() const
+{
+	return Vertices.GetAllocatedSize() + sizeof(int32);
+}
\ No newline at end of file
diff --git a/Engine/Source/Runtime/Engine/Private/PhysicsEngine/PhysicsAsset.cpp b/Engine/Source/Runtime/Engine/Private/PhysicsEngine/PhysicsAsset.cpp
--- a/Engine/Source/Runtime/Engine/Private/PhysicsEngine/PhysicsAsset.cpp
+++ b/Engine/Source/Runtime/Engine/Private/PhysicsEngine/PhysicsAsset.cpp
@@ -406,4 +406,20 @@ void UPhysicsAsset::BodyFindConstraints(int32 BodyIndex, TArray<int32>& Constrai
 			Constraints.Add(ConIdx);
 		}
 	}
+}
+
+SIZE_T UPhysicsAsset::GetResourceSize(EResourceSizeMode::Type Mode)
+{
+	SIZE_T ResourceSize = 0;
+
+	for (const auto& SingleBody : BodySetup)
+	{
+		ResourceSize += SingleBody->GetResourceSize(Mode);
+	}
+
+	ResourceSize += BodySetupIndexMap.GetAllocatedSize();
+	ResourceSize += CollisionDisableTable.GetAllocatedSize();
+
+	// @todo implement inclusive mode
+	return ResourceSize;
 }
\ No newline at end of file
diff --git a/Engine/Source/Runtime/Engine/Private/SkeletalMesh.cpp b/Engine/Source/Runtime/Engine/Private/SkeletalMesh.cpp
--- a/Engine/Source/Runtime/Engine/Private/SkeletalMesh.cpp
+++ b/Engine/Source/Runtime/Engine/Private/SkeletalMesh.cpp
@@ -117,6 +117,39 @@ static bool SaveApexClothingAssetToBlob(const NxClothingAsset *InAsset, TArray<u
 }
 
 #endif//#if WITH_APEX_CLOTHING
+
+////////////////////////////////////////////////
+SIZE_T FClothingAssetData::GetResourceSize() const
+{
+#if WITH_APEX_CLOTHING
+	SIZE_T ResourceSize = 0;
+
+	if (ApexClothingAsset)
+	{
+		physx::PxU32 LODLevel = ApexClothingAsset->getNumGraphicalLodLevels();
+		for(physx::PxU32 LODId=0; LODId < LODLevel; ++LODId)
+		{
+			if(const NxRenderMeshAsset* RenderAsset = ApexClothingAsset->getRenderMeshAsset(LODId))
+			{
+				NxRenderMeshAssetStats AssetStats;
+				RenderAsset->getStats(AssetStats);
+
+				ResourceSize += AssetStats.totalBytes;
+			}
+		}
+	}
+
+	ResourceSize += ClothCollisionVolumes.GetAllocatedSize();
+	ResourceSize += ClothCollisionConvexPlaneIndices.GetAllocatedSize();
+	ResourceSize += ClothCollisionVolumePlanes.GetAllocatedSize();
+	ResourceSize += ClothBoneSpheres.GetAllocatedSize();
+	ResourceSize += BoneSphereConnections.GetAllocatedSize();
+	ResourceSize += ClothVisualizationInfos.GetAllocatedSize();
+#endif // #if WITH_APEX_CLOTHING
+
+	return ResourceSize;
+}
+
 /*-----------------------------------------------------------------------------
 	FSkeletalMeshVertexBuffer
 -----------------------------------------------------------------------------*/
@@ -1595,6 +1628,40 @@ void FStaticLODModel::ReleaseCPUResources()
 	}
 }
 
+SIZE_T FStaticLODModel::GetResourceSize() const
+{
+	SIZE_T ResourceSize = 0;
+
+	ResourceSize += Sections.GetAllocatedSize();
+	ResourceSize += Chunks.GetAllocatedSize();
+	ResourceSize += ActiveBoneIndices.GetAllocatedSize();  
+	ResourceSize += RequiredBones.GetAllocatedSize();
+
+	const FRawStaticIndexBuffer16or32Interface* IndexBuffer = MultiSizeIndexContainer.GetIndexBuffer();
+	if (IndexBuffer)
+	{
+		ResourceSize += IndexBuffer->GetResourceDataSize(); 
+	}
+
+	const FRawStaticIndexBuffer16or32Interface* AdjacentIndexBuffer = AdjacencyMultiSizeIndexContainer.GetIndexBuffer();
+	if (AdjacentIndexBuffer)
+	{
+		ResourceSize += AdjacentIndexBuffer->GetResourceDataSize();
+	}
+
+	ResourceSize += VertexBufferGPUSkin.GetVertexDataSize();
+	ResourceSize += ColorVertexBuffer.GetVertexDataSize();
+	ResourceSize += APEXClothVertexBuffer.GetVertexDataSize();
+
+	ResourceSize += RawPointIndices.GetBulkDataSize();
+	ResourceSize += LegacyRawPointIndices.GetBulkDataSize();
+	ResourceSize += MeshToImportVertexMap.GetAllocatedSize();
+
+	// I suppose we add everything we could
+	ResourceSize += sizeof(int32);
+
+	return ResourceSize;
+}
 /*-----------------------------------------------------------------------------
 FStaticMeshSourceData
 -----------------------------------------------------------------------------*/
@@ -1906,6 +1973,19 @@ bool FSkeletalMeshResource::RequiresCPUSkinning(ERHIFeatureLevel::Type FeatureLe
 	return (MaxBonesPerChunk > MaxGPUSkinBones) || (HasExtraBoneInfluences() && FeatureLevel < ERHIFeatureLevel::ES3_1);
 }
 
+SIZE_T FSkeletalMeshResource::GetResourceSize()
+{
+	SIZE_T ResourceSize = 0;
+	for(int32 LODIndex = 0; LODIndex < LODModels.Num(); ++LODIndex)
+	{
+		const FStaticLODModel& Model = LODModels[LODIndex];
+
+		ResourceSize += Model.GetResourceSize();
+	}
+
+	return ResourceSize;
+}
+
 /*-----------------------------------------------------------------------------
 	USkeletalMesh
 -----------------------------------------------------------------------------*/
@@ -2061,7 +2141,55 @@ float USkeletalMesh::GetStreamingTextureFactor( int32 RequestedUVIndex )
 
 SIZE_T USkeletalMesh::GetResourceSize(EResourceSizeMode::Type Mode)
 {
-	return 0;
+	SIZE_T ResourceSize = 0;
+	if(ImportedResource.IsValid())
+	{
+		ResourceSize += ImportedResource->GetResourceSize();
+	}
+
+	for (const auto& MorphTarget : MorphTargets)
+	{
+		ResourceSize += MorphTarget->GetResourceSize(Mode);
+	}
+
+	for (const auto& ClothingAsset : ClothingAssets)
+	{
+		ResourceSize += ClothingAsset.GetResourceSize();
+	}
+
+	if(Mode == EResourceSizeMode::Inclusive)
+	{
+		TSet<UMaterialInterface*> UniqueMaterials;
+		for(int32 MaterialIndex = 0; MaterialIndex < Materials.Num(); ++MaterialIndex)
+		{
+			UMaterialInterface* Material = Materials[MaterialIndex].MaterialInterface;
+			bool bAlreadyCounted = false;
+			UniqueMaterials.Add(Material, &bAlreadyCounted);
+			if(!bAlreadyCounted && Material)
+			{
+				ResourceSize += Material->GetResourceSize(Mode);
+			}
+		}
+
+#if WITH_EDITORONLY_DATA
+		ResourceSize += RetargetBasePose.GetAllocatedSize();
+#endif
+
+		ResourceSize += RefBasesInvMatrix.GetAllocatedSize();
+		ResourceSize += RefSkeleton.GetDataSize();
+
+		if(BodySetup)
+		{
+			ResourceSize += BodySetup->GetResourceSize(Mode);
+		}
+
+		if (PhysicsAsset)
+		{
+			ResourceSize += PhysicsAsset->GetResourceSize(Mode);
+		}
+	}
+
+	return ResourceSize;
 }
 
 /**
@@ -4879,4 +5007,5 @@ void USkeletalMeshSocket::PostEditChangeProperty(FPropertyChangedEvent& Property
 		ChangedEvent.Broadcast(this, PropertyChangedEvent.MemberProperty);
 	}
 }
-#endif
\ No newline at end of file
+#endif
+
diff --git a/Engine/Source/Runtime/Engine/Public/RawIndexBuffer.h b/Engine/Source/Runtime/Engine/Public/RawIndexBuffer.h
--- a/Engine/Source/Runtime/Engine/Public/RawIndexBuffer.h
+++ b/Engine/Source/Runtime/Engine/Public/RawIndexBuffer.h
@@ -211,7 +211,7 @@ public:
 	virtual void Insert(int32 Idx, int32 Num = 1) = 0;
 	virtual void Remove(int32 Idx, int32 Num = 1) = 0;
 	virtual void Empty(int32 Slack = 0) = 0;
-	virtual int32 GetResourceDataSize() = 0;
+	virtual int32 GetResourceDataSize() const = 0;
 };
 
 template <typename INDEX_TYPE>
@@ -301,7 +301,7 @@ public:
 		Indices.Empty(Slack);
 	}
 
-	virtual int32 GetResourceDataSize() override
+	virtual int32 GetResourceDataSize() const override
 	{
 		return Indices.GetResourceDataSize();
 	}
diff --git a/Engine/Source/Runtime/Engine/Public/ReferenceSkeleton.h b/Engine/Source/Runtime/Engine/Public/ReferenceSkeleton.h
--- a/Engine/Source/Runtime/Engine/Public/ReferenceSkeleton.h
+++ b/Engine/Source/Runtime/Engine/Public/ReferenceSkeleton.h
@@ -255,6 +255,7 @@ public:
 
 	void RebuildNameToIndexMap();
 
+	SIZE_T GetDataSize() const;
 
 	friend FArchive & operator<<(FArchive & Ar, FReferenceSkeleton & F);
 };
\ No newline at end of file
diff --git a/Engine/Source/Runtime/Engine/Public/SkeletalMeshTypes.h b/Engine/Source/Runtime/Engine/Public/SkeletalMeshTypes.h
--- a/Engine/Source/Runtime/Engine/Public/SkeletalMeshTypes.h
+++ b/Engine/Source/Runtime/Engine/Public/SkeletalMeshTypes.h
@@ -1662,6 +1662,11 @@ public:
 
 		return false;
 	}
+
+	/**
+	 * Get Resource Size
+	 */
+	SIZE_T GetResourceSize() const;
 };
 
 /**
@@ -1696,6 +1701,11 @@ public:
 	/** Returns true if there are more than MAX_INFLUENCES_PER_STREAM influences per vertex. */
 	bool HasExtraBoneInfluences() const;
 
+	/** 
+	 *	Return the resource size
+	 */
+	SIZE_T GetResourceSize();
+
 private:
 	/** True if the resource has been initialized. */
 	bool bInitialized;
