Fix initialization order bug that caused evaluation skipping to not function properly

diff --git a/Engine/Source/Runtime/Engine/Private/Components/SkeletalMeshComponent.cpp b/Engine/Source/Runtime/Engine/Private/Components/SkeletalMeshComponent.cpp
--- a/Engine/Source/Runtime/Engine/Private/Components/SkeletalMeshComponent.cpp
+++ b/Engine/Source/Runtime/Engine/Private/Components/SkeletalMeshComponent.cpp
@@ -1092,10 +1092,10 @@ void USkeletalMeshComponent::RefreshBoneTransforms(FActorComponentTickFunction*
 
 	const bool bDoPAE = !!CVarUseParallelAnimationEvaluation.GetValueOnGameThread() && FApp::ShouldUseThreadingForPerformance();
 
-	const bool bDoParallelEvaluation = AnimEvaluationContext.bDoEvaluation && TickFunction && bDoPAE;
+	const bool bDoParallelEvaluation = bShouldDoEvaluation && TickFunction && bDoPAE;
 
-	const bool bBlockOnTask = !bDoParallelEvaluation;  // If we aren't trying to do parallel evaluation then we
-															// will need to wait on an existing task.
+	const bool bBlockOnTask = !bDoParallelEvaluation;   // If we aren't trying to do parallel evaluation then we
+														// will need to wait on an existing task.
 
 	const bool bPerformPostAnimEvaluation = true;
 	if (HandleExistingParallelEvaluationTask(bBlockOnTask, bPerformPostAnimEvaluation))
