Fix ensure with mixing additive, non additive in montage

diff --git a/Engine/Source/Runtime/Engine/Classes/Animation/AnimMontage.h b/Engine/Source/Runtime/Engine/Classes/Animation/AnimMontage.h
--- a/Engine/Source/Runtime/Engine/Classes/Animation/AnimMontage.h
+++ b/Engine/Source/Runtime/Engine/Classes/Animation/AnimMontage.h
@@ -449,6 +449,10 @@ public:
 	 */
 	ENGINE_API void UpdateLinkableElements(int32 SlotIdx, int32 SegmentIdx);
 #endif
+	/*
+	 * Check if this slot has valid animation data
+	 */
+	bool IsValidAdditiveSlot(const FName& SlotNodeName) const;
 
 	/** Get FCompositeSection with InSectionName */
 	FCompositeSection& GetAnimCompositeSection(int32 SectionIndex);
diff --git a/Engine/Source/Runtime/Engine/Private/Animation/AnimInstance.cpp b/Engine/Source/Runtime/Engine/Private/Animation/AnimInstance.cpp
--- a/Engine/Source/Runtime/Engine/Private/Animation/AnimInstance.cpp
+++ b/Engine/Source/Runtime/Engine/Private/Animation/AnimInstance.cpp
@@ -1297,7 +1297,7 @@ void UAnimInstance::GetSlotWeight(FName const & SlotNodeName, float& out_SlotNod
 		{
 			NodeTotalWeight += MontageInstance->Weight;
 
-			if( !MontageInstance->Montage->IsValidAdditive() )
+			if( !MontageInstance->Montage->IsValidAdditiveSlot(SlotNodeName) )
 			{
 				NonAdditiveTotalWeight += MontageInstance->Weight;
 			}
diff --git a/Engine/Source/Runtime/Engine/Private/Animation/AnimMontage.cpp b/Engine/Source/Runtime/Engine/Private/Animation/AnimMontage.cpp
--- a/Engine/Source/Runtime/Engine/Private/Animation/AnimMontage.cpp
+++ b/Engine/Source/Runtime/Engine/Private/Animation/AnimMontage.cpp
@@ -678,6 +678,23 @@ bool UAnimMontage::IsValidAdditive() const
 	return false;
 }
 
+bool UAnimMontage::IsValidAdditiveSlot(const FName& SlotNodeName) const
+{
+	// if first one is additive, this is additive
+	if ( SlotAnimTracks.Num() > 0 )
+	{
+		for (int32 I=0; I<SlotAnimTracks.Num(); ++I)
+		{
+			if (SlotAnimTracks[I].SlotName == SlotNodeName)
+			{
+				return SlotAnimTracks[I].AnimTrack.IsAdditive();
+			}
+		}
+	}
+
+	return false;
+}
+
 EAnimEventTriggerOffsets::Type UAnimMontage::CalculateOffsetFromSections(float Time) const
 {
 	for(auto Iter = CompositeSections.CreateConstIterator(); Iter; ++Iter)
