Use FName instead of string compares

diff --git a/Engine/Source/Runtime/Engine/Private/Components/PrimitiveComponent.cpp b/Engine/Source/Runtime/Engine/Private/Components/PrimitiveComponent.cpp
--- a/Engine/Source/Runtime/Engine/Private/Components/PrimitiveComponent.cpp
+++ b/Engine/Source/Runtime/Engine/Private/Components/PrimitiveComponent.cpp
@@ -679,15 +679,19 @@ bool UPrimitiveComponent::CanEditChange(const UProperty* InProperty) const
 
 void UPrimitiveComponent::PostEditChangeChainProperty(FPropertyChangedChainEvent& PropertyChangedEvent)
 {
+	const FName NAME_Scale3D(TEXT("Scale3D"));
+	const FName NAME_Scale(TEXT("Scale"));
+	const FName NAME_Translation(TEXT("Translation"));
+
 	for( FEditPropertyChain::TIterator It(PropertyChangedEvent.PropertyChain.GetHead()); It; ++It )
 	{
-		FName N = *It->GetName();
-		if( FCString::Stricmp( *It->GetName(), TEXT("Scale3D") )		== 0 || 
-			FCString::Stricmp( *It->GetName(), TEXT("Scale") )			== 0 || 
-			FCString::Stricmp( *It->GetName(), TEXT("Translation") )	== 0 || 
-			FCString::Stricmp( *It->GetName(), TEXT("Rotation") )		== 0 )
+		if( It->GetFName() == NAME_Scale3D		||
+			It->GetFName() == NAME_Scale		||
+			It->GetFName() == NAME_Translation	||
+			It->GetFName() == NAME_Rotation)
 		{
 			UpdateComponentToWorld();
+			break;
 		}
 	}
 
