Fix issues with EditReregisterContext entries not being cleared and eventual pointer reuse causing crashes by using TWeakObjectPtr.

diff --git a/Engine/Source/Runtime/Engine/Private/Components/ActorComponent.cpp b/Engine/Source/Runtime/Engine/Private/Components/ActorComponent.cpp
--- a/Engine/Source/Runtime/Engine/Private/Components/ActorComponent.cpp
+++ b/Engine/Source/Runtime/Engine/Private/Components/ActorComponent.cpp
@@ -428,10 +428,11 @@ bool UActorComponent::CallRemoteFunction( UFunction* Function, void* Parameters,
 	return false;
 }
 
+#if WITH_EDITOR
+
 /** FComponentReregisterContexts for components which have had PreEditChange called but not PostEditChange. */
-static TMap<UActorComponent*,FComponentReregisterContext*> EditReregisterContexts;
+static TMap<TWeakObjectPtr<UActorComponent>,FComponentReregisterContext*> EditReregisterContexts;
 
-#if WITH_EDITOR
 bool UActorComponent::Modify( bool bAlwaysMarkDirty/*=true*/ )
 {
 	// If this is a construction script component we don't store them in the transaction buffer.  Instead, mark
@@ -461,7 +462,7 @@ void UActorComponent::PreEditChange(UProperty* PropertyThatWillChange)
 		else
 		{
 			ExecuteUnregisterEvents();
-			World = NULL;
+			World = nullptr;
 		}
 	}
 
@@ -484,11 +485,21 @@ void UActorComponent::PostEditUndo()
 	if( IsPendingKill() )
 	{
 		// The reregister context won't bother attaching components that are 'pending kill'. 
-		FComponentReregisterContext* ReregisterContext = EditReregisterContexts.FindRef(this);
-		if(ReregisterContext)
+		FComponentReregisterContext* ReregisterContext = nullptr;
+		if (EditReregisterContexts.RemoveAndCopyValue(this, ReregisterContext))
 		{
 			delete ReregisterContext;
-			EditReregisterContexts.Remove(this);
+		}
+		else
+		{
+			// This means there are likely some stale elements left in there now, strip them out
+			for (auto It(EditReregisterContexts.CreateIterator()); It; ++It)
+			{
+				if (!It.Key().IsValid())
+				{
+					It.RemoveCurrent();
+				}
+			}
 		}
 	}
 	else
@@ -531,11 +542,10 @@ void UActorComponent::PostEditUndo()
 
 void UActorComponent::ConsolidatedPostEditChange(const FPropertyChangedEvent& PropertyChangedEvent)
 {
-	FComponentReregisterContext* ReregisterContext = EditReregisterContexts.FindRef(this);
-	if(ReregisterContext)
+	FComponentReregisterContext* ReregisterContext = nullptr;
+	if(EditReregisterContexts.RemoveAndCopyValue(this, ReregisterContext))
 	{
 		delete ReregisterContext;
-		EditReregisterContexts.Remove(this);
 
 		AActor* MyOwner = GetOwner();
 		if ( MyOwner && !MyOwner->IsTemplate() && PropertyChangedEvent.ChangeType != EPropertyChangeType::Interactive )
@@ -543,6 +553,17 @@ void UActorComponent::ConsolidatedPostEditChange(const FPropertyChangedEvent& Pr
 			MyOwner->RerunConstructionScripts();
 		}
 	}
+	else
+	{
+		// This means there are likely some stale elements left in there now, strip them out
+		for (auto It(EditReregisterContexts.CreateIterator()); It; ++It)
+		{
+			if (!It.Key().IsValid())
+			{
+				It.RemoveCurrent();
+			}
+		}
+	}
 
 	// The component or its outer could be pending kill when calling PostEditChange when applying a transaction.
 	// Don't do do a full recreate in this situation, and instead simply detach.
