Fix missing physx locks

diff --git a/Engine/Source/Runtime/Engine/Private/InstancedStaticMesh.cpp b/Engine/Source/Runtime/Engine/Private/InstancedStaticMesh.cpp
--- a/Engine/Source/Runtime/Engine/Private/InstancedStaticMesh.cpp
+++ b/Engine/Source/Runtime/Engine/Private/InstancedStaticMesh.cpp
@@ -1612,7 +1612,9 @@ void UInstancedStaticMeshComponent::SetupNewInstanceData(FInstancedStaticMeshIns
 				auto* Aggregate = GPhysXSDK->createAggregate(AggregateMaxSize, false);
 				const int32 AddedIndex = Aggregates.Add(Aggregate);
 				check(AddedIndex == AggregateIndex);
-				PhysScene->GetPhysXScene(SceneType)->addAggregate(*Aggregate);
+				PxScene* PScene = PhysScene->GetPhysXScene(SceneType);
+				SCOPED_SCENE_WRITE_LOCK(PScene);
+				PScene->addAggregate(*Aggregate);
 			}
 		}
 
diff --git a/Engine/Source/Runtime/Engine/Private/PhysicsEngine/BodyInstance.cpp b/Engine/Source/Runtime/Engine/Private/PhysicsEngine/BodyInstance.cpp
--- a/Engine/Source/Runtime/Engine/Private/PhysicsEngine/BodyInstance.cpp
+++ b/Engine/Source/Runtime/Engine/Private/PhysicsEngine/BodyInstance.cpp
@@ -1883,7 +1883,7 @@ bool FBodyInstance::Weld(FBodyInstance* TheirBody, const FTransform& TheirTM)
 
 		if (PxRigidActor* MyBody = RigidActorAsync)
 		{
-			TheirBody->BodySetup->AddShapesToRigidActor_AssumesLocked(this, MyBody, PST_Sync, Scale3D, PSimpleMat, ComplexPhysMats, ShapeData, RelativeTM, &PNewShapes);
+			TheirBody->BodySetup->AddShapesToRigidActor_AssumesLocked(this, MyBody, PST_Async, Scale3D, PSimpleMat, ComplexPhysMats, ShapeData, RelativeTM, &PNewShapes);
 			if (TheirBody->bGenerateWakeEvents)
 			{
 				MyBody->setActorFlag(PxActorFlag::eSEND_SLEEP_NOTIFIES, true);
diff --git a/Engine/Source/Runtime/Landscape/Private/LandscapeCollision.cpp b/Engine/Source/Runtime/Landscape/Private/LandscapeCollision.cpp
--- a/Engine/Source/Runtime/Landscape/Private/LandscapeCollision.cpp
+++ b/Engine/Source/Runtime/Landscape/Private/LandscapeCollision.cpp
@@ -923,7 +923,9 @@ void ULandscapeMeshCollisionComponent::CreatePhysicsState()
 				}
 
 				// Add to scenes
-				PhysScene->GetPhysXScene(PST_Sync)->addActor(*MeshActorSync);
+				PxScene* SyncScene = PhysScene->GetPhysXScene(PST_Sync);
+				SCOPED_SCENE_WRITE_LOCK(SyncScene);
+				SyncScene->addActor(*MeshActorSync);
 
 				if (PhysScene->HasAsyncScene())
 				{
