#ue4 - Make it clear that MoveIgnoreActors does not apply to movement from physics simulation.

diff --git a/Engine/Source/Runtime/Engine/Classes/Components/PrimitiveComponent.h b/Engine/Source/Runtime/Engine/Classes/Components/PrimitiveComponent.h
--- a/Engine/Source/Runtime/Engine/Classes/Components/PrimitiveComponent.h
+++ b/Engine/Source/Runtime/Engine/Classes/Components/PrimitiveComponent.h
@@ -485,6 +485,7 @@ public:
 	 * Set of actors to ignore during component sweeps in MoveComponent().
 	 * All components owned by these actors will be ignored when this component moves or updates overlaps.
 	 * Components on the other Actor may also need to be told to do the same when they move.
+	 * Does not affect movement of this component when simulating physics.
 	 * @see IgnoreActorWhenMoving()
 	 */
 	TArray<TWeakObjectPtr<AActor> > MoveIgnoreActors;
@@ -492,6 +493,7 @@ public:
 	/**
 	 * Tells this component whether to ignore collision with all components of a specific Actor when this component is moved.
 	 * Components on the other Actor may also need to be told to do the same when they move.
+	 * Does not affect movement of this component when simulating physics.
 	 */
 	UFUNCTION(BlueprintCallable, Category = "Collision", meta=(Keywords="Move MoveIgnore"))
 	void IgnoreActorWhenMoving(AActor* Actor, bool bShouldIgnore);
