Remove deprecated API functions (since 4.0)

diff --git a/Engine/Source/Runtime/Engine/Classes/GameFramework/Actor.h b/Engine/Source/Runtime/Engine/Classes/GameFramework/Actor.h
--- a/Engine/Source/Runtime/Engine/Classes/GameFramework/Actor.h
+++ b/Engine/Source/Runtime/Engine/Classes/GameFramework/Actor.h
@@ -1088,15 +1088,6 @@ public:
 	UFUNCTION(BlueprintCallable, meta=(DeprecatedFunction, DeprecationMessage="Use PrimitiveComponent.CreateAndSetMaterialInstanceDynamic instead.", BlueprintProtected = "true"), Category="Rendering|Material")
 	class UMaterialInstanceDynamic* MakeMIDForMaterial(class UMaterialInterface* Parent);
 
-	//=============================================================================
-	// Sound functions.
-	
-	DEPRECATED(4.0, "Actor::PlaySoundOnActor will be removed. Use UGameplayStatics::PlaySoundAttached instead.")
-	void PlaySoundOnActor(class USoundCue* InSoundCue, float VolumeMultiplier=1.f, float PitchMultiplier=1.f);
-
-	DEPRECATED(4.0, "Actor::PlaySoundOnActor will be removed. Use UGameplayStatics::PlaySoundAtLocation instead.")
-	void PlaySoundAtLocation(class USoundCue* InSoundCue, FVector SoundLocation, float VolumeMultiplier=1.f, float PitchMultiplier=1.f);
-
 	//=============================================================================
 	// AI functions.
 	
diff --git a/Engine/Source/Runtime/Engine/Private/Actor.cpp b/Engine/Source/Runtime/Engine/Private/Actor.cpp
--- a/Engine/Source/Runtime/Engine/Private/Actor.cpp
+++ b/Engine/Source/Runtime/Engine/Private/Actor.cpp
@@ -3102,16 +3102,6 @@ AWorldSettings * AActor::GetWorldSettings() const
 	return (World ? World->GetWorldSettings() : nullptr);
 }
 
-void AActor::PlaySoundOnActor(USoundCue* InSoundCue, float VolumeMultiplier/*=1.f*/, float PitchMultiplier/*=1.f*/)
-{
-	UGameplayStatics::PlaySoundAtLocation( this, InSoundCue, GetActorLocation(), VolumeMultiplier, PitchMultiplier );
-}
-
-void AActor::PlaySoundAtLocation(USoundCue* InSoundCue, FVector SoundLocation, float VolumeMultiplier/*=1.f*/, float PitchMultiplier/*=1.f*/)
-{
-	UGameplayStatics::PlaySoundAtLocation( this, InSoundCue, (SoundLocation.IsZero() ? GetActorLocation() : SoundLocation), VolumeMultiplier, PitchMultiplier );
-}
-
 ENetMode AActor::GetNetMode() const
 {
 	UNetDriver *NetDriver = GetNetDriver();
