Rename PlaySoundAttached to SpawnSoundAttached
Add SpawnSoundLocation which returns an Audio Component
Add rotation to Play/Spawn Sound/Dialogue functions
Make dialogue gameplay statics play/spawn API consistent with sound and eliminate duplicate code

diff --git a/Engine/Config/BaseEngine.ini b/Engine/Config/BaseEngine.ini
--- a/Engine/Config/BaseEngine.ini
+++ b/Engine/Config/BaseEngine.ini
@@ -1298,6 +1298,10 @@ bUseStreamingPause=false
 ;Rename MapRange to MapRangeUnclamped - DEPRECATED(4.9)
 +K2FieldRedirects=(OldFieldName="KismetMathLibrary.MapRange",NewFieldName="KismetMathLibrary.MapRangeUnclamped")
 
+;Rename Play[Sound|Dialogue]Attached to Spawn[Sound|Dialogue]Attached - DEPRECATED(4.9)
++K2FieldRedirects=(OldFieldName="GameplayStatics.PlaySoundAttached",NewFieldName="GameplayStatics.SpawnSoundAttached")
++K2FieldRedirects=(OldFieldName="GameplayStatics.PlayDialogueAttached",NewFieldName="GameplayStatics.SpawnDialogueAttached")
+
 [CoreUObject.Metadata]
 ; Note: UnrealHeaderTool should be rerun after making changes to MetadataRedirects to catch any keys specified in class headers
 +MetadataRedirects=(OldKey="K2Protected", NewKey="BlueprintProtected")
diff --git a/Engine/Source/Runtime/AIModule/Private/BehaviorTree/Tasks/BTTask_PlaySound.cpp b/Engine/Source/Runtime/AIModule/Private/BehaviorTree/Tasks/BTTask_PlaySound.cpp
--- a/Engine/Source/Runtime/AIModule/Private/BehaviorTree/Tasks/BTTask_PlaySound.cpp
+++ b/Engine/Source/Runtime/AIModule/Private/BehaviorTree/Tasks/BTTask_PlaySound.cpp
@@ -20,7 +20,7 @@ EBTNodeResult::Type UBTTask_PlaySound::ExecuteTask(UBehaviorTreeComponent& Owner
 	{
 		if (const APawn* MyPawn = MyController->GetPawn())
 		{
-			AC = UGameplayStatics::PlaySoundAttached(SoundToPlay, MyPawn->GetRootComponent());
+			AC = UGameplayStatics::SpawnSoundAttached(SoundToPlay, MyPawn->GetRootComponent());
 		}
 	}
 	return AC ? EBTNodeResult::Succeeded : EBTNodeResult::Failed;
diff --git a/Engine/Source/Runtime/Engine/Classes/Kismet/GameplayStatics.h b/Engine/Source/Runtime/Engine/Classes/Kismet/GameplayStatics.h
--- a/Engine/Source/Runtime/Engine/Classes/Kismet/GameplayStatics.h
+++ b/Engine/Source/Runtime/Engine/Classes/Kismet/GameplayStatics.h
@@ -300,59 +300,172 @@ class ENGINE_API UGameplayStatics : public UBlueprintFunctionLibrary
 	 * @param Sound - Sound to play.
 	 * @param VolumeMultiplier - Multiplied with the volume to make the sound louder or softer.
 	 * @param PitchMultiplier - Multiplies the pitch.
+	 * @param StartTime - How far in to the sound to begin playback at
 	 */
 	UFUNCTION(BlueprintCallable, BlueprintCosmetic, Category="Audio", meta=( WorldContext="WorldContextObject", AdvancedDisplay = "2", UnsafeDuringActorConstruction = "true" ))
 	static void PlaySound2D(UObject* WorldContextObject, class USoundBase* Sound, float VolumeMultiplier = 1.f, float PitchMultiplier = 1.f, float StartTime = 0.f);
 
+	/**
+	 * Spawns a sound with no attenuation, perfect for UI sounds.
+	 *
+	 * ● Not Replicated.
+	 * @param Sound - Sound to play.
+	 * @param VolumeMultiplier - Multiplied with the volume to make the sound louder or softer.
+	 * @param PitchMultiplier - Multiplies the pitch.
+	 * @param StartTime - How far in to the sound to begin playback at
+	 * @return An audio component to manipulate the spawned sound
+	 */
+	UFUNCTION(BlueprintCallable, BlueprintCosmetic, Category="Audio", meta=( WorldContext="WorldContextObject", AdvancedDisplay = "2", UnsafeDuringActorConstruction = "true", Keywords = "play" ))
+	static UAudioComponent* SpawnSound2D(UObject* WorldContextObject, class USoundBase* Sound, float VolumeMultiplier = 1.f, float PitchMultiplier = 1.f, float StartTime = 0.f);
+
 	/**
 	 * Plays a sound at the given location. This is a fire and forget sound and does not travel with any actor. Replication is also not handled at this point.
 	 * @param Sound - sound to play
 	 * @param Location - World position to play sound at
+	 * @param Rotation - World rotation to play sound at
 	 * @param VolumeMultiplier - Volume multiplier 
 	 * @param PitchMultiplier - PitchMultiplier
+	 * @param StartTime - How far in to the sound to begin playback at
 	 * @param AttenuationSettings - Override attenuation settings package to play sound with
 	 */
 	UFUNCTION(BlueprintCallable, Category="Audio", meta=(WorldContext="WorldContextObject", AdvancedDisplay = "3", UnsafeDuringActorConstruction = "true"))
-	static void PlaySoundAtLocation(UObject* WorldContextObject, class USoundBase* Sound, FVector Location, float VolumeMultiplier = 1.f, float PitchMultiplier = 1.f, float StartTime = 0.f, class USoundAttenuation* AttenuationSettings = NULL);
+	static void PlaySoundAtLocation(UObject* WorldContextObject, class USoundBase* Sound, FVector Location, FRotator Rotation, float VolumeMultiplier = 1.f, float PitchMultiplier = 1.f, float StartTime = 0.f, class USoundAttenuation* AttenuationSettings = NULL);
 
-	/** Plays a dialogue at the given location. This is a fire and forget sound and does not travel with any actor. Replication is also not handled at this point.
-	 * @param Dialogue - dialogue to play
+	static void PlaySoundAtLocation(UObject* WorldContextObject, class USoundBase* Sound, FVector Location, float VolumeMultiplier = 1.f, float PitchMultiplier = 1.f, float StartTime = 0.f, class USoundAttenuation* AttenuationSettings = NULL)
+	{
+		PlaySoundAtLocation(WorldContextObject, Sound, Location, FRotator::ZeroRotator, VolumeMultiplier, PitchMultiplier, StartTime, AttenuationSettings);
+	}
+
+	/**
+	 * Spawns a sound at the given location. This does not travel with any actor. Replication is also not handled at this point.
+	 * @param Sound - sound to play
 	 * @param Location - World position to play sound at
-	 * @param World - The World in which the sound is to be played
+	 * @param Rotation - World rotation to play sound at
 	 * @param VolumeMultiplier - Volume multiplier 
 	 * @param PitchMultiplier - PitchMultiplier
+	 * @param StartTime - How far in to the sound to begin playback at
 	 * @param AttenuationSettings - Override attenuation settings package to play sound with
+	 * @return An audio component to manipulate the spawned sound
 	 */
-	UFUNCTION(BlueprintCallable, Category="Audio", meta=(WorldContext="WorldContextObject", AdvancedDisplay = "3", UnsafeDuringActorConstruction = "true"))
-	static void PlayDialogueAtLocation(UObject* WorldContextObject, class UDialogueWave* Dialogue, const struct FDialogueContext& Context, FVector Location, float VolumeMultiplier = 1.f, float PitchMultiplier = 1.f, float StartTime = 0.f, class USoundAttenuation* AttenuationSettings = NULL);
+	UFUNCTION(BlueprintCallable, Category="Audio", meta=(WorldContext="WorldContextObject", AdvancedDisplay = "3", UnsafeDuringActorConstruction = "true", Keywords = "play"))
+	static UAudioComponent* SpawnSoundAtLocation(UObject* WorldContextObject, class USoundBase* Sound, FVector Location, FRotator Rotation = FRotator::ZeroRotator, float VolumeMultiplier = 1.f, float PitchMultiplier = 1.f, float StartTime = 0.f, class USoundAttenuation* AttenuationSettings = NULL);
 
 	/** Plays a sound attached to and following the specified component. This is a fire and forget sound. Replication is also not handled at this point.
 	 * @param Sound - sound to play
 	 * @param AttachComponent - Component to attach to.
 	 * @param AttachPointName - Optional named point within the AttachComponent to play the sound at
 	 * @param Location - Depending on the value of Location Type this is either a relative offset from the attach component/point or an absolute world position that will be translated to a relative offset
+	 * @param Rotation - Depending on the value of Location Type this is either a relative offset from the attach component/point or an absolute world rotation that will be translated to a relative offset
 	 * @param LocationType - Specifies whether Location is a relative offset or an absolute world position
 	 * @param bStopWhenAttachedToDestroyed - Specifies whether the sound should stop playing when the owner of the attach to component is destroyed.
 	 * @param VolumeMultiplier - Volume multiplier 
 	 * @param PitchMultiplier - PitchMultiplier	 
+	 * @param StartTime - How far in to the sound to begin playback at
 	 * @param AttenuationSettings - Override attenuation settings package to play sound with
+	 * @return An audio component to manipulate the spawned sound
+	 */
+	UFUNCTION(BlueprintCallable, Category="Audio", meta=(AdvancedDisplay = "2", UnsafeDuringActorConstruction = "true", Keywords = "play"))
+	static class UAudioComponent* SpawnSoundAttached(class USoundBase* Sound, class USceneComponent* AttachToComponent, FName AttachPointName = NAME_None, FVector Location = FVector(ForceInit), FRotator Rotation = FRotator::ZeroRotator, EAttachLocation::Type LocationType = EAttachLocation::KeepRelativeOffset, bool bStopWhenAttachedToDestroyed = false, float VolumeMultiplier = 1.f, float PitchMultiplier = 1.f, float StartTime = 0.f, class USoundAttenuation* AttenuationSettings = NULL);
+
+	DEPRECATED(4.9, "PlaySoundAttached has been renamed SpawnSoundAttached")
+	static class UAudioComponent* PlaySoundAttached(class USoundBase* Sound, class USceneComponent* AttachToComponent, FName AttachPointName = NAME_None, FVector Location = FVector(ForceInit), EAttachLocation::Type LocationType = EAttachLocation::KeepRelativeOffset, bool bStopWhenAttachedToDestroyed = false, float VolumeMultiplier = 1.f, float PitchMultiplier = 1.f, float StartTime = 0.f, class USoundAttenuation* AttenuationSettings = NULL)
+	{
+		return SpawnSoundAttached(Sound, AttachToComponent, AttachPointName, Location, FRotator::ZeroRotator, LocationType, bStopWhenAttachedToDestroyed, VolumeMultiplier, PitchMultiplier, StartTime, AttenuationSettings);
+	}
+
+	static class UAudioComponent* SpawnSoundAttached(class USoundBase* Sound, class USceneComponent* AttachToComponent, FName AttachPointName, FVector Location, EAttachLocation::Type LocationType = EAttachLocation::KeepRelativeOffset, bool bStopWhenAttachedToDestroyed = false, float VolumeMultiplier = 1.f, float PitchMultiplier = 1.f, float StartTime = 0.f, class USoundAttenuation* AttenuationSettings = NULL)
+	{
+		return SpawnSoundAttached(Sound, AttachToComponent, AttachPointName, Location, FRotator::ZeroRotator, LocationType, bStopWhenAttachedToDestroyed, VolumeMultiplier, PitchMultiplier, StartTime, AttenuationSettings);
+	}
+
+	/**
+	 * Plays a dialogue directly with no attenuation, perfect for UI.
+	 *
+	 * ● Fire and Forget.
+	 * ● Not Replicated.
+	 * @param Dialogue - dialogue to play
+	 * @param Context - context the dialogue is to play in
+	 * @param VolumeMultiplier - Multiplied with the volume to make the sound louder or softer.
+	 * @param PitchMultiplier - Multiplies the pitch.
+	 * @param StartTime - How far in to the dialogue to begin playback at
 	 */
-	UFUNCTION(BlueprintCallable, Category="Audio", meta=(AdvancedDisplay = "2", UnsafeDuringActorConstruction = "true"))
-	static class UAudioComponent* PlaySoundAttached(class USoundBase* Sound, class USceneComponent* AttachToComponent, FName AttachPointName = NAME_None, FVector Location = FVector(ForceInit), EAttachLocation::Type LocationType = EAttachLocation::KeepRelativeOffset, bool bStopWhenAttachedToDestroyed = false, float VolumeMultiplier = 1.f, float PitchMultiplier = 1.f, float StartTime = 0.f, class USoundAttenuation* AttenuationSettings = NULL);
+	UFUNCTION(BlueprintCallable, BlueprintCosmetic, Category="Audio", meta=( WorldContext="WorldContextObject", AdvancedDisplay = "3", UnsafeDuringActorConstruction = "true" ))
+	static void PlayDialogue2D(UObject* WorldContextObject, class UDialogueWave* Dialogue, const struct FDialogueContext& Context, float VolumeMultiplier = 1.f, float PitchMultiplier = 1.f, float StartTime = 0.f);
+
+	/**
+	 * Spawns a dialogue with no attenuation, perfect for UI.
+	 *
+	 * ● Not Replicated.
+	 * @param Dialogue - dialogue to play
+	 * @param Context - context the dialogue is to play in
+	 * @param VolumeMultiplier - Multiplied with the volume to make the sound louder or softer.
+	 * @param PitchMultiplier - Multiplies the pitch.
+	 * @param StartTime - How far in to the dialogue to begin playback at
+	 * @return An audio component to manipulate the spawned sound
+	 */
+	UFUNCTION(BlueprintCallable, BlueprintCosmetic, Category="Audio", meta=( WorldContext="WorldContextObject", AdvancedDisplay = "3", UnsafeDuringActorConstruction = "true", Keywords = "play" ))
+	static UAudioComponent* SpawnDialogue2D(UObject* WorldContextObject, class UDialogueWave* Dialogue, const struct FDialogueContext& Context, float VolumeMultiplier = 1.f, float PitchMultiplier = 1.f, float StartTime = 0.f);
+
+	/** Plays a dialogue at the given location. This is a fire and forget sound and does not travel with any actor. Replication is also not handled at this point.
+	 * @param Dialogue - dialogue to play
+	 * @param Context - context the dialogue is to play in
+	 * @param Location - World position to play dialogue at
+	 * @param Rotation - World rotation to play dialogue at
+	 * @param VolumeMultiplier - Volume multiplier 
+	 * @param PitchMultiplier - Pitch multiplier
+	 * @param StartTime - How far in to the dialogue to begin playback at
+	 * @param AttenuationSettings - Override attenuation settings package to play sound with
+	 */
+	UFUNCTION(BlueprintCallable, Category="Audio", meta=(WorldContext="WorldContextObject", AdvancedDisplay = "4", UnsafeDuringActorConstruction = "true"))
+	static void PlayDialogueAtLocation(UObject* WorldContextObject, class UDialogueWave* Dialogue, const struct FDialogueContext& Context, FVector Location, FRotator Rotation, float VolumeMultiplier = 1.f, float PitchMultiplier = 1.f, float StartTime = 0.f, class USoundAttenuation* AttenuationSettings = NULL);
+
+	static void PlayDialogueAtLocation(UObject* WorldContextObject, class UDialogueWave* Dialogue, const struct FDialogueContext& Context, FVector Location, float VolumeMultiplier = 1.f, float PitchMultiplier = 1.f, float StartTime = 0.f, class USoundAttenuation* AttenuationSettings = NULL)
+	{
+		PlayDialogueAtLocation(WorldContextObject, Dialogue, Context, Location, FRotator::ZeroRotator, VolumeMultiplier, PitchMultiplier, StartTime, AttenuationSettings);
+	}
+
+	/** Plays a dialogue at the given location. This is a fire and forget sound and does not travel with any actor. Replication is also not handled at this point.
+	 * @param Dialogue - dialogue to play
+	 * @param Context - context the dialogue is to play in
+	 * @param Location - World position to play dialogue at
+	 * @param Rotation - World rotation to play dialogue at
+	 * @param VolumeMultiplier - Volume multiplier 
+	 * @param PitchMultiplier - PitchMultiplier
+	 * @param StartTime - How far in to the dialogue to begin playback at
+	 * @param AttenuationSettings - Override attenuation settings package to play sound with
+	 * @return Audio Component to manipulate the playing dialogue with
+	 */
+	UFUNCTION(BlueprintCallable, Category="Audio", meta=(WorldContext="WorldContextObject", AdvancedDisplay = "4", UnsafeDuringActorConstruction = "true", Keywords = "play"))
+	static UAudioComponent* SpawnDialogueAtLocation(UObject* WorldContextObject, class UDialogueWave* Dialogue, const struct FDialogueContext& Context, FVector Location, FRotator Rotation = FRotator::ZeroRotator, float VolumeMultiplier = 1.f, float PitchMultiplier = 1.f, float StartTime = 0.f, class USoundAttenuation* AttenuationSettings = NULL);
 
 	/** Plays a dialogue attached to and following the specified component. This is a fire and forget sound. Replication is also not handled at this point.
 	 * @param Dialogue - dialogue to play
+	 * @param Context - context the dialogue is to play in
 	 * @param AttachComponent - Component to attach to.
 	 * @param AttachPointName - Optional named point within the AttachComponent to play the sound at
 	 * @param Location - Depending on the value of Location Type this is either a relative offset from the attach component/point or an absolute world position that will be translated to a relative offset
+	 * @param Rotation - Depending on the value of Location Type this is either a relative offset from the attach component/point or an absolute world rotation that will be translated to a relative offset
 	 * @param LocationType - Specifies whether Location is a relative offset or an absolute world position
 	 * @param bStopWhenAttachedToDestroyed - Specifies whether the sound should stop playing when the owner of the attach to component is destroyed.
 	 * @param VolumeMultiplier - Volume multiplier 
 	 * @param PitchMultiplier - PitchMultiplier	 
+	 * @param StartTime - How far in to the dialogue to begin playback at
 	 * @param AttenuationSettings - Override attenuation settings package to play sound with
+	 * @return Audio Component to manipulate the playing dialogue with
 	 */
-	UFUNCTION(BlueprintCallable, Category="Audio", meta=(AdvancedDisplay = "2", UnsafeDuringActorConstruction = "true"))
-	static class UAudioComponent* PlayDialogueAttached(class UDialogueWave* Dialogue, const struct FDialogueContext& Context, class USceneComponent* AttachToComponent, FName AttachPointName = NAME_None, FVector Location = FVector(ForceInit), EAttachLocation::Type LocationType = EAttachLocation::KeepRelativeOffset, bool bStopWhenAttachedToDestroyed = false, float VolumeMultiplier = 1.f, float PitchMultiplier = 1.f, float StartTime = 0.f, class USoundAttenuation* AttenuationSettings = NULL);
+	UFUNCTION(BlueprintCallable, Category="Audio", meta=(AdvancedDisplay = "2", UnsafeDuringActorConstruction = "true", Keywords = "play"))
+	static class UAudioComponent* SpawnDialogueAttached(class UDialogueWave* Dialogue, const struct FDialogueContext& Context, class USceneComponent* AttachToComponent, FName AttachPointName = NAME_None, FVector Location = FVector(ForceInit), FRotator Rotation = FRotator::ZeroRotator, EAttachLocation::Type LocationType = EAttachLocation::KeepRelativeOffset, bool bStopWhenAttachedToDestroyed = false, float VolumeMultiplier = 1.f, float PitchMultiplier = 1.f, float StartTime = 0.f, class USoundAttenuation* AttenuationSettings = NULL);
+
+	DEPRECATED(4.9, "PlayDialogueAttached has been renamed SpawnDialogueAttached")
+	static class UAudioComponent* PlayDialogueAttached(class UDialogueWave* Dialogue, const struct FDialogueContext& Context, class USceneComponent* AttachToComponent, FName AttachPointName = NAME_None, FVector Location = FVector(ForceInit), EAttachLocation::Type LocationType = EAttachLocation::KeepRelativeOffset, bool bStopWhenAttachedToDestroyed = false, float VolumeMultiplier = 1.f, float PitchMultiplier = 1.f, float StartTime = 0.f, class USoundAttenuation* AttenuationSettings = NULL)
+	{
+		return SpawnDialogueAttached(Dialogue, Context, AttachToComponent, AttachPointName, Location, FRotator::ZeroRotator, LocationType, bStopWhenAttachedToDestroyed, VolumeMultiplier, PitchMultiplier, StartTime, AttenuationSettings);
+	}
+
+	static class UAudioComponent* SpawnDialogueAttached(class UDialogueWave* Dialogue, const struct FDialogueContext& Context, class USceneComponent* AttachToComponent, FName AttachPointName, FVector Location, EAttachLocation::Type LocationType = EAttachLocation::KeepRelativeOffset, bool bStopWhenAttachedToDestroyed = false, float VolumeMultiplier = 1.f, float PitchMultiplier = 1.f, float StartTime = 0.f, class USoundAttenuation* AttenuationSettings = NULL)
+	{
+		return SpawnDialogueAttached(Dialogue, Context, AttachToComponent, AttachPointName, Location, FRotator::ZeroRotator, LocationType, bStopWhenAttachedToDestroyed, VolumeMultiplier, PitchMultiplier, StartTime, AttenuationSettings);
+	}
 
 	// --- Audio Functions ----------------------------
 	/** Set the sound mix of the audio system for special EQing **/
diff --git a/Engine/Source/Runtime/Engine/Private/GameplayStatics.cpp b/Engine/Source/Runtime/Engine/Private/GameplayStatics.cpp
--- a/Engine/Source/Runtime/Engine/Private/GameplayStatics.cpp
+++ b/Engine/Source/Runtime/Engine/Private/GameplayStatics.cpp
@@ -765,75 +765,44 @@ void UGameplayStatics::PlaySound2D(UObject* WorldContextObject, class USoundBase
 	}
 }
 
-void UGameplayStatics::PlaySoundAtLocation(UObject* WorldContextObject, class USoundBase* Sound, FVector Location, float VolumeMultiplier, float PitchMultiplier, float StartTime, class USoundAttenuation* AttenuationSettings)
+UAudioComponent* UGameplayStatics::SpawnSound2D(UObject* WorldContextObject, class USoundBase* Sound, float VolumeMultiplier, float PitchMultiplier, float StartTime)
 {
 	if (!Sound || !GEngine || !GEngine->UseSound())
 	{
-		return;
+		return nullptr;
 	}
 
 	UWorld* ThisWorld = GEngine->GetWorldFromContextObject(WorldContextObject);
 	if (!ThisWorld || !ThisWorld->bAllowAudioPlayback || ThisWorld->GetNetMode() == NM_DedicatedServer)
 	{
-		return;
+		return nullptr;
 	}
 
-	if (FAudioDevice* AudioDevice = ThisWorld->GetAudioDevice())
+	UAudioComponent* AudioComponent = FAudioDevice::CreateComponent(Sound, ThisWorld, ThisWorld->GetWorldSettings(), false);
+	if (AudioComponent)
 	{
-		if (Sound->IsAudibleSimple(AudioDevice, Location, AttenuationSettings))
-		{
-			const bool bIsInGameWorld = ThisWorld->IsGameWorld();
-
-			FActiveSound NewActiveSound;
-			NewActiveSound.World = ThisWorld;
-			NewActiveSound.Sound = Sound;
-
-			NewActiveSound.VolumeMultiplier = VolumeMultiplier;
-			NewActiveSound.PitchMultiplier = PitchMultiplier;
-
-			NewActiveSound.RequestedStartTime = FMath::Max(0.f, StartTime);
-
-			NewActiveSound.bLocationDefined = true;
-			NewActiveSound.Transform.SetTranslation(Location);
-
-			NewActiveSound.bIsUISound = !bIsInGameWorld;
-			NewActiveSound.bHandleSubtitles = true;
-			NewActiveSound.SubtitlePriority = 10000.f; // Fixme: pass in? Do we want that exposed to blueprints though?
-
-			const FAttenuationSettings* AttenuationSettingsToApply = (AttenuationSettings ? &AttenuationSettings->Attenuation : Sound->GetAttenuationSettingsToApply());
-
-			NewActiveSound.bHasAttenuationSettings = (bIsInGameWorld && AttenuationSettingsToApply);
-			if (NewActiveSound.bHasAttenuationSettings)
-			{
-				NewActiveSound.AttenuationSettings = *AttenuationSettingsToApply;
-			}
+		const bool bIsInGameWorld = AudioComponent->GetWorld()->IsGameWorld();
 
-			// TODO - Audio Threading. This call would be a task call to dispatch to the audio thread
-			AudioDevice->AddNewActiveSound(NewActiveSound);
-		}
-		else
-		{
-			// Don't play a sound for short sounds that start out of range of any listener
-			UE_LOG(LogAudio, Log, TEXT("Sound not played for out of range Sound %s"), *Sound->GetName());
-		}
+		AudioComponent->SetVolumeMultiplier(VolumeMultiplier);
+		AudioComponent->SetPitchMultiplier(PitchMultiplier);
+		AudioComponent->bAllowSpatialization	= false;
+		AudioComponent->bIsUISound				= true;
+		AudioComponent->bAutoDestroy			= true;
+		AudioComponent->SubtitlePriority		= 10000.f; // Fixme: pass in? Do we want that exposed to blueprints though?
+		AudioComponent->Play(StartTime);
 	}
+	return AudioComponent;
 }
 
-void UGameplayStatics::PlayDialogueAtLocation(UObject* WorldContextObject, class UDialogueWave* Dialogue, const FDialogueContext& Context, FVector Location, float VolumeMultiplier, float PitchMultiplier, float StartTime, class USoundAttenuation* AttenuationSettings)
+void UGameplayStatics::PlaySoundAtLocation(UObject* WorldContextObject, class USoundBase* Sound, FVector Location, FRotator Rotation, float VolumeMultiplier, float PitchMultiplier, float StartTime, class USoundAttenuation* AttenuationSettings)
 {
-	if (!Dialogue || !GEngine || !GEngine->UseSound())
+	if (!Sound || !GEngine || !GEngine->UseSound())
 	{
 		return;
 	}
 
 	UWorld* ThisWorld = GEngine->GetWorldFromContextObject(WorldContextObject);
-	if (!ThisWorld || !ThisWorld->bAllowAudioPlayback)
-	{
-		return;
-	}
-	
-	USoundBase* Sound = Dialogue->GetWaveFromContext(Context);
-	if (!Sound)
+	if (!ThisWorld || !ThisWorld->bAllowAudioPlayback || ThisWorld->GetNetMode() == NM_DedicatedServer)
 	{
 		return;
 	}
@@ -843,7 +812,7 @@ void UGameplayStatics::PlayDialogueAtLocation(UObject* WorldContextObject, class
 		if (Sound->IsAudibleSimple(AudioDevice, Location, AttenuationSettings))
 		{
 			const bool bIsInGameWorld = ThisWorld->IsGameWorld();
-			
+
 			FActiveSound NewActiveSound;
 			NewActiveSound.World = ThisWorld;
 			NewActiveSound.Sound = Sound;
@@ -855,6 +824,7 @@ void UGameplayStatics::PlayDialogueAtLocation(UObject* WorldContextObject, class
 
 			NewActiveSound.bLocationDefined = true;
 			NewActiveSound.Transform.SetTranslation(Location);
+			NewActiveSound.Transform.SetRotation(FQuat(Rotation));
 
 			NewActiveSound.bIsUISound = !bIsInGameWorld;
 			NewActiveSound.bHandleSubtitles = true;
@@ -868,6 +838,7 @@ void UGameplayStatics::PlayDialogueAtLocation(UObject* WorldContextObject, class
 				NewActiveSound.AttenuationSettings = *AttenuationSettingsToApply;
 			}
 
+			// TODO - Audio Threading. This call would be a task call to dispatch to the audio thread
 			AudioDevice->AddNewActiveSound(NewActiveSound);
 		}
 		else
@@ -878,40 +849,26 @@ void UGameplayStatics::PlayDialogueAtLocation(UObject* WorldContextObject, class
 	}
 }
 
-class UAudioComponent* UGameplayStatics::PlaySoundAttached(class USoundBase* Sound, class USceneComponent* AttachToComponent, FName AttachPointName, FVector Location, EAttachLocation::Type LocationType, bool bStopWhenAttachedToDestroyed, float VolumeMultiplier, float PitchMultiplier, float StartTime, class USoundAttenuation* AttenuationSettings)
+UAudioComponent* UGameplayStatics::SpawnSoundAtLocation(UObject* WorldContextObject, class USoundBase* Sound, FVector Location, FRotator Rotation, float VolumeMultiplier, float PitchMultiplier, float StartTime, class USoundAttenuation* AttenuationSettings)
 {
-	if (!Sound)
+	if (!Sound || !GEngine || !GEngine->UseSound())
 	{
 		return nullptr;
 	}
 
-	if (!AttachToComponent)
+	UWorld* ThisWorld = GEngine->GetWorldFromContextObject(WorldContextObject);
+	if (!ThisWorld || !ThisWorld->bAllowAudioPlayback || ThisWorld->GetNetMode() == NM_DedicatedServer)
 	{
-		UE_LOG(LogScript, Warning, TEXT("UGameplayStatics::PlaySoundAttached: NULL AttachComponent specified!"));
 		return nullptr;
 	}
 
-	// Location used to check whether to create a component if out of range
-	FVector TestLocation = Location;
-	if (LocationType != EAttachLocation::KeepWorldPosition)
-	{
-		TestLocation = AttachToComponent->GetRelativeTransform().TransformPosition(Location);
-	}
+	const bool bIsInGameWorld = ThisWorld->IsGameWorld();
+
+	UAudioComponent* AudioComponent = FAudioDevice::CreateComponent(Sound, ThisWorld, ThisWorld->GetWorldSettings(), false, false, &Location, AttenuationSettings);
 
-	UAudioComponent* AudioComponent = FAudioDevice::CreateComponent(Sound, AttachToComponent->GetWorld(), AttachToComponent->GetOwner(), false, bStopWhenAttachedToDestroyed, &TestLocation, AttenuationSettings);
 	if (AudioComponent)
 	{
-		const bool bIsInGameWorld = AudioComponent->GetWorld()->IsGameWorld();
-
-		AudioComponent->AttachTo(AttachToComponent, AttachPointName);
-		if (LocationType == EAttachLocation::KeepWorldPosition)
-		{
-			AudioComponent->SetWorldLocation(Location);
-		}
-		else
-		{
-			AudioComponent->SetRelativeLocation(Location);
-		}
+		AudioComponent->SetWorldLocationAndRotation(Location, Rotation);
 		AudioComponent->SetVolumeMultiplier(VolumeMultiplier);
 		AudioComponent->SetPitchMultiplier(PitchMultiplier);
 		AudioComponent->bAllowSpatialization	= bIsInGameWorld;
@@ -924,14 +881,8 @@ class UAudioComponent* UGameplayStatics::PlaySoundAttached(class USoundBase* Sou
 	return AudioComponent;
 }
 
-class UAudioComponent* UGameplayStatics::PlayDialogueAttached(class UDialogueWave* Dialogue, const FDialogueContext& Context, class USceneComponent* AttachToComponent, FName AttachPointName, FVector Location, EAttachLocation::Type LocationType, bool bStopWhenAttachedToDestroyed, float VolumeMultiplier, float PitchMultiplier, float StartTime, class USoundAttenuation* AttenuationSettings)
+class UAudioComponent* UGameplayStatics::SpawnSoundAttached(class USoundBase* Sound, class USceneComponent* AttachToComponent, FName AttachPointName, FVector Location, FRotator Rotation, EAttachLocation::Type LocationType, bool bStopWhenAttachedToDestroyed, float VolumeMultiplier, float PitchMultiplier, float StartTime, class USoundAttenuation* AttenuationSettings)
 {
-	if (!Dialogue)
-	{
-		return nullptr;
-	}
-
-	USoundBase* Sound = Dialogue->GetWaveFromContext(Context);
 	if (!Sound)
 	{
 		return nullptr;
@@ -953,21 +904,21 @@ class UAudioComponent* UGameplayStatics::PlayDialogueAttached(class UDialogueWav
 	UAudioComponent* AudioComponent = FAudioDevice::CreateComponent(Sound, AttachToComponent->GetWorld(), AttachToComponent->GetOwner(), false, bStopWhenAttachedToDestroyed, &TestLocation, AttenuationSettings);
 	if (AudioComponent)
 	{
-		const bool bIsGameWorld = AudioComponent->GetWorld()->IsGameWorld();
+		const bool bIsInGameWorld = AudioComponent->GetWorld()->IsGameWorld();
 
 		AudioComponent->AttachTo(AttachToComponent, AttachPointName);
 		if (LocationType == EAttachLocation::KeepWorldPosition)
 		{
-			AudioComponent->SetWorldLocation(Location);
+			AudioComponent->SetWorldLocationAndRotation(Location, Rotation);
 		}
 		else
 		{
-			AudioComponent->SetRelativeLocation(Location);
+			AudioComponent->SetRelativeLocationAndRotation(Location, Rotation);
 		}
 		AudioComponent->SetVolumeMultiplier(VolumeMultiplier);
 		AudioComponent->SetPitchMultiplier(PitchMultiplier);
-		AudioComponent->bAllowSpatialization	= bIsGameWorld;
-		AudioComponent->bIsUISound				= !bIsGameWorld;
+		AudioComponent->bAllowSpatialization	= bIsInGameWorld;
+		AudioComponent->bIsUISound				= !bIsInGameWorld;
 		AudioComponent->bAutoDestroy			= true;
 		AudioComponent->SubtitlePriority		= 10000.f; // Fixme: pass in? Do we want that exposed to blueprints though?
 		AudioComponent->Play(StartTime);
@@ -976,6 +927,49 @@ class UAudioComponent* UGameplayStatics::PlayDialogueAttached(class UDialogueWav
 	return AudioComponent;
 }
 
+void UGameplayStatics::PlayDialogue2D(UObject* WorldContextObject, class UDialogueWave* Dialogue, const FDialogueContext& Context, float VolumeMultiplier, float PitchMultiplier, float StartTime)
+{
+	if (Dialogue)
+	{
+		PlaySound2D(WorldContextObject, Dialogue->GetWaveFromContext(Context), VolumeMultiplier, PitchMultiplier, StartTime);
+	}
+}
+
+UAudioComponent* UGameplayStatics::SpawnDialogue2D(UObject* WorldContextObject, class UDialogueWave* Dialogue, const FDialogueContext& Context, float VolumeMultiplier, float PitchMultiplier, float StartTime)
+{
+	if (Dialogue)
+	{
+		return SpawnSound2D(WorldContextObject, Dialogue->GetWaveFromContext(Context), VolumeMultiplier, PitchMultiplier, StartTime);
+	}
+	return nullptr;
+}
+
+void UGameplayStatics::PlayDialogueAtLocation(UObject* WorldContextObject, class UDialogueWave* Dialogue, const FDialogueContext& Context, FVector Location, FRotator Rotation, float VolumeMultiplier, float PitchMultiplier, float StartTime, class USoundAttenuation* AttenuationSettings)
+{
+	if (Dialogue)
+	{
+		PlaySoundAtLocation(WorldContextObject, Dialogue->GetWaveFromContext(Context), Location, Rotation, VolumeMultiplier, PitchMultiplier, StartTime, AttenuationSettings);
+	}
+}
+
+UAudioComponent* UGameplayStatics::SpawnDialogueAtLocation(UObject* WorldContextObject, class UDialogueWave* Dialogue, const FDialogueContext& Context, FVector Location, FRotator Rotation, float VolumeMultiplier, float PitchMultiplier, float StartTime, class USoundAttenuation* AttenuationSettings)
+{
+	if (Dialogue)
+	{
+		return SpawnSoundAtLocation(WorldContextObject, Dialogue->GetWaveFromContext(Context), Location, Rotation, VolumeMultiplier, PitchMultiplier, StartTime, AttenuationSettings);
+	}
+	return nullptr;
+}
+
+class UAudioComponent* UGameplayStatics::SpawnDialogueAttached(class UDialogueWave* Dialogue, const FDialogueContext& Context, class USceneComponent* AttachToComponent, FName AttachPointName, FVector Location, FRotator Rotation, EAttachLocation::Type LocationType, bool bStopWhenAttachedToDestroyed, float VolumeMultiplier, float PitchMultiplier, float StartTime, class USoundAttenuation* AttenuationSettings)
+{
+	if (Dialogue)
+	{
+		return SpawnSoundAttached(Dialogue->GetWaveFromContext(Context), AttachToComponent, AttachPointName, Location, Rotation, LocationType, bStopWhenAttachedToDestroyed, VolumeMultiplier, PitchMultiplier, StartTime, AttenuationSettings);
+	}
+	return nullptr;
+}
+
 void UGameplayStatics::SetBaseSoundMix(UObject* WorldContextObject, USoundMix* InSoundMix)
 {
 	if (!InSoundMix || !GEngine || !GEngine->UseSound())
