Fixed lights not being hidden when changing the editor visibility via the world outliner

https://jira.ol.epicgames.net/browse/UE-18608

diff --git a/Engine/Source/Runtime/Engine/Private/Components/LightComponent.cpp b/Engine/Source/Runtime/Engine/Private/Components/LightComponent.cpp
--- a/Engine/Source/Runtime/Engine/Private/Components/LightComponent.cpp
+++ b/Engine/Source/Runtime/Engine/Private/Components/LightComponent.cpp
@@ -642,7 +642,7 @@ void ULightComponent::CreateRenderState_Concurrent()
 
 	bool bHidden = false;
 #if WITH_EDITORONLY_DATA
-	bHidden = GetOwner() ? GetOwner()->bHiddenEdLevel : false;
+	bHidden = GetOwner() ? GetOwner()->IsHiddenEd() : false;
 #endif // WITH_EDITORONLY_DATA
 
 	if(!ShouldComponentAddToScene())
