These cached strings are only used once... not caching them

diff --git a/Engine/Source/Runtime/Engine/Classes/EdGraph/EdGraphSchema.h b/Engine/Source/Runtime/Engine/Classes/EdGraph/EdGraphSchema.h
--- a/Engine/Source/Runtime/Engine/Classes/EdGraph/EdGraphSchema.h
+++ b/Engine/Source/Runtime/Engine/Classes/EdGraph/EdGraphSchema.h
@@ -80,18 +80,6 @@ struct ENGINE_API FEdGraphSchemaAction
 	UPROPERTY()
 	int32 SectionID;
 
-	/** Search title for the action (doesn't have to be set when instantiated, will be constructed by GetSearchTitle() if left empty). */
-	UPROPERTY()
-	FString CachedSearchTitle;
-
-	/** Search keywords for the action (doesn't have to be set when instantiated, will be constructed by GetSearchTitle() if left empty). */
-	UPROPERTY()
-	FString CachedSearchKeywords;
-
-	/** Search categories for the action (doesn't have to be set when instantiated, will be constructed by GetSearchTitle() if left empty). */
-	UPROPERTY()
-	FString CachedSearchCategories;
-
 	FEdGraphSchemaAction() 
 		: Grouping(0)
 		, SectionID(0)
@@ -131,42 +119,27 @@ struct ENGINE_API FEdGraphSchemaAction
 	}
 
 	/** Retrieves the full searchable title for this action. */
-	const FString& GetSearchTitle()
+	FString GetSearchTitle()
 	{
-		if(CachedSearchTitle.IsEmpty())
-		{
-			CachedSearchTitle = MenuDescription.ToString();
-			CachedSearchTitle.AppendChar(TEXT(' '));
-			CachedSearchTitle.Append(MenuDescription.BuildSourceString());
-			CachedSearchTitle = CachedSearchTitle.ToLower();
-		}
-		return CachedSearchTitle;
+		return (MenuDescription.ToString() +
+			TEXT(' ') +
+			MenuDescription.BuildSourceString()).ToLower();
 	}
 
 	/** Retrieves the full searchable keywords for this action. */
-	const FString& GetSearchKeywords()
+	FString GetSearchKeywords()
 	{
-		if(CachedSearchKeywords.IsEmpty())
-		{
-			CachedSearchKeywords = Keywords.ToString();
-			CachedSearchKeywords.AppendChar(TEXT(' '));
-			CachedSearchKeywords.Append(Keywords.BuildSourceString());
-			CachedSearchKeywords = CachedSearchKeywords.ToLower();
-		}
-		return CachedSearchKeywords;
+		return (Keywords.ToString() +
+			TEXT(' ') +
+			Keywords.BuildSourceString()).ToLower();
 	}
 
 	/** Retrieves the full searchable categories for this action. */
-	const FString& GetSearchCategory()
+	FString GetSearchCategory()
 	{
-		if(CachedSearchCategories.IsEmpty())
-		{
-			CachedSearchCategories = Category.ToString();
-			CachedSearchCategories.AppendChar(TEXT(' '));
-			CachedSearchCategories.Append(Category.BuildSourceString());
-			CachedSearchCategories = CachedSearchCategories.ToLower();
-		}
-		return CachedSearchCategories;
+		return (Category.ToString() +
+			TEXT(' ') +
+			Category.BuildSourceString()).ToLower();
 	}
 
 	// GC.
diff --git a/Engine/Source/Runtime/Engine/Private/EdGraph/EdGraphSchema.cpp b/Engine/Source/Runtime/Engine/Private/EdGraph/EdGraphSchema.cpp
--- a/Engine/Source/Runtime/Engine/Private/EdGraph/EdGraphSchema.cpp
+++ b/Engine/Source/Runtime/Engine/Private/EdGraph/EdGraphSchema.cpp
@@ -168,15 +168,12 @@ void FGraphActionListBuilderBase::ActionGroup::InitScoringData()
 	if (Actions.Num() > 0)
 	{
 		FEdGraphSchemaAction& FirstAction = *Actions[0];
-		const FString& Keywords = FirstAction.GetSearchKeywords();
-		const FString& Title = FirstAction.GetSearchTitle();
-		const FString& Category = FirstAction.GetSearchCategory();
 
 		// We keep these individual arrays so that they can be weighted differently by the scoring algorithm in SGraphActionMenu::GetActionFilteredWeight
-		Keywords.ParseIntoArray(SearchKeywordsArray, TEXT(" "), true);
+		FirstAction.GetSearchKeywords().ParseIntoArray(SearchKeywordsArray, TEXT(" "), true);
 		FirstAction.MenuDescription.ToString().ToLower().ParseIntoArray(MenuDescriptionArray, TEXT(" "), true);
-		Title.ParseIntoArray(SearchTitleArray, TEXT(" "), true);
-		Category.ParseIntoArray(SearchCategoryArray, TEXT(" "), true);
+		FirstAction.GetSearchTitle().ParseIntoArray(SearchTitleArray, TEXT(" "), true);
+		FirstAction.GetSearchCategory().ParseIntoArray(SearchCategoryArray, TEXT(" "), true);
 
 		// Glob search text together, we use the SearchText string for basic filtering:
 		for (const auto& Entry : SearchTitleArray)
