Fixes Slave Component LOD issue

- https://jira.ol.epicgames.net/browse/UE-9206

diff --git a/Engine/Source/Runtime/Engine/Private/Components/SkinnedMeshComponent.cpp b/Engine/Source/Runtime/Engine/Private/Components/SkinnedMeshComponent.cpp
--- a/Engine/Source/Runtime/Engine/Private/Components/SkinnedMeshComponent.cpp
+++ b/Engine/Source/Runtime/Engine/Private/Components/SkinnedMeshComponent.cpp
@@ -343,13 +343,6 @@ void USkinnedMeshComponent::SendRenderDynamicData_Concurrent()
 		SCOPE_CYCLE_COUNTER(STAT_MeshObjectUpdate);
 
 		int32 UseLOD = PredictedLODLevel;
-		// If we have a MasterPoseComponent - force this component to render at that LOD, so all bones are present for it.
-		// Note that this currently relies on the behaviour where this mesh is rendered at the LOD we pass in here for all viewports. We should
-		// be able to render it at lower LOD on viewports where it is further away. That will make the MasterPoseComponent case a bit harder to solve.
-		if(MasterPoseComponent.IsValid())
-		{
-			UseLOD = FMath::Clamp(MasterPoseComponent->PredictedLODLevel, 0, MeshObject->GetSkeletalMeshResource().LODModels.Num()-1);
-		}
 
 		// Are morph targets disabled for this LOD?
 		if ( SkeletalMesh->LODInfo[ UseLOD ].bHasBeenSimplified || bDisableMorphTarget )
