fixed comment

diff --git a/Engine/Source/Runtime/Engine/Private/Components/BoxComponent.cpp b/Engine/Source/Runtime/Engine/Private/Components/BoxComponent.cpp
--- a/Engine/Source/Runtime/Engine/Private/Components/BoxComponent.cpp
+++ b/Engine/Source/Runtime/Engine/Private/Components/BoxComponent.cpp
@@ -83,7 +83,7 @@ FBoxSphereBounds UBoxComponent::CalcBounds(const FTransform& LocalToWorld) const
 
 FPrimitiveSceneProxy* UBoxComponent::CreateSceneProxy()
 {
-	/** Represents a UCapsuleComponent to the scene manager. */
+	/** Represents a UBoxComponent to the scene manager. */
 	class FBoxSceneProxy : public FPrimitiveSceneProxy
 	{
 	public:
