#ue4 - Don't allow negative values for CapsuleComponent radius or half-height.

diff --git a/Engine/Source/Runtime/Engine/Classes/Components/CapsuleComponent.h b/Engine/Source/Runtime/Engine/Classes/Components/CapsuleComponent.h
--- a/Engine/Source/Runtime/Engine/Classes/Components/CapsuleComponent.h
+++ b/Engine/Source/Runtime/Engine/Classes/Components/CapsuleComponent.h
@@ -14,17 +14,17 @@ class ENGINE_API UCapsuleComponent : public UShapeComponent
 
 protected:
 	/** 
-	 *	Half-height, i.e. from center of capsule to end of top or bottom hemisphere.  
+	 *	Half-height, from center of capsule to the end of top or bottom hemisphere.  
 	 *	This cannot be less than CapsuleRadius.
 	 */
-	UPROPERTY(EditAnywhere, BlueprintReadOnly, export, Category=Shape)
+	UPROPERTY(EditAnywhere, BlueprintReadOnly, export, Category=Shape, meta=(ClampMin="0", UIMin="0"))
 	float CapsuleHalfHeight;
 
 	/** 
 	 *	Radius of cap hemispheres and center cylinder. 
 	 *	This cannot be more than CapsuleHalfHeight.
 	 */
-	UPROPERTY(EditAnywhere, BlueprintReadOnly, export, Category=Shape)
+	UPROPERTY(EditAnywhere, BlueprintReadOnly, export, Category=Shape, meta=(ClampMin="0", UIMin="0"))
 	float CapsuleRadius;
 
 protected:
@@ -115,8 +115,8 @@ public:
 	// Sets the capsule size without triggering a render or physics update. This is the preferred method when initializing a component in a class constructor.
 	FORCEINLINE void InitCapsuleSize(float InRadius, float InHalfHeight)
 	{
-		CapsuleRadius = InRadius;
-		CapsuleHalfHeight = FMath::Max(InHalfHeight, InRadius);
+		CapsuleRadius = FMath::Max(0.f, InRadius);
+		CapsuleHalfHeight = FMath::Max3(0.f, InHalfHeight, InRadius);
 	}
 };
 
diff --git a/Engine/Source/Runtime/Engine/Private/Components/CapsuleComponent.cpp b/Engine/Source/Runtime/Engine/Private/Components/CapsuleComponent.cpp
--- a/Engine/Source/Runtime/Engine/Private/Components/CapsuleComponent.cpp
+++ b/Engine/Source/Runtime/Engine/Private/Components/CapsuleComponent.cpp
@@ -110,7 +110,7 @@ void UCapsuleComponent::Serialize(FArchive& Ar)
 		}
 	}
 
-	CapsuleHalfHeight = FMath::Max(CapsuleHalfHeight, CapsuleRadius);
+	CapsuleHalfHeight = FMath::Max3(0.f, CapsuleHalfHeight, CapsuleRadius);
 }
 
 #if WITH_EDITOR
@@ -122,11 +122,11 @@ void UCapsuleComponent::PostEditChangeProperty(FPropertyChangedEvent& PropertyCh
 	// things like propagation from CDO to instances don't work correctly if changing one property causes a different property to change
 	if (PropertyName == GET_MEMBER_NAME_CHECKED(UCapsuleComponent, CapsuleHalfHeight))
 	{
-		CapsuleHalfHeight = FMath::Max(CapsuleHalfHeight, CapsuleRadius);
+		CapsuleHalfHeight = FMath::Max3(0.f, CapsuleHalfHeight, CapsuleRadius);
 	}
 	else if (PropertyName == GET_MEMBER_NAME_CHECKED(UCapsuleComponent, CapsuleRadius))
 	{
-		CapsuleRadius = FMath::Min(CapsuleHalfHeight, CapsuleRadius);
+		CapsuleRadius = FMath::Clamp(CapsuleRadius, 0.f, CapsuleHalfHeight);
 	}
 
 	Super::PostEditChangeProperty(PropertyChangedEvent);
@@ -135,8 +135,8 @@ void UCapsuleComponent::PostEditChangeProperty(FPropertyChangedEvent& PropertyCh
 
 void UCapsuleComponent::SetCapsuleSize(float NewRadius, float NewHalfHeight, bool bUpdateOverlaps)
 {
-	CapsuleHalfHeight = FMath::Max(NewHalfHeight, NewRadius);
-	CapsuleRadius = NewRadius;
+	CapsuleHalfHeight = FMath::Max3(0.f, NewHalfHeight, NewRadius);
+	CapsuleRadius = FMath::Max(0.f, NewRadius);
 	MarkRenderStateDirty();
 
 	// do this if already created
