[UE-18152] Fix crash caused by earlier fix when loading some older Blueprint component templates.

change summary:
- relocated code from USCS_Node::PostLoad() to USimpleConstructionScript::PostLoad(), and further restricted to "new" component template names since they're really the only ones potentially affected.

diff --git a/Engine/Source/Runtime/Engine/Private/SCS_Node.cpp b/Engine/Source/Runtime/Engine/Private/SCS_Node.cpp
--- a/Engine/Source/Runtime/Engine/Private/SCS_Node.cpp
+++ b/Engine/Source/Runtime/Engine/Private/SCS_Node.cpp
@@ -407,16 +407,6 @@ void USCS_Node::PostLoad()
 {
 	Super::PostLoad();
 
-	if(ComponentTemplate != nullptr)
-	{
-		// Fix up component templates not in sync w/ the variable name on load; if not, duplicate it to a new instance w/ a corrected name in order to ensure that it is uniquely bound to this node only.
-		const FString ComponentTemplateName = VariableName.ToString() + TEXT("_GEN_VARIABLE");
-		if(ComponentTemplateName != ComponentTemplate->GetName())
-		{
-			ComponentTemplate = static_cast<UActorComponent*>(StaticDuplicateObject(ComponentTemplate, ComponentTemplate->GetOuter(), *ComponentTemplateName));
-		}
-	}
-
 	ValidateGuid();
 }
 
diff --git a/Engine/Source/Runtime/Engine/Private/SimpleConstructionScript.cpp b/Engine/Source/Runtime/Engine/Private/SimpleConstructionScript.cpp
--- a/Engine/Source/Runtime/Engine/Private/SimpleConstructionScript.cpp
+++ b/Engine/Source/Runtime/Engine/Private/SimpleConstructionScript.cpp
@@ -16,6 +16,8 @@
 
 namespace
 {
+	static const FString ComponentTemplateNameSuffix(TEXT("_GEN_VARIABLE"));
+
 	// Helper method to register instanced components post-construction
 	void RegisterInstancedComponent(UActorComponent* InstancedComponent)
 	{
@@ -177,6 +179,15 @@ void USimpleConstructionScript::PostLoad()
 		// Fix up components that may have switched from scene to non-scene type and vice-versa
 		if(Node->ComponentTemplate != nullptr)
 		{
+			// Fix up any component template objects whose name doesn't match the current variable name; this ensures that there is always one unique template per node.
+			FString VariableName = Node->GetVariableName().ToString();
+			FString ComponentTemplateName = Node->ComponentTemplate->GetName();
+			if(ComponentTemplateName.EndsWith(ComponentTemplateNameSuffix) && !ComponentTemplateName.StartsWith(VariableName) && !GIsDuplicatingClassForReinstancing)
+			{
+				Node->ComponentTemplate->ConditionalPostLoad();
+				Node->ComponentTemplate = static_cast<UActorComponent*>(StaticDuplicateObject(Node->ComponentTemplate, Node->ComponentTemplate->GetOuter(), *(VariableName + ComponentTemplateNameSuffix)));
+			}
+
 			// Check to see if switched from scene to a non-scene component type
 			if (!Node->ComponentTemplate->IsA<USceneComponent>())
 			{
@@ -1043,7 +1054,7 @@ USCS_Node* USimpleConstructionScript::CreateNode(UClass* NewComponentClass, FNam
 	// note that naming logic is duplicated in CreateNodeAndRenameComponent:
 	NewComponentVariableName = GenerateNewComponentName(NewComponentClass, NewComponentVariableName);
 
-	UActorComponent* NewComponentTemplate = NewObject<UActorComponent>(Blueprint->GeneratedClass, NewComponentClass, *(NewComponentVariableName.ToString() + TEXT("_GEN_VARIABLE")), RF_ArchetypeObject | RF_Transactional | RF_Public);
+	UActorComponent* NewComponentTemplate = NewObject<UActorComponent>(Blueprint->GeneratedClass, NewComponentClass, *(NewComponentVariableName.ToString() + ComponentTemplateNameSuffix), RF_ArchetypeObject | RF_Transactional | RF_Public);
 
 	return CreateNodeImpl(NewComponentTemplate, NewComponentVariableName);
 }
@@ -1056,7 +1067,7 @@ USCS_Node* USimpleConstructionScript::CreateNodeAndRenameComponent(UActorCompone
 	FName NewComponentVariableName = GenerateNewComponentName(NewComponentTemplate->GetClass());
 
 	// Relocate the instance from the transient package to the BPGC and assign it a unique object name
-	NewComponentTemplate->Rename(*(NewComponentVariableName.ToString() + TEXT("_GEN_VARIABLE")), GetBlueprint()->GeneratedClass, REN_DontCreateRedirectors | REN_DoNotDirty);
+	NewComponentTemplate->Rename(*(NewComponentVariableName.ToString() + ComponentTemplateNameSuffix), GetBlueprint()->GeneratedClass, REN_DontCreateRedirectors | REN_DoNotDirty);
 
 	return CreateNodeImpl(NewComponentTemplate, NewComponentVariableName);
 }
