Made PinFriendlyName and relevant methods defined only in editors that have WITH_EDITORONLY_DATA evaluating as true.

diff --git a/Engine/Source/Runtime/Engine/Classes/EdGraph/EdGraphPin.h b/Engine/Source/Runtime/Engine/Classes/EdGraph/EdGraphPin.h
--- a/Engine/Source/Runtime/Engine/Classes/EdGraph/EdGraphPin.h
+++ b/Engine/Source/Runtime/Engine/Classes/EdGraph/EdGraphPin.h
@@ -174,9 +174,11 @@ class UEdGraphPin : public UObject
 	UPROPERTY()
 	FString PinName;
 
+#if WITH_EDITORONLY_DATA
 	/** Used as the display name if set */
 	UPROPERTY()
 	FText PinFriendlyName;
+#endif
 
 	/** The tool-tip describing this pin's purpose */
 	UPROPERTY()
@@ -312,8 +314,10 @@ public:
 	/** Get the current DefaultObject path name, or DefaultValue if its null */
 	ENGINE_API FString GetDefaultAsString() const;
 
+#if WITH_EDITORONLY_DATA
 	/** Returns how the name of the pin should be displayed in the UI */
 	ENGINE_API FText GetDisplayName() const;
+#endif // WITH_EDITORONLY_DATA
 
 	void ResetDefaultValue()
 	{
@@ -329,7 +333,9 @@ public:
 		PinType.ResetToDefaults();
 
 		PinName.Empty();
+#if WITH_EDITORONLY_DATA
 		PinFriendlyName = FText::GetEmpty();
+#endif // WITH_EDITORONLY_DATA
 		AutogeneratedDefaultValue.Empty();
 		ResetDefaultValue();
 
diff --git a/Engine/Source/Runtime/Engine/Classes/EdGraph/EdGraphSchema.h b/Engine/Source/Runtime/Engine/Classes/EdGraph/EdGraphSchema.h
--- a/Engine/Source/Runtime/Engine/Classes/EdGraph/EdGraphSchema.h
+++ b/Engine/Source/Runtime/Engine/Classes/EdGraph/EdGraphSchema.h
@@ -610,8 +610,10 @@ class ENGINE_API UEdGraphSchema : public UObject
 	 */
 	virtual FLinearColor GetPinTypeColor(const FEdGraphPinType& PinType) const { return FLinearColor::Black; }
 
+#if WITH_EDITORONLY_DATA
 	/** Get the name to show in the editor */
 	virtual FText GetPinDisplayName(const UEdGraphPin* Pin) const;
+#endif // WITH_EDITORONLY_DATA
 
 	/**
 	 * Takes the PinDescription and tacks on any other data important to the 
diff --git a/Engine/Source/Runtime/Engine/Private/EdGraph/EdGraphPin.cpp b/Engine/Source/Runtime/Engine/Private/EdGraph/EdGraphPin.cpp
--- a/Engine/Source/Runtime/Engine/Private/EdGraph/EdGraphPin.cpp
+++ b/Engine/Source/Runtime/Engine/Private/EdGraph/EdGraphPin.cpp
@@ -260,6 +260,7 @@ FString UEdGraphPin::GetDefaultAsString() const
 	}
 }
 
+#if WITH_EDITORONLY_DATA
 FText UEdGraphPin::GetDisplayName() const
 {
 	FText DisplayName = FText::GetEmpty();
@@ -281,6 +282,7 @@ FText UEdGraphPin::GetDisplayName() const
 	}
 	return DisplayName;
 }
+#endif // WITH_EDITORONLY_DATA
 
 const FString UEdGraphPin::GetLinkInfoString( const FString& InFunctionName, const FString& InInfoData, const UEdGraphPin* InToPin ) const
 {
diff --git a/Engine/Source/Runtime/Engine/Private/EdGraph/EdGraphSchema.cpp b/Engine/Source/Runtime/Engine/Private/EdGraph/EdGraphSchema.cpp
--- a/Engine/Source/Runtime/Engine/Private/EdGraph/EdGraphSchema.cpp
+++ b/Engine/Source/Runtime/Engine/Private/EdGraph/EdGraphSchema.cpp
@@ -593,6 +593,7 @@ FPinConnectionResponse UEdGraphSchema::CopyPinLinks(UEdGraphPin& CopyFromPin, UE
 	return FinalResponse;
 }
 
+#if WITH_EDITORONLY_DATA
 FText UEdGraphSchema::GetPinDisplayName(const UEdGraphPin* Pin) const
 {
 	FText ResultPinName;
@@ -613,6 +614,7 @@ FText UEdGraphSchema::GetPinDisplayName(const UEdGraphPin* Pin) const
 	}
 	return ResultPinName;
 }
+#endif // WITH_EDITORONLY_DATA
 
 void UEdGraphSchema::ConstructBasicPinTooltip(UEdGraphPin const& Pin, FText const& PinDescription, FString& TooltipOut) const
 {
