Added a sanity check to APlayerCameraManager's processing of DefaultModifiers #UE4

UE-18630

diff --git a/Engine/Source/Runtime/Engine/Private/PlayerCameraManager.cpp b/Engine/Source/Runtime/Engine/Private/PlayerCameraManager.cpp
--- a/Engine/Source/Runtime/Engine/Private/PlayerCameraManager.cpp
+++ b/Engine/Source/Runtime/Engine/Private/PlayerCameraManager.cpp
@@ -733,13 +733,17 @@ void APlayerCameraManager::PostInitializeComponents()
 	{
 		for (auto ModifierClass : DefaultModifiers)
 		{
-			UCameraModifier* const NewMod = AddNewCameraModifier(ModifierClass);
-		
-			// cache ref to camera shake if this is it
-			UCameraModifier_CameraShake* const ShakeMod = Cast<UCameraModifier_CameraShake>(NewMod);
-			if (ShakeMod)
+			// empty entries are not valid here, do work only for actual classes
+			if (ModifierClass)
 			{
-				CachedCameraShakeMod = ShakeMod;
+				UCameraModifier* const NewMod = AddNewCameraModifier(ModifierClass);
+
+				// cache ref to camera shake if this is it
+				UCameraModifier_CameraShake* const ShakeMod = Cast<UCameraModifier_CameraShake>(NewMod);
+				if (ShakeMod)
+				{
+					CachedCameraShakeMod = ShakeMod;
+				}
 			}
 		}
 	}
