Fix for substepping crash when changing worlds due to scene interactions during simulation.

diff --git a/Engine/Source/Runtime/Engine/Private/Components/DestructibleComponent.cpp b/Engine/Source/Runtime/Engine/Private/Components/DestructibleComponent.cpp
--- a/Engine/Source/Runtime/Engine/Private/Components/DestructibleComponent.cpp
+++ b/Engine/Source/Runtime/Engine/Private/Components/DestructibleComponent.cpp
@@ -362,18 +362,16 @@ void UDestructibleComponent::CreatePhysicsState()
 	PhysxChunkUserData.Reset(ChunkCount);
 	PhysxChunkUserData.AddZeroed(ChunkCount);
 
-	if(PhysScene->DeferredCommandHandler.HasPendingCommands())
 	{
 		// Lock and flush deferred command handler here to stop any currently pending deletions from affecting new actors.
-		// Only do this if we have any commands to flush to avoid the scene lock if possible
 		SCOPED_SCENE_WRITE_LOCK(PScene);
 		PhysScene->DeferredCommandHandler.Flush();
+
+		// Create an APEX NxDestructibleActor from the Destructible asset and actor descriptor, has to be in a locked scope as it affects simulation state
+		ApexDestructibleActor = static_cast<NxDestructibleActor*>(TheDestructibleMesh->ApexDestructibleAsset->createApexActor(*ActorParams, *ApexScene));
+		check(ApexDestructibleActor);
 	}
 	
-	// Create an APEX NxDestructibleActor from the Destructible asset and actor descriptor
-	ApexDestructibleActor = static_cast<NxDestructibleActor*>(TheDestructibleMesh->ApexDestructibleAsset->createApexActor(*ActorParams, *ApexScene));
-	check(ApexDestructibleActor);
-
 	// Make a backpointer to this component
 	PhysxUserData = FPhysxUserData(this);
 	ApexDestructibleActor->userData = &PhysxUserData;
diff --git a/Engine/Source/Runtime/Engine/Private/World.cpp b/Engine/Source/Runtime/Engine/Private/World.cpp
--- a/Engine/Source/Runtime/Engine/Private/World.cpp
+++ b/Engine/Source/Runtime/Engine/Private/World.cpp
@@ -3063,6 +3063,12 @@ void UWorld::CleanupWorld(bool bSessionEnded, bool bCleanupResources, UWorld* Ne
 {
 	check(IsVisibilityRequestPending() == false);
 
+	// Wait on current physics scenes if they are processing
+	if(FPhysScene* Scene = GetPhysicsScene())
+	{
+		Scene->WaitPhysScenes();
+	}
+
 	FWorldDelegates::OnWorldCleanup.Broadcast(this, bSessionEnded, bCleanupResources);
 
 	if (AISystem != NULL)
