Make Master pose to be serialized - allow it to run during construction script

- UE-18621

diff --git a/Engine/Source/Runtime/Engine/Classes/Components/SkinnedMeshComponent.h b/Engine/Source/Runtime/Engine/Classes/Components/SkinnedMeshComponent.h
--- a/Engine/Source/Runtime/Engine/Classes/Components/SkinnedMeshComponent.h
+++ b/Engine/Source/Runtime/Engine/Classes/Components/SkinnedMeshComponent.h
@@ -168,6 +168,7 @@ class ENGINE_API USkinnedMeshComponent : public UMeshComponent
 	 *	use the SpaceBases array in the MasterPoseComponent. This is used when constructing a character using multiple skeletal meshes sharing the same
 	 *	skeleton within the same Actor.
 	 */
+	UPROPERTY(BlueprintReadOnly, Category="Mesh")
 	TWeakObjectPtr< class USkinnedMeshComponent > MasterPoseComponent;
 
 private:
@@ -665,7 +666,7 @@ public:
 	 *
 	 * @param NewMasterBoneComponent New MasterPoseComponent
 	 */
-	UFUNCTION(BlueprintCallable, Category="Components|SkinnedMesh", meta=(UnsafeDuringActorConstruction="true"))
+	UFUNCTION(BlueprintCallable, Category="Components|SkinnedMesh")
 	void SetMasterPoseComponent(USkinnedMeshComponent* NewMasterBoneComponent);
 
 	/** 
diff --git a/Engine/Source/Runtime/Engine/Private/Components/SkinnedMeshComponent.cpp b/Engine/Source/Runtime/Engine/Private/Components/SkinnedMeshComponent.cpp
--- a/Engine/Source/Runtime/Engine/Private/Components/SkinnedMeshComponent.cpp
+++ b/Engine/Source/Runtime/Engine/Private/Components/SkinnedMeshComponent.cpp
@@ -224,8 +224,15 @@ void USkinnedMeshComponent::OnRegister()
 
 	AnimUpdateRateParams = FAnimUpdateRateManager::GetUpdateRateParameters(this);
 
-	AllocateTransformData();
-	UpdateMasterBoneMap();	
+	if (MasterPoseComponent.IsValid())
+	{
+		// this has to be called again during register so that it can do related initialization
+		SetMasterPoseComponent(MasterPoseComponent.Get());
+	}
+	else
+	{
+		AllocateTransformData();
+	}
 
 	Super::OnRegister();
 
