Refactored colliison handling when spawning an actor. Users have 4 options, choosable in the actor class and overrideable at spawn time.  Old content should upgrade cleanly to new system.

#docs

diff --git a/Engine/Documentation/Source/Shared/GraphNodes/Blueprint/BlueprintNodes.INT.udn b/Engine/Documentation/Source/Shared/GraphNodes/Blueprint/BlueprintNodes.INT.udn
--- a/Engine/Documentation/Source/Shared/GraphNodes/Blueprint/BlueprintNodes.INT.udn
+++ b/Engine/Documentation/Source/Shared/GraphNodes/Blueprint/BlueprintNodes.INT.udn
@@ -94,7 +94,7 @@ Return Value - This outputs the new Actor instance that was spawned.
 
 The "Spawn Actor from Class Node" node takes an Actor Class (specified via the 'Class' input) and attempts to spawn an instance of that in the world.
 
-The 'Spawn Transform' input is used to define the Actor's spawn location (and its initial orientation). If that location is blocked by some collider, then the Actor will fail to spawn unless you set the 'Spawn Even if Colliding' input to _True_.
+The 'Spawn Transform' input is used to define the Actor's spawn location (and its initial orientation).
 
 ---
 
@@ -103,7 +103,7 @@ The 'Spawn Transform' input is used to define the Actor's spawn location (and it
 Unlabeled Execution Pin - This execution input triggers the spawning. 
 Class - Denotes the Actor Class you wish to spawn (required). 
 Spawn Transform - Takes in a Transform used to position and orient the Actor in the world. 
-Spawn Even if Colliding - Takes in a Boolean value used to determine if the Actor should be spawned when something is blocking the spawn location.  
+Collision Handling Override - Overrides actor class setting for how to handle a situation where something is blocking the spawn location. If "Default" is selected, the actor class's SpawnCollisionHandlingMethod is used.
 
 ---
 
diff --git a/Engine/Plugins/Runtime/GearVR/Source/GearVR/Private/HeadMountedDisplayCommon.cpp b/Engine/Plugins/Runtime/GearVR/Source/GearVR/Private/HeadMountedDisplayCommon.cpp
--- a/Engine/Plugins/Runtime/GearVR/Source/GearVR/Private/HeadMountedDisplayCommon.cpp
+++ b/Engine/Plugins/Runtime/GearVR/Source/GearVR/Private/HeadMountedDisplayCommon.cpp
@@ -1150,7 +1150,7 @@ void FHeadMountedDisplay::DrawSeaOfCubes(UWorld* World, FVector ViewLocation)
 	if (frame->Settings->Flags.bDrawCubes && !SeaOfCubesActorPtr.IsValid())
 	{
 		FActorSpawnParameters SpawnInfo;
-		SpawnInfo.bNoCollisionFail = true;
+		SpawnInfo.SpawnCollisionHandlingOverride = ESpawnActorCollisionHandlingMethod::AlwaysSpawn;
 		SpawnInfo.bNoFail = true;
 		SpawnInfo.ObjectFlags = RF_Transient;
 		AStaticMeshActor* SeaOfCubesActor;
diff --git a/Engine/Plugins/Runtime/OculusRift/Source/OculusRift/Private/HeadMountedDisplayCommon.cpp b/Engine/Plugins/Runtime/OculusRift/Source/OculusRift/Private/HeadMountedDisplayCommon.cpp
--- a/Engine/Plugins/Runtime/OculusRift/Source/OculusRift/Private/HeadMountedDisplayCommon.cpp
+++ b/Engine/Plugins/Runtime/OculusRift/Source/OculusRift/Private/HeadMountedDisplayCommon.cpp
@@ -1150,7 +1150,7 @@ void FHeadMountedDisplay::DrawSeaOfCubes(UWorld* World, FVector ViewLocation)
 	if (frame->Settings->Flags.bDrawCubes && !SeaOfCubesActorPtr.IsValid())
 	{
 		FActorSpawnParameters SpawnInfo;
-		SpawnInfo.bNoCollisionFail = true;
+		SpawnInfo.SpawnCollisionHandlingOverride = ESpawnActorCollisionHandlingMethod::AlwaysSpawn;
 		SpawnInfo.bNoFail = true;
 		SpawnInfo.ObjectFlags = RF_Transient;
 		AStaticMeshActor* SeaOfCubesActor;
diff --git a/Engine/Source/Developer/GameplayDebugger/Private/GameplayDebugger.cpp b/Engine/Source/Developer/GameplayDebugger/Private/GameplayDebugger.cpp
--- a/Engine/Source/Developer/GameplayDebugger/Private/GameplayDebugger.cpp
+++ b/Engine/Source/Developer/GameplayDebugger/Private/GameplayDebugger.cpp
@@ -294,7 +294,7 @@ bool FGameplayDebugger::CreateGameplayDebuggerForPlayerController(APlayerControl
 	}
 
 	FActorSpawnParameters SpawnInfo;
-	SpawnInfo.bNoCollisionFail = true;
+	SpawnInfo.SpawnCollisionHandlingOverride = ESpawnActorCollisionHandlingMethod::AlwaysSpawn;
 	SpawnInfo.Name = *FString::Printf(TEXT("GameplayDebuggingReplicator_%s"), *PlayerController->GetName());
 	AGameplayDebuggingReplicator* DestActor = World->SpawnActor<AGameplayDebuggingReplicator>(FVector::ZeroVector, FRotator::ZeroRotator, SpawnInfo);
 	if (DestActor != NULL)
@@ -359,7 +359,7 @@ void FGameplayDebugger::WorldAdded(UWorld* InWorld)
 
 	// create global replicator on level
 	FActorSpawnParameters SpawnInfo;
-	SpawnInfo.bNoCollisionFail = true;
+	SpawnInfo.SpawnCollisionHandlingOverride = ESpawnActorCollisionHandlingMethod::AlwaysSpawn;
 	SpawnInfo.Name = *FString::Printf(TEXT("GameplayDebuggingReplicator_Global"));
 	AGameplayDebuggingReplicator* DestActor = InWorld->SpawnActor<AGameplayDebuggingReplicator>(FVector::ZeroVector, FRotator::ZeroRotator, SpawnInfo);
 	if (DestActor != NULL)
diff --git a/Engine/Source/Developer/GameplayDebugger/Private/GameplayDebuggingControllerComponent.cpp b/Engine/Source/Developer/GameplayDebugger/Private/GameplayDebuggingControllerComponent.cpp
--- a/Engine/Source/Developer/GameplayDebugger/Private/GameplayDebuggingControllerComponent.cpp
+++ b/Engine/Source/Developer/GameplayDebugger/Private/GameplayDebuggingControllerComponent.cpp
@@ -446,7 +446,7 @@ void UGameplayDebuggingControllerComponent::ToggleDebugCamera()
 	{
 		{
 			FActorSpawnParameters SpawnInfo;
-			SpawnInfo.bNoCollisionFail = true;
+			SpawnInfo.SpawnCollisionHandlingOverride = ESpawnActorCollisionHandlingMethod::AlwaysSpawn;
 			SpawnInfo.Owner = PlayerOwner->GetWorldSettings();
 			SpawnInfo.Instigator = PlayerOwner->Instigator;
 			DebugCameraController = GetWorld()->SpawnActor<ADebugCameraController>(SpawnInfo);
@@ -463,7 +463,7 @@ void UGameplayDebuggingControllerComponent::ToggleDebugCamera()
 			FActorSpawnParameters SpawnInfo;
 			SpawnInfo.Owner = PlayerOwner.Get();
 			SpawnInfo.Instigator = PlayerOwner->Instigator;
-			SpawnInfo.bNoCollisionFail = true;
+			SpawnInfo.SpawnCollisionHandlingOverride = ESpawnActorCollisionHandlingMethod::AlwaysSpawn;
 			SpawnInfo.ObjectFlags |= RF_Transient;	// We never want these to save into a map
 			AGaneplayDebuggerProxyHUD* ProxyHUD = GetWorld()->SpawnActor<AGaneplayDebuggerProxyHUD>(SpawnInfo);
 			ProxyHUD->RedirectedHUD = PlayerOwner->MyHUD;
diff --git a/Engine/Source/Developer/GameplayDebugger/Private/GameplayDebuggingReplicator.cpp b/Engine/Source/Developer/GameplayDebugger/Private/GameplayDebuggingReplicator.cpp
--- a/Engine/Source/Developer/GameplayDebugger/Private/GameplayDebuggingReplicator.cpp
+++ b/Engine/Source/Developer/GameplayDebugger/Private/GameplayDebuggingReplicator.cpp
@@ -689,7 +689,7 @@ void AGameplayDebuggingReplicator::DrawDebugData(class UCanvas* Canvas, class AP
 		FActorSpawnParameters SpawnInfo;
 		SpawnInfo.Owner = NULL;
 		SpawnInfo.Instigator = NULL;
-		SpawnInfo.bNoCollisionFail = true;
+		SpawnInfo.SpawnCollisionHandlingOverride = ESpawnActorCollisionHandlingMethod::AlwaysSpawn;
 
 		DebugRenderer = GetWorld()->SpawnActor<AGameplayDebuggingHUDComponent>(DebugComponentHUDClass.Get(), SpawnInfo);
 		DebugRenderer->SetCanvas(Canvas);
diff --git a/Engine/Source/Developer/LogVisualizer/Private/LogVisualizer.cpp b/Engine/Source/Developer/LogVisualizer/Private/LogVisualizer.cpp
--- a/Engine/Source/Developer/LogVisualizer/Private/LogVisualizer.cpp
+++ b/Engine/Source/Developer/LogVisualizer/Private/LogVisualizer.cpp
@@ -193,7 +193,7 @@ AActor* FLogVisualizer::GetVisualLoggerHelperActor()
 	}
 
 	FActorSpawnParameters SpawnInfo;
-	SpawnInfo.bNoCollisionFail = true;
+	SpawnInfo.SpawnCollisionHandlingOverride = ESpawnActorCollisionHandlingMethod::AlwaysSpawn;
 	SpawnInfo.Name = *FString::Printf(TEXT("VisualLoggerRenderingActor"));
 	AActor* HelperActor = World->SpawnActor<AVisualLoggerRenderingActor>(SpawnInfo);
 
diff --git a/Engine/Source/Developer/LogVisualizer/Private/VisualLoggerCameraController.cpp b/Engine/Source/Developer/LogVisualizer/Private/VisualLoggerCameraController.cpp
--- a/Engine/Source/Developer/LogVisualizer/Private/VisualLoggerCameraController.cpp
+++ b/Engine/Source/Developer/LogVisualizer/Private/VisualLoggerCameraController.cpp
@@ -57,7 +57,7 @@ void AVisualLoggerCameraController::PostInitializeComponents()
 	FActorSpawnParameters SpawnInfo;
 	SpawnInfo.Owner = this;
 	SpawnInfo.Instigator = Instigator;
-	SpawnInfo.bNoCollisionFail = true;
+	SpawnInfo.SpawnCollisionHandlingOverride = ESpawnActorCollisionHandlingMethod::AlwaysSpawn;
 	MyHUD = GetWorld()->SpawnActor<AVisualLoggerHUD>(SpawnInfo);
 
 	/** TO DO: somehow this two configuration settings are set to zero for this controller, have to fix it later */
@@ -82,7 +82,7 @@ AVisualLoggerCameraController* AVisualLoggerCameraController::EnableCamera(UWorl
 		// in editor world (if running PIE). Hate it, but it works, and 
 		// this is a debugging tool		
 		FActorSpawnParameters SpawnInfo;
-		SpawnInfo.bNoCollisionFail = true;
+		SpawnInfo.SpawnCollisionHandlingOverride = ESpawnActorCollisionHandlingMethod::AlwaysSpawn;
 		SpawnInfo.Owner = PC->GetWorldSettings();
 		SpawnInfo.Instigator = PC->Instigator;
 		Instance = InWorld->SpawnActor<AVisualLoggerCameraController>( SpawnInfo );
diff --git a/Engine/Source/Editor/BlueprintGraph/Classes/K2Node_SpawnActorFromClass.h b/Engine/Source/Editor/BlueprintGraph/Classes/K2Node_SpawnActorFromClass.h
--- a/Engine/Source/Editor/BlueprintGraph/Classes/K2Node_SpawnActorFromClass.h
+++ b/Engine/Source/Editor/BlueprintGraph/Classes/K2Node_SpawnActorFromClass.h
@@ -53,8 +53,8 @@ class BLUEPRINTGRAPH_API UK2Node_SpawnActorFromClass : public UK2Node
 	UEdGraphPin* GetWorldContextPin() const;
 	/** Get the spawn transform input pin */	
 	UEdGraphPin* GetSpawnTransformPin() const;
-	/** Get the spawn NoCollisionFail input pin */
-	UEdGraphPin* GetNoCollisionFailPin() const;
+	/** Get the collision handling method input pin */
+	UEdGraphPin* GetCollisionHandlingOverridePin() const;
 	/** Get the actor owner pin */
 	UEdGraphPin* GetOwnerPin() const;
 	/** Get the result output pin */
@@ -72,4 +72,7 @@ protected:
 
 	/** Constructing FText strings can be costly, so we cache the node's title */
 	FNodeTextCache CachedNodeTitle;
+
+private:
+	void MaybeUpdateCollisionPin(TArray<UEdGraphPin*>& OldPins);
 };
diff --git a/Engine/Source/Editor/BlueprintGraph/Private/K2Node_SpawnActorFromClass.cpp b/Engine/Source/Editor/BlueprintGraph/Private/K2Node_SpawnActorFromClass.cpp
--- a/Engine/Source/Editor/BlueprintGraph/Private/K2Node_SpawnActorFromClass.cpp
+++ b/Engine/Source/Editor/BlueprintGraph/Private/K2Node_SpawnActorFromClass.cpp
@@ -6,6 +6,7 @@
 #include "BlueprintNodeSpawner.h"
 #include "EditorCategoryUtils.h"
 #include "BlueprintActionDatabaseRegistrar.h"
+#include "K2Node_EnumLiteral.h"
 
 struct FK2Node_SpawnActorFromClassHelper
 {
@@ -13,13 +14,15 @@ struct FK2Node_SpawnActorFromClassHelper
 	static FString ClassPinName;
 	static FString SpawnTransformPinName;
 	static FString NoCollisionFailPinName;
+	static FString CollisionHandlingOverridePinName;
 	static FString OwnerPinName;
 };
 
 FString FK2Node_SpawnActorFromClassHelper::WorldContextPinName(TEXT("WorldContextObject"));
 FString FK2Node_SpawnActorFromClassHelper::ClassPinName(TEXT("Class"));
 FString FK2Node_SpawnActorFromClassHelper::SpawnTransformPinName(TEXT("SpawnTransform"));
-FString FK2Node_SpawnActorFromClassHelper::NoCollisionFailPinName(TEXT("SpawnEvenIfColliding"));
+FString FK2Node_SpawnActorFromClassHelper::NoCollisionFailPinName(TEXT("SpawnEvenIfColliding"));		// deprecated pin, name kept for backwards compat
+FString FK2Node_SpawnActorFromClassHelper::CollisionHandlingOverridePinName(TEXT("CollisionHandlingOverride"));
 FString FK2Node_SpawnActorFromClassHelper::OwnerPinName(TEXT("Owner"));
 
 #define LOCTEXT_NAMESPACE "K2Node_SpawnActorFromClass"
@@ -53,9 +56,13 @@ void UK2Node_SpawnActorFromClass::AllocateDefaultPins()
 	UEdGraphPin* TransformPin = CreatePin(EGPD_Input, K2Schema->PC_Struct, TEXT(""), TransformStruct, false, false, FK2Node_SpawnActorFromClassHelper::SpawnTransformPinName);
 	K2Schema->ConstructBasicPinTooltip(*TransformPin, LOCTEXT("TransformPinDescription", "The transform to spawn the Actor with"), TransformPin->PinToolTip);
 
-	// bNoCollisionFail pin
-	UEdGraphPin* NoCollisionFailPin = CreatePin(EGPD_Input, K2Schema->PC_Boolean, TEXT(""), NULL, false, false, FK2Node_SpawnActorFromClassHelper::NoCollisionFailPinName);
-	K2Schema->ConstructBasicPinTooltip(*NoCollisionFailPin, LOCTEXT("NoCollisionFailPinDescription", "Determines if the Actor should be spawned when the location is blocked by a collision"), NoCollisionFailPin->PinToolTip);
+	// Collision handling method pin
+	UEnum* const MethodEnum = FindObjectChecked<UEnum>(ANY_PACKAGE, TEXT("ESpawnActorCollisionHandlingMethod"), true);
+	UEdGraphPin* const CollisionHandlingOverridePin = CreatePin(EGPD_Input, K2Schema->PC_Byte, TEXT(""), MethodEnum, false, false, FK2Node_SpawnActorFromClassHelper::CollisionHandlingOverridePinName);
+	CollisionHandlingOverridePin->DefaultValue = MethodEnum->GetEnumName(static_cast<int>(ESpawnActorCollisionHandlingMethod::Undefined));
+
+	K2Schema->ConstructBasicPinTooltip(*CollisionHandlingOverridePin, LOCTEXT("CollisionHandlingOverridePinDescription", "Specifies how to handle collisions at the spawn point. If undefined, uses actor class settings."), CollisionHandlingOverridePin->PinToolTip);
+
 
 	UEdGraphPin* OwnerPin = CreatePin(EGPD_Input, K2Schema->PC_Object, TEXT(""), AActor::StaticClass(),/*bIsArray =*/false, /*bIsReference =*/false, FK2Node_SpawnActorFromClassHelper::OwnerPinName);
 	OwnerPin->bAdvancedView = true;
@@ -139,6 +146,126 @@ UClass* UK2Node_SpawnActorFromClass::GetClassToSpawn(const TArray<UEdGraphPin*>*
 	return UseSpawnClass;
 }
 
+void UK2Node_SpawnActorFromClass::MaybeUpdateCollisionPin(TArray<UEdGraphPin*>& OldPins)
+{
+	// see if there's a bNoCollisionFail pin
+	for (UEdGraphPin* Pin : OldPins)
+	{
+		if (Pin->PinName == FK2Node_SpawnActorFromClassHelper::NoCollisionFailPinName)
+		{
+			bool bHadOldCollisionPin = true;
+			if (Pin->LinkedTo.Num() == 0)
+			{
+				// no links, use the default value of the pin
+				bool const bOldCollisionPinValue = (Pin->DefaultValue == FString(TEXT("true")));
+
+				UEdGraphPin* const CollisionHandlingOverridePin = GetCollisionHandlingOverridePin();
+				if (CollisionHandlingOverridePin)
+				{
+					UEnum const* const MethodEnum = FindObjectChecked<UEnum>(ANY_PACKAGE, TEXT("ESpawnActorCollisionHandlingMethod"), true);
+					CollisionHandlingOverridePin->DefaultValue =
+						bOldCollisionPinValue
+						? MethodEnum->GetEnumName(static_cast<int>(ESpawnActorCollisionHandlingMethod::AlwaysSpawn))
+						: MethodEnum->GetEnumName(static_cast<int>(ESpawnActorCollisionHandlingMethod::AdjustIfPossibleButDontSpawnIfColliding));
+				}
+			}
+			else
+			{
+				// something was linked. we will just move the links to the new pin
+				// #note: this will be an invalid linkage and the BP compiler will complain, and that's intentional
+				// so that users will be able to see and fix issues
+				UEdGraphPin* const CollisionHandlingOverridePin = GetCollisionHandlingOverridePin();
+				check(CollisionHandlingOverridePin);
+
+				UEnum* const MethodEnum = FindObjectChecked<UEnum>(ANY_PACKAGE, TEXT("ESpawnActorCollisionHandlingMethod"), true);
+				
+				FGraphNodeCreator<UK2Node_EnumLiteral> AlwaysSpawnLiteralCreator(*GetGraph());
+				UK2Node_EnumLiteral* const AlwaysSpawnLiteralNode = AlwaysSpawnLiteralCreator.CreateNode();
+				AlwaysSpawnLiteralNode->Enum = MethodEnum;
+				AlwaysSpawnLiteralNode->NodePosX = NodePosX;
+				AlwaysSpawnLiteralNode->NodePosY = NodePosY;
+				AlwaysSpawnLiteralCreator.Finalize();
+			
+				FGraphNodeCreator<UK2Node_EnumLiteral> AdjustIfNecessaryLiteralCreator(*GetGraph());
+				UK2Node_EnumLiteral* const AdjustIfNecessaryLiteralNode = AdjustIfNecessaryLiteralCreator.CreateNode();
+				AdjustIfNecessaryLiteralNode->Enum = MethodEnum;
+				AdjustIfNecessaryLiteralNode->NodePosX = NodePosX;
+				AdjustIfNecessaryLiteralNode->NodePosY = NodePosY;
+				AdjustIfNecessaryLiteralCreator.Finalize();
+
+				FGraphNodeCreator<UK2Node_Select> SelectCreator(*GetGraph());
+				UK2Node_Select* const SelectNode = SelectCreator.CreateNode();
+				SelectNode->NodePosX = NodePosX;
+				SelectNode->NodePosY = NodePosY;
+				SelectCreator.Finalize();
+
+				// find pins we want to set and link up
+				auto FindEnumInputPin = [](UK2Node_EnumLiteral const* Node)
+				{
+					for (UEdGraphPin* NodePin : Node->Pins)
+					{
+						if (NodePin->PinName == Node->GetEnumInputPinName())
+						{
+							return NodePin;
+						}
+					}
+					return (UEdGraphPin*)nullptr;
+				};
+
+				UEdGraphPin* const AlwaysSpawnLiteralNodeInputPin = FindEnumInputPin(AlwaysSpawnLiteralNode);
+				UEdGraphPin* const AdjustIfNecessaryLiteralInputPin = FindEnumInputPin(AdjustIfNecessaryLiteralNode);
+
+				TArray<UEdGraphPin*> SelectOptionPins;
+				SelectNode->GetOptionPins(SelectOptionPins);
+				UEdGraphPin* const SelectIndexPin = SelectNode->GetIndexPin();
+
+				auto FindResultPin = [](UK2Node const* Node)
+				{
+					for (UEdGraphPin* NodePin : Node->Pins)
+					{
+						if (EEdGraphPinDirection::EGPD_Output == NodePin->Direction)
+						{
+							return NodePin;
+						}
+					}
+					return (UEdGraphPin*)nullptr;
+				};
+				UEdGraphPin* const AlwaysSpawnLiteralNodeResultPin = FindResultPin(AlwaysSpawnLiteralNode);
+				check(AlwaysSpawnLiteralNodeResultPin);
+				UEdGraphPin* const AdjustIfNecessaryLiteralResultPin = FindResultPin(AdjustIfNecessaryLiteralNode);
+				check(AdjustIfNecessaryLiteralResultPin);
+
+				UEdGraphPin* const OldBoolPin = Pin->LinkedTo[0];
+				check(OldBoolPin);
+
+				//
+				// now set data and links that we want to set
+				//
+
+				AlwaysSpawnLiteralNodeInputPin->DefaultValue = MethodEnum->GetEnumName(static_cast<int>(ESpawnActorCollisionHandlingMethod::AlwaysSpawn));
+				AdjustIfNecessaryLiteralInputPin->DefaultValue = MethodEnum->GetEnumName(static_cast<int>(ESpawnActorCollisionHandlingMethod::AdjustIfPossibleButDontSpawnIfColliding));
+
+				OldBoolPin->BreakLinkTo(Pin);
+				OldBoolPin->MakeLinkTo(SelectIndexPin);
+
+				AlwaysSpawnLiteralNodeResultPin->MakeLinkTo(SelectOptionPins[0]);
+				AdjustIfNecessaryLiteralResultPin->MakeLinkTo(SelectOptionPins[1]);
+				
+				UEdGraphPin* const SelectOutputPin = SelectNode->GetReturnValuePin();
+				check(SelectOutputPin);
+				SelectOutputPin->MakeLinkTo(CollisionHandlingOverridePin);
+
+				// tell select node to update its wildcard status
+				SelectNode->NotifyPinConnectionListChanged(SelectIndexPin);
+				SelectNode->NotifyPinConnectionListChanged(SelectOptionPins[0]);
+				SelectNode->NotifyPinConnectionListChanged(SelectOptionPins[1]);
+				SelectNode->NotifyPinConnectionListChanged(SelectOutputPin);
+
+			}
+		}
+	}
+}
+
 
 void UK2Node_SpawnActorFromClass::ReallocatePinsDuringReconstruction(TArray<UEdGraphPin*>& OldPins) 
 {
@@ -149,6 +276,8 @@ void UK2Node_SpawnActorFromClass::ReallocatePinsDuringReconstruction(TArray<UEdG
 		TArray<UEdGraphPin*> ClassPins;
 		CreatePinsForClass(UseSpawnClass, ClassPins);
 	}
+
+	MaybeUpdateCollisionPin(OldPins);
 }
 
 void UK2Node_SpawnActorFromClass::PostPlacedNewNode()
@@ -182,7 +311,7 @@ bool UK2Node_SpawnActorFromClass::IsSpawnVarPin(UEdGraphPin* Pin)
 			Pin->PinName != K2Schema->PN_ReturnValue &&
 			Pin->PinName != FK2Node_SpawnActorFromClassHelper::ClassPinName &&
 			Pin->PinName != FK2Node_SpawnActorFromClassHelper::WorldContextPinName &&
-			Pin->PinName != FK2Node_SpawnActorFromClassHelper::NoCollisionFailPinName &&
+			Pin->PinName != FK2Node_SpawnActorFromClassHelper::CollisionHandlingOverridePinName &&
 			Pin->PinName != FK2Node_SpawnActorFromClassHelper::SpawnTransformPinName && 
 			Pin->PinName != FK2Node_SpawnActorFromClassHelper::OwnerPinName );
 }
@@ -297,10 +426,9 @@ UEdGraphPin* UK2Node_SpawnActorFromClass::GetSpawnTransformPin() const
 	return Pin;
 }
 
-
-UEdGraphPin* UK2Node_SpawnActorFromClass::GetNoCollisionFailPin() const
+UEdGraphPin* UK2Node_SpawnActorFromClass::GetCollisionHandlingOverridePin() const
 {
-	UEdGraphPin* Pin = FindPinChecked(FK2Node_SpawnActorFromClassHelper::NoCollisionFailPinName);
+	UEdGraphPin* const Pin = FindPinChecked(FK2Node_SpawnActorFromClassHelper::CollisionHandlingOverridePinName);
 	check(Pin->Direction == EGPD_Input);
 	return Pin;
 }
@@ -428,14 +556,14 @@ void UK2Node_SpawnActorFromClass::ExpandNode(class FKismetCompilerContext& Compi
 {
 	Super::ExpandNode(CompilerContext, SourceGraph);
 
-	static FName BeginSpawningBlueprintFuncName = GET_FUNCTION_NAME_CHECKED(UGameplayStatics, BeginSpawningActorFromClass);
+	static FName BeginSpawningBlueprintFuncName = GET_FUNCTION_NAME_CHECKED(UGameplayStatics, BeginDeferredActorSpawnFromClass);
 	static FString ActorClassParamName = FString(TEXT("ActorClass"));
 	static FString WorldContextParamName = FString(TEXT("WorldContextObject"));
 
 	static FName FinishSpawningFuncName = GET_FUNCTION_NAME_CHECKED(UGameplayStatics, FinishSpawningActor);
 	static FString ActorParamName = FString(TEXT("Actor"));
 	static FString TransformParamName = FString(TEXT("SpawnTransform"));
-	static FString NoCollisionFailParamName = FString(TEXT("bNoCollisionFail"));
+	static FString CollisionHandlingOverrideParamName = FString(TEXT("CollisionHandlingOverride"));
 	static FString OwnerParamName = FString(TEXT("Owner"));
 
 	static FString ObjectParamName = FString(TEXT("Object"));
@@ -445,7 +573,7 @@ void UK2Node_SpawnActorFromClass::ExpandNode(class FKismetCompilerContext& Compi
 	UK2Node_SpawnActorFromClass* SpawnNode = this;
 	UEdGraphPin* SpawnNodeExec = SpawnNode->GetExecPin();
 	UEdGraphPin* SpawnNodeTransform = SpawnNode->GetSpawnTransformPin();
-	UEdGraphPin* SpawnNodeNoCollisionFail = GetNoCollisionFailPin();
+	UEdGraphPin* SpawnNodeCollisionHandlingOverride = GetCollisionHandlingOverridePin();
 	UEdGraphPin* SpawnWorldContextPin = SpawnNode->GetWorldContextPin();
 	UEdGraphPin* SpawnClassPin = SpawnNode->GetClassPin();
 	UEdGraphPin* SpawnNodeOwnerPin = SpawnNode->GetOwnerPin();
@@ -471,7 +599,8 @@ void UK2Node_SpawnActorFromClass::ExpandNode(class FKismetCompilerContext& Compi
 	UEdGraphPin* CallBeginWorldContextPin = CallBeginSpawnNode->FindPinChecked(WorldContextParamName);
 	UEdGraphPin* CallBeginActorClassPin = CallBeginSpawnNode->FindPinChecked(ActorClassParamName);
 	UEdGraphPin* CallBeginTransform = CallBeginSpawnNode->FindPinChecked(TransformParamName);
-	UEdGraphPin* CallBeginNoCollisionFail = CallBeginSpawnNode->FindPinChecked(NoCollisionFailParamName);
+	UEdGraphPin* CallBeginCollisionHandlingOverride = CallBeginSpawnNode->FindPinChecked(CollisionHandlingOverrideParamName);
+
 	UEdGraphPin* CallBeginOwnerPin = CallBeginSpawnNode->FindPinChecked(FK2Node_SpawnActorFromClassHelper::OwnerPinName);
 	UEdGraphPin* CallBeginResult = CallBeginSpawnNode->GetReturnValuePin();	
 
@@ -504,7 +633,7 @@ void UK2Node_SpawnActorFromClass::ExpandNode(class FKismetCompilerContext& Compi
 	CompilerContext.MovePinLinksToIntermediate(*SpawnNodeTransform, *CallBeginTransform);
 
 	// Copy the 'bNoCollisionFail' connection from the spawn node to 'begin spawn'
-	CompilerContext.MovePinLinksToIntermediate(*SpawnNodeNoCollisionFail, *CallBeginNoCollisionFail);
+	CompilerContext.MovePinLinksToIntermediate(*SpawnNodeCollisionHandlingOverride, *CallBeginCollisionHandlingOverride);
 
 	//////////////////////////////////////////////////////////////////////////
 	// create 'finish spawn' call node
diff --git a/Engine/Source/Editor/Kismet/Private/BlueprintEditor.cpp b/Engine/Source/Editor/Kismet/Private/BlueprintEditor.cpp
--- a/Engine/Source/Editor/Kismet/Private/BlueprintEditor.cpp
+++ b/Engine/Source/Editor/Kismet/Private/BlueprintEditor.cpp
@@ -7373,7 +7373,7 @@ void FBlueprintEditor::UpdatePreviewActor(UBlueprint* InBlueprint, bool bInForce
 
 			// Spawn an Actor based on the Blueprint's generated class
 			FActorSpawnParameters SpawnInfo;
-			SpawnInfo.bNoCollisionFail = true;
+			SpawnInfo.SpawnCollisionHandlingOverride = ESpawnActorCollisionHandlingMethod::AlwaysSpawn;
 			SpawnInfo.bNoFail = true;
 			SpawnInfo.ObjectFlags = RF_Transient|RF_Transactional;
 
diff --git a/Engine/Source/Editor/UnrealEd/Private/Commandlets/PackageUtilities.cpp b/Engine/Source/Editor/UnrealEd/Private/Commandlets/PackageUtilities.cpp
--- a/Engine/Source/Editor/UnrealEd/Private/Commandlets/PackageUtilities.cpp
+++ b/Engine/Source/Editor/UnrealEd/Private/Commandlets/PackageUtilities.cpp
@@ -2442,7 +2442,7 @@ int32 UReplaceActorCommandlet::Main(const FString& Params)
 
 						FActorSpawnParameters SpawnInfo;
 						SpawnInfo.OverrideLevel = Level;
-						SpawnInfo.bNoCollisionFail = true;
+						SpawnInfo.SpawnCollisionHandlingOverride = ESpawnActorCollisionHandlingMethod::AlwaysSpawn;
 						// spawn the new actor
 						AActor* NewActor = World->SpawnActor<AActor>( ReplaceWithClass, OldLocation, OldRotator, SpawnInfo );
 
diff --git a/Engine/Source/Editor/UnrealEd/Private/EditorActor.cpp b/Engine/Source/Editor/UnrealEd/Private/EditorActor.cpp
--- a/Engine/Source/Editor/UnrealEd/Private/EditorActor.cpp
+++ b/Engine/Source/Editor/UnrealEd/Private/EditorActor.cpp
@@ -945,7 +945,7 @@ AActor* UUnrealEdEngine::ReplaceActor( AActor* CurrentActor, UClass* NewActorCla
 	FRotator SpawnRot = CurrentActor->GetActorRotation();
 	FActorSpawnParameters SpawnInfo;
 	SpawnInfo.Template = Cast<AActor>(Archetype);
-	SpawnInfo.bNoCollisionFail = true;
+	SpawnInfo.SpawnCollisionHandlingOverride = ESpawnActorCollisionHandlingMethod::AlwaysSpawn;
 	AActor* NewActor = CurrentActor->GetWorld()->SpawnActor( NewActorClass, &SpawnLoc, &SpawnRot, SpawnInfo );
 	if( NewActor )
 	{
diff --git a/Engine/Source/Editor/UnrealEd/Private/EditorEngine.cpp b/Engine/Source/Editor/UnrealEd/Private/EditorEngine.cpp
--- a/Engine/Source/Editor/UnrealEd/Private/EditorEngine.cpp
+++ b/Engine/Source/Editor/UnrealEd/Private/EditorEngine.cpp
@@ -3294,7 +3294,7 @@ void UEditorEngine::ConvertActorsFromClass( UClass* FromClass, UClass* ToClass )
 				
 				FActorSpawnParameters SpawnInfo;
 				SpawnInfo.OverrideLevel = Info.SourceLevel;
-				SpawnInfo.bNoCollisionFail = true;
+				SpawnInfo.SpawnCollisionHandlingOverride = ESpawnActorCollisionHandlingMethod::AlwaysSpawn;
 				if( bToStaticMesh )
 				{					
 					AStaticMeshActor* SMActor = CastChecked<AStaticMeshActor>( World->SpawnActor( ToClass, &Info.Location, &Info.Rotation, SpawnInfo ) );
@@ -5240,7 +5240,7 @@ void UEditorEngine::DoConvertActors( const TArray<AActor*>& ActorsToConvert, UCl
 				{
 					FActorSpawnParameters SpawnInfo;
 					SpawnInfo.OverrideLevel = ActorLevel;
-					SpawnInfo.bNoCollisionFail = true;
+					SpawnInfo.SpawnCollisionHandlingOverride = ESpawnActorCollisionHandlingMethod::AlwaysSpawn;
 					NewActor = World->SpawnActor( ConvertToClass, &SpawnLoc, &SpawnRot, SpawnInfo );
 
 					if (NewActor)
@@ -5524,7 +5524,7 @@ AActor* UEditorEngine::AddActor(ULevel* InLevel, UClass* Class, const FTransform
 
 		FActorSpawnParameters SpawnInfo;
 		SpawnInfo.OverrideLevel = DesiredLevel;
-		SpawnInfo.bNoCollisionFail = true;
+		SpawnInfo.SpawnCollisionHandlingOverride = ESpawnActorCollisionHandlingMethod::AlwaysSpawn;
 		SpawnInfo.ObjectFlags = ObjectFlags;
 		const auto Location = Transform.GetLocation();
 		const auto Rotation = Transform.GetRotation().Rotator();
diff --git a/Engine/Source/Editor/UnrealEd/Private/EditorLevelUtils.cpp b/Engine/Source/Editor/UnrealEd/Private/EditorLevelUtils.cpp
--- a/Engine/Source/Editor/UnrealEd/Private/EditorLevelUtils.cpp
+++ b/Engine/Source/Editor/UnrealEd/Private/EditorLevelUtils.cpp
@@ -230,7 +230,7 @@ namespace EditorLevelUtils
 				if (SubLevelWorld != NULL)
 				{
 					FActorSpawnParameters SpawnInfo;
-					SpawnInfo.bNoCollisionFail = true;
+					SpawnInfo.SpawnCollisionHandlingOverride = ESpawnActorCollisionHandlingMethod::AlwaysSpawn;
 					SpawnInfo.Name = GEngine->WorldSettingsClass->GetFName();
 					AWorldSettings* NewWorldSettings = SubLevelWorld->SpawnActor<AWorldSettings>( GEngine->WorldSettingsClass, SpawnInfo );
 					NewLevel->Actors[0] = NewWorldSettings;
diff --git a/Engine/Source/Editor/UnrealEd/Private/Factories/EditorFactories.cpp b/Engine/Source/Editor/UnrealEd/Private/Factories/EditorFactories.cpp
--- a/Engine/Source/Editor/UnrealEd/Private/Factories/EditorFactories.cpp
+++ b/Engine/Source/Editor/UnrealEd/Private/Factories/EditorFactories.cpp
@@ -551,7 +551,7 @@ UObject* ULevelFactory::FactoryCreateText
 						FActorSpawnParameters SpawnInfo;
 						SpawnInfo.Name = ActorUniqueName;
 						SpawnInfo.Template = Archetype;
-						SpawnInfo.bNoCollisionFail = true;
+						SpawnInfo.SpawnCollisionHandlingOverride = ESpawnActorCollisionHandlingMethod::AlwaysSpawn;
 						AActor* NewActor = World->SpawnActor( TempClass, nullptr, nullptr, SpawnInfo );
 						
 						if( NewActor )
diff --git a/Engine/Source/Editor/UnrealEd/Private/Kismet2/KismetReinstanceUtilities.cpp b/Engine/Source/Editor/UnrealEd/Private/Kismet2/KismetReinstanceUtilities.cpp
--- a/Engine/Source/Editor/UnrealEd/Private/Kismet2/KismetReinstanceUtilities.cpp
+++ b/Engine/Source/Editor/UnrealEd/Private/Kismet2/KismetReinstanceUtilities.cpp
@@ -999,7 +999,7 @@ void FBlueprintCompileReinstancer::ReplaceInstancesOfClass(UClass* OldClass, UCl
 				FActorSpawnParameters SpawnInfo;
 				SpawnInfo.OverrideLevel      = ActorLevel;
 				SpawnInfo.Template           = NewArchetype;
-				SpawnInfo.bNoCollisionFail   = true;
+				SpawnInfo.SpawnCollisionHandlingOverride = ESpawnActorCollisionHandlingMethod::AlwaysSpawn;
 				SpawnInfo.bDeferConstruction = true;
 
 				// Temporarily remove the deprecated flag so we can respawn the Blueprint in the level
diff --git a/Engine/Source/Editor/UnrealEd/Private/PlayLevel.cpp b/Engine/Source/Editor/UnrealEd/Private/PlayLevel.cpp
--- a/Engine/Source/Editor/UnrealEd/Private/PlayLevel.cpp
+++ b/Engine/Source/Editor/UnrealEd/Private/PlayLevel.cpp
@@ -1085,7 +1085,7 @@ void UEditorEngine::SaveWorldForPlay(TArray<FString>& SavedMapNames)
 		if( PlayerStart == NULL )
 		{
 			FActorSpawnParameters SpawnInfo;
-			SpawnInfo.bNoCollisionFail = true;
+			SpawnInfo.SpawnCollisionHandlingOverride = ESpawnActorCollisionHandlingMethod::AlwaysSpawn;
 			PlayerStart = World->SpawnActor<AActor>( APlayerStart::StaticClass(), FVector::ZeroVector, FRotator::ZeroRotator, SpawnInfo );
 
 			bCreatedPlayerStart = true;
diff --git a/Engine/Source/Editor/UnrealEd/Private/ThumbnailHelpers.cpp b/Engine/Source/Editor/UnrealEd/Private/ThumbnailHelpers.cpp
--- a/Engine/Source/Editor/UnrealEd/Private/ThumbnailHelpers.cpp
+++ b/Engine/Source/Editor/UnrealEd/Private/ThumbnailHelpers.cpp
@@ -239,7 +239,7 @@ FMaterialThumbnailScene::FMaterialThumbnailScene()
 	// Create preview actor
 	// checked
 	FActorSpawnParameters SpawnInfo;
-	SpawnInfo.bNoCollisionFail = true;
+	SpawnInfo.SpawnCollisionHandlingOverride = ESpawnActorCollisionHandlingMethod::AlwaysSpawn;
 	SpawnInfo.bNoFail = true;
 	SpawnInfo.ObjectFlags = RF_Transient;
 	PreviewActor = GetWorld()->SpawnActor<AStaticMeshActor>( SpawnInfo );
@@ -382,7 +382,7 @@ FSkeletalMeshThumbnailScene::FSkeletalMeshThumbnailScene()
 	// Create preview actor
 	// checked
 	FActorSpawnParameters SpawnInfo;
-	SpawnInfo.bNoCollisionFail = true;
+	SpawnInfo.SpawnCollisionHandlingOverride = ESpawnActorCollisionHandlingMethod::AlwaysSpawn;
 	SpawnInfo.bNoFail = true;
 	SpawnInfo.ObjectFlags = RF_Transient;
 	PreviewActor = GetWorld()->SpawnActor<ASkeletalMeshActor>( SpawnInfo );
@@ -452,7 +452,7 @@ FDestructibleMeshThumbnailScene::FDestructibleMeshThumbnailScene()
 	// Create preview actor
 	// checked
 	FActorSpawnParameters SpawnInfo;
-	SpawnInfo.bNoCollisionFail = true;
+	SpawnInfo.SpawnCollisionHandlingOverride = ESpawnActorCollisionHandlingMethod::AlwaysSpawn;
 	SpawnInfo.bNoFail = true;
 	SpawnInfo.ObjectFlags = RF_Transient;
 	PreviewActor = GetWorld()->SpawnActor<ADestructibleActor>(SpawnInfo);
@@ -523,7 +523,7 @@ FStaticMeshThumbnailScene::FStaticMeshThumbnailScene()
 	// Create preview actor
 	// checked
 	FActorSpawnParameters SpawnInfo;
-	SpawnInfo.bNoCollisionFail = true;
+	SpawnInfo.SpawnCollisionHandlingOverride = ESpawnActorCollisionHandlingMethod::AlwaysSpawn;
 	SpawnInfo.bNoFail = true;
 	SpawnInfo.ObjectFlags = RF_Transient;
 	PreviewActor = GetWorld()->SpawnActor<AStaticMeshActor>( SpawnInfo );
@@ -600,7 +600,7 @@ FAnimationSequenceThumbnailScene::FAnimationSequenceThumbnailScene()
 	// Create preview actor
 	// checked
 	FActorSpawnParameters SpawnInfo;
-	SpawnInfo.bNoCollisionFail = true;
+	SpawnInfo.SpawnCollisionHandlingOverride = ESpawnActorCollisionHandlingMethod::AlwaysSpawn;
 	SpawnInfo.bNoFail = true;
 	SpawnInfo.ObjectFlags = RF_Transient;
 	PreviewActor = GetWorld()->SpawnActor<ASkeletalMeshActor>(SpawnInfo);
@@ -724,7 +724,7 @@ FBlendSpaceThumbnailScene::FBlendSpaceThumbnailScene()
 	// Create preview actor
 	// checked
 	FActorSpawnParameters SpawnInfo;
-	SpawnInfo.bNoCollisionFail = true;
+	SpawnInfo.SpawnCollisionHandlingOverride = ESpawnActorCollisionHandlingMethod::AlwaysSpawn;
 	SpawnInfo.bNoFail = true;
 	SpawnInfo.ObjectFlags = RF_Transient;
 	PreviewActor = GetWorld()->SpawnActor<ASkeletalMeshActor>(SpawnInfo);
@@ -850,7 +850,7 @@ FAnimBlueprintThumbnailScene::FAnimBlueprintThumbnailScene()
 	
 	// Create preview actor
 	FActorSpawnParameters SpawnInfo;
-	SpawnInfo.bNoCollisionFail = true;
+	SpawnInfo.SpawnCollisionHandlingOverride = ESpawnActorCollisionHandlingMethod::AlwaysSpawn;
 	SpawnInfo.bNoFail = true;
 	SpawnInfo.ObjectFlags = RF_Transient;
 	PreviewActor = GetWorld()->SpawnActor<ASkeletalMeshActor>(SpawnInfo);
diff --git a/Engine/Source/Runtime/AIModule/Private/Blueprint/AIBlueprintHelperLibrary.cpp b/Engine/Source/Runtime/AIModule/Private/Blueprint/AIBlueprintHelperLibrary.cpp
--- a/Engine/Source/Runtime/AIModule/Private/Blueprint/AIBlueprintHelperLibrary.cpp
+++ b/Engine/Source/Runtime/AIModule/Private/Blueprint/AIBlueprintHelperLibrary.cpp
@@ -140,7 +140,8 @@ APawn* UAIBlueprintHelperLibrary::SpawnAIFromClass(UObject* WorldContextObject,
 	if (World && *PawnClass)
 	{
 		FActorSpawnParameters ActorSpawnParams;
-		ActorSpawnParams.bNoCollisionFail = bNoCollisionFail;
+		ActorSpawnParams.SpawnCollisionHandlingOverride = bNoCollisionFail ? ESpawnActorCollisionHandlingMethod::AlwaysSpawn : ESpawnActorCollisionHandlingMethod::AdjustIfPossibleButDontSpawnIfColliding;
+
 		NewPawn = World->SpawnActor<APawn>(*PawnClass, Location, Rotation, ActorSpawnParams);
 
 		if (NewPawn != NULL)
diff --git a/Engine/Source/Runtime/Engine/Classes/Components/SceneComponent.h b/Engine/Source/Runtime/Engine/Classes/Components/SceneComponent.h
--- a/Engine/Source/Runtime/Engine/Classes/Components/SceneComponent.h
+++ b/Engine/Source/Runtime/Engine/Classes/Components/SceneComponent.h
@@ -839,7 +839,7 @@ public:
 	/** Update the Bounds of the component.*/
 	virtual void UpdateBounds();
 
-	/** If true, bounds should be used when placing component/actor in level, and spawning may fail */
+	/** If true, bounds should be used when placing component/actor in level. Does not affect spawning. */
 	virtual bool ShouldCollideWhenPlacing() const
 	{
 		return false;
diff --git a/Engine/Source/Runtime/Engine/Classes/Engine/EngineTypes.h b/Engine/Source/Runtime/Engine/Classes/Engine/EngineTypes.h
--- a/Engine/Source/Runtime/Engine/Classes/Engine/EngineTypes.h
+++ b/Engine/Source/Runtime/Engine/Classes/Engine/EngineTypes.h
@@ -3641,3 +3641,20 @@ struct FCanvasUVTri
 
 
 template <> struct TIsZeroConstructType<FCanvasUVTri> { enum { Value = true }; };
+
+/** Defines available strategies for handling the case where an actor is spawned in such a way that it penetrates blocking collision. */
+UENUM(BlueprintType)
+enum class ESpawnActorCollisionHandlingMethod : uint8
+{
+	/** Fall back to default settings. */
+	Undefined								UMETA(DisplayName = "Default"),
+	/** Actor will spawn in desired location, regardless of collisions. */
+	AlwaysSpawn								UMETA(DisplayName = "Always Spawn, Ignore Collisions"),
+	/** Actor will try to find a nearby non-colliding location (based on shape components), but will always spawn even if one cannot be found. */
+	AdjustIfPossibleButAlwaysSpawn			UMETA(DisplayName = "Try To Adjust Location, But Always Spawn"),
+	/** Actor will try to find a nearby non-colliding location (based on shape components), but will NOT spawn unless one is found. */
+	AdjustIfPossibleButDontSpawnIfColliding	UMETA(DisplayName = "Try To Adjust Location, Don't Spawn If Still Colliding"),
+	/** Actor will fail to spawn. */
+	DontSpawnIfColliding					UMETA(DisplayName = "Do Not Spawn"),
+};
+
diff --git a/Engine/Source/Runtime/Engine/Classes/Engine/World.h b/Engine/Source/Runtime/Engine/Classes/Engine/World.h
--- a/Engine/Source/Runtime/Engine/Classes/Engine/World.h
+++ b/Engine/Source/Runtime/Engine/Classes/Engine/World.h
@@ -425,6 +425,7 @@ struct ENGINE_API FActorSpawnParameters
 		,	Owner(NULL)
 		,	Instigator(NULL)
 		,	OverrideLevel(NULL)
+		,	SpawnCollisionHandlingOverride(ESpawnActorCollisionHandlingMethod::Undefined)
 		,	bNoCollisionFail(false)
 		,	bRemoteOwned(false)
 		,	bNoFail(false)
@@ -449,20 +450,24 @@ struct ENGINE_API FActorSpawnParameters
 	/* The ULevel to spawn the Actor in, i.e. the Outer of the Actor. If left as NULL the Outer of the Owner is used. If the Owner is NULL the persistent level is used. */
 	class	ULevel* OverrideLevel;
 
+	/** Method for resolving collisions at the spawn point. Undefined means no override, use the actor's setting. */
+	ESpawnActorCollisionHandlingMethod SpawnCollisionHandlingOverride;
+
 	/* Determines whether a collision test will be performed when spawning the Actor. If true, no collision test will be performed when spawning the Actor regardless of the collision settings of the root component or template Actor. */
-	uint32	bNoCollisionFail:1;
+	DEPRECATED(4.9, "bNoCollisionFail is deprecated. Use SpawnCollisionHandlingOverride to override the actor class's spawn collision handling setting.")
+	uint16	bNoCollisionFail:1;
 
 	/* Is the actor remotely owned. */
-	uint32	bRemoteOwned:1;
+	uint16	bRemoteOwned:1;
 
 	/* Determines whether spawning will not fail if certain conditions are not met. If true, spawning will not fail because the class being spawned is `bStatic=true` or because the class of the template Actor is not the same as the class of the Actor being spawned. */
-	uint32	bNoFail:1;
+	uint16	bNoFail:1;
 
 	/* Determines whether the construction script will be run. If true, the construction script will not be run on the spawned Actor. Only applicable if the Actor is being spawned from a Blueprint. */
-	uint32	bDeferConstruction:1;
+	uint16	bDeferConstruction:1;
 	
 	/* Determines whether or not the actor may be spawned when running a construction script. If true spawning will fail if a construction script is being run. */
-	uint32	bAllowDuringConstructionScript:1;
+	uint16	bAllowDuringConstructionScript:1;
 	
 	/* Flags used to describe the spawned actor/object instance. */
 	EObjectFlags ObjectFlags;		
@@ -2852,9 +2857,9 @@ public:
 		UClass* Class,
 		FVector const& Location,
 		FRotator const& Rotation,
-		AActor* Owner=NULL,
-		APawn* Instigator=NULL,
-		bool bNoCollisionFail=false
+		AActor* Owner,
+		APawn* Instigator,
+		bool bNoCollisionFail
 		)
 	{
 		if( Owner )
@@ -2862,7 +2867,7 @@ public:
 			check(this==Owner->GetWorld());
 		}
 		FActorSpawnParameters SpawnInfo;
-		SpawnInfo.bNoCollisionFail = bNoCollisionFail;
+		SpawnInfo.SpawnCollisionHandlingOverride = bNoCollisionFail ? ESpawnActorCollisionHandlingMethod::AlwaysSpawn : ESpawnActorCollisionHandlingMethod::AdjustIfPossibleButDontSpawnIfColliding;
 		SpawnInfo.Owner = Owner;
 		SpawnInfo.Instigator = Instigator;
 		SpawnInfo.bDeferConstruction = true;
@@ -2870,17 +2875,44 @@ public:
 	}
 
 	/**
-	* Spawns given class and returns class T pointer, forcibly sets world transform(note this allows scale as well). WILL NOT run Construction Script of Blueprints 
+	* Spawns given class and returns class T pointer, forcibly sets world position. WILL NOT run Construction Script of Blueprints 
 	* to give caller an opportunity to set parameters beforehand.  Caller is responsible for invoking construction
 	* manually by calling UGameplayStatics::FinishSpawningActor (see AActor::OnConstruction).
 	*/
 	template< class T >
+	T* SpawnActorDeferred(
+		UClass* Class,
+		FVector const& Location,
+		FRotator const& Rotation,
+		AActor* Owner = nullptr,
+		APawn* Instigator = nullptr,
+		ESpawnActorCollisionHandlingMethod CollisionHandlingOverride = ESpawnActorCollisionHandlingMethod::Undefined
+		)
+	{
+		if( Owner )
+		{
+			check(this==Owner->GetWorld());
+		}
+		FActorSpawnParameters SpawnInfo;
+		SpawnInfo.SpawnCollisionHandlingOverride = CollisionHandlingOverride;
+		SpawnInfo.Owner = Owner;
+		SpawnInfo.Instigator = Instigator;
+		SpawnInfo.bDeferConstruction = true;
+		return (Class != NULL) ? Cast<T>(SpawnActor(Class, &Location, &Rotation, SpawnInfo )) : NULL;
+	}
+
+	/**
+	 * Spawns given class and returns class T pointer, forcibly sets world transform (note this allows scale as well). WILL NOT run Construction Script of Blueprints 
+	 * to give caller an opportunity to set parameters beforehand.  Caller is responsible for invoking construction
+	 * manually by calling UGameplayStatics::FinishSpawningActor (see AActor::OnConstruction).
+	 */
+	template< class T >
 	T* SpawnActorDeferred(
 		UClass* Class,
 		FTransform const& Transform,
 		AActor* Owner=NULL,
 		APawn* Instigator=NULL,
-		bool bNoCollisionFail=false
+		ESpawnActorCollisionHandlingMethod CollisionHandlingOverride = ESpawnActorCollisionHandlingMethod::Undefined
 		)
 	{
 		if( Owner )
@@ -2888,11 +2920,11 @@ public:
 			check(this==Owner->GetWorld());
 		}
 		FActorSpawnParameters SpawnInfo;
-		SpawnInfo.bNoCollisionFail = bNoCollisionFail;
+		SpawnInfo.SpawnCollisionHandlingOverride = CollisionHandlingOverride;
 		SpawnInfo.Owner = Owner;
 		SpawnInfo.Instigator = Instigator;
 		SpawnInfo.bDeferConstruction = true;
-		return (Class != NULL) ? Cast<T>(SpawnActor(Class, &Transform, SpawnInfo )) : NULL;
+		return (Class != NULL) ? Cast<T>(SpawnActor(Class, &Transform, SpawnInfo)) : NULL;
 	}
 
 	/** 
diff --git a/Engine/Source/Runtime/Engine/Classes/GameFramework/Actor.h b/Engine/Source/Runtime/Engine/Classes/GameFramework/Actor.h
--- a/Engine/Source/Runtime/Engine/Classes/GameFramework/Actor.h
+++ b/Engine/Source/Runtime/Engine/Classes/GameFramework/Actor.h
@@ -341,18 +341,22 @@ private:
 
 public:
 
-	/** This actor collides with the world when placing in the editor or when spawned, even if RootComponent collision is disabled */
+	/** This actor collides with the world when placing in the editor, even if RootComponent collision is disabled. Does not affect spawning, @see SpawnCollisionHandlingMethod */
 	UPROPERTY()
-	uint32 bCollideWhenPlacing:1;    
+	uint32 bCollideWhenPlacing:1;
 
 	/** If true, this actor should search for an owned camera component to view through when used as a view target. */
 	UPROPERTY(EditAnywhere, BlueprintReadWrite, Category=Actor, AdvancedDisplay)
 	uint32 bFindCameraComponentWhenViewTarget:1;
 	
-	/** If true, this actor will be replicated to network replys (default is true) */
+	/** If true, this actor will be replicated to network replays (default is true) */
 	UPROPERTY()
 	uint32 bRelevantForNetworkReplays:1;
 
+	/** Controls how to handle spawning this actor in a situation where it's colliding with something else. "Default" means AlwaysSpawn here. */
+	UPROPERTY(EditAnywhere, BlueprintReadWrite, Category=Actor)
+	ESpawnActorCollisionHandlingMethod SpawnCollisionHandlingMethod;
+	
 	/** The time this actor was created, relative to World->GetTimeSeconds().
 	* @see UWorld::GetTimeSeconds()
 	*/
diff --git a/Engine/Source/Runtime/Engine/Classes/Kismet/GameplayStatics.h b/Engine/Source/Runtime/Engine/Classes/Kismet/GameplayStatics.h
--- a/Engine/Source/Runtime/Engine/Classes/Kismet/GameplayStatics.h
+++ b/Engine/Source/Runtime/Engine/Classes/Kismet/GameplayStatics.h
@@ -42,10 +42,14 @@ class ENGINE_API UGameplayStatics : public UBlueprintFunctionLibrary
 	UFUNCTION(BlueprintCallable, Category="Spawning", meta=(WorldContext="WorldContextObject", UnsafeDuringActorConstruction = "true", BlueprintInternalUseOnly = "true", DeprecatedFunction, DeprecationMessage="Use BeginSpawningActorFromClass"))
 	static class AActor* BeginSpawningActorFromBlueprint(UObject* WorldContextObject, const class UBlueprint* Blueprint, const FTransform& SpawnTransform, bool bNoCollisionFail);
 
-	/** Spawns an instance of an actor class, but does not automatically run its construction script.  */
+	DEPRECATED(4.9, "This function is deprecated. Please use BeginDeferredActorSpawnFromClass instead.")
 	UFUNCTION(BlueprintCallable, Category="Spawning", meta=(WorldContext="WorldContextObject", UnsafeDuringActorConstruction = "true", BlueprintInternalUseOnly = "true"))
 	static class AActor* BeginSpawningActorFromClass(UObject* WorldContextObject, TSubclassOf<AActor> ActorClass, const FTransform& SpawnTransform, bool bNoCollisionFail = false, AActor* Owner = nullptr);
 
+	/** Spawns an instance of an actor class, but does not automatically run its construction script.  */
+	UFUNCTION(BlueprintCallable, Category = "Spawning", meta = (WorldContext = "WorldContextObject", UnsafeDuringActorConstruction = "true", BlueprintInternalUseOnly = "true"))
+	static class AActor* BeginDeferredActorSpawnFromClass(UObject* WorldContextObject, TSubclassOf<AActor> ActorClass, const FTransform& SpawnTransform, ESpawnActorCollisionHandlingMethod CollisionHandlingOverride = ESpawnActorCollisionHandlingMethod::Undefined, AActor* Owner = nullptr);
+
 	/** 'Finish' spawning an actor.  This will run the construction script. */
 	UFUNCTION(BlueprintCallable, Category="Spawning", meta=(UnsafeDuringActorConstruction = "true", BlueprintInternalUseOnly = "true"))
 	static class AActor* FinishSpawningActor(class AActor* Actor, const FTransform& SpawnTransform);
diff --git a/Engine/Source/Runtime/Engine/Private/Actor.cpp b/Engine/Source/Runtime/Engine/Private/Actor.cpp
--- a/Engine/Source/Runtime/Engine/Private/Actor.cpp
+++ b/Engine/Source/Runtime/Engine/Private/Actor.cpp
@@ -96,6 +96,7 @@ void AActor::InitializeDefaults()
 #if WITH_EDITORONLY_DATA
 	PivotOffset = FVector::ZeroVector;
 #endif
+	SpawnCollisionHandlingMethod = ESpawnActorCollisionHandlingMethod::AlwaysSpawn;
 }
 
 void FActorTickFunction::ExecuteTick(float DeltaTime, enum ELevelTick TickType, ENamedThreads::Type CurrentThread, const FGraphEventRef& MyCompletionGraphEvent)
@@ -300,7 +301,7 @@ bool AActor::TeleportTo( const FVector& DestLocation, const FRotator& DestRotati
 	UPrimitiveComponent* ActorPrimComp = Cast<UPrimitiveComponent>(RootComponent);
 	if ( ActorPrimComp )
 	{
-		if (!bNoCheck && (ActorPrimComp->IsCollisionEnabled() || (bCollideWhenPlacing && (GetNetMode() != NM_Client))) )
+		if (!bNoCheck && (ActorPrimComp->IsCollisionEnabled() || (GetNetMode() != NM_Client)) )
 		{
 			// Apply the pivot offset to the desired location
 			FVector Offset = GetRootComponent()->Bounds.Origin - PrevLocation;
@@ -1649,7 +1650,7 @@ FVector AActor::GetPlacementExtent() const
 		FBox ActorBox(0.f);
 		for (int32 ComponentID=0; ComponentID<Components.Num(); ++ComponentID)
 		{
-			USceneComponent * SceneComp = Components[ComponentID];
+			USceneComponent* SceneComp = Components[ComponentID];
 			if (SceneComp->ShouldCollideWhenPlacing() )
 			{
 				ActorBox += SceneComp->GetPlacementExtent().GetBox();
@@ -2469,7 +2470,7 @@ void AActor::FinishSpawning(const FTransform& Transform, bool bIsDefaultTransfor
 
 void AActor::PostActorConstruction()
 {
-	UWorld* World = GetWorld();
+	UWorld* const World = GetWorld();
 	bool const bActorsInitialized = World && World->AreActorsInitialized();
 
 	if (bActorsInitialized)
@@ -2485,15 +2486,66 @@ void AActor::PostActorConstruction()
 		// Call InitializeComponent on components
 		InitializeComponents();
 
-		PostInitializeComponents();
-		if (!bActorInitialized && !IsPendingKill())
+		// actor should have all of its components created and registered now, do any collision checking and handling that we need to do
+		if (World)
 		{
-			UE_LOG(LogActor, Fatal, TEXT("%s failed to route PostInitializeComponents.  Please call Super::PostInitializeComponents() in your <className>::PostInitializeComponents() function. "), *GetFullName() );
+			switch (SpawnCollisionHandlingMethod)
+			{
+			case ESpawnActorCollisionHandlingMethod::AdjustIfPossibleButAlwaysSpawn:
+			{
+				// Try to find a spawn position
+				FVector AdjustedLocation = GetActorLocation();
+				FRotator AdjustedRotation = GetActorRotation();
+				if (World->FindTeleportSpot(this, AdjustedLocation, AdjustedRotation))
+				{
+					SetActorLocationAndRotation(AdjustedLocation, AdjustedRotation);
+				}
+			}
+			break;
+			case ESpawnActorCollisionHandlingMethod::AdjustIfPossibleButDontSpawnIfColliding:
+			{
+				// Try to find a spawn position			
+				FVector AdjustedLocation = GetActorLocation();
+				FRotator AdjustedRotation = GetActorRotation();
+				if (World->FindTeleportSpot(this, AdjustedLocation, AdjustedRotation))
+				{
+					SetActorLocationAndRotation(AdjustedLocation, AdjustedRotation);
+				}
+				else
+				{
+					UE_LOG(LogSpawn, Warning, TEXT("SpawnActor failed because of collision at the spawn location [%s] for [%s]"), *AdjustedLocation.ToString(), *GetClass()->GetName());
+					Destroy();
+				}
+			}
+			break;
+			case ESpawnActorCollisionHandlingMethod::DontSpawnIfColliding:
+				if (World->EncroachingBlockingGeometry(this, GetActorLocation(), GetActorRotation()))
+				{
+					UE_LOG(LogSpawn, Warning, TEXT("SpawnActor failed because of collision at the spawn location [%s] for [%s]"), *GetActorLocation().ToString(), *GetClass()->GetName());
+					Destroy();
+				}
+				break;
+			case ESpawnActorCollisionHandlingMethod::Undefined:
+			case ESpawnActorCollisionHandlingMethod::AlwaysSpawn:
+			default:
+				// note we use "always spawn" as default, so treat undefined as that
+				// nothing to do here, just proceed as normal
+				break;
+			}
 		}
 
-		if (World->HasBegunPlay() && !deferBeginPlayAndUpdateOverlaps)
+		if (!IsPendingKill())
 		{
-			BeginPlay();
+			PostInitializeComponents();
+			if (!bActorInitialized && !IsPendingKill())
+			{
+				UE_LOG(LogActor, Fatal, TEXT("%s failed to route PostInitializeComponents.  Please call Super::PostInitializeComponents() in your <className>::PostInitializeComponents() function. "), *GetFullName());
+			}
+
+			if (World->HasBegunPlay() && !deferBeginPlayAndUpdateOverlaps)
+			{
+				BeginPlay();
+			}
 		}
 	}
 	else
@@ -2506,14 +2558,17 @@ void AActor::PostActorConstruction()
 		ClearFlags(RF_PendingKill);
 	}
 
-	// Components are all there and we've begun play, init overlapping state
-	if (!deferBeginPlayAndUpdateOverlaps)
+	if (!IsPendingKill())
 	{
-		UpdateOverlaps();
-	}
+		// Components are all there and we've begun play, init overlapping state
+		if (!deferBeginPlayAndUpdateOverlaps)
+		{
+			UpdateOverlaps();
+		}
 
-	// Notify the texture streaming manager about the new actor.
-	IStreamingManager::Get().NotifyActorSpawned(this);
+		// Notify the texture streaming manager about the new actor.
+		IStreamingManager::Get().NotifyActorSpawned(this);
+	}
 }
 
 void AActor::SetReplicates(bool bInReplicates)
diff --git a/Engine/Source/Runtime/Engine/Private/Brush.cpp b/Engine/Source/Runtime/Engine/Private/Brush.cpp
--- a/Engine/Source/Runtime/Engine/Private/Brush.cpp
+++ b/Engine/Source/Runtime/Engine/Private/Brush.cpp
@@ -36,6 +36,7 @@ ABrush::ABrush(const FObjectInitializer& ObjectInitializer)
 	bNotForClientOrServer = false;
 	bCanBeDamaged = false;
 	bCollideWhenPlacing = true;
+	SpawnCollisionHandlingMethod = ESpawnActorCollisionHandlingMethod::AdjustIfPossibleButDontSpawnIfColliding;
 }
 
 #if WITH_EDITOR
diff --git a/Engine/Source/Runtime/Engine/Private/Components/ChildActorComponent.cpp b/Engine/Source/Runtime/Engine/Private/Components/ChildActorComponent.cpp
--- a/Engine/Source/Runtime/Engine/Private/Components/ChildActorComponent.cpp
+++ b/Engine/Source/Runtime/Engine/Private/Components/ChildActorComponent.cpp
@@ -243,7 +243,7 @@ void UChildActorComponent::CreateChildActor()
 			if (bSpawn)
 			{
 				FActorSpawnParameters Params;
-				Params.bNoCollisionFail = true;
+				Params.SpawnCollisionHandlingOverride = ESpawnActorCollisionHandlingMethod::AlwaysSpawn;
 				Params.bDeferConstruction = true; // We defer construction so that we set ParentComponentActor prior to component registration so they appear selected
 				Params.bAllowDuringConstructionScript = true;
 				Params.OverrideLevel = (MyOwner ? MyOwner->GetLevel() : nullptr);
diff --git a/Engine/Source/Runtime/Engine/Private/Controller.cpp b/Engine/Source/Runtime/Engine/Private/Controller.cpp
--- a/Engine/Source/Runtime/Engine/Private/Controller.cpp
+++ b/Engine/Source/Runtime/Engine/Private/Controller.cpp
@@ -438,7 +438,7 @@ void AController::InitPlayerState()
 			FActorSpawnParameters SpawnInfo;
 			SpawnInfo.Owner = this;
 			SpawnInfo.Instigator = Instigator;
-			SpawnInfo.bNoCollisionFail = true;
+			SpawnInfo.SpawnCollisionHandlingOverride = ESpawnActorCollisionHandlingMethod::AlwaysSpawn;
 			SpawnInfo.ObjectFlags |= RF_Transient;	// We never want player states to save into a map
 			PlayerState = World->SpawnActor<APlayerState>(GameMode->PlayerStateClass, SpawnInfo );
 	
diff --git a/Engine/Source/Runtime/Engine/Private/Debug/LogVisualizerCameraController.cpp b/Engine/Source/Runtime/Engine/Private/Debug/LogVisualizerCameraController.cpp
--- a/Engine/Source/Runtime/Engine/Private/Debug/LogVisualizerCameraController.cpp
+++ b/Engine/Source/Runtime/Engine/Private/Debug/LogVisualizerCameraController.cpp
@@ -57,7 +57,7 @@ void ALogVisualizerCameraController::PostInitializeComponents()
 	FActorSpawnParameters SpawnInfo;
 	SpawnInfo.Owner = this;
 	SpawnInfo.Instigator = Instigator;
-	SpawnInfo.bNoCollisionFail = true;
+	SpawnInfo.SpawnCollisionHandlingOverride = ESpawnActorCollisionHandlingMethod::AlwaysSpawn;
 	MyHUD = GetWorld()->SpawnActor<ALogVisualizerHUD>( SpawnInfo );
 
 	/** TO DO: somehow this two configuration settings are set to zero for this controller, have to fix it later */
@@ -82,7 +82,7 @@ ALogVisualizerCameraController* ALogVisualizerCameraController::EnableCamera(UWo
 		// in editor world (if running PIE). Hate it, but it works, and 
 		// this is a debugging tool		
 		FActorSpawnParameters SpawnInfo;
-		SpawnInfo.bNoCollisionFail = true;
+		SpawnInfo.SpawnCollisionHandlingOverride = ESpawnActorCollisionHandlingMethod::AlwaysSpawn;
 		SpawnInfo.Owner = PC->GetWorldSettings();
 		SpawnInfo.Instigator = PC->Instigator;
 		Instance = InWorld->SpawnActor<ALogVisualizerCameraController>( SpawnInfo );
diff --git a/Engine/Source/Runtime/Engine/Private/DebugCameraController.cpp b/Engine/Source/Runtime/Engine/Private/DebugCameraController.cpp
--- a/Engine/Source/Runtime/Engine/Private/DebugCameraController.cpp
+++ b/Engine/Source/Runtime/Engine/Private/DebugCameraController.cpp
@@ -258,7 +258,7 @@ ASpectatorPawn* ADebugCameraController::SpawnSpectatorPawn()
 		{
 			FActorSpawnParameters SpawnParams;
 			SpawnParams.Owner = this;
-			SpawnParams.bNoCollisionFail = true;
+			SpawnParams.SpawnCollisionHandlingOverride = ESpawnActorCollisionHandlingMethod::AlwaysSpawn;
 			SpawnParams.ObjectFlags |= RF_Transient;	// We never want to save spectator pawns into a map
 			SpawnedSpectator = GetWorld()->SpawnActor<ASpectatorPawn>(ASpectatorPawn::StaticClass(), GetSpawnLocation(), GetControlRotation(), SpawnParams);
 			if (SpawnedSpectator)
@@ -324,7 +324,7 @@ void ADebugCameraController::PostInitializeComponents()
 	FActorSpawnParameters SpawnInfo;
 	SpawnInfo.Owner = this;
 	SpawnInfo.Instigator = Instigator;
-	SpawnInfo.bNoCollisionFail = true;
+	SpawnInfo.SpawnCollisionHandlingOverride = ESpawnActorCollisionHandlingMethod::AlwaysSpawn;
 	SpawnInfo.ObjectFlags |= RF_Transient;	// We never want these to save into a map
 	MyHUD = GetWorld()->SpawnActor<ADebugCameraHUD>( SpawnInfo );
 
diff --git a/Engine/Source/Runtime/Engine/Private/DefaultPawn.cpp b/Engine/Source/Runtime/Engine/Private/DefaultPawn.cpp
--- a/Engine/Source/Runtime/Engine/Private/DefaultPawn.cpp
+++ b/Engine/Source/Runtime/Engine/Private/DefaultPawn.cpp
@@ -23,6 +23,7 @@ ADefaultPawn::ADefaultPawn(const FObjectInitializer& ObjectInitializer)
 
 	BaseEyeHeight = 0.0f;
 	bCollideWhenPlacing = false;
+	SpawnCollisionHandlingMethod = ESpawnActorCollisionHandlingMethod::AlwaysSpawn;
 
 	CollisionComponent = CreateDefaultSubobject<USphereComponent>(ADefaultPawn::CollisionComponentName);
 	CollisionComponent->InitSphereRadius(35.0f);
diff --git a/Engine/Source/Runtime/Engine/Private/GameplayStatics.cpp b/Engine/Source/Runtime/Engine/Private/GameplayStatics.cpp
--- a/Engine/Source/Runtime/Engine/Private/GameplayStatics.cpp
+++ b/Engine/Source/Runtime/Engine/Private/GameplayStatics.cpp
@@ -367,7 +367,8 @@ AActor* UGameplayStatics::BeginSpawningActorFromBlueprint(UObject* WorldContextO
 	{
 		if( Blueprint->GeneratedClass->IsChildOf(AActor::StaticClass()) )
 		{
-			NewActor = BeginSpawningActorFromClass(WorldContextObject, *Blueprint->GeneratedClass, SpawnTransform, bNoCollisionFail);
+			ESpawnActorCollisionHandlingMethod const CollisionHandlingOverride = bNoCollisionFail ? ESpawnActorCollisionHandlingMethod::AdjustIfPossibleButDontSpawnIfColliding : ESpawnActorCollisionHandlingMethod::AlwaysSpawn;
+			NewActor = BeginDeferredActorSpawnFromClass(WorldContextObject, *Blueprint->GeneratedClass, SpawnTransform, CollisionHandlingOverride);
 		}
 		else
 		{
@@ -377,7 +378,14 @@ AActor* UGameplayStatics::BeginSpawningActorFromBlueprint(UObject* WorldContextO
 	return NewActor;
 }
 
+// deprecated
 AActor* UGameplayStatics::BeginSpawningActorFromClass(UObject* WorldContextObject, TSubclassOf<AActor> ActorClass, const FTransform& SpawnTransform, bool bNoCollisionFail, AActor* Owner)
+{
+	ESpawnActorCollisionHandlingMethod const CollisionHandlingOverride = bNoCollisionFail ? ESpawnActorCollisionHandlingMethod::AdjustIfPossibleButDontSpawnIfColliding : ESpawnActorCollisionHandlingMethod::AlwaysSpawn;
+	return BeginDeferredActorSpawnFromClass(WorldContextObject, ActorClass, SpawnTransform, CollisionHandlingOverride, Owner);
+}
+
+AActor* UGameplayStatics::BeginDeferredActorSpawnFromClass(UObject* WorldContextObject, TSubclassOf<AActor> ActorClass, const FTransform& SpawnTransform, ESpawnActorCollisionHandlingMethod CollisionHandlingMethod, AActor* Owner)
 {
 	AActor* NewActor = NULL;
 
@@ -400,7 +408,7 @@ AActor* UGameplayStatics::BeginSpawningActorFromClass(UObject* WorldContextObjec
 		UWorld* World = GEngine->GetWorldFromContextObject(WorldContextObject);
 		if (World)
 		{
-			NewActor = World->SpawnActorDeferred<AActor>(Class, SpawnTransform, Owner, AutoInstigator, bNoCollisionFail);
+			NewActor = World->SpawnActorDeferred<AActor>(Class, SpawnTransform, Owner, AutoInstigator, CollisionHandlingMethod);
 		}
 		else
 		{
diff --git a/Engine/Source/Runtime/Engine/Private/LevelActor.cpp b/Engine/Source/Runtime/Engine/Private/LevelActor.cpp
--- a/Engine/Source/Runtime/Engine/Private/LevelActor.cpp
+++ b/Engine/Source/Runtime/Engine/Private/LevelActor.cpp
@@ -13,6 +13,7 @@
 #endif // WITH_PHYSX
 #include "Components/SphereComponent.h"
 #include "Components/CapsuleComponent.h"
+#include "Components/BoxComponent.h"
 #include "GameFramework/GameMode.h"
 
 // CVars
@@ -324,18 +325,46 @@ AActor* UWorld::SpawnActor( UClass* Class, FTransform const* Transform, const FA
 	FRotator NewRotation = Transform ? Transform->GetRotation().Rotator() :	(Template->GetRootComponent() ? Template->GetRootComponent()->RelativeRotation : FRotator::ZeroRotator);
 	FVector NewScale = Transform ? Transform->GetScale3D() :	(Template->GetRootComponent() ? Template->GetRootComponent()->RelativeScale3D : FVector(1.f) );
 
-	// If we can fail, and we have a root component, and either the component should collide or bCollideWhenPlacing is true, test
-	if( !SpawnParameters.bNoCollisionFail && Template->GetRootComponent() && Template->bCollideWhenPlacing && Template->GetRootComponent()->ShouldCollideWhenPlacing() )
+	PRAGMA_DISABLE_DEPRECATION_WARNINGS;
+	// handle existing (but deprecated) uses of bNoCollisionFail where user set it to true
+	ESpawnActorCollisionHandlingMethod CollisionHandlingOverride = SpawnParameters.SpawnCollisionHandlingOverride;
+	if ((CollisionHandlingOverride == ESpawnActorCollisionHandlingMethod::Undefined) && SpawnParameters.bNoCollisionFail)
 	{
-		// Try to find a spawn position
-		if ( !FindTeleportSpot(Template, NewLocation, NewRotation) )
+		CollisionHandlingOverride = ESpawnActorCollisionHandlingMethod::AlwaysSpawn;
+	}
+	PRAGMA_ENABLE_DEPRECATION_WARNINGS;
+
+	// "no fail" take preedence over collision handling settings that include fails
+	if (SpawnParameters.bNoFail)
+	{
+		// maybe upgrade to disallow fail
+		if (CollisionHandlingOverride == ESpawnActorCollisionHandlingMethod::AdjustIfPossibleButDontSpawnIfColliding)
 		{
-			UE_LOG(LogSpawn, Warning, TEXT("SpawnActor failed because of collision at the spawn location [%s] for [%s]"), *NewLocation.ToString(), *Class->GetName() );
-			return NULL;
+			CollisionHandlingOverride = ESpawnActorCollisionHandlingMethod::AdjustIfPossibleButAlwaysSpawn;
+		}
+		else if (CollisionHandlingOverride == ESpawnActorCollisionHandlingMethod::DontSpawnIfColliding)
+		{
+			CollisionHandlingOverride = ESpawnActorCollisionHandlingMethod::AlwaysSpawn;
 		}
 	}
 
-	AActor* Actor = NewObject<AActor>(LevelToSpawnIn, Class, NewActorName, SpawnParameters.ObjectFlags, Template);
+	// use override if set, else fall back to actor's preference
+	ESpawnActorCollisionHandlingMethod const CollisionHandlingMethod = (CollisionHandlingOverride == ESpawnActorCollisionHandlingMethod::Undefined) ? Template->SpawnCollisionHandlingMethod : CollisionHandlingOverride;
+
+	// see if we can avoid spawning altogether by checking native components
+	// note: we can't handle all cases here, since we don't know the full component hierarchy until after the actor is spawned
+	if (CollisionHandlingMethod == ESpawnActorCollisionHandlingMethod::DontSpawnIfColliding)
+	{
+		if (EncroachingBlockingGeometry(Template, NewLocation, NewRotation))
+		{
+			// a native component is colliding, that's enough to reject spawning
+			UE_LOG(LogSpawn, Warning, TEXT("SpawnActor failed because of collision at the spawn location [%s] for [%s]"), *NewLocation.ToString(), *Class->GetName());
+			return nullptr;
+		}
+	}
+
+	// actually make the actor object
+	AActor* const Actor = NewObject<AActor>(LevelToSpawnIn, Class, NewActorName, SpawnParameters.ObjectFlags, Template);
 	check(Actor);
 
 #if WITH_EDITOR
@@ -359,6 +388,9 @@ AActor* UWorld::SpawnActor( UClass* Class, FTransform const* Transform, const FA
 	}
 #endif
 
+	// tell the actor what method to use, in case it was overridden
+	Actor->SpawnCollisionHandlingMethod = CollisionHandlingMethod;
+
 	Actor->PostSpawnInitialize(FTransform(NewRotation, NewLocation, NewScale), SpawnParameters.Owner, SpawnParameters.Instigator, SpawnParameters.bRemoteOwned, SpawnParameters.bNoFail, SpawnParameters.bDeferConstruction);
 
 	if (Actor->IsPendingKill() && !SpawnParameters.bNoFail)
@@ -366,6 +398,55 @@ AActor* UWorld::SpawnActor( UClass* Class, FTransform const* Transform, const FA
 		UE_LOG(LogSpawn, Warning, TEXT("SpawnActor failed because the spawned actor IsPendingKill"));
 		return NULL;
 	}
+// 
+// 	// actor should have all of its components now, do any collision checking and handling that we need to do
+// 	switch (CollisionHandlingMethod)
+// 	{
+// 	case ESpawnActorCollisionHandlingMethod::AdjustIfPossibleButAlwaysSpawn:
+// 	 	// Try to find a spawn position
+// 		{
+// 			FVector AdjustedLocation = NewLocation;
+// 			FRotator AdjustedRotation = NewRotation;
+// 			if (FindTeleportSpot(Actor, AdjustedLocation, AdjustedRotation))
+// 			{
+// 				Actor->SetActorLocationAndRotation(AdjustedLocation, AdjustedRotation);
+// 			}
+// 		}
+// 	 	break;
+// 	case ESpawnActorCollisionHandlingMethod::AdjustIfPossibleButDontSpawnIfColliding:
+// 	 	// Try to find a spawn position
+// 		{
+// 			FVector AdjustedLocation = NewLocation;
+// 			FRotator AdjustedRotation = NewRotation;
+// 			if (FindTeleportSpot(Actor, AdjustedLocation, AdjustedRotation))
+// 			{
+// 				Actor->SetActorLocationAndRotation(AdjustedLocation, AdjustedRotation);
+// 			}
+// 			else
+// 			{
+// 				UE_LOG(LogSpawn, Warning, TEXT("SpawnActor failed because of collision at the spawn location [%s] for [%s]"), *NewLocation.ToString(), *Class->GetName());
+// 				DestroyActor(Actor);
+// 				return nullptr;
+// 			}
+// 		}
+// 	 	break;
+// 	case ESpawnActorCollisionHandlingMethod::DontSpawnIfColliding:
+// 		// #todo: don't recheck components checked above?
+// 		if (EncroachingBlockingGeometry(Actor, NewLocation, NewRotation))
+// 	 	{
+// 	 		UE_LOG(LogSpawn, Warning, TEXT("SpawnActor failed because of collision at the spawn location [%s] for [%s]"), *NewLocation.ToString(), *Class->GetName());
+// 			DestroyActor(Actor);
+// 			return nullptr;
+// 	 	}
+// 	 	break;
+// 	// note we use "always spawn" as default, so treat undefined as that
+// 	case ESpawnActorCollisionHandlingMethod::Undefined:
+// 	case ESpawnActorCollisionHandlingMethod::AlwaysSpawn:
+// 	default:
+// 	 	// nothing to do, just proceed as normal
+// 	 	break;
+// 	}
+
 	Actor->CheckDefaultSubobjects();
 
 	// Broadcast notification of spawn
@@ -385,8 +466,8 @@ AActor* UWorld::SpawnActor( UClass* Class, FTransform const* Transform, const FA
 ABrush* UWorld::SpawnBrush()
 {
 	FActorSpawnParameters SpawnInfo;
-	SpawnInfo.bNoCollisionFail = true;
-	ABrush* Result = SpawnActor<ABrush>( SpawnInfo );
+	SpawnInfo.SpawnCollisionHandlingOverride = ESpawnActorCollisionHandlingMethod::AlwaysSpawn;
+	ABrush* const Result = SpawnActor<ABrush>( SpawnInfo );
 	check(Result);
 	return Result;
 }
@@ -637,7 +718,7 @@ bool UWorld::FindTeleportSpot(AActor* TestActor, FVector& TestLocation, FRotator
 	}
 
 	// first do only Z
-	if ( Adjust.Z != 0.f )
+	if (!FMath::IsNearlyZero(Adjust.Z)) 
 	{
 		TestLocation.Z += Adjust.Z;
 		if( !EncroachingBlockingGeometry(TestActor, TestLocation, TestRotation, &Adjust) )
@@ -647,7 +728,7 @@ bool UWorld::FindTeleportSpot(AActor* TestActor, FVector& TestLocation, FRotator
 	}
 
 	// now try just XY
-	if ( (Adjust.X != 0.f) || (Adjust.Y != 0.f) )
+	if (!FMath::IsNearlyZero(Adjust.X) || !FMath::IsNearlyZero(Adjust.Y))
 	{
 		for (int i = 0; i < 8; ++i)
 		{
@@ -665,7 +746,7 @@ bool UWorld::FindTeleportSpot(AActor* TestActor, FVector& TestLocation, FRotator
 	}
 
 	// now z again
-	if ( Adjust.Z != 0.f )
+	if (!FMath::IsNearlyZero(Adjust.Z))
 	{
 		TestLocation.Z += Adjust.Z;
 		if( !EncroachingBlockingGeometry(TestActor, TestLocation, TestRotation, &Adjust) )
@@ -673,7 +754,8 @@ bool UWorld::FindTeleportSpot(AActor* TestActor, FVector& TestLocation, FRotator
 			return true;
 		}
 	}
-	if ( Adjust.IsZero() )
+
+	if ( Adjust.IsNearlyZero() )
 	{
 		return false;
 	}
@@ -684,85 +766,237 @@ bool UWorld::FindTeleportSpot(AActor* TestActor, FVector& TestLocation, FRotator
 	return !EncroachingBlockingGeometry(TestActor, TestLocation, TestRotation, &Adjust);
 }
 
-bool UWorld::EncroachingBlockingGeometry(AActor* TestActor, FVector TestLocation, FRotator TestRotation, FVector* ProposedAdjustment )
-{
-	float Epsilon = CVarEncroachEpsilon.GetValueOnGameThread();
-
-	// currently just tests RootComponent.  @TODO FIXME should test all colliding components?  Cost/benefit?
-	UPrimitiveComponent* PrimComp = TestActor ? Cast<UPrimitiveComponent>(TestActor->GetRootComponent()) : NULL;
-	bool bFoundBlockingHit = false;
-	if (ProposedAdjustment)
-	{
-		*ProposedAdjustment = FVector::ZeroVector;
-	}
-	if ( PrimComp )
+/** Tests shape components more efficiently than the with-adjustment case, but does less-efficient ppr-poly collision for meshes. */
+static bool ComponentEncroachesBlockingGeometry_NoAdjustment(UWorld const* World, AActor const* TestActor, UPrimitiveComponent const* PrimComp, FTransform const& TestWorldTransform)
+{	
+	float const Epsilon = CVarEncroachEpsilon.GetValueOnGameThread();
+	
+	if (World && PrimComp)
 	{
-		TArray<FOverlapResult> Overlaps;
-		static FName NAME_EncroachingBlockingGeometry = FName(TEXT("EncroachingBlockingGeometry"));
-		ECollisionChannel BlockingChannel = PrimComp->GetCollisionObjectType();
+		bool bFoundBlockingHit = false;
+		
+		static FName NAME_ComponentEncroachesBlockingGeometry_NoAdjustment = FName(TEXT("ComponentEncroachesBlockingGeometry_NoAdjustment"));
+		ECollisionChannel const BlockingChannel = PrimComp->GetCollisionObjectType();
+		FCollisionShape const CollisionShape = PrimComp->GetCollisionShape(-Epsilon);
 
-		// @TODO support more than just capsules and spheres for spawning, don't yet have - add primitive component function to calculate these, use boundingbox as fallback
-		UCapsuleComponent* const Capsule = Cast<UCapsuleComponent>(PrimComp);
-		if (Capsule)
+		if (CollisionShape.IsBox() && (Cast<UBoxComponent>(PrimComp) == nullptr))
 		{
-			FCollisionQueryParams Params(NAME_EncroachingBlockingGeometry, false, TestActor);
-			bFoundBlockingHit = OverlapMultiByChannel(Overlaps, TestLocation, FQuat::Identity, BlockingChannel, FCollisionShape::MakeCapsule(FMath::Max(Capsule->GetScaledCapsuleRadius() - Epsilon, 0.1f), FMath::Max(Capsule->GetScaledCapsuleHalfHeight() - Epsilon, 0.1f)), Params);
+			// we have a bounding box not for a box component, which means this was the fallback aabb
+			// since we don't need the penetration info, go ahead and test the component itself for overlaps, which is more accurate
+			if (PrimComp->IsRegistered())
+			{
+				// must be registered
+				TArray<FOverlapResult> Overlaps;
+				FComponentQueryParams Params(NAME_ComponentEncroachesBlockingGeometry_NoAdjustment, TestActor);
+				return World->ComponentOverlapMultiByChannel(Overlaps, PrimComp, TestWorldTransform.GetLocation(), TestWorldTransform.GetRotation(), BlockingChannel, Params);
+			}
+			else
+			{
+				UE_LOG(LogPhysics, Log, TEXT("Components must be registered in order to be used in a ComponentOverlapMulti call. PriComp: %s TestActor: %s"), *PrimComp->GetName(), *TestActor->GetName());
+				return false;
+			}
 		}
 		else
 		{
-			USphereComponent* const Sphere = Cast<USphereComponent>(PrimComp);
-			if (Sphere)
-			{
-				FCollisionQueryParams Params(NAME_EncroachingBlockingGeometry, false, TestActor);
-				bFoundBlockingHit = OverlapMultiByChannel(Overlaps, TestLocation, FQuat::Identity, BlockingChannel, FCollisionShape::MakeSphere(FMath::Max(Sphere->GetScaledSphereRadius() - Epsilon, 0.1f)), Params);
-			}
-			else if (PrimComp->IsRegistered())
+			FCollisionQueryParams Params(NAME_ComponentEncroachesBlockingGeometry_NoAdjustment, false, TestActor);
+			return World->OverlapAnyTestByChannel(TestWorldTransform.GetLocation(), TestWorldTransform.GetRotation(), BlockingChannel, CollisionShape, Params);
+		}
+	}
+
+	return false;
+}
+
+/** Tests shape components less efficiently than the no-adjustment case, but does quicker aabb collision for meshes. */
+static bool ComponentEncroachesBlockingGeometry_WithAdjustment(UWorld const* World, AActor const* TestActor, UPrimitiveComponent const* PrimComp, FTransform const& TestWorldTransform, FVector& OutProposedAdjustment)
+{
+	// init our output
+	OutProposedAdjustment = FVector::ZeroVector;
+
+	float const Epsilon = CVarEncroachEpsilon.GetValueOnGameThread();
+
+	if (World && PrimComp)
+	{
+		bool bFoundBlockingHit = false;
+		bool bComputePenetrationAdjustment = true;
+		
+		TArray<FOverlapResult> Overlaps;
+		static FName NAME_ComponentEncroachesBlockingGeometry_WithAdjustment = FName(TEXT("ComponentEncroachesBlockingGeometry_WithAdjustment"));
+		ECollisionChannel const BlockingChannel = PrimComp->GetCollisionObjectType();
+		FCollisionShape const CollisionShape = PrimComp->GetCollisionShape(-Epsilon);
+
+		if (CollisionShape.IsBox() && (Cast<UBoxComponent>(PrimComp) == nullptr))
+		{
+			// we have a bounding box not for a box component, which means this was the fallback aabb
+			// so lets test the actual component instead of it's aabb
+			// note we won't get penetration adjustment but that's ok
+			if (PrimComp->IsRegistered())
 			{
 				// must be registered
-				FComponentQueryParams Params(NAME_EncroachingBlockingGeometry, TestActor);
-				bFoundBlockingHit = ComponentOverlapMultiByChannel(Overlaps, PrimComp, TestLocation, TestActor->GetActorQuat(), BlockingChannel, Params);
+				FComponentQueryParams Params(NAME_ComponentEncroachesBlockingGeometry_WithAdjustment, TestActor);
+				bFoundBlockingHit = World->ComponentOverlapMultiByChannel(Overlaps, PrimComp, TestWorldTransform.GetLocation(), TestWorldTransform.GetRotation(), BlockingChannel, Params);
+				bComputePenetrationAdjustment = false;
 			}
 			else
 			{
 				UE_LOG(LogPhysics, Log, TEXT("Components must be registered in order to be used in a ComponentOverlapMulti call. PriComp: %s TestActor: %s"), *PrimComp->GetName(), *TestActor->GetName());
 			}
 		}
-
-		if ( !bFoundBlockingHit )
+		else
 		{
-			return false;
+			// overlap our shape
+			FCollisionQueryParams Params(NAME_ComponentEncroachesBlockingGeometry_WithAdjustment, false, TestActor);
+			bFoundBlockingHit = World->OverlapMultiByChannel(Overlaps, TestWorldTransform.GetLocation(), TestWorldTransform.GetRotation(), BlockingChannel, CollisionShape, Params);
 		}
 
-		if ( ProposedAdjustment && Capsule )
+		// compute adjustment
+		if (bFoundBlockingHit && bComputePenetrationAdjustment)
 		{
 			// if encroaching, add up all the MTDs of overlapping shapes
-			float CapsuleRadius, CapsuleHalfHeight;
-			Capsule->GetScaledCapsuleSize(CapsuleRadius, CapsuleHalfHeight);
-			FCollisionShape CapsuleShape = FCollisionShape::MakeCapsule(CapsuleRadius, CapsuleHalfHeight);
 			FMTDResult MTDResult;
-
 			for (int32 HitIdx = 0; HitIdx < Overlaps.Num(); HitIdx++)
 			{
-				UPrimitiveComponent * OverlapComponent = Overlaps[HitIdx].Component.Get();
+				UPrimitiveComponent* const OverlapComponent = Overlaps[HitIdx].Component.Get();
 				// first determine closest impact point along each axis
 				if (OverlapComponent && OverlapComponent->GetCollisionResponseToChannel(BlockingChannel) == ECR_Block)
-
 				{
-					bool bSuccess = OverlapComponent->ComputePenetration(MTDResult, CapsuleShape, TestLocation, FQuat::Identity);
+					FCollisionShape const NonShrunkenCollisionShape = PrimComp->GetCollisionShape();
+					bool bSuccess = OverlapComponent->ComputePenetration(MTDResult, NonShrunkenCollisionShape, TestWorldTransform.GetLocation(), TestWorldTransform.GetRotation());
 					if (bSuccess)
 					{
-						*ProposedAdjustment += MTDResult.Direction * MTDResult.Distance;
+						OutProposedAdjustment += MTDResult.Direction * MTDResult.Distance;
 					}
 					else
 					{
 						UE_LOG(LogPhysics, Log, TEXT("OverlapTest says we are overlapping, yet MTD says we're not. Something is wrong"));
 					}
+
+					// #hack: sometimes for boxes, physx returns a 0 MTD even though it reports a contact (returns true)
+					// to get around this, let's go ahead and test again with the epsilon-shrunken collision shape to see if we're really in 
+					// the clear.  if so, we'll say we have no contact (despite what OverlapMultiByChannel said -- it uses a different algorithm)
+					if (FMath::IsNearlyZero(MTDResult.Distance))
+					{
+						FCollisionShape const ShrunkenCollisionShape = PrimComp->GetCollisionShape(-Epsilon);
+						bool bSuccess = OverlapComponent->ComputePenetration(MTDResult, ShrunkenCollisionShape, TestWorldTransform.GetLocation(), TestWorldTransform.GetRotation());
+						if (bSuccess)
+						{
+							OutProposedAdjustment += MTDResult.Direction * MTDResult.Distance;
+						}
+						else
+						{
+							// let's call this "no contact" and be done
+							return false;
+						}
+					}
 				}
 			}
+		}
+
+		return bFoundBlockingHit;
+	}
+
+	return false;
+}
+
+/** Tests if the given component overlaps any blocking geometry if it were placed at the given world transform, optionally returns a suggested translation to get the component away from its overlaps. */
+static bool ComponentEncroachesBlockingGeometry(UWorld const* World, AActor const* TestActor, UPrimitiveComponent const* PrimComp, FTransform const& TestWorldTransform, FVector* OutProposedAdjustment)
+{
+	return OutProposedAdjustment
+		? ComponentEncroachesBlockingGeometry_WithAdjustment(World, TestActor, PrimComp, TestWorldTransform, *OutProposedAdjustment)
+		: ComponentEncroachesBlockingGeometry_NoAdjustment(World, TestActor, PrimComp, TestWorldTransform);
+}
+
+
+static FVector CombineAdjustments(FVector CurrentAdjustment, FVector AdjustmentToAdd)
+{
+	// remove the part of the new adjustment that's parallel to the current adjustment
+	if (CurrentAdjustment.IsZero())
+	{
+		return AdjustmentToAdd;
+	}
+
+	FVector Projection = AdjustmentToAdd.ProjectOnTo(CurrentAdjustment);
+	Projection = Projection.GetClampedToMaxSize(CurrentAdjustment.Size());
 
+	FVector OrthogalAdjustmentToAdd = AdjustmentToAdd - Projection;
+	return CurrentAdjustment + OrthogalAdjustmentToAdd;
+}
+
+// perf note: this is faster if ProposedAdjustment is null, since it can early out on first penetration
+bool UWorld::EncroachingBlockingGeometry(AActor* TestActor, FVector TestLocation, FRotator TestRotation, FVector* ProposedAdjustment)
+{
+	if (TestActor == nullptr)
+	{
+		return false;
+	}
+
+	USceneComponent* const RootComponent = TestActor->GetRootComponent();
+	if (RootComponent == nullptr)
+	{
+		return false;
+	}
+
+	bool bFoundEncroacher = false;
+	FVector TotalAdjustment(0.f);
+
+	FTransform const TestRootToWorld = FTransform(TestRotation, TestLocation);
+	
+	UPrimitiveComponent* const RootPrimComp = Cast<UPrimitiveComponent>(RootComponent);
+	if (RootPrimComp && RootPrimComp->IsCollisionEnabled())
+	{
+		if (ComponentEncroachesBlockingGeometry(this, TestActor, RootPrimComp, TestRootToWorld, ProposedAdjustment))
+		{
+			if (ProposedAdjustment == nullptr)
+			{
+				// don't need an adjustment and we know we are overlapping, so we can be done
+				return true;
+			}
+			else
+			{
+				TotalAdjustment = *ProposedAdjustment;
+			}
+
+			bFoundEncroacher = true;
+		}
+	}
+
+	// now test all colliding children for encroachment
+	TArray<USceneComponent*> Children;
+	RootComponent->GetChildrenComponents(true, Children);
+
+	FTransform const WorldToOldRoot = RootComponent->GetComponentToWorld().Inverse();
+
+	for (auto Child : Children)
+	{
+		if (Child->IsCollisionEnabled())
+		{
+			UPrimitiveComponent* const PrimComp = Cast<UPrimitiveComponent>(Child);
+			if (PrimComp)
+			{
+				FTransform const CompToRoot = Child->GetComponentToWorld() * WorldToOldRoot;
+				FTransform const CompToNewWorld = CompToRoot * TestRootToWorld;
+
+				if (ComponentEncroachesBlockingGeometry(this, TestActor, PrimComp, CompToNewWorld, ProposedAdjustment))
+				{
+					if (ProposedAdjustment == nullptr)
+					{
+						// don't need an adjustment and we know we are overlapping, so we can be done
+						return true;
+					}
+
+					TotalAdjustment = CombineAdjustments(TotalAdjustment, *ProposedAdjustment);
+					bFoundEncroacher = true;
+				}
+			}
 		}
 	}
-	return bFoundBlockingHit;
+
+	// copy over total adjustment
+	if (ProposedAdjustment)
+	{
+		*ProposedAdjustment = TotalAdjustment;
+	}
+
+	return bFoundEncroacher;
 }
 
 
diff --git a/Engine/Source/Runtime/Engine/Private/Light.cpp b/Engine/Source/Runtime/Engine/Private/Light.cpp
--- a/Engine/Source/Runtime/Engine/Private/Light.cpp
+++ b/Engine/Source/Runtime/Engine/Private/Light.cpp
@@ -19,6 +19,7 @@ ALight::ALight(const FObjectInitializer& ObjectInitializer)
 
 	bHidden = true;
 	bCollideWhenPlacing = true;
+	SpawnCollisionHandlingMethod = ESpawnActorCollisionHandlingMethod::AdjustIfPossibleButDontSpawnIfColliding;
 }
 
 /** The quality level to use for half-resolution lightmaps (not exposed)		*/
diff --git a/Engine/Source/Runtime/Engine/Private/NavigationObjectBase.cpp b/Engine/Source/Runtime/Engine/Private/NavigationObjectBase.cpp
--- a/Engine/Source/Runtime/Engine/Private/NavigationObjectBase.cpp
+++ b/Engine/Source/Runtime/Engine/Private/NavigationObjectBase.cpp
@@ -20,6 +20,7 @@ ANavigationObjectBase::ANavigationObjectBase(const FObjectInitializer& ObjectIni
 	CapsuleComponent->Mobility = EComponentMobility::Static;
 	RootComponent = CapsuleComponent;
 	bCollideWhenPlacing = true;
+	SpawnCollisionHandlingMethod = ESpawnActorCollisionHandlingMethod::AdjustIfPossibleButDontSpawnIfColliding;
 
 	GoodSprite = CreateEditorOnlyDefaultSubobject<UBillboardComponent>(TEXT("Sprite"));
 
diff --git a/Engine/Source/Runtime/Engine/Private/PackageMapClient.cpp b/Engine/Source/Runtime/Engine/Private/PackageMapClient.cpp
--- a/Engine/Source/Runtime/Engine/Private/PackageMapClient.cpp
+++ b/Engine/Source/Runtime/Engine/Private/PackageMapClient.cpp
@@ -338,7 +338,7 @@ bool UPackageMapClient::SerializeNewActor(FArchive& Ar, class UActorChannel *Cha
 				{
 					FActorSpawnParameters SpawnInfo;
 					SpawnInfo.Template = Cast<AActor>(Archetype);
-					SpawnInfo.bNoCollisionFail = true;
+					SpawnInfo.SpawnCollisionHandlingOverride = ESpawnActorCollisionHandlingMethod::AlwaysSpawn;
 					SpawnInfo.bRemoteOwned = true;
 					SpawnInfo.bNoFail = true;
 					Actor = Connection->Driver->GetWorld()->SpawnActor(Archetype->GetClass(), &Location, &Rotation, SpawnInfo );
diff --git a/Engine/Source/Runtime/Engine/Private/Pawn.cpp b/Engine/Source/Runtime/Engine/Private/Pawn.cpp
--- a/Engine/Source/Runtime/Engine/Private/Pawn.cpp
+++ b/Engine/Source/Runtime/Engine/Private/Pawn.cpp
@@ -56,6 +56,7 @@ APawn::APawn(const FObjectInitializer& ObjectInitializer)
 	BaseEyeHeight = 64.0f;
 	AllowedYawError = 10.99f;
 	bCollideWhenPlacing = true;
+	SpawnCollisionHandlingMethod = ESpawnActorCollisionHandlingMethod::AdjustIfPossibleButDontSpawnIfColliding;
 	bProcessingOutsideWorldBounds = false;
 
 	bUseControllerRotationPitch = false;
@@ -296,7 +297,7 @@ void APawn::SpawnDefaultController()
 	{
 		FActorSpawnParameters SpawnInfo;
 		SpawnInfo.Instigator = Instigator;
-		SpawnInfo.bNoCollisionFail = true;
+		SpawnInfo.SpawnCollisionHandlingOverride = ESpawnActorCollisionHandlingMethod::AlwaysSpawn;
 		SpawnInfo.OverrideLevel = GetLevel();
 		SpawnInfo.ObjectFlags |= RF_Transient;	// We never want to save AI controllers into a map
 		AController* NewController = GetWorld()->SpawnActor<AController>(AIControllerClass, GetActorLocation(), GetActorRotation(), SpawnInfo);
diff --git a/Engine/Source/Runtime/Engine/Private/PlayerCameraManager.cpp b/Engine/Source/Runtime/Engine/Private/PlayerCameraManager.cpp
--- a/Engine/Source/Runtime/Engine/Private/PlayerCameraManager.cpp
+++ b/Engine/Source/Runtime/Engine/Private/PlayerCameraManager.cpp
@@ -761,7 +761,7 @@ void APlayerCameraManager::PostInitializeComponents()
 	FActorSpawnParameters SpawnInfo;
 	SpawnInfo.Owner = this;
 	SpawnInfo.Instigator = Instigator;
-	SpawnInfo.bNoCollisionFail = true;
+	SpawnInfo.SpawnCollisionHandlingOverride = ESpawnActorCollisionHandlingMethod::AlwaysSpawn;
 	SpawnInfo.ObjectFlags |= RF_Transient;	// We never want to save these temp actors into a map
 	AnimCameraActor = GetWorld()->SpawnActor<ACameraActor>(SpawnInfo);
 }
@@ -1149,7 +1149,7 @@ AEmitterCameraLensEffectBase* APlayerCameraManager::AddCameraLensEffect(TSubclas
 			FActorSpawnParameters SpawnInfo;
 			SpawnInfo.Owner = PCOwner->GetViewTarget();
 			SpawnInfo.Instigator = Instigator;
-			SpawnInfo.bNoCollisionFail = true;
+			SpawnInfo.SpawnCollisionHandlingOverride = ESpawnActorCollisionHandlingMethod::AlwaysSpawn;
 			SpawnInfo.ObjectFlags |= RF_Transient;	// We never want to save these into a map
 			LensEffect = GetWorld()->SpawnActor<AEmitterCameraLensEffectBase>(LensEffectEmitterClass, SpawnInfo);
 			if (LensEffect != NULL)
diff --git a/Engine/Source/Runtime/Engine/Private/PlayerController.cpp b/Engine/Source/Runtime/Engine/Private/PlayerController.cpp
--- a/Engine/Source/Runtime/Engine/Private/PlayerController.cpp
+++ b/Engine/Source/Runtime/Engine/Private/PlayerController.cpp
@@ -3919,7 +3919,7 @@ void APlayerController::SetSpectatorPawn(class ASpectatorPawn* NewSpectatorPawn)
 		SpectatorPawn = NewSpectatorPawn;
 		AttachToPawn(SpectatorPawn);
 		AddPawnTickDependency(SpectatorPawn);
-
+		
 		if (NewSpectatorPawn)
 		{
 			AutoManageActiveCameraTarget(NewSpectatorPawn);
@@ -3943,7 +3943,8 @@ ASpectatorPawn* APlayerController::SpawnSpectatorPawn()
 		{
 			FActorSpawnParameters SpawnParams;
 			SpawnParams.Owner = this;
-			SpawnParams.bNoCollisionFail = true;
+			SpawnParams.SpawnCollisionHandlingOverride = ESpawnActorCollisionHandlingMethod::AlwaysSpawn;
+
 			SpawnParams.ObjectFlags |= RF_Transient;	// We never want to save spectator pawns into a map
 			SpawnedSpectator = GetWorld()->SpawnActor<ASpectatorPawn>(GameState->SpectatorClass, GetSpawnLocation(), GetControlRotation(), SpawnParams);
 			if (SpawnedSpectator)
diff --git a/Engine/Source/Runtime/Engine/Private/PlayerState.cpp b/Engine/Source/Runtime/Engine/Private/PlayerState.cpp
--- a/Engine/Source/Runtime/Engine/Private/PlayerState.cpp
+++ b/Engine/Source/Runtime/Engine/Private/PlayerState.cpp
@@ -287,7 +287,7 @@ APlayerState* APlayerState::Duplicate()
 {
 	FActorSpawnParameters SpawnInfo;
 	SpawnInfo.Instigator = Instigator;
-	SpawnInfo.bNoCollisionFail = true;
+	SpawnInfo.SpawnCollisionHandlingOverride = ESpawnActorCollisionHandlingMethod::AlwaysSpawn;
 	SpawnInfo.ObjectFlags |= RF_Transient;	// We never want to save player states into a map
 	APlayerState* NewPlayerState = GetWorld()->SpawnActor<APlayerState>(GetClass(), SpawnInfo );
 	// Can fail in case of multiplayer PIE teardown
diff --git a/Engine/Source/Runtime/Engine/Private/SpectatorPawn.cpp b/Engine/Source/Runtime/Engine/Private/SpectatorPawn.cpp
--- a/Engine/Source/Runtime/Engine/Private/SpectatorPawn.cpp
+++ b/Engine/Source/Runtime/Engine/Private/SpectatorPawn.cpp
@@ -18,6 +18,7 @@ ASpectatorPawn::ASpectatorPawn(const FObjectInitializer& ObjectInitializer)
 
 	BaseEyeHeight = 0.0f;
 	bCollideWhenPlacing = false;
+	SpawnCollisionHandlingMethod = ESpawnActorCollisionHandlingMethod::AlwaysSpawn;
 
 	static FName CollisionProfileName(TEXT("Spectator"));
 	GetCollisionComponent()->SetCollisionProfileName(CollisionProfileName);
diff --git a/Engine/Source/Runtime/Engine/Private/TriggerActors.cpp b/Engine/Source/Runtime/Engine/Private/TriggerActors.cpp
--- a/Engine/Source/Runtime/Engine/Private/TriggerActors.cpp
+++ b/Engine/Source/Runtime/Engine/Private/TriggerActors.cpp
@@ -26,6 +26,7 @@ ATriggerCapsule::ATriggerCapsule(const FObjectInitializer& ObjectInitializer)
 	CapsuleCollisionComponent->SetCollisionProfileName(CollisionProfileName);
 
 	bCollideWhenPlacing = true;
+	SpawnCollisionHandlingMethod = ESpawnActorCollisionHandlingMethod::AdjustIfPossibleButDontSpawnIfColliding;
 
 	if (GetSpriteComponent())
 	{
diff --git a/Engine/Source/Runtime/Engine/Private/World.cpp b/Engine/Source/Runtime/Engine/Private/World.cpp
--- a/Engine/Source/Runtime/Engine/Private/World.cpp
+++ b/Engine/Source/Runtime/Engine/Private/World.cpp
@@ -813,9 +813,9 @@ void UWorld::RepairWorldSettings()
 		}
 		
 		FActorSpawnParameters SpawnInfo;
-		SpawnInfo.bNoCollisionFail = true;
+		SpawnInfo.SpawnCollisionHandlingOverride = ESpawnActorCollisionHandlingMethod::AlwaysSpawn;
 		SpawnInfo.Name = GEngine->WorldSettingsClass->GetFName();
-		AWorldSettings* NewWorldSettings = SpawnActor<AWorldSettings>( GEngine->WorldSettingsClass, SpawnInfo );
+		AWorldSettings* const NewWorldSettings = SpawnActor<AWorldSettings>( GEngine->WorldSettingsClass, SpawnInfo );
 
 		const int32 NewWorldSettingsActorIndex = PersistentLevel->Actors.Find( NewWorldSettings );
 
@@ -1047,7 +1047,7 @@ void UWorld::InitializeNewWorld(const InitializationValues IVS)
 
 	// Create the WorldInfo actor.
 	FActorSpawnParameters SpawnInfo;
-	SpawnInfo.bNoCollisionFail = true;
+	SpawnInfo.SpawnCollisionHandlingOverride = ESpawnActorCollisionHandlingMethod::AlwaysSpawn;
 	// Set constant name for WorldSettings to make a network replication work between new worlds on host and client
 	SpawnInfo.Name = GEngine->WorldSettingsClass->GetFName();
 	AActor* WorldSettings = SpawnActor( GEngine->WorldSettingsClass, NULL, NULL, SpawnInfo );
@@ -2879,7 +2879,7 @@ bool UWorld::SetGameMode(const FURL& InURL)
 		// Spawn the GameMode.
 		UE_LOG(LogWorld, Log,  TEXT("Game class is '%s'"), *GameClass->GetName() );
 		FActorSpawnParameters SpawnInfo;
-		SpawnInfo.bNoCollisionFail = true;
+		SpawnInfo.SpawnCollisionHandlingOverride = ESpawnActorCollisionHandlingMethod::AlwaysSpawn;
 		SpawnInfo.ObjectFlags |= RF_Transient;	// We never want to save game modes into a map
 		AuthorityGameMode = SpawnActor<AGameMode>( GameClass, SpawnInfo );
 		
diff --git a/Engine/Source/Runtime/Foliage/Private/InteractiveFoliageActor.cpp b/Engine/Source/Runtime/Foliage/Private/InteractiveFoliageActor.cpp
--- a/Engine/Source/Runtime/Foliage/Private/InteractiveFoliageActor.cpp
+++ b/Engine/Source/Runtime/Foliage/Private/InteractiveFoliageActor.cpp
@@ -25,6 +25,7 @@ AInteractiveFoliageActor::AInteractiveFoliageActor(const FObjectInitializer& Obj
 	PrimaryActorTick.bCanEverTick = true;
 	bCanBeDamaged = true;
 	bCollideWhenPlacing = true;
+	SpawnCollisionHandlingMethod = ESpawnActorCollisionHandlingMethod::AdjustIfPossibleButDontSpawnIfColliding;
 	FoliageDamageImpulseScale = 20.0f;
 	FoliageTouchImpulseScale = 10.0f;
 	FoliageStiffness = 10.0f;
diff --git a/Engine/Source/Runtime/GameplayAbilities/Private/Abilities/Tasks/AbilityTask_SpawnActor.cpp b/Engine/Source/Runtime/GameplayAbilities/Private/Abilities/Tasks/AbilityTask_SpawnActor.cpp
--- a/Engine/Source/Runtime/GameplayAbilities/Private/Abilities/Tasks/AbilityTask_SpawnActor.cpp
+++ b/Engine/Source/Runtime/GameplayAbilities/Private/Abilities/Tasks/AbilityTask_SpawnActor.cpp
@@ -25,8 +25,11 @@ bool UAbilityTask_SpawnActor::BeginSpawningActor(UObject* WorldContextObject, FG
 {
 	if (Ability.IsValid() && Ability.Get()->GetCurrentActorInfo()->IsNetAuthority())
 	{
-		UWorld* World = GEngine->GetWorldFromContextObject(WorldContextObject);
-		SpawnedActor = World->SpawnActorDeferred<AActor>(InClass, FTransform::Identity, NULL, NULL, true);
+		UWorld* const World = GEngine->GetWorldFromContextObject(WorldContextObject);
+		if (World)
+		{
+			SpawnedActor = World->SpawnActorDeferred<AActor>(InClass, FTransform::Identity, NULL, NULL, ESpawnActorCollisionHandlingMethod::AlwaysSpawn);
+		}
 	}
 	
 	if (SpawnedActor == nullptr)
diff --git a/Engine/Source/Runtime/GameplayAbilities/Private/Abilities/Tasks/AbilityTask_VisualizeTargeting.cpp b/Engine/Source/Runtime/GameplayAbilities/Private/Abilities/Tasks/AbilityTask_VisualizeTargeting.cpp
--- a/Engine/Source/Runtime/GameplayAbilities/Private/Abilities/Tasks/AbilityTask_VisualizeTargeting.cpp
+++ b/Engine/Source/Runtime/GameplayAbilities/Private/Abilities/Tasks/AbilityTask_VisualizeTargeting.cpp
@@ -72,8 +72,11 @@ bool UAbilityTask_VisualizeTargeting::BeginSpawningActor(UObject* WorldContextOb
 			UClass* Class = *InTargetClass;
 			if (Class != NULL)
 			{
-				UWorld* World = GEngine->GetWorldFromContextObject(WorldContextObject);
-				SpawnedActor = World->SpawnActorDeferred<AGameplayAbilityTargetActor>(Class, FTransform::Identity, NULL, NULL, true);
+				UWorld* const World = GEngine->GetWorldFromContextObject(WorldContextObject);
+				if (World)
+				{
+					SpawnedActor = World->SpawnActorDeferred<AGameplayAbilityTargetActor>(Class, FTransform::Identity, NULL, NULL, ESpawnActorCollisionHandlingMethod::AlwaysSpawn);
+				}
 			}
 
 			if (SpawnedActor)
diff --git a/Engine/Source/Runtime/GameplayAbilities/Private/Abilities/Tasks/AbilityTask_WaitTargetData.cpp b/Engine/Source/Runtime/GameplayAbilities/Private/Abilities/Tasks/AbilityTask_WaitTargetData.cpp
--- a/Engine/Source/Runtime/GameplayAbilities/Private/Abilities/Tasks/AbilityTask_WaitTargetData.cpp
+++ b/Engine/Source/Runtime/GameplayAbilities/Private/Abilities/Tasks/AbilityTask_WaitTargetData.cpp
@@ -82,7 +82,7 @@ bool UAbilityTask_WaitTargetData::BeginSpawningActor(UObject* WorldContextObject
 			if (Class != NULL)
 			{
 				UWorld* World = GEngine->GetWorldFromContextObject(WorldContextObject);
-				SpawnedActor = World->SpawnActorDeferred<AGameplayAbilityTargetActor>(Class, FTransform::Identity, nullptr, nullptr, true);
+				SpawnedActor = World->SpawnActorDeferred<AGameplayAbilityTargetActor>(Class, FTransform::Identity, nullptr, nullptr, ESpawnActorCollisionHandlingMethod::AlwaysSpawn);
 			}
 
 			if (SpawnedActor)
diff --git a/Engine/Source/Runtime/GameplayTasks/Private/Tasks/GameplayTask_SpawnActor.cpp b/Engine/Source/Runtime/GameplayTasks/Private/Tasks/GameplayTask_SpawnActor.cpp
--- a/Engine/Source/Runtime/GameplayTasks/Private/Tasks/GameplayTask_SpawnActor.cpp
+++ b/Engine/Source/Runtime/GameplayTasks/Private/Tasks/GameplayTask_SpawnActor.cpp
@@ -40,8 +40,11 @@ UGameplayTask_SpawnActor* UGameplayTask_SpawnActor::SpawnActor(TScriptInterface<
 
 bool UGameplayTask_SpawnActor::BeginSpawningActor(UObject* WorldContextObject, AActor*& SpawnedActor)
 {
-	UWorld* World = GEngine->GetWorldFromContextObject(WorldContextObject);
-	SpawnedActor = World->SpawnActorDeferred<AActor>(ClassToSpawn, FTransform(CachedSpawnRotation, CachedSpawnLocation), NULL, NULL, true);
+	UWorld* const World = GEngine->GetWorldFromContextObject(WorldContextObject);
+	if (World)
+	{
+		SpawnedActor = World->SpawnActorDeferred<AActor>(ClassToSpawn, FTransform(CachedSpawnRotation, CachedSpawnLocation), NULL, NULL, ESpawnActorCollisionHandlingMethod::AlwaysSpawn);
+	}
 	
 	if (SpawnedActor == nullptr)
 	{
diff --git a/Engine/Source/Runtime/UMG/Private/Components/Viewport.cpp b/Engine/Source/Runtime/UMG/Private/Components/Viewport.cpp
--- a/Engine/Source/Runtime/UMG/Private/Components/Viewport.cpp
+++ b/Engine/Source/Runtime/UMG/Private/Components/Viewport.cpp
@@ -471,7 +471,7 @@ AActor* UViewport::Spawn(TSubclassOf<AActor> ActorClass)
 	{
 		UWorld* World = GetViewportWorld();
 		FActorSpawnParameters SpawnParameters;
-		SpawnParameters.bNoCollisionFail = true;
+		SpawnParameters.SpawnCollisionHandlingOverride = ESpawnActorCollisionHandlingMethod::AlwaysSpawn;
 
 		return World->SpawnActor<AActor>(ActorClass, FVector(0, 0, 0), FRotator(), SpawnParameters);
 	}
