When clearing a PlayerController's spectator pawn, use regular pawn for attachment and auto-managed camera target if regular pawn has already been set up. (UE-15786)

diff --git a/Engine/Source/Runtime/Engine/Private/PlayerController.cpp b/Engine/Source/Runtime/Engine/Private/PlayerController.cpp
--- a/Engine/Source/Runtime/Engine/Private/PlayerController.cpp
+++ b/Engine/Source/Runtime/Engine/Private/PlayerController.cpp
@@ -3917,16 +3917,28 @@ void APlayerController::SetSpectatorPawn(class ASpectatorPawn* NewSpectatorPawn)
 	{
 		RemovePawnTickDependency(SpectatorPawn);
 		SpectatorPawn = NewSpectatorPawn;
-		AttachToPawn(SpectatorPawn);
-		AddPawnTickDependency(SpectatorPawn);
 		
 		if (NewSpectatorPawn)
 		{
+			// setting to a new valid spectator pawn
+			AttachToPawn(NewSpectatorPawn);
+			AddPawnTickDependency(NewSpectatorPawn);
 			AutoManageActiveCameraTarget(NewSpectatorPawn);
 		}
 		else
 		{
-			AutoManageActiveCameraTarget(this);
+			// clearing the spectator pawn, try to attach to the regular pawn
+			APawn* const Pawn = GetPawn();
+			AttachToPawn(Pawn);
+			AddPawnTickDependency(Pawn);
+			if (Pawn)
+			{
+				AutoManageActiveCameraTarget(Pawn);
+			}
+			else
+			{
+				AutoManageActiveCameraTarget(this);
+			}
 		}
 	}
 }
