Fix crash removing cooling down tick functions

diff --git a/Engine/Source/Runtime/Engine/Private/TickTaskManager.cpp b/Engine/Source/Runtime/Engine/Private/TickTaskManager.cpp
--- a/Engine/Source/Runtime/Engine/Private/TickTaskManager.cpp
+++ b/Engine/Source/Runtime/Engine/Private/TickTaskManager.cpp
@@ -935,7 +935,7 @@ public:
 			FTickFunction* PrevComparisionFunction = nullptr;
 			FTickFunction* ComparisonFunction = AllCoolingDownTickFunctions.Head;
 			bool bFound = false;
-			while (ComparisonFunction)
+			while (ComparisonFunction && !bFound)
 			{
 				if (ComparisonFunction == TickFunction)
 				{
