Fixes Camera Position material expression when computing the previous world position offset

diff --git a/Engine/Source/Editor/UnrealEd/Private/Lightmass/LightmassRender.cpp b/Engine/Source/Editor/UnrealEd/Private/Lightmass/LightmassRender.cpp
--- a/Engine/Source/Editor/UnrealEd/Private/Lightmass/LightmassRender.cpp
+++ b/Engine/Source/Editor/UnrealEd/Private/Lightmass/LightmassRender.cpp
@@ -98,12 +98,6 @@ struct FLightmassMaterialCompiler : public FProxyMaterialCompiler
 		return Compiler->Constant3(0.0f,0.0f,0.0f);
 	}
 
-	virtual int32 CameraWorldPosition() override
-	{
-		//UE_LOG(LogLightmassRender, Log, TEXT("Lightmass material compiler has encountered CameraWorldPosition... Forcing constant (0.0f,0.0f,0.0f)."));
-		return Compiler->Constant3(0.0f,0.0f,0.0f);
-	}
-
 	virtual int32 CameraVector() override
 	{
 		//UE_LOG(LogLightmassRender, Log, TEXT("Lightmass material compiler has encountered CameraVector... Forcing constant (0.0f,0.0f,1.0f)."));
diff --git a/Engine/Source/Runtime/Engine/Classes/Materials/MaterialExpressionViewProperty.h b/Engine/Source/Runtime/Engine/Classes/Materials/MaterialExpressionViewProperty.h
--- a/Engine/Source/Runtime/Engine/Classes/Materials/MaterialExpressionViewProperty.h
+++ b/Engine/Source/Runtime/Engine/Classes/Materials/MaterialExpressionViewProperty.h
@@ -16,6 +16,8 @@ enum EMaterialExposedViewProperty
 	MEVP_TanHalfFieldOfView UMETA(DisplayName="TanHalfFieldOfView"),
 	/** Horizontal and vertical size of the view in pixels */
 	MEVP_ViewSize UMETA(DisplayName="ViewSize"),
+	/** Absolute world space camera position */
+	MEVP_WorldSpaceCameraPosition UMETA(DisplayName="CameraPosition (Absolute World Space)"),
 
 	MEVP_MAX,
 };
diff --git a/Engine/Source/Runtime/Engine/Private/Materials/HLSLMaterialTranslator.h b/Engine/Source/Runtime/Engine/Private/Materials/HLSLMaterialTranslator.h
--- a/Engine/Source/Runtime/Engine/Private/Materials/HLSLMaterialTranslator.h
+++ b/Engine/Source/Runtime/Engine/Private/Materials/HLSLMaterialTranslator.h
@@ -1660,6 +1660,7 @@ protected:
 			{MEVP_FieldOfView, MCT_Float2, TEXT("View.<PREV>FieldOfViewWideAngles"), nullptr},
 			{MEVP_TanHalfFieldOfView, MCT_Float2, TEXT("Get<PREV>TanHalfFieldOfView()"), TEXT("Get<PREV>CotanHalfFieldOfView()")},
 			{MEVP_ViewSize, MCT_Float2, TEXT("View.ViewSizeAndInvSize.xy"), TEXT("View.ViewSizeAndInvSize.zw")},
+			{MEVP_WorldSpaceCameraPosition, MCT_Float3, TEXT("View.<PREV>ViewOrigin"), nullptr},
 		};
 		static_assert((sizeof(ViewPropertyMetaArray) / sizeof(ViewPropertyMetaArray[0])) == MEVP_MAX, "incoherency between EMaterialExposedViewProperty and ViewPropertyMetaArray");
 
@@ -1922,11 +1923,6 @@ protected:
 		return AddInlinedCodeChunk(MCT_Float3,TEXT("Parameters.CameraVector"));
 	}
 
-	virtual int32 CameraWorldPosition() override
-	{
-		return AddInlinedCodeChunk(MCT_Float3,TEXT("View.ViewOrigin.xyz"));
-	}
-
 	virtual int32 LightVector() override
 	{
 		if (ShaderFrequency != SF_Pixel && ShaderFrequency != SF_Compute)
diff --git a/Engine/Source/Runtime/Engine/Private/Materials/MaterialExpressions.cpp b/Engine/Source/Runtime/Engine/Private/Materials/MaterialExpressions.cpp
--- a/Engine/Source/Runtime/Engine/Private/Materials/MaterialExpressions.cpp
+++ b/Engine/Source/Runtime/Engine/Private/Materials/MaterialExpressions.cpp
@@ -2893,7 +2893,7 @@ UMaterialExpressionCameraPositionWS::UMaterialExpressionCameraPositionWS(const F
 
 int32 UMaterialExpressionCameraPositionWS::Compile(class FMaterialCompiler* Compiler, int32 OutputIndex, int32 MultiplexIndex)
 {
-	return Compiler->CameraWorldPosition();
+	return Compiler->ViewProperty(MEVP_WorldSpaceCameraPosition);
 }
 
 void UMaterialExpressionCameraPositionWS::GetCaption(TArray<FString>& OutCaptions) const
diff --git a/Engine/Source/Runtime/Engine/Public/MaterialCompiler.h b/Engine/Source/Runtime/Engine/Public/MaterialCompiler.h
--- a/Engine/Source/Runtime/Engine/Public/MaterialCompiler.h
+++ b/Engine/Source/Runtime/Engine/Public/MaterialCompiler.h
@@ -76,7 +76,6 @@ public:
 	virtual int32 ReflectionVector() = 0;
 	virtual int32 ReflectionAboutCustomWorldNormal(int32 CustomWorldNormal, int32 bNormalizeCustomWorldNormal) = 0;
 	virtual int32 CameraVector() = 0;
-	virtual int32 CameraWorldPosition() = 0;
 	virtual int32 LightVector() = 0;
 
 	virtual int32 ScreenPosition() = 0;
@@ -242,7 +241,6 @@ public:
 
 	virtual int32 ReflectionVector() override { return Compiler->ReflectionVector(); }
 	virtual int32 CameraVector() override { return Compiler->CameraVector(); }
-	virtual int32 CameraWorldPosition() override { return Compiler->CameraWorldPosition(); }
 	virtual int32 LightVector() override { return Compiler->LightVector(); }
 
 	virtual int32 ScreenPosition() override { return Compiler->ScreenPosition(); }
