Provide default implementations for members so that the base type isn't abstract and can be used in arrays. (VS2015)

diff --git a/Engine/Source/Runtime/Engine/Public/GPUSkinVertexFactory.h b/Engine/Source/Runtime/Engine/Public/GPUSkinVertexFactory.h
--- a/Engine/Source/Runtime/Engine/Public/GPUSkinVertexFactory.h
+++ b/Engine/Source/Runtime/Engine/Public/GPUSkinVertexFactory.h
@@ -448,7 +448,7 @@ private:
 		return ShaderData;
 	}
 
-	virtual bool UsesExtraBoneInfluences() const = 0;
+	virtual bool UsesExtraBoneInfluences() const { return false; }
 
 	/**
 	 * Set the data with the given frame number, keeps data for the last frame.
diff --git a/Engine/Source/Runtime/ShaderCore/Public/VertexFactory.h b/Engine/Source/Runtime/ShaderCore/Public/VertexFactory.h
--- a/Engine/Source/Runtime/ShaderCore/Public/VertexFactory.h
+++ b/Engine/Source/Runtime/ShaderCore/Public/VertexFactory.h
@@ -353,7 +353,7 @@ public:
 	/**
 	 * @return The vertex factory's type.
 	 */
-	virtual FVertexFactoryType* GetType() const = 0;
+	virtual FVertexFactoryType* GetType() const { return NULL; }
 
 	/**
 	 * Activates the vertex factory.
