#ue4 - Various minor fixes and improvements to CharacterMovementComponent.

- Added "virtual" modifier to some existing functions.
- Improved "const" correctness at several places.
- Removed useless whitespace.
- Improved standing downward force code and moved it to the new overridable function ApplyDownwardForce().
- Moved calculation of root motion velocity code to the new overridable function CalcRootMotionVelocity().
- Minor improvements in functions GetLedgeMove() / MoveAlongFloor() / PhysWalking() / StepUp() / ApplyRepulsionForce() / ...

Git #954 - author Xaklse
UE-12297

diff --git a/Engine/Source/Runtime/Engine/Classes/GameFramework/CharacterMovementComponent.h b/Engine/Source/Runtime/Engine/Classes/GameFramework/CharacterMovementComponent.h
--- a/Engine/Source/Runtime/Engine/Classes/GameFramework/CharacterMovementComponent.h
+++ b/Engine/Source/Runtime/Engine/Classes/GameFramework/CharacterMovementComponent.h
@@ -956,7 +956,7 @@ public:
 	virtual void UpdateBasedMovement(float DeltaSeconds);
 
 	/** Update controller's view rotation as pawn's base rotates */
-	virtual void UpdateBasedRotation(FRotator &FinalRotation, const FRotator& ReducedRotation);
+	virtual void UpdateBasedRotation(FRotator& FinalRotation, const FRotator& ReducedRotation);
 
 	/** Update (or defer updating) OldBaseLocation and OldBaseQuat if there is a valid movement base. */
 	DEPRECATED(4.4, "CharacterMovementComponent::MaybeSaveBaseLocation() will be removed, call SaveBaseLocation().")
@@ -1006,7 +1006,7 @@ public:
 	virtual void JumpOutOfWater(FVector WallNormal);
 
 	/** @return how far to rotate character during the time interval DeltaTime. */
-	virtual FRotator GetDeltaRotation(float DeltaTime);
+	virtual FRotator GetDeltaRotation(float DeltaTime) const;
 
 	/**
 	  * Compute a target rotation based on current movement. Used by PhysicsRotation() when bOrientRotationToMovement is true.
@@ -1018,7 +1018,7 @@ public:
 	  *
 	  * @return The target rotation given current movement.
 	  */
-	virtual FRotator ComputeOrientToMovementRotation(const FRotator& CurrentRotation, float DeltaTime, FRotator& DeltaRotation);
+	virtual FRotator ComputeOrientToMovementRotation(const FRotator& CurrentRotation, float DeltaTime, FRotator& DeltaRotation) const;
 
 	/**
 	 * Use velocity requested by path following to compute a requested acceleration and speed.
@@ -1047,7 +1047,7 @@ public:
 	/* Determine how deep in water the character is immersed.
 	 * @return float in range 0.0 = not in water, 1.0 = fully immersed
 	 */
-	virtual float ImmersionDepth();
+	virtual float ImmersionDepth() const;
 
 	/** 
 	 * Updates Velocity and Acceleration based on the current state, applying the effects of friction and acceleration or deceleration. Does not apply gravity.
@@ -1103,7 +1103,7 @@ public:
 	/** 
 	 * Move up steps or slope. Does nothing and returns false if CanStepUp(Hit) returns false.
 	 *
-	 * @param GravDir			Gravity vector
+	 * @param GravDir			Gravity vector direction (assumed normalized or zero)
 	 * @param Delta				Requested move
 	 * @param Hit				[In] The hit before the step up.
 	 * @param OutStepDownResult	[Out] If non-null, a floor check will be performed if possible as part of the final step down, and it will be updated to reflect this result.
@@ -1119,6 +1119,12 @@ public:
 	 */
 	void SetBaseFromFloor(const FFindFloorResult& FloorResult);
 
+	/**
+	 * Applies downward force when walking on top of physics objects.
+	 * @param DeltaSeconds Time elapsed since last frame.
+	 */
+	virtual void ApplyDownwardForce(float DeltaSeconds);
+
 	/** Applies repulsion force to all touched components. */
 	virtual void ApplyRepulsionForce(float DeltaSeconds);
 	
@@ -1136,7 +1142,7 @@ public:
 	void StartSwimming(FVector OldLocation, FVector OldVelocity, float timeTick, float remainingTime, int32 Iterations);
 
 	/* Swimming uses gravity - but scaled by (1.f - buoyancy) */
-	float Swim(FVector Delta, FHitResult &Hit);
+	float Swim(FVector Delta, FHitResult& Hit);
 
 	/** Get as close to waterline as possible, staying on same side as currently. */
 	FVector FindWaterLine(FVector Start, FVector End);
@@ -1185,8 +1191,7 @@ protected:
 	virtual float BoostAirControl(float DeltaTime, float TickAirControl, const FVector& FallAcceleration);
 
 	/**
-	 * Checks if air control will cause the player collision shape to hit something given the current location.
-	 * This function is used internally by PhysFalling().
+	 * (DEPRECATED) Checks if air control will cause the player collision shape to hit something given the current location.
 	 *
 	 * @param DeltaTime			Time step for the current update.
 	 * @param AdditionalTime	Time to look ahead further, applying acceleration and gravity.
@@ -1197,10 +1202,11 @@ protected:
 	 * @return True if there is an impact, in which case OutHitResult contains the result of that impact.
 	 * @see GetAirControl()
 	 */
+	DEPRECATED(4.9, "FindAirControlImpact is no longer used by engine code.")
 	virtual bool FindAirControlImpact(float DeltaTime, float AdditionalTime, const FVector& FallVelocity, const FVector& FallAcceleration, const FVector& Gravity, FHitResult& OutHitResult);
 
 	/**
-	 * Limits the air control to use during falling movement, given an impact from FindAirControlImpact().
+	 * Limits the air control to use during falling movement, given an impact while falling.
 	 * This function is used internally by PhysFalling().
 	 *
 	 * @param DeltaTime			Time step for the current update.
@@ -1208,7 +1214,7 @@ protected:
 	 * @param HitResult			Result of impact.
 	 * @param bCheckForValidLandingSpot If true, will use IsValidLandingSpot() to determine if HitResult is a walkable surface. If false, this check is skipped.
 	 * @return Modified air control acceleration to use during falling movement.
-	 * @see FindAirControlImpact()
+	 * @see PhysFalling()
 	 */
 	virtual FVector LimitAirControl(float DeltaTime, const FVector& FallAcceleration, const FHitResult& HitResult, bool bCheckForValidLandingSpot);
 	
@@ -1262,13 +1268,13 @@ public:
 	virtual bool CanCrouchInCurrentState() const;
 	
 	/** @return true if there is a suitable floor SideStep from current position. */
-	bool CheckLedgeDirection(const FVector& OldLocation, const FVector& SideStep, const FVector& GravDir);
+	virtual bool CheckLedgeDirection(const FVector& OldLocation, const FVector& SideStep, const FVector& GravDir) const;
 
 	/** 
 	 * @param Delta is the current move delta (which ended up going over a ledge).
 	 * @return new delta which moves along the ledge
 	 */
-	FVector GetLedgeMove(const FVector& OldLocation, const FVector& Delta, const FVector& GravDir);
+	virtual FVector GetLedgeMove(const FVector& OldLocation, const FVector& Delta, const FVector& GravDir) const;
 
 	/** Check if pawn is falling */
 	virtual bool CheckFall(const FFindFloorResult& OldFloor, const FHitResult& Hit, const FVector& Delta, const FVector& OldLocation, float remainingTime, float timeTick, int32 Iterations, bool bMustJump);
@@ -1445,7 +1451,7 @@ protected:
 	virtual bool ResolvePenetrationImpl(const FVector& Adjustment, const FHitResult& Hit, const FQuat& NewRotation) override;
 
 	/** Handle a blocking impact. Calls ApplyImpactPhysicsForces for the hit, if bEnablePhysicsInteraction is true. */
-	virtual void HandleImpact(FHitResult const& Hit, float TimeSlice=0.f, const FVector& MoveDelta = FVector::ZeroVector) override;
+	virtual void HandleImpact(const FHitResult& Hit, float TimeSlice=0.f, const FVector& MoveDelta = FVector::ZeroVector) override;
 
 	/**
 	 * Apply physics forces to the impacted component, if bEnablePhysicsInteraction is true.
@@ -1456,10 +1462,10 @@ protected:
 	virtual void ApplyImpactPhysicsForces(const FHitResult& Impact, const FVector& ImpactAcceleration, const FVector& ImpactVelocity);
 
 	/** Custom version of SlideAlongSurface that handles different movement modes separately; namely during walking physics we might not want to slide up slopes. */
-	virtual float SlideAlongSurface(const FVector& Delta, float Time, const FVector& Normal, FHitResult &Hit, bool bHandleImpact) override;
+	virtual float SlideAlongSurface(const FVector& Delta, float Time, const FVector& Normal, FHitResult& Hit, bool bHandleImpact) override;
 
 	/** Custom version that allows upwards slides when walking if the surface is walkable. */
-	virtual void TwoWallAdjust(FVector &Delta, const FHitResult& Hit, const FVector &OldHitNormal) const override;
+	virtual void TwoWallAdjust(FVector& Delta, const FHitResult& Hit, const FVector& OldHitNormal) const override;
 
 	/**
 	 * Calculate slide vector along a surface.
@@ -1497,7 +1503,7 @@ protected:
 	 * Return true if the 2D distance to the impact point is inside the edge tolerance (CapsuleRadius minus a small rejection threshold).
 	 * Useful for rejecting adjacent hits when finding a floor or landing spot.
 	 */
-	bool IsWithinEdgeTolerance(const FVector& CapsuleLocation, const FVector& TestImpactPoint, const float CapsuleRadius) const;
+	virtual bool IsWithinEdgeTolerance(const FVector& CapsuleLocation, const FVector& TestImpactPoint, const float CapsuleRadius) const;
 
 	/**
 	 * Sweeps a vertical trace to find the floor for the capsule at the given location. Will attempt to perch if ShouldComputePerchResult() returns true for the downward sweep result.
@@ -1581,7 +1587,7 @@ protected:
 
 	/** Called when the collision capsule touches another primitive component */
 	UFUNCTION()
-	void CapsuleTouched(AActor* Other, UPrimitiveComponent* OtherComp, int32 OtherBodyIndex, bool bFromSweep, const FHitResult& SweepResult);
+	virtual void CapsuleTouched(AActor* Other, UPrimitiveComponent* OtherComp, int32 OtherBodyIndex, bool bFromSweep, const FHitResult& SweepResult);
 
 	// Enum used to control GetPawnCapsuleExtent behavior
 	enum EShrinkCapsuleExtent
@@ -1849,6 +1855,14 @@ public:
 	/** Simulate Root Motion physics on Simulated Proxies */
 	void SimulateRootMotion(float DeltaSeconds, const FTransform& LocalRootMotionTransform);
 
+	/**
+	 * Calculate velocity from root motion. Under some movement conditions, only portions of root motion may be used (e.g. when falling Z may be ignored).
+	 * @param RootMotionDeltaMove	Change in location from root motion.
+	 * @param DeltaSeconds			Elapsed time
+	 * @param CurrentVelocity		Non-root motion velocity at current time, used for components of result that may ignore root motion.
+	 */
+	virtual FVector CalcRootMotionVelocity(const FVector& RootMotionDeltaMove, float DeltaSeconds, const FVector& CurrentVelocity) const;
+
 	// RVO Avoidance
 
 	/** calculate RVO avoidance and apply it to current velocity */
diff --git a/Engine/Source/Runtime/Engine/Private/Components/CharacterMovementComponent.cpp b/Engine/Source/Runtime/Engine/Private/Components/CharacterMovementComponent.cpp
--- a/Engine/Source/Runtime/Engine/Private/Components/CharacterMovementComponent.cpp
+++ b/Engine/Source/Runtime/Engine/Private/Components/CharacterMovementComponent.cpp
@@ -894,15 +894,9 @@ void UCharacterMovementComponent::TickComponent(float DeltaTime, enum ELevelTick
 
 	// See if we fell out of the world.
 	const bool bIsSimulatingPhysics = UpdatedComponent->IsSimulatingPhysics();
-	if (CharacterOwner->Role == ROLE_Authority)
+	if (CharacterOwner->Role == ROLE_Authority && (!bCheatFlying || bIsSimulatingPhysics) && !CharacterOwner->CheckStillInWorld())
 	{
-		if (!bCheatFlying || bIsSimulatingPhysics)
-		{
-			if (!CharacterOwner->CheckStillInWorld())
-			{
-				return;
-			}
-		}
+		return;
 	}
 
 	// We don't update if simulating physics (eg ragdolls).
@@ -969,21 +963,7 @@ void UCharacterMovementComponent::TickComponent(float DeltaTime, enum ELevelTick
 	if (bEnablePhysicsInteraction)
 	{
 		SCOPE_CYCLE_COUNTER(STAT_CharPhysicsInteraction);
-
-		if (CurrentFloor.HitResult.IsValidBlockingHit())
-		{
-			// Apply downwards force when walking on top of physics objects
-			if (UPrimitiveComponent* BaseComp = CurrentFloor.HitResult.GetComponent())
-			{
-				if (StandingDownwardForceScale != 0.f && BaseComp->IsAnySimulatingPhysics())
-				{
-					const float GravZ = GetGravityZ();
-					const FVector ForceLocation = CurrentFloor.HitResult.ImpactPoint;
-					BaseComp->AddForceAtLocation(FVector(0.f, 0.f, GravZ * Mass * StandingDownwardForceScale), ForceLocation, CurrentFloor.HitResult.BoneName);
-				}
-			}
-		}
-
+		ApplyDownwardForce(DeltaTime);
 		ApplyRepulsionForce(DeltaTime);
 	}
 
@@ -1146,9 +1126,7 @@ void UCharacterMovementComponent::SimulateRootMotion(float DeltaSeconds, const F
 		const FTransform WorldSpaceRootMotionTransform = CharacterOwner->GetMesh()->ConvertLocalRootMotionToWorld(LocalRootMotionTransform);
 
 		// Compute root motion velocity to be used by physics
-		const FVector RootMotionVelocity = WorldSpaceRootMotionTransform.GetTranslation() / DeltaSeconds;
-		// Do not override Velocity.Z if falling.
-		Velocity = FVector(RootMotionVelocity.X, RootMotionVelocity.Y, (MovementMode == MOVE_Falling ? Velocity.Z : RootMotionVelocity.Z));
+		Velocity = CalcRootMotionVelocity(WorldSpaceRootMotionTransform.GetTranslation(), DeltaSeconds, Velocity);
 
 		// Update replicated movement mode.
 		if (bNetworkMovementModeChanged)
@@ -1172,6 +1150,20 @@ void UCharacterMovementComponent::SimulateRootMotion(float DeltaSeconds, const F
 }
 
 
+FVector UCharacterMovementComponent::CalcRootMotionVelocity(const FVector& RootMotionDeltaMove, float DeltaSeconds, const FVector& CurrentVelocity) const
+{
+	FVector RootMotionVelocity = RootMotionDeltaMove / DeltaSeconds;
+	
+	// Do not override Velocity.Z if in falling physics, we want to keep the effect of gravity.
+	if (IsFalling())
+	{
+		RootMotionVelocity.Z = CurrentVelocity.Z;
+	}
+	
+	return RootMotionVelocity;
+}
+
+
 void UCharacterMovementComponent::SimulateMovement(float DeltaSeconds)
 {
 	if (!HasValidData() || UpdatedComponent->Mobility != EComponentMobility::Movable || UpdatedComponent->IsSimulatingPhysics())
@@ -1249,14 +1241,13 @@ void UCharacterMovementComponent::SimulateMovement(float DeltaSeconds)
 		const bool bEnableFloorCheck = (!CharacterOwner->bSimGravityDisabled || !bIsSimulatedProxy);
 		if (bEnableFloorCheck && (IsMovingOnGround() || MovementMode == MOVE_Falling))
 		{
-			const FVector CollisionCenter = UpdatedComponent->GetComponentLocation();
 			if (StepDownResult.bComputedFloor)
 			{
 				CurrentFloor = StepDownResult.FloorResult;
 			}
 			else if (Velocity.Z <= 0.f)
 			{
-				FindFloor(CollisionCenter, CurrentFloor, Velocity.IsZero(), NULL);
+				FindFloor(UpdatedComponent->GetComponentLocation(), CurrentFloor, Velocity.IsZero(), NULL);
 			}
 			else
 			{
@@ -1497,7 +1488,7 @@ void UCharacterMovementComponent::OnUnableToFollowBaseMove(const FVector& DeltaP
 }
 
 
-void UCharacterMovementComponent::UpdateBasedRotation(FRotator &FinalRotation, const FRotator& ReducedRotation)
+void UCharacterMovementComponent::UpdateBasedRotation(FRotator& FinalRotation, const FRotator& ReducedRotation)
 {
 	AController* Controller = CharacterOwner ? CharacterOwner->Controller : NULL;
 	float ControllerRoll = 0.f;
@@ -1613,9 +1604,7 @@ void UCharacterMovementComponent::PerformMovement(float DeltaSeconds)
 				// Then turn root motion to velocity to be used by various physics modes.
 				if (DeltaSeconds > 0.f)
 				{
-					const FVector RootMotionVelocity = RootMotionParams.RootMotionTransform.GetTranslation() / DeltaSeconds;
-					// Do not override Velocity.Z if in falling physics, we want to keep the effect of gravity.
-					Velocity = FVector(RootMotionVelocity.X, RootMotionVelocity.Y, (MovementMode == MOVE_Falling ? Velocity.Z : RootMotionVelocity.Z));
+					Velocity = CalcRootMotionVelocity(RootMotionParams.RootMotionTransform.GetTranslation(), DeltaSeconds, Velocity);
 				}
 
 				UE_LOG(LogRootMotion, Log,  TEXT("PerformMovement WorldSpaceRootMotion Translation: %s, Rotation: %s, Actor Facing: %s, Velocity: %s")
@@ -1985,8 +1974,10 @@ void UCharacterMovementComponent::UnCrouch(bool bClientSimulation)
 
 void UCharacterMovementComponent::StartNewPhysics(float deltaTime, int32 Iterations)
 {
-	if ( (deltaTime < MIN_TICK_TIME) || (Iterations >= MaxSimulationIterations) || !HasValidData())
+	if ((deltaTime < MIN_TICK_TIME) || (Iterations >= MaxSimulationIterations) || !HasValidData())
+	{
 		return;
+	}
 
 	if (UpdatedComponent->IsSimulatingPhysics())
 	{
@@ -2067,7 +2058,7 @@ bool UCharacterMovementComponent::ResolvePenetrationImpl(const FVector& Adjustme
 }
 
 
-float UCharacterMovementComponent::SlideAlongSurface(const FVector& Delta, float Time, const FVector& InNormal, FHitResult &Hit, bool bHandleImpact)
+float UCharacterMovementComponent::SlideAlongSurface(const FVector& Delta, float Time, const FVector& InNormal, FHitResult& Hit, bool bHandleImpact)
 {
 	if (!Hit.bBlockingHit)
 	{
@@ -2106,9 +2097,9 @@ float UCharacterMovementComponent::SlideAlongSurface(const FVector& Delta, float
 }
 
 
-void UCharacterMovementComponent::TwoWallAdjust(FVector &Delta, const FHitResult& Hit, const FVector &OldHitNormal) const
+void UCharacterMovementComponent::TwoWallAdjust(FVector& Delta, const FHitResult& Hit, const FVector& OldHitNormal) const
 {
-	FVector InDelta = Delta;
+	const FVector InDelta = Delta;
 	Super::TwoWallAdjust(Delta, Hit, OldHitNormal);
 
 	if (IsMovingOnGround())
@@ -2219,7 +2210,7 @@ FVector UCharacterMovementComponent::NewFallVelocity(const FVector& InitialVeloc
 }
 
 
-float UCharacterMovementComponent::ImmersionDepth()
+float UCharacterMovementComponent::ImmersionDepth() const
 {
 	float depth = 0.f;
 
@@ -2940,7 +2931,7 @@ void UCharacterMovementComponent::StartSwimming(FVector OldLocation, FVector Old
 	}
 }
 
-float UCharacterMovementComponent::Swim(FVector Delta, FHitResult &Hit)
+float UCharacterMovementComponent::Swim(FVector Delta, FHitResult& Hit)
 {
 	FVector Start = UpdatedComponent->GetComponentLocation();
 	float airTime = 0.f;
@@ -3246,7 +3237,7 @@ void UCharacterMovementComponent::PhysFalling(float deltaTime, int32 Iterations)
 						// bDitch=true means that pawn is straddling two slopes, neither of which he can stand on
 						bool bDitch = ( (OldHitImpactNormal.Z > 0.f) && (Hit.ImpactNormal.Z > 0.f) && (FMath::Abs(Delta.Z) <= KINDA_SMALL_NUMBER) && ((Hit.ImpactNormal | OldHitImpactNormal) < 0.f) );
 						SafeMoveUpdatedComponent( Delta, PawnRotation, true, Hit);
-						if ( Hit.Time == 0 )
+						if ( Hit.Time == 0.f )
 						{
 							// if we are stuck then try to side step
 							FVector SideDelta = (OldHitNormal + Hit.ImpactNormal).GetSafeNormal2D();
@@ -3257,7 +3248,7 @@ void UCharacterMovementComponent::PhysFalling(float deltaTime, int32 Iterations)
 							SafeMoveUpdatedComponent( SideDelta, PawnRotation, true, Hit);
 						}
 							
-						if ( bDitch || IsValidLandingSpot(UpdatedComponent->GetComponentLocation(), Hit) || Hit.Time == 0  )
+						if ( bDitch || IsValidLandingSpot(UpdatedComponent->GetComponentLocation(), Hit) || Hit.Time == 0.f  )
 						{
 							remainingTime = 0.f;
 							ProcessLanded(Hit, remainingTime, Iterations);
@@ -3292,6 +3283,7 @@ void UCharacterMovementComponent::PhysFalling(float deltaTime, int32 Iterations)
 }
 
 
+// Note: Deprecated for 4.9
 bool UCharacterMovementComponent::FindAirControlImpact(float DeltaTime, float AdditionalTime, const FVector& FallVelocity, const FVector& FallAcceleration, const FVector& Gravity, FHitResult& OutHitResult)
 {
 	// Test for slope to avoid using air control to climb walls.
@@ -3346,13 +3338,13 @@ FVector UCharacterMovementComponent::LimitAirControl(float DeltaTime, const FVec
 	return Result;
 }
 
-bool UCharacterMovementComponent::CheckLedgeDirection(const FVector& OldLocation, const FVector& SideStep, const FVector& GravDir)
+bool UCharacterMovementComponent::CheckLedgeDirection(const FVector& OldLocation, const FVector& SideStep, const FVector& GravDir) const
 {
-	FVector SideDest = OldLocation + SideStep;
+	const FVector SideDest = OldLocation + SideStep;
 	FCollisionQueryParams CapsuleParams(CharacterMovementComponentStatics::CheckLedgeDirectionName, false, CharacterOwner);
 	FCollisionResponseParams ResponseParam;
 	InitCollisionParams(CapsuleParams, ResponseParam);
-	FCollisionShape CapsuleShape = GetPawnCapsuleCollisionShape(SHRINK_None);
+	const FCollisionShape CapsuleShape = GetPawnCapsuleCollisionShape(SHRINK_None);
 	const ECollisionChannel CollisionChannel = UpdatedComponent->GetCollisionObjectType();
 	FHitResult Result(1.f);
 	GetWorld()->SweepSingleByChannel(Result, OldLocation, SideDest, FQuat::Identity, CollisionChannel, CapsuleShape, CapsuleParams, ResponseParam);
@@ -3372,34 +3364,28 @@ bool UCharacterMovementComponent::CheckLedgeDirection(const FVector& OldLocation
 }
 
 
-FVector UCharacterMovementComponent::GetLedgeMove(const FVector& OldLocation, const FVector& Delta, const FVector& GravDir)
+FVector UCharacterMovementComponent::GetLedgeMove(const FVector& OldLocation, const FVector& Delta, const FVector& GravDir) const
 {
-	// We have a ledge!
-	if (!HasValidData())
+	if (!HasValidData() || Delta.IsZero())
 	{
 		return FVector::ZeroVector;
 	}
 
-	// check which direction ledge goes
-	float DesiredDistSq = Delta.SizeSquared();
-
-	if ( DesiredDistSq > 0.f )
-	{
-		FVector SideDir(Delta.Y, -1.f * Delta.X, 0.f);
+	FVector SideDir(Delta.Y, -1.f * Delta.X, 0.f);
 		
-		// try left
-		if ( CheckLedgeDirection(OldLocation, SideDir, GravDir) )
-		{
-			return SideDir;
-		}
+	// try left
+	if ( CheckLedgeDirection(OldLocation, SideDir, GravDir) )
+	{
+		return SideDir;
+	}
 
-		// try right
-		SideDir *= -1.f;
-		if ( CheckLedgeDirection(OldLocation, SideDir, GravDir) )
-		{
-			return SideDir;
-		}
+	// try right
+	SideDir *= -1.f;
+	if ( CheckLedgeDirection(OldLocation, SideDir, GravDir) )
+	{
+		return SideDir;
 	}
+	
 	return FVector::ZeroVector;
 }
 
@@ -3534,14 +3520,14 @@ void UCharacterMovementComponent::OnCharacterStuckInGeometry()
 
 void UCharacterMovementComponent::MoveAlongFloor(const FVector& InVelocity, float DeltaSeconds, FStepDownResult* OutStepDownResult)
 {
-	const FVector Delta = FVector(InVelocity.X, InVelocity.Y, 0.f) * DeltaSeconds;
-	
+
 	if (!CurrentFloor.IsWalkableFloor())
 	{
 		return;
 	}
 
 	// Move along the current floor
+	const FVector Delta = FVector(InVelocity.X, InVelocity.Y, 0.f) * DeltaSeconds;
 	FHitResult Hit(1.f);
 	FVector RampVector = ComputeGroundMovementDelta(Delta, CurrentFloor.HitResult, CurrentFloor.bLineTrace);
 	SafeMoveUpdatedComponent(RampVector, UpdatedComponent->GetComponentQuat(), true, Hit);
@@ -3667,11 +3653,10 @@ void UCharacterMovementComponent::PhysWalking(float deltaTime, int32 Iterations)
 
 		// Ensure velocity is horizontal.
 		MaintainHorizontalGroundVelocity();
-		Velocity.Z = 0.f;
 		const FVector OldVelocity = Velocity;
+		Acceleration.Z = 0.f;
 
 		// Apply acceleration
-		Acceleration.Z = 0.f;
 		if( !HasRootMotion() )
 		{
 			CalcVelocity(timeTick, GroundFriction, false, BrakingDecelerationWalking);
@@ -4358,12 +4343,12 @@ FVector CharAngularVelocity(FRotator const& OldRot, FRotator const& NewRot, floa
 	return RetAngVel;
 }
 
-FRotator UCharacterMovementComponent::GetDeltaRotation(float DeltaTime)
+FRotator UCharacterMovementComponent::GetDeltaRotation(float DeltaTime) const
 {
 	return (RotationRate * DeltaTime);
 }
 
-FRotator UCharacterMovementComponent::ComputeOrientToMovementRotation(const FRotator& CurrentRotation, float DeltaTime, FRotator& DeltaRotation)
+FRotator UCharacterMovementComponent::ComputeOrientToMovementRotation(const FRotator& CurrentRotation, float DeltaTime, FRotator& DeltaRotation) const
 {
 	if (Acceleration.SizeSquared() < KINDA_SMALL_NUMBER)
 	{
@@ -5184,17 +5169,11 @@ bool UCharacterMovementComponent::CanStepUp(const FHitResult& Hit) const
 }
 
 
-bool UCharacterMovementComponent::StepUp(const FVector& InGravDir, const FVector& Delta, const FHitResult &InHit, FStepDownResult* OutStepDownResult)
+bool UCharacterMovementComponent::StepUp(const FVector& GravDir, const FVector& Delta, const FHitResult &InHit, FStepDownResult* OutStepDownResult)
 {
 	SCOPE_CYCLE_COUNTER(STAT_CharStepUp);
 
-
-	if (!CanStepUp(InHit))
-	{
-		return false;
-	}
-
-	if (MaxStepHeight <= 0.f)
+	if (!CanStepUp(InHit) || MaxStepHeight <= 0.f)
 	{
 		return false;
 	}
@@ -5216,12 +5195,14 @@ bool UCharacterMovementComponent::StepUp(const FVector& InGravDir, const FVector
 		return false;
 	}
 
-	const FVector GravDir = InGravDir.GetSafeNormal();
 	if (GravDir.IsZero())
 	{
 		return false;
 	}
 
+	// Gravity should be a normalized direction
+	ensure(GravDir.IsNormalized());
+
 	float StepTravelUpHeight = MaxStepHeight;
 	float StepTravelDownHeight = StepTravelUpHeight;
 	const float StepSideZ = -1.f * (InHit.ImpactNormal | GravDir);
@@ -5404,7 +5385,7 @@ bool UCharacterMovementComponent::StepUp(const FVector& InGravDir, const FVector
 	return true;
 }
 
-void UCharacterMovementComponent::HandleImpact(FHitResult const& Impact, float TimeSlice, const FVector& MoveDelta)
+void UCharacterMovementComponent::HandleImpact(const FHitResult& Impact, float TimeSlice, const FVector& MoveDelta)
 {
 	if (CharacterOwner)
 	{
@@ -7001,6 +6982,20 @@ void UCharacterMovementComponent::SetAvoidanceEnabled(bool bEnable)
 	}
 }
 
+void UCharacterMovementComponent::ApplyDownwardForce(float DeltaSeconds)
+{
+	if (StandingDownwardForceScale != 0.0f && CurrentFloor.HitResult.IsValidBlockingHit())
+	{
+		UPrimitiveComponent* BaseComp = CurrentFloor.HitResult.GetComponent();
+		const FVector Gravity = FVector(0.0f, 0.0f, GetGravityZ());
+
+		if (BaseComp && BaseComp->IsAnySimulatingPhysics() && !Gravity.IsZero())
+		{
+			BaseComp->AddForceAtLocation(Gravity * Mass * StandingDownwardForceScale, CurrentFloor.HitResult.ImpactPoint, CurrentFloor.HitResult.BoneName);
+		}
+	}
+}
+
 void UCharacterMovementComponent::ApplyRepulsionForce(float DeltaSeconds)
 {
 	if (UpdatedPrimitive && RepulsionForce > 0.0f)
@@ -7062,27 +7057,31 @@ void UCharacterMovementComponent::ApplyRepulsionForce(float DeltaSeconds)
 																	QueryParams);
 
 				FVector HitLoc = Hit.ImpactPoint;
-				bool bIsPenetrating = Hit.bStartPenetrating || Hit.PenetrationDepth > 2.5f;
+				bool bIsPenetrating = Hit.bStartPenetrating || Hit.PenetrationDepth > StopBodyDistance;
 
 				// If we didn't hit the capsule, we're inside the capsule
-				if(!bHasHit) 
-				{ 
-					HitLoc = BodyLocation; 
+				if (!bHasHit) 
+				{
+					HitLoc = BodyLocation;
 					bIsPenetrating = true;
 				}
 
 				const float DistanceNow = (HitLoc - BodyLocation).SizeSquared2D();
 				const float DistanceLater = (HitLoc - (BodyLocation + BodyVelocity * DeltaSeconds)).SizeSquared2D();
 
-				if (BodyLocation.SizeSquared() > 0.1f && bHasHit && DistanceNow < StopBodyDistance && !bIsPenetrating)
+				if (bHasHit && DistanceNow < StopBodyDistance && !bIsPenetrating)
 				{
 					OverlapBody->SetLinearVelocity(FVector(0.0f, 0.0f, 0.0f), false);
 				}
 				else if (DistanceLater <= DistanceNow || bIsPenetrating)
 				{
-					FVector ForceCenter(MyLocation.X, MyLocation.Y, bHasHit ? HitLoc.Z : MyLocation.Z);
+					FVector ForceCenter = MyLocation;
 
-					if (!bHasHit)
+					if (bHasHit)
+					{
+						ForceCenter.Z = HitLoc.Z;
+					}
+					else
 					{
 						ForceCenter.Z = FMath::Clamp(BodyLocation.Z, MyLocation.Z - CapsuleHalfHeight, MyLocation.Z + CapsuleHalfHeight);
 					}
