Assert to help track down crash in GetUsedPrimitivesBasedOnMaterials

diff --git a/Engine/Source/Runtime/Engine/Private/Materials/MaterialShared.cpp b/Engine/Source/Runtime/Engine/Private/Materials/MaterialShared.cpp
--- a/Engine/Source/Runtime/Engine/Private/Materials/MaterialShared.cpp
+++ b/Engine/Source/Runtime/Engine/Private/Materials/MaterialShared.cpp
@@ -1659,6 +1659,7 @@ FMaterialRenderProxy::FMaterialRenderProxy()
 	: bSelected(false)
 	, bHovered(false)
 	, SubsurfaceProfileRT(0)
+	, DeletedFlag(0)
 {
 }
 
@@ -1666,6 +1667,7 @@ FMaterialRenderProxy::FMaterialRenderProxy(bool bInSelected, bool bInHovered)
 	: bSelected(bInSelected)
 	, bHovered(bInHovered)
 	, SubsurfaceProfileRT(0)
+	, DeletedFlag(0)
 {
 }
 
@@ -1676,6 +1678,8 @@ FMaterialRenderProxy::~FMaterialRenderProxy()
 		check(IsInRenderingThread());
 		ReleaseResource();
 	}
+
+	DeletedFlag = 1;
 }
 
 void FMaterialRenderProxy::InitDynamicRHI()
diff --git a/Engine/Source/Runtime/Engine/Public/MaterialShared.h b/Engine/Source/Runtime/Engine/Public/MaterialShared.h
--- a/Engine/Source/Runtime/Engine/Public/MaterialShared.h
+++ b/Engine/Source/Runtime/Engine/Public/MaterialShared.h
@@ -1343,6 +1343,7 @@ public:
 	virtual bool GetTextureValue(const FName ParameterName,const UTexture** OutValue, const FMaterialRenderContext& Context) const = 0;
 	bool IsSelected() const { return bSelected; }
 	bool IsHovered() const { return bHovered; }
+	bool IsDeleted() const { return DeletedFlag != 0; }
 
 	// FRenderResource interface.
 	ENGINE_API virtual void InitDynamicRHI() override;
@@ -1365,6 +1366,9 @@ private:
 	/** 0 if not set, game thread pointer, do not dereference, only for comparison */
 	const USubsurfaceProfile* SubsurfaceProfileRT;
 
+	/** For tracking down a bug accessing a deleted proxy. */
+	int32 DeletedFlag;
+
 	/** 
 	 * Tracks all material render proxies in all scenes, can only be accessed on the rendering thread.
 	 * This is used to propagate new shader maps to materials being used for rendering.
diff --git a/Engine/Source/Runtime/Renderer/Private/StaticMeshDrawList.inl b/Engine/Source/Runtime/Renderer/Private/StaticMeshDrawList.inl
--- a/Engine/Source/Runtime/Renderer/Private/StaticMeshDrawList.inl
+++ b/Engine/Source/Runtime/Renderer/Private/StaticMeshDrawList.inl
@@ -682,6 +682,7 @@ void TStaticMeshDrawList<DrawingPolicyType>::GetUsedPrimitivesBasedOnMaterials(E
 
 			if (Proxy)
 			{
+				check(!Proxy->IsDeleted());
 				FMaterial* MaterialResource = Proxy->GetMaterialNoFallback(FeatureLevel);
 
 				if (Materials.Contains(MaterialResource))
