#ue4 - Zero out FRepMovement::AngularVelocity when it's not used. More efficient for replication if it previously had non-zero values (switched from simulated to not simulated).

diff --git a/Engine/Source/Runtime/Engine/Private/ActorReplication.cpp b/Engine/Source/Runtime/Engine/Private/ActorReplication.cpp
--- a/Engine/Source/Runtime/Engine/Private/ActorReplication.cpp
+++ b/Engine/Source/Runtime/Engine/Private/ActorReplication.cpp
@@ -221,6 +221,7 @@ void AActor::GatherCurrentMovement()
 			ReplicatedMovement.Location = RootComponent->GetComponentLocation();
 			ReplicatedMovement.Rotation = RootComponent->GetComponentRotation();
 			ReplicatedMovement.LinearVelocity = GetVelocity();
+			ReplicatedMovement.AngularVelocity = FVector::ZeroVector;
 			ReplicatedMovement.bRepPhysics = false;
 		}
 	}
