GetNetMode returns incorrect value. Special cased the incorrect return, this was caused by PlayNetMode in ULevelEditorPlaySettings being set to PIE_Client
UE-17316

diff --git a/Engine/Source/Runtime/Engine/Private/World.cpp b/Engine/Source/Runtime/Engine/Private/World.cpp
--- a/Engine/Source/Runtime/Engine/Private/World.cpp
+++ b/Engine/Source/Runtime/Engine/Private/World.cpp
@@ -5401,8 +5401,6 @@ ENetMode UWorld::GetNetMode() const
 
 	// Use NextURL or PendingNetURL to derive NetMode
 	return AttemptDeriveFromURL();
-
-	//return NM_Standalone;
 }
 
 #if WITH_EDITOR
@@ -5419,9 +5417,33 @@ ENetMode UWorld::AttemptDeriveFromPlayInSettings() const
 			switch (PlayNetMode)
 			{
 			case EPlayNetMode::PIE_Client:
+			{
+				int32 NumberOfClients = 0;
+				PlayInSettings->GetPlayNumberOfClients(NumberOfClients);
+
+				bool bAutoConnectToServer = false;
+				PlayInSettings->GetAutoConnectToServer(bAutoConnectToServer);
+
+				// Playing as client without listen server in single process,
+				// or as a client not going to connect to a server
+				if(NumberOfClients == 1 || bAutoConnectToServer == false)
+				{
+					return NM_Standalone;
+				}
 				return NM_Client;
+			}
 			case EPlayNetMode::PIE_ListenServer:
+			{
+				bool bDedicatedServer = false;
+				PlayInSettings->GetPlayNetDedicated(bDedicatedServer);
+
+				if(bDedicatedServer == true)
+				{
+					return NM_DedicatedServer;
+				}
+
 				return NM_ListenServer;
+			}
 			case EPlayNetMode::PIE_Standalone:
 				return NM_Standalone;
 			default:
