Fixed crash when destroying a pawn on the client by Kill Z volume by checking if character is pending kill before calling TickCharacterPose, also added safety checks on accessing FAnimUpdateRateParameters pointer in TickCharacterPose() : UE-18218

diff --git a/Engine/Source/Runtime/Engine/Private/Components/CharacterMovementComponent.cpp b/Engine/Source/Runtime/Engine/Private/Components/CharacterMovementComponent.cpp
--- a/Engine/Source/Runtime/Engine/Private/Components/CharacterMovementComponent.cpp
+++ b/Engine/Source/Runtime/Engine/Private/Components/CharacterMovementComponent.cpp
@@ -6581,6 +6581,12 @@ void UCharacterMovementComponent::MoveAutonomous
 	
 	PerformMovement(DeltaTime);
 
+	// Check if data is valid as PerformMovement can mark character for pending kill
+	if (!HasValidData())
+	{
+		return;
+	}
+
 	// If not playing root motion, tick animations after physics. We do this here to keep events, notifies, states and transitions in sync with client updates.
 	if( !CharacterOwner->bClientUpdating && !CharacterOwner->IsPlayingRootMotion() && CharacterOwner->GetMesh() )
 	{
diff --git a/Engine/Source/Runtime/Engine/Private/Components/SkeletalMeshComponent.cpp b/Engine/Source/Runtime/Engine/Private/Components/SkeletalMeshComponent.cpp
--- a/Engine/Source/Runtime/Engine/Private/Components/SkeletalMeshComponent.cpp
+++ b/Engine/Source/Runtime/Engine/Private/Components/SkeletalMeshComponent.cpp
@@ -611,7 +611,7 @@ void USkeletalMeshComponent::TickPose(float DeltaTime, bool bNeedsValidRootMotio
 {
 	Super::TickPose(DeltaTime, bNeedsValidRootMotion);
 
-	if (!bEnableUpdateRateOptimizations || !AnimUpdateRateParams->ShouldSkipUpdate())
+	if (AnimUpdateRateParams != nullptr && (!bEnableUpdateRateOptimizations || !AnimUpdateRateParams->ShouldSkipUpdate()))
 	{
 		float TimeAdjustment = AnimUpdateRateParams->GetTimeAdjustment();
 		TickAnimation(DeltaTime + TimeAdjustment);
