Notify client in non shipping build if we disconnect due to exploit security risk : UE-2710

diff --git a/Engine/Source/Runtime/Core/Public/Misc/AssertionMacros.h b/Engine/Source/Runtime/Core/Public/Misc/AssertionMacros.h
--- a/Engine/Source/Runtime/Core/Public/Misc/AssertionMacros.h
+++ b/Engine/Source/Runtime/Core/Public/Misc/AssertionMacros.h
@@ -341,13 +341,12 @@ struct FTCharArrayTester
 			{ \
 				if (!LogSecurity.IsSuppressed(ELogVerbosity::Warning)) \
 				{ \
-				FString Test = FString::Printf(TEXT("%s: %s: %s"), *(NetConnection->RemoteAddressToString()), ToString(SecurityEventType), Format); \
-					FMsg::Logf_Internal(__FILE__, __LINE__, LogSecurity.GetCategoryName(), ELogVerbosity::Warning, *Test, ##__VA_ARGS__); \
+					FString Print = FString::Printf(TEXT("%s: %s: %s"), *(NetConnection->RemoteAddressToString()), ESecurityEvent::ToString(SecurityEventType), Format); \
+					FMsg::Logf_Internal(__FILE__, __LINE__, LogSecurity.GetCategoryName(), ELogVerbosity::Warning, *Print, ##__VA_ARGS__); \
 				} \
 			} \
 		 }
 
-
 		// Conditional logging. Will only log if Condition is met.
 		#define UE_CLOG(Condition, CategoryName, Verbosity, Format, ...) \
 		{ \
@@ -448,3 +447,32 @@ struct FTCharArrayTester
 #if PLATFORM_HTML5
 #include "HTML5/HTML5AssertionMacros.h"
 #endif 
+
+#if UE_BUILD_SHIPPING
+#define NOTIFY_CLIENT_OF_SECURITY_EVENT_IF_NOT_SHIPPING(NetConnection, SecurityPrint) ;
+#else
+#define NOTIFY_CLIENT_OF_SECURITY_EVENT_IF_NOT_SHIPPING(NetConnection, SecurityPrint) \
+	FNetControlMessage<NMT_SecurityViolation>::Send(NetConnection, SecurityPrint); \
+	NetConnection->FlushNet(true)
+#endif
+
+/**
+		* A  macro that closes the connection and logs the security event on the server and the client
+		* @param NetConnection, a valid UNetConnection
+		* @param SecurityEventType, a security event type (ESecurityEvent::Type)
+		* @param Format, format text
+***/
+#define CLOSE_CONNECTION_DUE_TO_SECURITY_VIOLATION(NetConnection, SecurityEventType, Format, ...) \
+{ \
+	static_assert(IS_TCHAR_ARRAY(Format), "Formatting string must be a TCHAR array."); \
+	check(NetConnection != nullptr); \
+	FString SecurityPrint = FString::Printf(Format,##__VA_ARGS__); \
+	if(Channels[0] == nullptr) \
+	{ \
+		Channels[0] = CreateChannel(EChannelType::CHTYPE_Control, false, 0); \
+	} \
+	UE_SECURITY_LOG(NetConnection, SecurityEventType, Format, __VA_ARGS__); \
+	UE_SECURITY_LOG(NetConnection, ESecurityEvent::Closed, TEXT("Connection closed")); \
+	NOTIFY_CLIENT_OF_SECURITY_EVENT_IF_NOT_SHIPPING(NetConnection, SecurityPrint); \
+	NetConnection->Close(); \
+}
\ No newline at end of file
diff --git a/Engine/Source/Runtime/Engine/Private/DataChannel.cpp b/Engine/Source/Runtime/Engine/Private/DataChannel.cpp
--- a/Engine/Source/Runtime/Engine/Private/DataChannel.cpp
+++ b/Engine/Source/Runtime/Engine/Private/DataChannel.cpp
@@ -1037,6 +1037,7 @@ IMPLEMENT_CONTROL_CHANNEL_MESSAGE(Abort);
 IMPLEMENT_CONTROL_CHANNEL_MESSAGE(PCSwap);
 IMPLEMENT_CONTROL_CHANNEL_MESSAGE(ActorChannelFailure);
 IMPLEMENT_CONTROL_CHANNEL_MESSAGE(DebugText);
+IMPLEMENT_CONTROL_CHANNEL_MESSAGE(SecurityViolation);
 IMPLEMENT_CONTROL_CHANNEL_MESSAGE(BeaconWelcome);
 IMPLEMENT_CONTROL_CHANNEL_MESSAGE(BeaconJoin);
 IMPLEMENT_CONTROL_CHANNEL_MESSAGE(BeaconAssignGUID);
@@ -1177,6 +1178,14 @@ void UControlChannel::ReceivedBunch( FInBunch& Bunch )
 				}
 
 			}
+			else if(MessageType == NMT_SecurityViolation)
+			{
+				FString DebugMessage;
+				FNetControlMessage<NMT_SecurityViolation>::Receive(Bunch, DebugMessage);
+
+				UE_SECURITY_LOG(Connection, ESecurityEvent::Closed, TEXT("%s"), *DebugMessage);
+				break;
+			}
 		}
 		else
 		{
@@ -1233,6 +1242,10 @@ void UControlChannel::ReceivedBunch( FInBunch& Bunch )
 					break;
 				case NMT_NetGUIDAssign:
 					FNetControlMessage<NMT_NetGUIDAssign>::Discard(Bunch);
+					break;
+				case NMT_SecurityViolation:
+					FNetControlMessage<NMT_SecurityViolation>::Discard(Bunch);
+					break;
 				case NMT_BeaconWelcome:
 					//FNetControlMessage<NMT_BeaconWelcome>::Discard(Bunch);
 					break;
diff --git a/Engine/Source/Runtime/Engine/Private/NetConnection.cpp b/Engine/Source/Runtime/Engine/Private/NetConnection.cpp
--- a/Engine/Source/Runtime/Engine/Private/NetConnection.cpp
+++ b/Engine/Source/Runtime/Engine/Private/NetConnection.cpp
@@ -497,17 +497,13 @@ void UNetConnection::ReceivedRawPacket( void* InData, int32 Count )
 		// MalformedPacket - Received a packet with 0's in the last byte
 		else 
 		{
-			UE_SECURITY_LOG(this, ESecurityEvent::Malformed_Packet, TEXT("Received packet with 0's in last byte of packet"));
-			Close();	// If they have the secret key (or haven't opened the control channel yet) and get here, assume they are being malicious
-			UE_SECURITY_LOG(this, ESecurityEvent::Closed, TEXT("Connection closed"));
+			CLOSE_CONNECTION_DUE_TO_SECURITY_VIOLATION(this, ESecurityEvent::Malformed_Packet, TEXT("Received packet with 0's in last byte of packet"));
 		}
 	}
 	// MalformedPacket - Received a packet of 0 bytes
 	else 
 	{
-		UE_SECURITY_LOG(this, ESecurityEvent::Malformed_Packet, TEXT("Received zero-size packet"));
-		Close();	// If they have the secret key (or haven't opened the control channel yet) and get here, assume they are being malicious
-		UE_SECURITY_LOG(this, ESecurityEvent::Closed, TEXT("Connection closed"));
+		CLOSE_CONNECTION_DUE_TO_SECURITY_VIOLATION(this, ESecurityEvent::Malformed_Packet, TEXT("Received zero-size packet"));
 	}
 }
 
@@ -737,9 +733,7 @@ void UNetConnection::ReceivedPacket( FBitReader& Reader )
 		bool IsAck = !!Reader.ReadBit();
 		if ( Reader.IsError() )
 		{
-			UE_SECURITY_LOG(this, ESecurityEvent::Invalid_Data, TEXT("Bunch missing ack flag"));
-			Close();
-			UE_SECURITY_LOG(this, ESecurityEvent::Closed, TEXT("Connection Closed"));
+			CLOSE_CONNECTION_DUE_TO_SECURITY_VIOLATION(this, ESecurityEvent::Malformed_Packet, TEXT("Bunch missing ack flag"));
 			return;
 		}
 
@@ -753,9 +747,7 @@ void UNetConnection::ReceivedPacket( FBitReader& Reader )
 
 			if( Reader.IsError() )
 			{
-				UE_SECURITY_LOG(this, ESecurityEvent::Invalid_Data, TEXT("Bunch missing ack"));
-				Close();
-				UE_SECURITY_LOG(this, ESecurityEvent::Closed, TEXT("Connection Closed"));
+				CLOSE_CONNECTION_DUE_TO_SECURITY_VIOLATION(this, ESecurityEvent::Malformed_Packet, TEXT("Bunch missing ack"));
 				return;
 			}
 
@@ -914,16 +906,14 @@ void UNetConnection::ReceivedPacket( FBitReader& Reader )
 
 			if( Reader.IsError() )
 			{
-				UE_LOG( LogNetTraffic, Error, TEXT( "Bunch header overflowed" ) );
-				Close();
+				CLOSE_CONNECTION_DUE_TO_SECURITY_VIOLATION(this, ESecurityEvent::Malformed_Packet, TEXT("Bunch header overflowed"));
 				return;
 			}
 			Bunch.SetData( Reader, BunchDataBits );
 			if( Reader.IsError() )
 			{
 				// Bunch claims it's larger than the enclosing packet.
-				UE_LOG( LogNetTraffic, Error, TEXT( "Bunch data overflowed (%i %i+%i/%i)" ), IncomingStartPos, HeaderPos, BunchDataBits, Reader.GetNumBits() );
-				Close();
+				CLOSE_CONNECTION_DUE_TO_SECURITY_VIOLATION(this, ESecurityEvent::Invalid_Data, TEXT("Bunch data overflowed (%i %i+%i/%i)"), IncomingStartPos, HeaderPos, BunchDataBits, Reader.GetNumBits());
 				return;
 			}
 
@@ -951,8 +941,7 @@ void UNetConnection::ReceivedPacket( FBitReader& Reader )
 				// Can't handle other channels until control channel exists.
 				if ( Channels[0] == NULL )
 				{
-					UE_LOG( LogNetTraffic, Error, TEXT( "UNetConnection::ReceivedPacket: Received non-control bunch before control channel was created. ChIndex: %i, ChType: %i" ), Bunch.ChIndex, Bunch.ChType );
-					Close();
+					CLOSE_CONNECTION_DUE_TO_SECURITY_VIOLATION(this, ESecurityEvent::Malformed_Packet, TEXT( "UNetConnection::ReceivedPacket: Received non-control bunch before control channel was created. ChIndex: %i, ChType: %i" ), Bunch.ChIndex, Bunch.ChType);
 					return;
 				}
 				// on the server, if we receive bunch data for a channel that doesn't exist while we're still logging in,
@@ -960,16 +949,14 @@ void UNetConnection::ReceivedPacket( FBitReader& Reader )
 				// so reject it
 				else if ( PlayerController == NULL && Driver->ClientConnections.Contains( this ) )
 				{
-					UE_LOG( LogNetTraffic, Error, TEXT( "UNetConnection::ReceivedPacket: Received non-control bunch before player controller was assigned. ChIndex: %i, ChType: %i" ), Bunch.ChIndex, Bunch.ChType );
-					Close();
+					CLOSE_CONNECTION_DUE_TO_SECURITY_VIOLATION(this, ESecurityEvent::Malformed_Packet, TEXT( "UNetConnection::ReceivedPacket: Received non-control bunch before player controller was assigned. ChIndex: %i, ChType: %i" ), Bunch.ChIndex, Bunch.ChType);
 					return;
 				}
 			}
 			// ignore control channel close if it hasn't been opened yet
 			if ( Bunch.ChIndex == 0 && Channels[0] == NULL && Bunch.bClose && Bunch.ChType == CHTYPE_Control )
 			{
-				UE_LOG( LogNetTraffic, Error, TEXT( "UNetConnection::ReceivedPacket: Received control channel close before open" ) );
-				Close();
+				CLOSE_CONNECTION_DUE_TO_SECURITY_VIOLATION(this, ESecurityEvent::Malformed_Packet, TEXT( "UNetConnection::ReceivedPacket: Received control channel close before open" ));
 				return;
 			}
 
@@ -1021,8 +1008,7 @@ void UNetConnection::ReceivedPacket( FBitReader& Reader )
 				if ( !UChannel::IsKnownChannelType( Bunch.ChType ) )
 				{
 					// Unknown type.
-					UE_LOG( LogNetTraffic, Error, TEXT( "UNetConnection::ReceivedPacket: Connection unknown channel type (%i)" ), (int)Bunch.ChType );
-					Close();
+					CLOSE_CONNECTION_DUE_TO_SECURITY_VIOLATION(this, ESecurityEvent::Invalid_Data, TEXT( "UNetConnection::ReceivedPacket: Connection unknown channel type (%i)" ), (int)Bunch.ChType);
 					return;
 				}
 
diff --git a/Engine/Source/Runtime/Engine/Public/Net/DataChannel.h b/Engine/Source/Runtime/Engine/Public/Net/DataChannel.h
--- a/Engine/Source/Runtime/Engine/Public/Net/DataChannel.h
+++ b/Engine/Source/Runtime/Engine/Public/Net/DataChannel.h
@@ -385,6 +385,7 @@ DEFINE_CONTROL_CHANNEL_MESSAGE_ONEPARAM(PCSwap, 15, int32); // client tells serv
 DEFINE_CONTROL_CHANNEL_MESSAGE_ONEPARAM(ActorChannelFailure, 16, int32); // client tells server that it failed to open an Actor channel sent by the server (e.g. couldn't serialize Actor archetype)
 DEFINE_CONTROL_CHANNEL_MESSAGE_ONEPARAM(DebugText, 17, FString); // debug text sent to all clients or to server
 DEFINE_CONTROL_CHANNEL_MESSAGE_TWOPARAM(NetGUIDAssign, 18, FNetworkGUID, FString); // Explicit NetworkGUID assignment. This is rare and only happens if a netguid is only serialized client->server (this msg goes server->client to tell client what ID to use in that case)
+DEFINE_CONTROL_CHANNEL_MESSAGE_ONEPARAM(SecurityViolation, 19, FString); // server tells client that it has violated security and has been disconnected
 
 // 			Beacon control channel flow
 // Client												Server
