Deprecated version of SpawnActorDeferred() that takes separate location and rotation and the old bNoCollisionFail bool.

diff --git a/Engine/Source/Runtime/Engine/Classes/Engine/World.h b/Engine/Source/Runtime/Engine/Classes/Engine/World.h
--- a/Engine/Source/Runtime/Engine/Classes/Engine/World.h
+++ b/Engine/Source/Runtime/Engine/Classes/Engine/World.h
@@ -2837,7 +2837,6 @@ public:
 		return CastChecked<T>(SpawnActor(Class, &Location, &Rotation, SpawnParameters),ECastCheckedType::NullAllowed);
 	}
 	/** 
-	
 	 *  Templated version of SpawnActor that allows you to specify whole Transform
 	 *  class type via parameter while the return type is a parent class of that type 
 	 */
@@ -2853,40 +2852,14 @@ public:
 	* manually by calling UGameplayStatics::FinishSpawningActor (see AActor::OnConstruction).
 	*/
 	template< class T >
-	T* SpawnActorDeferred(
-		UClass* Class,
-		FVector const& Location,
-		FRotator const& Rotation,
-		AActor* Owner,
-		APawn* Instigator,
-		bool bNoCollisionFail
-		)
-	{
-		if( Owner )
-		{
-			check(this==Owner->GetWorld());
-		}
-		FActorSpawnParameters SpawnInfo;
-		SpawnInfo.SpawnCollisionHandlingOverride = bNoCollisionFail ? ESpawnActorCollisionHandlingMethod::AlwaysSpawn : ESpawnActorCollisionHandlingMethod::AdjustIfPossibleButDontSpawnIfColliding;
-		SpawnInfo.Owner = Owner;
-		SpawnInfo.Instigator = Instigator;
-		SpawnInfo.bDeferConstruction = true;
-		return (Class != NULL) ? Cast<T>(SpawnActor(Class, &Location, &Rotation, SpawnInfo )) : NULL;
-	}
-
-	/**
-	* Spawns given class and returns class T pointer, forcibly sets world position. WILL NOT run Construction Script of Blueprints 
-	* to give caller an opportunity to set parameters beforehand.  Caller is responsible for invoking construction
-	* manually by calling UGameplayStatics::FinishSpawningActor (see AActor::OnConstruction).
-	*/
-	template< class T >
+	DEPRECATED(4.9, "This version of SpawnActorDeferred is deprecated. Please use the version that takes an FTransform and ESpawnActorCollisionHandlingMethod.")
 	T* SpawnActorDeferred(
 		UClass* Class,
 		FVector const& Location,
 		FRotator const& Rotation,
 		AActor* Owner = nullptr,
 		APawn* Instigator = nullptr,
-		ESpawnActorCollisionHandlingMethod CollisionHandlingOverride = ESpawnActorCollisionHandlingMethod::Undefined
+		bool bNoCollisionFail = false
 		)
 	{
 		if( Owner )
@@ -2894,7 +2867,7 @@ public:
 			check(this==Owner->GetWorld());
 		}
 		FActorSpawnParameters SpawnInfo;
-		SpawnInfo.SpawnCollisionHandlingOverride = CollisionHandlingOverride;
+		SpawnInfo.SpawnCollisionHandlingOverride = bNoCollisionFail ? ESpawnActorCollisionHandlingMethod::AlwaysSpawn : ESpawnActorCollisionHandlingMethod::AdjustIfPossibleButDontSpawnIfColliding;
 		SpawnInfo.Owner = Owner;
 		SpawnInfo.Instigator = Instigator;
 		SpawnInfo.bDeferConstruction = true;
@@ -2910,8 +2883,8 @@ public:
 	T* SpawnActorDeferred(
 		UClass* Class,
 		FTransform const& Transform,
-		AActor* Owner=NULL,
-		APawn* Instigator=NULL,
+		AActor* Owner = nullptr,
+		APawn* Instigator = nullptr,
 		ESpawnActorCollisionHandlingMethod CollisionHandlingOverride = ESpawnActorCollisionHandlingMethod::Undefined
 		)
 	{
@@ -2924,7 +2897,7 @@ public:
 		SpawnInfo.Owner = Owner;
 		SpawnInfo.Instigator = Instigator;
 		SpawnInfo.bDeferConstruction = true;
-		return (Class != NULL) ? Cast<T>(SpawnActor(Class, &Transform, SpawnInfo)) : NULL;
+		return (Class != nullptr) ? Cast<T>(SpawnActor(Class, &Transform, SpawnInfo)) : nullptr;
 	}
 
 	/** 
