Replicate APlayerController::SpawnLocation over network which stops remote spectator players spawning on ZeroVector position

UE-5986

diff --git a/Engine/Source/Runtime/Engine/Classes/GameFramework/PlayerController.h b/Engine/Source/Runtime/Engine/Classes/GameFramework/PlayerController.h
--- a/Engine/Source/Runtime/Engine/Classes/GameFramework/PlayerController.h
+++ b/Engine/Source/Runtime/Engine/Classes/GameFramework/PlayerController.h
@@ -1532,7 +1532,7 @@ private:
 
 protected:
 	/** The location used internally when there is no pawn or spectator, to know where to spawn the spectator or focus the camera on death. */
-	UPROPERTY()
+	UPROPERTY(Replicated)
 	FVector SpawnLocation;
 
 	/** Set the SpawnLocation for use when changing states or when there is no pawn or spectator. */
diff --git a/Engine/Source/Runtime/Engine/Private/PlayerController.cpp b/Engine/Source/Runtime/Engine/Private/PlayerController.cpp
--- a/Engine/Source/Runtime/Engine/Private/PlayerController.cpp
+++ b/Engine/Source/Runtime/Engine/Private/PlayerController.cpp
@@ -3583,6 +3583,9 @@ void APlayerController::GetLifetimeReplicatedProps( TArray< FLifetimeProperty >
 	// But, since they also don't update unless that condition is true, these values won't change, thus won't send
 	// This is a little less efficient, but fits into the new condition system well, and shouldn't really add much overhead
 	DOREPLIFETIME_CONDITION( APlayerController, TargetViewRotation,		COND_OwnerOnly );
+
+	// Replicate SpawnLocation for remote spectators
+	DOREPLIFETIME_CONDITION( APlayerController, SpawnLocation, COND_OwnerOnly );
 }
 
 void APlayerController::SetPawn(APawn* InPawn)
