Keys bound via key or action events will now count as having consumed the event even if there wasn't a pressed or released event this frame preventing axes using that key in lower priority input components from being active on other frames.

diff --git a/Engine/Source/Runtime/Engine/Private/UserInterface/PlayerInput.cpp b/Engine/Source/Runtime/Engine/Private/UserInterface/PlayerInput.cpp
--- a/Engine/Source/Runtime/Engine/Private/UserInterface/PlayerInput.cpp
+++ b/Engine/Source/Runtime/Engine/Private/UserInterface/PlayerInput.cpp
@@ -601,6 +601,7 @@ void UPlayerInput::ConditionalBuildKeyMappings()
 void UPlayerInput::GetChordsForKeyMapping(const FInputActionKeyMapping& KeyMapping, const FInputActionBinding& ActionBinding, const bool bGamePaused, TArray<FDelegateDispatchDetails>& FoundChords, TArray<FKey>& KeysToConsume)
 {
 	TArray<uint32> EventIndices;
+	bool bConsumeInput = false;
 
 	// test modifier conditions and ignore the event if they failed
 	if (	(KeyMapping.bAlt == false || IsAltPressed())
@@ -646,12 +647,13 @@ void UPlayerInput::GetChordsForKeyMapping(const FInputActionKeyMapping& KeyMappi
 				FoundChord.EventIndex = EventIndices[EventsIndex];
 				FoundChords.Add(FoundChord);
 			}
-			if (ActionBinding.bConsumeInput)
-			{
-				KeysToConsume.AddUnique(KeyMapping.Key);
-			}
+			bConsumeInput = true;
 		}
 	}
+	if (ActionBinding.bConsumeInput && (bConsumeInput || !(KeyMapping.bAlt || KeyMapping.bCtrl || KeyMapping.bShift || KeyMapping.bCmd)))
+	{
+		KeysToConsume.AddUnique(KeyMapping.Key);
+	}
 }
 
 void UPlayerInput::GetChordsForAction(const FInputActionBinding& ActionBinding, const bool bGamePaused, TArray<FDelegateDispatchDetails>& FoundChords, TArray<FKey>& KeysToConsume)
@@ -685,7 +687,7 @@ void UPlayerInput::GetChordsForAction(const FInputActionBinding& ActionBinding,
 
 void UPlayerInput::GetChordForKey(const FInputKeyBinding& KeyBinding, const bool bGamePaused, TArray<FDelegateDispatchDetails>& FoundChords, TArray<FKey>& KeysToConsume)
 {
-	bool bKeyOccurred = false;
+	bool bConsumeInput = false;
 
 	if (KeyBinding.Chord.Key == EKeys::AnyKey)
 	{
@@ -744,13 +746,14 @@ void UPlayerInput::GetChordForKey(const FInputKeyBinding& KeyBinding, const bool
 					FoundChord.EventIndex = EventIndices[EventsIndex];
 					FoundChords.Add(FoundChord);
 				}
-				if (KeyBinding.bConsumeInput)
-				{
-					KeysToConsume.AddUnique(KeyBinding.Chord.Key);
-				}
+				bConsumeInput = true;
 			}
 		}
 	}
+	if (KeyBinding.bConsumeInput && (bConsumeInput || !(KeyBinding.Chord.bAlt || KeyBinding.Chord.bCtrl || KeyBinding.Chord.bShift || KeyBinding.Chord.bCmd)))
+	{
+		KeysToConsume.AddUnique(KeyBinding.Chord.Key);
+	}
 }
 
 float UPlayerInput::DetermineAxisValue(const FInputAxisBinding& AxisBinding, const bool bGamePaused, TArray<FKey>& KeysToConsume)
