Fixing issue with broken thumbnails and material preview. Inverse view matrix was computed by transposing, but since it contains translation, a full inverse needs to be used instead. Replaced Transpose with InverseFast.

diff --git a/Engine/Source/Runtime/Engine/Private/SceneView.cpp b/Engine/Source/Runtime/Engine/Private/SceneView.cpp
--- a/Engine/Source/Runtime/Engine/Private/SceneView.cpp
+++ b/Engine/Source/Runtime/Engine/Private/SceneView.cpp
@@ -342,7 +342,7 @@ FSceneView::FSceneView(const FSceneViewInitOptions& InitOptions)
 
 	/** The view transform, starting from world-space points translated by -ViewOrigin. */
 	FMatrix TranslatedViewMatrix = InitOptions.ViewRotationMatrix;
-	FMatrix InvTranslatedViewMatrix = TranslatedViewMatrix.GetTransposed();
+	FMatrix InvTranslatedViewMatrix = TranslatedViewMatrix.Inverse();
 
 	// Translate world-space so its origin is at ViewOrigin for improved precision.
 	// Note that this isn't exactly right for orthogonal projections (See the above special case), but we still use ViewOrigin
