Fixing shadow variables

diff --git a/Engine/Source/Runtime/Engine/Private/LevelActor.cpp b/Engine/Source/Runtime/Engine/Private/LevelActor.cpp
--- a/Engine/Source/Runtime/Engine/Private/LevelActor.cpp
+++ b/Engine/Source/Runtime/Engine/Private/LevelActor.cpp
@@ -876,7 +876,7 @@ static bool ComponentEncroachesBlockingGeometry_WithAdjustment(UWorld const* Wor
 					if (FMath::IsNearlyZero(MTDResult.Distance))
 					{
 						FCollisionShape const ShrunkenCollisionShape = PrimComp->GetCollisionShape(-Epsilon);
-						bool bSuccess = OverlapComponent->ComputePenetration(MTDResult, ShrunkenCollisionShape, TestWorldTransform.GetLocation(), TestWorldTransform.GetRotation());
+						bSuccess = OverlapComponent->ComputePenetration(MTDResult, ShrunkenCollisionShape, TestWorldTransform.GetLocation(), TestWorldTransform.GetRotation());
 						if (bSuccess)
 						{
 							OutProposedAdjustment += MTDResult.Direction * MTDResult.Distance;
diff --git a/Engine/Source/Runtime/Engine/Private/World.cpp b/Engine/Source/Runtime/Engine/Private/World.cpp
--- a/Engine/Source/Runtime/Engine/Private/World.cpp
+++ b/Engine/Source/Runtime/Engine/Private/World.cpp
@@ -3064,9 +3064,9 @@ void UWorld::CleanupWorld(bool bSessionEnded, bool bCleanupResources, UWorld* Ne
 	check(IsVisibilityRequestPending() == false);
 
 	// Wait on current physics scenes if they are processing
-	if(FPhysScene* Scene = GetPhysicsScene())
+	if(FPhysScene* CurrPhysicsScene = GetPhysicsScene())
 	{
-		Scene->WaitPhysScenes();
+		CurrPhysicsScene->WaitPhysScenes();
 	}
 
 	FWorldDelegates::OnWorldCleanup.Broadcast(this, bSessionEnded, bCleanupResources);
