Fixing CIS build warnings: mono, FBodyInstance construction list not matching order of field declarations.

diff --git a/Engine/Source/Runtime/Engine/Private/PhysicsEngine/BodyInstance.cpp b/Engine/Source/Runtime/Engine/Private/PhysicsEngine/BodyInstance.cpp
--- a/Engine/Source/Runtime/Engine/Private/PhysicsEngine/BodyInstance.cpp
+++ b/Engine/Source/Runtime/Engine/Private/PhysicsEngine/BodyInstance.cpp
@@ -206,65 +206,71 @@ void FCollisionResponse::UpdateResponseContainerFromArray()
 ////////////////////////////////////////////////////////////////////////////
 
 FBodyInstance::FBodyInstance()
-: InstanceBodyIndex(INDEX_NONE)
-, InstanceBoneIndex(INDEX_NONE)
-, Scale3D(1.0f)
-, SceneIndexSync(0)
-, SceneIndexAsync(0)
-, CollisionProfileName(UCollisionProfile::CustomCollisionProfileName)
-, CollisionEnabled(ECollisionEnabled::QueryAndPhysics)
-, ObjectType(ECC_WorldStatic)
-, bUseCCD(false)
-, bNotifyRigidBodyCollision(false)
-, bSimulatePhysics(false)
-, bOverrideMass(false)
-, MassInKg(100.f)
-, LinearDamping(0.01)
-, AngularDamping(0.0)
-, bEnableGravity(true)
-, bAutoWeld(false)
-, bWelded(false)
-, bStartAwake(true)
-, bGenerateWakeEvents(false)
-, bUpdateMassWhenScaleChanges(false)
-, bLockTranslation(true)
-, bLockRotation(true)
-, bLockXTranslation(false)
-, bLockYTranslation(false)
-, bLockZTranslation(false)
-, bLockXRotation(false)
-, bLockYRotation(false)
-, bLockZRotation(false)
-, DOFMode(0)
-, CustomDOFPlaneNormal(FVector::ZeroVector)
-, COMNudge(ForceInit)
-, MassScale(1.f)
-, DOFConstraint(NULL)
-, WeldParent(NULL)
-, bUseAsyncScene(false)
-, bOverrideMaxDepenetrationVelocity(false)
-, MaxDepenetrationVelocity(0.f)
-, bOverrideWalkableSlopeOnInstance(false)
-, PhysMaterialOverride(NULL)
-, SleepFamily(ESleepFamily::Normal)
-, CustomSleepThresholdMultiplier(1.f)
-, PhysicsBlendWeight(0.f)
-, PositionSolverIterationCount(8)
-, VelocitySolverIterationCount(1)
+	: InstanceBodyIndex(INDEX_NONE)
+	, InstanceBoneIndex(INDEX_NONE)
+	, Scale3D(1.0f)
+	, SceneIndexSync(0)
+	, SceneIndexAsync(0)
+	, CollisionProfileName(UCollisionProfile::CustomCollisionProfileName)
+	, bUseCCD(false)
+	, bNotifyRigidBodyCollision(false)
+	, bSimulatePhysics(false)
+	, bOverrideMass(false)
+	, bEnableGravity(true)
+	, bAutoWeld(false)
+	, bWelded(false)
+	, bStartAwake(true)
+	, bGenerateWakeEvents(false)
+	, bUpdateMassWhenScaleChanges(false)
+	, bLockTranslation(true)
+	, bLockRotation(true)
+	, bLockXTranslation(false)
+	, bLockYTranslation(false)
+	, bLockZTranslation(false)
+	, bLockXRotation(false)
+	, bLockYRotation(false)
+	, bLockZRotation(false)
 #if WITH_PHYSX
-, RigidActorSync(NULL)
-, RigidActorAsync(NULL)
-, BodyAggregate(NULL)
-, RigidActorSyncId(PX_SERIAL_OBJECT_ID_INVALID)
-, RigidActorAsyncId(PX_SERIAL_OBJECT_ID_INVALID)
-, InitialLinearVelocity(0.0f)
-, bWokenExternally(false)
-, PhysxUserData(this)
-, CurrentSceneState(BodyInstanceSceneState::NotAdded)
-#endif	//WITH_PHYSX
+	, bWokenExternally(false)
+#endif // WITH_PHYSX
+	, bUseAsyncScene(false)
+	, bOverrideMaxDepenetrationVelocity(false)
+	, bOverrideWalkableSlopeOnInstance(false)
+	, MaxDepenetrationVelocity(0.f)
+	, MassInKg(100.f)
+	, LinearDamping(0.01)
+	, AngularDamping(0.0)
+	, CustomDOFPlaneNormal(FVector::ZeroVector)
+	, COMNudge(ForceInit)
+	, MassScale(1.f)
+	, DOFConstraint(NULL)
+	, WeldParent(NULL)
+	, PhysMaterialOverride(NULL)
+	, CustomSleepThresholdMultiplier(1.f)
+	, PhysicsBlendWeight(0.f)
+	, PositionSolverIterationCount(8)
+#if WITH_PHYSX
+	, RigidActorSync(NULL)
+	, RigidActorAsync(NULL)
+	, BodyAggregate(NULL)
+	, RigidActorSyncId(PX_SERIAL_OBJECT_ID_INVALID)
+	, RigidActorAsyncId(PX_SERIAL_OBJECT_ID_INVALID)
+#endif // WITH_PHYSX
+	, VelocitySolverIterationCount(1)
+#if WITH_PHYSX
+	, InitialLinearVelocity(0.0f)
+	, PhysxUserData(this)
+#endif // WITH_PHYSX
 #if WITH_BOX2D
-, BodyInstancePtr(nullptr)
-#endif
+	, BodyInstancePtr(nullptr)
+#endif // WITH_BOX2D
+#if WITH_PHYSX
+	, CurrentSceneState(BodyInstanceSceneState::NotAdded)
+#endif // WITH_PHYSX
+	, SleepFamily(ESleepFamily::Normal)
+	, DOFMode(0)
+	, CollisionEnabled(ECollisionEnabled::QueryAndPhysics)
+	, ObjectType(ECC_WorldStatic)
 {
 	MaxAngularVelocity = UPhysicsSettings::Get()->MaxAngularVelocity;
 }
