Fixed right click menus no longer working in some viewport clients.

diff --git a/Engine/Source/Runtime/Engine/Private/Slate/SceneViewport.cpp b/Engine/Source/Runtime/Engine/Private/Slate/SceneViewport.cpp
--- a/Engine/Source/Runtime/Engine/Private/Slate/SceneViewport.cpp
+++ b/Engine/Source/Runtime/Engine/Private/Slate/SceneViewport.cpp
@@ -372,13 +372,19 @@ FReply FSceneViewport::OnMouseButtonDown( const FGeometry& InGeometry, const FPo
 			ApplyModifierKeys( KeysState );
 		}
 
+		const bool bAnyMenusVisible = FSlateApplication::Get().AnyMenusVisible();
+
 		// Process the mouse event
 		if (!ViewportClient->InputKey(this, InMouseEvent.GetUserIndex(), InMouseEvent.GetEffectingButton(), IE_Pressed))
 		{
 			CurrentReplyState = FReply::Unhandled(); 
 		}
 
+		// a new menu was opened if there was previously not a menu visible but now there is
+		const bool bNewMenuWasOpened = !bAnyMenusVisible && FSlateApplication::Get().AnyMenusVisible();
+
 		if (!ViewportClient->IgnoreInput() &&
+			!bNewMenuWasOpened && // We should not focus the viewport if a menu was opened as it would close the menu
 			( ViewportClient->CaptureMouseOnClick() == EMouseCaptureMode::CapturePermanently ||
 			  ViewportClient->CaptureMouseOnClick() == EMouseCaptureMode::CaptureDuringMouseDown ||
 			  ( ViewportClient->CaptureMouseOnClick() == EMouseCaptureMode::CaptureDuringRightMouseDown && InMouseEvent.GetEffectingButton() == EKeys::RightMouseButton ) ) )
