TActorIterator now skips Actors in hidden levels the same as FActorIterator

diff --git a/Engine/Source/Runtime/Engine/Public/EngineUtils.h b/Engine/Source/Runtime/Engine/Public/EngineUtils.h
--- a/Engine/Source/Runtime/Engine/Public/EngineUtils.h
+++ b/Engine/Source/Runtime/Engine/Public/EngineUtils.h
@@ -474,6 +474,18 @@ public:
 	{
 		return **this;
 	}
+
+private:
+	/**
+	 * Used to examine whether this level is valid for iteration or not
+	 *
+	 * @param Level the level to check for iteration
+	 * @return true if the level can be iterated, false otherwise
+	 */
+	static bool CanIterateLevel(ULevel* Level)
+	{
+		return Level->bIsVisible;
+	}
 };
 
 /**
