Add a remove player static function

diff --git a/Engine/Source/Runtime/Engine/Classes/Kismet/GameplayStatics.h b/Engine/Source/Runtime/Engine/Classes/Kismet/GameplayStatics.h
--- a/Engine/Source/Runtime/Engine/Classes/Kismet/GameplayStatics.h
+++ b/Engine/Source/Runtime/Engine/Classes/Kismet/GameplayStatics.h
@@ -111,6 +111,13 @@ class ENGINE_API UGameplayStatics : public UBlueprintFunctionLibrary
 	UFUNCTION(BlueprintCallable, Category="Game", meta=(WorldContext="WorldContextObject", AdvancedDisplay="2", UnsafeDuringActorConstruction="true"))
 	static class APlayerController* CreatePlayer(UObject* WorldContextObject, int32 ControllerId = -1, bool bSpawnPawn = true);
 
+	/** Removes a player from this game.  
+	 *  @param Player			The player controller of the player to be removed
+	 *  @param bDestroyPawn		Whether the controlled pawn should be deleted as well
+	 */
+	UFUNCTION(BlueprintCallable, Category="Game", meta=(UnsafeDuringActorConstruction="true"))
+	static void RemovePlayer(APlayerController* Player, bool bDestroyPawn);
+
 	// --- Level Streaming functions ------------------------
 	
 	/** Stream the level with the LevelName ; Calling again before it finishes has no effect */
diff --git a/Engine/Source/Runtime/Engine/Private/GameplayStatics.cpp b/Engine/Source/Runtime/Engine/Private/GameplayStatics.cpp
--- a/Engine/Source/Runtime/Engine/Private/GameplayStatics.cpp
+++ b/Engine/Source/Runtime/Engine/Private/GameplayStatics.cpp
@@ -102,6 +102,24 @@ APlayerController* UGameplayStatics::CreatePlayer(UObject* WorldContextObject, i
 	return (LocalPlayer ? LocalPlayer->PlayerController : NULL);
 }
 
+void UGameplayStatics::RemovePlayer(APlayerController* PlayerController, bool bDestroyPawn)
+{
+	if (PlayerController)
+	{
+		if (UWorld* World = PlayerController->GetWorld())
+		{
+			if (ULocalPlayer* LocalPlayer = PlayerController->GetLocalPlayer())
+			{
+				APawn* PlayerPawn = (bDestroyPawn ? PlayerController->GetPawn() : nullptr);
+				if (World->GetGameInstance()->RemoveLocalPlayer(LocalPlayer) && PlayerPawn)
+				{
+					PlayerPawn->Destroy();
+				}
+			}
+		}
+	}
+}
+
 AGameMode* UGameplayStatics::GetGameMode(UObject* WorldContextObject)
 {
 	UWorld* const World = GEngine->GetWorldFromContextObject(WorldContextObject);
