Chagned from Empty to Reset for curve realloc

diff --git a/Engine/Source/Runtime/Engine/Private/Animation/AnimCurveTypes.cpp b/Engine/Source/Runtime/Engine/Private/Animation/AnimCurveTypes.cpp
--- a/Engine/Source/Runtime/Engine/Private/Animation/AnimCurveTypes.cpp
+++ b/Engine/Source/Runtime/Engine/Private/Animation/AnimCurveTypes.cpp
@@ -436,7 +436,7 @@ void FBlendedCurve::Set(USkeleton::AnimCurveUID InUid, float InValue, int32 InFl
 
 void FBlendedCurve::Reset(int32 Count)
 {
-	Elements.Empty(Count);
+	Elements.Reset();
 	Elements.Reserve(Count);
 }
 
diff --git a/Engine/Source/Runtime/Engine/Public/Animation/AnimCurveTypes.h b/Engine/Source/Runtime/Engine/Public/Animation/AnimCurveTypes.h
--- a/Engine/Source/Runtime/Engine/Public/Animation/AnimCurveTypes.h
+++ b/Engine/Source/Runtime/Engine/Public/Animation/AnimCurveTypes.h
@@ -270,7 +270,7 @@ struct ENGINE_API FBlendedCurve
 	void InitFrom(const FBlendedCurve& InCurveToInitFrom)
 	{
 		UIDList = InCurveToInitFrom.UIDList;
-		Elements.Empty(UIDList.Num());
+		Elements.Reset(UIDList.Num());
 		Elements.AddZeroed(UIDList.Num());
 		bInitialized = true;
 	}
@@ -316,8 +316,8 @@ struct ENGINE_API FBlendedCurve
 	/** Empty */
 	void Empty()
 	{
-		UIDList.Empty();
-		Elements.Empty();
+		UIDList.Reset();
+		Elements.Reset();
 	}
 private:
 	/**  Whether initialized or not */
