Change display names of timer functions that use delegates to event to better match the Create Event node that is used with it.

diff --git a/Engine/Source/Runtime/Engine/Classes/Kismet/KismetSystemLibrary.h b/Engine/Source/Runtime/Engine/Classes/Kismet/KismetSystemLibrary.h
--- a/Engine/Source/Runtime/Engine/Classes/Kismet/KismetSystemLibrary.h
+++ b/Engine/Source/Runtime/Engine/Classes/Kismet/KismetSystemLibrary.h
@@ -298,74 +298,74 @@ class ENGINE_API UKismetSystemLibrary : public UBlueprintFunctionLibrary
 
 	/**
 	 * Set a timer to execute delegate. Setting an existing timer will reset that timer with updated parameters.
-	 * @param Delegate		Delegate. Can be a K2 function or a Custom Event.
+	 * @param Event			Event. Can be a K2 function or a Custom Event.
 	 * @param Time			How long to wait before executing the delegate, in seconds. Setting a timer to <= 0 seconds will clear it if it is set.
 	 * @param bLooping		True to keep executing the delegate every Time seconds, false to execute delegate only once.
 	 * @return				The timer handle to pass to other timer functions to manipulate this timer.
 	 */
-	UFUNCTION(BlueprintCallable, meta=(DisplayName = "Set Timer by Delegate"), Category="Utilities|Time")
-	static FTimerHandle K2_SetTimerDelegate(FTimerDynamicDelegate Delegate, float Time, bool bLooping);
+	UFUNCTION(BlueprintCallable, meta=(DisplayName = "Set Timer by Event"), Category="Utilities|Time")
+	static FTimerHandle K2_SetTimerDelegate(UPARAM(DisplayName="Event") FTimerDynamicDelegate Delegate, float Time, bool bLooping);
 
 	/**
 	 * Clears a set timer.
-	 * @param Delegate		Delegate. Can be a K2 function or a Custom Event.
+	 * @param Event  Can be a K2 function or a Custom Event.
 	 */
-	UFUNCTION(BlueprintCallable, meta=(DeprecatedFunction, DeprecationMessage = "Use Clear Timer by Handle", DisplayName = "Clear Timer by Delegate"), Category="Utilities|Time")
-	static void K2_ClearTimerDelegate(FTimerDynamicDelegate Delegate);
+	UFUNCTION(BlueprintCallable, meta=(DeprecatedFunction, DeprecationMessage = "Use Clear Timer by Handle", DisplayName = "Clear Timer by Event"), Category="Utilities|Time")
+	static void K2_ClearTimerDelegate(UPARAM(DisplayName="Event") FTimerDynamicDelegate Delegate);
 
 	/**
 	 * Pauses a set timer at its current elapsed time.
-	 * @param Delegate		Delegate. Can be a K2 function or a Custom Event.
+	 * @param Event  Can be a K2 function or a Custom Event.
 	 */
-	UFUNCTION(BlueprintCallable, meta=(DeprecatedFunction, DeprecationMessage = "Use Pause Timer by Handle", DisplayName = "Pause Timer by Delegate"), Category="Utilities|Time")
-	static void K2_PauseTimerDelegate(FTimerDynamicDelegate Delegate);
+	UFUNCTION(BlueprintCallable, meta=(DeprecatedFunction, DeprecationMessage = "Use Pause Timer by Handle", DisplayName = "Pause Timer by Event"), Category="Utilities|Time")
+	static void K2_PauseTimerDelegate(UPARAM(DisplayName="Event") FTimerDynamicDelegate Delegate);
 
 	/**
 	 * Resumes a paused timer from its current elapsed time.
-	 * @param Delegate		Delegate. Can be a K2 function or a Custom Event.
+	 * @param Event  Can be a K2 function or a Custom Event.
 	 */
-	UFUNCTION(BlueprintCallable, meta=(DeprecatedFunction, DeprecationMessage = "Use Unpause Timer by Handle", DisplayName = "Unpause Timer by Delegate"), Category="Utilities|Time")
-	static void K2_UnPauseTimerDelegate(FTimerDynamicDelegate Delegate);
+	UFUNCTION(BlueprintCallable, meta=(DeprecatedFunction, DeprecationMessage = "Use Unpause Timer by Handle", DisplayName = "Unpause Timer by Event"), Category="Utilities|Time")
+	static void K2_UnPauseTimerDelegate(UPARAM(DisplayName="Event") FTimerDynamicDelegate Delegate);
 
 	/**
 	 * Returns true if a timer exists and is active for the given delegate, false otherwise.
-	 * @param Delegate		Delegate. Can be a K2 function or a Custom Event.
+	 * @param Event  Can be a K2 function or a Custom Event.
 	 * @return				True if the timer exists and is active.
 	 */
-	UFUNCTION(BlueprintPure, meta=(DeprecatedFunction, DeprecationMessage = "Use Is Timer Active by Handle", DisplayName = "Is Timer Active by Delegate"), Category="Utilities|Time")
-	static bool K2_IsTimerActiveDelegate(FTimerDynamicDelegate Delegate);
+	UFUNCTION(BlueprintPure, meta=(DeprecatedFunction, DeprecationMessage = "Use Is Timer Active by Handle", DisplayName = "Is Timer Active by Event"), Category="Utilities|Time")
+	static bool K2_IsTimerActiveDelegate(UPARAM(DisplayName="Event") FTimerDynamicDelegate Delegate);
 
 	/**
 	 * Returns true if a timer exists and is paused for the given delegate, false otherwise.
-	 * @param Delegate		Delegate. Can be a K2 function or a Custom Event.
+	 * @param Event  Can be a K2 function or a Custom Event.
 	 * @return				True if the timer exists and is paused.
 	 */
-	UFUNCTION(BlueprintPure, meta=(DeprecatedFunction, DeprecationMessage = "Use Is Timer Paused by Handle", DisplayName = "Is Timer Paused by Delegate"), Category = "Utilities|Time")
-	static bool K2_IsTimerPausedDelegate(FTimerDynamicDelegate Delegate);
+	UFUNCTION(BlueprintPure, meta=(DeprecatedFunction, DeprecationMessage = "Use Is Timer Paused by Handle", DisplayName = "Is Timer Paused by Event"), Category = "Utilities|Time")
+	static bool K2_IsTimerPausedDelegate(UPARAM(DisplayName="Event") FTimerDynamicDelegate Delegate);
 
 	/**
 	 * Returns true is a timer for the given delegate exists, false otherwise.
-	 * @param Delegate		Delegate. Can be a K2 function or a Custom Event.
+	 * @param Event  Can be a K2 function or a Custom Event.
 	 * @return				True if the timer exists.
 	 */
-	UFUNCTION(BlueprintPure, meta=(DeprecatedFunction, DeprecationMessage = "Use Does Timer Exist by Handle", DisplayName = "Does Timer Exist by Delegate"), Category = "Utilities|Time")
-	static bool K2_TimerExistsDelegate(FTimerDynamicDelegate Delegate);
+	UFUNCTION(BlueprintPure, meta=(DeprecatedFunction, DeprecationMessage = "Use Does Timer Exist by Handle", DisplayName = "Does Timer Exist by Event"), Category = "Utilities|Time")
+	static bool K2_TimerExistsDelegate(UPARAM(DisplayName="Event") FTimerDynamicDelegate Delegate);
 	
 	/**
 	 * Returns elapsed time for the given delegate (time since current countdown iteration began).
-	 * @param Delegate		Delegate. Can be a K2 function or a Custom Event.
+	 * @param Event  Can be a K2 function or a Custom Event.
 	 * @return				How long has elapsed since the current iteration of the timer began.
 	 */
-	UFUNCTION(BlueprintPure, meta=(DeprecatedFunction, DeprecationMessage = "Use Get Timer Elapsed Time by Handle", DisplayName = "Get Timer Elapsed Time by Delegate"), Category="Utilities|Time")
-	static float K2_GetTimerElapsedTimeDelegate(FTimerDynamicDelegate Delegate);
+	UFUNCTION(BlueprintPure, meta=(DeprecatedFunction, DeprecationMessage = "Use Get Timer Elapsed Time by Handle", DisplayName = "Get Timer Elapsed Time by Event"), Category="Utilities|Time")
+	static float K2_GetTimerElapsedTimeDelegate(UPARAM(DisplayName="Event") FTimerDynamicDelegate Delegate);
 
 	/**
 	 * Returns time until the timer will next execute its delegate.
-	 * @param Delegate		Delegate. Can be a K2 function or a Custom Event.
+	 * @param Event  Can be a K2 function or a Custom Event.
 	 * @return				How long is remaining in the current iteration of the timer.
 	 */
-	UFUNCTION(BlueprintPure, meta=(DeprecatedFunction, DeprecationMessage = "Use Get Timer Remaining Time by Handle", DisplayName = "Get Timer Remaining Time by Delegate"), Category="Utilities|Time")
-	static float K2_GetTimerRemainingTimeDelegate(FTimerDynamicDelegate Delegate);
+	UFUNCTION(BlueprintPure, meta=(DeprecatedFunction, DeprecationMessage = "Use Get Timer Remaining Time by Handle", DisplayName = "Get Timer Remaining Time by Event"), Category="Utilities|Time")
+	static float K2_GetTimerRemainingTimeDelegate(UPARAM(DisplayName="Event") FTimerDynamicDelegate Delegate);
 
 	// --- Timer functions with handle input ----------
 
@@ -456,7 +456,7 @@ class ENGINE_API UKismetSystemLibrary : public UBlueprintFunctionLibrary
 	 * @param bLooping		true to keep executing the delegate every Time seconds, false to execute delegate only once.
 	 * @return				The timer handle to pass to other timer functions to manipulate this timer.
 	 */
-	UFUNCTION(BlueprintCallable, meta=(DeprecatedFunction, DeprecationMessage = "Use Set Timer by Delegate", DisplayName = "Set Timer by Function Name", DefaultToSelf = "Object"), Category="Utilities|Time")
+	UFUNCTION(BlueprintCallable, meta=(DeprecatedFunction, DeprecationMessage = "Use Set Timer by Event", DisplayName = "Set Timer by Function Name", DefaultToSelf = "Object"), Category="Utilities|Time")
 	static FTimerHandle K2_SetTimer(UObject* Object, FString FunctionName, float Time, bool bLooping);
 
 	/**
