PR #481: Expose DebugDrawCircle to blueprints (Contributed by [REDACTED_PERSON])

diff --git a/Engine/Source/Runtime/Engine/Classes/Kismet/KismetSystemLibrary.h b/Engine/Source/Runtime/Engine/Classes/Kismet/KismetSystemLibrary.h
--- a/Engine/Source/Runtime/Engine/Classes/Kismet/KismetSystemLibrary.h
+++ b/Engine/Source/Runtime/Engine/Classes/Kismet/KismetSystemLibrary.h
@@ -1233,6 +1233,10 @@ class ENGINE_API UKismetSystemLibrary : public UBlueprintFunctionLibrary
 	UFUNCTION(BlueprintCallable, Category="Rendering|Debug", meta=(WorldContext="WorldContextObject"))
 	static void DrawDebugLine(UObject* WorldContextObject, const FVector LineStart, const FVector LineEnd, FLinearColor LineColor, float Duration=0.f, float Thickness = 0.f);
 
+	/** Draw a debug circle! */
+	UFUNCTION(BlueprintCallable, Category="Rendering|Debug", meta=(WorldContext="WorldContextObject"))
+	static void DrawDebugCircle(UObject* WorldContextObject, FVector Center, float Radius, int32 NumSegments=12, FLinearColor LineColor = FLinearColor::White, float Duration=0.f, float Thickness=0.f, FVector YAxis=FVector(0.f,1.f,0.f),FVector ZAxis=FVector(0.f,0.f,1.f), bool bDrawAxis=false);
+	
 	/** Draw a debug point */
 	UFUNCTION(BlueprintCallable, Category="Rendering|Debug", meta=(WorldContext="WorldContextObject"))
 	static void DrawDebugPoint(UObject* WorldContextObject, const FVector Position, float Size, FLinearColor PointColor, float Duration=0.f);
diff --git a/Engine/Source/Runtime/Engine/Private/DrawDebugHelpers.cpp b/Engine/Source/Runtime/Engine/Private/DrawDebugHelpers.cpp
--- a/Engine/Source/Runtime/Engine/Private/DrawDebugHelpers.cpp
+++ b/Engine/Source/Runtime/Engine/Private/DrawDebugHelpers.cpp
@@ -310,7 +310,7 @@ void DrawDebugCoordinateSystem(const UWorld* InWorld, FVector const& AxisLoc, FR
 	}
 }
 
-static void InternalDrawDebugCircle(const UWorld* InWorld, const FMatrix& TransformMatrix, float Radius, int32 Segments, const FColor& Color, bool bPersistentLines, float LifeTime, uint8 DepthPriority)
+static void InternalDrawDebugCircle(const UWorld* InWorld, const FMatrix& TransformMatrix, float Radius, int32 Segments, const FColor& Color, bool bPersistentLines, float LifeTime, uint8 DepthPriority, float Thickness = 0.f)
 {
 	// no debug line drawing on dedicated server
 	if( GEngine->GetNetMode(InWorld) != NM_DedicatedServer )
@@ -337,14 +337,14 @@ static void InternalDrawDebugCircle(const UWorld* InWorld, const FMatrix& Transf
 				const FVector Vertex1 = Center + Radius * (AxisY * FMath::Cos(Angle) + AxisZ * FMath::Sin(Angle));
 				Angle += AngleStep;
 				const FVector Vertex2 = Center + Radius * (AxisY * FMath::Cos(Angle) + AxisZ * FMath::Sin(Angle));
-				Lines.Add(FBatchedLine(Vertex1, Vertex2, Color, LineLifeTime, 0.0f, DepthPriority));
+				Lines.Add(FBatchedLine(Vertex1, Vertex2, Color, LineLifeTime, Thickness, DepthPriority));
 			}
 			LineBatcher->DrawLines(Lines);
 		}
 	}
 }
 
-void DrawDebugCircle(const UWorld* InWorld, const FMatrix& TransformMatrix, float Radius, int32 Segments, const FColor& Color, bool bPersistentLines, float LifeTime, uint8 DepthPriority)
+void DrawDebugCircle(const UWorld* InWorld, const FMatrix& TransformMatrix, float Radius, int32 Segments, const FColor& Color, bool bPersistentLines, float LifeTime, uint8 DepthPriority, float Thickness, bool bDrawAxis)
 {
 	// no debug line drawing on dedicated server
 	if( GEngine->GetNetMode(InWorld) != NM_DedicatedServer )
@@ -356,21 +356,46 @@ void DrawDebugCircle(const UWorld* InWorld, const FMatrix& TransformMatrix, floa
 
 			// Need at least 4 segments
 			Segments = FMath::Max((Segments - 2) / 2, 4);
-			InternalDrawDebugCircle(InWorld, TransformMatrix, Radius, Segments, Color, bPersistentLines, LifeTime, DepthPriority);
-
-			const FVector Center = TransformMatrix.GetOrigin();
-			const FVector AxisY = TransformMatrix.GetScaledAxis( EAxis::Y );
-			const FVector AxisZ = TransformMatrix.GetScaledAxis( EAxis::Z );
+			InternalDrawDebugCircle(InWorld, TransformMatrix, Radius, Segments, Color, bPersistentLines, LifeTime, DepthPriority, Thickness);
 
-			TArray<FBatchedLine> Lines;
-			Lines.Empty(2);
-			Lines.Add(FBatchedLine(Center - Radius * AxisY, Center + Radius * AxisY, Color, LineLifeTime, 0.0f, DepthPriority));
-			Lines.Add(FBatchedLine(Center - Radius * AxisZ, Center + Radius * AxisZ, Color, LineLifeTime, 0.0f, DepthPriority));
-			LineBatcher->DrawLines(Lines);
+			if (bDrawAxis)
+			{
+				const FVector Center = TransformMatrix.GetOrigin();
+				const FVector AxisY = TransformMatrix.GetScaledAxis( EAxis::Y );
+				const FVector AxisZ = TransformMatrix.GetScaledAxis( EAxis::Z );
+
+				TArray<FBatchedLine> Lines;
+				Lines.Empty(2);
+				Lines.Add(FBatchedLine(Center - Radius * AxisY, Center + Radius * AxisY, Color, LineLifeTime, Thickness, DepthPriority));
+				Lines.Add(FBatchedLine(Center - Radius * AxisZ, Center + Radius * AxisZ, Color, LineLifeTime, Thickness, DepthPriority));
+				LineBatcher->DrawLines(Lines);
+			}
 		}
 	}
 }
 
+void DrawDebugCircle(const UWorld* InWorld, FVector Center, float Radius, int32 Segments, const FColor& Color, bool PersistentLines, float LifeTime, uint8 DepthPriority, float Thickness, FVector YAxis, FVector ZAxis, bool bDrawAxis)
+{
+	FMatrix TM;
+	TM.SetOrigin(Center);
+	TM.SetAxis(0, FVector(1,0,0));
+	TM.SetAxis(1, YAxis);
+	TM.SetAxis(2, ZAxis);
+	
+	DrawDebugCircle(
+		InWorld,
+		TM,
+		Radius,
+		Segments,
+		Color,
+		PersistentLines,
+		LifeTime,
+		DepthPriority,
+		Thickness,
+		bDrawAxis
+	);
+}
+
 void DrawDebug2DDonut(const UWorld* InWorld, const FMatrix& TransformMatrix, float InnerRadius, float OuterRadius, int32 Segments, const FColor& Color, bool bPersistentLines, float LifeTime, uint8 DepthPriority)
 {
 	// no debug line drawing on dedicated server
diff --git a/Engine/Source/Runtime/Engine/Private/KismetSystemLibrary.cpp b/Engine/Source/Runtime/Engine/Private/KismetSystemLibrary.cpp
--- a/Engine/Source/Runtime/Engine/Private/KismetSystemLibrary.cpp
+++ b/Engine/Source/Runtime/Engine/Private/KismetSystemLibrary.cpp
@@ -2652,6 +2652,12 @@ void UKismetSystemLibrary::DrawDebugLine(UObject* WorldContextObject, FVector co
 	}
 }
 
+/** Draw a debug circle */
+void UKismetSystemLibrary::DrawDebugCircle(UObject* WorldContextObject,FVector Center, float Radius, int32 NumSegments, FLinearColor LineColor, float LifeTime, float Thickness, FVector YAxis, FVector ZAxis, bool bDrawAxis)
+{ 
+	::DrawDebugCircle(GEngine->GetWorldFromContextObject(WorldContextObject),Center, Radius, NumSegments, LineColor.ToFColor(true), false, LifeTime, SDPG_World, Thickness, YAxis, ZAxis, bDrawAxis);
+}
+
 /** Draw a debug point */
 void UKismetSystemLibrary::DrawDebugPoint(UObject* WorldContextObject, FVector const Position, float Size, FLinearColor PointColor, float LifeTime)
 {
diff --git a/Engine/Source/Runtime/Engine/Public/DrawDebugHelpers.h b/Engine/Source/Runtime/Engine/Public/DrawDebugHelpers.h
--- a/Engine/Source/Runtime/Engine/Public/DrawDebugHelpers.h
+++ b/Engine/Source/Runtime/Engine/Public/DrawDebugHelpers.h
@@ -22,7 +22,9 @@ ENGINE_API void DrawDebugBox(const UWorld* InWorld, FVector const& Center, FVect
 /** Draw Debug coordinate system */
 ENGINE_API void DrawDebugCoordinateSystem(const UWorld* InWorld, FVector const& AxisLoc, FRotator const& AxisRot, float Scale, bool bPersistentLines = false, float LifeTime=-1.f, uint8 DepthPriority = 0);
 /** Draw Debug Circle */
-ENGINE_API void DrawDebugCircle(const UWorld* InWorld, const FMatrix& TransformMatrix, float Radius, int32 Segments, const FColor& Color, bool bPersistentLines = false, float LifeTime=-1.f, uint8 DepthPriority = 0);
+ENGINE_API void DrawDebugCircle(const UWorld* InWorld, const FMatrix& TransformMatrix, float Radius, int32 Segments, const FColor& Color, bool bPersistentLines = false, float LifeTime=-1.f, uint8 DepthPriority = 0, float Thickness = 0.f, bool bDrawAxis=true);
+/** Draw Debug Circle */
+ENGINE_API void DrawDebugCircle(const UWorld* InWorld, FVector Center, float Radius, int32 Segments, const FColor& Color, bool bPersistentLines = false, float LifeTime=-1.f, uint8 DepthPriority = 0, float Thickness = 0.f, FVector YAxis=FVector(0.f,1.f,0.f), FVector ZAxis=FVector(0.f,0.f,1.f), bool bDrawAxis=true);
 /** Draw Debug 2D donut */
 ENGINE_API void DrawDebug2DDonut(const UWorld* InWorld, const FMatrix& TransformMatrix, float InnerRadius, float OuterRadius, int32 Segments, const FColor& Color, bool bPersistentLines = false, float LifeTime=-1.f, uint8 DepthPriority = 0);
 /** Draw a debug sphere */
