Changed encroachment checking so that actors with a movement component only check the updated component.

diff --git a/Engine/Source/Runtime/Engine/Private/LevelActor.cpp b/Engine/Source/Runtime/Engine/Private/LevelActor.cpp
--- a/Engine/Source/Runtime/Engine/Private/LevelActor.cpp
+++ b/Engine/Source/Runtime/Engine/Private/LevelActor.cpp
@@ -11,17 +11,19 @@
 #if WITH_PHYSX
 	#include "PhysicsEngine/PhysXSupport.h"
 #endif // WITH_PHYSX
+
 #include "Components/SphereComponent.h"
 #include "Components/CapsuleComponent.h"
 #include "Components/BoxComponent.h"
+#include "GameFramework/MovementComponent.h"
 #include "GameFramework/GameMode.h"
 
 // CVars
 static TAutoConsoleVariable<float> CVarEncroachEpsilon(
 	TEXT("p.EncroachEpsilon"),
 	0.15f,
-	TEXT("Epsilon value encroach checking for capsules\n")
-	TEXT("0: use full sized capsule. > 0: shrink capsule size by this amount (world units)"),
+	TEXT("Epsilon value used during encroachment checking for shape components\n")
+	TEXT("0: use full sized shape. > 0: shrink shape size by this amount (world units)"),
 	ECVF_Default);
 
 
@@ -936,55 +938,87 @@ bool UWorld::EncroachingBlockingGeometry(AActor* TestActor, FVector TestLocation
 	}
 
 	bool bFoundEncroacher = false;
-	FVector TotalAdjustment(0.f);
 
+	FVector TotalAdjustment(0.f);
 	FTransform const TestRootToWorld = FTransform(TestRotation, TestLocation);
-	
-	UPrimitiveComponent* const RootPrimComp = Cast<UPrimitiveComponent>(RootComponent);
-	if (RootPrimComp && RootPrimComp->IsCollisionEnabled())
+	FTransform const WorldToOldRoot = RootComponent->GetComponentToWorld().Inverse();
+
+	UMovementComponent* const MoveComponent = TestActor->FindComponentByClass<UMovementComponent>();
+	if (MoveComponent && MoveComponent->UpdatedPrimitive)
 	{
-		if (ComponentEncroachesBlockingGeometry(this, TestActor, RootPrimComp, TestRootToWorld, ProposedAdjustment))
+		// This actor has a movement component, which we interpret to mean that this actor has a primary component being swept around
+		// the world, and that component is the only one we care about encroaching (since the movement code will happily embedding
+		// other components in the world during movement updates)
+		UPrimitiveComponent* const MovedPrimComp = MoveComponent->UpdatedPrimitive;
+		if (MovedPrimComp->IsCollisionEnabled())
 		{
-			if (ProposedAdjustment == nullptr)
-			{
-				// don't need an adjustment and we know we are overlapping, so we can be done
-				return true;
-			}
-			else
+			// might not be the root, so we need to compute the transform
+			FTransform const CompToRoot = MovedPrimComp->GetComponentToWorld() * WorldToOldRoot;
+			FTransform const CompToNewWorld = CompToRoot * TestRootToWorld;
+
+			if (ComponentEncroachesBlockingGeometry(this, TestActor, MovedPrimComp, CompToNewWorld, ProposedAdjustment))
 			{
-				TotalAdjustment = *ProposedAdjustment;
-			}
+				if (ProposedAdjustment == nullptr)
+				{
+					// don't need an adjustment and we know we are overlapping, so we can be done
+					return true;
+				}
+				else
+				{
+					TotalAdjustment = *ProposedAdjustment;
+				}
 
-			bFoundEncroacher = true;
+				bFoundEncroacher = true;
+			}
 		}
 	}
+	else
+	{
+		// This actor does not have a movement component, so we'll assume all components are potentially important to keep out of the world
+		UPrimitiveComponent* const RootPrimComp = Cast<UPrimitiveComponent>(RootComponent);
+		if (RootPrimComp && RootPrimComp->IsCollisionEnabled())
+		{
+			if (ComponentEncroachesBlockingGeometry(this, TestActor, RootPrimComp, TestRootToWorld, ProposedAdjustment))
+			{
+				if (ProposedAdjustment == nullptr)
+				{
+					// don't need an adjustment and we know we are overlapping, so we can be done
+					return true;
+				}
+				else
+				{
+					TotalAdjustment = *ProposedAdjustment;
+				}
 
-	// now test all colliding children for encroachment
-	TArray<USceneComponent*> Children;
-	RootComponent->GetChildrenComponents(true, Children);
+				bFoundEncroacher = true;
+			}
+		}
 
-	FTransform const WorldToOldRoot = RootComponent->GetComponentToWorld().Inverse();
+		// now test all colliding children for encroachment
+		TArray<USceneComponent*> Children;
+		RootComponent->GetChildrenComponents(true, Children);
 
-	for (auto Child : Children)
-	{
-		if (Child->IsCollisionEnabled())
+		for (auto Child : Children)
 		{
-			UPrimitiveComponent* const PrimComp = Cast<UPrimitiveComponent>(Child);
-			if (PrimComp)
+			if (Child->IsCollisionEnabled())
 			{
-				FTransform const CompToRoot = Child->GetComponentToWorld() * WorldToOldRoot;
-				FTransform const CompToNewWorld = CompToRoot * TestRootToWorld;
-
-				if (ComponentEncroachesBlockingGeometry(this, TestActor, PrimComp, CompToNewWorld, ProposedAdjustment))
+				UPrimitiveComponent* const PrimComp = Cast<UPrimitiveComponent>(Child);
+				if (PrimComp)
 				{
-					if (ProposedAdjustment == nullptr)
+					FTransform const CompToRoot = Child->GetComponentToWorld() * WorldToOldRoot;
+					FTransform const CompToNewWorld = CompToRoot * TestRootToWorld;
+
+					if (ComponentEncroachesBlockingGeometry(this, TestActor, PrimComp, CompToNewWorld, ProposedAdjustment))
 					{
-						// don't need an adjustment and we know we are overlapping, so we can be done
-						return true;
-					}
+						if (ProposedAdjustment == nullptr)
+						{
+							// don't need an adjustment and we know we are overlapping, so we can be done
+							return true;
+						}
 
-					TotalAdjustment = CombineAdjustments(TotalAdjustment, *ProposedAdjustment);
-					bFoundEncroacher = true;
+						TotalAdjustment = CombineAdjustments(TotalAdjustment, *ProposedAdjustment);
+						bFoundEncroacher = true;
+					}
 				}
 			}
 		}
