nicer order in PostprocessSettings for better usability

diff --git a/Engine/Source/Runtime/Engine/Classes/Engine/Scene.h b/Engine/Source/Runtime/Engine/Classes/Engine/Scene.h
--- a/Engine/Source/Runtime/Engine/Classes/Engine/Scene.h
+++ b/Engine/Source/Runtime/Engine/Classes/Engine/Scene.h
@@ -829,14 +829,6 @@ struct FPostProcessSettings
 	UPROPERTY(EditAnywhere, BlueprintReadWrite, Category=DepthOfField, meta=(editcondition = "bOverride_DepthOfFieldMethod", DisplayName = "Method"))
 	TEnumAsByte<enum EDepthOfFieldMethod> DepthOfFieldMethod;
 
-	/** CircleDOF only: Depth blur km for 50% */
-	UPROPERTY(interp, BlueprintReadWrite, Category=DepthOfField, meta=(ClampMin = "0.000001", ClampMax = "100.0", editcondition = "bOverride_DepthOfFieldDepthBlurAmount", DisplayName = "Depth Blur km for 50%"))
-	float DepthOfFieldDepthBlurAmount;
-
-	/** CircleDOF only: Depth blur radius in pixels at 1920x */
-	UPROPERTY(interp, BlueprintReadWrite, Category=DepthOfField, meta=(ClampMin = "0.0", ClampMax = "4.0", editcondition = "bOverride_DepthOfFieldDepthBlurRadius", DisplayName = "Depth Blur Radius"))
-	float DepthOfFieldDepthBlurRadius;
-	
 	/** CircleDOF only: Defines the opening of the camera lens, Aperture is 1/fstop, typical lens go down to f/1.2 (large opening), larger numbers reduce the DOF effect */
 	UPROPERTY(interp, BlueprintReadWrite, Category=DepthOfField, meta=(ClampMin = "1.0", ClampMax = "32.0", editcondition = "bOverride_DepthOfFieldFstop", DisplayName = "Aperture F-stop"))
 	float DepthOfFieldFstop;
@@ -845,6 +837,14 @@ struct FPostProcessSettings
 	UPROPERTY(interp, BlueprintReadWrite, Category=DepthOfField, meta=(ClampMin = "0.0", UIMin = "1.0", UIMax = "10000.0", editcondition = "bOverride_DepthOfFieldFocalDistance", DisplayName = "Focal Distance"))
 	float DepthOfFieldFocalDistance;
 
+	/** CircleDOF only: Depth blur km for 50% */
+	UPROPERTY(interp, BlueprintReadWrite, AdvancedDisplay, Category=DepthOfField, meta=(ClampMin = "0.000001", ClampMax = "100.0", editcondition = "bOverride_DepthOfFieldDepthBlurAmount", DisplayName = "Depth Blur km for 50%"))
+	float DepthOfFieldDepthBlurAmount;
+
+	/** CircleDOF only: Depth blur radius in pixels at 1920x */
+	UPROPERTY(interp, BlueprintReadWrite, AdvancedDisplay, Category=DepthOfField, meta=(ClampMin = "0.0", ClampMax = "4.0", editcondition = "bOverride_DepthOfFieldDepthBlurRadius", DisplayName = "Depth Blur Radius"))
+	float DepthOfFieldDepthBlurRadius;
+
 	/** Artificial region where all content is in focus, starting after DepthOfFieldFocalDistance, in unreal units  (cm) */
 	UPROPERTY(interp, BlueprintReadWrite, Category=DepthOfField, meta=(UIMin = "0.0", UIMax = "10000.0", editcondition = "bOverride_DepthOfFieldFocalRegion", DisplayName = "Focal Region"))
 	float DepthOfFieldFocalRegion;
