nicer order in PostprocessSettings for better usability

diff --git a/Engine/Source/Runtime/Engine/Private/Materials/MaterialShared.cpp b/Engine/Source/Runtime/Engine/Private/Materials/MaterialShared.cpp
--- a/Engine/Source/Runtime/Engine/Private/Materials/MaterialShared.cpp
+++ b/Engine/Source/Runtime/Engine/Private/Materials/MaterialShared.cpp
@@ -2134,7 +2134,7 @@ FMaterialUpdateContext::~FMaterialUpdateContext()
 
 	double EndTime = FPlatformTime::Seconds();
 	UE_LOG(LogMaterial, Log,
-		TEXT("%f seconds spent updating %d materials, %d interfaces, %d instances, %d with static permutations."),
+		TEXT("%.2f seconds spent updating %d materials, %d interfaces, %d instances, %d with static permutations."),
 		(float)(EndTime - StartTime),
 		UpdatedMaterials.Num(),
 		UpdatedMaterialInterfaces.Num(),
