[UE-18596] Fix for lerp node keywords not being set (typing lerp in the material editor context menu would not find linearinterpolate node), and fix for SearchTitleArray missing from move and copy constructor, this caused various scoring problems

diff --git a/Engine/Source/Runtime/Engine/Private/EdGraph/EdGraphSchema.cpp b/Engine/Source/Runtime/Engine/Private/EdGraph/EdGraphSchema.cpp
--- a/Engine/Source/Runtime/Engine/Private/EdGraph/EdGraphSchema.cpp
+++ b/Engine/Source/Runtime/Engine/Private/EdGraph/EdGraphSchema.cpp
@@ -125,6 +125,7 @@ void FGraphActionListBuilderBase::ActionGroup::Move(FGraphActionListBuilderBase:
 	CategoryChain = MoveTemp(Other.CategoryChain);
 	SearchKeywordsArray = MoveTemp(Other.SearchKeywordsArray);
 	MenuDescriptionArray = MoveTemp(Other.MenuDescriptionArray);
+	SearchTitleArray = MoveTemp(Other.SearchTitleArray);
 	SearchCategoryArray = MoveTemp(Other.SearchCategoryArray);
 	SearchText = MoveTemp(Other.SearchText);
 }
@@ -136,6 +137,7 @@ void FGraphActionListBuilderBase::ActionGroup::Copy(const ActionGroup& Other)
 	CategoryChain = Other.CategoryChain;
 	SearchKeywordsArray = Other.SearchKeywordsArray;
 	MenuDescriptionArray = Other.MenuDescriptionArray;
+	SearchTitleArray = Other.SearchTitleArray;
 	SearchCategoryArray = Other.SearchCategoryArray;
 	SearchText = Other.SearchText;
 }
