Added shader compilation stats

diff --git a/Engine/Source/Runtime/Engine/Private/Materials/MaterialShader.cpp b/Engine/Source/Runtime/Engine/Private/Materials/MaterialShader.cpp
--- a/Engine/Source/Runtime/Engine/Private/Materials/MaterialShader.cpp
+++ b/Engine/Source/Runtime/Engine/Private/Materials/MaterialShader.cpp
@@ -16,6 +16,132 @@ static FAutoConsoleVariableRef CVarCreateShadersOnLoad(
 	TEXT("Whether to create shaders on load, which can reduce hitching, but use more memory.  Otherwise they will be created as needed.")
 	);
 
+
+
+//
+// Cooking Stats
+// 
+#define MATERIAL_COOKING_STATS 1
+
+#if MATERIAL_COOKING_STATS
+
+#include "CookingStatsModule.h"
+#include "ModuleManager.h"
+
+
+FString GetMaterialShaderMapKeyString(const FMaterialShaderMapId& ShaderMapId, EShaderPlatform Platform);
+
+namespace MaterialStats
+{
+
+	ICookingStats* GetCookingStats()
+	{
+		static ICookingStats* CookingStats = nullptr;
+		static bool bInitialized = false;
+		if (bInitialized == false)
+		{
+			FCookingStatsModule* CookingStatsModule = FModuleManager::LoadModulePtr<FCookingStatsModule>(TEXT("CookingStats"));
+			if (CookingStatsModule)
+			{
+				CookingStats = &CookingStatsModule->Get();
+			}
+			bInitialized = true;
+		}
+		return CookingStats;
+	}
+
+	FCriticalSection CookingStatsSyncObject;
+
+	struct FMaterialTimingInfo
+	{
+	public:
+		FName TransactionGuid;
+		TMap<FName, double> TagStartTime;
+	};
+
+	TMap<FMaterialShaderMapId, FMaterialTimingInfo> CurrentCookingStats;
+	int32 CurrentIndex;
+	FName TransactionGuid;
+
+	FName GenerateNewTransactionName()
+	{
+		FScopeLock ScopeSync(&CookingStatsSyncObject);
+		++CurrentIndex;
+		if (CurrentIndex <= 1)
+		{
+			FString TransactionGuidString = TEXT("MaterialTransactionId");
+			TransactionGuidString += FGuid::NewGuid().ToString();
+			TransactionGuid = FName(*TransactionGuidString);
+			CurrentIndex = 1;
+		}
+		check(TransactionGuid != NAME_None);
+
+		return FName(TransactionGuid, CurrentIndex);
+	}
+
+	void StartStat(const FMaterialShaderMapId& ShaderMapId, const EShaderPlatform Platform, const FName& TagName)
+	{
+		ICookingStats* CookingStats = GetCookingStats();
+		if (CookingStats)
+		{
+			FScopeLock ScopeSync(&CookingStatsSyncObject);
+
+			FMaterialTimingInfo* TimingInfo = CurrentCookingStats.Find(ShaderMapId);
+
+			if (TimingInfo == nullptr)
+			{
+				TimingInfo = &CurrentCookingStats.Add(ShaderMapId);
+				TimingInfo->TransactionGuid = GenerateNewTransactionName();
+				static const FName NAME_CacheKey(TEXT("CacheKey"));
+
+				FString CacheKeyString = GetMaterialShaderMapKeyString(ShaderMapId, Platform);
+
+				CookingStats->AddTagValue(TimingInfo->TransactionGuid, NAME_CacheKey, CacheKeyString);
+			}
+
+			double* StartTime = TimingInfo->TagStartTime.Find(TagName);
+			if (!StartTime)
+			{
+				TimingInfo->TagStartTime.Add(TagName, FPlatformTime::Seconds());
+			}
+		}
+	}
+
+	void StopStat(const FMaterialShaderMapId& ShaderMapId, const FName& TagName)
+	{
+		ICookingStats* CookingStats = GetCookingStats();
+		if (CookingStats)
+		{
+			FScopeLock ScopeSync(&CookingStatsSyncObject);
+
+			FMaterialTimingInfo* TimingInfo = CurrentCookingStats.Find(ShaderMapId);
+			check(TimingInfo); // did you call stop stat before calling start
+			double* StartTime = TimingInfo->TagStartTime.Find(TagName);
+			check(StartTime); // did you call stop stat before calling start for that tag
+			
+			double Duration = FPlatformTime::Seconds() - *StartTime;
+			TimingInfo->TagStartTime.Remove(TagName);
+
+			CookingStats->AddTagValue(TimingInfo->TransactionGuid, TagName, FString::Printf(TEXT("%fms"), Duration*1000.0f));
+		}
+	}
+
+};
+
+#define COOKING_STAT_START(TagName, ShaderMapID, Platform) \
+	static const FName NAME_##TagName(TEXT(#TagName)); \
+	MaterialStats::StartStat(ShaderMapID, Platform, NAME_##TagName);
+
+#define COOKING_STAT_STOP(TagName, ShaderMapID) \
+	static const FName NAME_##TagName(TEXT(#TagName)); \
+	MaterialStats::StopStat(ShaderMapID, NAME_##TagName);
+
+
+#else
+#endif
+
+
+
 //
 // Globals
 //
@@ -63,6 +189,19 @@ FString GetBlendModeString(EBlendMode BlendMode)
 	return BlendModeName;
 }
 
+
+/** Creates a string key for the derived data cache given a shader map id. */
+FString GetMaterialShaderMapKeyString(const FMaterialShaderMapId& ShaderMapId, EShaderPlatform Platform)
+{
+	FName Format = LegacyShaderPlatformToShaderFormat(Platform);
+	FString ShaderMapKeyString = Format.ToString() + TEXT("_") + FString(FString::FromInt(GetTargetPlatformManagerRef().ShaderFormatVersion(Format))) + TEXT("_");
+	ShaderMapAppendKeyString(Platform, ShaderMapKeyString);
+	ShaderMapId.AppendKeyString(ShaderMapKeyString);
+	return FDerivedDataCacheInterface::BuildCacheKey(TEXT("MATSM"), MATERIALSHADERMAP_DERIVEDDATA_VER, *ShaderMapKeyString);
+}
+
+
+
 /** Called for every material shader to update the appropriate stats. */
 void UpdateMaterialShaderCompilingStats(const FMaterial* Material)
 {
@@ -740,16 +879,6 @@ void FMaterialShaderMap::FixupShaderTypes(EShaderPlatform Platform, const TMap<F
 	}
 }
 
-/** Creates a string key for the derived data cache given a shader map id. */
-FString GetMaterialShaderMapKeyString(const FMaterialShaderMapId& ShaderMapId, EShaderPlatform Platform)
-{
-	FName Format = LegacyShaderPlatformToShaderFormat(Platform);
-	FString ShaderMapKeyString = Format.ToString() + TEXT("_") + FString(FString::FromInt(GetTargetPlatformManagerRef().ShaderFormatVersion(Format))) + TEXT("_");
-	ShaderMapAppendKeyString(Platform, ShaderMapKeyString);
-	ShaderMapId.AppendKeyString(ShaderMapKeyString);
-	return FDerivedDataCacheInterface::BuildCacheKey(TEXT("MATSM"), MATERIALSHADERMAP_DERIVEDDATA_VER, *ShaderMapKeyString);
-}
-
 void FMaterialShaderMap::LoadFromDerivedDataCache(const FMaterial* Material, const FMaterialShaderMapId& ShaderMapId, EShaderPlatform Platform, TRefCountPtr<FMaterialShaderMap>& InOutShaderMap)
 {
 	if (InOutShaderMap != NULL)
@@ -1026,6 +1155,8 @@ void FMaterialShaderMap::LoadForRemoteRecompile(FArchive& Ar, EShaderPlatform Sh
 }
 
 
+
+
 /**
 * Compiles the shaders for a material and caches them in this shader map.
 * @param Material - The material to compile shaders for.
@@ -1047,6 +1178,10 @@ void FMaterialShaderMap::Compile(
 	}
 	else
 	{
+		COOKING_STAT_START(ShaderCompilation, InShaderMapId, InPlatform);
+
+
+
 		check(!Material->bContainsInlineShaders);
   
 		// Make sure we are operating on a referenced shader map or the below Find will cause this shader map to be deleted,
@@ -1311,6 +1446,8 @@ bool FMaterialShaderMap::ProcessCompilationResults(const TArray<FShaderCompileJo
 		// The shader map can now be used on the rendering thread
 		bCompilationFinalized = true;
 
+		COOKING_STAT_STOP(ShaderCompilation, ShaderMapId)
+
 		return true;
 	}
 
