Fixed issue with audio ddc hanging when requiring building of audio.

diff --git a/Engine/Source/Runtime/Engine/Private/AudioDerivedData.cpp b/Engine/Source/Runtime/Engine/Private/AudioDerivedData.cpp
--- a/Engine/Source/Runtime/Engine/Private/AudioDerivedData.cpp
+++ b/Engine/Source/Runtime/Engine/Private/AudioDerivedData.cpp
@@ -1082,6 +1082,12 @@ bool USoundWave::IsCachedCookedPlatformDataLoaded( const ITargetPlatform* Target
 			// we havne't called begincache
 			return false;
 		}
+
+		if (PlatformData->AsyncTask && PlatformData->AsyncTask->IsWorkDone())
+		{
+			PlatformData->FinishCache();
+		}
+
 		return PlatformData->IsFinishedCache();
 	}
 	return true; 
