Stopped the builder brush from erroneously being rendered by 'show collision'.

diff --git a/Engine/Source/Runtime/Engine/Private/GameViewportClient.cpp b/Engine/Source/Runtime/Engine/Private/GameViewportClient.cpp
--- a/Engine/Source/Runtime/Engine/Private/GameViewportClient.cpp
+++ b/Engine/Source/Runtime/Engine/Private/GameViewportClient.cpp
@@ -32,6 +32,7 @@
 #include "GameFramework/GameUserSettings.h"
 #include "Runtime/Engine/Classes/Engine/UserInterfaceSettings.h"
 #include "SGameLayerManager.h"
+#include "ActorEditorUtils.h"
 
 #define LOCTEXT_NAMESPACE "GameViewport"
 
@@ -2298,7 +2299,7 @@ void UGameViewportClient::ToggleShowVolumes()
 		AVolume* Owner = Cast<AVolume>(BrushComponent->GetOwner());
 
 		// Only bother with volume brushes that belong to the world's scene
-		if (Owner && BrushComponent->GetScene() == GetWorld()->Scene)
+		if (Owner && BrushComponent->GetScene() == GetWorld()->Scene && !FActorEditorUtils::IsABuilderBrush(Owner))
 		{
 			// We're expecting this to be in the game at this point
 			check(Owner->GetWorld()->IsGameWorld());
@@ -2375,7 +2376,9 @@ void UGameViewportClient::ToggleShowCollision()
 			UPrimitiveComponent* PrimitiveComponent = *It;
 			if (!PrimitiveComponent->IsVisible() && PrimitiveComponent->IsCollisionEnabled() && PrimitiveComponent->GetScene() == World->Scene)
 			{
-				if (PrimitiveComponent->GetOwner() && PrimitiveComponent->GetOwner()->GetWorld() && PrimitiveComponent->GetOwner()->GetWorld()->IsGameWorld())
+				AActor* Owner = PrimitiveComponent->GetOwner();
+
+				if (Owner && Owner->GetWorld() && Owner->GetWorld()->IsGameWorld() && !FActorEditorUtils::IsABuilderBrush(Owner))
 				{
 					// Save state before modifying the collision visibility
 					Mapping.Add(PrimitiveComponent, CollVisibilityState(PrimitiveComponent->bHiddenInGame, PrimitiveComponent->bVisible));
