fixed geometry projection for navmesh walking mode using cached location even if pawn moved far away on Z axis
#ue4 UE-18790

diff --git a/Engine/Source/Runtime/Engine/Private/Components/CharacterMovementComponent.cpp b/Engine/Source/Runtime/Engine/Private/Components/CharacterMovementComponent.cpp
--- a/Engine/Source/Runtime/Engine/Private/Components/CharacterMovementComponent.cpp
+++ b/Engine/Source/Runtime/Engine/Private/Components/CharacterMovementComponent.cpp
@@ -3862,7 +3862,14 @@ void UCharacterMovementComponent::PhysNavWalking(float deltaTime, int32 Iteratio
 	{
 		if (bProjectNavMeshWalking)
 		{
-			bSameNavLocation = (OldLocation - CachedNavLocation.Location).SizeSquared2D() <= KINDA_SMALL_NUMBER;
+			const float DistSq2D = (OldLocation - CachedNavLocation.Location).SizeSquared2D();
+			const float DistZ = FMath::Abs(OldLocation.Z - CachedNavLocation.Location.Z);
+
+			const float TotalCapsuleHeight = CharacterOwner->GetCapsuleComponent()->GetScaledCapsuleHalfHeight() * 2.0f;
+			const float ProjectionScale = (OldLocation.Z > CachedNavLocation.Location.Z) ? NavMeshProjectionHeightScaleUp : NavMeshProjectionHeightScaleDown;
+			const float DistZThr = TotalCapsuleHeight * FMath::Max(0.f, ProjectionScale);
+
+			bSameNavLocation = (DistSq2D <= KINDA_SMALL_NUMBER) && (DistZ < DistZThr);
 		}
 		else
 		{
