UE4 - Change runaway recursive spawning from an assert to a spammy message.

diff --git a/Engine/Source/Runtime/Engine/Private/TickTaskManager.cpp b/Engine/Source/Runtime/Engine/Private/TickTaskManager.cpp
--- a/Engine/Source/Runtime/Engine/Private/TickTaskManager.cpp
+++ b/Engine/Source/Runtime/Engine/Private/TickTaskManager.cpp
@@ -702,6 +702,27 @@ public:
 		NewlySpawnedTickFunctions.Empty();
 		return Num;
 	}
+	/**
+	 * If there is infinite recursive spawning, log that and discard them
+	 */
+	void LogAndDisardRunawayNewlySpawned(ETickingGroup CurrentTickGroup)
+	{
+		Context.TickGroup = CurrentTickGroup;
+		FTickTaskSequencer& TTS = FTickTaskSequencer::Get();
+		for (TSet<FTickFunction*>::TIterator It(NewlySpawnedTickFunctions); It; ++It)
+		{
+			FTickFunction* TickFunction = *It;
+			UE_LOG(LogTick, Error, TEXT("Could not tick newly spawned in 100 iterations; runaway recursive spawing. Tick is %s."), *TickFunction->DiagnosticMessage());
+
+			if (TickFunction->TickInterval > 0.f)
+			{
+				AllEnabledTickFunctions.Remove(TickFunction);
+				TickFunctionsToReschedule.Add(FTickScheduleDetails(TickFunction, TickFunction->TickInterval));
+			}
+		}
+		ScheduleTickFunctionCooldowns();
+		NewlySpawnedTickFunctions.Empty();
+	}
 
 	/**
 	 * Run all of the ticks for a pause frame synchronously on the game thread.
@@ -1187,9 +1208,9 @@ public:
 		{
 			QUICK_SCOPE_CYCLE_COUNTER(STAT_TickTask_RunTickGroup_BlockTillComplete);
 
+			bool bFinished = false;
 			for (int32 Iterations = 0;Iterations < 101; Iterations++)
 			{
-				check(Iterations < 100); // this is runaway recursive spawning.
 				int32 Num = 0;
 				for( int32 LevelIndex = 0; LevelIndex < LevelList.Num(); LevelIndex++ )
 				{
@@ -1202,9 +1223,18 @@ public:
 				}
 				else
 				{
+					bFinished = true;
 					break;
 				}
 			}
+			if (!bFinished)
+			{
+				// this is runaway recursive spawning.
+				for( int32 LevelIndex = 0; LevelIndex < LevelList.Num(); LevelIndex++ )
+				{
+					LevelList[LevelIndex]->LogAndDisardRunawayNewlySpawned(Context.TickGroup);
+				}
+			}
 		}
 	}
 
