Renamed variables that detect if a menu was opened in a viewport to improve readability

diff --git a/Engine/Source/Runtime/Engine/Private/Slate/SceneViewport.cpp b/Engine/Source/Runtime/Engine/Private/Slate/SceneViewport.cpp
--- a/Engine/Source/Runtime/Engine/Private/Slate/SceneViewport.cpp
+++ b/Engine/Source/Runtime/Engine/Private/Slate/SceneViewport.cpp
@@ -372,7 +372,7 @@ FReply FSceneViewport::OnMouseButtonDown( const FGeometry& InGeometry, const FPo
 			ApplyModifierKeys( KeysState );
 		}
 
-		const bool bAnyMenusVisible = FSlateApplication::Get().AnyMenusVisible();
+		const bool bAnyMenuWasVisible = FSlateApplication::Get().AnyMenusVisible();
 
 		// Process the mouse event
 		if (!ViewportClient->InputKey(this, InMouseEvent.GetUserIndex(), InMouseEvent.GetEffectingButton(), IE_Pressed))
@@ -381,7 +381,7 @@ FReply FSceneViewport::OnMouseButtonDown( const FGeometry& InGeometry, const FPo
 		}
 
 		// a new menu was opened if there was previously not a menu visible but now there is
-		const bool bNewMenuWasOpened = !bAnyMenusVisible && FSlateApplication::Get().AnyMenusVisible();
+		const bool bNewMenuWasOpened = !bAnyMenuWasVisible && FSlateApplication::Get().AnyMenusVisible();
 
 		if (!ViewportClient->IgnoreInput() &&
 			!bNewMenuWasOpened && // We should not focus the viewport if a menu was opened as it would close the menu
