Fix crash when both parallel animation and parallel physics are turned off.

diff --git a/Engine/Source/Runtime/Engine/Classes/Components/SkeletalMeshComponent.h b/Engine/Source/Runtime/Engine/Classes/Components/SkeletalMeshComponent.h
--- a/Engine/Source/Runtime/Engine/Classes/Components/SkeletalMeshComponent.h
+++ b/Engine/Source/Runtime/Engine/Classes/Components/SkeletalMeshComponent.h
@@ -1280,6 +1280,7 @@ private:
 
 	friend class FParallelBlendPhysicsCompletionTask;
 	void CompleteParallelBlendPhysics();
+	void PostBlendPhysics();
 
 	friend class FTickClothingTask;
 	void PerformTickClothing(float DeltaTime);
diff --git a/Engine/Source/Runtime/Engine/Private/PhysicsEngine/PhysAnim.cpp b/Engine/Source/Runtime/Engine/Private/PhysicsEngine/PhysAnim.cpp
--- a/Engine/Source/Runtime/Engine/Private/PhysicsEngine/PhysAnim.cpp
+++ b/Engine/Source/Runtime/Engine/Private/PhysicsEngine/PhysAnim.cpp
@@ -95,7 +95,7 @@ public:
 
 
 // Use current pose to calculate world-space position of this bone without physics now.
-void UpdateWorldBoneTM(TArray<FAssetWorldBoneTM> & WorldBoneTMs, int32 BoneIndex, USkeletalMeshComponent* SkelComp, const FTransform &LocalToWorldTM, const FVector& Scale3D)
+void UpdateWorldBoneTM(TArray<FAssetWorldBoneTM> & WorldBoneTMs, const TArray<FTransform>& InLocalAtoms, int32 BoneIndex, USkeletalMeshComponent* SkelComp, const FTransform &LocalToWorldTM, const FVector& Scale3D)
 {
 	// If its already up to date - do nothing
 	if(	WorldBoneTMs[BoneIndex].bUpToDate )
@@ -113,11 +113,11 @@ void UpdateWorldBoneTM(TArray<FAssetWorldBoneTM> & WorldBoneTMs, int32 BoneIndex
 	{
 		// If not root, use our cached world-space bone transforms.
 		int32 ParentIndex = SkelComp->SkeletalMesh->RefSkeleton.GetParentIndex(BoneIndex);
-		UpdateWorldBoneTM(WorldBoneTMs, ParentIndex, SkelComp, LocalToWorldTM, Scale3D);
+		UpdateWorldBoneTM(WorldBoneTMs, InLocalAtoms, ParentIndex, SkelComp, LocalToWorldTM, Scale3D);
 		ParentTM = WorldBoneTMs[ParentIndex].TM;
 	}
 
-	RelTM = SkelComp->LocalAtoms[BoneIndex];
+	RelTM = InLocalAtoms[BoneIndex];
 	RelTM.ScaleTranslation( Scale3D );
 
 	WorldBoneTMs[BoneIndex].TM = RelTM * ParentTM;
@@ -237,7 +237,7 @@ void USkeletalMeshComponent::PerformBlendPhysicsBones(const TArray<FBoneIndexTyp
 					{
 						// If not root, get parent TM from cache (making sure its up-to-date).
 						int32 ParentIndex = SkeletalMesh->RefSkeleton.GetParentIndex(BoneIndex);
-						UpdateWorldBoneTM(WorldBoneTMs, ParentIndex, this, LocalToWorldTM, TotalScale3D);
+						UpdateWorldBoneTM(WorldBoneTMs, InLocalAtoms, ParentIndex, this, LocalToWorldTM, TotalScale3D);
 						ParentWorldTM = WorldBoneTMs[ParentIndex].TM;
 					}
 
@@ -411,16 +411,13 @@ void USkeletalMeshComponent::BlendInPhysics()
 		}else
 		{
 			PerformBlendPhysicsBones(RequiredBones, LocalAtoms);
-			CompleteParallelBlendPhysics();
+			PostBlendPhysics();
 		}
 	}
 }
 
-void USkeletalMeshComponent::CompleteParallelBlendPhysics()
+void USkeletalMeshComponent::PostBlendPhysics()
 {
-	Exchange(AnimEvaluationContext.LocalAtoms, AnimEvaluationContext.bDoInterpolation ? CachedLocalAtoms : LocalAtoms);
-		
-
 	FlipEditableSpaceBases();
 
 	// Update Child Transform - The above function changes bone transform, so will need to update child transform
@@ -433,6 +430,13 @@ void USkeletalMeshComponent::CompleteParallelBlendPhysics()
 	MarkRenderDynamicDataDirty();
 
 	FinalizeBoneTransform();
+}
+
+void USkeletalMeshComponent::CompleteParallelBlendPhysics()
+{
+	Exchange(AnimEvaluationContext.LocalAtoms, AnimEvaluationContext.bDoInterpolation ? CachedLocalAtoms : LocalAtoms);
+		
+	PostBlendPhysics();
 
 	ParallelAnimationEvaluationTask.SafeRelease();
 	ParallelBlendPhysicsCompletionTask.SafeRelease();
