Fix welded children not being able to dynamically change collision response

diff --git a/Engine/Source/Runtime/Engine/Private/PhysicsEngine/BodyInstance.cpp b/Engine/Source/Runtime/Engine/Private/PhysicsEngine/BodyInstance.cpp
--- a/Engine/Source/Runtime/Engine/Private/PhysicsEngine/BodyInstance.cpp
+++ b/Engine/Source/Runtime/Engine/Private/PhysicsEngine/BodyInstance.cpp
@@ -836,6 +836,12 @@ void FBodyInstance::UpdatePhysicsFilterData()
 {
 	SCOPE_CYCLE_COUNTER(STAT_UpdatePhysFilter);
 
+	if(WeldParent)
+	{
+		WeldParent->UpdatePhysicsFilterData();
+		return;
+	}
+
 	// Do nothing if no physics actor
 	if (!IsValidBodyInstance())
 	{
