Fix Non-Unity builds.

diff --git a/Engine/Source/Runtime/Engine/Private/ActorEditor.cpp b/Engine/Source/Runtime/Engine/Private/ActorEditor.cpp
--- a/Engine/Source/Runtime/Engine/Private/ActorEditor.cpp
+++ b/Engine/Source/Runtime/Engine/Private/ActorEditor.cpp
@@ -12,6 +12,8 @@
 
 #if WITH_EDITOR
 
+#include "Editor.h"
+
 #define LOCTEXT_NAMESPACE "ErrorChecking"
 
 void AActor::PreEditChange(UProperty* PropertyThatWillChange)
