USCS_Node::CategoryName is now wrapped in WITH_EDITORONLY_DATA so that it is not gathered for localization unless gathering editor only data.

diff --git a/Engine/Source/Runtime/Engine/Classes/Engine/SCS_Node.h b/Engine/Source/Runtime/Engine/Classes/Engine/SCS_Node.h
--- a/Engine/Source/Runtime/Engine/Classes/Engine/SCS_Node.h
+++ b/Engine/Source/Runtime/Engine/Classes/Engine/SCS_Node.h
@@ -20,9 +20,11 @@ class USCS_Node : public UObject
 	UPROPERTY()
 	FName VariableName;
 
+#if WITH_EDITORONLY_DATA
 	/** If non-None, the assigned category name */
 	UPROPERTY()
 	FText CategoryName;
+#endif //WITH_EDITORONLY_DATA
 
 	/** Socket/Bone that Node might attach to */
 	UPROPERTY()
