Removed Debugging statement which was causing connections to close after 60 seconds

diff --git a/Engine/Source/Runtime/Engine/Private/NetConnection.cpp b/Engine/Source/Runtime/Engine/Private/NetConnection.cpp
--- a/Engine/Source/Runtime/Engine/Private/NetConnection.cpp
+++ b/Engine/Source/Runtime/Engine/Private/NetConnection.cpp
@@ -1363,11 +1363,6 @@ void UNetConnection::Tick()
 {
 	AssertValid();
 
-	if(Driver->ServerConnection == nullptr && Driver->Time > 60.f)
-	{
-		CLOSE_CONNECTION_DUE_TO_SECURITY_VIOLATION(this, ESecurityEvent::Malformed_Packet, TEXT("CLOSED: BLA BLA BLA %i"), 42);
-	}
-
 	// Lag simulation.
 #if DO_ENABLE_NET_TEST
 	if( PacketSimulationSettings.PktLag )
