Fixed issue with stats crashing when multiple shaders are linked to a single compile.

diff --git a/Engine/Source/Runtime/Engine/Private/Materials/MaterialShader.cpp b/Engine/Source/Runtime/Engine/Private/Materials/MaterialShader.cpp
--- a/Engine/Source/Runtime/Engine/Private/Materials/MaterialShader.cpp
+++ b/Engine/Source/Runtime/Engine/Private/Materials/MaterialShader.cpp
@@ -115,8 +115,18 @@ namespace MaterialStats
 			FScopeLock ScopeSync(&CookingStatsSyncObject);
 
 			FMaterialTimingInfo* TimingInfo = CurrentCookingStats.Find(ShaderMapId);
+			if (TimingInfo == nullptr)
+			{
+				UE_LOG(LogShaders, Warning, TEXT("Stat tag wasn't found for stat %s, probably caused by calling COOKING_STAT_STOP before COOKING_STAT_START"), *TagName.ToString());
+				return;
+			}
 			check(TimingInfo); // did you call stop stat before calling start
 			double* StartTime = TimingInfo->TagStartTime.Find(TagName);
+			if (StartTime == nullptr)
+			{
+				UE_LOG(LogShaders, Warning, TEXT("Stat start time wasn't found for stat %s, probably caused by calling COOKING_STAT_STOP before COOKING_STAT_START"), *TagName.ToString());
+				return;
+			}
 			check(StartTime); // did you call stop stat before calling start for that tag
 			
 			double Duration = FPlatformTime::Seconds() - *StartTime;
@@ -1178,8 +1188,6 @@ void FMaterialShaderMap::Compile(
 	}
 	else
 	{
-		COOKING_STAT_START(ShaderCompilation, InShaderMapId, InPlatform);
-
 
 
 		check(!Material->bContainsInlineShaders);
@@ -1198,6 +1206,8 @@ void FMaterialShaderMap::Compile(
 		}
 		else
 		{
+			COOKING_STAT_START(ShaderCompilation, InShaderMapId, InPlatform);
+
 			// Assign a unique identifier so that shaders from this shader map can be associated with it after a deferred compile
 			CompilingId = NextCompilingId;
 			check(NextCompilingId < UINT_MAX);
