Remove redundant call to flip editable bone buffer

diff --git a/Engine/Source/Runtime/Engine/Private/PhysicsEngine/PhysAnim.cpp b/Engine/Source/Runtime/Engine/Private/PhysicsEngine/PhysAnim.cpp
--- a/Engine/Source/Runtime/Engine/Private/PhysicsEngine/PhysAnim.cpp
+++ b/Engine/Source/Runtime/Engine/Private/PhysicsEngine/PhysAnim.cpp
@@ -418,7 +418,7 @@ void USkeletalMeshComponent::BlendInPhysics()
 
 void USkeletalMeshComponent::PostBlendPhysics()
 {
-	FlipEditableSpaceBases();
+	FinalizeBoneTransform();
 
 	// Update Child Transform - The above function changes bone transform, so will need to update child transform
 	UpdateChildTransforms();
@@ -428,8 +428,6 @@ void USkeletalMeshComponent::PostBlendPhysics()
 
 	// New bone positions need to be sent to render thread
 	MarkRenderDynamicDataDirty();
-
-	FinalizeBoneTransform();
 }
 
 void USkeletalMeshComponent::CompleteParallelBlendPhysics()
