Don't show slate popup for successful high res screenshot in fullscreen game. Breaks fullscreen focus.

Merging

//depot/UE4-UT/...

to //depot/UE4/...

diff --git a/Engine/Source/Runtime/Engine/Private/UnrealClient.cpp b/Engine/Source/Runtime/Engine/Private/UnrealClient.cpp
--- a/Engine/Source/Runtime/Engine/Private/UnrealClient.cpp
+++ b/Engine/Source/Runtime/Engine/Private/UnrealClient.cpp
@@ -779,6 +779,7 @@ void FViewport::HighResScreenshot()
 	bTakeHighResScreenShot = false;
 
 	// Notification of a successful screenshot
+	if (GIsEditor || !IsFullscreen())
 	{
 		auto Message = NSLOCTEXT("UnrealClient", "HighResScreenshotSavedAs", "High resolution screenshot saved as");
 		FNotificationInfo Info(Message);
