UE4 - D3D12 - Do not allow real fullscreen

diff --git a/Engine/Source/Runtime/Engine/Private/GameEngine.cpp b/Engine/Source/Runtime/Engine/Private/GameEngine.cpp
--- a/Engine/Source/Runtime/Engine/Private/GameEngine.cpp
+++ b/Engine/Source/Runtime/Engine/Private/GameEngine.cpp
@@ -180,6 +180,17 @@ void UGameEngine::ConditionallyOverrideSettings(int32& ResolutionX, int32& Resol
 		auto CVar = IConsoleManager::Get().FindTConsoleVariableDataInt(TEXT("r.FullScreenMode"));
 		check(CVar);
 		WindowMode = CVar->GetValueOnGameThread() == 0 ? EWindowMode::Fullscreen : EWindowMode::WindowedFullscreen;
+
+		if (PLATFORM_WINDOWS && WindowMode == EWindowMode::Fullscreen)
+		{
+			// Handle fullscreen mode differently for D3D11/D3D12
+			static const bool bD3D12 = FParse::Param(FCommandLine::Get(), TEXT("d3d12")) || FParse::Param(FCommandLine::Get(), TEXT("dx12"));
+			if (bD3D12)
+			{
+				// Force D3D12 RHI to use windowed fullscreen mode
+				WindowMode = EWindowMode::WindowedFullscreen;
+			}
+		}
 	}
 
 	//fullscreen is always supported, but don't allow windowed mode on platforms that dont' support it.
diff --git a/Engine/Source/Runtime/Engine/Private/GameViewportClient.cpp b/Engine/Source/Runtime/Engine/Private/GameViewportClient.cpp
--- a/Engine/Source/Runtime/Engine/Private/GameViewportClient.cpp
+++ b/Engine/Source/Runtime/Engine/Private/GameViewportClient.cpp
@@ -2708,6 +2708,18 @@ bool UGameViewportClient::HandleToggleFullscreenCommand()
 			FullScreenMode = Viewport->IsFullscreen() ? EWindowMode::Windowed : EWindowMode::Fullscreen;
 		}
 	}
+
+	if (PLATFORM_WINDOWS && FullScreenMode == EWindowMode::Fullscreen)
+	{
+		// Handle fullscreen mode differently for D3D11/D3D12
+		static const bool bD3D12 = FParse::Param(FCommandLine::Get(), TEXT("d3d12")) || FParse::Param(FCommandLine::Get(), TEXT("dx12"));
+		if (bD3D12)
+		{
+			// Force D3D12 RHI to use windowed fullscreen mode
+			FullScreenMode = EWindowMode::WindowedFullscreen;
+		}
+	}
+
 	FSystemResolution::RequestResolutionChange(GSystemResolution.ResX, GSystemResolution.ResY, FullScreenMode);
 	return true;
 }
