#FRAMEWORK: Accumulated transform during spawn

- This is to preserve and respect the original template relative transform and apply the new as delta.
- Doesn't work for SpawnActor blueprint node since the transform gets stomped - wip on discussion

diff --git a/Engine/Source/Runtime/Engine/Private/LevelActor.cpp b/Engine/Source/Runtime/Engine/Private/LevelActor.cpp
--- a/Engine/Source/Runtime/Engine/Private/LevelActor.cpp
+++ b/Engine/Source/Runtime/Engine/Private/LevelActor.cpp
@@ -323,9 +323,25 @@ AActor* UWorld::SpawnActor( UClass* Class, FTransform const* Transform, const FA
 		return NULL;
 	}
 
-	FVector NewLocation = Transform ? Transform->GetLocation() : (Template->GetRootComponent() ? Template->GetRootComponent()->RelativeLocation : FVector::ZeroVector);
-	FRotator NewRotation = Transform ? Transform->GetRotation().Rotator() :	(Template->GetRootComponent() ? Template->GetRootComponent()->RelativeRotation : FRotator::ZeroRotator);
-	FVector NewScale = Transform ? Transform->GetScale3D() :	(Template->GetRootComponent() ? Template->GetRootComponent()->RelativeScale3D : FVector(1.f) );
+	const USceneComponent* TemplateRootComponent = Template->GetRootComponent();
+	FVector RelativeLocation = Transform ? Transform->GetLocation() : (TemplateRootComponent ? TemplateRootComponent->RelativeLocation : FVector::ZeroVector);
+	FRotator RelativeRotation = Transform ? Transform->GetRotation().Rotator() :	(TemplateRootComponent ? TemplateRootComponent->RelativeRotation : FRotator::ZeroRotator);
+	FVector RelativeScale = Transform ? Transform->GetScale3D() :	(TemplateRootComponent ? TemplateRootComponent->RelativeScale3D : FVector(1.f));
+
+	FTransform RelativeTransform(RelativeRotation, RelativeLocation, RelativeScale);
+	FTransform NewTransform;
+
+	// we do respect Template's RootComponent. If this is called using new relative transform, we apply to the template transform
+	// this way, we respect CDO's transform and multiply on top of it. 
+	if(TemplateRootComponent)
+	{
+		FTransform TemplateTransform(TemplateRootComponent->RelativeRotation, TemplateRootComponent->RelativeLocation, TemplateRootComponent->RelativeScale3D);
+		NewTransform = TemplateTransform * RelativeTransform;
+	}
+	else
+	{
+		NewTransform = RelativeTransform;
+	}
 
 	PRAGMA_DISABLE_DEPRECATION_WARNINGS;
 	// handle existing (but deprecated) uses of bNoCollisionFail where user set it to true
@@ -357,6 +373,8 @@ AActor* UWorld::SpawnActor( UClass* Class, FTransform const* Transform, const FA
 	// note: we can't handle all cases here, since we don't know the full component hierarchy until after the actor is spawned
 	if (CollisionHandlingMethod == ESpawnActorCollisionHandlingMethod::DontSpawnIfColliding)
 	{
+		FVector NewLocation = NewTransform.GetLocation();
+		FRotator NewRotation = NewTransform.GetRotation().Rotator();
 		if (EncroachingBlockingGeometry(Template, NewLocation, NewRotation))
 		{
 			// a native component is colliding, that's enough to reject spawning
@@ -393,7 +411,7 @@ AActor* UWorld::SpawnActor( UClass* Class, FTransform const* Transform, const FA
 	// tell the actor what method to use, in case it was overridden
 	Actor->SpawnCollisionHandlingMethod = CollisionHandlingMethod;
 
-	Actor->PostSpawnInitialize(FTransform(NewRotation, NewLocation, NewScale), SpawnParameters.Owner, SpawnParameters.Instigator, SpawnParameters.bRemoteOwned, SpawnParameters.bNoFail, SpawnParameters.bDeferConstruction);
+	Actor->PostSpawnInitialize(NewTransform, SpawnParameters.Owner, SpawnParameters.Instigator, SpawnParameters.bRemoteOwned, SpawnParameters.bNoFail, SpawnParameters.bDeferConstruction);
 
 	if (Actor->IsPendingKill() && !SpawnParameters.bNoFail)
 	{
