[UE-18490] Ensure 'bExchangedRoles' on the CDO is not set after serialization.

Note: Some Blueprint class assets created in v4.7 through converting Actor instances may have accidentally gotten this transient flag set 'true' in the CDO, resulting in it being 'true' for instances of the Blueprint class by default. I *think* that would have been fixed by CL# 2497181.

diff --git a/Engine/Source/Runtime/Engine/Private/Actor.cpp b/Engine/Source/Runtime/Engine/Private/Actor.cpp
--- a/Engine/Source/Runtime/Engine/Private/Actor.cpp
+++ b/Engine/Source/Runtime/Engine/Private/Actor.cpp
@@ -446,6 +446,12 @@ void AActor::PostLoad()
 		bReplicates = (RemoteRole != ROLE_None);
 	}
 
+	// Ensure that this is not set for CDO (there was a case where this might have occurred in an older version when converting actor instances to BPs - see UE-18490)
+	if (HasAnyFlags(RF_ClassDefaultObject))
+	{
+		bExchangedRoles = false;
+	}
+
 	if ( GIsEditor )
 	{
 #if WITH_EDITORONLY_DATA
