UE4 - D3D12 - Integrate & fix texture lock with different strides

diff --git a/Engine/Source/Runtime/Engine/Private/Slate/SlateTextures.cpp b/Engine/Source/Runtime/Engine/Private/Slate/SlateTextures.cpp
--- a/Engine/Source/Runtime/Engine/Private/Slate/SlateTextures.cpp
+++ b/Engine/Source/Runtime/Engine/Private/Slate/SlateTextures.cpp
@@ -55,10 +55,24 @@ void FSlateTexture2DRHIRef::InitDynamicRHI()
 			check(Height == TextureData->GetHeight());
 
 			uint32 Stride;
-			uint8* DestTextureData = (uint8*)RHILockTexture2D( ShaderResource, 0, RLM_WriteOnly, Stride, false );
-			FMemory::Memcpy( DestTextureData, TextureData->GetRawBytes().GetData(), Width*Height*GPixelFormats[PixelFormat].BlockBytes );
-			RHIUnlockTexture2D( ShaderResource, 0, false );
-
+			uint8* DestTextureData = (uint8*)RHILockTexture2D(ShaderResource, 0, RLM_WriteOnly, Stride, false);
+			const uint8* SourceTextureData = TextureData->GetRawBytes().GetData();
+			const uint32 DataStride = Width * GPixelFormats[PixelFormat].BlockBytes;
+			if (Stride == DataStride)
+			{
+				FMemory::Memcpy(DestTextureData, SourceTextureData, DataStride * Height);
+			}
+			else
+			{
+				checkf(GPixelFormats[PixelFormat].BlockSizeX == GPixelFormats[PixelFormat].BlockSizeY == GPixelFormats[PixelFormat].BlockSizeZ == 1, TEXT("Tried to use compressed format?"));
+				for (uint32 i = 0; i < Height; i++)
+				{
+					FMemory::Memcpy(DestTextureData, SourceTextureData, DataStride);
+					DestTextureData += Stride;
+					SourceTextureData += DataStride;
+				}
+			}
+			RHIUnlockTexture2D(ShaderResource, 0, false);
 			TextureData->Empty();
 		}
 	}
