Fixed destructible meshes sometimes recieving root chunk offsets on world start.

diff --git a/Engine/Source/Runtime/Engine/Private/Components/DestructibleComponent.cpp b/Engine/Source/Runtime/Engine/Private/Components/DestructibleComponent.cpp
--- a/Engine/Source/Runtime/Engine/Private/Components/DestructibleComponent.cpp
+++ b/Engine/Source/Runtime/Engine/Private/Components/DestructibleComponent.cpp
@@ -885,7 +885,15 @@ void UDestructibleComponent::SetChunkVisible( int32 ChunkIndex, bool bVisible )
 			{
 				const physx::PxMat44 ChunkPoseRT = ApexDestructibleActor->getChunkPose(ChunkIndex);	// Unscaled
 				const physx::PxTransform Transform(ChunkPoseRT);
-				SetChunkWorldRT(ChunkIndex, P2UQuat(Transform.q), P2UVector(Transform.p));
+				if(IsFracturedOrInitiallyStatic())
+				{
+					SetChunkWorldRT(ChunkIndex, P2UQuat(Transform.q), P2UVector(Transform.p));
+				}
+				else
+				{
+					// Not yet fractured, sync the component to the physics body
+					UPrimitiveComponent::SyncComponentToRBPhysics();
+				}
 			}
 		}
 	}
