fixed bug in updating navigation path not reseting flag: bIsPartial
#ue4

diff --git a/Engine/Source/Runtime/Engine/Private/AI/Navigation/NavigationPath.cpp b/Engine/Source/Runtime/Engine/Private/AI/Navigation/NavigationPath.cpp
--- a/Engine/Source/Runtime/Engine/Private/AI/Navigation/NavigationPath.cpp
+++ b/Engine/Source/Runtime/Engine/Private/AI/Navigation/NavigationPath.cpp
@@ -385,6 +385,7 @@ void FNavMeshPath::Reset()
 	PathCorridorCost.Reset();
 	bStringPulled = false;
 	bCorridorEdgesGenerated = false;
+	bIsPartial = false;
 }
 
 float FNavMeshPath::GetStringPulledLength(const int32 StartingPoint) const
