Stats - Removed dumplicated stats, removed unnecessary spaces from stat's descriptions

diff --git a/Engine/Source/Runtime/Engine/Private/Components/CharacterMovementComponent.cpp b/Engine/Source/Runtime/Engine/Private/Components/CharacterMovementComponent.cpp
--- a/Engine/Source/Runtime/Engine/Private/Components/CharacterMovementComponent.cpp
+++ b/Engine/Source/Runtime/Engine/Private/Components/CharacterMovementComponent.cpp
@@ -875,7 +875,6 @@ void UCharacterMovementComponent::PerformAirControlForPathFollowing(FVector Dire
 
 void UCharacterMovementComponent::TickComponent(float DeltaTime, enum ELevelTick TickType, FActorComponentTickFunction *ThisTickFunction)
 {
-	SCOPE_CYCLE_COUNTER(STAT_CharacterMovement);
 	SCOPE_CYCLE_COUNTER(STAT_CharacterMovementTick);
 
 	const FVector InputVector = ConsumeInputVector();
diff --git a/Engine/Source/Runtime/Engine/Private/LevelTick.cpp b/Engine/Source/Runtime/Engine/Private/LevelTick.cpp
--- a/Engine/Source/Runtime/Engine/Private/LevelTick.cpp
+++ b/Engine/Source/Runtime/Engine/Private/LevelTick.cpp
@@ -63,7 +63,6 @@ DEFINE_STAT(STAT_PostTickComponentUpdateWait);
 DEFINE_STAT(STAT_TickTime);
 DEFINE_STAT(STAT_WorldTickTime);
 DEFINE_STAT(STAT_UpdateCameraTime);
-DEFINE_STAT(STAT_CharacterMovement);
 DEFINE_STAT(STAT_PlayerControllerTick);
 
 DEFINE_STAT(STAT_VolumeStreamingTickTime);
diff --git a/Engine/Source/Runtime/Engine/Public/EngineStats.h b/Engine/Source/Runtime/Engine/Public/EngineStats.h
--- a/Engine/Source/Runtime/Engine/Public/EngineStats.h
+++ b/Engine/Source/Runtime/Engine/Public/EngineStats.h
@@ -53,12 +53,11 @@ DECLARE_CYCLE_STAT_EXTERN(TEXT("GC Sweep Time"),STAT_GCSweepTime,STATGROUP_Game,
 DECLARE_CYCLE_STAT_EXTERN(TEXT("Tick Time"),STAT_TickTime,STATGROUP_Game, );
 DECLARE_CYCLE_STAT_EXTERN(TEXT("World Tick Time"),STAT_WorldTickTime,STATGROUP_Game, );
 DECLARE_CYCLE_STAT_EXTERN(TEXT("Update Camera Time"),STAT_UpdateCameraTime,STATGROUP_Game, );
-DECLARE_CYCLE_STAT_EXTERN(TEXT("Char Movement Total"), STAT_CharacterMovement, STATGROUP_Game, );
 DECLARE_CYCLE_STAT_EXTERN(TEXT("PlayerController Tick"), STAT_PlayerControllerTick, STATGROUP_Game, );
 DECLARE_CYCLE_STAT_EXTERN(TEXT("Post Tick Component Update"),STAT_PostTickComponentUpdate,STATGROUP_Game, );
 DECLARE_CYCLE_STAT_EXTERN(TEXT("Post Tick Component Wait"),STAT_PostTickComponentUpdateWait,STATGROUP_Game, );
-DECLARE_CYCLE_STAT_EXTERN(TEXT("  Recreate"),STAT_PostTickComponentRecreate,STATGROUP_Game, );
-DECLARE_CYCLE_STAT_EXTERN(TEXT("  Transform or RenderData"),STAT_PostTickComponentLW,STATGROUP_Game, );
+DECLARE_CYCLE_STAT_EXTERN(TEXT("Recreate"),STAT_PostTickComponentRecreate,STATGROUP_Game, );
+DECLARE_CYCLE_STAT_EXTERN(TEXT("Transform or RenderData"),STAT_PostTickComponentLW,STATGROUP_Game, );
 DECLARE_CYCLE_STAT_EXTERN(TEXT("Update Particle Data"),STAT_ParticleManagerUpdateData,STATGROUP_Game, );
 DECLARE_CYCLE_STAT_EXTERN(TEXT("Async Work Wait"),STAT_AsyncWorkWaitTime,STATGROUP_Game, );
 DECLARE_CYCLE_STAT_EXTERN(TEXT("Visual Logging time"), STAT_VisualLog, STATGROUP_Game, ENGINE_API);
@@ -69,16 +68,16 @@ DECLARE_CYCLE_STAT_EXTERN(TEXT("GT Tickable Time"),STAT_TickableTickTime,STATGRO
 DECLARE_CYCLE_STAT_EXTERN(TEXT("Runtime Movie Tick Time"),STAT_RuntimeMovieSceneTickTime,STATGROUP_Game, );
 DECLARE_CYCLE_STAT_EXTERN(TEXT("Finish Async Trace Time"),STAT_FinishAsyncTraceTickTime,STATGROUP_Game, );
 DECLARE_CYCLE_STAT_EXTERN(TEXT("Net Broadcast Tick Time"),STAT_NetBroadcastTickTime,STATGROUP_Game, );
-DECLARE_CYCLE_STAT_EXTERN(TEXT("  ServerReplicateActors Time"),STAT_NetServerRepActorsTime,STATGROUP_Game, );
-DECLARE_CYCLE_STAT_EXTERN(TEXT("  Consider Actors Time"),STAT_NetConsiderActorsTime,STATGROUP_Game, );
-DECLARE_CYCLE_STAT_EXTERN(TEXT("  Inital Dormant Time"),STAT_NetInitialDormantCheckTime,STATGROUP_Game, );
-DECLARE_CYCLE_STAT_EXTERN(TEXT("  Prioritize Actors Time"),STAT_NetPrioritizeActorsTime,STATGROUP_Game, );
-DECLARE_CYCLE_STAT_EXTERN(TEXT("  Replicate Actors Time"),STAT_NetReplicateActorsTime,STATGROUP_Game, );
-DECLARE_CYCLE_STAT_EXTERN(TEXT("  Dynamic Property Rep Time"),STAT_NetReplicateDynamicPropTime,STATGROUP_Game, );
-DECLARE_DWORD_COUNTER_STAT_EXTERN(TEXT("  Skipped Dynamic Props"),STAT_NetSkippedDynamicProps,STATGROUP_Game, );
-DECLARE_CYCLE_STAT_EXTERN(TEXT("  NetSerializeItemDelta Time"),STAT_NetSerializeItemDeltaTime,STATGROUP_Game, );
-DECLARE_CYCLE_STAT_EXTERN(TEXT("  Static Property Rep Time"),STAT_NetReplicateStaticPropTime,STATGROUP_Game, );
-DECLARE_CYCLE_STAT_EXTERN(TEXT("  Rebuild Conditionals"),STAT_NetRebuildConditionalTime,STATGROUP_Game, );
+DECLARE_CYCLE_STAT_EXTERN(TEXT("ServerReplicateActors Time"),STAT_NetServerRepActorsTime,STATGROUP_Game, );
+DECLARE_CYCLE_STAT_EXTERN(TEXT("Consider Actors Time"),STAT_NetConsiderActorsTime,STATGROUP_Game, );
+DECLARE_CYCLE_STAT_EXTERN(TEXT("Inital Dormant Time"),STAT_NetInitialDormantCheckTime,STATGROUP_Game, );
+DECLARE_CYCLE_STAT_EXTERN(TEXT("Prioritize Actors Time"),STAT_NetPrioritizeActorsTime,STATGROUP_Game, );
+DECLARE_CYCLE_STAT_EXTERN(TEXT("Replicate Actors Time"),STAT_NetReplicateActorsTime,STATGROUP_Game, );
+DECLARE_CYCLE_STAT_EXTERN(TEXT("Dynamic Property Rep Time"),STAT_NetReplicateDynamicPropTime,STATGROUP_Game, );
+DECLARE_DWORD_COUNTER_STAT_EXTERN(TEXT("Skipped Dynamic Props"),STAT_NetSkippedDynamicProps,STATGROUP_Game, );
+DECLARE_CYCLE_STAT_EXTERN(TEXT("NetSerializeItemDelta Time"),STAT_NetSerializeItemDeltaTime,STATGROUP_Game, );
+DECLARE_CYCLE_STAT_EXTERN(TEXT("Static Property Rep Time"),STAT_NetReplicateStaticPropTime,STATGROUP_Game, );
+DECLARE_CYCLE_STAT_EXTERN(TEXT("Rebuild Conditionals"),STAT_NetRebuildConditionalTime,STATGROUP_Game, );
 DECLARE_CYCLE_STAT_EXTERN(TEXT("Net Post BC Tick Time"),STAT_NetBroadcastPostTickTime,STATGROUP_Game, );
 DECLARE_CYCLE_STAT_EXTERN(TEXT("Net PackageMap SerializeObject"),STAT_PackageMap_SerializeObjectTime,STATGROUP_Game, );
 
@@ -89,11 +88,11 @@ DECLARE_CYCLE_STAT_EXTERN(TEXT("Net PackageMap SerializeObject"),STAT_PackageMap
 DECLARE_CYCLE_STAT_EXTERN(TEXT("Async navmesh tile building"),STAT_Navigation_TileBuildAsync,STATGROUP_Navigation, );
 DECLARE_CYCLE_STAT_EXTERN(TEXT("Sync tiles building preparation"),STAT_Navigation_TileBuildPreparationSync,STATGROUP_Navigation, );
 DECLARE_CYCLE_STAT_EXTERN(TEXT("Processing world for building navmesh"),STAT_Navigation_ProcessingActorsForNavMeshBuilding,STATGROUP_Navigation, );
-DECLARE_CYCLE_STAT_EXTERN(TEXT("    Sync actors' geometry export"),STAT_Navigation_ActorsGeometryExportSync,STATGROUP_Navigation, );
-DECLARE_CYCLE_STAT_EXTERN(TEXT("    Sync BSP geometry export"),STAT_Navigation_BSPExportSync,STATGROUP_Navigation, );
-DECLARE_CYCLE_STAT_EXTERN(TEXT("    Sync gathering navigation modifiers"),STAT_Navigation_GatheringNavigationModifiersSync,STATGROUP_Navigation, );
-DECLARE_CYCLE_STAT_EXTERN(TEXT("    Async debug OBJ file export"),STAT_Navigation_TileGeometryExportToObjAsync,STATGROUP_Navigation, );
-DECLARE_CYCLE_STAT_EXTERN(TEXT("    Async voxel filtering"),STAT_Navigation_TileVoxelFilteringAsync,STATGROUP_Navigation, );
+DECLARE_CYCLE_STAT_EXTERN(TEXT("Sync actors' geometry export"),STAT_Navigation_ActorsGeometryExportSync,STATGROUP_Navigation, );
+DECLARE_CYCLE_STAT_EXTERN(TEXT("Sync BSP geometry export"),STAT_Navigation_BSPExportSync,STATGROUP_Navigation, );
+DECLARE_CYCLE_STAT_EXTERN(TEXT("Sync gathering navigation modifiers"),STAT_Navigation_GatheringNavigationModifiersSync,STATGROUP_Navigation, );
+DECLARE_CYCLE_STAT_EXTERN(TEXT("Async debug OBJ file export"),STAT_Navigation_TileGeometryExportToObjAsync,STATGROUP_Navigation, );
+DECLARE_CYCLE_STAT_EXTERN(TEXT("Async voxel filtering"),STAT_Navigation_TileVoxelFilteringAsync,STATGROUP_Navigation, );
 DECLARE_CYCLE_STAT_EXTERN(TEXT("Sync pathfinding"),STAT_Navigation_PathfindingSync,STATGROUP_Navigation, );
 DECLARE_CYCLE_STAT_EXTERN(TEXT("Sync requests for async pathfinding"),STAT_Navigation_RequestingAsyncPathfinding,STATGROUP_Navigation, );
 DECLARE_CYCLE_STAT_EXTERN(TEXT("Async pathfinding"),STAT_Navigation_PathfindingAsync,STATGROUP_Navigation, );
@@ -101,13 +100,13 @@ DECLARE_DWORD_COUNTER_STAT_EXTERN(TEXT("Num observed nav paths"), STAT_Navigatio
 DECLARE_CYCLE_STAT_EXTERN(TEXT("Offset from corners"), STAT_Navigation_OffsetFromCorners, STATGROUP_Navigation, );
 DECLARE_CYCLE_STAT_EXTERN(TEXT("Visibility test for path optimisation"), STAT_Navigation_PathVisibilityOptimisation, STATGROUP_Navigation, );
 DECLARE_CYCLE_STAT_EXTERN(TEXT("Sync queries"),STAT_Navigation_QueriesTimeSync,STATGROUP_Navigation, );
-DECLARE_CYCLE_STAT_EXTERN(TEXT("    Sync meta nav area preparation"),STAT_Navigation_MetaAreaTranslation,STATGROUP_Navigation, ); 
-DECLARE_CYCLE_STAT_EXTERN(TEXT("    Async nav areas sorting"),STAT_Navigation_TileNavAreaSorting,STATGROUP_Navigation, );
+DECLARE_CYCLE_STAT_EXTERN(TEXT("Sync meta nav area preparation"),STAT_Navigation_MetaAreaTranslation,STATGROUP_Navigation, ); 
+DECLARE_CYCLE_STAT_EXTERN(TEXT("Async nav areas sorting"),STAT_Navigation_TileNavAreaSorting,STATGROUP_Navigation, );
 DECLARE_DWORD_COUNTER_STAT_EXTERN(TEXT("Num Destructible shapes exported"),STAT_Navigation_DestructiblesShapesExported,STATGROUP_Navigation, );
 DECLARE_DWORD_COUNTER_STAT_EXTERN(TEXT("Num update nav octree"),STAT_Navigation_UpdateNavOctree,STATGROUP_Navigation, );
 DECLARE_CYCLE_STAT_EXTERN(TEXT("Adding actors to navoctree"),STAT_Navigation_AddingActorsToNavOctree,STATGROUP_Navigation, );
-DECLARE_CYCLE_STAT_EXTERN(TEXT("    Adjusting nav links"),STAT_Navigation_AdjustingNavLinks,STATGROUP_Navigation, );
-DECLARE_CYCLE_STAT_EXTERN(TEXT("    Sync AddGeneratedTiles"),STAT_Navigation_AddGeneratedTiles,STATGROUP_Navigation, );
+DECLARE_CYCLE_STAT_EXTERN(TEXT("Adjusting nav links"),STAT_Navigation_AdjustingNavLinks,STATGROUP_Navigation, );
+DECLARE_CYCLE_STAT_EXTERN(TEXT("Sync AddGeneratedTiles"),STAT_Navigation_AddGeneratedTiles,STATGROUP_Navigation, );
 DECLARE_CYCLE_STAT_EXTERN(TEXT("Recast tick"),STAT_Navigation_RecastTick,STATGROUP_Navigation, );
 DECLARE_CYCLE_STAT_EXTERN(TEXT("Recast pathfinding"), STAT_Navigation_RecastPathfinding, STATGROUP_Navigation, );
 DECLARE_CYCLE_STAT_EXTERN(TEXT("Recast: build compressed layers"),STAT_Navigation_RecastBuildCompressedLayers,STATGROUP_Navigation, );
