[UE-17070] Child components are no longer lost in a Blueprint Class asset with a default scene root when reparenting to a class that yields an inherited scene root.

change summary:
- modified UBlueprint::GetBlueprintHierarchyFromClass() to look at UBlueprint::ParentClass before UClass::GetSuperClass(), in case the BPGC has not yet been compiled after a reparenting operation
- modified USimpleConstructionScript::ValidateSceneRootNodes() to call RemoveNodeAndPromoteChildren() instead of RemoveNode() in order to ensure that any child nodes don't get removed from the tree along with the default root
- modified USimpleConstructionScript::ValidateSceneRootNodes() to ensure that any promoted child nodes are properly parented to the appropriate native or inherited root when the local default scene root is removed

diff --git a/Engine/Source/Runtime/Engine/Private/Blueprint.cpp b/Engine/Source/Runtime/Engine/Private/Blueprint.cpp
--- a/Engine/Source/Runtime/Engine/Private/Blueprint.cpp
+++ b/Engine/Source/Runtime/Engine/Private/Blueprint.cpp
@@ -941,7 +941,8 @@ bool UBlueprint::GetBlueprintHierarchyFromClass(const UClass* InClass, TArray<UB
 		bNoErrors &= (BP->Status != BS_Error);
 #endif // #if WITH_EDITORONLY_DATA
 
-		CurrentClass = CurrentClass->GetSuperClass();
+		// If valid, use stored ParentClass rather than the actual UClass::GetSuperClass(); handles the case when the class has not been recompiled yet after a reparent operation.
+		CurrentClass = BP->ParentClass ? BP->ParentClass : CurrentClass->GetSuperClass();
 	}
 
 	return bNoErrors;
diff --git a/Engine/Source/Runtime/Engine/Private/SimpleConstructionScript.cpp b/Engine/Source/Runtime/Engine/Private/SimpleConstructionScript.cpp
--- a/Engine/Source/Runtime/Engine/Private/SimpleConstructionScript.cpp
+++ b/Engine/Source/Runtime/Engine/Private/SimpleConstructionScript.cpp
@@ -926,9 +926,36 @@ void USimpleConstructionScript::ValidateSceneRootNodes()
 		else if(RootComponentTemplate != nullptr
 			&& RootNodes.Contains(DefaultSceneRootNode))
 		{
-			RootNodes.Remove(DefaultSceneRootNode);
+			// Temporarily parent to a native or inherited root (to propagate to any children that might be promoted on removal below)
+			if(UBlueprintGeneratedClass* BPClass = Cast<UBlueprintGeneratedClass>(RootComponentTemplate->GetOuter()))
+			{
+				if(BPClass->SimpleConstructionScript != nullptr
+					&& BPClass->SimpleConstructionScript != this)
+				{
+					DefaultSceneRootNode->bIsParentComponentNative = false;
+					DefaultSceneRootNode->ParentComponentOwnerClassName = BPClass->GetFName();
+					TArray<const USCS_Node*> ParentSCSNodes = BPClass->SimpleConstructionScript->GetAllNodesConst();
+					for(const USCS_Node* ParentSCSNode : ParentSCSNodes)
+					{
+						if(ParentSCSNode != nullptr && ParentSCSNode->ComponentTemplate == RootComponentTemplate)
+						{
+							DefaultSceneRootNode->ParentComponentOrVariableName = ParentSCSNode->VariableName;
+							break;
+						}
+					}
+				}
+			}
+			else
+			{
+				DefaultSceneRootNode->bIsParentComponentNative = true;
+				DefaultSceneRootNode->ParentComponentOrVariableName = RootComponentTemplate->GetFName();
+				DefaultSceneRootNode->ParentComponentOwnerClassName = NAME_None;
+			}
+
+			// Remove from the root node list (but keep child nodes)
+			RemoveNodeAndPromoteChildren(DefaultSceneRootNode);
 
-			// These shouldn't be set, but just in case...
+			// Ensure that these are reset back to their default state
 			DefaultSceneRootNode->bIsParentComponentNative = false;
 			DefaultSceneRootNode->ParentComponentOrVariableName = NAME_None;
 			DefaultSceneRootNode->ParentComponentOwnerClassName = NAME_None;
