Added a comment to GameMode::K2_FindPlayerStart

diff --git a/Engine/Source/Runtime/Engine/Classes/GameFramework/GameMode.h b/Engine/Source/Runtime/Engine/Classes/GameFramework/GameMode.h
--- a/Engine/Source/Runtime/Engine/Classes/GameFramework/GameMode.h
+++ b/Engine/Source/Runtime/Engine/Classes/GameFramework/GameMode.h
@@ -573,6 +573,7 @@ public:
 	UFUNCTION(BlueprintNativeEvent, Category="Game")
 	class AActor* FindPlayerStart( AController* Player, const FString& IncomingName = TEXT("") );
 
+	/** Calls code to select the best player start for this player to start from. */
 	UFUNCTION(BlueprintPure, Category=Game, meta=(DisplayName="FindPlayerStart"))
 	AActor* K2_FindPlayerStart(AController* Player);
 
