Spawning an actor will now respect default transform on a native root component in all cases.

diff --git a/Engine/Source/Runtime/Engine/Private/Actor.cpp b/Engine/Source/Runtime/Engine/Private/Actor.cpp
--- a/Engine/Source/Runtime/Engine/Private/Actor.cpp
+++ b/Engine/Source/Runtime/Engine/Private/Actor.cpp
@@ -2396,7 +2396,7 @@ static USceneComponent* FixupNativeActorComponents(AActor* Actor)
 	return SceneRootComponent;
 }
 
-void AActor::PostSpawnInitialize(FTransform const& SpawnTransform, AActor* InOwner, APawn* InInstigator, bool bRemoteOwned, bool bNoFail, bool bDeferConstruction)
+void AActor::PostSpawnInitialize(FTransform const& UserSpawnTransform, AActor* InOwner, APawn* InInstigator, bool bRemoteOwned, bool bNoFail, bool bDeferConstruction)
 {
 	// General flow here is like so
 	// - Actor sets up the basics.
@@ -2409,7 +2409,7 @@ void AActor::PostSpawnInitialize(FTransform const& SpawnTransform, AActor* InOwn
 	// This should be the same sequence for deferred or nondeferred spawning.
 
 	// It's not safe to call UWorld accessor functions till the world info has been spawned.
-	UWorld* World = GetWorld();
+	UWorld* const World = GetWorld();
 	bool const bActorsInitialized = World && World->AreActorsInitialized();
 
 	CreationTime = (World ? World->GetTimeSeconds() : 0.f);
@@ -2417,12 +2417,16 @@ void AActor::PostSpawnInitialize(FTransform const& SpawnTransform, AActor* InOwn
 	// Set network role.
 	check(Role == ROLE_Authority);
 	ExchangeNetRoles(bRemoteOwned);
-	
-	USceneComponent* SceneRootComponent = FixupNativeActorComponents(this);
-	// Set the actor's location and rotation.
+
+	USceneComponent* const SceneRootComponent = FixupNativeActorComponents(this);
 	if (SceneRootComponent != NULL)
 	{
-		SceneRootComponent->SetWorldTransform(SpawnTransform);
+		// Set the actor's location and rotation since it has a native rootcomponent
+		// Note that we respect any initial transformation the root component may have from the CDO, so the final transform
+		// might necessarily be exactly the passed-in UserSpawnTransform.
+ 		const FTransform RootTransform(SceneRootComponent->RelativeRotation, SceneRootComponent->RelativeLocation, SceneRootComponent->RelativeScale3D);
+ 		const FTransform FinalRootComponentTransform = RootTransform * UserSpawnTransform;
+		SceneRootComponent->SetWorldTransform(FinalRootComponentTransform);
 	}
 
 	// Call OnComponentCreated on all default (native) components
@@ -2459,17 +2463,21 @@ void AActor::PostSpawnInitialize(FTransform const& SpawnTransform, AActor* InOwn
 	if (!bDeferConstruction)
 	{
 		// Preserve original root component scale
-		FinishSpawning(SpawnTransform, true);
+		FinishSpawning(UserSpawnTransform, true);
 	}
 }
 
-void AActor::FinishSpawning(const FTransform& Transform, bool bIsDefaultTransform)
+void AActor::FinishSpawning(const FTransform& UserTransform, bool bIsDefaultTransform)
 {
 	if (ensure(!bHasFinishedSpawning))
 	{
 		bHasFinishedSpawning = true;
 
-		ExecuteConstruction(Transform, nullptr, bIsDefaultTransform);
+		// if we have a native root component, its transform is already set properly, so we use that.  it might not be the same as UserTransform 
+		// if the root component in the CDO has a nonzero transform.
+		const FTransform FinalRootComponentTransform = RootComponent ? RootComponent->ComponentToWorld : UserTransform;
+
+		ExecuteConstruction(FinalRootComponentTransform, nullptr, bIsDefaultTransform);
 		PostActorConstruction();
 	}
 }
diff --git a/Engine/Source/Runtime/Engine/Private/LevelActor.cpp b/Engine/Source/Runtime/Engine/Private/LevelActor.cpp
--- a/Engine/Source/Runtime/Engine/Private/LevelActor.cpp
+++ b/Engine/Source/Runtime/Engine/Private/LevelActor.cpp
@@ -245,7 +245,7 @@ AActor* UWorld::SpawnActor( UClass* Class, FVector const* Location, FRotator con
 	return SpawnActor(Class, &Transform, SpawnParameters);
 }
 
-AActor* UWorld::SpawnActor( UClass* Class, FTransform const* Transform, const FActorSpawnParameters& SpawnParameters )
+AActor* UWorld::SpawnActor( UClass* Class, FTransform const* UserTransformPtr, const FActorSpawnParameters& SpawnParameters )
 {
 	SCOPE_CYCLE_COUNTER(STAT_SpawnActorTime);
 	check( CurrentLevel ); 	
@@ -290,9 +290,9 @@ AActor* UWorld::SpawnActor( UClass* Class, FTransform const* Transform, const FA
 		UE_LOG(LogSpawn, Warning, TEXT("SpawnActor failed because we are in the process of tearing down the world"));
 		return NULL;
 	}
-	else if (Transform && Transform->ContainsNaN())
+	else if (UserTransformPtr && UserTransformPtr->ContainsNaN())
 	{
-		UE_LOG(LogSpawn, Warning, TEXT("SpawnActor failed because the given transform (%s) is invalid"), *Transform->ToString());
+		UE_LOG(LogSpawn, Warning, TEXT("SpawnActor failed because the given transform (%s) is invalid"), *(UserTransformPtr->ToString()));
 		return NULL;
 	}
 
@@ -305,16 +305,17 @@ AActor* UWorld::SpawnActor( UClass* Class, FTransform const* Transform, const FA
 
 	FName NewActorName = SpawnParameters.Name;
 	AActor* Template = SpawnParameters.Template;
-	// Use class's default actor as a template.
+
 	if( !Template )
 	{
+		// Use class's default actor as a template.
 		Template = Class->GetDefaultObject<AActor>();
 	}
 	else if (NewActorName.IsNone() && !Template->HasAnyFlags(RF_ClassDefaultObject))
 	{
 		NewActorName = MakeUniqueObjectName(LevelToSpawnIn, Template->GetClass(), *Template->GetFName().GetPlainNameString());
 	}
-	check(Template!=NULL);
+	check(Template);
 
 	// See if we can spawn on ded.server/client only etc (check NeedsLoadForClient & NeedsLoadForServer)
 	if(!CanCreateInCurrentContext(Template))
@@ -323,25 +324,7 @@ AActor* UWorld::SpawnActor( UClass* Class, FTransform const* Transform, const FA
 		return NULL;
 	}
 
-	const USceneComponent* TemplateRootComponent = Template->GetRootComponent();
-	FVector RelativeLocation = Transform ? Transform->GetLocation() : (TemplateRootComponent ? TemplateRootComponent->RelativeLocation : FVector::ZeroVector);
-	FRotator RelativeRotation = Transform ? Transform->GetRotation().Rotator() :	(TemplateRootComponent ? TemplateRootComponent->RelativeRotation : FRotator::ZeroRotator);
-	FVector RelativeScale = Transform ? Transform->GetScale3D() :	(TemplateRootComponent ? TemplateRootComponent->RelativeScale3D : FVector(1.f));
-
-	FTransform RelativeTransform(RelativeRotation, RelativeLocation, RelativeScale);
-	FTransform NewTransform;
-
-	// we do respect Template's RootComponent. If this is called using new relative transform, we apply to the template transform
-	// this way, we respect CDO's transform and multiply on top of it. 
-	if(TemplateRootComponent)
-	{
-		FTransform TemplateTransform(TemplateRootComponent->RelativeRotation, TemplateRootComponent->RelativeLocation, TemplateRootComponent->RelativeScale3D);
-		NewTransform = TemplateTransform * RelativeTransform;
-	}
-	else
-	{
-		NewTransform = RelativeTransform;
-	}
+	FTransform const UserTransform = UserTransformPtr ? *UserTransformPtr : FTransform::Identity;
 
 	PRAGMA_DISABLE_DEPRECATION_WARNINGS;
 	// handle existing (but deprecated) uses of bNoCollisionFail where user set it to true
@@ -373,12 +356,22 @@ AActor* UWorld::SpawnActor( UClass* Class, FTransform const* Transform, const FA
 	// note: we can't handle all cases here, since we don't know the full component hierarchy until after the actor is spawned
 	if (CollisionHandlingMethod == ESpawnActorCollisionHandlingMethod::DontSpawnIfColliding)
 	{
-		FVector NewLocation = NewTransform.GetLocation();
-		FRotator NewRotation = NewTransform.GetRotation().Rotator();
-		if (EncroachingBlockingGeometry(Template, NewLocation, NewRotation))
+		USceneComponent* const TemplateRootComponent = Template ? Template->GetRootComponent() : nullptr;
+
+		// Note that we respect any initial transformation the root component may have from the CDO, so the final transform
+		// might necessarily be exactly the passed-in UserTransform.
+		FTransform const FinalRootComponentTransform =
+			TemplateRootComponent
+			? FTransform(TemplateRootComponent->RelativeRotation, TemplateRootComponent->RelativeLocation, TemplateRootComponent->RelativeScale3D) * UserTransform
+			: UserTransform;
+
+		FVector const FinalRootLocation = FinalRootComponentTransform.GetLocation();
+		FRotator const FinalRootRotation = FinalRootComponentTransform.Rotator();
+
+		if (EncroachingBlockingGeometry(Template, FinalRootLocation, FinalRootRotation))
 		{
 			// a native component is colliding, that's enough to reject spawning
-			UE_LOG(LogSpawn, Warning, TEXT("SpawnActor failed because of collision at the spawn location [%s] for [%s]"), *NewLocation.ToString(), *Class->GetName());
+			UE_LOG(LogSpawn, Warning, TEXT("SpawnActor failed because of collision at the spawn location [%s] for [%s]"), *FinalRootLocation.ToString(), *Class->GetName());
 			return nullptr;
 		}
 	}
@@ -411,61 +404,13 @@ AActor* UWorld::SpawnActor( UClass* Class, FTransform const* Transform, const FA
 	// tell the actor what method to use, in case it was overridden
 	Actor->SpawnCollisionHandlingMethod = CollisionHandlingMethod;
 
-	Actor->PostSpawnInitialize(NewTransform, SpawnParameters.Owner, SpawnParameters.Instigator, SpawnParameters.bRemoteOwned, SpawnParameters.bNoFail, SpawnParameters.bDeferConstruction);
+	Actor->PostSpawnInitialize(UserTransform, SpawnParameters.Owner, SpawnParameters.Instigator, SpawnParameters.bRemoteOwned, SpawnParameters.bNoFail, SpawnParameters.bDeferConstruction);
 
 	if (Actor->IsPendingKill() && !SpawnParameters.bNoFail)
 	{
 		UE_LOG(LogSpawn, Warning, TEXT("SpawnActor failed because the spawned actor IsPendingKill"));
 		return NULL;
 	}
-// 
-// 	// actor should have all of its components now, do any collision checking and handling that we need to do
-// 	switch (CollisionHandlingMethod)
-// 	{
-// 	case ESpawnActorCollisionHandlingMethod::AdjustIfPossibleButAlwaysSpawn:
-// 	 	// Try to find a spawn position
-// 		{
-// 			FVector AdjustedLocation = NewLocation;
-// 			FRotator AdjustedRotation = NewRotation;
-// 			if (FindTeleportSpot(Actor, AdjustedLocation, AdjustedRotation))
-// 			{
-// 				Actor->SetActorLocationAndRotation(AdjustedLocation, AdjustedRotation);
-// 			}
-// 		}
-// 	 	break;
-// 	case ESpawnActorCollisionHandlingMethod::AdjustIfPossibleButDontSpawnIfColliding:
-// 	 	// Try to find a spawn position
-// 		{
-// 			FVector AdjustedLocation = NewLocation;
-// 			FRotator AdjustedRotation = NewRotation;
-// 			if (FindTeleportSpot(Actor, AdjustedLocation, AdjustedRotation))
-// 			{
-// 				Actor->SetActorLocationAndRotation(AdjustedLocation, AdjustedRotation);
-// 			}
-// 			else
-// 			{
-// 				UE_LOG(LogSpawn, Warning, TEXT("SpawnActor failed because of collision at the spawn location [%s] for [%s]"), *NewLocation.ToString(), *Class->GetName());
-// 				DestroyActor(Actor);
-// 				return nullptr;
-// 			}
-// 		}
-// 	 	break;
-// 	case ESpawnActorCollisionHandlingMethod::DontSpawnIfColliding:
-// 		// #todo: don't recheck components checked above?
-// 		if (EncroachingBlockingGeometry(Actor, NewLocation, NewRotation))
-// 	 	{
-// 	 		UE_LOG(LogSpawn, Warning, TEXT("SpawnActor failed because of collision at the spawn location [%s] for [%s]"), *NewLocation.ToString(), *Class->GetName());
-// 			DestroyActor(Actor);
-// 			return nullptr;
-// 	 	}
-// 	 	break;
-// 	// note we use "always spawn" as default, so treat undefined as that
-// 	case ESpawnActorCollisionHandlingMethod::Undefined:
-// 	case ESpawnActorCollisionHandlingMethod::AlwaysSpawn:
-// 	default:
-// 	 	// nothing to do, just proceed as normal
-// 	 	break;
-// 	}
 
 	Actor->CheckDefaultSubobjects();
 
