[UE-17070] CIS fix.

diff --git a/Engine/Source/Runtime/Engine/Private/Blueprint.cpp b/Engine/Source/Runtime/Engine/Private/Blueprint.cpp
--- a/Engine/Source/Runtime/Engine/Private/Blueprint.cpp
+++ b/Engine/Source/Runtime/Engine/Private/Blueprint.cpp
@@ -942,7 +942,11 @@ bool UBlueprint::GetBlueprintHierarchyFromClass(const UClass* InClass, TArray<UB
 #endif // #if WITH_EDITORONLY_DATA
 
 		// If valid, use stored ParentClass rather than the actual UClass::GetSuperClass(); handles the case when the class has not been recompiled yet after a reparent operation.
-		CurrentClass = BP->ParentClass ? BP->ParentClass : CurrentClass->GetSuperClass();
+		CurrentClass = BP->ParentClass;
+		if(CurrentClass == nullptr)
+		{
+			CurrentClass = CurrentClass->GetSuperClass();
+		}
 	}
 
 	return bNoErrors;
