UE4 - Fix crash when running editor with -game

diff --git a/Engine/Source/Runtime/Engine/Private/ActorEditor.cpp b/Engine/Source/Runtime/Engine/Private/ActorEditor.cpp
--- a/Engine/Source/Runtime/Engine/Private/ActorEditor.cpp
+++ b/Engine/Source/Runtime/Engine/Private/ActorEditor.cpp
@@ -47,7 +47,7 @@ void AActor::PostEditChangeProperty(FPropertyChangedEvent& PropertyChangedEvent)
 
 	// During SIE, allow components to reregistered and reconstructed in PostEditChangeProperty.
 	// This is essential as construction is deferred during spawning / duplication when in SIE.
-	if (GEditor->bIsSimulatingInEditor || ReregisterComponentsWhenModified())
+	if ((GEditor && GEditor->bIsSimulatingInEditor) || ReregisterComponentsWhenModified())
 	{
 		// In the Undo case we have an annotation storing information about constructed components and we do not want
 		// to improperly apply out of date changes so we need to skip registration of all blueprint created components
