#ue4 - Add nicer displayname for ECanBeCharacterBase enums.

diff --git a/Engine/Source/Runtime/Engine/Classes/Components/PrimitiveComponent.h b/Engine/Source/Runtime/Engine/Classes/Components/PrimitiveComponent.h
--- a/Engine/Source/Runtime/Engine/Classes/Components/PrimitiveComponent.h
+++ b/Engine/Source/Runtime/Engine/Classes/Components/PrimitiveComponent.h
@@ -38,14 +38,14 @@ UENUM()
 enum ECanBeCharacterBase
 {
 	/** Character cannot step up onto this Component. */
-	ECB_No,
+	ECB_No UMETA(DisplayName="No"),
 	/** Character can step up onto this Component. */
-	ECB_Yes,
+	ECB_Yes UMETA(DisplayName="Yes"),
 	/**
 	 * Owning actor determines whether character can step up onto this Component (default true unless overridden in code).
 	 * @see AActor::CanBeBaseForCharacter()
 	 */
-	ECB_Owner,
+	ECB_Owner UMETA(DisplayName="(Owner)"),
 	ECB_MAX,
 };
 
