#ue4 - FloatingPawnMovement uses faster FQuat code path for movement.

diff --git a/Engine/Source/Runtime/Engine/Private/FloatingPawnMovement.cpp b/Engine/Source/Runtime/Engine/Private/FloatingPawnMovement.cpp
--- a/Engine/Source/Runtime/Engine/Private/FloatingPawnMovement.cpp
+++ b/Engine/Source/Runtime/Engine/Private/FloatingPawnMovement.cpp
@@ -53,7 +53,7 @@ void UFloatingPawnMovement::TickComponent(float DeltaTime, enum ELevelTick TickT
 		if (!Delta.IsNearlyZero(1e-6f))
 		{
 			const FVector OldLocation = UpdatedComponent->GetComponentLocation();
-			const FRotator Rotation = UpdatedComponent->GetComponentRotation();
+			const FQuat Rotation = UpdatedComponent->GetComponentQuat();
 
 			FHitResult Hit(1.f);
 			SafeMoveUpdatedComponent(Delta, Rotation, true, Hit);
