[UE-17070] Fix previous CIS fix.

diff --git a/Engine/Source/Runtime/Engine/Private/Blueprint.cpp b/Engine/Source/Runtime/Engine/Private/Blueprint.cpp
--- a/Engine/Source/Runtime/Engine/Private/Blueprint.cpp
+++ b/Engine/Source/Runtime/Engine/Private/Blueprint.cpp
@@ -942,8 +942,11 @@ bool UBlueprint::GetBlueprintHierarchyFromClass(const UClass* InClass, TArray<UB
 #endif // #if WITH_EDITORONLY_DATA
 
 		// If valid, use stored ParentClass rather than the actual UClass::GetSuperClass(); handles the case when the class has not been recompiled yet after a reparent operation.
-		CurrentClass = BP->ParentClass;
-		if(CurrentClass == nullptr)
+		if(BP->ParentClass != nullptr)
+		{
+			CurrentClass = BP->ParentClass;
+		}
+		else
 		{
 			CurrentClass = CurrentClass->GetSuperClass();
 		}
