[UE-18875] Making it so EdGraphPins loaded early on, won't get added to the root set (which would prevent them from being deleted).

diff --git a/Engine/Source/Runtime/Engine/Classes/EdGraph/EdGraphPin.h b/Engine/Source/Runtime/Engine/Classes/EdGraph/EdGraphPin.h
--- a/Engine/Source/Runtime/Engine/Classes/EdGraph/EdGraphPin.h
+++ b/Engine/Source/Runtime/Engine/Classes/EdGraph/EdGraphPin.h
@@ -255,6 +255,10 @@ class UEdGraphPin : public UObject
 #endif // WITH_EDITORONLY_DATA
 
 public:
+	// UObject interface
+	virtual bool IsSafeForRootSet() const override { return false; }
+	// End UObject interface
+
 	/** Create a link. Note, this does not check that schema allows it, and will not break any existing connections */
 	ENGINE_API void MakeLinkTo(UEdGraphPin* ToPin);
 
