#ue4 - Remove invalid check() in FScopedMovementUpdate when processing overlaps. This was a valid case for sweeping while collision is disabled.

UE-18894

diff --git a/Engine/Source/Runtime/Engine/Private/Components/SceneComponent.cpp b/Engine/Source/Runtime/Engine/Private/Components/SceneComponent.cpp
--- a/Engine/Source/Runtime/Engine/Private/Components/SceneComponent.cpp
+++ b/Engine/Source/Runtime/Engine/Private/Components/SceneComponent.cpp
@@ -2417,7 +2417,6 @@ void FScopedMovementUpdate::AppendOverlapsAfterMove(const TArray<FOverlapInfo>&
 	else
 	{
 		// We don't expect any pending overlaps in the case of a teleport.
-		checkSlow(!bSweep);
 		checkSlow(NewPendingOverlaps.Num() == 0);
 		CurrentOverlapState = EOverlapState::eUnknown;
 		FinalOverlapCandidatesIndex = INDEX_NONE;
