Stereo rendering requires the full view matrix inverse.

diff --git a/Engine/Source/Runtime/Engine/Private/SceneView.cpp b/Engine/Source/Runtime/Engine/Private/SceneView.cpp
--- a/Engine/Source/Runtime/Engine/Private/SceneView.cpp
+++ b/Engine/Source/Runtime/Engine/Private/SceneView.cpp
@@ -573,9 +573,10 @@ void FSceneView::UpdateViewMatrix()
 	ViewMatrices.InvTranslatedViewProjectionMatrix = ViewMatrices.TranslatedViewProjectionMatrix.Inverse();
 
 	ViewProjectionMatrix = ViewMatrices.GetViewProjMatrix();
-	InvViewMatrix = ViewMatrices.GetInvViewMatrix();
-	InvViewProjectionMatrix = ViewMatrices.GetInvViewProjMatrix();
 
+	// Need to use a full view matrix inverse here as there are translations present
+	InvViewMatrix = ViewMatrices.ViewMatrix.Inverse();
+	InvViewProjectionMatrix = ViewMatrices.GetInvProjMatrix() * InvViewMatrix;
 }
 
 void FSceneView::SetScaledViewRect(FIntRect InScaledViewRect)
