Merge: #REPLAYS: Proper cleanup of replay net driver when world fails to load

diff --git a/Engine/Source/Runtime/Engine/Private/DemoNetDriver.cpp b/Engine/Source/Runtime/Engine/Private/DemoNetDriver.cpp
--- a/Engine/Source/Runtime/Engine/Private/DemoNetDriver.cpp
+++ b/Engine/Source/Runtime/Engine/Private/DemoNetDriver.cpp
@@ -350,6 +350,17 @@ bool UDemoNetDriver::InitConnectInternal( FString& Error )
 
 	if ( !GEngine->LoadMap( *WorldContext, DemoURL, NewPendingNetGame, LoadMapError ) )
 	{
+		StopDemo();
+
+		// If we don't have a world that means we failed loading the new world.
+		// Since there is no world, we must free the net driver ourselves
+		// Technically the pending net game should handle it, but things aren't quite setup properly to handle that either
+		if ( WorldContext->World() == NULL )
+		{
+			GEngine->DestroyNamedNetDriver( WorldContext->PendingNetGame, NetDriverName );
+			WorldContext->PendingNetGame = NULL;
+		}
+
 		Error = LoadMapError;
 		UE_LOG( LogDemo, Error, TEXT( "UDemoNetDriver::InitConnect: LoadMap failed: failed: %s" ), *Error );
 		GameInstance->HandleDemoPlaybackFailure( EDemoPlayFailure::Generic, FString( TEXT( "LoadMap failed" ) ) );
@@ -1489,7 +1500,10 @@ void UDemoNetDriver::ReplayStreamingReady( bool bSuccess, bool bRecord )
 		
 		const double StartTime = FPlatformTime::Seconds();
 
-		InitConnectInternal( Error );
+		if ( !InitConnectInternal( Error ) )
+		{
+			return;
+		}
 
 		if ( ReplayStreamer->IsLive() && ReplayStreamer->GetTotalDemoTime() > 15 * 1000 )
 		{
