Fixed compile error for materials using a CustomNode and TextureObject on OpenGL platforms
UE-18981

diff --git a/Engine/Source/Runtime/Engine/Private/Materials/HLSLMaterialTranslator.h b/Engine/Source/Runtime/Engine/Private/Materials/HLSLMaterialTranslator.h
--- a/Engine/Source/Runtime/Engine/Private/Materials/HLSLMaterialTranslator.h
+++ b/Engine/Source/Runtime/Engine/Private/Materials/HLSLMaterialTranslator.h
@@ -3777,14 +3777,14 @@ protected:
 			case MCT_Texture2D:
 				InputParamDecl += TEXT("Texture2D ");
 				InputParamDecl += Custom->Inputs[i].InputName;
-				InputParamDecl += TEXT(", sampler ");
+				InputParamDecl += TEXT(", SamplerState ");
 				InputParamDecl += Custom->Inputs[i].InputName;
 				InputParamDecl += TEXT("Sampler ");
 				break;
 			case MCT_TextureCube:
 				InputParamDecl += TEXT("TextureCube ");
 				InputParamDecl += Custom->Inputs[i].InputName;
-				InputParamDecl += TEXT(", sampler ");
+				InputParamDecl += TEXT(", SamplerState ");
 				InputParamDecl += Custom->Inputs[i].InputName;
 				InputParamDecl += TEXT("Sampler ");
 				break;
