Do not send events to hidden sub-levels

diff --git a/Engine/Source/Runtime/Engine/Private/LevelScriptActor.cpp b/Engine/Source/Runtime/Engine/Private/LevelScriptActor.cpp
--- a/Engine/Source/Runtime/Engine/Private/LevelScriptActor.cpp
+++ b/Engine/Source/Runtime/Engine/Private/LevelScriptActor.cpp
@@ -75,7 +75,7 @@ bool ALevelScriptActor::RemoteEvent(FName EventName)
 	for( TArray<ULevel*>::TConstIterator it = GetWorld()->GetLevels().CreateConstIterator(); it; ++it )
 	{
 		ULevel* CurLevel = *it;
-		if( CurLevel )
+		if( CurLevel && CurLevel->bIsVisible )
 		{
 			ALevelScriptActor* LSA = CurLevel->GetLevelScriptActor();
 			if( LSA )
