Small changes to Visual Logger, to limit overhead when  logging is disabled.

diff --git a/Engine/Source/Runtime/Engine/Private/VisualLogger/VisualLogger.cpp b/Engine/Source/Runtime/Engine/Private/VisualLogger/VisualLogger.cpp
--- a/Engine/Source/Runtime/Engine/Private/VisualLogger/VisualLogger.cpp
+++ b/Engine/Source/Runtime/Engine/Private/VisualLogger/VisualLogger.cpp
@@ -20,14 +20,16 @@ DEFINE_STAT(STAT_VisualLog);
 DEFINE_LOG_CATEGORY(LogVisual);
 
 TMap<UObject*, TArray<TWeakObjectPtr<const UObject> > > FVisualLogger::RedirectionMap;
+int32 FVisualLogger::bIsRecording = false;
 
 bool FVisualLogger::CheckVisualLogInputInternal(const class UObject* Object, const struct FLogCategoryBase& Category, ELogVerbosity::Type Verbosity, UWorld **World, FVisualLogEntry **CurrentEntry)
 {
-	FVisualLogger& VisualLogger = FVisualLogger::Get();
-	if (!Object || (GEngine && GEngine->bDisableAILogging) || VisualLogger.IsRecording() == false || Object->HasAnyFlags(RF_ClassDefaultObject))
+	if (FVisualLogger::IsRecording() == false || !Object || (GEngine && GEngine->bDisableAILogging) || Object->HasAnyFlags(RF_ClassDefaultObject))
 	{
 		return false;
 	}
+
+	FVisualLogger& VisualLogger = FVisualLogger::Get();
 	const FName CategoryName = Category.GetCategoryName();
 	if (VisualLogger.IsBlockedForAllCategories() && VisualLogger.IsWhiteListed(CategoryName) == false)
 	{
diff --git a/Engine/Source/Runtime/Engine/Public/VisualLogger/VisualLogger.h b/Engine/Source/Runtime/Engine/Public/VisualLogger/VisualLogger.h
--- a/Engine/Source/Runtime/Engine/Public/VisualLogger/VisualLogger.h
+++ b/Engine/Source/Runtime/Engine/Public/VisualLogger/VisualLogger.h
@@ -13,56 +13,56 @@
 #define CONNECT_OBJECT_WITH_VLOG(Src, Dest)
 
 // Text, regular log
-#define UE_VLOG(LogOwner, CategoryName, Verbosity, Format, ...) FVisualLogger::CategorizedLogf(LogOwner, CategoryName, ELogVerbosity::Verbosity, INDEX_NONE, Format, ##__VA_ARGS__)
-#define UE_CVLOG(Condition, LogOwner, CategoryName, Verbosity, Format, ...)  if(Condition) {UE_VLOG(LogOwner, CategoryName, Verbosity, Format, ##__VA_ARGS__);} 
+#define UE_VLOG(LogOwner, CategoryName, Verbosity, Format, ...) if( FVisualLogger::IsRecording() ) FVisualLogger::CategorizedLogf(LogOwner, CategoryName, ELogVerbosity::Verbosity, INDEX_NONE, Format, ##__VA_ARGS__)
+#define UE_CVLOG(Condition, LogOwner, CategoryName, Verbosity, Format, ...)  if(FVisualLogger::IsRecording() && Condition) {UE_VLOG(LogOwner, CategoryName, Verbosity, Format, ##__VA_ARGS__);} 
 // Text, log with output to regular unreal logs too
-#define UE_VLOG_UELOG(LogOwner, CategoryName, Verbosity, Format, ...) { FVisualLogger::CategorizedLogf(LogOwner, CategoryName, ELogVerbosity::Verbosity, INDEX_NONE, Format, ##__VA_ARGS__); UE_LOG(CategoryName, Verbosity, Format, ##__VA_ARGS__); }
-#define UE_CVLOG_UELOG(Condition, LogOwner, CategoryName, Verbosity, Format, ...)  if(Condition) {UE_VLOG_DEBUG(LogOwner, CategoryName, Verbosity, Format, ##__VA_ARGS__);} 
+#define UE_VLOG_UELOG(LogOwner, CategoryName, Verbosity, Format, ...) { if(FVisualLogger::IsRecording()) FVisualLogger::CategorizedLogf(LogOwner, CategoryName, ELogVerbosity::Verbosity, INDEX_NONE, Format, ##__VA_ARGS__); UE_LOG(CategoryName, Verbosity, Format, ##__VA_ARGS__); }
+#define UE_CVLOG_UELOG(Condition, LogOwner, CategoryName, Verbosity, Format, ...)  if(FVisualLogger::IsRecording() && Condition) {UE_VLOG_UELOG(LogOwner, CategoryName, Verbosity, Format, ##__VA_ARGS__);} 
 // Segment shape
-#define UE_VLOG_SEGMENT(LogOwner, CategoryName, Verbosity, SegmentStart, SegmentEnd, Color, Format, ...) FVisualLogger::GeometryShapeLogf(LogOwner, CategoryName, ELogVerbosity::Verbosity, INDEX_NONE, SegmentStart, SegmentEnd, Color, Format, ##__VA_ARGS__)
-#define UE_CVLOG_SEGMENT(Condition, LogOwner, CategoryName, Verbosity, SegmentStart, SegmentEnd, Color, Format, ...) if(Condition) {UE_VLOG_SEGMENT(LogOwner, CategoryName, Verbosity, INDEX_NONE, SegmentStart, SegmentEnd, Color, Format, ##__VA_ARGS__);}
+#define UE_VLOG_SEGMENT(LogOwner, CategoryName, Verbosity, SegmentStart, SegmentEnd, Color, Format, ...) if(FVisualLogger::IsRecording()) FVisualLogger::GeometryShapeLogf(LogOwner, CategoryName, ELogVerbosity::Verbosity, INDEX_NONE, SegmentStart, SegmentEnd, Color, Format, ##__VA_ARGS__)
+#define UE_CVLOG_SEGMENT(Condition, LogOwner, CategoryName, Verbosity, SegmentStart, SegmentEnd, Color, Format, ...) if(FVisualLogger::IsRecording() && Condition) {UE_VLOG_SEGMENT(LogOwner, CategoryName, Verbosity, INDEX_NONE, SegmentStart, SegmentEnd, Color, Format, ##__VA_ARGS__);}
 // Segment shape
-#define UE_VLOG_SEGMENT_THICK(LogOwner, CategoryName, Verbosity, SegmentStart, SegmentEnd, Color, Thickness, Format, ...) FVisualLogger::GeometryShapeLogf(LogOwner, CategoryName, ELogVerbosity::Verbosity, INDEX_NONE, SegmentStart, SegmentEnd, Color, Format, ##__VA_ARGS__)
-#define UE_CVLOG_SEGMENT_THICK(Condition, LogOwner, CategoryName, Verbosity, SegmentStart, SegmentEnd, Color, Thickness, Format, ...) if(Condition) {UE_VLOG_SEGMENT_THICK(LogOwner, CategoryName, Verbosity, INDEX_NONE, SegmentStart, SegmentEnd, Color, Format, ##__VA_ARGS__);} 
+#define UE_VLOG_SEGMENT_THICK(LogOwner, CategoryName, Verbosity, SegmentStart, SegmentEnd, Color, Thickness, Format, ...) if(FVisualLogger::IsRecording()) FVisualLogger::GeometryShapeLogf(LogOwner, CategoryName, ELogVerbosity::Verbosity, INDEX_NONE, SegmentStart, SegmentEnd, Color, Format, ##__VA_ARGS__)
+#define UE_CVLOG_SEGMENT_THICK(Condition, LogOwner, CategoryName, Verbosity, SegmentStart, SegmentEnd, Color, Thickness, Format, ...) if(FVisualLogger::IsRecording() && Condition) {UE_VLOG_SEGMENT_THICK(LogOwner, CategoryName, Verbosity, INDEX_NONE, SegmentStart, SegmentEnd, Color, Format, ##__VA_ARGS__);} 
 // Localization as sphere shape
-#define UE_VLOG_LOCATION(LogOwner, CategoryName, Verbosity, Location, Radius, Color, Format, ...) FVisualLogger::GeometryShapeLogf(LogOwner, CategoryName, ELogVerbosity::Verbosity, INDEX_NONE, Location, Radius, Color, Format, ##__VA_ARGS__)
-#define UE_CVLOG_LOCATION(Condition, LogOwner, CategoryName, Verbosity, Location, Radius, Color, Format, ...)  if(Condition) {UE_VLOG_LOCATION(LogOwner, CategoryName, Verbosity, Location, Radius, Color, Format, ##__VA_ARGS__);} 
+#define UE_VLOG_LOCATION(LogOwner, CategoryName, Verbosity, Location, Radius, Color, Format, ...) if(FVisualLogger::IsRecording()) FVisualLogger::GeometryShapeLogf(LogOwner, CategoryName, ELogVerbosity::Verbosity, INDEX_NONE, Location, Radius, Color, Format, ##__VA_ARGS__)
+#define UE_CVLOG_LOCATION(Condition, LogOwner, CategoryName, Verbosity, Location, Radius, Color, Format, ...)  if(FVisualLogger::IsRecording() && Condition) {UE_VLOG_LOCATION(LogOwner, CategoryName, Verbosity, Location, Radius, Color, Format, ##__VA_ARGS__);} 
 // Box shape
-#define UE_VLOG_BOX(LogOwner, CategoryName, Verbosity, Box, Color, Format, ...)  FVisualLogger::GeometryShapeLogf(LogOwner, CategoryName, ELogVerbosity::Verbosity, INDEX_NONE, Box, FMatrix::Identity, Color, Format, ##__VA_ARGS__)
-#define UE_CVLOG_BOX(Condition, LogOwner, CategoryName, Verbosity, Box, Color, Format, ...)  if(Condition) {UE_VLOG_BOX(Actor, CategoryName, Verbosity, Box, FMatrix::Identity, Color, Format, ##__VA_ARGS__);} 
+#define UE_VLOG_BOX(LogOwner, CategoryName, Verbosity, Box, Color, Format, ...) if(FVisualLogger::IsRecording()) FVisualLogger::GeometryShapeLogf(LogOwner, CategoryName, ELogVerbosity::Verbosity, INDEX_NONE, Box, FMatrix::Identity, Color, Format, ##__VA_ARGS__)
+#define UE_CVLOG_BOX(Condition, LogOwner, CategoryName, Verbosity, Box, Color, Format, ...)  if(FVisualLogger::IsRecording() && Condition) {UE_VLOG_BOX(Actor, CategoryName, Verbosity, Box, FMatrix::Identity, Color, Format, ##__VA_ARGS__);} 
 // Oriented box shape
-#define UE_VLOG_OBOX(LogOwner, CategoryName, Verbosity, Box, Matrix, Color, Format, ...)  FVisualLogger::GeometryShapeLogf(LogOwner, CategoryName, ELogVerbosity::Verbosity, INDEX_NONE, Box, Matrix, Color, Format, ##__VA_ARGS__)
-#define UE_CVLOG_OBOX(Condition, LogOwner, CategoryName, Verbosity, Box, Matrix, Color, Format, ...)  if(Condition) {UE_VLOG_BOX(Actor, CategoryName, Verbosity, Box, Matrix, Color, Format, ##__VA_ARGS__);} 
+#define UE_VLOG_OBOX(LogOwner, CategoryName, Verbosity, Box, Matrix, Color, Format, ...) if(FVisualLogger::IsRecording()) FVisualLogger::GeometryShapeLogf(LogOwner, CategoryName, ELogVerbosity::Verbosity, INDEX_NONE, Box, Matrix, Color, Format, ##__VA_ARGS__)
+#define UE_CVLOG_OBOX(Condition, LogOwner, CategoryName, Verbosity, Box, Matrix, Color, Format, ...)  if(FVisualLogger::IsRecording() && Condition) {UE_VLOG_BOX(Actor, CategoryName, Verbosity, Box, Matrix, Color, Format, ##__VA_ARGS__);} 
 // Cone shape
-#define UE_VLOG_CONE(LogOwner, CategoryName, Verbosity, Orgin, Direction, Length, Angle, Color, Format, ...) FVisualLogger::GeometryShapeLogf(LogOwner, CategoryName, ELogVerbosity::Verbosity, INDEX_NONE, Orgin, Direction, Length, Angle, Color, Format, ##__VA_ARGS__)
-#define UE_CVLOG_CONE(Condition, LogOwner, CategoryName, Verbosity, Orgin, Direction, Length, Angle, Color, Format, ...)  if(Condition) {UE_VLOG_CONE(Actor, CategoryName, Verbosity, Orgin, Direction, Length, Angle, Color, Format, ##__VA_ARGS__);} 
+#define UE_VLOG_CONE(LogOwner, CategoryName, Verbosity, Orgin, Direction, Length, Angle, Color, Format, ...) if(FVisualLogger::IsRecording()) FVisualLogger::GeometryShapeLogf(LogOwner, CategoryName, ELogVerbosity::Verbosity, INDEX_NONE, Orgin, Direction, Length, Angle, Color, Format, ##__VA_ARGS__)
+#define UE_CVLOG_CONE(Condition, LogOwner, CategoryName, Verbosity, Orgin, Direction, Length, Angle, Color, Format, ...)  if(FVisualLogger::IsRecording() && Condition) {UE_VLOG_CONE(Actor, CategoryName, Verbosity, Orgin, Direction, Length, Angle, Color, Format, ##__VA_ARGS__);} 
 // Cylinder shape
-#define UE_VLOG_CYLINDER(LogOwner, CategoryName, Verbosity, Start, End, Radius, Color, Format, ...) FVisualLogger::GeometryShapeLogf(LogOwner, CategoryName, ELogVerbosity::Verbosity, INDEX_NONE, Start, End, Radius, Color, Format, ##__VA_ARGS__)
-#define UE_CVLOG_CYLINDER(Condition, LogOwner, CategoryName, Verbosity, Start, End, Radius, Color, Format, ...)  if(Condition) {UE_VLOG_CYLINDER(LogOwner, CategoryName, Verbosity, Start, End, Radius, Color, Format, ##__VA_ARGS__);} 
+#define UE_VLOG_CYLINDER(LogOwner, CategoryName, Verbosity, Start, End, Radius, Color, Format, ...) if(FVisualLogger::IsRecording()) FVisualLogger::GeometryShapeLogf(LogOwner, CategoryName, ELogVerbosity::Verbosity, INDEX_NONE, Start, End, Radius, Color, Format, ##__VA_ARGS__)
+#define UE_CVLOG_CYLINDER(Condition, LogOwner, CategoryName, Verbosity, Start, End, Radius, Color, Format, ...)  if(FVisualLogger::IsRecording() && Condition) {UE_VLOG_CYLINDER(LogOwner, CategoryName, Verbosity, Start, End, Radius, Color, Format, ##__VA_ARGS__);} 
 // Capsule shape
-#define UE_VLOG_CAPSULE(LogOwner, CategoryName, Verbosity, Center, HalfHeight, Radius, Rotation, Color, Format, ...) FVisualLogger::GeometryShapeLogf(LogOwner, CategoryName, ELogVerbosity::Verbosity, INDEX_NONE, Center, HalfHeight, Radius, Rotation, Color, Format, ##__VA_ARGS__)
-#define UE_CVLOG_CAPSULE(Condition, LogOwner, CategoryName, Verbosity, Center, HalfHeight, Radius, Rotation, Color, Format, ...)  if(Condition) {UE_VLOG_CAPSULE(LogOwner, CategoryName, Verbosity, Center, HalfHeight, Radius, Rotation, Color, Format, ##__VA_ARGS__);} 
+#define UE_VLOG_CAPSULE(LogOwner, CategoryName, Verbosity, Center, HalfHeight, Radius, Rotation, Color, Format, ...) if(FVisualLogger::IsRecording()) FVisualLogger::GeometryShapeLogf(LogOwner, CategoryName, ELogVerbosity::Verbosity, INDEX_NONE, Center, HalfHeight, Radius, Rotation, Color, Format, ##__VA_ARGS__)
+#define UE_CVLOG_CAPSULE(Condition, LogOwner, CategoryName, Verbosity, Center, HalfHeight, Radius, Rotation, Color, Format, ...)  if(FVisualLogger::IsRecording() && Condition) {UE_VLOG_CAPSULE(LogOwner, CategoryName, Verbosity, Center, HalfHeight, Radius, Rotation, Color, Format, ##__VA_ARGS__);} 
 // Histogram data for 2d graphs 
-#define UE_VLOG_HISTOGRAM(LogOwner, CategoryName, Verbosity, GraphName, DataName, Data) FVisualLogger::HistogramDataLogf(LogOwner, CategoryName, ELogVerbosity::Verbosity, INDEX_NONE, GraphName, DataName, Data, FColor::White, TEXT(""))
-#define UE_CVLOG_HISTOGRAM(Condition, LogOwner, CategoryName, Verbosity, GraphName, DataName, Data) if(Condition) {UE_VLOG_HISTOGRAM(LogOwner, CategoryName, Verbosity, GraphName, DataName, Data);} 
+#define UE_VLOG_HISTOGRAM(LogOwner, CategoryName, Verbosity, GraphName, DataName, Data) if(FVisualLogger::IsRecording()) FVisualLogger::HistogramDataLogf(LogOwner, CategoryName, ELogVerbosity::Verbosity, INDEX_NONE, GraphName, DataName, Data, FColor::White, TEXT(""))
+#define UE_CVLOG_HISTOGRAM(Condition, LogOwner, CategoryName, Verbosity, GraphName, DataName, Data) if(FVisualLogger::IsRecording() && Condition) {UE_VLOG_HISTOGRAM(LogOwner, CategoryName, Verbosity, GraphName, DataName, Data);} 
 // NavArea or vertically pulled convex shape
-#define UE_VLOG_PULLEDCONVEX(LogOwner, CategoryName, Verbosity, ConvexPoints, MinZ, MaxZ, Color, Format, ...) FVisualLogger::NavAreaShapeLogf(LogOwner, CategoryName, ELogVerbosity::Verbosity, INDEX_NONE, ConvexPoints, MinZ, MaxZ, Color, Format, ##__VA_ARGS__)
-#define UE_CVLOG_PULLEDCONVEX(Condition, LogOwner, CategoryName, Verbosity, ConvexPoints, MinZ, MaxZ, Color, Format, ...) if(Condition) {UE_VLOG_PULLEDCONVEX(LogOwner, CategoryName, Verbosity, ConvexPoints, MinZ, MaxZ, Color, Format, ##__VA_ARGS__);}
+#define UE_VLOG_PULLEDCONVEX(LogOwner, CategoryName, Verbosity, ConvexPoints, MinZ, MaxZ, Color, Format, ...) if(FVisualLogger::IsRecording()) FVisualLogger::NavAreaShapeLogf(LogOwner, CategoryName, ELogVerbosity::Verbosity, INDEX_NONE, ConvexPoints, MinZ, MaxZ, Color, Format, ##__VA_ARGS__)
+#define UE_CVLOG_PULLEDCONVEX(Condition, LogOwner, CategoryName, Verbosity, ConvexPoints, MinZ, MaxZ, Color, Format, ...) if(FVisualLogger::IsRecording() && Condition) {UE_VLOG_PULLEDCONVEX(LogOwner, CategoryName, Verbosity, ConvexPoints, MinZ, MaxZ, Color, Format, ##__VA_ARGS__);}
 // regular 3d mesh shape to log
-#define UE_VLOG_MESH(LogOwner, CategoryName, Verbosity, Vertices, Indices, Color, Format, ...) FVisualLogger::GeometryShapeLogf(LogOwner, CategoryName, ELogVerbosity::Verbosity, INDEX_NONE, Vertices, Indices, Format, ##__VA_ARGS__)
-#define UE_CVLOG_MESH(Condition, LogOwner, CategoryName, Verbosity, Vertices, Indexes, Color, Format, ...) if(Condition) {UE_VLOG_MESH(LogOwner, CategoryName, Verbosity, Vertices, Indices, Color, Format, ##__VA_ARGS__);}
+#define UE_VLOG_MESH(LogOwner, CategoryName, Verbosity, Vertices, Indices, Color, Format, ...) if (FVisualLogger::IsRecording()) FVisualLogger::GeometryShapeLogf(LogOwner, CategoryName, ELogVerbosity::Verbosity, INDEX_NONE, Vertices, Indices, Format, ##__VA_ARGS__)
+#define UE_CVLOG_MESH(Condition, LogOwner, CategoryName, Verbosity, Vertices, Indexes, Color, Format, ...) if(FVisualLogger::IsRecording() && Condition) {UE_VLOG_MESH(LogOwner, CategoryName, Verbosity, Vertices, Indices, Color, Format, ##__VA_ARGS__);}
 // 2d convex poly shape
-#define UE_VLOG_CONVEXPOLY(LogOwner, CategoryName, Verbosity, Points, Color, Format, ...)  FVisualLogger::GeometryConvexLogf(LogOwner, CategoryName, ELogVerbosity::Verbosity, INDEX_NONE, Points, Color, Format, ##__VA_ARGS__)
-#define UE_CVLOG_CONVEXPOLY(Condition, LogOwner, CategoryName, Verbosity, Points, Color, Format, ...)  if(Condition) {UE_VLOG_CONVEXPOLY(LogOwner, CategoryName, Verbosity, Points, Color, Format, ##__VA_ARGS__);}
+#define UE_VLOG_CONVEXPOLY(LogOwner, CategoryName, Verbosity, Points, Color, Format, ...) if(FVisualLogger::IsRecording()) FVisualLogger::GeometryConvexLogf(LogOwner, CategoryName, ELogVerbosity::Verbosity, INDEX_NONE, Points, Color, Format, ##__VA_ARGS__)
+#define UE_CVLOG_CONVEXPOLY(Condition, LogOwner, CategoryName, Verbosity, Points, Color, Format, ...)  if(FVisualLogger::IsRecording() && Condition) {UE_VLOG_CONVEXPOLY(LogOwner, CategoryName, Verbosity, Points, Color, Format, ##__VA_ARGS__);}
 
 
 #define DECLARE_VLOG_EVENT(EventName) extern FVisualLogEventBase EventName;
 #define DEFINE_VLOG_EVENT(EventName, Verbosity, UserFriendlyDesc) FVisualLogEventBase EventName(TEXT(#EventName), TEXT(UserFriendlyDesc), ELogVerbosity::Verbosity); 
 
-#define UE_VLOG_EVENTS(LogOwner, TagNameToLog, ...) FVisualLogger::EventLog(LogOwner, TagNameToLog, ##__VA_ARGS__)
-#define UE_CVLOG_EVENTS(Condition, LogOwner, TagNameToLog, ...) if(Condition) {UE_VLOG_EVENTS(LogOwner, TagNameToLog, ##__VA_ARGS__);}
-#define UE_VLOG_EVENT_WITH_DATA(LogOwner, LogEvent, ...) FVisualLogger::EventLog(LogOwner, LogEvent, ##__VA_ARGS__)
-#define UE_CVLOG_EVENT_WITH_DATA(Condition, LogOwner, LogEvent, ...) if(Condition) {UE_VLOG_EVENT_WITH_DATA(LogOwner, LogEvent, ##__VA_ARGS__);}
+#define UE_VLOG_EVENTS(LogOwner, TagNameToLog, ...) if(FVisualLogger::IsRecording()) FVisualLogger::EventLog(LogOwner, TagNameToLog, ##__VA_ARGS__)
+#define UE_CVLOG_EVENTS(Condition, LogOwner, TagNameToLog, ...) if(FVisualLogger::IsRecording() && Condition) {UE_VLOG_EVENTS(LogOwner, TagNameToLog, ##__VA_ARGS__);}
+#define UE_VLOG_EVENT_WITH_DATA(LogOwner, LogEvent, ...) if(FVisualLogger::IsRecording()) FVisualLogger::EventLog(LogOwner, LogEvent, ##__VA_ARGS__)
+#define UE_CVLOG_EVENT_WITH_DATA(Condition, LogOwner, LogEvent, ...) if(FVisualLogger::IsRecording() && Condition) {UE_VLOG_EVENT_WITH_DATA(LogOwner, LogEvent, ##__VA_ARGS__);}
 
 #define UE_IFVLOG(__code_block__) __code_block__;
 
@@ -208,7 +208,7 @@ public:
 	/** Starts visual log collecting and recording */
 	void SetIsRecording(bool InIsRecording);
 	/** return information is vlog recording is enabled or not */
-	bool IsRecording() const { return !!bIsRecording; }
+	FORCEINLINE static bool IsRecording() { return !!bIsRecording; }
 
 	/** Starts visual log collecting and recording */
 	void SetIsRecordingToFile(bool InIsRecording);
@@ -274,8 +274,6 @@ protected:
 	TMap<const class UObject*, TWeakObjectPtr<const class UObject> > ObjectToPointerMap;
 	// if set all categories are blocked from logging
 	int32 bBlockedAllCategories : 1;
-	// if set we are recording and collecting all vlog data
-	int32 bIsRecording : 1;
 	// if set we are recording to file
 	int32 bIsRecordingToFile : 1;
 	// variable set (from cheat manager) when logging is active on server
@@ -288,6 +286,9 @@ protected:
 	bool UseBinaryFileDevice;
 	/** Handle for timer used to serialize all waiting logs */
 	FTimerHandle VisualLoggerCleanupTimerHandle;
+
+	// if set we are recording and collecting all vlog data
+	static int32 bIsRecording;
 };
 
 // Unfortunately needs to be a #define since it uses GET_VARARGS_RESULT which uses the va_list stuff which operates on the
