UE-18936 Fix crash when passing null Component to ComponentOverlapComponents_NEW

diff --git a/Engine/Source/Runtime/Engine/Private/KismetSystemLibrary.cpp b/Engine/Source/Runtime/Engine/Private/KismetSystemLibrary.cpp
--- a/Engine/Source/Runtime/Engine/Private/KismetSystemLibrary.cpp
+++ b/Engine/Source/Runtime/Engine/Private/KismetSystemLibrary.cpp
@@ -1282,31 +1282,34 @@ bool UKismetSystemLibrary::ComponentOverlapComponents_NEW(UPrimitiveComponent* C
 {
 	OutComponents.Empty();
 
-	static FName ComponentOverlapComponentsName(TEXT("ComponentOverlapComponents"));
-	FComponentQueryParams Params(ComponentOverlapComponentsName);	
-	Params.bTraceAsyncScene = true;
-	Params.AddIgnoredActors(ActorsToIgnore);
+	if(Component != nullptr)
+	{
+		static FName ComponentOverlapComponentsName(TEXT("ComponentOverlapComponents"));
+		FComponentQueryParams Params(ComponentOverlapComponentsName);	
+		Params.bTraceAsyncScene = true;
+		Params.AddIgnoredActors(ActorsToIgnore);
 
-	TArray<FOverlapResult> Overlaps;
+		TArray<FOverlapResult> Overlaps;
 
-	FCollisionObjectQueryParams ObjectParams;
-	for (auto Iter = ObjectTypes.CreateConstIterator(); Iter; ++Iter)
-	{
-		const ECollisionChannel & Channel = UCollisionProfile::Get()->ConvertToCollisionChannel(false, *Iter);
-		ObjectParams.AddObjectTypesToQuery(Channel);
-	}
+		FCollisionObjectQueryParams ObjectParams;
+		for (auto Iter = ObjectTypes.CreateConstIterator(); Iter; ++Iter)
+		{
+			const ECollisionChannel & Channel = UCollisionProfile::Get()->ConvertToCollisionChannel(false, *Iter);
+			ObjectParams.AddObjectTypesToQuery(Channel);
+		}
 
-	check( Component->GetWorld());
-	Component->GetWorld()->ComponentOverlapMulti(Overlaps, Component, ComponentTransform.GetTranslation(), ComponentTransform.GetRotation(), Params, ObjectParams);
+		check( Component->GetWorld());
+		Component->GetWorld()->ComponentOverlapMulti(Overlaps, Component, ComponentTransform.GetTranslation(), ComponentTransform.GetRotation(), Params, ObjectParams);
 
-	for (int32 OverlapIdx=0; OverlapIdx<Overlaps.Num(); ++OverlapIdx)
-	{
-		FOverlapResult const& O = Overlaps[OverlapIdx];
-		if (O.Component.IsValid())
-		{ 
-			if ( !ComponentClassFilter || O.Component.Get()->IsA(ComponentClassFilter) )
-			{
-				OutComponents.Add(O.Component.Get());
+		for (int32 OverlapIdx=0; OverlapIdx<Overlaps.Num(); ++OverlapIdx)
+		{
+			FOverlapResult const& O = Overlaps[OverlapIdx];
+			if (O.Component.IsValid())
+			{ 
+				if ( !ComponentClassFilter || O.Component.Get()->IsA(ComponentClassFilter) )
+				{
+					OutComponents.Add(O.Component.Get());
+				}
 			}
 		}
 	}
