Allow loading replay checkpoints while paused & update the timeline. Fixes UE-15066.

diff --git a/Engine/Source/Runtime/Engine/Private/DemoNetDriver.cpp b/Engine/Source/Runtime/Engine/Private/DemoNetDriver.cpp
--- a/Engine/Source/Runtime/Engine/Private/DemoNetDriver.cpp
+++ b/Engine/Source/Runtime/Engine/Private/DemoNetDriver.cpp
@@ -590,10 +590,7 @@ void UDemoNetDriver::TickDispatch(float DeltaSeconds)
 			}
 		}
 
-		if ( World->GetWorldSettings()->Pauser == NULL )
-		{
-			TickDemoPlayback( DeltaSeconds );
-		}
+		TickDemoPlayback( DeltaSeconds );
 	}
 }
 
@@ -1392,8 +1389,11 @@ void UDemoNetDriver::TickDemoPlayback( float DeltaSeconds )
 		return;
 	}
 
-	// Advance demo time by seconds passed
-	DemoCurrentTime += DeltaSeconds;
+	// Advance demo time by seconds passed if we're not paused
+	if ( World->GetWorldSettings()->Pauser == NULL )
+	{
+		DemoCurrentTime += DeltaSeconds;
+	}
 
 	// Clamp time
 	if ( DemoCurrentTime > DemoTotalTime )
