Make sure parallel blend and clothing tick happens before PreCloth ends

diff --git a/Engine/Source/Runtime/Engine/Classes/Components/SkeletalMeshComponent.h b/Engine/Source/Runtime/Engine/Classes/Components/SkeletalMeshComponent.h
--- a/Engine/Source/Runtime/Engine/Classes/Components/SkeletalMeshComponent.h
+++ b/Engine/Source/Runtime/Engine/Classes/Components/SkeletalMeshComponent.h
@@ -754,7 +754,7 @@ public:
 	void TickAnimation(float DeltaTime);
 
 	/** Tick Clothing Animation , basically this is called inside TickComponent */
-	void TickClothing(float DeltaTime);
+	void TickClothing(float DeltaTime, FTickFunction& ThisTickFunction);
 
 	/** Store cloth simulation data into OutClothSimData */
 	void GetUpdateClothSimulationData(TArray<FClothSimulData>& OutClothSimData, USkeletalMeshComponent* OverrideLocalRootComponent = nullptr);
@@ -947,7 +947,7 @@ public:
 
 
 	/** Take the results of the physics and blend them with the animation state (based on the PhysicsWeight parameter), and update the SpaceBases array. */
-	void BlendInPhysics();	
+	void BlendInPhysics(FTickFunction& ThisTickFunction);	
 
 	/** 
 	 * Intialize PhysicsAssetInstance for the physicsAsset 
@@ -1112,7 +1112,7 @@ public:
 	FVector GetClosestCollidingRigidBodyLocation(const FVector& TestLocation) const;
 
 	/** Calls needed cloth updates */
-	void PreClothTick(float DeltaTime);
+	void PreClothTick(float DeltaTime, FTickFunction& ThisTickFunction);
 	
 	/** Set physics transforms for all bodies */
 	void ApplyDeltaToAllPhysicsTransforms(const FVector& DeltaLocation, const FQuat& DeltaRotation);
diff --git a/Engine/Source/Runtime/Engine/Private/PhysicsEngine/PhysAnim.cpp b/Engine/Source/Runtime/Engine/Private/PhysicsEngine/PhysAnim.cpp
--- a/Engine/Source/Runtime/Engine/Private/PhysicsEngine/PhysAnim.cpp
+++ b/Engine/Source/Runtime/Engine/Private/PhysicsEngine/PhysAnim.cpp
@@ -127,6 +127,7 @@ void UpdateWorldBoneTM(TArray<FAssetWorldBoneTM> & WorldBoneTMs, const TArray<FT
 
 void USkeletalMeshComponent::PerformBlendPhysicsBones(const TArray<FBoneIndexType>& InRequiredBones, TArray<FTransform>& InLocalAtoms)
 {
+	SCOPE_CYCLE_COUNTER(STAT_BlendInPhysics);
 	// Get drawscale from Owner (if there is one)
 	FVector TotalScale3D = ComponentToWorld.GetScale3D();
 	FVector RecipScale3D = TotalScale3D.Reciprocal();
@@ -369,9 +370,8 @@ bool USkeletalMeshComponent::DoAnyPhysicsBodiesHaveWeight() const
 
 TAutoConsoleVariable<int32> CVarUseParallelBlendPhysics(TEXT("a.ParallelBlendPhysics"), 1, TEXT("If 1, physics blending will be run across the task graph system. If 0, blending will run purely on the game thread"));
 
-void USkeletalMeshComponent::BlendInPhysics()
+void USkeletalMeshComponent::BlendInPhysics(FTickFunction& ThisTickFunction)
 {
-	SCOPE_CYCLE_COUNTER(STAT_BlendInPhysics);
 	check(IsInGameThread());
 
 	// Can't do anything without a SkeletalMesh
@@ -408,6 +408,8 @@ void USkeletalMeshComponent::BlendInPhysics()
 
 			check(!IsValidRef(ParallelBlendPhysicsCompletionTask));
 			ParallelBlendPhysicsCompletionTask = TGraphTask<FParallelBlendPhysicsCompletionTask>::CreateTask(&Prerequistes).ConstructAndDispatchWhenReady(this);
+
+			ThisTickFunction.GetCompletionHandle()->DontCompleteUntil(ParallelBlendPhysicsCompletionTask);
 		}else
 		{
 			PerformBlendPhysicsBones(RequiredBones, LocalAtoms);
diff --git a/Engine/Source/Runtime/Engine/Private/SkeletalMeshComponentPhysics.cpp b/Engine/Source/Runtime/Engine/Private/SkeletalMeshComponentPhysics.cpp
--- a/Engine/Source/Runtime/Engine/Private/SkeletalMeshComponentPhysics.cpp
+++ b/Engine/Source/Runtime/Engine/Private/SkeletalMeshComponentPhysics.cpp
@@ -44,7 +44,7 @@ void FSkeletalMeshComponentPreClothTickFunction::ExecuteTick(float DeltaTime, en
 
 	if ((TickType == LEVELTICK_All) && Target && !Target->HasAnyFlags(RF_PendingKill | RF_Unreachable))
 	{
-		Target->PreClothTick(DeltaTime);
+		Target->PreClothTick(DeltaTime, *this);
 	}
 }
 
@@ -3115,7 +3115,7 @@ void USkeletalMeshComponent::ProcessClothCollisionWithEnvironment()
 
 #endif// #if WITH_CLOTH_COLLISION_DETECTION
 
-void USkeletalMeshComponent::PreClothTick(float DeltaTime)
+void USkeletalMeshComponent::PreClothTick(float DeltaTime, FTickFunction& ThisTickFunction)
 {
 	//IMPORTANT!
 	//
@@ -3147,7 +3147,7 @@ void USkeletalMeshComponent::PreClothTick(float DeltaTime)
 	{
 		if (IsRegistered())
 		{
-			BlendInPhysics();
+			BlendInPhysics(ThisTickFunction);
 		}
 	}
 
@@ -3157,7 +3157,7 @@ void USkeletalMeshComponent::PreClothTick(float DeltaTime)
 	// if skeletal mesh has clothing assets, call TickClothing
 	if (SkeletalMesh && SkeletalMesh->ClothingAssets.Num() > 0)
 	{
-		TickClothing(DeltaTime);
+		TickClothing(DeltaTime, ThisTickFunction);
 	}
 #endif
 }
@@ -3728,13 +3728,14 @@ void USkeletalMeshComponent::PerformTickClothing(float DeltaTime)
 #endif// #if WITH_APEX_CLOTHING
 }
 
-void USkeletalMeshComponent::TickClothing(float DeltaTime)
+void USkeletalMeshComponent::TickClothing(float DeltaTime, FTickFunction& ThisTickFunction)
 {
 	if(IsValidRef(ParallelBlendPhysicsCompletionTask))
 	{
 		FGraphEventArray Prerequistes;
 		Prerequistes.Add(ParallelBlendPhysicsCompletionTask);
-		TGraphTask<FTickClothingTask>::CreateTask(&Prerequistes).ConstructAndDispatchWhenReady(this, DeltaTime);
+		FGraphEventRef TickClothingCompletionEvent = TGraphTask<FTickClothingTask>::CreateTask(&Prerequistes).ConstructAndDispatchWhenReady(this, DeltaTime);
+		ThisTickFunction.GetCompletionHandle()->DontCompleteUntil(TickClothingCompletionEvent);
 	}else
 	{
 		PerformTickClothing(DeltaTime);
