Move EndCloth tick group until the end of the frame. This allows cloth sim to run during PostPhysics and PostUpdateWork.

diff --git a/Engine/Source/Runtime/Engine/Classes/Engine/EngineBaseTypes.h b/Engine/Source/Runtime/Engine/Classes/Engine/EngineBaseTypes.h
--- a/Engine/Source/Runtime/Engine/Classes/Engine/EngineBaseTypes.h
+++ b/Engine/Source/Runtime/Engine/Classes/Engine/EngineBaseTypes.h
@@ -59,15 +59,15 @@ enum ETickingGroup
 	/** Any item that needs to be updated after rigid body simulation is done, but before cloth is simulation is done. */
 	TG_StartCloth,
 
-	/** Any item that can be run during cloth simulation. */	
-	TG_EndCloth,
-
 	/** Any item that needs rigid body and cloth simulation to be complete before being executed. */
 	TG_PostPhysics,
 
 	/** Any item that needs the update work to be done before being ticked. */
 	TG_PostUpdateWork,
 
+	/** Special tick group that ends cloth simulation. */
+	TG_EndCloth,
+
 	/** Special tick group that is not actually a tick group. After every tick group this is repeatedly re-run until there are no more newly spawned items to run. */
 	TG_NewlySpawned,
 
diff --git a/Engine/Source/Runtime/Engine/Private/LevelTick.cpp b/Engine/Source/Runtime/Engine/Private/LevelTick.cpp
--- a/Engine/Source/Runtime/Engine/Private/LevelTick.cpp
+++ b/Engine/Source/Runtime/Engine/Private/LevelTick.cpp
@@ -1177,14 +1177,10 @@ void UWorld::Tick( ELevelTick TickType, float DeltaSeconds )
 			SCOPE_CYCLE_COUNTER(STAT_TG_StartCloth);
 		RunTickGroup(TG_StartCloth);
 		}
-		{
-			SCOPE_CYCLE_COUNTER(STAT_TG_EndCloth);
-		RunTickGroup(TG_EndCloth);
-		}
 		{
 			SCOPE_CYCLE_COUNTER(STAT_TG_PostPhysics);
-		RunTickGroup(TG_PostPhysics);
-	}
+			RunTickGroup(TG_PostPhysics);
+		}
 	}
 	else if( bIsPaused )
 	{
@@ -1279,7 +1275,11 @@ void UWorld::Tick( ELevelTick TickType, float DeltaSeconds )
 		SCOPE_CYCLE_COUNTER(STAT_TickTime);
 		{
 			SCOPE_CYCLE_COUNTER(STAT_TG_PostUpdateWork);
-		RunTickGroup(TG_PostUpdateWork);
+			RunTickGroup(TG_PostUpdateWork);
+		}
+		{
+			SCOPE_CYCLE_COUNTER(STAT_TG_EndCloth);
+			RunTickGroup(TG_EndCloth);
 		}
 		FTickTaskManagerInterface::Get().EndFrame(); 
 
