UE-8324 Default values for decal components are not suitable for a component workflow
added DecalSize to the decal component

diff --git a/Engine/Source/Editor/ComponentVisualizers/Private/ComponentVisualizers.cpp b/Engine/Source/Editor/ComponentVisualizers/Private/ComponentVisualizers.cpp
--- a/Engine/Source/Editor/ComponentVisualizers/Private/ComponentVisualizers.cpp
+++ b/Engine/Source/Editor/ComponentVisualizers/Private/ComponentVisualizers.cpp
@@ -15,6 +15,7 @@
 #include "SpringArmComponentVisualizer.h"
 #include "SplineComponentVisualizer.h"
 #include "SplineMeshComponentVisualizer.h"
+#include "DecalComponentVisualizer.h"
 #include "SensingComponentVisualizer.h"
 #include "SpringComponentVisualizer.h"
 #include "PrimitiveComponentVisualizer.h"
@@ -42,6 +43,7 @@ void FComponentVisualizersModule::StartupModule()
 	RegisterComponentVisualizer(UPawnSensingComponent::StaticClass()->GetFName(), MakeShareable(new FSensingComponentVisualizer));
 	RegisterComponentVisualizer(UPhysicsSpringComponent::StaticClass()->GetFName(), MakeShareable(new FSpringComponentVisualizer));
 	RegisterComponentVisualizer(UPrimitiveComponent::StaticClass()->GetFName(), MakeShareable(new FPrimitiveComponentVisualizer));
+	RegisterComponentVisualizer(UDecalComponent::StaticClass()->GetFName(), MakeShareable(new FDecalComponentVisualizer));
 }
 
 void FComponentVisualizersModule::ShutdownModule()
diff --git a/Engine/Source/Editor/ComponentVisualizers/Private/DecalComponentVisualizer.cpp b/Engine/Source/Editor/ComponentVisualizers/Private/DecalComponentVisualizer.cpp
new file mode 100644
--- /dev/null
+++ b/Engine/Source/Editor/ComponentVisualizers/Private/DecalComponentVisualizer.cpp
@@ -0,0 +1,19 @@
+// Copyright 1998-2015 Epic Games, Inc. All Rights Reserved.
+
+#include "ComponentVisualizersPrivatePCH.h"
+
+#include "DecalComponentVisualizer.h"
+
+
+void FDecalComponentVisualizer::DrawVisualization( const UActorComponent* Component, const FSceneView* View, FPrimitiveDrawInterface* PDI )
+{
+	const UDecalComponent* DecalComponent = Cast<const UDecalComponent>(Component);
+	if(DecalComponent)
+	{
+		const FMatrix LocalToWorld = DecalComponent->ComponentToWorld.ToMatrixWithScale();
+		
+		const FLinearColor DrawColor = FColor(0, 157, 0, 255);
+
+		DrawOrientedWireBox(PDI, LocalToWorld.GetOrigin(), LocalToWorld.GetScaledAxis( EAxis::X ), LocalToWorld.GetScaledAxis( EAxis::Y ), LocalToWorld.GetScaledAxis( EAxis::Z ), DecalComponent->DecalSize, DrawColor, SDPG_World);
+	}
+}
diff --git a/Engine/Source/Editor/ComponentVisualizers/Private/DecalComponentVisualizer.h b/Engine/Source/Editor/ComponentVisualizers/Private/DecalComponentVisualizer.h
new file mode 100644
--- /dev/null
+++ b/Engine/Source/Editor/ComponentVisualizers/Private/DecalComponentVisualizer.h
@@ -0,0 +1,14 @@
+// Copyright 1998-2015 Epic Games, Inc. All Rights Reserved.
+
+#pragma once
+
+#include "ComponentVisualizer.h"
+#include "Components/DecalComponent.h"
+
+class FDecalComponentVisualizer : public FComponentVisualizer
+{
+public:
+	// Begin FComponentVisualizer interface
+	virtual void DrawVisualization(const UActorComponent* Component, const FSceneView* View, FPrimitiveDrawInterface* PDI) override;
+	// End FComponentVisualizer interface
+};
diff --git a/Engine/Source/Runtime/Core/Public/UObject/ObjectVersion.h b/Engine/Source/Runtime/Core/Public/UObject/ObjectVersion.h
--- a/Engine/Source/Runtime/Core/Public/UObject/ObjectVersion.h
+++ b/Engine/Source/Runtime/Core/Public/UObject/ObjectVersion.h
@@ -560,6 +560,8 @@ enum EUnrealEngineObjectUE4Version
 	VER_UE4_MAX_ANGULAR_VELOCITY_DEFAULT,
 	// Build Adjacency index buffer for clothing tessellation
 	VER_UE4_APEX_CLOTH_TESSELLATION,
+	// Added DecalSize member, solved backward compatibility
+	VER_UE4_DECAL_SIZE,
 	// -----<new versions can be added before this line>-------------------------------------------------
 	// - this needs to be the last line (see note below)
 	VER_UE4_AUTOMATIC_VERSION_PLUS_ONE,
diff --git a/Engine/Source/Runtime/Engine/Classes/Components/DecalComponent.h b/Engine/Source/Runtime/Engine/Classes/Components/DecalComponent.h
--- a/Engine/Source/Runtime/Engine/Classes/Components/DecalComponent.h
+++ b/Engine/Source/Runtime/Engine/Classes/Components/DecalComponent.h
@@ -32,6 +32,10 @@ class ENGINE_API UDecalComponent : public USceneComponent
 	UPROPERTY(EditAnywhere, BlueprintReadOnly, Category = Decal)
 	float FadeScreenSize;
 
+	/** Size if the scale if 1,1,1, technically redundant but there for convenience */
+	UPROPERTY(EditAnywhere, BlueprintReadOnly, Category = Decal, meta=(AllowPreserveRatio = "true"))
+	FVector DecalSize;
+
 	/** Sets the sort order for the decal component. Higher values draw later (on top). This will force the decal to reattach */
 	UFUNCTION(BlueprintCallable, Category = "Rendering|Components|Decal")
 	void SetSortOrder(int32 Value);
@@ -101,6 +105,8 @@ public:
 	virtual const UObject* AdditionalStatObject() const override;
 	// End UActorComponent Interface
 
+	virtual void Serialize(FArchive& Ar) override;
+
 	// Begin USceneComponent Interface
 	virtual FBoxSphereBounds CalcBounds(const FTransform& LocalToWorld) const override;
 	// End USceneComponent Interface
diff --git a/Engine/Source/Runtime/Engine/Classes/Engine/DecalActor.h b/Engine/Source/Runtime/Engine/Classes/Engine/DecalActor.h
--- a/Engine/Source/Runtime/Engine/Classes/Engine/DecalActor.h
+++ b/Engine/Source/Runtime/Engine/Classes/Engine/DecalActor.h
@@ -39,10 +39,9 @@ private_subobject:
 	UPROPERTY()
 	UBillboardComponent* SpriteComponent;
 
-	/* Reference to the selected visualization box component */
-	DEPRECATED_FORGAME(4.6, "BoxComponent should not be accessed directly, please use GetBoxComponent() function instead. BoxComponent will soon be private and your code will not compile.")
+	// formerly we used this component to draw a box, now we use the DecalComponentVisualizer
 	UPROPERTY()
-	UBoxComponent* BoxComponent;
+	UBoxComponent* BoxComponent_DEPRECATED;
 #endif
 
 public:
@@ -69,6 +68,10 @@ public:
 	// End AActor interface.
 #endif // WITH_EDITOR
 
+	
+	virtual void Serialize(FArchive& Ar) override;
+	virtual void PostLoad() override;
+
 public:
 
 	/** Returns Decal subobject **/
@@ -78,7 +81,5 @@ public:
 	UArrowComponent* GetArrowComponent() const;
 	/** Returns SpriteComponent subobject **/
 	UBillboardComponent* GetSpriteComponent() const;
-	/** Returns BoxComponent subobject **/
-	UBoxComponent* GetBoxComponent() const;
 #endif
 };
diff --git a/Engine/Source/Runtime/Engine/Private/Components/DecalComponent.cpp b/Engine/Source/Runtime/Engine/Private/Components/DecalComponent.cpp
--- a/Engine/Source/Runtime/Engine/Private/Components/DecalComponent.cpp
+++ b/Engine/Source/Runtime/Engine/Private/Components/DecalComponent.cpp
@@ -24,7 +24,10 @@ FDeferredDecalProxy::FDeferredDecalProxy(const UDecalComponent* InComponent)
 
 	Component = InComponent;
 	DecalMaterial = EffectiveMaterial;
+
 	ComponentTrans = InComponent->GetComponentToWorld();
+	ComponentTrans.SetScale3D(ComponentTrans.GetScale3D() * InComponent->DecalSize);
+
 	DrawInGame = InComponent->ShouldRender();
 	bOwnerSelected = InComponent->IsOwnerSelected();
 	SortOrder = InComponent->SortOrder;
@@ -36,8 +39,20 @@ void FDeferredDecalProxy::SetTransform(const FTransform& InComponentToWorld)
 }
 
 UDecalComponent::UDecalComponent(const FObjectInitializer& ObjectInitializer)
-: Super(ObjectInitializer), FadeScreenSize(0.01)
+	: Super(ObjectInitializer)
+	, FadeScreenSize(0.01)
+	, DecalSize(128.0f, 256.0f, 256.0f)
+{
+}
+
+void UDecalComponent::Serialize(FArchive& Ar)
 {
+	Super::Serialize(Ar);
+
+	if (Ar.UE4Ver() < VER_UE4_DECAL_SIZE)
+	{
+		DecalSize = FVector(1.0f, 1.0f, 1.0f);
+	}
 }
 
 void UDecalComponent::SetLifeSpan(const float LifeSpan)
@@ -104,7 +119,7 @@ FDeferredDecalProxy* UDecalComponent::CreateSceneProxy()
 
 FBoxSphereBounds UDecalComponent::CalcBounds(const FTransform& LocalToWorld) const
 {
-	return FBoxSphereBounds();
+	return FBoxSphereBounds(FVector(0, 0, 0), DecalSize, DecalSize.Size());
 }
 
 void UDecalComponent::CreateRenderState_Concurrent()
diff --git a/Engine/Source/Runtime/Engine/Private/DecalActor.cpp b/Engine/Source/Runtime/Engine/Private/DecalActor.cpp
--- a/Engine/Source/Runtime/Engine/Private/DecalActor.cpp
+++ b/Engine/Source/Runtime/Engine/Private/DecalActor.cpp
@@ -21,25 +21,11 @@ ADecalActor::ADecalActor(const FObjectInitializer& ObjectInitializer)
 : Super(ObjectInitializer)
 {
 	Decal = CreateDefaultSubobject<UDecalComponent>(TEXT("NewDecalComponent"));
-	Decal->RelativeScale3D = FVector(128.0f, 256.0f, 256.0f);
-
 	Decal->RelativeRotation = FRotator(-90, 0, 0);
 
 	RootComponent = Decal;
 
 #if WITH_EDITORONLY_DATA
-	BoxComponent = CreateEditorOnlyDefaultSubobject<UBoxComponent>(TEXT("DrawBox0"));
-	if (BoxComponent != nullptr)
-	{
-		BoxComponent->SetCollisionProfileName(UCollisionProfile::NoCollision_ProfileName);
-
-		BoxComponent->ShapeColor = FColor(80, 80, 200, 255);
-		BoxComponent->bDrawOnlyIfSelected = true;
-		BoxComponent->InitBoxExtent(FVector(1.0f, 1.0f, 1.0f));
-
-		BoxComponent->AttachParent = Decal;
-	}
-
 	ArrowComponent = CreateEditorOnlyDefaultSubobject<UArrowComponent>(TEXT("ArrowComponent0"));
 	SpriteComponent = CreateEditorOnlyDefaultSubobject<UBillboardComponent>(TEXT("Sprite"));
 
@@ -109,8 +95,21 @@ void ADecalActor::PostEditChangeProperty(FPropertyChangedEvent& PropertyChangedE
 	}
 }
 
+// @return 1 / x or 0 if the result would be a NaN
+static float SaveInv(float x)
+{
+	return FMath::IsNearlyZero(x) ? 0.0f : (1.0f / x);
+}
+
 void ADecalActor::EditorApplyScale(const FVector& DeltaScale, const FVector* PivotLocation, bool bAltDown, bool bShiftDown, bool bCtrlDown)
 {
+	FVector Mul(1.0f, 1.0f, 1.0f);
+
+	if(Decal)
+	{
+		Mul = FVector(SaveInv(Decal->DecalSize.X), SaveInv(Decal->DecalSize.Y), SaveInv(Decal->DecalSize.Z));
+	}
+
 	const FVector ModifiedScale = DeltaScale * (AActor::bUsePercentageBasedScaling ? DecalEditorConstants::PercentageScalingMultiplier : DecalEditorConstants::AdditiveScalingMultiplier);
 
 	Super::EditorApplyScale(ModifiedScale, PivotLocation, bAltDown, bShiftDown, bCtrlDown);
@@ -148,6 +147,43 @@ class UMaterialInstanceDynamic* ADecalActor::CreateDynamicMaterialInstance()
 	return Decal ? Decal->CreateDynamicMaterialInstance() : NULL;
 }
 
+void ADecalActor::Serialize(FArchive& Ar)
+{
+	if (Ar.UE4Ver() < VER_UE4_DECAL_SIZE)
+	{
+		if(Decal)
+		{
+			// before Super::Serialize(Ar);
+			Decal->RelativeScale3D = FVector(128.0f, 256.0f, 256.0f);
+		}
+	}
+
+	Super::Serialize(Ar);
+
+	if (Ar.UE4Ver() < VER_UE4_DECAL_SIZE)
+	{
+		if(Decal)
+		{
+			// after Super::Serialize(Ar);
+			Decal->DecalSize = FVector(1.0f, 1.0f, 1.0f);
+		}
+	}
+}
+
+void ADecalActor::PostLoad()
+{
+	 Super::PostLoad();
+
+#if WITH_EDITORONLY_DATA
+	if(BoxComponent_DEPRECATED)
+	{
+		// formerly we used this component to draw a box, now we use the DecalComponentVisualizer
+		BoxComponent_DEPRECATED->DestroyComponent();
+		BoxComponent_DEPRECATED = 0;
+	}
+#endif
+}
+
 /** Returns Decal subobject **/
 UDecalComponent* ADecalActor::GetDecal() const { return Decal; }
 #if WITH_EDITORONLY_DATA
@@ -155,6 +191,4 @@ UDecalComponent* ADecalActor::GetDecal() const { return Decal; }
 UArrowComponent* ADecalActor::GetArrowComponent() const { return ArrowComponent; }
 /** Returns SpriteComponent subobject **/
 UBillboardComponent* ADecalActor::GetSpriteComponent() const { return SpriteComponent; }
-/** Returns BoxComponent subobject **/
-UBoxComponent* ADecalActor::GetBoxComponent() const { return BoxComponent; }
 #endif
