Updated use of deprecated linked list method.

diff --git a/Engine/Source/Runtime/Engine/Private/Materials/MaterialUniformExpressions.cpp b/Engine/Source/Runtime/Engine/Private/Materials/MaterialUniformExpressions.cpp
--- a/Engine/Source/Runtime/Engine/Private/Materials/MaterialUniformExpressions.cpp
+++ b/Engine/Source/Runtime/Engine/Private/Materials/MaterialUniformExpressions.cpp
@@ -43,7 +43,7 @@ FMaterialUniformExpressionType::FMaterialUniformExpressionType(
 	SerializationConstructor(InSerializationConstructor)
 {
 	// Put the type in the type list until the name subsystem/type map are initialized.
-	(new TLinkedList<FMaterialUniformExpressionType*>(this))->Link(GetTypeList());
+	(new TLinkedList<FMaterialUniformExpressionType*>(this))->LinkHead(GetTypeList());
 }
 
 FArchive& operator<<(FArchive& Ar,FMaterialUniformExpression*& Ref)
diff --git a/Engine/Source/Runtime/Engine/Private/SceneView.cpp b/Engine/Source/Runtime/Engine/Private/SceneView.cpp
--- a/Engine/Source/Runtime/Engine/Private/SceneView.cpp
+++ b/Engine/Source/Runtime/Engine/Private/SceneView.cpp
@@ -170,7 +170,7 @@ void FSceneViewStateReference::Allocate()
 	check(!Reference);
 	Reference = GetRendererModule().AllocateViewState();
 	GlobalListLink = TLinkedList<FSceneViewStateReference*>(this);
-	GlobalListLink.Link(GetSceneViewStateList());
+	GlobalListLink.LinkHead(GetSceneViewStateList());
 }
 
 void FSceneViewStateReference::Destroy()
diff --git a/Engine/Source/Runtime/Engine/Private/StaticBoundShaderState.cpp b/Engine/Source/Runtime/Engine/Private/StaticBoundShaderState.cpp
--- a/Engine/Source/Runtime/Engine/Private/StaticBoundShaderState.cpp
+++ b/Engine/Source/Runtime/Engine/Private/StaticBoundShaderState.cpp
@@ -25,14 +25,14 @@ FGlobalBoundShaderStateResource::FGlobalBoundShaderStateResource()
 	// Add this resource to the global list in the rendering thread.
 	if(IsInRenderingThread())
 	{
-		GlobalListLink.Link(GetGlobalBoundShaderStateList());
+		GlobalListLink.LinkHead(GetGlobalBoundShaderStateList());
 	}
 	else
 	{
 		ENQUEUE_UNIQUE_RENDER_COMMAND_ONEPARAMETER(
 			LinkGlobalBoundShaderStateResource,FGlobalBoundShaderStateResource*,Resource,this,
 			{
-				Resource->GlobalListLink.Link(GetGlobalBoundShaderStateList());
+				Resource->GlobalListLink.LinkHead(GetGlobalBoundShaderStateList());
 			});
 	}
 }
diff --git a/Engine/Source/Runtime/Engine/Private/TextureRenderTarget.cpp b/Engine/Source/Runtime/Engine/Private/TextureRenderTarget.cpp
--- a/Engine/Source/Runtime/Engine/Private/TextureRenderTarget.cpp
+++ b/Engine/Source/Runtime/Engine/Private/TextureRenderTarget.cpp
@@ -157,7 +157,7 @@ void FDeferredUpdateResource::AddToDeferredUpdateList( bool OnlyUpdateOnce )
 	if( !bExists )
 	{
 		UpdateListLink = TLinkedList<FDeferredUpdateResource*>(this);
-		UpdateListLink.Link(UpdateList);
+		UpdateListLink.LinkHead(UpdateList);
 		bNeedsUpdate = true;
 	}
 	bOnlyUpdateOnce=OnlyUpdateOnce;
diff --git a/Engine/Source/Runtime/RenderCore/Private/RenderResource.cpp b/Engine/Source/Runtime/RenderCore/Private/RenderResource.cpp
--- a/Engine/Source/Runtime/RenderCore/Private/RenderResource.cpp
+++ b/Engine/Source/Runtime/RenderCore/Private/RenderResource.cpp
@@ -22,7 +22,7 @@ void FRenderResource::InitResource()
 	if(!bInitialized)
 	{
 		ResourceLink = TLinkedList<FRenderResource*>(this);
-		ResourceLink.Link(GetResourceList());
+		ResourceLink.LinkHead(GetResourceList());
 		if(GIsRHIInitialized)
 		{
 			InitDynamicRHI();
diff --git a/Engine/Source/Runtime/RenderCore/Public/UniformBuffer.h b/Engine/Source/Runtime/RenderCore/Public/UniformBuffer.h
--- a/Engine/Source/Runtime/RenderCore/Public/UniformBuffer.h
+++ b/Engine/Source/Runtime/RenderCore/Public/UniformBuffer.h
@@ -221,7 +221,7 @@ public:
 
 		if (bRegisterForAutoBinding)
 		{
-			GlobalListLink.Link(GetStructList());
+			GlobalListLink.LinkHead(GetStructList());
 			FName StrutTypeFName(StructTypeName);
 			// Verify that during FName creation there's no case conversion
 			checkSlow(FCString::Strcmp(StructTypeName, *StrutTypeFName.GetPlainNameString()) == 0);
diff --git a/Engine/Source/Runtime/ShaderCore/Private/Shader.cpp b/Engine/Source/Runtime/ShaderCore/Private/Shader.cpp
--- a/Engine/Source/Runtime/ShaderCore/Private/Shader.cpp
+++ b/Engine/Source/Runtime/ShaderCore/Private/Shader.cpp
@@ -99,7 +99,7 @@ FShaderType::FShaderType(
 	check(FCString::Strlen(InName) < NAME_SIZE);
 
 	// register this shader type
-	GlobalListLink.Link(GetTypeList());
+	GlobalListLink.LinkHead(GetTypeList());
 	GetNameToTypeMap().Add(FName(InName), this);
 
 	// Assign the shader type the next unassigned hash index.
diff --git a/Engine/Source/Runtime/ShaderCore/Private/VertexFactory.cpp b/Engine/Source/Runtime/ShaderCore/Private/VertexFactory.cpp
--- a/Engine/Source/Runtime/ShaderCore/Private/VertexFactory.cpp
+++ b/Engine/Source/Runtime/ShaderCore/Private/VertexFactory.cpp
@@ -121,7 +121,7 @@ FVertexFactoryType::FVertexFactoryType(
 	checkf(!bInitializedSerializationHistory, TEXT("VF type was loaded after engine init, use ELoadingPhase::PostConfigInit on your module to cause it to load earlier."));
 
 	// Add this vertex factory type to the global list.
-	GlobalListLink.Link(GetTypeList());
+	GlobalListLink.LinkHead(GetTypeList());
 
 	// Assign the vertex factory type the next unassigned hash index.
 	HashIndex = NextHashIndex++;
