Guard against null pointer in lightmap enforcement check. UE-15548

diff --git a/Engine/Source/Runtime/Engine/Private/StaticMesh.cpp b/Engine/Source/Runtime/Engine/Private/StaticMesh.cpp
--- a/Engine/Source/Runtime/Engine/Private/StaticMesh.cpp
+++ b/Engine/Source/Runtime/Engine/Private/StaticMesh.cpp
@@ -2311,9 +2311,16 @@ void UStaticMesh::EnforceLightmapRestrictions()
 
 	int32 NumUVs = 16;
 
-	for (int32 LODIndex = 0; LODIndex < RenderData->LODResources.Num(); ++LODIndex)
+	if (RenderData)
+	{
+		for (int32 LODIndex = 0; LODIndex < RenderData->LODResources.Num(); ++LODIndex)
+		{
+			NumUVs = FMath::Min(RenderData->LODResources[LODIndex].GetNumTexCoords(),NumUVs);
+		}
+	}
+	else
 	{
-		NumUVs = FMath::Min(RenderData->LODResources[LODIndex].GetNumTexCoords(), NumUVs);
+		NumUVs = 1;
 	}
 
 	// Clamp LightMapCoordinateIndex to be valid for all lightmap uvs
