UE-18852 ElementalDemo crashes due to nil constant buffer
added error messages for materials using nodes that are only supported in the Surface domain, some content might need to be fixed

diff --git a/Engine/Source/Runtime/Engine/Private/Materials/HLSLMaterialTranslator.h b/Engine/Source/Runtime/Engine/Private/Materials/HLSLMaterialTranslator.h
--- a/Engine/Source/Runtime/Engine/Private/Materials/HLSLMaterialTranslator.h
+++ b/Engine/Source/Runtime/Engine/Private/Materials/HLSLMaterialTranslator.h
@@ -2161,12 +2161,12 @@ protected:
 
 	virtual int32 ObjectRadius() override
 	{
-		return AddInlinedCodeChunk(MCT_Float,TEXT("Primitive.ObjectWorldPositionAndRadius.w"));		
+		return GetPrimitiveProperty(MCT_Float, TEXT("ObjectRadius"), TEXT("ObjectWorldPositionAndRadius.w"));		
 	}
 
 	virtual int32 ObjectBounds() override
 	{
-		return AddInlinedCodeChunk(MCT_Float3,TEXT("Primitive.ObjectBounds.xyz"));		
+		return GetPrimitiveProperty(MCT_Float3, TEXT("ObjectBounds"), TEXT("ObjectBounds.xyz"));
 	}
 
 	virtual int32 DistanceCullFade() override
@@ -3302,8 +3302,17 @@ protected:
 				}
 				else if (DestCoordBasis == MCB_Local)
 				{
-					//TODO: need Primitive.PrevWorldToLocal
-					//TODO: inconsistent with TransformLocal<TO>World with instancing
+					const EMaterialDomain Domain = (const EMaterialDomain)Material->GetMaterialDomain();
+
+					if(Domain != MD_Surface)
+					{
+						// TODO: for decals we could support it
+						Errorf(TEXT("This transformation is only supported in the 'Surface' material domain."));
+						return INDEX_NONE;
+					}
+
+					// TODO: need Primitive.PrevWorldToLocal
+					// TODO: inconsistent with TransformLocal<TO>World with instancing
 					CodeStr = TEXT("mul(<A>, <MATRIX>(Primitive.WorldToLocal))");
 				}
 				else if (DestCoordBasis == MCB_TranslatedWorld)
@@ -3467,7 +3476,7 @@ protected:
 
 	virtual int32 ObjectOrientation() override
 	{ 
-		return AddInlinedCodeChunk(MCT_Float3,TEXT("Primitive.ObjectOrientation.xyz"));	
+		return GetPrimitiveProperty(MCT_Float3, TEXT("ObjectOrientation"), TEXT("ObjectOrientation.xyz"));
 	}
 
 	virtual int32 RotateAboutAxis(int32 NormalizedRotationAxisAndAngleIndex, int32 PositionOnAxisIndex, int32 PositionIndex) override
@@ -3982,6 +3991,21 @@ protected:
 
 		return AddInlinedCodeChunk(MCT_Float, TEXT("EyeAdaptationLookup()"));
 	}
+
+	// to only have one piece of code dealing with error handling if the Primitive constant buffer is not used.
+	// @param Name e.g. TEXT("ObjectWorldPositionAndRadius.w")
+	int32 GetPrimitiveProperty(EMaterialValueType Type, const TCHAR* ExpressionName, const TCHAR* HLSLName)
+	{
+		const EMaterialDomain Domain = (const EMaterialDomain)Material->GetMaterialDomain();
+
+		if(Domain != MD_Surface)
+		{
+			Errorf(TEXT("The material expression '%s' is only supported in the 'Surface' material domain."), ExpressionName);
+			return INDEX_NONE;
+		}
+
+		return AddInlinedCodeChunk(Type, TEXT("Primitive.%s"), HLSLName);
+	}
 };
 
 #endif
