Merge: #REPLAYS: Ability to play replay from cmdline option

diff --git a/Engine/Source/Runtime/Engine/Private/GameInstance.cpp b/Engine/Source/Runtime/Engine/Private/GameInstance.cpp
--- a/Engine/Source/Runtime/Engine/Private/GameInstance.cpp
+++ b/Engine/Source/Runtime/Engine/Private/GameInstance.cpp
@@ -304,6 +304,14 @@ void UGameInstance::StartGameInstance()
 	Tmp = TEXT("");
 #endif // UE_BUILD_SHIPPING
 
+	// Parse replay name if specified on cmdline
+	FString ReplayCommand;
+	if ( FParse::Value( Tmp, TEXT( "-REPLAY=" ), ReplayCommand ) )
+	{
+		PlayReplay( ReplayCommand );
+		return;
+	}
+
 	const UGameMapsSettings* GameMapsSettings = GetDefault<UGameMapsSettings>();
 	const FString& DefaultMap = GameMapsSettings->GetGameDefaultMap();
 	if (!FParse::Token(Tmp, Parm, ARRAY_COUNT(Parm), 0) || Parm[0] == '-')
