Fix typo in comments

diff --git a/Engine/Source/Runtime/Engine/Classes/PhysicsEngine/ConstraintInstance.h b/Engine/Source/Runtime/Engine/Classes/PhysicsEngine/ConstraintInstance.h
--- a/Engine/Source/Runtime/Engine/Classes/PhysicsEngine/ConstraintInstance.h
+++ b/Engine/Source/Runtime/Engine/Classes/PhysicsEngine/ConstraintInstance.h
@@ -156,7 +156,7 @@ struct ENGINE_API FConstraintInstance
 	UPROPERTY(EditAnywhere, Category=Linear)
 	TEnumAsByte<enum ELinearConstraintMotion> LinearZMotion;
 
-	/** The limiting extent in world untis of the linear motion for limitied motion axes. */
+	/** The limiting extent in world units of the linear motion for limited motion axes. */
 	UPROPERTY(EditAnywhere, Category=Linear, meta=(ClampMin = "0.0"))
 	float LinearLimitSize;
 
