#ue4 - Projectile uses faster FQuat code path when passing rotation to MoveComponent.

diff --git a/Engine/Source/Runtime/Engine/Private/Components/ProjectileMovementComponent.cpp b/Engine/Source/Runtime/Engine/Private/Components/ProjectileMovementComponent.cpp
--- a/Engine/Source/Runtime/Engine/Private/Components/ProjectileMovementComponent.cpp
+++ b/Engine/Source/Runtime/Engine/Private/Components/ProjectileMovementComponent.cpp
@@ -133,7 +133,7 @@ void UProjectileMovementComponent::TickComponent(float DeltaTime, enum ELevelTic
 		const FVector OldVelocity = Velocity;
 		const FVector MoveDelta = ComputeMoveDelta(OldVelocity, TimeTick);
 
-		const FRotator NewRotation = (bRotationFollowsVelocity && !OldVelocity.IsNearlyZero(0.01f)) ? OldVelocity.Rotation() : ActorOwner->GetActorRotation();
+		const FQuat NewRotation = (bRotationFollowsVelocity && !OldVelocity.IsNearlyZero(0.01f)) ? OldVelocity.ToOrientationQuat() : UpdatedComponent->GetComponentQuat();
 
 		// Move the component
 		if (bShouldBounce)
