fixed typo

diff --git a/Engine/Source/Runtime/Engine/Classes/Components/DecalComponent.h b/Engine/Source/Runtime/Engine/Classes/Components/DecalComponent.h
--- a/Engine/Source/Runtime/Engine/Classes/Components/DecalComponent.h
+++ b/Engine/Source/Runtime/Engine/Classes/Components/DecalComponent.h
@@ -32,7 +32,7 @@ class ENGINE_API UDecalComponent : public USceneComponent
 	UPROPERTY(EditAnywhere, BlueprintReadOnly, Category = Decal)
 	float FadeScreenSize;
 
-	/** Size if the scale if 1,1,1, technically redundant but there for convenience */
+	/** Decal size in local space (does not include the component scale), technically redundant but there for convenience */
 	UPROPERTY(EditAnywhere, BlueprintReadOnly, Category = Decal, meta=(AllowPreserveRatio = "true"))
 	FVector DecalSize;
 
