Framework: Hid various internal-use-only tick groups from being displayed as options in Blueprints for SetTickGroup or the default tick group
Added friendly display names to drop the TG_ prefix on the rest of the tick groups

diff --git a/Engine/Source/Runtime/Engine/Classes/Engine/EngineBaseTypes.h b/Engine/Source/Runtime/Engine/Classes/Engine/EngineBaseTypes.h
--- a/Engine/Source/Runtime/Engine/Classes/Engine/EngineBaseTypes.h
+++ b/Engine/Source/Runtime/Engine/Classes/Engine/EngineBaseTypes.h
@@ -42,34 +42,34 @@ UENUM(BlueprintType)
 enum ETickingGroup
 {
 	/** Any item that needs to be executed before physics simulation starts. */
-	TG_PrePhysics,
+	TG_PrePhysics UMETA(DisplayName="Pre Physics"),
 
 	/** Special tick group that starts physics simulation. */							
-	TG_StartPhysics,
+	TG_StartPhysics UMETA(Hidden, DisplayName="Start Physics"),
 
 	/** Any item that can be run in parallel with our physics simulation work. */
-	TG_DuringPhysics,
+	TG_DuringPhysics UMETA(DisplayName="During Physics"),
 
 	/** Special tick group that ends physics simulation. */
-	TG_EndPhysics,
+	TG_EndPhysics UMETA(Hidden, DisplayName="End Physics"),
 
 	/** Any item that needs physics to be complete before being executed. */
-	TG_PreCloth,
+	TG_PreCloth UMETA(Hidden, DisplayName="Pre Cloth"),
 
 	/** Any item that needs to be updated after rigid body simulation is done, but before cloth is simulation is done. */
-	TG_StartCloth,
+	TG_StartCloth UMETA(Hidden, DisplayName = "Start Cloth"),
 
 	/** Any item that needs rigid body and cloth simulation to be complete before being executed. */
-	TG_PostPhysics,
+	TG_PostPhysics UMETA(DisplayName="Post Physics"),
 
 	/** Any item that needs the update work to be done before being ticked. */
-	TG_PostUpdateWork,
+	TG_PostUpdateWork UMETA(DisplayName="Post Update Work"),
 
 	/** Special tick group that ends cloth simulation. */
-	TG_EndCloth,
+	TG_EndCloth UMETA(Hidden, DisplayName="End Cloth"),
 
 	/** Special tick group that is not actually a tick group. After every tick group this is repeatedly re-run until there are no more newly spawned items to run. */
-	TG_NewlySpawned,
+	TG_NewlySpawned UMETA(Hidden, DisplayName="Newly Spawned"),
 
 	TG_MAX,
 };
