fixed decals scaling of newly created decals (missing code of former code change)

diff --git a/Engine/Source/Runtime/Engine/Private/DecalActor.cpp b/Engine/Source/Runtime/Engine/Private/DecalActor.cpp
--- a/Engine/Source/Runtime/Engine/Private/DecalActor.cpp
+++ b/Engine/Source/Runtime/Engine/Private/DecalActor.cpp
@@ -110,7 +110,7 @@ void ADecalActor::EditorApplyScale(const FVector& DeltaScale, const FVector* Piv
 		Mul = FVector(SaveInv(Decal->DecalSize.X), SaveInv(Decal->DecalSize.Y), SaveInv(Decal->DecalSize.Z));
 	}
 
-	const FVector ModifiedScale = DeltaScale * (AActor::bUsePercentageBasedScaling ? DecalEditorConstants::PercentageScalingMultiplier : DecalEditorConstants::AdditiveScalingMultiplier);
+	const FVector ModifiedScale = Mul * DeltaScale * (AActor::bUsePercentageBasedScaling ? DecalEditorConstants::PercentageScalingMultiplier : DecalEditorConstants::AdditiveScalingMultiplier);
 
 	Super::EditorApplyScale(ModifiedScale, PivotLocation, bAltDown, bShiftDown, bCtrlDown);
 }
