polish UI for Texture Sample material expression

diff --git a/Engine/Source/Runtime/Engine/Private/Materials/MaterialExpressions.cpp b/Engine/Source/Runtime/Engine/Private/Materials/MaterialExpressions.cpp
--- a/Engine/Source/Runtime/Engine/Private/Materials/MaterialExpressions.cpp
+++ b/Engine/Source/Runtime/Engine/Private/Materials/MaterialExpressions.cpp
@@ -1148,7 +1148,17 @@ bool UMaterialExpressionTextureSample::CanEditChange(const UProperty* InProperty
 	bool bIsEditable = Super::CanEditChange(InProperty);
 	if (bIsEditable && InProperty != NULL)
 	{
-		if (InProperty->GetFName() == TEXT("Texture"))
+		FName PropertyFName = InProperty->GetFName();
+
+		if (PropertyFName == GET_MEMBER_NAME_CHECKED(UMaterialExpressionTextureSample, ConstMipValue))
+		{
+			bIsEditable = MipValueMode == TMVM_MipLevel || MipValueMode == TMVM_MipBias;
+		}
+		else if (PropertyFName == GET_MEMBER_NAME_CHECKED(UMaterialExpressionTextureSample, ConstCoordinate))
+		{
+			bIsEditable = !Coordinates.Expression;
+		}
+		else if (PropertyFName == GET_MEMBER_NAME_CHECKED(UMaterialExpressionTextureSample, Texture))
 		{
 			// The Texture property is overridden by a connection to TextureObject
 			bIsEditable = TextureObject.Expression == NULL;
