#ue4 - Perf: Convert most uses of FQuat::FindBetween() to FQuat::FindBetweenNormals() where the inputs are known to be normalized.

diff --git a/Engine/Source/Runtime/AnimGraphRuntime/Private/BoneControllers/AnimNodeRotationMultiplier.cpp b/Engine/Source/Runtime/AnimGraphRuntime/Private/BoneControllers/AnimNodeRotationMultiplier.cpp
--- a/Engine/Source/Runtime/AnimGraphRuntime/Private/BoneControllers/AnimNodeRotationMultiplier.cpp
+++ b/Engine/Source/Runtime/AnimGraphRuntime/Private/BoneControllers/AnimNodeRotationMultiplier.cpp
@@ -48,7 +48,7 @@ FQuat FAnimNode_RotationMultiplier::ExtractAngle(const FTransform& RefPoseTransf
 	const FVector LocalRotationVector = LocalBoneTransform.GetRotation().RotateVector(RotationAxis);
 	const FVector ReferenceRotationVector = ReferenceBoneTransform.GetRotation().RotateVector(RotationAxis);
 
-	const FQuat LocalToRefQuat = FQuat::FindBetween(LocalRotationVector, ReferenceRotationVector);
+	const FQuat LocalToRefQuat = FQuat::FindBetweenNormals(LocalRotationVector, ReferenceRotationVector);
 	checkSlow( LocalToRefQuat.IsNormalized() );
 
 	// Rotate parent bone atom from position in local space to reference skeleton
diff --git a/Engine/Source/Runtime/AnimGraphRuntime/Private/BoneControllers/AnimNode_LookAt.cpp b/Engine/Source/Runtime/AnimGraphRuntime/Private/BoneControllers/AnimNode_LookAt.cpp
--- a/Engine/Source/Runtime/AnimGraphRuntime/Private/BoneControllers/AnimNode_LookAt.cpp
+++ b/Engine/Source/Runtime/AnimGraphRuntime/Private/BoneControllers/AnimNode_LookAt.cpp
@@ -166,11 +166,11 @@ void FAnimNode_LookAt::EvaluateBoneTransforms(USkeletalMeshComponent* SkelComp,
 		// project target to the plane
 		FVector NewTarget = FVector::VectorPlaneProject(ToTarget, LookUpVector);
 		NewTarget.Normalize();
-		DeltaRot = FQuat::FindBetween(LookAtVector, NewTarget);
+		DeltaRot = FQuat::FindBetweenNormals(LookAtVector, NewTarget);
 	}
 	else
 	{
-		DeltaRot = FQuat::FindBetween(LookAtVector, ToTarget);
+		DeltaRot = FQuat::FindBetweenNormals(LookAtVector, ToTarget);
 	}
 
 	// transform current rotation to delta rotation
diff --git a/Engine/Source/Runtime/AnimGraphRuntime/Private/BoneControllers/AnimNode_SpringBone.cpp b/Engine/Source/Runtime/AnimGraphRuntime/Private/BoneControllers/AnimNode_SpringBone.cpp
--- a/Engine/Source/Runtime/AnimGraphRuntime/Private/BoneControllers/AnimNode_SpringBone.cpp
+++ b/Engine/Source/Runtime/AnimGraphRuntime/Private/BoneControllers/AnimNode_SpringBone.cpp
@@ -195,7 +195,7 @@ void FAnimNode_SpringBone::EvaluateBoneTransforms(USkeletalMeshComponent* SkelCo
 		FVector ParentToTarget = (TargetPos - ParentSpaceBase.GetLocation()).GetSafeNormal();
 		FVector ParentToCurrent = (BoneLocation - ParentSpaceBase.GetLocation()).GetSafeNormal();
 
-		FQuat AdditionalRotation = FQuat::FindBetween(ParentToTarget, ParentToCurrent);
+		FQuat AdditionalRotation = FQuat::FindBetweenNormals(ParentToTarget, ParentToCurrent);
 
 		// Filter rotation based on our filter properties
 		FVector EularRot = AdditionalRotation.Euler();
diff --git a/Engine/Source/Runtime/AnimGraphRuntime/Private/BoneControllers/AnimNode_Trail.cpp b/Engine/Source/Runtime/AnimGraphRuntime/Private/BoneControllers/AnimNode_Trail.cpp
--- a/Engine/Source/Runtime/AnimGraphRuntime/Private/BoneControllers/AnimNode_Trail.cpp
+++ b/Engine/Source/Runtime/AnimGraphRuntime/Private/BoneControllers/AnimNode_Trail.cpp
@@ -180,7 +180,7 @@ void FAnimNode_Trail::EvaluateBoneTransforms(USkeletalMeshComponent* SkelComp, F
 		FVector NewBoneDir = FVector(OutBoneTransforms[i].Transform.GetTranslation() - OutBoneTransforms[i - 1].Transform.GetTranslation()).GetSafeNormal(SMALL_NUMBER);
 
 		// Calculate a quaternion that gets us from our current rotation to the desired one.
-		FQuat DeltaLookQuat = FQuat::FindBetween(CurrentBoneDir, NewBoneDir);
+		FQuat DeltaLookQuat = FQuat::FindBetweenNormals(CurrentBoneDir, NewBoneDir);
 		FTransform DeltaTM( DeltaLookQuat, FVector(0.f) );
 
 		// Apply to the current parent bone transform.
diff --git a/Engine/Source/Runtime/AnimGraphRuntime/Private/BoneControllers/AnimNode_TwoBoneIK.cpp b/Engine/Source/Runtime/AnimGraphRuntime/Private/BoneControllers/AnimNode_TwoBoneIK.cpp
--- a/Engine/Source/Runtime/AnimGraphRuntime/Private/BoneControllers/AnimNode_TwoBoneIK.cpp
+++ b/Engine/Source/Runtime/AnimGraphRuntime/Private/BoneControllers/AnimNode_TwoBoneIK.cpp
@@ -239,7 +239,7 @@ void FAnimNode_TwoBoneIK::EvaluateBoneTransforms(USkeletalMeshComponent* SkelCom
 		FVector const OldDir = (InitialJointPos - RootPos).GetSafeNormal();
 		FVector const NewDir = (OutJointPos - RootPos).GetSafeNormal();
 		// Find Delta Rotation take takes us from Old to New dir
-		FQuat const DeltaRotation = FQuat::FindBetween(OldDir, NewDir);
+		FQuat const DeltaRotation = FQuat::FindBetweenNormals(OldDir, NewDir);
 		// Rotate our Joint quaternion by this delta rotation
 		UpperLimbCSTransform.SetRotation( DeltaRotation * UpperLimbCSTransform.GetRotation() );
 		// And put joint where it should be.
@@ -256,7 +256,7 @@ void FAnimNode_TwoBoneIK::EvaluateBoneTransforms(USkeletalMeshComponent* SkelCom
 		FVector const NewDir = (OutEndPos - OutJointPos).GetSafeNormal();
 
 		// Find Delta Rotation take takes us from Old to New dir
-		FQuat const DeltaRotation = FQuat::FindBetween(OldDir, NewDir);
+		FQuat const DeltaRotation = FQuat::FindBetweenNormals(OldDir, NewDir);
 		// Rotate our Joint quaternion by this delta rotation
 		LowerLimbCSTransform.SetRotation( DeltaRotation * LowerLimbCSTransform.GetRotation() );
 		// And put joint where it should be.
diff --git a/Engine/Source/Runtime/Engine/Private/Animation/AnimCompress_RemoveLinearKeys.cpp b/Engine/Source/Runtime/Engine/Private/Animation/AnimCompress_RemoveLinearKeys.cpp
--- a/Engine/Source/Runtime/Engine/Private/Animation/AnimCompress_RemoveLinearKeys.cpp
+++ b/Engine/Source/Runtime/Engine/Private/Animation/AnimCompress_RemoveLinearKeys.cpp
@@ -745,8 +745,7 @@ void UAnimCompress_RemoveLinearKeys::ProcessAnimationTracks(
 								// limit the range we will retarget to something reasonable (~60 degrees)
 								if (DotResult < 1.0f && DotResult > 0.5f)
 								{
-									FQuat Adjustment= FQuat::FindBetween(CurrentHeading, DesiredHeading); 
-									Adjustment.Normalize();
+									FQuat Adjustment= FQuat::FindBetweenVectors(CurrentHeading, DesiredHeading);
 									Adjustment= EnforceShortestArc(FQuat::Identity, Adjustment);
 
 									const FVector Test = Adjustment.RotateVector(CurrentHeading);
diff --git a/Engine/Source/Runtime/Engine/Private/Particles/ParticleComponents.cpp b/Engine/Source/Runtime/Engine/Private/Particles/ParticleComponents.cpp
--- a/Engine/Source/Runtime/Engine/Private/Particles/ParticleComponents.cpp
+++ b/Engine/Source/Runtime/Engine/Private/Particles/ParticleComponents.cpp
@@ -3399,7 +3399,7 @@ void UParticleSystemComponent::OrientZAxisTowardCamera()
 		const FVector LocalZAxis = FVector(0,0,1);
 
 		// Find angle between z-axis and the camera direction
-		const FQuat PointTo = FQuat::FindBetween(LocalZAxis, DirToCamera);
+		const FQuat PointTo = FQuat::FindBetweenNormals(LocalZAxis, DirToCamera);
 		
 		// Adjust our rotation
 		const FRotator AdjustmentAngle(PointTo);
diff --git a/Engine/Source/Runtime/Engine/Private/Particles/ParticleEmitterInstances.cpp b/Engine/Source/Runtime/Engine/Private/Particles/ParticleEmitterInstances.cpp
--- a/Engine/Source/Runtime/Engine/Private/Particles/ParticleEmitterInstances.cpp
+++ b/Engine/Source/Runtime/Engine/Private/Particles/ParticleEmitterInstances.cpp
@@ -2934,7 +2934,7 @@ void FParticleMeshEmitterInstance::Tick(float DeltaTime, bool bSuppressSpawning)
 				NewDirection.Normalize();
 				FVector	OldDirection(1.0f, 0.0f, 0.0f);
 
-				FQuat Rotation	= FQuat::FindBetween(OldDirection, NewDirection);
+				FQuat Rotation	= FQuat::FindBetweenNormals(OldDirection, NewDirection);
 				FVector Euler	= Rotation.Euler();
 				PayloadData->Rotation = PayloadData->InitRotation + Euler;
 				PayloadData->Rotation += PayloadData->CurContinuousRotation;
@@ -3177,7 +3177,7 @@ void FParticleMeshEmitterInstance::PostSpawn(FBaseParticle* Particle, float Inte
 		NewDirection.Normalize();
 		FVector	OldDirection(1.0f, 0.0f, 0.0f);
 
-		FQuat Rotation	= FQuat::FindBetween(OldDirection, NewDirection);
+		FQuat Rotation	= FQuat::FindBetweenNormals(OldDirection, NewDirection);
 		FVector Euler	= Rotation.Euler();
 
 		PayloadData->Rotation.X	+= Euler.X;
diff --git a/Engine/Source/Runtime/Engine/Private/Particles/ParticleSystemRender.cpp b/Engine/Source/Runtime/Engine/Private/Particles/ParticleSystemRender.cpp
--- a/Engine/Source/Runtime/Engine/Private/Particles/ParticleSystemRender.cpp
+++ b/Engine/Source/Runtime/Engine/Private/Particles/ParticleSystemRender.cpp
@@ -1679,14 +1679,14 @@ void FDynamicMeshEmitterData::GetParticleTransform(const FBaseParticle& InPartic
 					bFacingDirectionIsValid = false;
 				}
 				// Velocity align the X-axis, and camera face the selected axis
-				PointTo = FQuat::FindBetween(FVector(1.0f, 0.0f, 0.0f), VelocityDirection);
+				PointTo = FQuat::FindBetweenNormals(FVector(1.0f, 0.0f, 0.0f), VelocityDirection);
 				FacingDir = VelocityDirection;
 				DesiredDir = DirToCamera;
 			}
 			else if (CameraFacingOption <= XAxisFacing_NegativeYUp)
 			{
 				// Camera face the X-axis, and point the selected axis towards the world up
-				PointTo = FQuat::FindBetween(FVector(1,0,0), DirToCamera);
+				PointTo = FQuat::FindBetweenNormals(FVector(1,0,0), DirToCamera);
 				FacingDir = DirToCamera;
 				DesiredDir = FVector(0,0,1);
 			}
@@ -1709,7 +1709,7 @@ void FDynamicMeshEmitterData::GetParticleTransform(const FBaseParticle& InPartic
 
 			FVector	DirToDesiredInRotationPlane = DesiredDir - ((DesiredDir | FacingDir) * FacingDir);
 			DirToDesiredInRotationPlane.Normalize();
-			FQuat FacingRotation = FQuat::FindBetween(PointTo.RotateVector(CameraFacingOpVector), DirToDesiredInRotationPlane);
+			FQuat FacingRotation = FQuat::FindBetweenNormals(PointTo.RotateVector(CameraFacingOpVector), DirToDesiredInRotationPlane);
 			PointTo = FacingRotation * PointTo;
 
 			// Add in additional rotation about either the directional or camera facing axis
@@ -1729,7 +1729,7 @@ void FDynamicMeshEmitterData::GetParticleTransform(const FBaseParticle& InPartic
 		}
 		else
 		{
-			PointTo = FQuat::FindBetween(FVector(1,0,0), DirToCamera);
+			PointTo = FQuat::FindBetweenNormals(FVector(1,0,0), DirToCamera);
 			// Add in additional rotation about facing axis
 			FQuat AddedRotation = FQuat(DirToCamera, InParticle.Rotation);
 			kLockedAxisRotator = FRotator(AddedRotation * PointTo);
@@ -1769,12 +1769,12 @@ void FDynamicMeshEmitterData::GetParticleTransform(const FBaseParticle& InPartic
 			// For the locked axis behavior, only rotate to	face the camera	about the
 			// locked direction, and maintain the up vector	pointing towards the locked	direction
 			// Find	the	rotation that points the localupaxis towards the targetupaxis
-			FQuat PointToUp	= FQuat::FindBetween(LocalSpaceUpAxis, Source.LockedAxis);
+			FQuat PointToUp	= FQuat::FindBetweenNormals(LocalSpaceUpAxis, Source.LockedAxis);
 
 			// Add in rotation about the TargetUpAxis to point the facing vector towards the camera
 			FVector	DirToCameraInRotationPlane = DirToCamera - ((DirToCamera | Source.LockedAxis)*Source.LockedAxis);
 			DirToCameraInRotationPlane.Normalize();
-			FQuat PointToCamera	= FQuat::FindBetween(PointToUp.RotateVector(LocalSpaceFacingAxis), DirToCameraInRotationPlane);
+			FQuat PointToCamera	= FQuat::FindBetweenNormals(PointToUp.RotateVector(LocalSpaceFacingAxis), DirToCameraInRotationPlane);
 
 			// Set kRotator	to the composed	rotation
 			FQuat MeshRotation = PointToCamera*PointToUp;
@@ -1787,7 +1787,7 @@ void FDynamicMeshEmitterData::GetParticleTransform(const FBaseParticle& InPartic
 				// (c.f. the roll rotation is about	the	radial axis)
 
 				// Find	the	rotation that points the facing	axis towards the camera
-				FRotator PointToRotation = FRotator(FQuat::FindBetween(LocalSpaceFacingAxis, DirToCamera));
+				FRotator PointToRotation = FRotator(FQuat::FindBetweenNormals(LocalSpaceFacingAxis, DirToCamera));
 
 				// When	constructing the rotation, we need to eliminate	roll around	the	dirtocamera	axis,
 				// otherwise the particle appears to rotate	around the dircamera axis when it or the camera	moves
@@ -1814,7 +1814,7 @@ void FDynamicMeshEmitterData::GetParticleTransform(const FBaseParticle& InPartic
 				// with	no roll, and then rotates about	the	direction_to_camera	by the spriteparticle rotation.
 
 				// Find	the	rotation that points the facing	axis towards the camera
-				FRotator PointToRotation = FRotator(FQuat::FindBetween(LocalSpaceFacingAxis, DirToCamera));
+				FRotator PointToRotation = FRotator(FQuat::FindBetweenNormals(LocalSpaceFacingAxis, DirToCamera));
 
 				// When	constructing the rotation, we need to eliminate	roll around	the	dirtocamera	axis,
 				// otherwise the particle appears to rotate	around the dircamera axis when it or the camera	moves
@@ -1893,7 +1893,7 @@ void FDynamicMeshEmitterData::GetInstanceData(void* InstanceData, void* DynamicP
 	if (bUseMeshLockedAxis == true)
 	{
 		// facing axis is taken to be the local x axis.	
-		PointToLockedAxis = FQuat::FindBetween(FVector(1,0,0), Source.LockedAxis);
+		PointToLockedAxis = FQuat::FindBetweenNormals(FVector(1,0,0), Source.LockedAxis);
 	}
 
 	FVector CameraFacingOpVector = FVector::ZeroVector;
diff --git a/Engine/Source/Runtime/Engine/Public/ParticleHelper.h b/Engine/Source/Runtime/Engine/Public/ParticleHelper.h
--- a/Engine/Source/Runtime/Engine/Public/ParticleHelper.h
+++ b/Engine/Source/Runtime/Engine/Public/ParticleHelper.h
@@ -638,15 +638,16 @@ struct FBeamParticleModifierPayloadData
 	{
 		if (bModifyTangent == true)
 		{
-			FVector ModTangent = Tangent;
-
+			FVector ModTangent;
 			if (bAbsolute == false)
 			{
 				// Transform the modified tangent so it is relative to the real tangent
-				FQuat RotQuat = FQuat::FindBetween(FVector(1.0f, 0.0f, 0.0f), Value);
-				FMatrix RotMat = FQuatRotationTranslationMatrix(RotQuat, FVector::ZeroVector);
-
-				ModTangent = RotMat.TransformVector(Tangent);
+				const FQuat RotQuat = FQuat::FindBetweenNormals(FVector(1.0f, 0.0f, 0.0f), Value);
+				ModTangent = RotQuat.RotateVector(Tangent);
+			}
+			else
+			{
+				ModTangent = Tangent;
 			}
 
 			if (bScaleTangent == false)
