#ue4 - Perf: Particles: optimize FDynamicMeshEmitterData::GetParticleTransform() to avoid needless Quat->Rotator and Rotator->Matrix conversions.

diff --git a/Engine/Source/Runtime/Engine/Private/Particles/ParticleSystemRender.cpp b/Engine/Source/Runtime/Engine/Private/Particles/ParticleSystemRender.cpp
--- a/Engine/Source/Runtime/Engine/Private/Particles/ParticleSystemRender.cpp
+++ b/Engine/Source/Runtime/Engine/Private/Particles/ParticleSystemRender.cpp
@@ -1593,7 +1593,7 @@ void FDynamicMeshEmitterData::GetParticleTransform(const FBaseParticle& InPartic
 	FMatrix LocalToWorld;
 	FTranslationMatrix kTransMat(FVector::ZeroVector);
 	FScaleMatrix kScaleMat(FVector(1.0f));
-	FRotator kLockedAxisRotator = FRotator::ZeroRotator;
+	FQuat kLockedAxisQuat = FQuat::Identity;
 	FVector Location;
 	FVector ScaledSize;
 	FVector	DirToCamera;
@@ -1617,7 +1617,7 @@ void FDynamicMeshEmitterData::GetParticleTransform(const FBaseParticle& InPartic
 	kScaleMat.M[1][1] = ScaledSize.Y;
 	kScaleMat.M[2][2] = ScaledSize.Z;
 
-	FRotator kRotator(0,0,0);
+	FMatrix kRotMat(FMatrix::Identity);
 	LocalToWorld = Proxy->GetLocalToWorld();
 	if (bUseCameraFacing == true)
 	{
@@ -1718,13 +1718,13 @@ void FDynamicMeshEmitterData::GetParticleTransform(const FBaseParticle& InPartic
 				if (bFacingDirectionIsValid)
 				{
 					FQuat AddedRotation = FQuat(FacingDir, InParticle.Rotation);
-					kLockedAxisRotator = FRotator(AddedRotation * PointTo);
+					kLockedAxisQuat = (AddedRotation * PointTo);
 				}
 			}
 			else
 			{
 				FQuat AddedRotation = FQuat(DirToCamera, InParticle.Rotation);
-				kLockedAxisRotator = FRotator(AddedRotation * PointTo);
+				kLockedAxisQuat = (AddedRotation * PointTo);
 			}
 		}
 		else
@@ -1732,14 +1732,14 @@ void FDynamicMeshEmitterData::GetParticleTransform(const FBaseParticle& InPartic
 			PointTo = FQuat::FindBetweenNormals(FVector(1,0,0), DirToCamera);
 			// Add in additional rotation about facing axis
 			FQuat AddedRotation = FQuat(DirToCamera, InParticle.Rotation);
-			kLockedAxisRotator = FRotator(AddedRotation * PointTo);
+			kLockedAxisQuat = (AddedRotation * PointTo);
 		}
 	}
 	else if (bUseMeshLockedAxis == true)
 	{
 		// Add any 'sprite rotation' about the locked axis
 		FQuat AddedRotation = FQuat(Source.LockedAxis, InParticle.Rotation);
-		kLockedAxisRotator = FRotator(AddedRotation * PointTo);
+		kLockedAxisQuat = (AddedRotation * PointTo);
 	}
 	else if (Source.ScreenAlignment == PSA_TypeSpecific)
 	{
@@ -1776,9 +1776,9 @@ void FDynamicMeshEmitterData::GetParticleTransform(const FBaseParticle& InPartic
 			DirToCameraInRotationPlane.Normalize();
 			FQuat PointToCamera	= FQuat::FindBetweenNormals(PointToUp.RotateVector(LocalSpaceFacingAxis), DirToCameraInRotationPlane);
 
-			// Set kRotator	to the composed	rotation
+			// Set kRotMat to the composed rotation
 			FQuat MeshRotation = PointToCamera*PointToUp;
-			kRotator = FRotator(MeshRotation);
+			kRotMat = FQuatRotationMatrix(MeshRotation);
 		}
 		else if (Source.MeshAlignment == PSMA_MeshFaceCameraWithSpin)
 		{
@@ -1804,9 +1804,9 @@ void FDynamicMeshEmitterData::GetParticleTransform(const FBaseParticle& InPartic
 
 				FQuat AddedTangentialRotation =	FQuat(vTangentAxis,	InParticle.Rotation);
 
-				// Set kRotator	to the composed	rotation
+				// Set kRotMat to the composed rotation
 				FQuat MeshRotation = AddedTangentialRotation*PointToRotation.Quaternion();
-				kRotator = FRotator(MeshRotation);
+				kRotMat = FQuatRotationMatrix(MeshRotation);
 		}
 		else 
 		{
@@ -1823,22 +1823,23 @@ void FDynamicMeshEmitterData::GetParticleTransform(const FBaseParticle& InPartic
 				// Add in the roll we do want.
 				FQuat AddedRollRotation	= FQuat(DirToCamera, InParticle.Rotation);
 
-				// Set kRotator	to the composed	rotation
+				// Set kRotMat to the composed	rotation
 				FQuat MeshRotation = AddedRollRotation*PointToRotation.Quaternion();
-				kRotator = FRotator(MeshRotation);
+				kRotMat = FQuatRotationMatrix(MeshRotation);
 		}
 	}
 	else
 	{
 		float fRot = InParticle.Rotation * 180.0f / PI;
 		FVector kRotVec = FVector(fRot, fRot, fRot);
-		kRotator = FRotator::MakeFromEuler(kRotVec);
+		FRotator kRotator = FRotator::MakeFromEuler(kRotVec);
 
 		const FMeshRotationPayloadData* PayloadData = (const FMeshRotationPayloadData*)((const uint8*)&InParticle + Source.MeshRotationOffset);
 		kRotator += FRotator::MakeFromEuler(PayloadData->Rotation);
+
+		kRotMat = FRotationMatrix(kRotator);
 	}
 
-	FRotationMatrix kRotMat(kRotator);
 	if (bApplyPreRotation == true)
 	{
 		const FMeshRotationPayloadData* PayloadData = (const FMeshRotationPayloadData*)((const uint8*)&InParticle + Source.MeshRotationOffset);
@@ -1847,7 +1848,7 @@ void FDynamicMeshEmitterData::GetParticleTransform(const FBaseParticle& InPartic
 
 		if ((bUseCameraFacing == true) || (bUseMeshLockedAxis == true))
 		{
-			OutTransformMat = (OrientMat * kScaleMat) * FRotationMatrix(kLockedAxisRotator) * kRotMat * kTransMat;
+			OutTransformMat = (OrientMat * kScaleMat) * FQuatRotationMatrix(kLockedAxisQuat) * kRotMat * kTransMat;
 		}
 		else
 		{
@@ -1856,7 +1857,7 @@ void FDynamicMeshEmitterData::GetParticleTransform(const FBaseParticle& InPartic
 	}
 	else if ((bUseCameraFacing == true) || (bUseMeshLockedAxis == true))
 	{
-		OutTransformMat = kScaleMat * FRotationMatrix(kLockedAxisRotator) * kRotMat * kTransMat;
+		OutTransformMat = kScaleMat * FQuatRotationMatrix(kLockedAxisQuat) * kRotMat * kTransMat;
 	}
 	else
 	{
