Stats - Restored STAT_CharacterMovement

diff --git a/Engine/Source/Runtime/Engine/Private/Components/CharacterMovementComponent.cpp b/Engine/Source/Runtime/Engine/Private/Components/CharacterMovementComponent.cpp
--- a/Engine/Source/Runtime/Engine/Private/Components/CharacterMovementComponent.cpp
+++ b/Engine/Source/Runtime/Engine/Private/Components/CharacterMovementComponent.cpp
@@ -875,6 +875,7 @@ void UCharacterMovementComponent::PerformAirControlForPathFollowing(FVector Dire
 
 void UCharacterMovementComponent::TickComponent(float DeltaTime, enum ELevelTick TickType, FActorComponentTickFunction *ThisTickFunction)
 {
+	SCOPE_CYCLE_COUNTER(STAT_CharacterMovement);
 	SCOPE_CYCLE_COUNTER(STAT_CharacterMovementTick);
 
 	const FVector InputVector = ConsumeInputVector();
diff --git a/Engine/Source/Runtime/Engine/Private/LevelTick.cpp b/Engine/Source/Runtime/Engine/Private/LevelTick.cpp
--- a/Engine/Source/Runtime/Engine/Private/LevelTick.cpp
+++ b/Engine/Source/Runtime/Engine/Private/LevelTick.cpp
@@ -63,6 +63,7 @@ DEFINE_STAT(STAT_PostTickComponentUpdateWait);
 DEFINE_STAT(STAT_TickTime);
 DEFINE_STAT(STAT_WorldTickTime);
 DEFINE_STAT(STAT_UpdateCameraTime);
+DEFINE_STAT(STAT_CharacterMovement);
 DEFINE_STAT(STAT_PlayerControllerTick);
 
 DEFINE_STAT(STAT_VolumeStreamingTickTime);
diff --git a/Engine/Source/Runtime/Engine/Public/EngineStats.h b/Engine/Source/Runtime/Engine/Public/EngineStats.h
--- a/Engine/Source/Runtime/Engine/Public/EngineStats.h
+++ b/Engine/Source/Runtime/Engine/Public/EngineStats.h
@@ -53,6 +53,7 @@ DECLARE_CYCLE_STAT_EXTERN(TEXT("GC Sweep Time"),STAT_GCSweepTime,STATGROUP_Game,
 DECLARE_CYCLE_STAT_EXTERN(TEXT("Tick Time"),STAT_TickTime,STATGROUP_Game, );
 DECLARE_CYCLE_STAT_EXTERN(TEXT("World Tick Time"),STAT_WorldTickTime,STATGROUP_Game, );
 DECLARE_CYCLE_STAT_EXTERN(TEXT("Update Camera Time"),STAT_UpdateCameraTime,STATGROUP_Game, );
+DECLARE_CYCLE_STAT_EXTERN(TEXT("Char Movement Total"), STAT_CharacterMovement, STATGROUP_Game, );
 DECLARE_CYCLE_STAT_EXTERN(TEXT("PlayerController Tick"), STAT_PlayerControllerTick, STATGROUP_Game, );
 DECLARE_CYCLE_STAT_EXTERN(TEXT("Post Tick Component Update"),STAT_PostTickComponentUpdate,STATGROUP_Game, );
 DECLARE_CYCLE_STAT_EXTERN(TEXT("Post Tick Component Wait"),STAT_PostTickComponentUpdateWait,STATGROUP_Game, );
