BytecodeOnly compilation recreates UberGraphPersistentFrame.

diff --git a/Engine/Source/Editor/KismetCompiler/Private/KismetCompilerModule.cpp b/Engine/Source/Editor/KismetCompiler/Private/KismetCompilerModule.cpp
--- a/Engine/Source/Editor/KismetCompiler/Private/KismetCompilerModule.cpp
+++ b/Engine/Source/Editor/KismetCompiler/Private/KismetCompilerModule.cpp
@@ -225,8 +225,10 @@ void FKismet2CompilerModule::CompileBlueprint(class UBlueprint* Blueprint, const
 			StubResults.bSilentMode = true;
 			FKismetCompilerOptions StubCompileOptions(CompileOptions);
 			StubCompileOptions.CompileType = EKismetCompileType::StubAfterFailure;
-
-			CompileBlueprintInner(Blueprint, StubCompileOptions, StubResults, StubReinstancer, ObjLoaded);
+			{
+				FRecreateUberGraphFrameScope RecreateUberGraphFrameScope(Blueprint->GeneratedClass, bBytecodeOnly);
+				CompileBlueprintInner(Blueprint, StubCompileOptions, StubResults, StubReinstancer, ObjLoaded);
+			}
 
 			StubReinstancer->UpdateBytecodeReferences();
 			if( !Blueprint->bIsRegeneratingOnLoad )
diff --git a/Engine/Source/Editor/UnrealEd/Private/Kismet2/Kismet2.cpp b/Engine/Source/Editor/UnrealEd/Private/Kismet2/Kismet2.cpp
--- a/Engine/Source/Editor/UnrealEd/Private/Kismet2/Kismet2.cpp
+++ b/Engine/Source/Editor/UnrealEd/Private/Kismet2/Kismet2.cpp
@@ -919,7 +919,10 @@ void FKismetEditorUtilities::RecompileBlueprintBytecode(UBlueprint* BlueprintObj
 
 	FKismetCompilerOptions CompileOptions;
 	CompileOptions.CompileType = EKismetCompileType::BytecodeOnly;
-	Compiler.CompileBlueprint(BlueprintObj, CompileOptions, Results, NULL, ObjLoaded);
+	{
+		FRecreateUberGraphFrameScope RecreateUberGraphFrameScope(BlueprintObj->GeneratedClass, true);
+		Compiler.CompileBlueprint(BlueprintObj, CompileOptions, Results, NULL, ObjLoaded);
+	}
 
 	ReinstanceHelper->UpdateBytecodeReferences();
 
diff --git a/Engine/Source/Editor/UnrealEd/Private/Kismet2/KismetReinstanceUtilities.cpp b/Engine/Source/Editor/UnrealEd/Private/Kismet2/KismetReinstanceUtilities.cpp
--- a/Engine/Source/Editor/UnrealEd/Private/Kismet2/KismetReinstanceUtilities.cpp
+++ b/Engine/Source/Editor/UnrealEd/Private/Kismet2/KismetReinstanceUtilities.cpp
@@ -21,6 +21,7 @@ DECLARE_CYCLE_STAT(TEXT("Recompile Child Classes"), EKismetReinstancerStats_Reco
 DECLARE_CYCLE_STAT(TEXT("Replace Classes Without Reinstancing"), EKismetReinstancerStats_ReplaceClassNoReinsancing, STATGROUP_KismetReinstancer );
 DECLARE_CYCLE_STAT(TEXT("Reinstance Objects"), EKismetCompilerStats_ReinstanceObjects, STATGROUP_KismetCompiler);
 DECLARE_CYCLE_STAT(TEXT("Refresh Dependent Blueprints In Reinstancer"), EKismetCompilerStats_RefreshDependentBlueprintsInReinstancer, STATGROUP_KismetCompiler);
+DECLARE_CYCLE_STAT(TEXT("Recreate UberGraphPersistentFrame"), EKismetCompilerStats_RecreateUberGraphPersistentFrame, STATGROUP_KismetCompiler);
 
 struct FReplaceReferenceHelper
 {
@@ -1324,3 +1325,35 @@ FBlueprintCompileReinstancer::FCDODuplicatesProvider& FBlueprintCompileReinstanc
 	static FCDODuplicatesProvider Delegate;
 	return Delegate;
 }
+
+TArray<UObject*> Objects;
+UClass* Class;
+
+FRecreateUberGraphFrameScope::FRecreateUberGraphFrameScope(UClass* InClass, bool bRecreate)
+	: Class(InClass)
+{
+	if (bRecreate && ensure(Class))
+	{
+		BP_SCOPED_COMPILER_EVENT_STAT(EKismetCompilerStats_RecreateUberGraphPersistentFrame);
+
+		const bool bIncludeDerivedClasses = true;
+		GetObjectsOfClass(Class, Objects, bIncludeDerivedClasses);
+
+		for (auto Obj : Objects)
+		{
+			Class->DestroyPersistentUberGraphFrame(Obj, true);
+		}
+	}
+}
+
+FRecreateUberGraphFrameScope::~FRecreateUberGraphFrameScope()
+{
+	BP_SCOPED_COMPILER_EVENT_STAT(EKismetCompilerStats_RecreateUberGraphPersistentFrame);
+	for (auto Obj : Objects)
+	{
+		if (IsValid(Obj))
+		{
+			Class->CreatePersistentUberGraphFrame(Obj, false, true);
+		}
+	}
+}
\ No newline at end of file
diff --git a/Engine/Source/Editor/UnrealEd/Public/Kismet2/KismetReinstanceUtilities.h b/Engine/Source/Editor/UnrealEd/Public/Kismet2/KismetReinstanceUtilities.h
--- a/Engine/Source/Editor/UnrealEd/Public/Kismet2/KismetReinstanceUtilities.h
+++ b/Engine/Source/Editor/UnrealEd/Public/Kismet2/KismetReinstanceUtilities.h
@@ -12,6 +12,16 @@ DECLARE_STATS_GROUP(TEXT("Kismet Reinstancer"), STATGROUP_KismetReinstancer, STA
 
 class FReinstanceFinalizer;
 
+struct UNREALED_API FRecreateUberGraphFrameScope
+{
+private:
+	TArray<UObject*> Objects;
+	UClass* Class;
+public:
+	FRecreateUberGraphFrameScope(UClass* InClass, bool bRecreate);
+	~FRecreateUberGraphFrameScope();
+};
+
 class UNREALED_API FBlueprintCompileReinstancer : public TSharedFromThis<FBlueprintCompileReinstancer>, public FGCObject
 {
 public:
diff --git a/Engine/Source/Runtime/CoreUObject/Public/UObject/Class.h b/Engine/Source/Runtime/CoreUObject/Public/UObject/Class.h
--- a/Engine/Source/Runtime/CoreUObject/Public/UObject/Class.h
+++ b/Engine/Source/Runtime/CoreUObject/Public/UObject/Class.h
@@ -2027,11 +2027,11 @@ public:
 		return NULL;
 	}
 
-	virtual void CreatePersistentUberGraphFrame(UObject* Obj, bool bCreateOnlyIfEmpty = false) const
+	virtual void CreatePersistentUberGraphFrame(UObject* Obj, bool bCreateOnlyIfEmpty = false, bool bSkipSuperClass = false) const
 	{
 	}
 
-	virtual void DestroyPersistentUberGraphFrame(UObject* Obj) const
+	virtual void DestroyPersistentUberGraphFrame(UObject* Obj, bool bSkipSuperClass = false) const
 	{
 	}
 
diff --git a/Engine/Source/Runtime/Engine/Classes/Engine/BlueprintGeneratedClass.h b/Engine/Source/Runtime/Engine/Classes/Engine/BlueprintGeneratedClass.h
--- a/Engine/Source/Runtime/Engine/Classes/Engine/BlueprintGeneratedClass.h
+++ b/Engine/Source/Runtime/Engine/Classes/Engine/BlueprintGeneratedClass.h
@@ -467,8 +467,8 @@ public:
 #endif //WITH_EDITOR
 	virtual bool IsFunctionImplementedInBlueprint(FName InFunctionName) const override;
 	virtual uint8* GetPersistentUberGraphFrame(UObject* Obj, UFunction* FuncToCheck) const override;
-	virtual void CreatePersistentUberGraphFrame(UObject* Obj, bool bCreateOnlyIfEmpty = false) const override;
-	virtual void DestroyPersistentUberGraphFrame(UObject* Obj) const override;
+	virtual void CreatePersistentUberGraphFrame(UObject* Obj, bool bCreateOnlyIfEmpty = false, bool bSkipSuperClass = false) const override;
+	virtual void DestroyPersistentUberGraphFrame(UObject* Obj, bool bSkipSuperClass = false) const override;
 	virtual void Link(FArchive& Ar, bool bRelinkExistingProperties) override;
 	virtual void PurgeClass(bool bRecompilingOnLoad) override;
 	virtual void Bind() override;
diff --git a/Engine/Source/Runtime/Engine/Private/BlueprintGeneratedClass.cpp b/Engine/Source/Runtime/Engine/Private/BlueprintGeneratedClass.cpp
--- a/Engine/Source/Runtime/Engine/Private/BlueprintGeneratedClass.cpp
+++ b/Engine/Source/Runtime/Engine/Private/BlueprintGeneratedClass.cpp
@@ -544,7 +544,7 @@ uint8* UBlueprintGeneratedClass::GetPersistentUberGraphFrame(UObject* Obj, UFunc
 	return ParentClass->GetPersistentUberGraphFrame(Obj, FuncToCheck);
 }
 
-void UBlueprintGeneratedClass::CreatePersistentUberGraphFrame(UObject* Obj, bool bCreateOnlyIfEmpty) const
+void UBlueprintGeneratedClass::CreatePersistentUberGraphFrame(UObject* Obj, bool bCreateOnlyIfEmpty, bool bSkipSuperClass) const
 {
 	checkSlow(!UberGraphFramePointerProperty == !UberGraphFunction);
 	if (Obj && UsePersistentUberGraphFrame() && UberGraphFramePointerProperty && UberGraphFunction)
@@ -577,12 +577,15 @@ void UBlueprintGeneratedClass::CreatePersistentUberGraphFrame(UObject* Obj, bool
 		}
 	}
 
-	auto ParentClass = GetSuperClass();
-	checkSlow(ParentClass);
-	return ParentClass->CreatePersistentUberGraphFrame(Obj, bCreateOnlyIfEmpty);
+	if (!bSkipSuperClass)
+	{
+		auto ParentClass = GetSuperClass();
+		checkSlow(ParentClass);
+		ParentClass->CreatePersistentUberGraphFrame(Obj, bCreateOnlyIfEmpty);
+	}
 }
 
-void UBlueprintGeneratedClass::DestroyPersistentUberGraphFrame(UObject* Obj) const
+void UBlueprintGeneratedClass::DestroyPersistentUberGraphFrame(UObject* Obj, bool bSkipSuperClass) const
 {
 	checkSlow(!UberGraphFramePointerProperty == !UberGraphFunction);
 	if (Obj && UsePersistentUberGraphFrame() && UberGraphFramePointerProperty && UberGraphFunction)
@@ -606,9 +609,12 @@ void UBlueprintGeneratedClass::DestroyPersistentUberGraphFrame(UObject* Obj) con
 		}
 	}
 
-	auto ParentClass = GetSuperClass();
-	checkSlow(ParentClass);
-	return ParentClass->DestroyPersistentUberGraphFrame(Obj);
+	if (!bSkipSuperClass)
+	{
+		auto ParentClass = GetSuperClass();
+		checkSlow(ParentClass);
+		ParentClass->DestroyPersistentUberGraphFrame(Obj);
+	}
 }
 
 void UBlueprintGeneratedClass::Link(FArchive& Ar, bool bRelinkExistingProperties)
