NavigationData's EndPlay and Destroy fixes #UE4

diff --git a/Engine/Source/Runtime/Engine/Classes/AI/Navigation/NavigationData.h b/Engine/Source/Runtime/Engine/Classes/AI/Navigation/NavigationData.h
--- a/Engine/Source/Runtime/Engine/Classes/AI/Navigation/NavigationData.h
+++ b/Engine/Source/Runtime/Engine/Classes/AI/Navigation/NavigationData.h
@@ -455,10 +455,16 @@ public:
 	virtual void PostEditUndo() override;
 #endif // WITH_EDITOR
 	virtual void Destroyed() override;
+	virtual void EndPlay(const EEndPlayReason::Type EndPlayReason) override;
 	// End UObject Interface
 		
 	virtual void CleanUp();
 
+private:
+	/** Simply unregisters self from navigation system and calls CleanUp */
+	void UnregisterAndCleanUp();
+
+public:
 	virtual void CleanUpAndMarkPendingKill();
 
 	FORCEINLINE bool IsRegistered() const { return bRegistered; }
diff --git a/Engine/Source/Runtime/Engine/Private/AI/Navigation/NavigationData.cpp b/Engine/Source/Runtime/Engine/Private/AI/Navigation/NavigationData.cpp
--- a/Engine/Source/Runtime/Engine/Private/AI/Navigation/NavigationData.cpp
+++ b/Engine/Source/Runtime/Engine/Private/AI/Navigation/NavigationData.cpp
@@ -400,19 +400,31 @@ bool ANavigationData::DoesSupportAgent(const FNavAgentProperties& AgentProps) co
 	return NavDataConfig.IsEquivalent(AgentProps);
 }
 
-void ANavigationData::Destroyed()
+void ANavigationData::EndPlay(const EEndPlayReason::Type EndPlayReason)
 {
-	UWorld* WorldOuter = GetWorld();
+	UnregisterAndCleanUp();
+	Super::EndPlay(EndPlayReason);
+}
 
-	bRegistered = false;
-	if (WorldOuter != NULL && WorldOuter->GetNavigationSystem() != NULL)
+void ANavigationData::Destroyed()
+{
+	UnregisterAndCleanUp();
+	Super::Destroyed();
+}
+void ANavigationData::UnregisterAndCleanUp()
+{
+	if (bRegistered)
 	{
-		WorldOuter->GetNavigationSystem()->UnregisterNavData(this);
-	}
+		UWorld* WorldOuter = GetWorld();
 
-	CleanUp();
+		bRegistered = false;
+		if (WorldOuter != NULL && WorldOuter->GetNavigationSystem() != NULL)
+		{
+			WorldOuter->GetNavigationSystem()->UnregisterNavData(this);
+		}
 
-	Super::Destroyed();
+		CleanUp();
+	}
 }
 
 void ANavigationData::CleanUp()
