UE-18845: Crash in audio after hot reload.

diff --git a/Engine/Source/Runtime/Engine/Classes/Engine/GameViewportClient.h b/Engine/Source/Runtime/Engine/Classes/Engine/GameViewportClient.h
--- a/Engine/Source/Runtime/Engine/Classes/Engine/GameViewportClient.h
+++ b/Engine/Source/Runtime/Engine/Classes/Engine/GameViewportClient.h
@@ -54,6 +54,11 @@ class ENGINE_API UGameViewportClient : public UScriptViewportClient, public FExe
 	GENERATED_UCLASS_BODY()
 
 public:
+#if WITH_HOT_RELOAD_CTORS
+	/** DO NOT USE. This constructor is for internal usage only for hot-reload purposes. */
+	UGameViewportClient(FVTableHelper& Helper);
+#endif // WITH_HOT_RELOAD_CTORS
+
 	virtual ~UGameViewportClient();
 
 	/** The viewport's console.   Might be null on consoles */
diff --git a/Engine/Source/Runtime/Engine/Private/GameViewportClient.cpp b/Engine/Source/Runtime/Engine/Private/GameViewportClient.cpp
--- a/Engine/Source/Runtime/Engine/Private/GameViewportClient.cpp
+++ b/Engine/Source/Runtime/Engine/Private/GameViewportClient.cpp
@@ -155,6 +155,21 @@ UGameViewportClient::UGameViewportClient(const FObjectInitializer& ObjectInitial
 	}
 }
 
+#if WITH_HOT_RELOAD_CTORS
+UGameViewportClient::UGameViewportClient(FVTableHelper& Helper)
+	: Super(Helper)
+	, EngineShowFlags(ESFIM_Game)
+	, CurrentBufferVisualizationMode(NAME_None)
+	, HighResScreenshotDialog(NULL)
+	, bIgnoreInput(false)
+	, MouseCaptureMode(EMouseCaptureMode::CapturePermanently)
+	, bHideCursorDuringCapture(false)
+	, AudioDeviceHandle(INDEX_NONE)
+	, bHasAudioFocus(false)
+{
+
+}
+#endif // WITH_HOT_RELOAD_CTORS
 
 UGameViewportClient::~UGameViewportClient()
 {
