Fix constraints not working when constrained against the world (no actor in the first slot)

diff --git a/Engine/Source/Runtime/Engine/Private/PhysicsEngine/ConstraintInstance.cpp b/Engine/Source/Runtime/Engine/Private/PhysicsEngine/ConstraintInstance.cpp
--- a/Engine/Source/Runtime/Engine/Private/PhysicsEngine/ConstraintInstance.cpp
+++ b/Engine/Source/Runtime/Engine/Private/PhysicsEngine/ConstraintInstance.cpp
@@ -343,7 +343,7 @@ PxScene* GetPScene_LockFree(const FBodyInstance* Body1, const FBodyInstance* Bod
 	}
 	else
 	{
-		PScene = GetPhysXSceneFromIndex(SceneIndex1);
+		PScene = GetPhysXSceneFromIndex(SceneIndex1 >= 0 ? SceneIndex1 : SceneIndex2);
 	}
 
 	return PScene;
