Fixed incorrect transform calculation in AActor::FinishSpawning.

diff --git a/Engine/Source/Runtime/Engine/Private/Actor.cpp b/Engine/Source/Runtime/Engine/Private/Actor.cpp
--- a/Engine/Source/Runtime/Engine/Private/Actor.cpp
+++ b/Engine/Source/Runtime/Engine/Private/Actor.cpp
@@ -2475,7 +2475,7 @@ void AActor::FinishSpawning(const FTransform& UserTransform, bool bIsDefaultTran
 
 		// if we have a native root component, its transform is already set properly, so we use that.  it might not be the same as UserTransform 
 		// if the root component in the CDO has a nonzero transform.
-		const FTransform FinalRootComponentTransform = RootComponent ? RootComponent->ComponentToWorld : UserTransform;
+		const FTransform FinalRootComponentTransform = RootComponent ? RootComponent->ComponentToWorld * UserTransform : UserTransform;
 
 		ExecuteConstruction(FinalRootComponentTransform, nullptr, bIsDefaultTransform);
 		PostActorConstruction();
