Fix physics handle being capped at 1.6 degrees tolerance. (Changed it to 0.016)

diff --git a/Engine/Source/Runtime/Engine/Private/PhysicsEngine/PhysicsHandleComponent.cpp b/Engine/Source/Runtime/Engine/Private/PhysicsEngine/PhysicsHandleComponent.cpp
--- a/Engine/Source/Runtime/Engine/Private/PhysicsEngine/PhysicsHandleComponent.cpp
+++ b/Engine/Source/Runtime/Engine/Private/PhysicsEngine/PhysicsHandleComponent.cpp
@@ -247,7 +247,7 @@ void UPhysicsHandleComponent::UpdateHandleTransform(const FTransform& NewTransfo
 	// Check if the new rotation is worthy of change
 	PxQuat PNewOrientation = U2PQuat(NewTransform.GetRotation());
 	PxQuat PCurrentOrientation = KinActor->getGlobalPose().q;
-	if(!(FMath::Abs(PNewOrientation.dot(PCurrentOrientation)) < (1.f - KINDA_SMALL_NUMBER)))
+	if((FMath::Abs(PNewOrientation.dot(PCurrentOrientation)) > (1.f - SMALL_NUMBER)))
 	{
 		PNewOrientation = PCurrentOrientation;
 		bChangedRotation = false;
