#ue4 - Fill in FHitResult.PenetrationDepth with ContactInfo.ContactPenetration in physics collision event.

UDN: https://udn.unrealengine.com/questions/253253/detecting-object-overlaps-on-startup.html

diff --git a/Engine/Source/Runtime/Engine/Private/Actor.cpp b/Engine/Source/Runtime/Engine/Private/Actor.cpp
--- a/Engine/Source/Runtime/Engine/Private/Actor.cpp
+++ b/Engine/Source/Runtime/Engine/Private/Actor.cpp
@@ -3423,6 +3423,7 @@ void AActor::DispatchPhysicsCollisionHit(const FRigidBodyCollisionInfo& MyInfo,
 	FHitResult Result;
 	Result.Location = Result.ImpactPoint = ContactInfo.ContactPosition;
 	Result.Normal = Result.ImpactNormal = ContactInfo.ContactNormal;
+	Result.PenetrationDepth = ContactInfo.ContactPenetration;
 	Result.PhysMaterial = ContactInfo.PhysMaterial[1];
 	Result.Actor = OtherInfo.Actor;
 	Result.Component = OtherInfo.Component;
