Added proper load error message for missing bone as well

diff --git a/Engine/Source/Runtime/Engine/Private/Animation/Skeleton.cpp b/Engine/Source/Runtime/Engine/Private/Animation/Skeleton.cpp
--- a/Engine/Source/Runtime/Engine/Private/Animation/Skeleton.cpp
+++ b/Engine/Source/Runtime/Engine/Private/Animation/Skeleton.cpp
@@ -10,6 +10,7 @@
 #include "AssetRegistryModule.h"
 #include "Animation/AnimSequence.h"
 #include "Animation/Rig.h"
+#include "MessageLog.h"
 
 #define LOCTEXT_NAMESPACE "Skeleton"
 #define ROOT_BONE_PARENT	INDEX_NONE
@@ -380,6 +381,12 @@ int32 USkeleton::BuildLinkup(const USkeletalMesh* InSkelMesh)
 				bDismissedMessage = true;
 			}
 
+			static FName NAME_LoadErrors("LoadErrors");
+			FMessageLog LoadErrors(NAME_LoadErrors);
+
+			TSharedRef<FTokenizedMessage> Message = LoadErrors.Info();
+			Message->AddToken(FTextToken::Create(FText::Format(LOCTEXT("SkeletonBuildLinkupMissingBones", "The Skeleton {0}, is missing bones that SkeletalMesh {1} needs. They will be added now. Please save the Skeleton!"), FText::FromString(GetNameSafe(this)), FText::FromString(GetNameSafe(InSkelMesh)))));
+
 			// Re-add all SkelMesh bones to the Skeleton.
 			MergeAllBonesToBoneTree(InSkelMesh);
 
