Made Asset selectable, so that you can go to the skeleton right away.

diff --git a/Engine/Source/Runtime/Engine/Private/Animation/Skeleton.cpp b/Engine/Source/Runtime/Engine/Private/Animation/Skeleton.cpp
--- a/Engine/Source/Runtime/Engine/Private/Animation/Skeleton.cpp
+++ b/Engine/Source/Runtime/Engine/Private/Animation/Skeleton.cpp
@@ -385,7 +385,11 @@ int32 USkeleton::BuildLinkup(const USkeletalMesh* InSkelMesh)
 			FMessageLog LoadErrors(NAME_LoadErrors);
 
 			TSharedRef<FTokenizedMessage> Message = LoadErrors.Info();
-			Message->AddToken(FTextToken::Create(FText::Format(LOCTEXT("SkeletonBuildLinkupMissingBones", "The Skeleton {0}, is missing bones that SkeletalMesh {1} needs. They will be added now. Please save the Skeleton!"), FText::FromString(GetNameSafe(this)), FText::FromString(GetNameSafe(InSkelMesh)))));
+			Message->AddToken(FTextToken::Create(LOCTEXT("SkeletonBuildLinkupMissingBones1", "The Skeleton ")));
+			Message->AddToken(FAssetNameToken::Create(GetPathName(), FText::FromString( GetNameSafe(this) ) ));
+			Message->AddToken(FTextToken::Create(LOCTEXT("SkeletonBuildLinkupMissingBones2", " is missing bones that SkeletalMesh ")));
+			Message->AddToken(FAssetNameToken::Create(InSkelMesh->GetPathName(), FText::FromString( GetNameSafe(InSkelMesh) )));
+			Message->AddToken(FTextToken::Create(LOCTEXT("SkeletonBuildLinkupMissingBones3", "  needs. They will be added now. Please save the Skeleton!")));
 
 			// Re-add all SkelMesh bones to the Skeleton.
 			MergeAllBonesToBoneTree(InSkelMesh);
