#UE4 Keys bound via key or action events will no longer count as having consumed the event if the event did not occur and event type was "DoubleClick". This was preventing axis mapped bindings from firing when a "DoubleClick" action was also mapped. (i.e Hold D to strafe right, but double tap D to dash right). UE-18962

diff --git a/Engine/Source/Runtime/Engine/Private/UserInterface/PlayerInput.cpp b/Engine/Source/Runtime/Engine/Private/UserInterface/PlayerInput.cpp
--- a/Engine/Source/Runtime/Engine/Private/UserInterface/PlayerInput.cpp
+++ b/Engine/Source/Runtime/Engine/Private/UserInterface/PlayerInput.cpp
@@ -650,7 +650,7 @@ void UPlayerInput::GetChordsForKeyMapping(const FInputActionKeyMapping& KeyMappi
 			bConsumeInput = true;
 		}
 	}
-	if (ActionBinding.bConsumeInput && (bConsumeInput || !(KeyMapping.bAlt || KeyMapping.bCtrl || KeyMapping.bShift || KeyMapping.bCmd)))
+	if (ActionBinding.bConsumeInput && (bConsumeInput || !(KeyMapping.bAlt || KeyMapping.bCtrl || KeyMapping.bShift || KeyMapping.bCmd || ActionBinding.KeyEvent == EInputEvent::IE_DoubleClick)))
 	{
 		KeysToConsume.AddUnique(KeyMapping.Key);
 	}
@@ -750,7 +750,7 @@ void UPlayerInput::GetChordForKey(const FInputKeyBinding& KeyBinding, const bool
 			}
 		}
 	}
-	if (KeyBinding.bConsumeInput && (bConsumeInput || !(KeyBinding.Chord.bAlt || KeyBinding.Chord.bCtrl || KeyBinding.Chord.bShift || KeyBinding.Chord.bCmd)))
+	if (KeyBinding.bConsumeInput && (bConsumeInput || !(KeyBinding.Chord.bAlt || KeyBinding.Chord.bCtrl || KeyBinding.Chord.bShift || KeyBinding.Chord.bCmd || KeyBinding.KeyEvent == EInputEvent::IE_DoubleClick)))
 	{
 		KeysToConsume.AddUnique(KeyBinding.Chord.Key);
 	}
