[UE-19067] Backed out CL# 2627557 temporarily to unblock promotable while I investigate further.

diff --git a/Engine/Source/Runtime/Engine/Private/Blueprint.cpp b/Engine/Source/Runtime/Engine/Private/Blueprint.cpp
--- a/Engine/Source/Runtime/Engine/Private/Blueprint.cpp
+++ b/Engine/Source/Runtime/Engine/Private/Blueprint.cpp
@@ -941,15 +941,7 @@ bool UBlueprint::GetBlueprintHierarchyFromClass(const UClass* InClass, TArray<UB
 		bNoErrors &= (BP->Status != BS_Error);
 #endif // #if WITH_EDITORONLY_DATA
 
-		// If valid, use stored ParentClass rather than the actual UClass::GetSuperClass(); handles the case when the class has not been recompiled yet after a reparent operation.
-		if(BP->ParentClass != nullptr)
-		{
-			CurrentClass = BP->ParentClass;
-		}
-		else
-		{
-			CurrentClass = CurrentClass->GetSuperClass();
-		}
+		CurrentClass = CurrentClass->GetSuperClass();
 	}
 
 	return bNoErrors;
diff --git a/Engine/Source/Runtime/Engine/Private/SimpleConstructionScript.cpp b/Engine/Source/Runtime/Engine/Private/SimpleConstructionScript.cpp
--- a/Engine/Source/Runtime/Engine/Private/SimpleConstructionScript.cpp
+++ b/Engine/Source/Runtime/Engine/Private/SimpleConstructionScript.cpp
@@ -926,36 +926,9 @@ void USimpleConstructionScript::ValidateSceneRootNodes()
 		else if(RootComponentTemplate != nullptr
 			&& RootNodes.Contains(DefaultSceneRootNode))
 		{
-			// Temporarily parent to a native or inherited root (to propagate to any children that might be promoted on removal below)
-			if(UBlueprintGeneratedClass* BPClass = Cast<UBlueprintGeneratedClass>(RootComponentTemplate->GetOuter()))
-			{
-				if(BPClass->SimpleConstructionScript != nullptr
-					&& BPClass->SimpleConstructionScript != this)
-				{
-					DefaultSceneRootNode->bIsParentComponentNative = false;
-					DefaultSceneRootNode->ParentComponentOwnerClassName = BPClass->GetFName();
-					TArray<const USCS_Node*> ParentSCSNodes = BPClass->SimpleConstructionScript->GetAllNodesConst();
-					for(const USCS_Node* ParentSCSNode : ParentSCSNodes)
-					{
-						if(ParentSCSNode != nullptr && ParentSCSNode->ComponentTemplate == RootComponentTemplate)
-						{
-							DefaultSceneRootNode->ParentComponentOrVariableName = ParentSCSNode->VariableName;
-							break;
-						}
-					}
-				}
-			}
-			else
-			{
-				DefaultSceneRootNode->bIsParentComponentNative = true;
-				DefaultSceneRootNode->ParentComponentOrVariableName = RootComponentTemplate->GetFName();
-				DefaultSceneRootNode->ParentComponentOwnerClassName = NAME_None;
-			}
-
-			// Remove from the root node list (but keep child nodes)
-			RemoveNodeAndPromoteChildren(DefaultSceneRootNode);
+			RootNodes.Remove(DefaultSceneRootNode);
 
-			// Ensure that these are reset back to their default state
+			// These shouldn't be set, but just in case...
 			DefaultSceneRootNode->bIsParentComponentNative = false;
 			DefaultSceneRootNode->ParentComponentOrVariableName = NAME_None;
 			DefaultSceneRootNode->ParentComponentOwnerClassName = NAME_None;
