Added some additional logging for a rare crash. UE-3518

diff --git a/Engine/Source/Runtime/Engine/Private/Materials/HLSLMaterialTranslator.h b/Engine/Source/Runtime/Engine/Private/Materials/HLSLMaterialTranslator.h
--- a/Engine/Source/Runtime/Engine/Private/Materials/HLSLMaterialTranslator.h
+++ b/Engine/Source/Runtime/Engine/Private/Materials/HLSLMaterialTranslator.h
@@ -2713,7 +2713,21 @@ protected:
 
 		EMaterialValueType ShaderType = InTexture->GetMaterialType();
 		const int32 TextureReferenceIndex = Material->GetReferencedTextures().Find(InTexture);
+
+#if DO_CHECK
+		// UE-3518: Additional pre-assert logging to help determine the cause of this failure.
+		if (TextureReferenceIndex == INDEX_NONE)
+		{
+			const TArray<UTexture*>& ReferencedTextures = Material->GetReferencedTextures();
+			UE_LOG(LogMaterial, Error, TEXT("Compiler->Texture() failed to find texture '%s' in referenced list of size '%i':"), *InTexture->GetName(), ReferencedTextures.Num());
+			for (int32 i = 0; i < ReferencedTextures.Num(); ++i)
+			{
+				UE_LOG(LogMaterial, Error, TEXT("%i: '%s'"), i, ReferencedTextures[i] ? *ReferencedTextures[i]->GetName() : TEXT("nullptr"));
+			}
+		}
+#endif
 		checkf(TextureReferenceIndex != INDEX_NONE, TEXT("Material expression called Compiler->Texture() without implementing UMaterialExpression::GetReferencedTexture properly"));
+
 		return AddUniformExpression(new FMaterialUniformExpressionTexture(TextureReferenceIndex, SamplerSource),ShaderType,TEXT(""));
 	}
 
