Added upgrade path for legacy SourceFilePath data stored on UAssetImportData structures

Contributed by [REDACTED_PERSON]

PR#1403

diff --git a/Engine/Source/Runtime/Engine/Classes/EditorFramework/AssetImportData.h b/Engine/Source/Runtime/Engine/Classes/EditorFramework/AssetImportData.h
--- a/Engine/Source/Runtime/Engine/Classes/EditorFramework/AssetImportData.h
+++ b/Engine/Source/Runtime/Engine/Classes/EditorFramework/AssetImportData.h
@@ -61,6 +61,14 @@ public:
 	UPROPERTY()
 	uint32 bDirty : 1;
 
+	/** Path to the resource used to construct this static mesh. Relative to the object's package, BaseDir() or absolute */
+	UPROPERTY()
+	FString SourceFilePath_DEPRECATED;
+
+	/** Date/Time-stamp of the file from the last import */
+	UPROPERTY()
+	FString SourceFileTimestamp_DEPRECATED;
+
 public:
 
 	/** Static event that is broadcast whenever any asset has updated its import data */
@@ -84,6 +92,8 @@ public:
 	/** Resolve a filename that is relative to either the specified package, BaseDir() or absolute */
 	static FString ResolveImportFilename(const FString& InRelativePath, const UPackage* Outermost);
 
+	virtual void PostLoad() override;
+
 protected:
 
 	/** Convert an absolute import path so that it's relative to either this object's package, BaseDir() or leave it absolute */
diff --git a/Engine/Source/Runtime/Engine/Private/EditorFramework/AssetImportData.cpp b/Engine/Source/Runtime/Engine/Private/EditorFramework/AssetImportData.cpp
--- a/Engine/Source/Runtime/Engine/Private/EditorFramework/AssetImportData.cpp
+++ b/Engine/Source/Runtime/Engine/Private/EditorFramework/AssetImportData.cpp
@@ -198,6 +198,25 @@ void UAssetImportData::Serialize(FArchive& Ar)
 	Super::Serialize(Ar);
 }
 
+void UAssetImportData::PostLoad()
+{
+	if (!SourceFilePath_DEPRECATED.IsEmpty() && SourceFiles.Num() == 0)
+	{
+		FDateTime SourceDateTime;
+		if (!FDateTime::Parse(SourceFileTimestamp_DEPRECATED,SourceDateTime))
+		{
+			SourceDateTime = 0;
+		}
+
+		SourceFiles.Add(FSourceFile(MoveTemp(SourceFilePath_DEPRECATED),SourceDateTime));
+
+		SourceFilePath_DEPRECATED.Empty();
+		SourceFileTimestamp_DEPRECATED.Empty();
+	}
+
+	Super::PostLoad();
+}
+
 void UAssetImportData::PostEditChangeProperty(FPropertyChangedEvent& PropertyChangedEvent)
 {
 	Super::PostEditChangeProperty(PropertyChangedEvent);
