Fix spawn transform of spot light
Now you shouldn't have to update component to world manually in the constructor

- https://jira.ol.epicgames.net/browse/UE-19034

diff --git a/Engine/Source/Runtime/Engine/Private/SpotLight.cpp b/Engine/Source/Runtime/Engine/Private/SpotLight.cpp
--- a/Engine/Source/Runtime/Engine/Private/SpotLight.cpp
+++ b/Engine/Source/Runtime/Engine/Private/SpotLight.cpp
@@ -23,7 +23,6 @@ ASpotLight::ASpotLight(const FObjectInitializer& ObjectInitializer)
 	SpotLightComponent = CastChecked<USpotLightComponent>(GetLightComponent());
 	SpotLightComponent->Mobility = EComponentMobility::Stationary;
 	SpotLightComponent->RelativeRotation = FRotator(-90, 0, 0);
-	SpotLightComponent->UpdateComponentToWorld();
 
 	RootComponent = SpotLightComponent;
 
