Fix new maps becoming dirty on creation

https://jira.ol.epicgames.net/browse/UE-19119

diff --git a/Engine/Source/Runtime/Engine/Classes/GameFramework/Actor.h b/Engine/Source/Runtime/Engine/Classes/GameFramework/Actor.h
--- a/Engine/Source/Runtime/Engine/Classes/GameFramework/Actor.h
+++ b/Engine/Source/Runtime/Engine/Classes/GameFramework/Actor.h
@@ -1597,8 +1597,9 @@ public:
 	/**
 	 * Assigns a new label to this actor.  Actor labels are only available in development builds.
 	 * @param	NewActorLabel	The new label string to assign to the actor.  If empty, the actor will have a default label.
+	 * @param	bMarkDirty		If true the actor's package will be marked dirty for saving.  Otherwise it will not be.  You should pass false for this parameter if dirtying is not allowed (like during loads)
 	 */
-	void SetActorLabel( const FString& NewActorLabel );
+	void SetActorLabel( const FString& NewActorLabel, bool bMarkDirty = true );
 
 	/** Advanced - clear the actor label. */
 	void ClearActorLabel();
@@ -2534,7 +2535,7 @@ private:
 #endif // ENABLE_VISUAL_LOG
 
 	//* Sets the friendly actor label and name */
-	void SetActorLabelInternal( const FString& NewActorLabelDirty, bool bMakeGloballyUniqueFName );
+	void SetActorLabelInternal( const FString& NewActorLabelDirty, bool bMakeGloballyUniqueFName, bool bMarkDirty );
 
 	static FMakeNoiseDelegate MakeNoiseDelegate;
 
diff --git a/Engine/Source/Runtime/Engine/Private/AI/Navigation/NavigationSystem.cpp b/Engine/Source/Runtime/Engine/Private/AI/Navigation/NavigationSystem.cpp
--- a/Engine/Source/Runtime/Engine/Private/AI/Navigation/NavigationSystem.cpp
+++ b/Engine/Source/Runtime/Engine/Private/AI/Navigation/NavigationSystem.cpp
@@ -3202,9 +3202,10 @@ ANavigationData* UNavigationSystem::CreateNavigationDataInstance(const FNavDataC
 			}
 
 			// Set descriptive name
-			Instance->Rename(*StrName);
+			Instance->Rename(*StrName, NULL, REN_DoNotDirty);
 #if WITH_EDITOR
-			Instance->SetActorLabel(StrName);
+			const bool bMarkDirty = false;
+			Instance->SetActorLabel(StrName, bMarkDirty );
 #endif // WITH_EDITOR
 		}
 	}
diff --git a/Engine/Source/Runtime/Engine/Private/ActorEditor.cpp b/Engine/Source/Runtime/Engine/Private/ActorEditor.cpp
--- a/Engine/Source/Runtime/Engine/Private/ActorEditor.cpp
+++ b/Engine/Source/Runtime/Engine/Private/ActorEditor.cpp
@@ -605,12 +605,13 @@ const FString& AActor::GetActorLabel() const
 	return ActorLabel;
 }
 
-void AActor::SetActorLabel( const FString& NewActorLabelDirty )
+void AActor::SetActorLabel( const FString& NewActorLabelDirty, bool bMarkDirty )
 {
-	SetActorLabelInternal(NewActorLabelDirty, false);
+	const bool bMakeGloballyUniqueFName = false;
+	SetActorLabelInternal(NewActorLabelDirty, bMakeGloballyUniqueFName, bMarkDirty );
 }
 
-void AActor::SetActorLabelInternal( const FString& NewActorLabelDirty, bool bMakeGloballyUniqueFName )
+void AActor::SetActorLabelInternal( const FString& NewActorLabelDirty, bool bMakeGloballyUniqueFName, bool bMarkDirty )
 {
 	// Clean up the incoming string a bit
 	FString NewActorLabel = NewActorLabelDirty;
@@ -624,7 +625,7 @@ void AActor::SetActorLabelInternal( const FString& NewActorLabelDirty, bool bMak
 		if( FCString::Strcmp( *NewActorLabel, *GetActorLabel() ) != 0 )
 		{
 			// Store new label
-			Modify();
+			Modify( bMarkDirty );
 			ActorLabel = NewActorLabel;
 		}
 	}
