Disabling single process PIE seamless travel while we work on a long term fix

diff --git a/Engine/Source/Runtime/Engine/Private/World.cpp b/Engine/Source/Runtime/Engine/Private/World.cpp
--- a/Engine/Source/Runtime/Engine/Private/World.cpp
+++ b/Engine/Source/Runtime/Engine/Private/World.cpp
@@ -3268,13 +3268,6 @@ void UWorld::AddNetworkActor( AActor* Actor )
 		return;
 	}
 
-#if 0
-	if ( ( !NetDriver || NetDriver->ServerConnection ) && !GIsEditor )
-	{
-		return;		// Clients or standalone games don't use this list
-	}
-#endif
-
 	if ( !ContainsLevel(Actor->GetLevel()) )
 	{
 		return;
@@ -5292,6 +5285,14 @@ void UWorld::SetSelectedLevels( const TArray<class ULevel*>& InLevels )
  */
 void UWorld::ServerTravel(const FString& FURL, bool bAbsolute, bool bShouldSkipGameNotify)
 {
+	// NOTE - This is a temp check while we work on a long term fix
+	// There are a few issues with seamless travel using single process PIE, so we're disabling that for now while working on a fix
+	if ( WorldType == EWorldType::PIE && AuthorityGameMode && AuthorityGameMode->bUseSeamlessTravel && !FParse::Param( FCommandLine::Get(), TEXT( "MultiprocessOSS" ) ) )
+	{
+		UE_LOG( LogWorld, Warning, TEXT( "UWorld::ServerTravel: Seamless travel currently NOT supported in single process PIE." ) );
+		return;
+	}
+
 	if (FURL.Contains(TEXT("%")) )
 	{
 		UE_LOG(LogWorld, Log, TEXT("FURL %s Contains illegal character '%%'."), *FURL);
