Fixed another incorrect transform calculation in SpawnActor. Also properly handle case where callers pass different transforms into SpawnActor and FinishSpawning.

diff --git a/Engine/Source/Runtime/Engine/Private/Actor.cpp b/Engine/Source/Runtime/Engine/Private/Actor.cpp
--- a/Engine/Source/Runtime/Engine/Private/Actor.cpp
+++ b/Engine/Source/Runtime/Engine/Private/Actor.cpp
@@ -2396,6 +2396,23 @@ static USceneComponent* FixupNativeActorComponents(AActor* Actor)
 	return SceneRootComponent;
 }
 
+// Simple and short-lived cache for storing transforms between beginning and finishing spawning.
+static TMap< TWeakObjectPtr<AActor>, FTransform > GSpawnActorDeferredTransformCache;
+
+static void ValidateDeferredTransformCache()
+{
+	// clean out any entries where the actor is no longer valid
+	// could happen if an actor is destroyed before FinishSpawning is called
+	for (auto It = GSpawnActorDeferredTransformCache.CreateIterator(); It; ++It)
+	{
+		TWeakObjectPtr<AActor> ActorRef = It.Key();
+		if (ActorRef.IsValid() == false)
+		{
+			It.RemoveCurrent();
+		}
+	}
+}
+
 void AActor::PostSpawnInitialize(FTransform const& UserSpawnTransform, AActor* InOwner, APawn* InInstigator, bool bRemoteOwned, bool bNoFail, bool bDeferConstruction)
 {
 	// General flow here is like so
@@ -2462,9 +2479,14 @@ void AActor::PostSpawnInitialize(FTransform const& UserSpawnTransform, AActor* I
 	// After this, we can assume all components are created and assembled.
 	if (!bDeferConstruction)
 	{
-		// Preserve original root component scale
 		FinishSpawning(UserSpawnTransform, true);
 	}
+	else if (SceneRootComponent != nullptr)
+	{
+		// we have a native root component and are deferring construction, store our original UserSpawnTransform
+		// so we can do the proper thing if the user passes in a different transform during FinishSpawning
+		GSpawnActorDeferredTransformCache.Emplace(this, UserSpawnTransform);
+	}
 }
 
 void AActor::FinishSpawning(const FTransform& UserTransform, bool bIsDefaultTransform)
@@ -2473,9 +2495,30 @@ void AActor::FinishSpawning(const FTransform& UserTransform, bool bIsDefaultTran
 	{
 		bHasFinishedSpawning = true;
 
-		// if we have a native root component, its transform is already set properly, so we use that.  it might not be the same as UserTransform 
-		// if the root component in the CDO has a nonzero transform.
-		const FTransform FinalRootComponentTransform = RootComponent ? RootComponent->ComponentToWorld * UserTransform : UserTransform;
+		FTransform FinalRootComponentTransform = UserTransform;
+
+		// see if we need to adjust the transform (i.e. in deferred cases where the caller passes in a different transform here 
+		// than was passed in during the original SpawnActor call)
+		if (RootComponent && !bIsDefaultTransform)
+		{
+			FTransform const* const OriginalSpawnTransform = GSpawnActorDeferredTransformCache.Find(this);
+			if (OriginalSpawnTransform)
+			{
+				GSpawnActorDeferredTransformCache.Remove(this);
+
+				if (OriginalSpawnTransform->Equals(UserTransform) == false)
+				{
+					// caller passed a different transform!
+					// undo the original spawn transform to get back to the template transform, so we can recompute a good
+					// final transform that takes into account the template's transform
+					FTransform const TemplateTransform = RootComponent->ComponentToWorld * OriginalSpawnTransform->Inverse();
+					FinalRootComponentTransform = TemplateTransform * UserTransform;
+				}
+			}
+
+			// should be fast and relatively rare
+			ValidateDeferredTransformCache();
+		}
 
 		ExecuteConstruction(FinalRootComponentTransform, nullptr, bIsDefaultTransform);
 		PostActorConstruction();
