Fixed a bug in navigation system that was resulting in rejecting save navigation data chunks in the editor #UE4

This change might require rebuilding and saving map indicated in the ticket
UE-19106

diff --git a/Engine/Source/Runtime/Engine/Private/AI/Navigation/NavigationSystem.cpp b/Engine/Source/Runtime/Engine/Private/AI/Navigation/NavigationSystem.cpp
--- a/Engine/Source/Runtime/Engine/Private/AI/Navigation/NavigationSystem.cpp
+++ b/Engine/Source/Runtime/Engine/Private/AI/Navigation/NavigationSystem.cpp
@@ -2208,7 +2208,7 @@ void UNavigationSystem::InitializeForWorld(UWorld* World, FNavigationSystemRunMo
 
 		// Remove old/stale chunk data from all levels, when navigation auto-update is enabled
 		// In case navigation system will be created chunks will be regenerated anyway
-		if (Mode == FNavigationSystemRunMode::EditorMode)
+		if (Mode == FNavigationSystemRunMode::EditorMode && bNavigationAutoUpdateEnabled == true)
 		{
 			DiscardNavigationDataChunks(World);
 		}
