Travel to default map when replay map fails to load

diff --git a/Engine/Source/Runtime/Engine/Private/DemoNetDriver.cpp b/Engine/Source/Runtime/Engine/Private/DemoNetDriver.cpp
--- a/Engine/Source/Runtime/Engine/Private/DemoNetDriver.cpp
+++ b/Engine/Source/Runtime/Engine/Private/DemoNetDriver.cpp
@@ -359,6 +359,7 @@ bool UDemoNetDriver::InitConnectInternal( FString& Error )
 		{
 			GEngine->DestroyNamedNetDriver( WorldContext->PendingNetGame, NetDriverName );
 			WorldContext->PendingNetGame = NULL;
+			GEngine->BrowseToDefaultMap( *WorldContext );
 		}
 
 		Error = LoadMapError;
