Fixed translucent self-shadowing on static meshes using the static draw path

diff --git a/Engine/Source/Runtime/Engine/Private/StaticMeshRender.cpp b/Engine/Source/Runtime/Engine/Private/StaticMeshRender.cpp
--- a/Engine/Source/Runtime/Engine/Private/StaticMeshRender.cpp
+++ b/Engine/Source/Runtime/Engine/Private/StaticMeshRender.cpp
@@ -499,6 +499,7 @@ void FStaticMeshSceneProxy::DrawStaticElements(FStaticPrimitiveDrawInterface* PD
 					bool bAnySectionUsesDitheredLODTransition = false;
 					bool bAllSectionsUseDitheredLODTransition = true;
 					bool bIsMovable = IsMovable();
+
 					for (int32 SectionIndex = 0; bSafeToUseShadowOnlyMesh && SectionIndex < LODModel.Sections.Num(); SectionIndex++)
 					{
 						const FMaterial* Material = ProxyLODInfo.Sections[SectionIndex].Material->GetRenderProxy(false)->GetMaterial(FeatureLevel);
@@ -506,11 +507,13 @@ void FStaticMeshSceneProxy::DrawStaticElements(FStaticPrimitiveDrawInterface* PD
 						bAnySectionUsesDitheredLODTransition = bAnySectionUsesDitheredLODTransition || (!bIsMovable && Material->IsDitheredLODTransition());
 						bAllSectionsUseDitheredLODTransition = bAllSectionsUseDitheredLODTransition && (!bIsMovable && Material->IsDitheredLODTransition());
 						const FStaticMeshSection& Section = LODModel.Sections[SectionIndex];
+
 						bSafeToUseShadowOnlyMesh =
 							Section.bCastShadow
 							&& !(bAnySectionUsesDitheredLODTransition && !bAllSectionsUseDitheredLODTransition) // can't use a single section if they are not homogeneous
 							&& Material->WritesEveryPixel()
 							&& !Material->IsTwoSided()
+							&& !IsTranslucentBlendMode(Material->GetBlendMode())
 							&& !Material->MaterialModifiesMeshPosition_RenderThread();
 						bAnySectionCastsShadow |= Section.bCastShadow;
 					}
