Submitting CL 2624473 for MichaelS:
Undo/Redo across a Blueprint compile of transactions involving instanced sub-objects in the CDO will no longer lead to crashes.

diff --git a/Engine/Source/Editor/Kismet/Private/BlueprintEditor.cpp b/Engine/Source/Editor/Kismet/Private/BlueprintEditor.cpp
--- a/Engine/Source/Editor/Kismet/Private/BlueprintEditor.cpp
+++ b/Engine/Source/Editor/Kismet/Private/BlueprintEditor.cpp
@@ -98,6 +98,7 @@
 #include "Engine/LevelStreamingKismet.h"
 
 #include "IMenu.h"
+#include "InstancedReferenceSubobjectHelper.h"
 
 #define LOCTEXT_NAMESPACE "BlueprintEditor"
 
@@ -2704,6 +2705,51 @@ bool FBlueprintEditor::CanNavigateToChildGraph() const
 	return FocusedGraphEdPtr.IsValid() && (FocusedGraphEdPtr.Pin()->GetCurrentGraph()->SubGraphs.Num() > 0);
 }
 
+void FBlueprintEditor::FixSubObjectReferencesPostUndoRedo(UObject* InObject)
+{
+	// Post undo/redo, these may have the correct Outer but are not referenced by the CDO's UProperties
+	TArray<UObject*> SubObjects;
+	GetObjectsWithOuter(InObject, SubObjects, false);
+
+	// Post undo/redo, these may have the in-correct Outer but are incorrectly referenced by the CDO's UProperties
+	TSet<UObject*> PropertySubObjectReferences;
+	UClass* ObjectClass = InObject->GetClass();
+	FFindInstancedReferenceSubobjectHelper::Get(ObjectClass, reinterpret_cast<uint8*>(InObject), PropertySubObjectReferences);
+
+	TMap<UObject*, UObject*> OldToNewInstanceMap;
+	for (UObject* PropertySubObject : PropertySubObjectReferences)
+	{
+		bool bFoundMatchingSubObject = false;
+		for (UObject* SubObject : SubObjects)
+		{
+			// The property and sub-objects should have the same name.
+			if (PropertySubObject->GetFName() == SubObject->GetFName())
+			{
+				if (PropertySubObject->GetOuter() != InObject)
+				{
+					OldToNewInstanceMap.Add(PropertySubObject, SubObject);
+					bFoundMatchingSubObject = true;
+				}
+				// Recurse on the SubObject to correct any sub-object/property references
+				FixSubObjectReferencesPostUndoRedo(SubObject);
+				break;
+			}
+		}
+
+		// If the property referenced does not exist in the current context as a subobject, we need to duplicate it and fix up references
+		// This will occur during post-undo/redo of deletions
+		if (!bFoundMatchingSubObject)
+		{
+			UObject* NewSubObject = DuplicateObject(PropertySubObject, InObject, *PropertySubObject->GetName());
+
+			// Don't forget to fix up all references and sub-object references
+			OldToNewInstanceMap.Add(PropertySubObject, NewSubObject);
+		}
+	}
+
+	FArchiveReplaceObjectRef<UObject> Replacer(InObject, OldToNewInstanceMap, false, false, false, false);
+}
+
 void FBlueprintEditor::PostUndo(bool bSuccess)
 {	
 	// Clear selection, to avoid holding refs to nodes that go away
@@ -2733,6 +2779,8 @@ void FBlueprintEditor::PostUndo(bool bSuccess)
 		{
 			SetUISelectionState(NAME_None);
 
+			FixSubObjectReferencesPostUndoRedo(GetBlueprintObj()->GeneratedClass->GetDefaultObject());
+
 			// Will cause a call to RefreshEditors()
 			FBlueprintEditorUtils::MarkBlueprintAsStructurallyModified(GetBlueprintObj());
 
@@ -2765,6 +2813,8 @@ void FBlueprintEditor::PostRedo(bool bSuccess)
 		// Transaction affects the Blueprint this editor handles, so react as necessary
 		if (bAffectsBlueprint)
 		{
+			FixSubObjectReferencesPostUndoRedo(GetBlueprintObj()->GeneratedClass->GetDefaultObject());
+
 			// Will cause a call to RefreshEditors()
 			FBlueprintEditorUtils::MarkBlueprintAsStructurallyModified( BlueprintObj );
 
diff --git a/Engine/Source/Editor/Kismet/Public/BlueprintEditor.h b/Engine/Source/Editor/Kismet/Public/BlueprintEditor.h
--- a/Engine/Source/Editor/Kismet/Public/BlueprintEditor.h
+++ b/Engine/Source/Editor/Kismet/Public/BlueprintEditor.h
@@ -955,6 +955,9 @@ private:
 	/** Util to try and get doc link for the currently selected node */
 	FString GetDocLinkForSelectedNode();
 
+	/** Fixes SubObject references of the passed object so they match up to sub-object UProperty references */
+	void FixSubObjectReferencesPostUndoRedo(UObject* InObject);
+
 public://@TODO
 	TSharedPtr<FDocumentTracker> DocumentManager;
 
diff --git a/Engine/Source/Editor/UnrealEd/Classes/Editor/Transactor.h b/Engine/Source/Editor/UnrealEd/Classes/Editor/Transactor.h
--- a/Engine/Source/Editor/UnrealEd/Classes/Editor/Transactor.h
+++ b/Engine/Source/Editor/UnrealEd/Classes/Editor/Transactor.h
@@ -49,6 +49,33 @@ inline bool OuterIsCDO(const UObject* Obj, UObject*& OutCDO)
 	return bOuterIsCDO;
 }
 
+inline bool OuterIsCDO(const UObject* Obj, UObject*& OutCDO, TArray<FName>& OutHierarchyNames)
+{
+	bool bOuterIsCDO = false;
+	UObject* Iter = Obj->GetOuter();
+	while (Iter)
+	{
+		if (Iter->HasAllFlags(RF_ClassDefaultObject))
+		{
+			bOuterIsCDO = true;
+			OutCDO = Iter;
+			break;
+		}
+		else
+		{
+			OutHierarchyNames.Add(Iter->GetFName());
+		}
+		Iter = Iter->GetOuter();
+	}
+
+	// If the outer is not a CDO, clear the hierarchy of names, they will not be used
+	if (bOuterIsCDO == false)
+	{
+		OutHierarchyNames.Empty();
+	}
+	return bOuterIsCDO;
+}
+
 /*-----------------------------------------------------------------------------
 	FTransaction.
 -----------------------------------------------------------------------------*/
@@ -97,13 +124,11 @@ protected:
 		private:
 			UObject* Object;
 			UClass* SourceCDO;
-			FName SubObjectId;
-
+			TArray<FName> SubObjectHierarchyID;
 		public:
 			FPersistentObjectRef()
 				: Object(nullptr)
 				, SourceCDO(nullptr)
-				, SubObjectId()
 			{}
 
 			FPersistentObjectRef(UObject* InObject)
@@ -117,7 +142,7 @@ protected:
 				// objects' names so we can look them up later in Get() 
 				const bool bIsCDO = InObject && InObject->HasAllFlags(RF_ClassDefaultObject);
 				UObject* CDO = bIsCDO ? InObject : nullptr;
-				const bool bIsSubobjectOfCDO = OuterIsCDO(InObject, CDO);
+				const bool bIsSubobjectOfCDO = OuterIsCDO(InObject, CDO, SubObjectHierarchyID);
 				
 				// we have to be careful though, Blueprints also duplicate CDOs 
 				// and their sub-objects; we don't want changes to the 
@@ -136,7 +161,7 @@ protected:
 
 					if (bIsSubobjectOfCDO)
 					{
-						SubObjectId = InObject->GetFName();
+						SubObjectHierarchyID.Add(InObject->GetFName());
 					}
 				}
 				else
@@ -152,15 +177,55 @@ protected:
 				//ensure(Get() == InObject);
 			}
 
+			/** Returns TRUE if the recorded object is part of the CDO */
+			bool IsPartOfCDO() const
+			{
+				return SourceCDO != nullptr;
+			}
+
 			UObject* Get() const
 			{
 				checkSlow(!SourceCDO || !Object);
 				if (SourceCDO)
 				{
-					if (SubObjectId != FName())
+					if (SubObjectHierarchyID.Num() > 0)
 					{
 						// find the subobject:
-						return SourceCDO->GetDefaultSubobjectByName(SubObjectId);
+						UObject* CurrentObject = SourceCDO->GetDefaultObject();
+						UObject* NextObject = CurrentObject->GetDefaultSubobjectByName(SubObjectHierarchyID[0]);
+
+						// Current increasing depth into sub-objects, starts at 1 to avoid the sub-object found and placed in NextObject.
+						int SubObjectDepth = 1;
+
+						// Only continue digging deeper if the current target sub-object was found
+						bool bFoundTargetSubObject = true;
+						while(NextObject && bFoundTargetSubObject)
+						{
+							CurrentObject = NextObject;
+							NextObject = nullptr;
+
+							// If there is no more depth to dig for our object, we are done.
+							if (SubObjectDepth < SubObjectHierarchyID.Num())
+							{
+								bFoundTargetSubObject = false;
+
+								// Look for any UObject's with the CurrentObject's outer to find the next sub-object
+								TArray<UObject*> OutDefaultSubobjects;
+								GetObjectsWithOuter(CurrentObject, OutDefaultSubobjects, false);
+								for (int32 SubobjectIndex = 0; SubobjectIndex < OutDefaultSubobjects.Num(); SubobjectIndex++)
+								{
+									if (OutDefaultSubobjects[SubobjectIndex]->GetFName() == SubObjectHierarchyID[SubObjectDepth])
+									{
+										SubObjectDepth++;
+										NextObject = OutDefaultSubobjects[SubobjectIndex];
+										bFoundTargetSubObject = true;
+										break;
+									}
+								}
+							}
+						}
+
+						return bFoundTargetSubObject? CurrentObject : nullptr;
 					}
 					else
 					{
diff --git a/Engine/Source/Editor/UnrealEd/Private/EditorTransaction.cpp b/Engine/Source/Editor/UnrealEd/Private/EditorTransaction.cpp
--- a/Engine/Source/Editor/UnrealEd/Private/EditorTransaction.cpp
+++ b/Engine/Source/Editor/UnrealEd/Private/EditorTransaction.cpp
@@ -36,6 +36,14 @@ FTransaction::FObjectRecord::FObjectRecord(FTransaction* Owner, UObject* InObjec
 
 void FTransaction::FObjectRecord::SerializeContents( FArchive& Ar, int32 InOper )
 {
+	// Cache to restore at the end
+	bool bWasArIgnoreOuterRef = Ar.ArIgnoreOuterRef;
+
+	if (Object.IsPartOfCDO())
+	{
+		Ar.ArIgnoreOuterRef = true;
+	}
+
 	if( Array )
 	{
 		//UE_LOG( LogEditorTransaction, Log, TEXT("Array %s %i*%i: %i"), Object ? *Object->GetFullName() : TEXT("Invalid Object"), Index, ElementSize, InOper);
@@ -89,6 +97,8 @@ void FTransaction::FObjectRecord::SerializeContents( FArchive& Ar, int32 InOper
 		check(Serializer==NULL);
 		Object->Serialize( Ar );
 	}
+
+	Ar.ArIgnoreOuterRef = bWasArIgnoreOuterRef;
 }
 
 void FTransaction::FObjectRecord::Restore( FTransaction* Owner )
@@ -169,15 +179,18 @@ void FTransaction::DumpObjectMap(FOutputDevice& Ar) const
 
 FArchive& operator<<( FArchive& Ar, FTransaction::FObjectRecord& R )
 {
-	UObject* Object = R.Object.Get();
-	check(Object);
-	FMemMark Mark(FMemStack::Get());
-	Ar << Object;
-	R.Object = Object;
-	Ar << R.Data;
-	Ar << R.ReferencedObjects;
-	Ar << R.ReferencedNames;
-	Mark.Pop();
+	if (!Ar.IsObjectReferenceCollector() || (Ar.IsObjectReferenceCollector() && !R.Object.IsPartOfCDO()))
+	{
+		UObject* Object = R.Object.Get();
+		check(Object);
+		FMemMark Mark(FMemStack::Get());
+		Ar << Object;
+		R.Object = Object;
+		Ar << R.Data;
+		Ar << R.ReferencedObjects;
+		Ar << R.ReferencedNames;
+		Mark.Pop();
+	}
 	return Ar;
 }
 
diff --git a/Engine/Source/Runtime/Engine/Private/UnrealEngine.cpp b/Engine/Source/Runtime/Engine/Private/UnrealEngine.cpp
--- a/Engine/Source/Runtime/Engine/Private/UnrealEngine.cpp
+++ b/Engine/Source/Runtime/Engine/Private/UnrealEngine.cpp
@@ -107,6 +107,8 @@
 
 #include "GameFramework/OnlineSession.h"
 
+#include "InstancedReferenceSubobjectHelper.h"
+
 DEFINE_LOG_CATEGORY(LogEngine);
 
 IMPLEMENT_MODULE( FEngineModule, Engine );
@@ -10736,109 +10738,6 @@ static TAutoConsoleVariable<int32> CVarDumpCopyPropertiesForUnrelatedObjects(
 	TEXT("Dump the objects that are cross class copied")
 	);
 
-struct FFindInstancedReferenceSubobjectHelper
-{
-	template<typename T>
-	static void Get(const UStruct* Struct, const uint8* ContainerAddress, T& OutObjects)
-	{
-		check(ContainerAddress);
-		for (UProperty* Prop = (Struct ? Struct->RefLink : NULL); Prop; Prop = Prop->NextRef)
-		{
-			auto ArrayProperty = Cast<const UArrayProperty>(Prop);
-			if (ArrayProperty && Prop->HasAnyPropertyFlags(CPF_PersistentInstance | CPF_ContainsInstancedReference))
-			{
-				auto InnerStructProperty = Cast<const UStructProperty>(ArrayProperty->Inner);
-				auto InnerObjectProperty = Cast<const UObjectProperty>(ArrayProperty->Inner);
-				if (InnerStructProperty && InnerStructProperty->Struct)
-				{
-					FScriptArrayHelper_InContainer ArrayHelper(ArrayProperty, ContainerAddress);
-					for (int32 ElementIndex = 0; ElementIndex < ArrayHelper.Num(); ++ElementIndex)
-					{
-						const uint8* ValueAddress = ArrayHelper.GetRawPtr(ElementIndex);
-						if (ValueAddress)
-						{
-							Get(InnerStructProperty->Struct, ValueAddress, OutObjects);
-						}
-					}
-				}
-				else if (InnerObjectProperty && InnerObjectProperty->HasAllPropertyFlags(CPF_PersistentInstance))
-				{
-					ensure(InnerObjectProperty->HasAllPropertyFlags(CPF_InstancedReference));
-					FScriptArrayHelper_InContainer ArrayHelper(ArrayProperty, ContainerAddress);
-					for (int32 ElementIndex = 0; ElementIndex < ArrayHelper.Num(); ++ElementIndex)
-					{
-						UObject* ObjectValue = InnerObjectProperty->GetObjectPropertyValue(ArrayHelper.GetRawPtr(ElementIndex));
-						if (ObjectValue)
-						{
-							OutObjects.Add(ObjectValue);
-						}
-					}
-				}
-			}
-			else if (Prop->HasAllPropertyFlags(CPF_PersistentInstance))
-			{
-				ensure(Prop->HasAllPropertyFlags(CPF_InstancedReference));
-				auto ObjectProperty = Cast<const UObjectProperty>(Prop);
-				if (ObjectProperty)
-				{
-					for (int32 ArrayIdx = 0; ArrayIdx < Prop->ArrayDim; ++ArrayIdx)
-					{
-						auto ObjectValue = ObjectProperty->GetObjectPropertyValue_InContainer(ContainerAddress, ArrayIdx);
-						if (ObjectValue)
-						{
-							OutObjects.Add(ObjectValue);
-						}
-					}
-				}
-			}
-			else if (Prop->HasAnyPropertyFlags(CPF_ContainsInstancedReference))
-			{
-				auto StructProperty = Cast<const UStructProperty>(Prop);
-				if (StructProperty && StructProperty->Struct)
-				{
-					for (int32 ArrayIdx = 0; ArrayIdx < Prop->ArrayDim; ++ArrayIdx)
-					{
-						const uint8* ValueAddress = StructProperty->ContainerPtrToValuePtr<uint8>(ContainerAddress, ArrayIdx);
-						Get(StructProperty->Struct, ValueAddress, OutObjects);
-					}
-				}
-			}
-		}
-	}
-
-	static void Duplicate(UObject* OldObject, UObject* NewObject, TMap<UObject*, UObject*>& ReferenceReplacementMap, TArray<UObject*>& DuplicatedObjects)
-	{
-		if (OldObject->GetClass()->HasAnyClassFlags(CLASS_HasInstancedReference) &&
-			NewObject->GetClass()->HasAnyClassFlags(CLASS_HasInstancedReference))
-		{
-			TSet<UObject*> OldEditInlineObjects;
-			Get(OldObject->GetClass(), reinterpret_cast<uint8*>(OldObject), OldEditInlineObjects);
-			if (OldEditInlineObjects.Num())
-			{
-				TSet<UObject*> NewEditInlineObjects;
-				Get(NewObject->GetClass(), reinterpret_cast<uint8*>(NewObject), NewEditInlineObjects);
-				for (auto Obj : NewEditInlineObjects)
-				{
-					const bool bProperOuter = (Obj->GetOuter() == OldObject);
-					const bool bEditInlineNew = Obj->GetClass()->HasAnyClassFlags(CLASS_EditInlineNew | CLASS_DefaultToInstanced);
-					if (bProperOuter && bEditInlineNew)
-					{
-						const bool bKeptByOld = OldEditInlineObjects.Contains(Obj);
-						const bool bNotHandledYet = !ReferenceReplacementMap.Contains(Obj);
-						if (bKeptByOld && bNotHandledYet)
-						{
-							UObject* NewEditInlineSubobject = StaticDuplicateObject(Obj, NewObject, NULL);
-							ReferenceReplacementMap.Add(Obj, NewEditInlineSubobject);
-
-							// We also need to make sure to fixup any properties here
-							DuplicatedObjects.Add(NewEditInlineSubobject);
-						}
-					}
-				}
-			}
-		}
-	}
-};
 
 void UEngine::CopyPropertiesForUnrelatedObjects(UObject* OldObject, UObject* NewObject, FCopyPropertiesForUnrelatedObjectsParams Params)
 {
diff --git a/Engine/Source/Runtime/Engine/Public/InstancedReferenceSubobjectHelper.h b/Engine/Source/Runtime/Engine/Public/InstancedReferenceSubobjectHelper.h
new file mode 100644
--- /dev/null
+++ b/Engine/Source/Runtime/Engine/Public/InstancedReferenceSubobjectHelper.h
@@ -0,0 +1,111 @@
+// Copyright 1998-2015 Epic Games, Inc. All Rights Reserved.
+
+#pragma once
+
+class FFindInstancedReferenceSubobjectHelper
+{
+public:
+	template<typename T>
+	static void Get(const UStruct* Struct, const uint8* ContainerAddress, T& OutObjects)
+	{
+		check(ContainerAddress);
+		for (UProperty* Prop = (Struct ? Struct->RefLink : NULL); Prop; Prop = Prop->NextRef)
+		{
+			auto ArrayProperty = Cast<const UArrayProperty>(Prop);
+			if (ArrayProperty && Prop->HasAnyPropertyFlags(CPF_PersistentInstance | CPF_ContainsInstancedReference))
+			{
+				auto InnerStructProperty = Cast<const UStructProperty>(ArrayProperty->Inner);
+				auto InnerObjectProperty = Cast<const UObjectProperty>(ArrayProperty->Inner);
+				if (InnerStructProperty && InnerStructProperty->Struct)
+				{
+					FScriptArrayHelper_InContainer ArrayHelper(ArrayProperty, ContainerAddress);
+					for (int32 ElementIndex = 0; ElementIndex < ArrayHelper.Num(); ++ElementIndex)
+					{
+						const uint8* ValueAddress = ArrayHelper.GetRawPtr(ElementIndex);
+						if (ValueAddress)
+						{
+							Get(InnerStructProperty->Struct, ValueAddress, OutObjects);
+						}
+					}
+				}
+				else if (InnerObjectProperty && InnerObjectProperty->HasAllPropertyFlags(CPF_PersistentInstance))
+				{
+					ensure(InnerObjectProperty->HasAllPropertyFlags(CPF_InstancedReference));
+					FScriptArrayHelper_InContainer ArrayHelper(ArrayProperty, ContainerAddress);
+					for (int32 ElementIndex = 0; ElementIndex < ArrayHelper.Num(); ++ElementIndex)
+					{
+						UObject* ObjectValue = InnerObjectProperty->GetObjectPropertyValue(ArrayHelper.GetRawPtr(ElementIndex));
+						if (ObjectValue)
+						{
+							OutObjects.Add(ObjectValue);
+						}
+					}
+				}
+			}
+			else if (Prop->HasAllPropertyFlags(CPF_PersistentInstance))
+			{
+				ensure(Prop->HasAllPropertyFlags(CPF_InstancedReference));
+				auto ObjectProperty = Cast<const UObjectProperty>(Prop);
+				if (ObjectProperty)
+				{
+					for (int32 ArrayIdx = 0; ArrayIdx < Prop->ArrayDim; ++ArrayIdx)
+					{
+						auto ObjectValue = ObjectProperty->GetObjectPropertyValue_InContainer(ContainerAddress, ArrayIdx);
+						if (ObjectValue)
+						{
+							OutObjects.Add(ObjectValue);
+						}
+					}
+				}
+			}
+			else if (Prop->HasAnyPropertyFlags(CPF_ContainsInstancedReference))
+			{
+				auto StructProperty = Cast<const UStructProperty>(Prop);
+				if (StructProperty && StructProperty->Struct)
+				{
+					for (int32 ArrayIdx = 0; ArrayIdx < Prop->ArrayDim; ++ArrayIdx)
+					{
+						const uint8* ValueAddress = StructProperty->ContainerPtrToValuePtr<uint8>(ContainerAddress, ArrayIdx);
+						Get(StructProperty->Struct, ValueAddress, OutObjects);
+					}
+				}
+			}
+		}
+	}
+
+	static void Duplicate(UObject* OldObject, UObject* NewObject, TMap<UObject*, UObject*>& ReferenceReplacementMap, TArray<UObject*>& DuplicatedObjects)
+	{
+		if (OldObject->GetClass()->HasAnyClassFlags(CLASS_HasInstancedReference) &&
+			NewObject->GetClass()->HasAnyClassFlags(CLASS_HasInstancedReference))
+		{
+			TSet<UObject*> OldEditInlineObjects;
+			Get(OldObject->GetClass(), reinterpret_cast<uint8*>(OldObject), OldEditInlineObjects);
+			if (OldEditInlineObjects.Num())
+			{
+				TSet<UObject*> NewEditInlineObjects;
+				Get(NewObject->GetClass(), reinterpret_cast<uint8*>(NewObject), NewEditInlineObjects);
+				for (auto Obj : NewEditInlineObjects)
+				{
+					const bool bProperOuter = (Obj->GetOuter() == OldObject);
+					const bool bEditInlineNew = Obj->GetClass()->HasAnyClassFlags(CLASS_EditInlineNew | CLASS_DefaultToInstanced);
+					if (bProperOuter && bEditInlineNew)
+					{
+						const bool bKeptByOld = OldEditInlineObjects.Contains(Obj);
+						const bool bNotHandledYet = !ReferenceReplacementMap.Contains(Obj);
+						if (bKeptByOld && bNotHandledYet)
+						{
+							UObject* NewEditInlineSubobject = StaticDuplicateObject(Obj, NewObject, NULL);
+							ReferenceReplacementMap.Add(Obj, NewEditInlineSubobject);
+
+							TArray<UObject*> OutDefaultOuters;
+							GetObjectsWithOuter(NewEditInlineSubobject, OutDefaultOuters, false);
+
+							// We also need to make sure to fixup any properties here
+							DuplicatedObjects.Add(NewEditInlineSubobject);
+						}
+					}
+				}
+			}
+		}
+	}
+};
\ No newline at end of file
