#ANIM: Consolidate the post physics blend

diff --git a/Engine/Source/Runtime/Engine/Private/Components/SkeletalMeshComponent.cpp b/Engine/Source/Runtime/Engine/Private/Components/SkeletalMeshComponent.cpp
--- a/Engine/Source/Runtime/Engine/Private/Components/SkeletalMeshComponent.cpp
+++ b/Engine/Source/Runtime/Engine/Private/Components/SkeletalMeshComponent.cpp
@@ -1235,23 +1235,9 @@ void USkeletalMeshComponent::PostAnimEvaluation(FAnimationEvaluationContext& Eva
 	UpdateKinematicBonesToAnim(GetEditableSpaceBases(), ETeleportType::None, true);
 	UpdateRBJointMotors();
 
-	// @todo anim : hack TTP 224385	ANIM: Skeletalmesh double buffer
-	// this is problem because intermediate buffer changes physics position as well
-	// this causes issue where a half of frame, physics position is fixed with anim pose, and the other half is real simulated position
-	// if you enable physics in tick, since that's before physics update, you'll get animation pose dominating physics pose, which isn't what you want. (Or what you'll see)
-	// so do not update transform if physics is on. This problem will be solved by double buffer, when we keep one buffer for intermediate, and the other buffer for result query
 	if (!ShouldBlendPhysicsBones())
 	{
-		SCOPE_CYCLE_COUNTER(STAT_UpdateLocalToWorldAndOverlaps);
-
-		// Updated last good bone positions
-		FinalizeBoneTransform();
-
-		// New bone positions need to be sent to render thread
-		UpdateComponentToWorld();
-
-		// animation often change overlap. 
-		UpdateOverlaps();
+		PostBlendPhysics();
 	}
 
 
diff --git a/Engine/Source/Runtime/Engine/Private/PhysicsEngine/PhysAnim.cpp b/Engine/Source/Runtime/Engine/Private/PhysicsEngine/PhysAnim.cpp
--- a/Engine/Source/Runtime/Engine/Private/PhysicsEngine/PhysAnim.cpp
+++ b/Engine/Source/Runtime/Engine/Private/PhysicsEngine/PhysAnim.cpp
@@ -420,6 +420,8 @@ void USkeletalMeshComponent::BlendInPhysics(FTickFunction& ThisTickFunction)
 
 void USkeletalMeshComponent::PostBlendPhysics()
 {
+	SCOPE_CYCLE_COUNTER(STAT_UpdateLocalToWorldAndOverlaps);
+	
 	FinalizeBoneTransform();
 
 	// Update Child Transform - The above function changes bone transform, so will need to update child transform
@@ -428,6 +430,9 @@ void USkeletalMeshComponent::PostBlendPhysics()
 	// animation often change overlap. 
 	UpdateOverlaps();
 
+	// update bounds
+	UpdateBounds();
+
 	// New bone positions need to be sent to render thread
 	MarkRenderDynamicDataDirty();
 }
