Use correct net driver for component PreReplication. Previously, during demo recording, the component pre-replication code would either:
- crash, if the world's NetDriver was null, or
- use the game net driver for component pre-replication instead of the demo net driver

diff --git a/Engine/Source/Runtime/Engine/Classes/GameFramework/Actor.h b/Engine/Source/Runtime/Engine/Classes/GameFramework/Actor.h
--- a/Engine/Source/Runtime/Engine/Classes/GameFramework/Actor.h
+++ b/Engine/Source/Runtime/Engine/Classes/GameFramework/Actor.h
@@ -13,6 +13,7 @@
 struct FHitResult;
 class AActor;
 class FTimerManager; 
+class UNetDriver;
 
 #include "Actor.generated.h"
 
@@ -320,6 +321,9 @@ public:
 	/** Called on the actor right before replication occurs */
 	virtual void PreReplication( IRepChangedPropertyTracker & ChangedPropertyTracker );
 
+	/** Called by the networking system to call PreReplication on this actor and its components using the given NetDriver to find or create RepChangedPropertyTrackers. */
+	void CallPreReplication(UNetDriver* NetDriver);
+
 	/** If true then destroy self when "finished", meaning all relevant components report that they are done and no timelines or timers are in flight. */
 	UPROPERTY(BlueprintReadWrite, Category=Actor)
 	uint32 bAutoDestroyWhenFinished:1;
diff --git a/Engine/Source/Runtime/Engine/Private/Actor.cpp b/Engine/Source/Runtime/Engine/Private/Actor.cpp
--- a/Engine/Source/Runtime/Engine/Private/Actor.cpp
+++ b/Engine/Source/Runtime/Engine/Private/Actor.cpp
@@ -753,6 +753,16 @@ void AActor::PreReplication( IRepChangedPropertyTracker & ChangedPropertyTracker
 	{
 		BPClass->InstancePreReplication(ChangedPropertyTracker);
 	}
+}
+
+void AActor::CallPreReplication(UNetDriver* NetDriver)
+{
+	if (NetDriver == nullptr)
+	{
+		return;
+	}
+
+	PreReplication(*NetDriver->FindOrCreateRepChangedPropertyTracker(this).Get());
 
 	// Call PreReplication on all owned components that are replicated
 	for (UActorComponent* Component : OwnedComponents)
@@ -760,7 +770,7 @@ void AActor::PreReplication( IRepChangedPropertyTracker & ChangedPropertyTracker
 		// Only call on components that aren't pending kill
 		if (Component && !Component->IsPendingKill() && Component->GetIsReplicated())
 		{
-			Component->PreReplication(*GetNetDriver()->FindOrCreateRepChangedPropertyTracker(Component).Get());
+			Component->PreReplication(*NetDriver->FindOrCreateRepChangedPropertyTracker(Component).Get());
 		}
 	}
 }
diff --git a/Engine/Source/Runtime/Engine/Private/DemoNetDriver.cpp b/Engine/Source/Runtime/Engine/Private/DemoNetDriver.cpp
--- a/Engine/Source/Runtime/Engine/Private/DemoNetDriver.cpp
+++ b/Engine/Source/Runtime/Engine/Private/DemoNetDriver.cpp
@@ -884,7 +884,7 @@ void UDemoNetDriver::SaveCheckpoint()
 	{
 		if ( CheckpointConnection->ActorChannels.Contains( Actor ) )
 		{
-			Actor->PreReplication( *FindOrCreateRepChangedPropertyTracker( Actor ).Get() );
+			Actor->CallPreReplication( this );
 			DemoReplicateActor( Actor, CheckpointConnection, false );
 		}
 	}
@@ -998,7 +998,7 @@ void UDemoNetDriver::TickDemoRecord( float DeltaSeconds )
 			continue;
 		}
 
-		Actor->PreReplication( *FindOrCreateRepChangedPropertyTracker( Actor ).Get() );
+		Actor->CallPreReplication( this );
 		DemoReplicateActor( Actor, ClientConnections[0], IsNetClient );
 	}
 
diff --git a/Engine/Source/Runtime/Engine/Private/NetworkDriver.cpp b/Engine/Source/Runtime/Engine/Private/NetworkDriver.cpp
--- a/Engine/Source/Runtime/Engine/Private/NetworkDriver.cpp
+++ b/Engine/Source/Runtime/Engine/Private/NetworkDriver.cpp
@@ -2065,7 +2065,7 @@ int32 UNetDriver::ServerReplicateActors(float DeltaSeconds)
 
 				if ( bWasConsidered )
 				{
-					Actor->PreReplication( *FindOrCreateRepChangedPropertyTracker( Actor ).Get() );
+					Actor->CallPreReplication( this );
 				}
 			}
 			/*
