[INTEGRATE] Fixed Print String messages not visible on HMD devices

diff --git a/Engine/Source/Runtime/Engine/Private/UnrealEngine.cpp b/Engine/Source/Runtime/Engine/Private/UnrealEngine.cpp
--- a/Engine/Source/Runtime/Engine/Private/UnrealEngine.cpp
+++ b/Engine/Source/Runtime/Engine/Private/UnrealEngine.cpp
@@ -7298,7 +7298,7 @@ void DrawStatsHUD( UWorld* World, FViewport* Viewport, FCanvas* Canvas, UCanvas*
 #if !(UE_BUILD_SHIPPING || UE_BUILD_TEST)
 	if( !GIsHighResScreenshot && !GIsDumpingMovie && GAreScreenMessagesEnabled )
 	{
-		const int32 MessageX = 40;
+		const int32 MessageX = ( GEngine->IsStereoscopic3D( Viewport ) ) ? Viewport->GetSizeXY().X * 0.5f * 0.3f : 40;
 
 		if (!GEngine->bSuppressMapWarnings)
 		{
