The DemoNetDriver gets the spectator controller class from the default game mode instead of the instance.
Allows live game clients to spawn a spectator controller even though they don't have a game mode object replicated to them.

diff --git a/Engine/Source/Runtime/Engine/Private/DemoNetDriver.cpp b/Engine/Source/Runtime/Engine/Private/DemoNetDriver.cpp
--- a/Engine/Source/Runtime/Engine/Private/DemoNetDriver.cpp
+++ b/Engine/Source/Runtime/Engine/Private/DemoNetDriver.cpp
@@ -1448,7 +1448,12 @@ void UDemoNetDriver::SpawnDemoRecSpectator( UNetConnection* Connection )
 {
 	check( Connection != NULL );
 
-	UClass* C = GetWorld()->GetAuthGameMode()->ReplaySpectatorPlayerControllerClass;
+	// Get the replay spectator controller class from the default game mode object,
+	// since the game mode instance isn't replicated to clients of live games.
+	AGameState* GameState = GetWorld() != nullptr ? GetWorld()->GetGameState() : nullptr;
+	TSubclassOf<AGameMode> DefaultGameModeClass = GameState != nullptr ? GameState->GameModeClass : nullptr;
+	AGameMode* DefaultGameMode = CastChecked<AGameMode>(DefaultGameModeClass.GetDefaultObject());
+	UClass* C = DefaultGameMode != nullptr ? DefaultGameMode->ReplaySpectatorPlayerControllerClass : nullptr;
 
 	if ( C == NULL )
 	{
