refactoring, should be sample result

diff --git a/Engine/Source/Runtime/Engine/Private/Rendering/SeparableSSS.cpp b/Engine/Source/Runtime/Engine/Private/Rendering/SeparableSSS.cpp
--- a/Engine/Source/Runtime/Engine/Private/Rendering/SeparableSSS.cpp
+++ b/Engine/Source/Runtime/Engine/Private/Rendering/SeparableSSS.cpp
@@ -151,9 +151,9 @@ void ComputeMirroredSSSKernel(FLinearColor* TargetBuffer, uint32 TargetBufferSiz
 
 		// Tweak them using the desired strength. The first one is:
 		//     lerp(1.0, kernel[0].rgb, strength)
-		kernel[0].R = (1.0f - SubsurfaceColor.R) * 1.0f + SubsurfaceColor.R * kernel[0].R;
-		kernel[0].G = (1.0f - SubsurfaceColor.G) * 1.0f + SubsurfaceColor.G * kernel[0].G;
-		kernel[0].B = (1.0f - SubsurfaceColor.B) * 1.0f + SubsurfaceColor.B * kernel[0].B;
+		kernel[0].R = FMath::Lerp(1.0f, kernel[0].R, SubsurfaceColor.R);
+		kernel[0].G = FMath::Lerp(1.0f, kernel[0].G, SubsurfaceColor.G);
+		kernel[0].B = FMath::Lerp(1.0f, kernel[0].B, SubsurfaceColor.B);
 
 		for (int i = 1; i < nTotalSamples; i++)
 		{
