removed unnecessary cached world pointer from navigation data
#ue4

diff --git a/Engine/Source/Runtime/Engine/Classes/AI/Navigation/NavigationData.h b/Engine/Source/Runtime/Engine/Classes/AI/Navigation/NavigationData.h
--- a/Engine/Source/Runtime/Engine/Classes/AI/Navigation/NavigationData.h
+++ b/Engine/Source/Runtime/Engine/Classes/AI/Navigation/NavigationData.h
@@ -439,9 +439,6 @@ protected:
 	UPROPERTY(EditAnywhere, Category = Runtime, config)
 	float ObservedPathsTickInterval;
 
-	UPROPERTY(Transient)
-	UWorld* CachedWorld;
-
 public:
 	//----------------------------------------------------------------------//
 	// Life cycle                                                                
diff --git a/Engine/Source/Runtime/Engine/Private/AI/Navigation/NavigationData.cpp b/Engine/Source/Runtime/Engine/Private/AI/Navigation/NavigationData.cpp
--- a/Engine/Source/Runtime/Engine/Private/AI/Navigation/NavigationData.cpp
+++ b/Engine/Source/Runtime/Engine/Private/AI/Navigation/NavigationData.cpp
@@ -182,10 +182,11 @@ void ANavigationData::PostInitProperties()
 void ANavigationData::PostInitializeComponents()
 {
 	Super::PostInitializeComponents();
-	
-	CachedWorld = GetWorld();
 
-	if (CachedWorld == NULL || CachedWorld->GetNavigationSystem() == NULL)
+	UWorld* MyWorld = GetWorld();
+	UNavigationSystem* NavSys = UNavigationSystem::GetCurrent(MyWorld);
+
+	if (NavSys == nullptr || MyWorld == nullptr)
 	{
 		CleanUpAndMarkPendingKill();
 	}
@@ -193,7 +194,7 @@ void ANavigationData::PostInitializeComponents()
 	{
 		// note: this is not a final fix for world composition's issues with navmesh generation
 		// but it's good for now, and navmesh creation is going to get a face-lift soon anyway
-		bWantsUpdate |= CachedWorld->GetWorldSettings()->bEnableWorldComposition;
+		bWantsUpdate |= MyWorld->GetWorldSettings()->bEnableWorldComposition;
 	}
 }
 
@@ -209,8 +210,6 @@ void ANavigationData::PostLoad()
 
 	InstantiateAndRegisterRenderingComponent();
 
-	CachedWorld = GetWorld();
-
 	bNetLoadOnClient = (*GEngine->NavigationSystemClass != nullptr) && (GEngine->NavigationSystemClass->GetDefaultObject<UNavigationSystem>()->ShouldLoadNavigationOnClient(this));
 }
 
