compilation fix for NavigationData
#ue4

diff --git a/Engine/Source/Runtime/Engine/Classes/AI/Navigation/NavigationData.h b/Engine/Source/Runtime/Engine/Classes/AI/Navigation/NavigationData.h
--- a/Engine/Source/Runtime/Engine/Classes/AI/Navigation/NavigationData.h
+++ b/Engine/Source/Runtime/Engine/Classes/AI/Navigation/NavigationData.h
@@ -724,7 +724,6 @@ public:
 	 *	@param Location is expressed in WorldSpace, navigation data is responsible for tansforming if need be */
 	virtual bool DoesNodeContainLocation(NavNodeRef NodeRef, const FVector& WorldSpaceLocation) const PURE_VIRTUAL(ANavigationData::DoesNodeContainLocation, return false;);
 
-	const UWorld* GetCachedWorld() const;
 	float GetWorldTimeStamp() const;
 
 	//----------------------------------------------------------------------//
diff --git a/Engine/Source/Runtime/Engine/Private/AI/Navigation/NavigationData.cpp b/Engine/Source/Runtime/Engine/Private/AI/Navigation/NavigationData.cpp
--- a/Engine/Source/Runtime/Engine/Private/AI/Navigation/NavigationData.cpp
+++ b/Engine/Source/Runtime/Engine/Private/AI/Navigation/NavigationData.cpp
@@ -558,14 +558,9 @@ void ANavigationData::DrawDebugPath(FNavigationPath* Path, FColor PathColor, UCa
 	Path->DebugDraw(this, PathColor, Canvas, bPersistent, NextPathPointIndex);
 }
 
-const UWorld* ANavigationData::GetCachedWorld() const
-{
-	return CachedWorld;
-}
-
 float ANavigationData::GetWorldTimeStamp() const
 {
-	const UWorld* World = GetCachedWorld();
+	const UWorld* World = GetWorld();
 	return World ? World->GetTimeSeconds() : 0.f;
 }
 
