- Added check to make sure BeginPay is not called again on actor components (this might happen if subobjects are renamed).

Merging

//depot/UE4-Orion/Engine/Source/Runtime/Engine/Private/Actor.cpp

to //depot/UE4/Engine/Source/Runtime/Engine/Private/Actor.cpp

diff --git a/Engine/Source/Runtime/Engine/Private/Actor.cpp b/Engine/Source/Runtime/Engine/Private/Actor.cpp
--- a/Engine/Source/Runtime/Engine/Private/Actor.cpp
+++ b/Engine/Source/Runtime/Engine/Private/Actor.cpp
@@ -2722,7 +2722,8 @@ void AActor::BeginPlay()
 
 	for (UActorComponent* Component : Components)
 	{
-		if (Component->IsRegistered())
+		// bHasBegunPlay will be true for the component if the component was renamed and moved to a new outer during initialization
+		if (Component->IsRegistered() && !Component->HasBegunPlay())
 		{
 			Component->BeginPlay();
 		}
