PR #1393: Added ENGINE_API to export FNotifyBufferFinishedHooks to allow custom sound nodes that rely on sounds finishing like the USoundNodeLooping node. (Contributed by [REDACTED_PERSON])

diff --git a/Engine/Source/Runtime/Engine/Public/Audio.h b/Engine/Source/Runtime/Engine/Public/Audio.h
--- a/Engine/Source/Runtime/Engine/Public/Audio.h
+++ b/Engine/Source/Runtime/Engine/Public/Audio.h
@@ -116,7 +116,7 @@ enum ELoopingMode
 	LOOP_Forever
 };
 
-struct FNotifyBufferFinishedHooks
+struct ENGINE_API FNotifyBufferFinishedHooks
 {
 	void AddNotify(USoundNode* NotifyNode, UPTRINT WaveInstanceHash);
 	UPTRINT GetHashForNode(USoundNode* NotifyNode) const;
