[UE-17070] Revised fix for lost child components in a Blueprint Class asset with a default scene root when reparenting to a class that yields an inherited scene root.

change summary:
- restored the change in UBlueprint::GetBlueprintHierarchyFromClass() to look at UBlueprint::ParentClass before UClass::GetSuperClass(), in case the BPGC has not yet been compiled after a reparenting operation.
- restored/improved the change in USimpleConstructionScript::ValidateSceneRootNodes() to ensure that any promoted child nodes are properly parented to the appropriate native or inherited root if the local default scene root is removed in favor of an existing root component, inherited or otherwise.

diff --git a/Engine/Source/Runtime/Engine/Private/Blueprint.cpp b/Engine/Source/Runtime/Engine/Private/Blueprint.cpp
--- a/Engine/Source/Runtime/Engine/Private/Blueprint.cpp
+++ b/Engine/Source/Runtime/Engine/Private/Blueprint.cpp
@@ -941,7 +941,15 @@ bool UBlueprint::GetBlueprintHierarchyFromClass(const UClass* InClass, TArray<UB
 		bNoErrors &= (BP->Status != BS_Error);
 #endif // #if WITH_EDITORONLY_DATA
 
-		CurrentClass = CurrentClass->GetSuperClass();
+		// If valid, use stored ParentClass rather than the actual UClass::GetSuperClass(); handles the case when the class has not been recompiled yet after a reparent operation.
+		if(const UClass* ParentClass = BP->ParentClass)
+		{
+			CurrentClass = ParentClass;
+		}
+		else
+		{
+			CurrentClass = CurrentClass->GetSuperClass();
+		}
 	}
 
 	return bNoErrors;
diff --git a/Engine/Source/Runtime/Engine/Private/SimpleConstructionScript.cpp b/Engine/Source/Runtime/Engine/Private/SimpleConstructionScript.cpp
--- a/Engine/Source/Runtime/Engine/Private/SimpleConstructionScript.cpp
+++ b/Engine/Source/Runtime/Engine/Private/SimpleConstructionScript.cpp
@@ -926,7 +926,66 @@ void USimpleConstructionScript::ValidateSceneRootNodes()
 		else if(RootComponentTemplate != nullptr
 			&& RootNodes.Contains(DefaultSceneRootNode))
 		{
+			// If the default scene root has any child nodes, determine what they should parent to.
+			USCS_Node* RootNode = nullptr;
+			bool bIsParentComponentNative = false;
+			FName ParentComponentOrVariableName = NAME_None;
+			FName ParentComponentOwnerClassName = NAME_None;
+			if(UBlueprintGeneratedClass* BPClass = Cast<UBlueprintGeneratedClass>(RootComponentTemplate->GetOuter()))
+			{
+				// The root scene component is an SCS node.
+				if(BPClass->SimpleConstructionScript != nullptr)
+				{
+					const TArray<USCS_Node*> SCSRootNodes = BPClass->SimpleConstructionScript->GetRootNodes();
+					for(USCS_Node* SCSNode : SCSRootNodes)
+					{
+						if(SCSNode != nullptr && SCSNode->ComponentTemplate == RootComponentTemplate)
+						{
+							if(BPClass->SimpleConstructionScript != this)
+							{
+								// The root node is inherited from a parent BP class.
+								ParentComponentOwnerClassName = BPClass->GetFName();
+								ParentComponentOrVariableName = SCSNode->VariableName;
+							}
+							else
+							{
+								// The root node belongs to the current BP class.
+								RootNode = SCSNode;
+							}
+							
+							break;
+						}
+					}
+				}
+			}
+			else
+			{
+				// The root scene component is a native component.
+				bIsParentComponentNative = true;
+				ParentComponentOrVariableName = RootComponentTemplate->GetFName();
+			}
+
+			// Reparent any child nodes within the current hierarchy.
+			for(USCS_Node* ChildNode : DefaultSceneRootNode->ChildNodes)
+			{
+				if(RootNode != nullptr)
+				{
+					// We have an existing root node within the current BP class.
+					RootNode->AddChildNode(ChildNode);
+				}
+				else
+				{
+					// The current root node is inherited from a parent class (may be BP or native).
+					RootNodes.Add(ChildNode);
+					ChildNode->bIsParentComponentNative = bIsParentComponentNative;
+					ChildNode->ParentComponentOrVariableName = ParentComponentOrVariableName;
+					ChildNode->ParentComponentOwnerClassName = ParentComponentOwnerClassName;
+				}
+			}
+
+			// Remove the default scene root node from the current hierarchy.
 			RootNodes.Remove(DefaultSceneRootNode);
+			DefaultSceneRootNode->ChildNodes.Empty();
 
 			// These shouldn't be set, but just in case...
 			DefaultSceneRootNode->bIsParentComponentNative = false;
