#UE4 Removed needless UPROPERTY. This property is handled entirely natively.

diff --git a/Engine/Source/Runtime/Engine/Classes/Engine/LevelStreaming.h b/Engine/Source/Runtime/Engine/Classes/Engine/LevelStreaming.h
--- a/Engine/Source/Runtime/Engine/Classes/Engine/LevelStreaming.h
+++ b/Engine/Source/Runtime/Engine/Classes/Engine/LevelStreaming.h
@@ -342,7 +342,6 @@ private:
 	class ULevel* PendingUnloadLevel;
 
 	/** The cached package name of the world asset that is loaded by the levelstreaming */
-	UPROPERTY(Transient)
 	FName CachedWorldAssetPackageFName;
 	
 	/** Friend classes to allow access to SetLoadedLevel */
