Display a graceful message for failure to load the default map when travelling instead of crashing via a fatal log message

diff --git a/Engine/Source/Runtime/Engine/Private/UnrealEngine.cpp b/Engine/Source/Runtime/Engine/Private/UnrealEngine.cpp
--- a/Engine/Source/Runtime/Engine/Private/UnrealEngine.cpp
+++ b/Engine/Source/Runtime/Engine/Private/UnrealEngine.cpp
@@ -8707,10 +8707,13 @@ EBrowseReturnVal::Type UEngine::Browse( FWorldContext& WorldContext, FURL URL, F
 		}
 		
 		const UGameMapsSettings* GameMapsSettings = GetDefault<UGameMapsSettings>();
-		bool LoadSucces = LoadMap(WorldContext, FURL(&URL, *(GameMapsSettings->GetGameDefaultMap() + GameMapsSettings->LocalMapOptions), TRAVEL_Partial), NULL, Error);
-		if (LoadSucces==false)
+		if (!LoadMap(WorldContext, FURL(&URL, *(GameMapsSettings->GetGameDefaultMap() + GameMapsSettings->LocalMapOptions), TRAVEL_Partial), NULL, Error))
 		{
-			UE_LOG(LogNet, Fatal, TEXT("Failed to load default map (%s). Error: (%s)"), *(GameMapsSettings->GetGameDefaultMap() + GameMapsSettings->LocalMapOptions), *Error);
+			UE_LOG(LogNet, Error, TEXT("Failed to load default map (%s). Error: (%s)"), *(GameMapsSettings->GetGameDefaultMap() + GameMapsSettings->LocalMapOptions), *Error);
+			const FText Message = FText::Format(NSLOCTEXT("Engine", "FailedToLoadDefaultMap", "Error '{1}'. Exiting."), FText::FromString(Error));
+			FMessageDialog::Open(EAppMsgType::Ok, Message);
+			FPlatformMisc::RequestExit(false);
+			return EBrowseReturnVal::Failure;
 		}
 
 		CollectGarbage( GARBAGE_COLLECTION_KEEPFLAGS );
@@ -8718,7 +8721,7 @@ EBrowseReturnVal::Type UEngine::Browse( FWorldContext& WorldContext, FURL URL, F
 		// now remove "failed" and "closed" options from LastURL so it doesn't get copied on to future URLs
 		WorldContext.LastURL.RemoveOption(TEXT("failed"));
 		WorldContext.LastURL.RemoveOption(TEXT("closed"));
-		return (LoadSucces ? EBrowseReturnVal::Success : EBrowseReturnVal::Failure);
+		return EBrowseReturnVal::Success;
 	}
 	else if( URL.HasOption(TEXT("restart")) )
 	{
