Fix for deprecation warnings.

diff --git a/Engine/Source/Runtime/Engine/Classes/Engine/World.h b/Engine/Source/Runtime/Engine/Classes/Engine/World.h
--- a/Engine/Source/Runtime/Engine/Classes/Engine/World.h
+++ b/Engine/Source/Runtime/Engine/Classes/Engine/World.h
@@ -419,23 +419,7 @@ struct FEndClothSimulationFunction : public FTickFunction
 /* Struct of optional parameters passed to SpawnActor function(s). */
 struct ENGINE_API FActorSpawnParameters
 {
-	FActorSpawnParameters()
-		:	Name(NAME_None)
-		,	Template(NULL)
-		,	Owner(NULL)
-		,	Instigator(NULL)
-		,	OverrideLevel(NULL)
-		,	SpawnCollisionHandlingOverride(ESpawnActorCollisionHandlingMethod::Undefined)
-		,	bRemoteOwned(false)
-		,	bNoFail(false)
-		,	bDeferConstruction(false)		
-		,	bAllowDuringConstructionScript(false)
-		,	ObjectFlags(RF_Transactional)
-	{
-PRAGMA_DISABLE_DEPRECATION_WARNINGS
-		bNoCollisionFail = false;
-PRAGMA_ENABLE_DEPRECATION_WARNINGS
-	}
+	FActorSpawnParameters();
 
 	/* A name to assign as the Name of the Actor being spawned. If no value is specified, the name of the spawned Actor will be automatically generated using the form [Class]_[Number]. */
 	FName Name;
@@ -475,6 +459,28 @@ PRAGMA_ENABLE_DEPRECATION_WARNINGS
 	EObjectFlags ObjectFlags;		
 };
 
+// Deprecation warnings disabled to initialize bNoCollisionFail
+PRAGMA_DISABLE_DEPRECATION_WARNINGS
+
+	inline FActorSpawnParameters::FActorSpawnParameters()
+		: Name(NAME_None)
+		, Template(NULL)
+		, Owner(NULL)
+		, Instigator(NULL)
+		, OverrideLevel(NULL)
+		, SpawnCollisionHandlingOverride(ESpawnActorCollisionHandlingMethod::Undefined)
+		, bNoCollisionFail(false)
+		, bRemoteOwned(false)
+		, bNoFail(false)
+		, bDeferConstruction(false)
+		, bAllowDuringConstructionScript(false)
+		, ObjectFlags(RF_Transactional)
+	{
+	}
+
+PRAGMA_ENABLE_DEPRECATION_WARNINGS
+
+
 /**
  *  This encapsulate World's async trace functionality. This contains two buffers of trace data buffer and alternates it for each tick. 
  *  
