introduce nodebug pvd. Helps with profiling

diff --git a/Engine/Source/Runtime/Engine/Private/PhysicsEngine/PhysLevel.cpp b/Engine/Source/Runtime/Engine/Private/PhysicsEngine/PhysLevel.cpp
--- a/Engine/Source/Runtime/Engine/Private/PhysicsEngine/PhysLevel.cpp
+++ b/Engine/Source/Runtime/Engine/Private/PhysicsEngine/PhysLevel.cpp
@@ -306,7 +306,7 @@ FString FEndClothSimulationFunction::DiagnosticMessage()
 }
 
 
-void PvdConnect(FString Host);
+void PvdConnect(FString Host, bool bVisualization);
 
 //////// GAME-LEVEL RIGID BODY PHYSICS STUFF ///////
 void InitGamePhys()
@@ -369,7 +369,7 @@ void InitGamePhys()
 
 	if( FParse::Param( FCommandLine::Get(), TEXT( "PVD" ) ) )
 	{
-		PvdConnect(TEXT("localhost"));
+		PvdConnect(TEXT("localhost"), true);
 	}
 
 
diff --git a/Engine/Source/Runtime/Engine/Private/PhysicsEngine/PhysUtils.cpp b/Engine/Source/Runtime/Engine/Private/PhysicsEngine/PhysUtils.cpp
--- a/Engine/Source/Runtime/Engine/Private/PhysicsEngine/PhysUtils.cpp
+++ b/Engine/Source/Runtime/Engine/Private/PhysicsEngine/PhysUtils.cpp
@@ -446,18 +446,23 @@ static bool ExecApexVis(UWorld* InWorld, uint32 SceneType, const TCHAR* Cmd, FOu
 }
 
 #if WITH_PHYSX
-void PvdConnect(FString Host)
+void PvdConnect(FString Host, bool bVisualization)
 {
 	int32		Port = 5425;         // TCP port to connect to, where PVD is listening
 	uint32	Timeout = 100;          // timeout in milliseconds to wait for PVD to respond, consoles and remote PCs need a higher timeout.
+	
 	PxVisualDebuggerConnectionFlags ConnectionFlags = PxVisualDebuggerExt::getAllConnectionFlags();
+	if (!bVisualization)
+	{
+		ConnectionFlags &= ~PxVisualDebuggerConnectionFlag::eDEBUG;
+	}
 
 	// and now try to connect
 	PxVisualDebuggerExt::createConnection(GPhysXSDK->getPvdConnectionManager(), TCHAR_TO_ANSI(*Host), Port, Timeout, ConnectionFlags);
 
-	GPhysXSDK->getVisualDebugger()->setVisualDebuggerFlag(PxVisualDebuggerFlag::eTRANSMIT_CONTACTS, true);
-	GPhysXSDK->getVisualDebugger()->setVisualDebuggerFlag(PxVisualDebuggerFlag::eTRANSMIT_SCENEQUERIES, true);
-	GPhysXSDK->getVisualDebugger()->setVisualDebuggerFlag(PxVisualDebuggerFlag::eTRANSMIT_CONSTRAINTS, true);
+	GPhysXSDK->getVisualDebugger()->setVisualDebuggerFlag(PxVisualDebuggerFlag::eTRANSMIT_CONTACTS, bVisualization);
+	GPhysXSDK->getVisualDebugger()->setVisualDebuggerFlag(PxVisualDebuggerFlag::eTRANSMIT_SCENEQUERIES, bVisualization);
+	GPhysXSDK->getVisualDebugger()->setVisualDebuggerFlag(PxVisualDebuggerFlag::eTRANSMIT_CONSTRAINTS, bVisualization);
 }
 #endif
 
@@ -557,7 +562,9 @@ bool ExecPhysCommands(const TCHAR* Cmd, FOutputDevice* Ar, UWorld* InWorld)
 					Host = Cmd;
 				}
 
-				PvdConnect(Host);
+				const bool bVizualization = !FParse::Command(&Cmd, TEXT("NODEBUG"));
+
+				PvdConnect(Host, bVizualization);
 
 			}
 			else if(FParse::Command(&Cmd, TEXT("DISCONNECT")))
