Removing redundant flush in FDebugCanvasDrawer::DrawRenderThread(), result of a merge error.

diff --git a/Engine/Source/Runtime/Engine/Private/Slate/DebugCanvas.cpp b/Engine/Source/Runtime/Engine/Private/Slate/DebugCanvas.cpp
--- a/Engine/Source/Runtime/Engine/Private/Slate/DebugCanvas.cpp
+++ b/Engine/Source/Runtime/Engine/Private/Slate/DebugCanvas.cpp
@@ -181,7 +181,6 @@ void FDebugCanvasDrawer::DrawRenderThread(FRHICommandListImmediate& RHICmdList,
 		else
 		{
 			RenderThreadCanvas->Canvas.SetRenderTargetRect( RenderTarget->GetViewRect() );
-			RenderThreadCanvas->Canvas.Flush_RenderThread(RHICmdList, true);
 		}
 		RenderThreadCanvas->Canvas.Flush_RenderThread(RHICmdList, true);
 		RenderThreadCanvas->Canvas.SetAllowSwitchVerticalAxis(bNeedToFlipVertical);
