Niagara: adding mesh effect renderer; few additional bug fixes

diff --git a/Engine/Source/Editor/NiagaraEditor/Private/NiagaraEffectEditor.cpp b/Engine/Source/Editor/NiagaraEditor/Private/NiagaraEffectEditor.cpp
--- a/Engine/Source/Editor/NiagaraEditor/Private/NiagaraEffectEditor.cpp
+++ b/Engine/Source/Editor/NiagaraEditor/Private/NiagaraEffectEditor.cpp
@@ -91,6 +91,7 @@ void FNiagaraEffectEditor::InitNiagaraEffectEditor(const EToolkitMode::Type Mode
 		MovieScene->InTime = InTime;
 		MovieScene->OutTime = OutTime;
 		MovieScene->EndTime = OutTime;
+		NewAnimation->MovieScene = MovieScene;
 
 		FSequencerViewParams ViewParams(TEXT("NiagaraSequencerSettings"));
 		{
diff --git a/Engine/Source/Editor/NiagaraEditor/Private/SNiagaraEffectEditorWidget.cpp b/Engine/Source/Editor/NiagaraEditor/Private/SNiagaraEffectEditorWidget.cpp
--- a/Engine/Source/Editor/NiagaraEditor/Private/SNiagaraEffectEditorWidget.cpp
+++ b/Engine/Source/Editor/NiagaraEditor/Private/SNiagaraEffectEditorWidget.cpp
@@ -675,8 +675,12 @@ void SEmitterWidget::OnRenderModuleChanged(TSharedPtr<FString> ModName, ESelectI
 		InEmitter->GetProperties()->RenderModuleType = InType;
 		InEffect->RenderModuleupdate();
 	});
-	TComponentReregisterContext<UNiagaraComponent> ComponentReregisterContext;
+	FlushRenderingCommands();
+
+	UMaterial *Material = UMaterial::GetDefaultMaterial(MD_Surface);
+	Emitter->GetEffectRenderer()->SetMaterial(Material, EffectInstance->GetComponent()->GetWorld()->FeatureLevel);
 
+	TComponentReregisterContext<UNiagaraComponent> ComponentReregisterContext;
 	UObject *Props = Emitter->GetProperties()->RendererProperties;
 	RendererPropertiesView->SetObject(Props);
 }
diff --git a/Engine/Source/Runtime/Engine/Classes/Engine/NiagaraEffectRenderer.h b/Engine/Source/Runtime/Engine/Classes/Engine/NiagaraEffectRenderer.h
--- a/Engine/Source/Runtime/Engine/Classes/Engine/NiagaraEffectRenderer.h
+++ b/Engine/Source/Runtime/Engine/Classes/Engine/NiagaraEffectRenderer.h
@@ -13,9 +13,11 @@ NiagaraEffectRenderer.h: Base class for Niagara render modules
 #include "NiagaraEffectRendererProperties.h"
 #include "NiagaraSpriteRendererProperties.h"
 #include "NiagaraRibbonRendererProperties.h"
+#include "NiagaraMeshRendererProperties.h"
 
 DECLARE_CYCLE_STAT(TEXT("Generate Sprite Vertex Data"), STAT_NiagaraGenSpriteVertexData, STATGROUP_Niagara);
 DECLARE_CYCLE_STAT(TEXT("Generate Ribbon Vertex Data"), STAT_NiagaraGenRibbonVertexData, STATGROUP_Niagara);
+DECLARE_CYCLE_STAT(TEXT("Generate Mesh Vertex Data"), STAT_NiagaraGenMeshVertexData, STATGROUP_Niagara);
 
 
 /** Struct used to pass dynamic data from game thread to render thread */
@@ -100,15 +102,17 @@ public:
 	virtual ~NiagaraEffectRenderer() {}
 
 	virtual UClass *GetPropertiesClass() = 0;
+	virtual void SetRendererProperties(UNiagaraEffectRendererProperties *Props) = 0;
 
 	float GetCPUTimeMS() { return CPUTimeMS; }
 	FBox GetBounds()	{ return CachedBounds;  }
 
 protected:
 	NiagaraEffectRenderer()	
-		: Material(nullptr), TickState(NTS_Running)
+		: TickState(NTS_Running)
 
 	{
+		Material = UMaterial::GetDefaultMaterial(MD_Surface);
 	}
 
 	mutable float CPUTimeMS;
@@ -152,6 +156,7 @@ public:
 	bool HasDynamicData() override;
 
 	UClass *GetPropertiesClass() override { return UNiagaraSpriteRendererProperties::StaticClass(); }
+	void SetRendererProperties(UNiagaraEffectRendererProperties *Props) override { Properties = Cast<UNiagaraSpriteRendererProperties>(Props); }
 
 	virtual const TArray<FNiagaraVariableInfo>& GetRequiredAttributes() override;
 
@@ -217,6 +222,7 @@ public:
 	virtual const TArray<FNiagaraVariableInfo>& GetRequiredAttributes() override;
 
 	UClass *GetPropertiesClass() override { return UNiagaraRibbonRendererProperties::StaticClass(); }
+	void SetRendererProperties(UNiagaraEffectRendererProperties *Props) override { Properties = Cast<UNiagaraRibbonRendererProperties>(Props); }
 
 private:
 	class FParticleBeamTrailVertexFactory *VertexFactory;
@@ -226,3 +232,46 @@ private:
 };
 
 
+
+
+
+
+/**
+* NiagaraEffectRendererSprites renders an FNiagaraSimulation as sprite particles
+*/
+class ENGINE_API NiagaraEffectRendererMeshes : public NiagaraEffectRenderer
+{
+public:
+
+	explicit NiagaraEffectRendererMeshes(ERHIFeatureLevel::Type FeatureLevel, UNiagaraEffectRendererProperties *Props);
+	~NiagaraEffectRendererMeshes()
+	{
+		ReleaseRenderThreadResources();
+	}
+
+
+	virtual void ReleaseRenderThreadResources() override;
+
+	// FPrimitiveSceneProxy interface.
+	virtual void CreateRenderThreadResources() override;
+
+	virtual void GetDynamicMeshElements(const TArray<const FSceneView*>& Views, const FSceneViewFamily& ViewFamily, uint32 VisibilityMap, FMeshElementCollector& Collector, const FNiagaraSceneProxy *SceneProxy) const override;
+	virtual bool SetMaterialUsage() override;
+	/** Update render data buffer from attributes */
+	FNiagaraDynamicDataBase *GenerateVertexData(const FNiagaraEmitterParticleData &Data) override;
+
+	virtual void SetDynamicData_RenderThread(FNiagaraDynamicDataBase* NewDynamicData) override;
+	int GetDynamicDataSize() override;
+	bool HasDynamicData() override;
+
+	UClass *GetPropertiesClass() override { return UNiagaraMeshRendererProperties::StaticClass(); }
+	void SetRendererProperties(UNiagaraEffectRendererProperties *Props) override { Properties = Cast<UNiagaraMeshRendererProperties>(Props); }
+	virtual const TArray<FNiagaraVariableInfo>& GetRequiredAttributes() override;
+
+	void SetupVertexFactory(FMeshParticleVertexFactory *InVertexFactory, const FStaticMeshLODResources& LODResources) const;
+private:
+	UNiagaraMeshRendererProperties *Properties;
+	struct FNiagaraDynamicDataMesh *DynamicDataRender;
+	mutable TUniformBuffer<FPrimitiveUniformShaderParameters> WorldSpacePrimitiveUniformBuffer;
+	class FMeshParticleVertexFactory* VertexFactory;
+};
diff --git a/Engine/Source/Runtime/Engine/Classes/Engine/NiagaraMeshRendererProperties.h b/Engine/Source/Runtime/Engine/Classes/Engine/NiagaraMeshRendererProperties.h
new file mode 100644
--- /dev/null
+++ b/Engine/Source/Runtime/Engine/Classes/Engine/NiagaraMeshRendererProperties.h
@@ -0,0 +1,37 @@
+// Copyright 1998-2015 Epic Games, Inc. All Rights Reserved.
+
+#pragma once
+
+#include "NiagaraEffectRendererProperties.h"
+#include "StaticMeshResources.h"
+#include "NiagaraMeshRendererProperties.generated.h"
+
+UCLASS()
+class UNiagaraMeshRendererProperties : public UNiagaraEffectRendererProperties
+{
+public:
+	GENERATED_UCLASS_BODY()
+
+	UNiagaraMeshRendererProperties()
+	{
+		ParticleMesh = nullptr;
+	}
+
+	void PostEditChangeProperty(FPropertyChangedEvent& PropertyChangedEvent)
+	{
+		if (ParticleMesh && PropertyChangedEvent.Property->GetName() == "ParticleMesh")
+		{
+			const FStaticMeshLODResources& LODModel = ParticleMesh->RenderData->LODResources[0];
+			for (int32 SectionIndex = 0; SectionIndex < LODModel.Sections.Num(); SectionIndex++)
+			{
+				const FStaticMeshSection& Section = LODModel.Sections[SectionIndex];
+				UMaterialInterface *Material = ParticleMesh->GetMaterial(Section.MaterialIndex);
+				FMaterialRenderProxy* MaterialProxy = Material->GetRenderProxy(false, false);
+				Material->CheckMaterialUsage(MATUSAGE_MeshParticles);
+			}
+		}
+	}
+
+	UPROPERTY(EditAnywhere, Category = "Mesh Rendering")
+	UStaticMesh *ParticleMesh;
+};
diff --git a/Engine/Source/Runtime/Engine/Private/NiagaraEffectRenderer.cpp b/Engine/Source/Runtime/Engine/Private/NiagaraEffectRenderer.cpp
--- a/Engine/Source/Runtime/Engine/Private/NiagaraEffectRenderer.cpp
+++ b/Engine/Source/Runtime/Engine/Private/NiagaraEffectRenderer.cpp
@@ -4,18 +4,27 @@
 #include "EnginePrivate.h"
 #include "Particles/ParticleResources.h"
 #include "Engine/NiagaraEffectRenderer.h"
+#include "StaticMeshResources.h"
 
 DECLARE_CYCLE_STAT(TEXT("Render Total"), STAT_NiagaraRender, STATGROUP_Niagara);
 DECLARE_CYCLE_STAT(TEXT("Render Sprites"), STAT_NiagaraRenderSprites, STATGROUP_Niagara);
 DECLARE_CYCLE_STAT(TEXT("Render Ribbons"), STAT_NiagaraRenderRibbons, STATGROUP_Niagara);
+DECLARE_CYCLE_STAT(TEXT("Render Meshes"), STAT_NiagaraRenderMeshes, STATGROUP_Niagara);
 
 UNiagaraEffectRendererProperties::UNiagaraEffectRendererProperties(FObjectInitializer const &Initializer)
+	:Super(Initializer)
 {}
 
 UNiagaraSpriteRendererProperties::UNiagaraSpriteRendererProperties(FObjectInitializer const &Initializer)
+	:Super(Initializer)
 {}
 
 UNiagaraRibbonRendererProperties::UNiagaraRibbonRendererProperties(FObjectInitializer const &Initializer)
+	:Super(Initializer)
+{}
+
+UNiagaraMeshRendererProperties::UNiagaraMeshRendererProperties(FObjectInitializer const &Initializer)
+	:Super(Initializer)
 {}
 
 
@@ -30,6 +39,10 @@ struct FNiagaraDynamicDataRibbon : public FNiagaraDynamicDataBase
 	TArray<FParticleBeamTrailVertex> VertexData;
 };
 
+struct FNiagaraDynamicDataMesh : public FNiagaraDynamicDataBase
+{
+	TArray<FMeshParticleInstanceVertex> VertexData;
+};
 
 
 /* Mesh collector classes */
@@ -59,6 +72,18 @@ public:
 };
 
 
+class FNiagaraMeshCollectorResourcesMesh : public FOneFrameResource
+{
+public:
+	FMeshParticleVertexFactory VertexFactory;
+	FMeshParticleUniformBufferRef UniformBuffer;
+
+	virtual ~FNiagaraMeshCollectorResourcesMesh()
+	{
+		VertexFactory.ReleaseResource();
+	}
+};
+
 
 
 
@@ -205,7 +230,7 @@ void NiagaraEffectRendererSprites::GetDynamicMeshElements(const TArray<const FSc
 bool NiagaraEffectRendererSprites::SetMaterialUsage()
 {
 	//Causes deadlock :S Need to look at / rework the setting of materials and render modules.
-	return Material && Material->CheckMaterialUsage(MATUSAGE_ParticleSprites);
+	return Material && Material->CheckMaterialUsage_Concurrent(MATUSAGE_ParticleSprites);
 }
 
 /** Update render data buffer from attributes */
@@ -574,4 +599,468 @@ FNiagaraDynamicDataBase *NiagaraEffectRendererRibbon::GenerateVertexData(const F
 	CPUTimeMS = VertexDataTimer.GetElapsedMilliseconds();
 
 	return DynamicData;
-}
\ No newline at end of file
+}
+
+
+
+
+
+
+
+
+
+
+
+NiagaraEffectRendererMeshes::NiagaraEffectRendererMeshes(ERHIFeatureLevel::Type FeatureLevel, UNiagaraEffectRendererProperties *InProps) :
+NiagaraEffectRenderer(),
+DynamicDataRender(NULL)
+{
+	//check(InProps);
+	VertexFactory = new FMeshParticleVertexFactory(PVFT_Mesh, FeatureLevel, sizeof(FMeshParticleInstanceVertex), 0);
+	Properties = Cast<UNiagaraMeshRendererProperties>(InProps);
+
+
+	if (Properties && Properties->ParticleMesh)
+	{
+		const FStaticMeshLODResources& LODModel = Properties->ParticleMesh->RenderData->LODResources[0];
+		for (int32 SectionIndex = 0; SectionIndex < LODModel.Sections.Num(); SectionIndex++)
+		{
+			//FMaterialRenderProxy* MaterialProxy = MeshMaterials[SectionIndex]->GetRenderProxy(bSelected);
+			const FStaticMeshSection& Section = LODModel.Sections[SectionIndex];
+			UMaterialInterface *Material = Properties->ParticleMesh->GetMaterial(Section.MaterialIndex);
+			Material->CheckMaterialUsage_Concurrent(MATUSAGE_MeshParticles);
+		}
+	}
+
+}
+
+void NiagaraEffectRendererMeshes::SetupVertexFactory(FMeshParticleVertexFactory *InVertexFactory, const FStaticMeshLODResources& LODResources) const
+{
+	FMeshParticleVertexFactory::DataType Data;
+
+	Data.PositionComponent = FVertexStreamComponent(
+		&LODResources.PositionVertexBuffer,
+		STRUCT_OFFSET(FPositionVertex, Position),
+		LODResources.PositionVertexBuffer.GetStride(),
+		VET_Float3
+		);
+
+	Data.TangentBasisComponents[0] = FVertexStreamComponent(
+		&LODResources.VertexBuffer,
+		STRUCT_OFFSET(FStaticMeshFullVertex, TangentX),
+		LODResources.VertexBuffer.GetStride(),
+		VET_PackedNormal
+		);
+
+	Data.TangentBasisComponents[1] = FVertexStreamComponent(
+		&LODResources.VertexBuffer,
+		STRUCT_OFFSET(FStaticMeshFullVertex, TangentZ),
+		LODResources.VertexBuffer.GetStride(),
+		VET_PackedNormal
+		);
+
+	Data.TextureCoordinates.Empty();
+	if (!LODResources.VertexBuffer.GetUseFullPrecisionUVs())
+	{
+		uint32 NumTexCoords = FMath::Min<uint32>(LODResources.VertexBuffer.GetNumTexCoords(), MAX_TEXCOORDS);
+		for (uint32 UVIndex = 0; UVIndex < NumTexCoords; UVIndex++)
+		{
+			Data.TextureCoordinates.Add(FVertexStreamComponent(
+				&LODResources.VertexBuffer,
+				STRUCT_OFFSET(TStaticMeshFullVertexFloat16UVs<MAX_TEXCOORDS>, UVs) + sizeof(FVector2DHalf)* UVIndex,
+				LODResources.VertexBuffer.GetStride(),
+				VET_Half2
+				));
+		}
+	}
+	else
+	{
+		for (uint32 UVIndex = 0; UVIndex < LODResources.VertexBuffer.GetNumTexCoords(); UVIndex++)
+		{
+			Data.TextureCoordinates.Add(FVertexStreamComponent(
+				&LODResources.VertexBuffer,
+				STRUCT_OFFSET(TStaticMeshFullVertexFloat32UVs<MAX_TEXCOORDS>, UVs) + sizeof(FVector2D)* UVIndex,
+				LODResources.VertexBuffer.GetStride(),
+				VET_Float2
+				));
+		}
+	}
+
+	if (LODResources.ColorVertexBuffer.GetNumVertices() > 0)
+	{
+		Data.VertexColorComponent = FVertexStreamComponent(
+			&LODResources.ColorVertexBuffer,
+			0,
+			LODResources.ColorVertexBuffer.GetStride(),
+			VET_Color
+			);
+	}
+
+
+	// Initialize instanced data. Vertex buffer and stride are set before render.
+	// Particle color
+	Data.ParticleColorComponent = FVertexStreamComponent(
+		NULL,
+		STRUCT_OFFSET(FMeshParticleInstanceVertex, Color),
+		0,
+		VET_Float4,
+		true
+		);
+
+	// Particle transform matrix
+	for (int32 MatrixRow = 0; MatrixRow < 3; MatrixRow++)
+	{
+		Data.TransformComponent[MatrixRow] = FVertexStreamComponent(
+			NULL,
+			STRUCT_OFFSET(FMeshParticleInstanceVertex, Transform) + sizeof(FVector4)* MatrixRow,
+			0,
+			VET_Float4,
+			true
+			);
+	}
+
+	Data.VelocityComponent = FVertexStreamComponent(
+		NULL,
+		STRUCT_OFFSET(FMeshParticleInstanceVertex, Velocity),
+		0,
+		VET_Float4,
+		true
+		);
+
+	// SubUVs.
+	Data.SubUVs = FVertexStreamComponent(
+		NULL,
+		STRUCT_OFFSET(FMeshParticleInstanceVertex, SubUVParams),
+		0,
+		VET_Short4,
+		true
+		);
+
+	// Pack SubUV Lerp and the particle's relative time
+	Data.SubUVLerpAndRelTime = FVertexStreamComponent(
+		NULL,
+		STRUCT_OFFSET(FMeshParticleInstanceVertex, SubUVLerp),
+		0,
+		VET_Float2,
+		true
+		);
+
+	Data.bInitialized = true;
+	InVertexFactory->SetData(Data);
+}
+
+
+
+void NiagaraEffectRendererMeshes::ReleaseRenderThreadResources()
+{
+	VertexFactory->ReleaseResource();
+	WorldSpacePrimitiveUniformBuffer.ReleaseResource();
+}
+
+void NiagaraEffectRendererMeshes::CreateRenderThreadResources()
+{
+	VertexFactory->InitResource();
+}
+
+
+
+
+void NiagaraEffectRendererMeshes::GetDynamicMeshElements(const TArray<const FSceneView*>& Views, const FSceneViewFamily& ViewFamily, uint32 VisibilityMap, FMeshElementCollector& Collector, const FNiagaraSceneProxy *SceneProxy) const
+{
+	SCOPE_CYCLE_COUNTER(STAT_NiagaraRender);
+	SCOPE_CYCLE_COUNTER(STAT_NiagaraRenderMeshes);
+
+	SimpleTimer MeshElementsTimer;
+
+	//check(DynamicDataRender)
+
+	if (!DynamicDataRender || DynamicDataRender->VertexData.Num() == 0)
+	{
+		return;
+	}
+
+	int32 SizeInBytes = DynamicDataRender->VertexData.GetTypeSize() * DynamicDataRender->VertexData.Num();
+	FGlobalDynamicVertexBuffer::FAllocation LocalDynamicVertexAllocation = FGlobalDynamicVertexBuffer::Get().Allocate(SizeInBytes);
+
+	if (LocalDynamicVertexAllocation.IsValid())
+	{
+		// Update the primitive uniform buffer if needed.
+		if (!WorldSpacePrimitiveUniformBuffer.IsInitialized())
+		{
+			FPrimitiveUniformShaderParameters PrimitiveUniformShaderParameters = GetPrimitiveUniformShaderParameters(
+				FMatrix::Identity,
+				SceneProxy->GetActorPosition(),
+				SceneProxy->GetBounds(),
+				SceneProxy->GetLocalBounds(),
+				SceneProxy->ReceivesDecals(),
+				false,
+				false,
+				SceneProxy->UseEditorDepthTest()
+				);
+			WorldSpacePrimitiveUniformBuffer.SetContents(PrimitiveUniformShaderParameters);
+			WorldSpacePrimitiveUniformBuffer.InitResource();
+		}
+
+		// Copy the vertex data over.
+		FMemory::Memcpy(LocalDynamicVertexAllocation.Buffer, DynamicDataRender->VertexData.GetData(), SizeInBytes);
+
+		// Compute the per-view uniform buffers.
+		for (int32 ViewIndex = 0; ViewIndex < Views.Num(); ViewIndex++)
+		{
+			if (VisibilityMap & (1 << ViewIndex))
+			{
+				const FSceneView* View = Views[ViewIndex];
+				const FStaticMeshLODResources& LODModel = Properties->ParticleMesh->RenderData->LODResources[0];
+
+				FNiagaraMeshCollectorResourcesMesh& CollectorResources = Collector.AllocateOneFrameResource<FNiagaraMeshCollectorResourcesMesh>();
+				SetupVertexFactory(&CollectorResources.VertexFactory, LODModel);
+				FMeshParticleUniformParameters PerViewUniformParameters;// = UniformParameters;
+				/*
+				if (Properties)
+				{
+					PerViewUniformParameters.SubImageSize = FVector4(Properties->SubImageInfo.X, Properties->SubImageInfo.Y, 1.0f / Properties->SubImageInfo.X, 1.0f / Properties->SubImageInfo.Y);
+				}
+				*/
+
+				// Collector.AllocateOneFrameResource uses default ctor, initialize the vertex factory
+				CollectorResources.VertexFactory.SetFeatureLevel(ViewFamily.GetFeatureLevel());
+				CollectorResources.VertexFactory.SetParticleFactoryType(PVFT_Mesh);
+				CollectorResources.UniformBuffer = FMeshParticleUniformBufferRef::CreateUniformBufferImmediate(PerViewUniformParameters, UniformBuffer_SingleFrame);
+				CollectorResources.VertexFactory.SetStrides(sizeof(FMeshParticleInstanceVertex), 0);
+
+				CollectorResources.VertexFactory.InitResource();
+				CollectorResources.VertexFactory.SetUniformBuffer(CollectorResources.UniformBuffer);
+				CollectorResources.VertexFactory.SetInstanceBuffer(
+					LocalDynamicVertexAllocation.VertexBuffer,
+					LocalDynamicVertexAllocation.VertexOffset,
+					sizeof(FMeshParticleInstanceVertex)
+					);
+				CollectorResources.VertexFactory.SetDynamicParameterBuffer(NULL, 0, 0);
+
+				const bool bIsWireframe = AllowDebugViewmodes() && View->Family->EngineShowFlags.Wireframe;
+
+				for (int32 SectionIndex = 0; SectionIndex < LODModel.Sections.Num(); SectionIndex++)
+				{
+					const FStaticMeshSection& Section = LODModel.Sections[SectionIndex];
+					UMaterialInterface *Material = Properties->ParticleMesh->GetMaterial(Section.MaterialIndex);
+					FMaterialRenderProxy* MaterialProxy = Material->GetRenderProxy(false, false);
+
+					if ((Section.NumTriangles == 0) || (MaterialProxy == NULL))
+					{
+						//@todo. This should never occur, but it does occasionally.
+						continue;
+					}
+
+					FMeshBatch& Mesh = Collector.AllocateMesh();
+					Mesh.VertexFactory = &CollectorResources.VertexFactory;
+					Mesh.DynamicVertexData = NULL;
+					Mesh.LCI = NULL;
+					Mesh.UseDynamicData = false;
+					Mesh.ReverseCulling = SceneProxy->IsLocalToWorldDeterminantNegative();
+					Mesh.CastShadow = SceneProxy->CastsDynamicShadow();
+					Mesh.DepthPriorityGroup = (ESceneDepthPriorityGroup)SceneProxy->GetDepthPriorityGroup(View);
+
+					FMeshBatchElement& BatchElement = Mesh.Elements[0];
+					BatchElement.PrimitiveUniformBufferResource = &WorldSpacePrimitiveUniformBuffer;
+					BatchElement.FirstIndex = Section.FirstIndex;
+					BatchElement.MinVertexIndex = Section.MinVertexIndex;
+					BatchElement.MaxVertexIndex = Section.MaxVertexIndex;
+					BatchElement.NumInstances = DynamicDataRender->VertexData.Num();
+
+					if (bIsWireframe)
+					{
+						if (LODModel.WireframeIndexBuffer.IsInitialized())
+						{
+							Mesh.Type = PT_LineList;
+							Mesh.MaterialRenderProxy = UMaterial::GetDefaultMaterial(MD_Surface)->GetRenderProxy(SceneProxy->IsSelected(), SceneProxy->IsHovered());
+							BatchElement.FirstIndex = 0;
+							BatchElement.IndexBuffer = &LODModel.WireframeIndexBuffer;
+							BatchElement.NumPrimitives = LODModel.WireframeIndexBuffer.GetNumIndices() / 2;
+
+						}
+						else
+						{
+							Mesh.Type = PT_TriangleList;
+							Mesh.MaterialRenderProxy = MaterialProxy;
+							Mesh.bWireframe = true;
+							BatchElement.FirstIndex = 0;
+							BatchElement.IndexBuffer = &LODModel.IndexBuffer;
+							BatchElement.NumPrimitives = LODModel.IndexBuffer.GetNumIndices() / 3;
+						}
+					}
+					else
+					{
+						Mesh.Type = PT_TriangleList;
+						Mesh.MaterialRenderProxy = MaterialProxy;
+						BatchElement.IndexBuffer = &LODModel.IndexBuffer;
+						BatchElement.FirstIndex = Section.FirstIndex;
+						BatchElement.NumPrimitives = Section.NumTriangles;
+					}
+
+					/*  TODO: make this work
+					if (!bInstanced)
+					{
+						FMeshParticleVertexFactory::FBatchParametersCPU& BatchParameters = Collector.AllocateOneFrameResource<FMeshParticleVertexFactory::FBatchParametersCPU>();
+						BatchParameters.InstanceBuffer = InstanceVerticesCPU->InstanceDataAllocationsCPU.GetData();
+						BatchParameters.DynamicParameterBuffer = InstanceVerticesCPU->DynamicParameterDataAllocationsCPU.GetData();
+						BatchElement.UserData = &BatchParameters;
+						BatchElement.UserIndex = 0;
+
+						Mesh.Elements.Reserve(ParticleCount);
+						for (int32 ParticleIndex = 1; ParticleIndex < ParticleCount; ++ParticleIndex)
+						{
+							FMeshBatchElement* NextElement = new(Mesh.Elements) FMeshBatchElement();
+							*NextElement = Mesh.Elements[0];
+							NextElement->UserIndex = ParticleIndex;
+						}
+					}*/
+
+					Mesh.bCanApplyViewModeOverrides = true;
+					Mesh.bUseWireframeSelectionColoring = SceneProxy->IsSelected();
+
+					Collector.AddMesh(ViewIndex, Mesh);
+				}
+			}
+		}
+	}
+
+	CPUTimeMS += MeshElementsTimer.GetElapsedMilliseconds();
+}
+
+
+
+bool NiagaraEffectRendererMeshes::SetMaterialUsage()
+{
+	//Causes deadlock :S Need to look at / rework the setting of materials and render modules.
+	return Material && Material->CheckMaterialUsage_Concurrent(MATUSAGE_MeshParticles);
+}
+
+
+
+/** Update render data buffer from attributes */
+FNiagaraDynamicDataBase *NiagaraEffectRendererMeshes::GenerateVertexData(const FNiagaraEmitterParticleData &Data)
+{
+	SCOPE_CYCLE_COUNTER(STAT_NiagaraGenMeshVertexData);
+
+	if (!Properties || Properties->ParticleMesh == nullptr)
+	{
+		return nullptr;
+	}
+
+
+	SimpleTimer VertexDataTimer;
+	FVector MaxSize;
+
+	FNiagaraDynamicDataMesh *DynamicData = new FNiagaraDynamicDataMesh;
+	TArray<FMeshParticleInstanceVertex>& RenderData = DynamicData->VertexData;
+
+	RenderData.Reset(Data.GetNumParticles());
+	CachedBounds.Init();
+
+	const FVector4 *PosPtr = Data.GetAttributeData(FNiagaraVariableInfo(FName(TEXT("Position")), ENiagaraDataType::Vector));
+	const FVector4 *VelPtr = Data.GetAttributeData(FNiagaraVariableInfo(FName(TEXT("Velocity")), ENiagaraDataType::Vector));
+	const FVector4 *ColPtr = Data.GetAttributeData(FNiagaraVariableInfo(FName(TEXT("Color")), ENiagaraDataType::Vector));
+	const FVector4 *AgePtr = Data.GetAttributeData(FNiagaraVariableInfo(FName(TEXT("Age")), ENiagaraDataType::Vector));
+	const FVector4 *RotPtr = Data.GetAttributeData(FNiagaraVariableInfo(FName(TEXT("Rotation")), ENiagaraDataType::Vector));
+	const FVector4 *XformPtr = Data.GetAttributeData(FNiagaraVariableInfo(FName(TEXT("Transform")), ENiagaraDataType::Vector));
+
+	//Bail if we don't have the required attributes to render this emitter.
+	if (!PosPtr || !ColPtr || !AgePtr || !RotPtr)
+	{
+		return DynamicData;
+	}
+
+	UStaticMesh *Mesh = nullptr;
+	if (Properties)
+	{
+		Mesh = Properties->ParticleMesh;
+	}
+	float ParticleId = 0.0f, IdInc = 1.0f / Data.GetNumParticles();
+	RenderData.AddUninitialized(Data.GetNumParticles());
+	for (uint32 ParticleIndex = 0; ParticleIndex < Data.GetNumParticles(); ParticleIndex++)
+	{
+		FMeshParticleInstanceVertex& NewVertex = RenderData[ParticleIndex];
+		NewVertex.Color = FLinearColor(ColPtr[ParticleIndex]);
+		NewVertex.Color.A = ColPtr[ParticleIndex].W;
+		//NewVertex.ParticleId = ParticleId;
+		ParticleId += IdInc;
+		NewVertex.RelativeTime = AgePtr[ParticleIndex].X;
+
+		FRotationMatrix RotMat(FRotator(0.0f, 0.0f, 0.0f));
+		float Scale = 1;
+		if (XformPtr)
+		{
+			FVector Euler(XformPtr[ParticleIndex]);
+			FRotator Rotator = FRotator::MakeFromEuler(Euler);
+			RotMat = Rotator;
+			Scale = XformPtr[ParticleIndex].W;
+		}
+
+		NewVertex.Transform[0].X = RotMat.GetColumn(0).X*Scale;  NewVertex.Transform[0].Y = RotMat.GetColumn(0).Y*Scale;   NewVertex.Transform[0].Z = RotMat.GetColumn(0).Z*Scale;
+		NewVertex.Transform[1].X = RotMat.GetColumn(1).X*Scale;  NewVertex.Transform[1].Y = RotMat.GetColumn(1).Y*Scale;   NewVertex.Transform[1].Z = RotMat.GetColumn(1).Z*Scale;
+		NewVertex.Transform[2].X = RotMat.GetColumn(2).X*Scale;  NewVertex.Transform[2].Y = RotMat.GetColumn(2).Y*Scale;   NewVertex.Transform[2].Z = RotMat.GetColumn(2).Z*Scale;
+		NewVertex.Transform[0].W = PosPtr[ParticleIndex].X;
+		NewVertex.Transform[1].W = PosPtr[ParticleIndex].Y;
+		NewVertex.Transform[2].W = PosPtr[ParticleIndex].Z;
+		NewVertex.Velocity = VelPtr[ParticleIndex];
+
+		FPlatformMisc::Prefetch(PosPtr + ParticleIndex + 1);
+		FPlatformMisc::Prefetch(XformPtr + ParticleIndex + 1);
+		FPlatformMisc::Prefetch(VelPtr + ParticleIndex + 1);
+		FPlatformMisc::Prefetch(AgePtr + ParticleIndex + 1);
+		MaxSize = MaxSize.ComponentMax(Mesh->GetBoundingBox().GetExtent());
+		CachedBounds += PosPtr[ParticleIndex];
+	}
+	CachedBounds.ExpandBy(MaxSize);
+
+	CPUTimeMS = VertexDataTimer.GetElapsedMilliseconds();
+
+
+	return DynamicData;
+}
+
+
+
+void NiagaraEffectRendererMeshes::SetDynamicData_RenderThread(FNiagaraDynamicDataBase* NewDynamicData)
+{
+	check(IsInRenderingThread());
+
+	if (DynamicDataRender)
+	{
+		delete DynamicDataRender;
+		DynamicDataRender = NULL;
+	}
+	DynamicDataRender = static_cast<FNiagaraDynamicDataMesh*>(NewDynamicData);
+}
+
+int NiagaraEffectRendererMeshes::GetDynamicDataSize()
+{
+	uint32 Size = sizeof(FNiagaraDynamicDataMesh);
+	if (DynamicDataRender)
+	{
+		Size += DynamicDataRender->VertexData.GetAllocatedSize();
+	}
+
+	return Size;
+}
+
+bool NiagaraEffectRendererMeshes::HasDynamicData()
+{
+	return DynamicDataRender && DynamicDataRender->VertexData.Num() > 0;
+}
+
+
+const TArray<FNiagaraVariableInfo>& NiagaraEffectRendererMeshes::GetRequiredAttributes()
+{
+	static TArray<FNiagaraVariableInfo> Attrs;
+
+	if (Attrs.Num() == 0)
+	{
+		Attrs.Add(FNiagaraVariableInfo(FName(TEXT("Position")), ENiagaraDataType::Vector));
+		Attrs.Add(FNiagaraVariableInfo(FName(TEXT("Color")), ENiagaraDataType::Vector));
+		//Attrs.Add(FNiagaraVariableInfo(FName(TEXT("Rotation")), ENiagaraDataType::Vector));
+		Attrs.Add(FNiagaraVariableInfo(FName(TEXT("Age")), ENiagaraDataType::Vector));
+	}
+
+	return Attrs;
+}
diff --git a/Engine/Source/Runtime/Niagara/NiagaraConstantSet.h b/Engine/Source/Runtime/Niagara/NiagaraConstantSet.h
--- a/Engine/Source/Runtime/Niagara/NiagaraConstantSet.h
+++ b/Engine/Source/Runtime/Niagara/NiagaraConstantSet.h
@@ -20,6 +20,9 @@ private:
 	TMap<FNiagaraVariableInfo, class UNiagaraDataObject*> DataConstants;
 
 public:
+	FNiagaraConstantMap()
+	{
+	}
 
 	void Empty();
 
diff --git a/Engine/Source/Runtime/Niagara/Private/NiagaraEffect.cpp b/Engine/Source/Runtime/Niagara/Private/NiagaraEffect.cpp
--- a/Engine/Source/Runtime/Niagara/Private/NiagaraEffect.cpp
+++ b/Engine/Source/Runtime/Niagara/Private/NiagaraEffect.cpp
@@ -40,6 +40,7 @@ void UNiagaraEffect::CreateEffectRendererProps(TSharedPtr<FNiagaraSimulation> Si
 	{
 		Sim->GetProperties()->RendererProperties = nullptr;
 	}
+	Sim->GetEffectRenderer()->SetRendererProperties(Sim->GetProperties()->RendererProperties);
 }
 
 void UNiagaraEffect::Serialize(FArchive& Ar)
diff --git a/Engine/Source/Runtime/Niagara/Private/NiagaraSimulation.cpp b/Engine/Source/Runtime/Niagara/Private/NiagaraSimulation.cpp
--- a/Engine/Source/Runtime/Niagara/Private/NiagaraSimulation.cpp
+++ b/Engine/Source/Runtime/Niagara/Private/NiagaraSimulation.cpp
@@ -336,12 +336,15 @@ void FNiagaraSimulation::SetRenderModuleType(EEmitterRenderModuleType Type, ERHI
 			break;
 		case RMT_Ribbon: EffectRenderer = new NiagaraEffectRendererRibbon(FeatureLevel, Props->RendererProperties);
 			break;
+		case RMT_Meshes: EffectRenderer = new NiagaraEffectRendererMeshes(FeatureLevel, Props->RendererProperties);
+			break;
 		default:	EffectRenderer = new NiagaraEffectRendererSprites(FeatureLevel, Props->RendererProperties);
 					Props->RenderModuleType = RMT_Sprites;
 					break;
 		}
 
 		EffectRenderer->SetMaterial(Material, FeatureLevel);
+
 		CheckAttriubtesForRenderer();
 	}
 }
diff --git a/Engine/Source/Runtime/Niagara/Public/NiagaraAnimation.h b/Engine/Source/Runtime/Niagara/Public/NiagaraAnimation.h
--- a/Engine/Source/Runtime/Niagara/Public/NiagaraAnimation.h
+++ b/Engine/Source/Runtime/Niagara/Public/NiagaraAnimation.h
@@ -16,7 +16,6 @@ class UNiagaraAnimation
 	GENERATED_UCLASS_BODY()
 
 public:
-
 	// UMovieSceneAnimation overrides
 
 	virtual bool AllowsSpawnableObjects() const override;
@@ -34,8 +33,6 @@ public:
 	virtual bool TryGetObjectDisplayName(const FGuid& ObjectId, FText& OutDisplayName) const override;
 #endif
 
-private:
-
 	/** Pointer to the movie scene that controls this animation. */
 	UPROPERTY()
 	UMovieScene* MovieScene;
diff --git a/Engine/Source/Runtime/VectorVM/Private/VectorVM.cpp b/Engine/Source/Runtime/VectorVM/Private/VectorVM.cpp
--- a/Engine/Source/Runtime/VectorVM/Private/VectorVM.cpp
+++ b/Engine/Source/Runtime/VectorVM/Private/VectorVM.cpp
@@ -64,7 +64,7 @@ UNiagaraDataObject::UNiagaraDataObject(const FObjectInitializer& ObjectInitializ
 }
 
 UNiagaraCurveDataObject::UNiagaraCurveDataObject(const FObjectInitializer& ObjectInitializer)
-: Super(ObjectInitializer)
+: Super(ObjectInitializer), CurveObj(nullptr)
 {
 }
 
diff --git a/Engine/Source/Runtime/VectorVM/Public/VectorVMDataObject.h b/Engine/Source/Runtime/VectorVM/Public/VectorVMDataObject.h
--- a/Engine/Source/Runtime/VectorVM/Public/VectorVMDataObject.h
+++ b/Engine/Source/Runtime/VectorVM/Public/VectorVMDataObject.h
@@ -25,6 +25,10 @@ class UNiagaraCurveDataObject : public UNiagaraDataObject
 	
 	UPROPERTY(EditAnywhere, Category="Curve")
 	class UCurveVector *CurveObj;
+public:
+	UNiagaraCurveDataObject(class UCurveVector *InCurve) : CurveObj(InCurve)
+	{
+	}
 
 	FVector4 Sample(const FVector4& InCoords) const;
 	virtual FVector4 Write(const FVector4& InCoords, const FVector4& InValue) override
@@ -72,4 +76,4 @@ public:
 		Data[MakeHash(InCoords)] = InValue;
 		return InValue;
 	}
-};
\ No newline at end of file
+};
