Make AttachChildren replicated so that it is consistent with AttachParent

diff --git a/Engine/Source/Runtime/Engine/Classes/Components/SceneComponent.h b/Engine/Source/Runtime/Engine/Classes/Components/SceneComponent.h
--- a/Engine/Source/Runtime/Engine/Classes/Components/SceneComponent.h
+++ b/Engine/Source/Runtime/Engine/Classes/Components/SceneComponent.h
@@ -108,7 +108,7 @@ public:
 	class USceneComponent* AttachParent;
 
 	/** List of child SceneComponents that are attached to us. */
-	UPROPERTY(transient)
+	UPROPERTY(Replicated, transient)
 	TArray< USceneComponent* > AttachChildren;
 
 	/** Optional socket name on AttachParent that we are attached to. */
