Calculate material billboard bounds using actual sizes. UE-4725

diff --git a/Engine/Source/Runtime/Engine/Private/Components/MaterialBillboardComponent.cpp b/Engine/Source/Runtime/Engine/Private/Components/MaterialBillboardComponent.cpp
--- a/Engine/Source/Runtime/Engine/Private/Components/MaterialBillboardComponent.cpp
+++ b/Engine/Source/Runtime/Engine/Private/Components/MaterialBillboardComponent.cpp
@@ -263,7 +263,13 @@ FPrimitiveSceneProxy* UMaterialBillboardComponent::CreateSceneProxy()
 
 FBoxSphereBounds UMaterialBillboardComponent::CalcBounds(const FTransform& LocalToWorld) const
 {
-	const float BoundsSize = 32.0f;
+	float BoundsSize = 1.0f;
+	for (int32 i = 0; i < Elements.Num(); ++i)
+	{
+		BoundsSize = FMath::Max3(BoundsSize, Elements[i].BaseSizeX, Elements[i].BaseSizeY);
+	}
+	BoundsSize *= LocalToWorld.GetMaximumAxisScale();
+
 	return FBoxSphereBounds(LocalToWorld.GetLocation(),FVector(BoundsSize,BoundsSize,BoundsSize),FMath::Sqrt(3.0f * FMath::Square(BoundsSize)));
 }
 
