merged pull request fixing HitNormal of navmesh raycasts
#ue4
#1397

diff --git a/Engine/Source/Runtime/Engine/Private/AI/Navigation/PImplRecastNavMesh.cpp b/Engine/Source/Runtime/Engine/Private/AI/Navigation/PImplRecastNavMesh.cpp
--- a/Engine/Source/Runtime/Engine/Private/AI/Navigation/PImplRecastNavMesh.cpp
+++ b/Engine/Source/Runtime/Engine/Private/AI/Navigation/PImplRecastNavMesh.cpp
@@ -745,10 +745,14 @@ void FPImplRecastNavMesh::Raycast2D(const FVector& StartLoc, const FVector& EndL
 
 	if (StartNode != INVALID_NAVNODEREF)
 	{
+		float RecastHitNormal[3];
+
 		const dtStatus RaycastStatus = NavQuery.raycast(StartNode, &RecastStart.X, &RecastEnd.X
-			, QueryFilter, &RaycastResult.HitTime, &RaycastResult.HitNormal.X
+			, QueryFilter, &RaycastResult.HitTime, RecastHitNormal
 			, RaycastResult.CorridorPolys, &RaycastResult.CorridorPolysCount, RaycastResult.GetMaxCorridorSize());
 
+		RaycastResult.HitNormal = Recast2UnrVector(RecastHitNormal);
+
 		if (dtStatusSucceed(RaycastStatus) == false)
 		{
 			UE_VLOG(NavMeshOwner, LogNavigation, Log, TEXT("FPImplRecastNavMesh::Raycast2D failed"));
@@ -784,9 +788,13 @@ void FPImplRecastNavMesh::Raycast2D(NavNodeRef StartNode, const FVector& StartLo
 
 	if (StartNode != INVALID_NAVNODEREF)
 	{
+		float RecastHitNormal[3];
+
 		const dtStatus RaycastStatus = NavQuery.raycast(StartNode, &RecastStart.X, &RecastEnd.X
-					, QueryFilter, &RaycastResult.HitTime, &RaycastResult.HitNormal.X
-					, RaycastResult.CorridorPolys, &RaycastResult.CorridorPolysCount, RaycastResult.GetMaxCorridorSize());
+			, QueryFilter, &RaycastResult.HitTime, RecastHitNormal
+			, RaycastResult.CorridorPolys, &RaycastResult.CorridorPolysCount, RaycastResult.GetMaxCorridorSize());
+
+		RaycastResult.HitNormal = Recast2UnrVector(RecastHitNormal);
 
 		if (dtStatusSucceed(RaycastStatus) == false)
 		{
diff --git a/Engine/Source/Runtime/Engine/Private/AI/Navigation/RecastNavMesh.cpp b/Engine/Source/Runtime/Engine/Private/AI/Navigation/RecastNavMesh.cpp
--- a/Engine/Source/Runtime/Engine/Private/AI/Navigation/RecastNavMesh.cpp
+++ b/Engine/Source/Runtime/Engine/Private/AI/Navigation/RecastNavMesh.cpp
@@ -1903,10 +1903,14 @@ void ARecastNavMesh::BatchRaycast(TArray<FNavigationRaycastWork>& Workload, TSha
 
 		if (StartNode != INVALID_NAVNODEREF)
 		{
+			float RecastHitNormal[3];
+
 			const dtStatus RaycastStatus = NavQuery.raycast(StartNode, &RecastStart.X, &RecastEnd.X
-				, QueryFilter, &RaycastResult.HitTime, &RaycastResult.HitNormal.X
+				, QueryFilter, &RaycastResult.HitTime, RecastHitNormal
 				, RaycastResult.CorridorPolys, &RaycastResult.CorridorPolysCount, RaycastResult.GetMaxCorridorSize());
 
+			RaycastResult.HitNormal = Recast2UnrVector(RecastHitNormal);
+
 			if (dtStatusSucceed(RaycastStatus) && RaycastResult.HasHit())
 			{
 				WorkItem.bDidHit = true;
