compilation fix for RecastNavMesh

diff --git a/Engine/Source/Runtime/Engine/Private/AI/Navigation/RecastNavMesh.cpp b/Engine/Source/Runtime/Engine/Private/AI/Navigation/RecastNavMesh.cpp
--- a/Engine/Source/Runtime/Engine/Private/AI/Navigation/RecastNavMesh.cpp
+++ b/Engine/Source/Runtime/Engine/Private/AI/Navigation/RecastNavMesh.cpp
@@ -1909,7 +1909,7 @@ void ARecastNavMesh::BatchRaycast(TArray<FNavigationRaycastWork>& Workload, TSha
 				, QueryFilter, &RaycastResult.HitTime, RecastHitNormal
 				, RaycastResult.CorridorPolys, &RaycastResult.CorridorPolysCount, RaycastResult.GetMaxCorridorSize());
 
-			RaycastResult.HitNormal = Recast2UnrVector(RecastHitNormal);
+			RaycastResult.HitNormal = Recast2UnrealPoint(RecastHitNormal);
 
 			if (dtStatusSucceed(RaycastStatus) && RaycastResult.HasHit())
 			{
