Add project setting to force CCD off. Can save about 0.5ms

diff --git a/Engine/Source/Runtime/Engine/Classes/PhysicsEngine/PhysicsSettings.h b/Engine/Source/Runtime/Engine/Classes/PhysicsEngine/PhysicsSettings.h
--- a/Engine/Source/Runtime/Engine/Classes/PhysicsEngine/PhysicsSettings.h
+++ b/Engine/Source/Runtime/Engine/Classes/PhysicsEngine/PhysicsSettings.h
@@ -172,6 +172,11 @@ class ENGINE_API UPhysicsSettings : public UDeveloperSettings
 	UPROPERTY(config, EditAnywhere, Category = Simulation)
 	bool bDefaultHasComplexCollision;
 
+	/**
+	*  If true CCD will be ignored. This is an optimization when CCD is never used which removes the need for physx to check it internally. */
+	UPROPERTY(config, EditAnywhere, Category = Simulation)
+	bool bDisableCCD;
+
 
 	/** Max Physics Delta Time to be clamped. */
 	UPROPERTY(config, EditAnywhere, meta=(ClampMin="0.0013", UIMin = "0.0013", ClampMax="1.0", UIMax="1.0"), Category=Framerate)
diff --git a/Engine/Source/Runtime/Engine/Private/PhysicsEngine/PhysScene.cpp b/Engine/Source/Runtime/Engine/Private/PhysicsEngine/PhysScene.cpp
--- a/Engine/Source/Runtime/Engine/Private/PhysicsEngine/PhysScene.cpp
+++ b/Engine/Source/Runtime/Engine/Private/PhysicsEngine/PhysScene.cpp
@@ -1372,7 +1372,7 @@ void FPhysScene::InitPhysScene(uint32 SceneType)
 	SCENE_LOCK_WRITE(PScene);
 
 	// enable CCD at scene level
-	if (bGlobalCCD)
+	if (UPhysicsSettings::Get()->bDisableCCD == false)
 	{
 		PScene->setFlag(PxSceneFlag::eENABLE_CCD, true);
 	}
diff --git a/Engine/Source/Runtime/Engine/Private/PhysicsEngine/PhysXSupport.h b/Engine/Source/Runtime/Engine/Private/PhysicsEngine/PhysXSupport.h
--- a/Engine/Source/Runtime/Engine/Private/PhysicsEngine/PhysXSupport.h
+++ b/Engine/Source/Runtime/Engine/Private/PhysicsEngine/PhysXSupport.h
@@ -54,8 +54,6 @@ private:
 const uint32 AggregateMaxSize	   = 128;
 const uint32 AggregateBodyShapesThreshold	   = 999999999;
 
-/** Global CCD Switch*/
-const bool bGlobalCCD = true;
 
 /////// UTILS
 
