Added SynchronousLoadType<> to the streamable manager to avoid having to cast on the other end.

diff --git a/Engine/Source/Runtime/Engine/Classes/Engine/StreamableManager.h b/Engine/Source/Runtime/Engine/Classes/Engine/StreamableManager.h
--- a/Engine/Source/Runtime/Engine/Classes/Engine/StreamableManager.h
+++ b/Engine/Source/Runtime/Engine/Classes/Engine/StreamableManager.h
@@ -19,6 +19,12 @@ struct ENGINE_API FStreamableManager
 	 * No references are made to the object, and this can be very slow.
 	 */
 	UObject* SynchronousLoad(FStringAssetReference const& Target);
+
+	template< typename T >
+	T* SynchronousLoadType( FStringAssetReference const& Target )
+	{
+		return Cast< T >( SynchronousLoad( Target ) );
+	}
 	
 	/** 
 	 * Perform a simple asynchronous load of a single object.
