Engine: Added configurable threshold settings for fps/frame time displays
- t.TargetFrameTimeThreshold and t.UnacceptableFrameTimeThreshold set the thresholds (in ms, defaulting to 33.9 / 50.0)
- Values below the target level are green, values above the unacceptable level are red, and the ones in between are yellow

diff --git a/Engine/Source/Runtime/Engine/Classes/Engine/Engine.h b/Engine/Source/Runtime/Engine/Classes/Engine/Engine.h
--- a/Engine/Source/Runtime/Engine/Classes/Engine/Engine.h
+++ b/Engine/Source/Runtime/Engine/Classes/Engine/Engine.h
@@ -1763,6 +1763,11 @@ public:
 	 */
 	void SetAverageUnitTimes(float FrameTime, float RenderThreadTime, float GameThreadTime, float GPUFrameTime);
 
+	/**
+	 * Returns the display color for a given frame time (based on t.TargetFrameTimeThreshold and t.UnacceptableFrameTimeThreshold)
+	 */
+	FColor GetFrameTimeDisplayColor(float FrameTimeMS) const;
+
 	/**
 	 * @return true to throttle CPU usage based on current state (usually editor minimized or not in foreground)
 	 */
diff --git a/Engine/Source/Runtime/Engine/Private/UnrealClient.cpp b/Engine/Source/Runtime/Engine/Private/UnrealClient.cpp
--- a/Engine/Source/Runtime/Engine/Private/UnrealClient.cpp
+++ b/Engine/Source/Runtime/Engine/Private/UnrealClient.cpp
@@ -361,7 +361,6 @@ int32 FStatUnitData::DrawStat(FViewport* InViewport, FCanvas* InCanvas, int32 In
 	const bool bStereoRendering = (GEngine && GEngine->IsStereoscopic3D(InViewport));
 	UFont* Font = (!FPlatformProperties::SupportsWindowedMode() && GEngine->GetMediumFont()) ? GEngine->GetMediumFont() : GEngine->GetSmallFont();
 
-	FColor Color;
 	// Note InX should already be within the safe zone
 	int32 X3 = InX * (bStereoRendering ? 0.5f : 1.0f);
 	if (bShowUnitMaxTimes)
@@ -373,43 +372,52 @@ int32 FStatUnitData::DrawStat(FViewport* InViewport, FCanvas* InCanvas, int32 In
 	const int32 RowHeight = FMath::TruncToInt(Font->GetMaxCharHeight() * 1.1f);
 	const bool bShowUnitTimeGraph = InViewport->GetClient() ? InViewport->GetClient()->IsStatEnabled(TEXT("UnitGraph")) : false;
 
-	// 0-34 ms: Green, 34-50 ms: Yellow, 50+ ms: Red
-	Color = FrameTime < 34.0f ? FColor::Green : (FrameTime < 50.0f ? FColor::Yellow : FColor::Red);
-	InCanvas->DrawShadowedString(X1, InY, TEXT("Frame:"), Font, bShowUnitTimeGraph ? FColor(100, 255, 100) : FColor::White);
-	InCanvas->DrawShadowedString(X2, InY, *FString::Printf(TEXT("%3.2f ms"), FrameTime), Font, Color);
-	if (bShowUnitMaxTimes)
 	{
-		InCanvas->DrawShadowedString(X3, InY, *FString::Printf(TEXT("%4.2f ms"), Max_FrameTime), Font, Color);
+		const FColor FrameTimeAverageColor = GEngine->GetFrameTimeDisplayColor(FrameTime);
+		InCanvas->DrawShadowedString(X1, InY, TEXT("Frame:"), Font, bShowUnitTimeGraph ? FColor(100, 255, 100) : FColor::White);
+		InCanvas->DrawShadowedString(X2, InY, *FString::Printf(TEXT("%3.2f ms"), FrameTime), Font, FrameTimeAverageColor);
+		if (bShowUnitMaxTimes)
+		{
+			const FColor MaxFrameTimeColor = GEngine->GetFrameTimeDisplayColor(Max_FrameTime);
+			InCanvas->DrawShadowedString(X3, InY, *FString::Printf(TEXT("%4.2f ms"), Max_FrameTime), Font, MaxFrameTimeColor);
+		}
+		InY += RowHeight;
 	}
-	InY += RowHeight;
 
-	Color = GameThreadTime < 34.0f ? FColor::Green : (GameThreadTime < 50.0f ? FColor::Yellow : FColor::Red);
-	InCanvas->DrawShadowedString(X1, InY, TEXT("Game:"), Font, bShowUnitTimeGraph ? FColor(255, 100, 100) : FColor::White);
-	InCanvas->DrawShadowedString(X2, InY, *FString::Printf(TEXT("%3.2f ms"), GameThreadTime), Font, Color);
-	if (bShowUnitMaxTimes)
 	{
-		InCanvas->DrawShadowedString(X3, InY, *FString::Printf(TEXT("%4.2f ms"), Max_GameThreadTime), Font, Color);
+		const FColor GameThreadAverageColor = GEngine->GetFrameTimeDisplayColor(GameThreadTime);
+		InCanvas->DrawShadowedString(X1, InY, TEXT("Game:"), Font, bShowUnitTimeGraph ? FColor(255, 100, 100) : FColor::White);
+		InCanvas->DrawShadowedString(X2, InY, *FString::Printf(TEXT("%3.2f ms"), GameThreadTime), Font, GameThreadAverageColor);
+		if (bShowUnitMaxTimes)
+		{
+			const FColor GameThreadMaxColor = GEngine->GetFrameTimeDisplayColor(Max_GameThreadTime);
+			InCanvas->DrawShadowedString(X3, InY, *FString::Printf(TEXT("%4.2f ms"), Max_GameThreadTime), Font, GameThreadMaxColor);
+		}
+		InY += RowHeight;
 	}
-	InY += RowHeight;
 
-	Color = RenderThreadTime < 34.0f ? FColor::Green : (RenderThreadTime < 50.0f ? FColor::Yellow : FColor::Red);
-	InCanvas->DrawShadowedString(X1, InY, TEXT("Draw:"), Font, bShowUnitTimeGraph ? FColor(100, 100, 255) : FColor::White);
-	InCanvas->DrawShadowedString(X2, InY, *FString::Printf(TEXT("%3.2f ms"), RenderThreadTime), Font, Color);
-	if (bShowUnitMaxTimes)
 	{
-		InCanvas->DrawShadowedString(X3, InY, *FString::Printf(TEXT("%4.2f ms"), Max_RenderThreadTime), Font, Color);
+		const FColor RenderThreadAverageColor = GEngine->GetFrameTimeDisplayColor(RenderThreadTime);
+		InCanvas->DrawShadowedString(X1, InY, TEXT("Draw:"), Font, bShowUnitTimeGraph ? FColor(100, 100, 255) : FColor::White);
+		InCanvas->DrawShadowedString(X2, InY, *FString::Printf(TEXT("%3.2f ms"), RenderThreadTime), Font, RenderThreadAverageColor);
+		if (bShowUnitMaxTimes)
+		{
+			const FColor RenderThreadMaxColor = GEngine->GetFrameTimeDisplayColor(Max_RenderThreadTime);
+			InCanvas->DrawShadowedString(X3, InY, *FString::Printf(TEXT("%4.2f ms"), Max_RenderThreadTime), Font, RenderThreadMaxColor);
+		}
+		InY += RowHeight;
 	}
-	InY += RowHeight;
 
 	const bool bHaveGPUData = GPUCycles > 0;
 	if (bHaveGPUData)
 	{
-		Color = GPUFrameTime < 34.0f ? FColor::Green : (GPUFrameTime < 50.0f ? FColor::Yellow : FColor::Red);
+		const FColor GPUAverageColor = GEngine->GetFrameTimeDisplayColor(GPUFrameTime);
 		InCanvas->DrawShadowedString(X1, InY, TEXT("GPU:"), Font, bShowUnitTimeGraph ? FColor(255, 255, 100) : FColor::White);
-		InCanvas->DrawShadowedString(X2, InY, *FString::Printf(TEXT("%3.2f ms"), GPUFrameTime), Font, Color);
+		InCanvas->DrawShadowedString(X2, InY, *FString::Printf(TEXT("%3.2f ms"), GPUFrameTime), Font, GPUAverageColor);
 		if (bShowUnitMaxTimes)
 		{
-			InCanvas->DrawShadowedString(X3, InY, *FString::Printf(TEXT("%4.2f ms"), Max_GPUFrameTime), Font, Color);
+			const FColor GPUMaxColor = GEngine->GetFrameTimeDisplayColor(Max_GPUFrameTime);
+			InCanvas->DrawShadowedString(X3, InY, *FString::Printf(TEXT("%4.2f ms"), Max_GPUFrameTime), Font, GPUMaxColor);
 		}
 		InY += RowHeight;
 	}
@@ -427,7 +435,10 @@ int32 FStatUnitData::DrawStat(FViewport* InViewport, FCanvas* InCanvas, int32 In
 		// The horizontal axis is the frame number (NOT time!!!)
 
 		// Threshold where graph lines will pulsate for slow frames
-		const float AlertTimeMS = 33.33f;
+		extern TAutoConsoleVariable<float> GTargetFrameTimeThresholdCVar;
+		const float TargetTimeMS = GTargetFrameTimeThresholdCVar.GetValueOnGameThread();
+
+		const float AlertTimeMS = TargetTimeMS;
 
 		// Graph layout
 		const float GraphLeftXPos = 80.0f;
diff --git a/Engine/Source/Runtime/Engine/Private/UnrealEngine.cpp b/Engine/Source/Runtime/Engine/Private/UnrealEngine.cpp
--- a/Engine/Source/Runtime/Engine/Private/UnrealEngine.cpp
+++ b/Engine/Source/Runtime/Engine/Private/UnrealEngine.cpp
@@ -156,7 +156,23 @@ static TAutoConsoleVariable<float> GHitchThresholdCVar(
 	TEXT("Time in seconds that is considered a hitch by \"stat dumphitches\"")
 	);
 
-static TAutoConsoleVariable<int32> CVarAllowOneFrameThreadLag(
+// The maximum threshold for an 'OK' frame time in miliseconds (*cosmetic only* and used for fps/stat display, should not be used in scalability code)
+// Values above this will be red, values between this and the acceptable limit will be yellow, and values below will be green.
+static TAutoConsoleVariable<float> GUnacceptableFrameTimeThresholdCVar(
+	TEXT("t.UnacceptableFrameTimeThreshold"),
+	50.0f,
+	TEXT("The frame time theshold for what is considered completely unacceptable; values above this will be drawn as red\n")
+	TEXT(" default: 50.0 ms"));
+
+// The target threshold for frame time in miliseconds (*cosmetic only* and used for fps/stat display, should not be used in scalability code)
+// Values below this will be green, values between this and the unacceptable threshold will be yellow, and values above that will be red
+TAutoConsoleVariable<float> GTargetFrameTimeThresholdCVar(
+	TEXT("t.TargetFrameTimeThreshold"),
+	33.9f,
+	TEXT("The target frame time; values below this will be drawn in green, values above will be yellow or red depending on the severity\n")
+	TEXT(" default: 33.9 ms"));
+
+TAutoConsoleVariable<int32> CVarAllowOneFrameThreadLag(
 	TEXT("r.OneFrameThreadLag"),
 	1,
 	TEXT("Whether to allow the rendering thread to lag one frame behind the game thread (0: disabled, otherwise enabled)")
@@ -7256,6 +7272,14 @@ void UEngine::SetAverageUnitTimes(float FrameTime, float RenderThreadTime, float
 	}
 }
 
+FColor UEngine::GetFrameTimeDisplayColor(float FrameTimeMS) const
+{
+	const float UnacceptableTime = GUnacceptableFrameTimeThresholdCVar.GetValueOnGameThread();
+	const float TargetTime = GTargetFrameTimeThresholdCVar.GetValueOnGameThread();
+
+	return (FrameTimeMS > UnacceptableTime) ? FColor::Red : ((FrameTimeMS > TargetTime) ? FColor::Yellow : FColor::Green);
+}
+
 bool UEngine::ShouldThrottleCPUUsage() const
 {
 	return false;
@@ -11310,11 +11334,12 @@ int32 UEngine::RenderStatFPS(UWorld* World, FViewport* Viewport, FCanvas* Canvas
 	// Pick a larger font on console.
 	UFont* Font = FPlatformProperties::SupportsWindowedMode() ? GetSmallFont() : GetMediumFont();
 
-	// Choose the counter color based on the average framerate.
-	FColor FPSColor = GAverageFPS < 20.0f ? FColor::Red : (GAverageFPS < 29.5f ? FColor::Yellow : FColor::Green);
+	// Choose the counter color based on the average frame time.
+	const FColor FPSColor = GetFrameTimeDisplayColor(GAverageMS);
 
 	// Start drawing the various counters.
 	const int32 RowHeight = FMath::TruncToInt(Font->GetMaxCharHeight() * 1.1f);
+
 	// Draw the FPS counter.
 	Canvas->DrawShadowedString(
 		X,
