[BuildWatcher] Trying to fix Mac CIS deprecation warnings

diff --git a/Engine/Source/Runtime/Engine/Classes/Engine/World.h b/Engine/Source/Runtime/Engine/Classes/Engine/World.h
--- a/Engine/Source/Runtime/Engine/Classes/Engine/World.h
+++ b/Engine/Source/Runtime/Engine/Classes/Engine/World.h
@@ -459,27 +459,6 @@ struct ENGINE_API FActorSpawnParameters
 	EObjectFlags ObjectFlags;		
 };
 
-// Deprecation warnings disabled to initialize bNoCollisionFail
-PRAGMA_DISABLE_DEPRECATION_WARNINGS
-
-	inline FActorSpawnParameters::FActorSpawnParameters()
-		: Name(NAME_None)
-		, Template(NULL)
-		, Owner(NULL)
-		, Instigator(NULL)
-		, OverrideLevel(NULL)
-		, SpawnCollisionHandlingOverride(ESpawnActorCollisionHandlingMethod::Undefined)
-		, bNoCollisionFail(false)
-		, bRemoteOwned(false)
-		, bNoFail(false)
-		, bDeferConstruction(false)
-		, bAllowDuringConstructionScript(false)
-		, ObjectFlags(RF_Transactional)
-	{
-	}
-
-PRAGMA_ENABLE_DEPRECATION_WARNINGS
-
 
 /**
  *  This encapsulate World's async trace functionality. This contains two buffers of trace data buffer and alternates it for each tick. 
diff --git a/Engine/Source/Runtime/Engine/Private/World.cpp b/Engine/Source/Runtime/Engine/Private/World.cpp
--- a/Engine/Source/Runtime/Engine/Private/World.cpp
+++ b/Engine/Source/Runtime/Engine/Private/World.cpp
@@ -87,6 +87,27 @@ DEFINE_LOG_CATEGORY(LogSpawn);
 
 #define LOCTEXT_NAMESPACE "World"
 
+// Deprecation warnings disabled to initialize bNoCollisionFail
+PRAGMA_DISABLE_DEPRECATION_WARNINGS
+
+FActorSpawnParameters::FActorSpawnParameters()
+: Name(NAME_None)
+, Template(NULL)
+, Owner(NULL)
+, Instigator(NULL)
+, OverrideLevel(NULL)
+, SpawnCollisionHandlingOverride(ESpawnActorCollisionHandlingMethod::Undefined)
+, bNoCollisionFail(false)
+, bRemoteOwned(false)
+, bNoFail(false)
+, bDeferConstruction(false)
+, bAllowDuringConstructionScript(false)
+, ObjectFlags(RF_Transactional)
+{
+}
+
+PRAGMA_ENABLE_DEPRECATION_WARNINGS
+
 /*-----------------------------------------------------------------------------
 	UWorld implementation.
 -----------------------------------------------------------------------------*/
