Added simulation stat so we can tell the full time it takes for a specific physx scene from start to finish

diff --git a/Engine/Source/Runtime/Engine/Private/PhysicsEngine/PhysScene.cpp b/Engine/Source/Runtime/Engine/Private/PhysicsEngine/PhysScene.cpp
--- a/Engine/Source/Runtime/Engine/Private/PhysicsEngine/PhysScene.cpp
+++ b/Engine/Source/Runtime/Engine/Private/PhysicsEngine/PhysScene.cpp
@@ -418,7 +418,7 @@ bool FPhysScene::SubstepSimulation(uint32 SceneType, FGraphEventRef &InOutComple
 	}else
 	{
 		//we have valid scene and subtime so enqueue task
-		PhysXCompletionTask* Task = new PhysXCompletionTask(InOutCompletionEvent, PScene->getTaskManager());
+		PhysXCompletionTask* Task = new PhysXCompletionTask(InOutCompletionEvent, SceneType, PScene->getTaskManager());
 		ENamedThreads::Type NamedThread = PhysSingleThreadedMode() ? ENamedThreads::GameThread : ENamedThreads::AnyThread;
 
 		DECLARE_CYCLE_STAT(TEXT("FSimpleDelegateGraphTask.SubstepSimulationImp"),
@@ -540,6 +540,30 @@ void GatherPhysXStats(PxScene* PScene, uint32 SceneType)
 #endif
 }
 
+DECLARE_FLOAT_COUNTER_STAT(TEXT("Sync Sim Time (ms)"), STAT_PhysSyncSim, STATGROUP_Physics);
+DECLARE_FLOAT_COUNTER_STAT(TEXT("Async Sim Time (ms)"), STAT_PhysAsyncSim, STATGROUP_Physics);
+DECLARE_FLOAT_COUNTER_STAT(TEXT("Cloth Sim Time (ms)"), STAT_PhysClothSim, STATGROUP_Physics);
+
+double GSimStartTime[PST_MAX] = {0.f, 0.f, 0.f};
+
+void FinishSceneStat(uint32 Scene)
+{
+	if (Scene < PST_MAX)	//PST_MAX used when we don't care
+	{
+		float SceneTime = float(FPlatformTime::Seconds() - GSimStartTime[Scene]) * 1000.0f;
+		switch(Scene)
+		{
+			case PST_Sync:
+			INC_FLOAT_STAT_BY(STAT_PhysSyncSim, SceneTime); break;
+			case PST_Async:
+			INC_FLOAT_STAT_BY(STAT_PhysAsyncSim, SceneTime); break;
+			case PST_Cloth:
+			INC_FLOAT_STAT_BY(STAT_PhysClothSim, SceneTime); break;
+		}
+	}
+}
+
+
 /** Exposes ticking of physics-engine scene outside Engine. */
 void FPhysScene::TickPhysScene(uint32 SceneType, FGraphEventRef& InOutCompletionEvent)
 {
@@ -548,6 +572,8 @@ void FPhysScene::TickPhysScene(uint32 SceneType, FGraphEventRef& InOutCompletion
 
 	check(SceneType < NumPhysScenes);
 
+	GSimStartTime[SceneType] = FPlatformTime::Seconds();
+
 	if (bPhysXSceneExecuting[SceneType] != 0)
 	{
 		// Already executing this scene, must call WaitPhysScene before calling this function again.
@@ -641,14 +667,14 @@ void FPhysScene::TickPhysScene(uint32 SceneType, FGraphEventRef& InOutCompletion
 #endif
 		{
 #if !WITH_APEX
-			PhysXCompletionTask* Task = new PhysXCompletionTask(InOutCompletionEvent, PScene->getTaskManager());
+			PhysXCompletionTask* Task = new PhysXCompletionTask(InOutCompletionEvent, SceneType, PScene->getTaskManager());
 			PScene->lockWrite();
 			PScene->simulate(AveragedFrameTime[SceneType], Task);
 			PScene->unlockWrite();
 			Task->removeReference();
 			bTaskOutstanding = true;
 #else
-			PhysXCompletionTask* Task = new PhysXCompletionTask(InOutCompletionEvent, ApexScene->getTaskManager());
+			PhysXCompletionTask* Task = new PhysXCompletionTask(InOutCompletionEvent, SceneType, ApexScene->getTaskManager());
 			ApexScene->simulate(AveragedFrameTime[SceneType], true, Task);
 			Task->removeReference();
 			bTaskOutstanding = true;
diff --git a/Engine/Source/Runtime/Engine/Private/PhysicsEngine/PhysSubstepTasks.cpp b/Engine/Source/Runtime/Engine/Private/PhysicsEngine/PhysSubstepTasks.cpp
--- a/Engine/Source/Runtime/Engine/Private/PhysicsEngine/PhysSubstepTasks.cpp
+++ b/Engine/Source/Runtime/Engine/Private/PhysicsEngine/PhysSubstepTasks.cpp
@@ -358,7 +358,9 @@ void FPhysSubstepTask::SubstepSimulationStart()
 	
 	check(!CompletionEvent.GetReference());	//should be done
 	CompletionEvent = FGraphEvent::CreateGraphEvent();
-	PhysXCompletionTask* SubstepTask = new PhysXCompletionTask(CompletionEvent, PAScene->getTaskManager());
+	PhysXCompletionTask* SubstepTask = new PhysXCompletionTask(CompletionEvent,
+		 PST_MAX //we don't care about sub-step time. The full time is recorded by FullSimulationTask
+		,PAScene->getTaskManager());
 	ENamedThreads::Type NamedThread = PhysSingleThreadedMode() ? ENamedThreads::GameThread : ENamedThreads::AnyThread;
 
 	DECLARE_CYCLE_STAT(TEXT("FDelegateGraphTask.ProcessPhysSubstepSimulation"),
diff --git a/Engine/Source/Runtime/Engine/Private/PhysicsEngine/PhysSubstepTasks.h b/Engine/Source/Runtime/Engine/Private/PhysicsEngine/PhysSubstepTasks.h
--- a/Engine/Source/Runtime/Engine/Private/PhysicsEngine/PhysSubstepTasks.h
+++ b/Engine/Source/Runtime/Engine/Private/PhysicsEngine/PhysSubstepTasks.h
@@ -2,15 +2,19 @@
 
 #pragma once
 
+void FinishSceneStat(uint32 Scene);
+
 //This is only here for now while we transition into substepping
 #if WITH_PHYSX
 #include "pxtask/PxTask.h"
 class PhysXCompletionTask : public PxLightCpuTask
 {
 	FGraphEventRef& EventToFire;
+	uint32 Scene;
 public:
-	PhysXCompletionTask(FGraphEventRef& InEventToFire, PxTaskManager* TaskManager)
+	PhysXCompletionTask(FGraphEventRef& InEventToFire, uint32 InScene, PxTaskManager* TaskManager)
 		: EventToFire(InEventToFire)
+		, Scene(InScene)
 	{
 		setContinuation(*TaskManager, NULL);
 	}
@@ -20,6 +24,7 @@ public:
 	virtual void release()
 	{
 		PxLightCpuTask::release();
+		FinishSceneStat(Scene);
 		EventToFire->DispatchSubsequents();
 		delete this;
 	}
