Engine: Exposed some more UGameUserSettings methods to Blueprints
[INTEGRATE] Change 2638057 by [REDACTED_PERSON][REDACTED_HANDLE] on 2015/07/29 21:16:42

diff --git a/Engine/Source/Runtime/Engine/Classes/GameFramework/GameUserSettings.h b/Engine/Source/Runtime/Engine/Classes/GameFramework/GameUserSettings.h
--- a/Engine/Source/Runtime/Engine/Classes/GameFramework/GameUserSettings.h
+++ b/Engine/Source/Runtime/Engine/Classes/GameFramework/GameUserSettings.h
@@ -192,21 +192,26 @@ public:
 	virtual bool IsDirty() const;
 
 	/** Validates and resets bad user settings to default. Deletes stale user settings file if necessary. */
+	UFUNCTION(BlueprintCallable, Category=Settings)
 	virtual void ValidateSettings();
 
 	/** Loads the user settings from persistent storage */
+	UFUNCTION(BlueprintCallable, Category=Settings)
 	virtual void LoadSettings(bool bForceReload = false);
 
-	/** Save the user settings to persistent storage */
+	/** Save the user settings to persistent storage (automatically happens as part of ApplySettings) */
+	UFUNCTION(BlueprintCallable, Category=Settings)
 	virtual void SaveSettings();
 
 	/** This function resets all settings to the current system settings */
+	UFUNCTION(BlueprintCallable, Category=Settings)
 	virtual void ResetToCurrentSettings();
 
 	virtual void SetWindowPosition(int32 WindowPosX, int32 WindowPosY);
 
 	virtual FIntPoint GetWindowPosition();
 	
+	UFUNCTION(BlueprintCallable, Category=Settings)
 	virtual void SetToDefaults();
 
 	/** Loads the resolution settings before is object is available */
