Fix for destructible mesh flickering due to the root bone transform not being updated, sometimes propogating NaNs to the GPU skinning code.

diff --git a/Engine/Source/Editor/UnrealEd/Private/Factories/ApexDestructibleAssetImport.cpp b/Engine/Source/Editor/UnrealEd/Private/Factories/ApexDestructibleAssetImport.cpp
--- a/Engine/Source/Editor/UnrealEd/Private/Factories/ApexDestructibleAssetImport.cpp
+++ b/Engine/Source/Editor/UnrealEd/Private/Factories/ApexDestructibleAssetImport.cpp
@@ -842,6 +842,8 @@ bool SetApexDestructibleAsset(UDestructibleMesh& DestructibleMesh, NxDestructibl
 
 	FStaticLODModel& LODModel = DestructibleMeshResource.LODModels[0];
 	
+	LODModel.ActiveBoneIndices.Add(0);
+
 	// Pass the number of texture coordinate sets to the LODModel.  Ensure there is at least one UV coord
 	LODModel.NumTexCoords = FMath::Max<uint32>(1,SkelMeshImportDataPtr->NumTexCoords);
 //	if( bCreateRenderData )	// We always create render data
diff --git a/Engine/Source/Runtime/Engine/Private/DestructibleMesh.cpp b/Engine/Source/Runtime/Engine/Private/DestructibleMesh.cpp
--- a/Engine/Source/Runtime/Engine/Private/DestructibleMesh.cpp
+++ b/Engine/Source/Runtime/Engine/Private/DestructibleMesh.cpp
@@ -47,6 +47,13 @@ void UDestructibleMesh::PostLoad()
 		FSkeletalMeshLODInfo& ThisLODInfo = LODInfo[LodIndex];
 		FStaticLODModel& ThisLODModel = ImportedResource->LODModels[LodIndex];
 
+		// Check that we list the root bone as an active bone.
+		if(!ThisLODModel.ActiveBoneIndices.Contains(0))
+		{
+			ThisLODModel.ActiveBoneIndices.Add(0);
+			ThisLODModel.ActiveBoneIndices.Sort();
+		}
+
 		for (int32 ChunkIndex = 0; ChunkIndex < ThisLODModel.Chunks.Num(); ++ChunkIndex)
 		{
 			if (ThisLODModel.Chunks[ChunkIndex].BoneMap.Num() > MaxGPUSkinBones)
