Fixing crash with Material Overrides brought up by user https://udn.unrealengine.com/questions/224531/using-materialoverride-can-cause-a-crash.html#answer-255497

Distortion rendering was not taking blend mode overrides into account.

diff --git a/Engine/Shaders/ShaderVersion.usf b/Engine/Shaders/ShaderVersion.usf
--- a/Engine/Shaders/ShaderVersion.usf
+++ b/Engine/Shaders/ShaderVersion.usf
@@ -2,4 +2,4 @@
 // Each time the console command is executed it generates a new GUID. As this file is included
 // in common.usf (which should be included in any shader) it allows to invalidate the shader DDC.
 // 
-// GUID = 614697354776EE6EA312B887F2FFD5EF
\ No newline at end of file
+// GUID = 1E7DB6E145A42CAE76519AA01490523B
\ No newline at end of file
diff --git a/Engine/Source/Editor/MaterialEditor/Private/MaterialEditorInstanceDetailCustomization.cpp b/Engine/Source/Editor/MaterialEditor/Private/MaterialEditorInstanceDetailCustomization.cpp
--- a/Engine/Source/Editor/MaterialEditor/Private/MaterialEditorInstanceDetailCustomization.cpp
+++ b/Engine/Source/Editor/MaterialEditor/Private/MaterialEditorInstanceDetailCustomization.cpp
@@ -440,7 +440,7 @@ void FMaterialInstanceParameterDetails::ResetToDefault( class UDEditorParameterV
 
 EVisibility FMaterialInstanceParameterDetails::ShouldShowMaterialRefractionSettings() const
 {
-	return (MaterialEditorInstance->SourceInstance->GetMaterial()->bUsesDistortion) ? EVisibility::Visible : EVisibility::Collapsed;
+	return (MaterialEditorInstance->SourceInstance->GetMaterial()->bUsesDistortion && IsTranslucentBlendMode(MaterialEditorInstance->SourceInstance->GetBlendMode())) ? EVisibility::Visible : EVisibility::Collapsed;
 }
 
 EVisibility FMaterialInstanceParameterDetails::ShouldShowSubsurfaceProfile() const
diff --git a/Engine/Source/Runtime/Engine/Private/Materials/Material.cpp b/Engine/Source/Runtime/Engine/Private/Materials/Material.cpp
--- a/Engine/Source/Runtime/Engine/Private/Materials/Material.cpp
+++ b/Engine/Source/Runtime/Engine/Private/Materials/Material.cpp
@@ -2630,15 +2630,11 @@ void UMaterial::PostEditChangeProperty(FPropertyChangedEvent& PropertyChangedEve
 	// check for distortion in material 
 	{
 		bUsesDistortion = false;
-		// can only have distortion with translucent blend modes
-		if(IsTranslucentBlendMode((EBlendMode)BlendMode))
-		{
-			// check for a distortion value
-			if(Refraction.Expression
+		// check for a distortion value
+		if (Refraction.Expression
 			|| (Refraction.UseConstant && FMath::Abs(Refraction.Constant) >= KINDA_SMALL_NUMBER))
-			{
-				bUsesDistortion = true;
-			}
+		{
+			bUsesDistortion = true;
 		}
 	}
 
diff --git a/Engine/Source/Runtime/Engine/Private/Materials/MaterialInterface.cpp b/Engine/Source/Runtime/Engine/Private/Materials/MaterialInterface.cpp
--- a/Engine/Source/Runtime/Engine/Private/Materials/MaterialInterface.cpp
+++ b/Engine/Source/Runtime/Engine/Private/Materials/MaterialInterface.cpp
@@ -64,7 +64,7 @@ FMaterialRelevance UMaterialInterface::GetRelevance_Internal(const UMaterial* Ma
 		MaterialRelevance.ShadingModelMask = 1 << ShadingModel;
 		MaterialRelevance.bOpaque = !bIsTranslucent;
 		MaterialRelevance.bMasked = IsMasked();
-		MaterialRelevance.bDistortion = Material->bUsesDistortion && bIsTranslucent;
+		MaterialRelevance.bDistortion = MaterialResource->IsDistorted();
 		MaterialRelevance.bSeparateTranslucency = bIsTranslucent && Material->bEnableSeparateTranslucency;
 		MaterialRelevance.bNormalTranslucency = bIsTranslucent && !Material->bEnableSeparateTranslucency;
 		MaterialRelevance.bDisableDepthTest = bIsTranslucent && Material->bDisableDepthTest;		
diff --git a/Engine/Source/Runtime/Engine/Private/Materials/MaterialShared.cpp b/Engine/Source/Runtime/Engine/Private/Materials/MaterialShared.cpp
--- a/Engine/Source/Runtime/Engine/Private/Materials/MaterialShared.cpp
+++ b/Engine/Source/Runtime/Engine/Private/Materials/MaterialShared.cpp
@@ -861,7 +861,7 @@ bool FMaterialResource::IsMasked() const
 	return MaterialInstance ? MaterialInstance->IsMasked() : Material->IsMasked();
 }
 
-bool FMaterialResource::IsDistorted() const { return Material->bUsesDistortion; }
+bool FMaterialResource::IsDistorted() const { return Material->bUsesDistortion && IsTranslucentBlendMode(GetBlendMode()); }
 float FMaterialResource::GetTranslucencyDirectionalLightingIntensity() const { return Material->TranslucencyDirectionalLightingIntensity; }
 float FMaterialResource::GetTranslucentShadowDensityScale() const { return Material->TranslucentShadowDensityScale; }
 float FMaterialResource::GetTranslucentSelfShadowDensityScale() const { return Material->TranslucentSelfShadowDensityScale; }
