Workaround to avoid culling screen-space material billboard components incorrectly. UE-4725

diff --git a/Engine/Source/Runtime/Engine/Private/Components/MaterialBillboardComponent.cpp b/Engine/Source/Runtime/Engine/Private/Components/MaterialBillboardComponent.cpp
--- a/Engine/Source/Runtime/Engine/Private/Components/MaterialBillboardComponent.cpp
+++ b/Engine/Source/Runtime/Engine/Private/Components/MaterialBillboardComponent.cpp
@@ -266,7 +266,17 @@ FBoxSphereBounds UMaterialBillboardComponent::CalcBounds(const FTransform& Local
 	float BoundsSize = 1.0f;
 	for (int32 i = 0; i < Elements.Num(); ++i)
 	{
-		BoundsSize = FMath::Max3(BoundsSize, Elements[i].BaseSizeX, Elements[i].BaseSizeY);
+		if (Elements[i].bSizeIsInScreenSpace)
+		{
+			// Workaround static bounds by disabling culling. Still allows override such as 'use parent bounds', etc.
+			// Note: Bounds are dynamically calculated at draw time per view, so difficult to cull correctly. (UE-4725)
+			BoundsSize = float(HALF_WORLD_MAX); 
+			break;
+		}
+		else
+		{
+			BoundsSize = FMath::Max3(BoundsSize, Elements[i].BaseSizeX, Elements[i].BaseSizeY);
+		}
 	}
 	BoundsSize *= LocalToWorld.GetMaximumAxisScale();
 
