Stats - Better formatting for displayed stats

diff --git a/Engine/Source/Runtime/Engine/Private/StatsRender2.cpp b/Engine/Source/Runtime/Engine/Private/StatsRender2.cpp
--- a/Engine/Source/Runtime/Engine/Private/StatsRender2.cpp
+++ b/Engine/Source/Runtime/Engine/Private/StatsRender2.cpp
@@ -437,27 +437,44 @@ static int32 RenderMemoryHeadings(class FCanvas* Canvas,int32 X,int32 Y)
 	return GetStatRenderGlobals().GetFontHeight() + (GetStatRenderGlobals().GetFontHeight() / 3);
 }
 
+static FString FormatStatValueFloat( const float Value )
+{
+	const float Frac = FMath::Frac( Value );
+	// #YRX_Stats: 2015-07-30 Move to stats thread, add support for int64 type, int32 may not be sufficient all the time.
+	const int32 Integer = FMath::FloorToInt( FMath::Abs( Value ) );
+	const FString IntString = FString::FormatAsNumber( Integer );
+	const FString FracString = FString::Printf( TEXT( "%0.2f" ), Frac );
+	const FString Result = FString::Printf( TEXT( "%s.%s" ), *IntString, *FracString.Mid(2) ); 
+	return Result;
+}
+
+static FString FormatStatValueInt64( const int64 Value )
+{
+	const FString IntString = FString::FormatAsNumber( (int32)Value );
+	return IntString;	
+}
+
 // @param bAutoType true: automatically choose GB/MB/KB/... false: always use MB for easier comparisons
-static FString GetMemoryString(double Value, bool bAutoType = true)
+static FString GetMemoryString( double Value, bool bAutoType = true )
 {
-	if(bAutoType)
+	if (bAutoType)
 	{
 		if (Value > 1024.0 * 1024.0 * 1024.0)
 		{
-			return FString::Printf(TEXT("%.2f GB"), float(Value / (1024.0 * 1024.0 * 1024.0)));
+			return FString::Printf( TEXT( "%.2f GB" ), float( Value / (1024.0 * 1024.0 * 1024.0) ) );
 		}
 		if (Value > 1024.0 * 1024.0)
 		{
-			return FString::Printf(TEXT("%.2f MB"), float(Value / (1024.0 * 1024.0)));
+			return FString::Printf( TEXT( "%.2f MB" ), float( Value / (1024.0 * 1024.0) ) );
 		}
 		if (Value > 1024.0)
 		{
-			return FString::Printf(TEXT("%.2f KB"), float(Value / (1024.0)));
+			return FString::Printf( TEXT( "%.2f KB" ), float( Value / (1024.0) ) );
 		}
-		return FString::Printf(TEXT("%.2f B"), float(Value));
+		return FString::Printf( TEXT( "%.2f B" ), float( Value ) );
 	}
 
-	return FString::Printf(TEXT("%.2f MB"), float(Value / (1024.0 * 1024.0)));
+	return FString::Printf( TEXT( "%.2f MB" ), float( Value / (1024.0 * 1024.0) ) );
 }
 
 static int32 RenderMemoryCounter(const FGameThreadHudData& ViewData, const FComplexStatMessage& All,class FCanvas* Canvas,int32 X,int32 Y)
@@ -517,11 +534,13 @@ static int32 RenderCounter(const FGameThreadHudData& ViewData, const FComplexSta
 		// Append the average.
 		if (All.NameAndInfo.GetField<EStatDataType>() == EStatDataType::ST_double)
 		{
-			RightJustify(Canvas,CurrX,Y,*FString::Printf(TEXT("%.2f"), All.GetValue_double(EComplexStatField::IncAve)),GetStatRenderGlobals().StatColor);
+			const FString ValueFormatted = FormatStatValueFloat( All.GetValue_double( EComplexStatField::IncAve ) );
+			RightJustify( Canvas, CurrX, Y, *ValueFormatted, GetStatRenderGlobals().StatColor );
 		}
 		else if (All.NameAndInfo.GetField<EStatDataType>() == EStatDataType::ST_int64)
 		{
-			RightJustify(Canvas,CurrX,Y,*FString::Printf(TEXT("%lld"), All.GetValue_int64(EComplexStatField::IncAve)),GetStatRenderGlobals().StatColor);
+			const FString ValueFormatted = FormatStatValueInt64( All.GetValue_int64( EComplexStatField::IncAve ) );
+			RightJustify( Canvas, CurrX, Y, *ValueFormatted, GetStatRenderGlobals().StatColor );
 		}
 	}
 	CurrX += GetStatRenderGlobals().InterColumnOffset;
@@ -529,11 +548,13 @@ static int32 RenderCounter(const FGameThreadHudData& ViewData, const FComplexSta
 	// Append the maximum.
 	if (All.NameAndInfo.GetField<EStatDataType>() == EStatDataType::ST_double)
 	{
-		RightJustify(Canvas,CurrX,Y,*FString::Printf(TEXT("%.2f"), All.GetValue_double(EComplexStatField::IncMax)),GetStatRenderGlobals().StatColor);
+		const FString ValueFormatted = FormatStatValueFloat( All.GetValue_double( EComplexStatField::IncMax ) );
+		RightJustify( Canvas, CurrX, Y, *ValueFormatted, GetStatRenderGlobals().StatColor );
 	}
 	else if (All.NameAndInfo.GetField<EStatDataType>() == EStatDataType::ST_int64)
 	{
-		RightJustify(Canvas,CurrX,Y,*FString::Printf(TEXT("%lld"), All.GetValue_int64(EComplexStatField::IncMax)),GetStatRenderGlobals().StatColor);
+		const FString ValueFormatted = FormatStatValueInt64( All.GetValue_int64( EComplexStatField::IncMax ) );
+		RightJustify( Canvas, CurrX, Y, *ValueFormatted, GetStatRenderGlobals().StatColor );
 	}
 	return GetStatRenderGlobals().GetFontHeight();
 }
