Enabled clearing source memory in the cooker after we generate ddc content.
Reduces memory footprint.

diff --git a/Engine/Source/Runtime/Engine/Private/AudioDerivedData.cpp b/Engine/Source/Runtime/Engine/Private/AudioDerivedData.cpp
--- a/Engine/Source/Runtime/Engine/Private/AudioDerivedData.cpp
+++ b/Engine/Source/Runtime/Engine/Private/AudioDerivedData.cpp
@@ -1139,8 +1139,8 @@ void USoundWave::WillNeverCacheCookedPlatformDataAgain()
 	Super::WillNeverCacheCookedPlatformDataAgain();
 
 	// TODO: We can clear these arrays if we never need to cook again. 
-	// RawData.RemoveBulkData();
-	// CompressedFormatData.FlushData();
+	RawData.RemoveBulkData();
+	CompressedFormatData.FlushData();
 }
 #endif
 
diff --git a/Engine/Source/Runtime/Engine/Private/TextureDerivedData.cpp b/Engine/Source/Runtime/Engine/Private/TextureDerivedData.cpp
--- a/Engine/Source/Runtime/Engine/Private/TextureDerivedData.cpp
+++ b/Engine/Source/Runtime/Engine/Private/TextureDerivedData.cpp
@@ -1647,7 +1647,7 @@ bool UTexture::IsCachedCookedPlatformDataLoaded( const ITargetPlatform* TargetPl
 void UTexture2D::WillNeverCacheCookedPlatformDataAgain()
 {
 	Super::WillNeverCacheCookedPlatformDataAgain();
-#if WITH_EDITORONLY_DATA && 0 // enable this when I feel safe about it
+#if WITH_EDITORONLY_DATA 
 	// "Source" is only in WITH_EDITORONLY_DATA
 	UE_LOG(LogTemp, Warning, TEXT("Cleared source texture data for texture %s"), *GetName());
 	// clear source mips if we are in the editor then we don't have this luxury 
@@ -1721,10 +1721,10 @@ void UTexture::FinishCachePlatformData()
 #if DO_CHECK
 			if (!(GetOutermost()->PackageFlags & PKG_FilterEditorOnly))
 			{
-				FString DerivedDataKey;
-				FTextureBuildSettings BuildSettings;
-				GetBuildSettingsForRunningPlatform(*this, BuildSettings);
-				GetTextureDerivedDataKey(*this, BuildSettings, DerivedDataKey);
+			FString DerivedDataKey;
+			FTextureBuildSettings BuildSettings;
+			GetBuildSettingsForRunningPlatform(*this, BuildSettings);
+			GetTextureDerivedDataKey(*this, BuildSettings, DerivedDataKey);
 
 				check(RunningPlatformData->DerivedDataKey == DerivedDataKey);
 			}
@@ -1853,31 +1853,31 @@ void UTexture::SerializeCookedPlatformData(FArchive& Ar)
 			}
 			else
 			{
-				TMap<FString, FTexturePlatformData*> *CookedPlatformDataPtr = GetCookedPlatformData();
+			TMap<FString, FTexturePlatformData*> *CookedPlatformDataPtr = GetCookedPlatformData();
 				if (CookedPlatformDataPtr == NULL)
-					return;
-				
-				TArray<FName> PlatformFormats;
-
-				Ar.CookingTarget()->GetTextureFormats(this, PlatformFormats);
-				for (int32 FormatIndex = 0; FormatIndex < PlatformFormats.Num(); FormatIndex++)
-				{
-					FString DerivedDataKey;
-					BuildSettings.TextureFormatName = PlatformFormats[FormatIndex];
-					GetTextureDerivedDataKey(*this, BuildSettings, DerivedDataKey);
+				return;
 
-					FTexturePlatformData *PlatformDataPtr = (*CookedPlatformDataPtr).FindRef(DerivedDataKey);
+			TArray<FName> PlatformFormats;
 
-					if (PlatformDataPtr == NULL)
-					{
-						PlatformDataPtr = new FTexturePlatformData();
-						PlatformDataPtr->Cache(*this, BuildSettings, ETextureCacheFlags::InlineMips | ETextureCacheFlags::Async);
+			Ar.CookingTarget()->GetTextureFormats(this, PlatformFormats);
+			for (int32 FormatIndex = 0; FormatIndex < PlatformFormats.Num(); FormatIndex++)
+			{
+				FString DerivedDataKey;
+				BuildSettings.TextureFormatName = PlatformFormats[FormatIndex];
+				GetTextureDerivedDataKey(*this, BuildSettings, DerivedDataKey);
 
+				FTexturePlatformData *PlatformDataPtr = (*CookedPlatformDataPtr).FindRef(DerivedDataKey);
+				
+				if (PlatformDataPtr == NULL)
+				{
+					PlatformDataPtr = new FTexturePlatformData();
+					PlatformDataPtr->Cache(*this, BuildSettings, ETextureCacheFlags::InlineMips | ETextureCacheFlags::Async);
+					
 						CookedPlatformDataPtr->Add(DerivedDataKey, PlatformDataPtr);
-
-					}
-					PlatformDataToSerialize.Add(PlatformDataPtr);
+					
 				}
+				PlatformDataToSerialize.Add(PlatformDataPtr);
+			}
 			}
 
 			for (int32 i = 0; i < PlatformDataToSerialize.Num(); ++i)
