Fixing debug ensure command logic

diff --git a/Engine/Source/Runtime/Engine/Private/UnrealEngine.cpp b/Engine/Source/Runtime/Engine/Private/UnrealEngine.cpp
--- a/Engine/Source/Runtime/Engine/Private/UnrealEngine.cpp
+++ b/Engine/Source/Runtime/Engine/Private/UnrealEngine.cpp
@@ -4726,18 +4726,20 @@ bool UEngine::HandleDebugCommand( const TCHAR* Cmd, FOutputDevice& Ar )
 	else if( FParse::Command( &Cmd, TEXT("ENSURE") ) )
 	{
 		UE_LOG(LogEngine, Warning, TEXT("Printed warning to log.") );
-		if( !ensure( 0 ) )
+		if (!ensure(0))
 		{
-			return true;
+			UE_LOG(LogEngine, Warning, TEXT("Ensure condition failed (this is the expected behavior)."));
 		}
+		return true;
 	}
 	else if( FParse::Command( &Cmd, TEXT("ENSUREALWAYS") ) )
 	{
 		UE_LOG(LogEngine, Warning, TEXT("Printed warning to log.") );
 		if( !ensureAlways( 0 ) )
 		{
-			return true;
+			UE_LOG(LogEngine, Warning, TEXT("Ensure condition failed (this is the expected behavior)."));
 		}
+		return true;
 	}
 	else if( FParse::Command( &Cmd, TEXT( "FATAL" ) ) )
 	{
