Clients connected to a live game can now record their view of the world as a replay.
- Swap actors' Role and RemoteRole while replicating them for the replay.
- Properties with custom replication conditions are always considered active for client recording, since the game server has already computed the condition, and game code may be using the actor's Role to determine the property's active status. This led to missed RepNotify calls since Role is not the same on the client and server.
- The DemoNetDriver handles bTearOff actors so that the bTearOff flag is recorded properly in client replays.
- ULevel::SortActorList no longer adds non-replicated actors to UWorld::NetworkActors list on clients. Greatly reduces the number actors that are iterated through during client recording.

diff --git a/Engine/Source/Runtime/Engine/Classes/Engine/Level.h b/Engine/Source/Runtime/Engine/Classes/Engine/Level.h
--- a/Engine/Source/Runtime/Engine/Classes/Engine/Level.h
+++ b/Engine/Source/Runtime/Engine/Classes/Engine/Level.h
@@ -515,6 +515,9 @@ private:
 	// Actors awaiting input to be enabled once the appropriate PlayerController has been created
 	TArray<FPendingAutoReceiveInputActor> PendingAutoReceiveInputActors;
 
+	// Used internally to determine which actors should go on the world's NetworkActor list
+	static bool IsNetActor(const AActor* Actor);
+
 public:
 	/** Called when a level package has been dirtied. */
 	ENGINE_API static FSimpleMulticastDelegate LevelDirtiedEvent;
diff --git a/Engine/Source/Runtime/Engine/Classes/Engine/World.h b/Engine/Source/Runtime/Engine/Classes/Engine/World.h
--- a/Engine/Source/Runtime/Engine/Classes/Engine/World.h
+++ b/Engine/Source/Runtime/Engine/Classes/Engine/World.h
@@ -3011,6 +3011,11 @@ public:
 		NetDriver = NewDriver;
 	}
 
+	/**
+	 * Returns true if the game net driver exists and is a client and the demo net driver exists and is a server.
+	 */
+	bool IsRecordingClientReplay() const;
+
 	/**
 	 * Sets the number of frames to delay Streaming Volume updating, 
 	 * useful if you preload a bunch of levels but the camera hasn't caught up yet 
diff --git a/Engine/Source/Runtime/Engine/Classes/GameFramework/Actor.h b/Engine/Source/Runtime/Engine/Classes/GameFramework/Actor.h
--- a/Engine/Source/Runtime/Engine/Classes/GameFramework/Actor.h
+++ b/Engine/Source/Runtime/Engine/Classes/GameFramework/Actor.h
@@ -1687,6 +1687,9 @@ public:
 	/** Swaps Role and RemoteRole if client */
 	void ExchangeNetRoles(bool bRemoteOwner);
 
+	/** The replay system calls this to hack the Role and RemoteRole while recording replays on a client. Only call this if you know what you're doing! */
+	void SwapRolesForReplay();
+
 	/**
 	 * When called, will call the virtual call chain to register all of the tick functions for both the actor and optionally all components
 	 * Do not override this function or make it virtual
diff --git a/Engine/Source/Runtime/Engine/Private/Actor.cpp b/Engine/Source/Runtime/Engine/Private/Actor.cpp
--- a/Engine/Source/Runtime/Engine/Private/Actor.cpp
+++ b/Engine/Source/Runtime/Engine/Private/Actor.cpp
@@ -2712,6 +2712,11 @@ void AActor::ExchangeNetRoles(bool bRemoteOwned)
 	}
 }
 
+void AActor::SwapRolesForReplay()
+{
+	Swap(Role, RemoteRole);
+}
+
 void AActor::BeginPlay()
 {
 	ensure(!bActorHasBegunPlay);
diff --git a/Engine/Source/Runtime/Engine/Private/DemoNetDriver.cpp b/Engine/Source/Runtime/Engine/Private/DemoNetDriver.cpp
--- a/Engine/Source/Runtime/Engine/Private/DemoNetDriver.cpp
+++ b/Engine/Source/Runtime/Engine/Private/DemoNetDriver.cpp
@@ -725,8 +725,33 @@ void UDemoNetDriver::StopDemo()
 Demo Recording tick.
 -----------------------------------------------------------------------------*/
 
-static void DemoReplicateActor(AActor* Actor, UNetConnection* Connection, bool IsNetClient)
+static void DemoReplicateActor(AActor* Actor, UNetConnection* Connection, bool IsNetClient, APlayerController* SpectatorController)
 {
+	// RAII object to swap the Role and RemoteRole of an actor within a scope. Used for recording replays on a client.
+	class FScopedActorRoleSwap
+	{
+	public:
+		FScopedActorRoleSwap(AActor* InActor)
+			: Actor(InActor)
+		{
+			if (Actor != nullptr)
+			{
+				Actor->SwapRolesForReplay();
+			}
+		}
+
+		~FScopedActorRoleSwap()
+		{
+			if (Actor != nullptr)
+			{
+				Actor->SwapRolesForReplay();
+			}
+		}
+
+	private:
+		AActor* Actor;
+	};
+
 	// All actors marked for replication are assumed to be relevant for demo recording.
 	/*
 	if
@@ -736,37 +761,53 @@ static void DemoReplicateActor(AActor* Actor, UNetConnection* Connection, bool I
 		&& (Actor == Connection->PlayerController || Cast<APlayerController>(Actor) == NULL)
 		)
 	*/
-	if ( Actor != NULL && Actor->GetRemoteRole() != ROLE_None && ( Actor == Connection->PlayerController || Cast< APlayerController >( Actor ) == NULL ) )
-	{
-		const bool bShouldHaveChannel =
-			Actor->bRelevantForNetworkReplays &&
+	if ( Actor != NULL )
+	{
+		// We need to swap roles if:
+		//  1. We're recording a replay on a client that's connected to a live server, and
+		//  2. the actor isn't bTearOff, and
+		//  3. the actor isn't the replay spectator controller.
+		// This is to ensure the roles appear correct when playing back this demo.
+		UWorld* ActorWorld = Actor->GetWorld();
+		bool bShouldSwapRoles =
+			ActorWorld != nullptr && ActorWorld->IsRecordingClientReplay() &&
 			!Actor->bTearOff &&
-			(!Actor->IsNetStartupActor() || Connection->ClientHasInitializedLevelFor(Actor));
+			Actor != SpectatorController;
 
-		UActorChannel* Channel = Connection->ActorChannels.FindRef(Actor);
+		FScopedActorRoleSwap RoleSwap(bShouldSwapRoles ? Actor : nullptr);
 
-		if (bShouldHaveChannel && Channel == NULL)
+		if ( (Actor->GetRemoteRole() != ROLE_None || Actor->bTearOff) && ( Actor == Connection->PlayerController || Cast< APlayerController >( Actor ) == NULL ) )
 		{
-			// Create a new channel for this actor.
-			Channel = (UActorChannel*)Connection->CreateChannel(CHTYPE_Actor, 1);
-			if (Channel != NULL)
+			const bool bShouldHaveChannel =
+				Actor->bRelevantForNetworkReplays &&
+				!Actor->bTearOff &&
+				(!Actor->IsNetStartupActor() || Connection->ClientHasInitializedLevelFor(Actor));
+
+			UActorChannel* Channel = Connection->ActorChannels.FindRef(Actor);
+
+			if (bShouldHaveChannel && Channel == NULL)
 			{
-				Channel->SetChannelActor(Actor);
+				// Create a new channel for this actor.
+				Channel = (UActorChannel*)Connection->CreateChannel(CHTYPE_Actor, 1);
+				if (Channel != NULL)
+				{
+					Channel->SetChannelActor(Actor);
+				}
 			}
-		}
 
-		if (Channel != NULL && !Channel->Closing)
-		{
-			// Send it out!
-			if (Channel->IsNetReady(0))
+			if (Channel != NULL && !Channel->Closing)
 			{
-				Channel->ReplicateActor();
-			}
+				// Send it out!
+				if (Channel->IsNetReady(0))
+				{
+					Channel->ReplicateActor();
+				}
 			
-			// Close the channel if this actor shouldn't have one
-			if (!bShouldHaveChannel)
-			{
-				Channel->Close();
+				// Close the channel if this actor shouldn't have one
+				if (!bShouldHaveChannel)
+				{
+					Channel->Close();
+				}
 			}
 		}
 	}
@@ -877,7 +918,7 @@ void UDemoNetDriver::SaveCheckpoint()
 	// It's important that we don't catch any new actors that the next frame will also catch, that will cause conflict with bOpen (the open will occur twice on the same channel)
 	if ( CheckpointConnection->ActorChannels.Contains( World->GetWorldSettings() ) )
 	{
-		DemoReplicateActor( World->GetWorldSettings(), CheckpointConnection, false );
+		DemoReplicateActor( World->GetWorldSettings(), CheckpointConnection, false, SpectatorController );
 	}
 
 	for ( AActor* Actor : World->NetworkActors )
@@ -885,7 +926,7 @@ void UDemoNetDriver::SaveCheckpoint()
 		if ( CheckpointConnection->ActorChannels.Contains( Actor ) )
 		{
 			Actor->CallPreReplication( this );
-			DemoReplicateActor( Actor, CheckpointConnection, false );
+			DemoReplicateActor( Actor, CheckpointConnection, false, SpectatorController );
 		}
 	}
 
@@ -992,7 +1033,7 @@ void UDemoNetDriver::TickDemoRecord( float DeltaSeconds )
 
 	const bool IsNetClient = ( GetWorld()->GetNetDriver() != NULL && GetWorld()->GetNetDriver()->GetNetMode() == NM_Client );
 
-	DemoReplicateActor( World->GetWorldSettings(), ClientConnections[0], IsNetClient );
+	DemoReplicateActor( World->GetWorldSettings(), ClientConnections[0], IsNetClient, SpectatorController );
 
 	for ( TSet<AActor*>::TIterator ActorIt = World->NetworkActors.CreateIterator(); ActorIt; ++ActorIt)
 	{
@@ -1004,14 +1045,16 @@ void UDemoNetDriver::TickDemoRecord( float DeltaSeconds )
 			continue;
 		}
 
-		if ( Actor->GetRemoteRole() == ROLE_None )
+		// During client recording, a torn-off actor will already have its remote role set to None, but
+		// we still need to replicate it one more time so that the recorded replay knows it's been torn-off as well.
+		if ( Actor->GetRemoteRole() == ROLE_None && !Actor->bTearOff)
 		{
 			ActorIt.RemoveCurrent();
 			continue;
 		}
 
 		Actor->CallPreReplication( this );
-		DemoReplicateActor( Actor, ClientConnections[0], IsNetClient );
+		DemoReplicateActor( Actor, ClientConnections[0], IsNetClient, SpectatorController );
 	}
 
 	// Make sure nothing is left over
diff --git a/Engine/Source/Runtime/Engine/Private/Level.cpp b/Engine/Source/Runtime/Engine/Private/Level.cpp
--- a/Engine/Source/Runtime/Engine/Private/Level.cpp
+++ b/Engine/Source/Runtime/Engine/Private/Level.cpp
@@ -422,6 +422,21 @@ void ULevel::Serialize( FArchive& Ar )
 	}
 }
 
+bool ULevel::IsNetActor(const AActor* Actor)
+{
+	if (Actor == nullptr)
+	{
+		return false;
+	}
+
+	// If this is a server, use RemoteRole.
+	// If this is a client, use Role.
+	const ENetRole NetRole = (Actor->GetNetMode() < NM_Client) ? Actor->GetRemoteRole() : Actor->Role;
+
+	// This test will return true on clients for actors with ROLE_Authority, which might be counterintuitive,
+	// but clients will need to consider these actors in some cases, such as if their bTearOff flag is true.
+	return NetRole > ROLE_None;
+}
 
 void ULevel::SortActorList()
 {
@@ -443,7 +458,8 @@ void ULevel::SortActorList()
 	for (int32 ActorIndex = StartIndex; ActorIndex < Actors.Num(); ActorIndex++)
 	{
 		AActor* Actor = Actors[ActorIndex];
-		if (Actor != NULL && !Actor->IsPendingKill() && Actor->GetRemoteRole() == ROLE_None)
+
+		if (Actor != NULL && !Actor->IsPendingKill() && !IsNetActor(Actor))
 		{
 			NewActors.Add(Actor);
 		}
@@ -454,7 +470,7 @@ void ULevel::SortActorList()
 	for (int32 ActorIndex = StartIndex; ActorIndex < Actors.Num(); ActorIndex++)
 	{
 		AActor* Actor = Actors[ActorIndex];		
-		if (Actor != NULL && !Actor->IsPendingKill() && Actor->GetRemoteRole() > ROLE_None)
+		if (Actor != NULL && !Actor->IsPendingKill() && IsNetActor(Actor))
 		{
 			NewActors.Add(Actor);
 		}
diff --git a/Engine/Source/Runtime/Engine/Private/NetworkDriver.cpp b/Engine/Source/Runtime/Engine/Private/NetworkDriver.cpp
--- a/Engine/Source/Runtime/Engine/Private/NetworkDriver.cpp
+++ b/Engine/Source/Runtime/Engine/Private/NetworkDriver.cpp
@@ -2816,7 +2816,8 @@ TSharedPtr<FRepChangedPropertyTracker> UNetDriver::FindOrCreateRepChangedPropert
 
 	if ( !GlobalPropertyTrackerPtr ) 
 	{
-		FRepChangedPropertyTracker * Tracker = new FRepChangedPropertyTracker();
+		const bool bForceAlwaysActive = GetWorld() != nullptr ? GetWorld()->IsRecordingClientReplay() : false;
+		FRepChangedPropertyTracker * Tracker = new FRepChangedPropertyTracker(bForceAlwaysActive);
 
 		GetObjectClassRepLayout( Obj->GetClass() )->InitChangedTracker( Tracker );
 
diff --git a/Engine/Source/Runtime/Engine/Private/World.cpp b/Engine/Source/Runtime/Engine/Private/World.cpp
--- a/Engine/Source/Runtime/Engine/Private/World.cpp
+++ b/Engine/Source/Runtime/Engine/Private/World.cpp
@@ -5425,6 +5425,19 @@ ENetMode UWorld::GetNetMode() const
 	return AttemptDeriveFromURL();
 }
 
+bool UWorld::IsRecordingClientReplay() const
+{
+	if (GetNetDriver() != nullptr && !GetNetDriver()->IsServer())
+	{
+		if (DemoNetDriver != nullptr && DemoNetDriver->IsServer())
+		{
+			return true;
+		}
+	}
+
+	return false;
+}
+
 #if WITH_EDITOR
 ENetMode UWorld::AttemptDeriveFromPlayInSettings() const
 {
diff --git a/Engine/Source/Runtime/Engine/Public/Net/RepLayout.h b/Engine/Source/Runtime/Engine/Public/Net/RepLayout.h
--- a/Engine/Source/Runtime/Engine/Public/Net/RepLayout.h
+++ b/Engine/Source/Runtime/Engine/Public/Net/RepLayout.h
@@ -31,7 +31,12 @@ public:
 class FRepChangedPropertyTracker : public IRepChangedPropertyTracker
 {
 public:
-	FRepChangedPropertyTracker() : LastReplicationGroupFrame( 0 ), LastReplicationFrame( 0 ), ActiveStatusChanged( false ), UnconditionalPropChanged( false ) { }
+	FRepChangedPropertyTracker(bool InForceAlwaysActive)
+		: LastReplicationGroupFrame( 0 )
+		, LastReplicationFrame( 0 )
+		, ActiveStatusChanged( false )
+		, UnconditionalPropChanged( false )
+		, ForceAlwaysActive( InForceAlwaysActive ) { }
 	virtual ~FRepChangedPropertyTracker() { }
 
 	virtual void SetCustomIsActiveOverride( const uint16 RepIndex, const bool bIsActive ) override
@@ -40,7 +45,7 @@ public:
 
 		checkSlow( Parent.IsConditional );
 
-		Parent.Active = bIsActive ? 1 : 0;
+		Parent.Active = (bIsActive || ForceAlwaysActive) ? 1 : 0;
 
 		if ( Parent.Active != Parent.OldActive )
 		{
@@ -58,6 +63,7 @@ public:
 
 	uint32						ActiveStatusChanged;
 	bool						UnconditionalPropChanged;
+	bool						ForceAlwaysActive;				// Used for client replay recording. The server has already evaluated any custom conditions, so the client doesn't need to.
 };
 
 class FRepLayout;
