Corrected fixup for Brushes with non-zero prepivots, saved prior to their deprecation.

diff --git a/Engine/Source/Runtime/Engine/Private/Components/BrushComponent.cpp b/Engine/Source/Runtime/Engine/Private/Components/BrushComponent.cpp
--- a/Engine/Source/Runtime/Engine/Private/Components/BrushComponent.cpp
+++ b/Engine/Source/Runtime/Engine/Private/Components/BrushComponent.cpp
@@ -460,6 +460,7 @@ UBrushComponent::UBrushComponent(const FObjectInitializer& ObjectInitializer)
 	bUseAsOccluder = true;
 	bUseEditorCompositing = true;
 	bCanEverAffectNavigation = true;
+	PrePivot_DEPRECATED = FVector::ZeroVector;
 }
 
 FPrimitiveSceneProxy* UBrushComponent::CreateSceneProxy()
@@ -547,9 +548,12 @@ void UBrushComponent::PostLoad()
 
 #if WITH_EDITOR
 	AActor* Owner = GetOwner();
+
 	if (Owner)
 	{
+		AddRelativeLocation(GetComponentTransform().TransformVector(-PrePivot_DEPRECATED));
 		Owner->SetPivotOffset(PrePivot_DEPRECATED);
+		PrePivot_DEPRECATED = FVector::ZeroVector;
 	}
 #endif
 }
