Added 'debug renderensure' command

diff --git a/Engine/Source/Runtime/Engine/Private/UnrealEngine.cpp b/Engine/Source/Runtime/Engine/Private/UnrealEngine.cpp
--- a/Engine/Source/Runtime/Engine/Private/UnrealEngine.cpp
+++ b/Engine/Source/Runtime/Engine/Private/UnrealEngine.cpp
@@ -4585,6 +4585,18 @@ bool UEngine::HandleDebugCommand( const TCHAR* Cmd, FOutputDevice& Ar )
 		} );
 		return true;
 	}
+	if (FParse::Command(&Cmd, TEXT("RENDERENSURE")))
+	{
+		ENQUEUE_UNIQUE_RENDER_COMMAND(CauseRenderThreadEnsure,
+		{
+			UE_LOG(LogEngine, Warning, TEXT("Printed warning to log."));
+			if (!ensure(0))
+			{
+				UE_LOG(LogEngine, Warning, TEXT("Ensure condition failed (this is the expected behavior)."));
+			}
+		});
+		return true;
+	}
 	if( FParse::Command(&Cmd,TEXT("THREADCRASH")) )
 	{
 		struct FThread
