Attempt to fix compile failure

diff --git a/Engine/Source/Runtime/Engine/Private/Level.cpp b/Engine/Source/Runtime/Engine/Private/Level.cpp
--- a/Engine/Source/Runtime/Engine/Private/Level.cpp
+++ b/Engine/Source/Runtime/Engine/Private/Level.cpp
@@ -431,7 +431,7 @@ bool ULevel::IsNetActor(const AActor* Actor)
 
 	// If this is a server, use RemoteRole.
 	// If this is a client, use Role.
-	const ENetRole NetRole = (Actor->GetNetMode() < NM_Client) ? Actor->GetRemoteRole() : Actor->Role;
+	const ENetRole NetRole = (Actor->GetNetMode() < NM_Client) ? Actor->GetRemoteRole() : (ENetRole)Actor->Role;
 
 	// This test will return true on clients for actors with ROLE_Authority, which might be counterintuitive,
 	// but clients will need to consider these actors in some cases, such as if their bTearOff flag is true.
