Stats - Removed redundant Abs from FormatStatValueFloat

diff --git a/Engine/Source/Runtime/Engine/Private/StatsRender2.cpp b/Engine/Source/Runtime/Engine/Private/StatsRender2.cpp
--- a/Engine/Source/Runtime/Engine/Private/StatsRender2.cpp
+++ b/Engine/Source/Runtime/Engine/Private/StatsRender2.cpp
@@ -441,7 +441,7 @@ static FString FormatStatValueFloat( const float Value )
 {
 	const float Frac = FMath::Frac( Value );
 	// #YRX_Stats: 2015-07-30 Move to stats thread, add support for int64 type, int32 may not be sufficient all the time.
-	const int32 Integer = FMath::FloorToInt( FMath::Abs( Value ) );
+	const int32 Integer = FMath::FloorToInt( Value );
 	const FString IntString = FString::FormatAsNumber( Integer );
 	const FString FracString = FString::Printf( TEXT( "%0.2f" ), Frac );
 	const FString Result = FString::Printf( TEXT( "%s.%s" ), *IntString, *FracString.Mid(2) ); 
