[UE-18631] Fix inconsistent slider behavior on CapsuleComponent shape properties in the Blueprint editor.

change summary:
- moved CapsuleHalfHeight clamping from UCapsuleComponent::Serialize() to PostLoad(); this avoids a situation where we can clamp it to a yet-unclamped CapsuleRadius value during non-interactive property edits on the CDO.

diff --git a/Engine/Source/Runtime/Engine/Classes/Components/CapsuleComponent.h b/Engine/Source/Runtime/Engine/Classes/Components/CapsuleComponent.h
--- a/Engine/Source/Runtime/Engine/Classes/Components/CapsuleComponent.h
+++ b/Engine/Source/Runtime/Engine/Classes/Components/CapsuleComponent.h
@@ -61,6 +61,7 @@ public:
 
 	// Begin UObject interface
 	virtual void Serialize(FArchive& Ar) override;
+	virtual void PostLoad() override;
 #if WITH_EDITOR
 	virtual void PostEditChangeProperty(FPropertyChangedEvent& PropertyChangedEvent) override;
 #endif // WITH_EDITOR
diff --git a/Engine/Source/Runtime/Engine/Private/Components/CapsuleComponent.cpp b/Engine/Source/Runtime/Engine/Private/Components/CapsuleComponent.cpp
--- a/Engine/Source/Runtime/Engine/Private/Components/CapsuleComponent.cpp
+++ b/Engine/Source/Runtime/Engine/Private/Components/CapsuleComponent.cpp
@@ -109,8 +109,17 @@ void UCapsuleComponent::Serialize(FArchive& Ar)
 			CapsuleHeight_DEPRECATED = 0.0f;
 		}
 	}
+}
+
+void UCapsuleComponent::PostLoad()
+{
+	Super::PostLoad();
 
-	CapsuleHalfHeight = FMath::Max3(0.f, CapsuleHalfHeight, CapsuleRadius);
+	// Ensure this value is clamped only in the case where we're not re-running construction scripts.
+	if(!GIsReconstructingBlueprintInstances)
+	{
+		CapsuleHalfHeight = FMath::Max3(0.f, CapsuleHalfHeight, CapsuleRadius);
+	}
 }
 
 #if WITH_EDITOR
