Convert HidePin/DefaultToSelf cases to WorldContext

diff --git a/Engine/Source/Runtime/AIModule/Classes/Actions/PawnAction.h b/Engine/Source/Runtime/AIModule/Classes/Actions/PawnAction.h
--- a/Engine/Source/Runtime/AIModule/Classes/Actions/PawnAction.h
+++ b/Engine/Source/Runtime/AIModule/Classes/Actions/PawnAction.h
@@ -197,7 +197,7 @@ public:
 	UFUNCTION(BlueprintPure, Category = "AI|PawnActions")
 	TEnumAsByte<EAIRequestPriority::Type> GetActionPriority();
 
-	UFUNCTION(BlueprintCallable, Category = "AI|PawnActions", meta = (HidePin = "WorldContextObject", DefaultToSelf = "WorldContextObject"))
+	UFUNCTION(BlueprintCallable, Category = "AI|PawnActions", meta = (WorldContext="WorldContextObject"))
 	static UPawnAction* CreateActionInstance(UObject* WorldContextObject, TSubclassOf<UPawnAction> ActionClass);
 
 	//----------------------------------------------------------------------//
diff --git a/Engine/Source/Runtime/AIModule/Classes/Perception/AIPerceptionSystem.h b/Engine/Source/Runtime/AIModule/Classes/Perception/AIPerceptionSystem.h
--- a/Engine/Source/Runtime/AIModule/Classes/Perception/AIPerceptionSystem.h
+++ b/Engine/Source/Runtime/AIModule/Classes/Perception/AIPerceptionSystem.h
@@ -121,7 +121,7 @@ public:
 	UFUNCTION(BlueprintCallable, Category = "AI|Perception")
 	void ReportEvent(UAISenseEvent* PerceptionEvent);
 
-	UFUNCTION(BlueprintCallable, Category = "AI|Perception", meta = (HidePin = "WorldContext", DefaultToSelf = "WorldContext"))
+	UFUNCTION(BlueprintCallable, Category = "AI|Perception", meta = (WorldContext="WorldContext"))
 	static void ReportPerceptionEvent(UObject* WorldContext, UAISenseEvent* PerceptionEvent);
 
 	template<typename FSenseClass>
@@ -140,12 +140,12 @@ public:
 
 	static void MakeNoiseImpl(AActor* NoiseMaker, float Loudness, APawn* NoiseInstigator, const FVector& NoiseLocation, float MaxRange, FName Tag);
 
-	UFUNCTION(BlueprintCallable, Category = "AI|Perception", meta = (HidePin = "WorldContext", DefaultToSelf = "WorldContext"))
+	UFUNCTION(BlueprintCallable, Category = "AI|Perception", meta = (WorldContext="WorldContext"))
 	static bool RegisterPerceptionStimuliSource(UObject* WorldContext, TSubclassOf<UAISense> Sense, AActor* Target);
 
 	FAISenseID RegisterSenseClass(TSubclassOf<UAISense> SenseClass);
 
-	UFUNCTION(BlueprintCallable, Category = "AI|Perception", meta = (HidePin = "WorldContext", DefaultToSelf = "WorldContext"))
+	UFUNCTION(BlueprintCallable, Category = "AI|Perception", meta = (WorldContext="WorldContext"))
 	static TSubclassOf<UAISense> GetSenseClassForStimulus(UObject* WorldContext, const FAIStimulus& Stimulus);
 	
 protected:
diff --git a/Engine/Source/Runtime/AIModule/Classes/Perception/AISense_Damage.h b/Engine/Source/Runtime/AIModule/Classes/Perception/AISense_Damage.h
--- a/Engine/Source/Runtime/AIModule/Classes/Perception/AISense_Damage.h
+++ b/Engine/Source/Runtime/AIModule/Classes/Perception/AISense_Damage.h
@@ -63,7 +63,7 @@ public:
 	virtual void RegisterWrappedEvent(UAISenseEvent& PerceptionEvent) override;
 
 	/** EventLocation will be reported as Instigator's location at the moment of event happening*/
-	UFUNCTION(BlueprintCallable, Category = "AI|Perception", meta = (HidePin = "WorldContext", DefaultToSelf = "WorldContext", AdvancedDisplay="HitLocation"))
+	UFUNCTION(BlueprintCallable, Category = "AI|Perception", meta = (WorldContext="WorldContext", AdvancedDisplay="HitLocation"))
 	static void ReportDamageEvent(UObject* WorldContext, AActor* DamagedActor, AActor* Instigator, float DamageAmount, FVector EventLocation, FVector HitLocation);
 
 protected:
diff --git a/Engine/Source/Runtime/AIModule/Classes/Perception/AISense_Hearing.h b/Engine/Source/Runtime/AIModule/Classes/Perception/AISense_Hearing.h
--- a/Engine/Source/Runtime/AIModule/Classes/Perception/AISense_Hearing.h
+++ b/Engine/Source/Runtime/AIModule/Classes/Perception/AISense_Hearing.h
@@ -97,7 +97,7 @@ public:
 	 * @param MaxRange Max range at which the sound can be heard, multiplied by Loudness. Values <= 0 mean no limit (still limited by listener's range however).
 	 * @param Tag Identifier for the event.
 	 */
-	UFUNCTION(BlueprintCallable, Category = "AI|Perception", meta = (HidePin = "WorldContext", DefaultToSelf = "WorldContext"))
+	UFUNCTION(BlueprintCallable, Category = "AI|Perception", meta = (WorldContext="WorldContext"))
 	static void ReportNoiseEvent(UObject* WorldContext, FVector NoiseLocation, float Loudness = 1.f, AActor* Instigator = nullptr, float MaxRange = 0.f, FName Tag = NAME_None);
 
 protected:
diff --git a/Engine/Source/Runtime/Engine/Classes/AI/Navigation/NavigationSystem.h b/Engine/Source/Runtime/Engine/Classes/AI/Navigation/NavigationSystem.h
--- a/Engine/Source/Runtime/Engine/Classes/AI/Navigation/NavigationSystem.h
+++ b/Engine/Source/Runtime/Engine/Classes/AI/Navigation/NavigationSystem.h
@@ -292,13 +292,13 @@ public:
 
 	/** Finds path instantly, in a FindPath Synchronously. 
 	 *	@param PathfindingContext could be one of following: NavigationData (like Navmesh actor), Pawn or Controller. This parameter determines parameters of specific pathfinding query */
-	UFUNCTION(BlueprintCallable, Category = "AI|Navigation", meta = (HidePin = "WorldContext", DefaultToSelf = "WorldContext"))
+	UFUNCTION(BlueprintCallable, Category = "AI|Navigation", meta = (WorldContext="WorldContext"))
 	static UNavigationPath* FindPathToLocationSynchronously(UObject* WorldContext, const FVector& PathStart, const FVector& PathEnd, AActor* PathfindingContext = NULL, TSubclassOf<UNavigationQueryFilter> FilterClass = NULL);
 
 	/** Finds path instantly, in a FindPath Synchronously. Main advantage over FindPathToLocationSynchronously is that 
 	 *	the resulting path with automatically get updated if goal actor moves more then TetherDistance away from last path node
 	 *	@param PathfindingContext could be one of following: NavigationData (like Navmesh actor), Pawn or Controller. This parameter determines parameters of specific pathfinding query */
-	UFUNCTION(BlueprintCallable, Category = "AI|Navigation", meta = (HidePin = "WorldContext", DefaultToSelf = "WorldContext"))
+	UFUNCTION(BlueprintCallable, Category = "AI|Navigation", meta = (WorldContext="WorldContext"))
 	static UNavigationPath* FindPathToActorSynchronously(UObject* WorldContext, const FVector& PathStart, AActor* GoalActor, float TetherDistance = 50.f, AActor* PathfindingContext = NULL, TSubclassOf<UNavigationQueryFilter> FilterClass = NULL);
 
 	/** Performs navigation raycast on NavigationData appropriate for given Querier.
diff --git a/Engine/Source/Runtime/Engine/Classes/Kismet/GameplayStatics.h b/Engine/Source/Runtime/Engine/Classes/Kismet/GameplayStatics.h
--- a/Engine/Source/Runtime/Engine/Classes/Kismet/GameplayStatics.h
+++ b/Engine/Source/Runtime/Engine/Classes/Kismet/GameplayStatics.h
@@ -617,7 +617,7 @@ class ENGINE_API UGameplayStatics : public UBlueprintFunctionLibrary
 	static bool DeleteGameInSlot(const FString& SlotName, const int32 UserIndex);
 
 	/** Returns the frame delta time in seconds adjusted by e.g. time dilation. */
-	UFUNCTION(BlueprintPure, Category = "Utilities|Time", meta = (HidePin = "WorldContextObject", DefaultToSelf = "WorldContextObject"))
+	UFUNCTION(BlueprintPure, Category = "Utilities|Time", meta = (WorldContext="WorldContextObject"))
 	static float GetWorldDeltaSeconds(UObject* WorldContextObject);
 
 	/** Returns time in seconds since world was brought up for play, does NOT stop when game pauses, NOT dilated/clamped */
diff --git a/Engine/Source/Runtime/Engine/Classes/Kismet/KismetSystemLibrary.h b/Engine/Source/Runtime/Engine/Classes/Kismet/KismetSystemLibrary.h
--- a/Engine/Source/Runtime/Engine/Classes/Kismet/KismetSystemLibrary.h
+++ b/Engine/Source/Runtime/Engine/Classes/Kismet/KismetSystemLibrary.h
@@ -844,7 +844,7 @@ class ENGINE_API UKismetSystemLibrary : public UBlueprintFunctionLibrary
 	* @param OutHit			Properties of the trace hit.
 	* @return				True if there was a hit, false otherwise.
 	*/
-	UFUNCTION(BlueprintCallable, Category = "Collision", meta = (bIgnoreSelf = "true", HidePin = "WorldContextObject", DefaultToSelf = "WorldContextObject", AutoCreateRefTerm = "ActorsToIgnore", DisplayName = "BoxTraceByChannel", Keywords="sweep"))
+	UFUNCTION(BlueprintCallable, Category = "Collision", meta = (bIgnoreSelf = "true", WorldContext="WorldContextObject", AutoCreateRefTerm = "ActorsToIgnore", DisplayName = "BoxTraceByChannel", Keywords="sweep"))
 	static bool BoxTraceSingle(UObject* WorldContextObject, const FVector Start, const FVector End, const FVector HalfSize, const FRotator Orientation, ETraceTypeQuery TraceChannel, bool bTraceComplex, const TArray<AActor*>& ActorsToIgnore, EDrawDebugTrace::Type DrawDebugType, FHitResult& OutHit, bool bIgnoreSelf);
 
 	/**
@@ -860,7 +860,7 @@ class ENGINE_API UKismetSystemLibrary : public UBlueprintFunctionLibrary
 	* @param OutHits		A list of hits, sorted along the trace from start to finish. The blocking hit will be the last hit, if there was one.
 	* @return				True if there was a blocking hit, false otherwise.
 	*/
-	UFUNCTION(BlueprintCallable, Category = "Collision", meta = (bIgnoreSelf = "true", HidePin = "WorldContextObject", DefaultToSelf = "WorldContextObject", AutoCreateRefTerm = "ActorsToIgnore", DisplayName = "MultiBoxTraceByChannel", Keywords="sweep"))
+	UFUNCTION(BlueprintCallable, Category = "Collision", meta = (bIgnoreSelf = "true", WorldContext="WorldContextObject", AutoCreateRefTerm = "ActorsToIgnore", DisplayName = "MultiBoxTraceByChannel", Keywords="sweep"))
 	static bool BoxTraceMulti(UObject* WorldContextObject, const FVector Start, const FVector End, FVector HalfSize, const FRotator Orientation, ETraceTypeQuery TraceChannel, bool bTraceComplex, const TArray<AActor*>& ActorsToIgnore, EDrawDebugTrace::Type DrawDebugType, TArray<FHitResult>& OutHits, bool bIgnoreSelf);
 
 	
@@ -973,7 +973,7 @@ class ENGINE_API UKismetSystemLibrary : public UBlueprintFunctionLibrary
 	* @param OutHit			Properties of the trace hit.
 	* @return				True if there was a hit, false otherwise.
 	*/
-	UFUNCTION(BlueprintCallable, Category = "Collision", meta = (bIgnoreSelf = "true", HidePin = "WorldContextObject", DefaultToSelf = "WorldContextObject", AutoCreateRefTerm = "ActorsToIgnore", DisplayName = "BoxTraceForObjects", Keywords="sweep"))
+	UFUNCTION(BlueprintCallable, Category = "Collision", meta = (bIgnoreSelf = "true", WorldContext="WorldContextObject", AutoCreateRefTerm = "ActorsToIgnore", DisplayName = "BoxTraceForObjects", Keywords="sweep"))
 	static bool BoxTraceSingleForObjects(UObject* WorldContextObject, const FVector Start, const FVector End, const FVector HalfSize, const FRotator Orientation, const TArray<TEnumAsByte<EObjectTypeQuery> > & ObjectTypes, bool bTraceComplex, const TArray<AActor*>& ActorsToIgnore, EDrawDebugTrace::Type DrawDebugType, FHitResult& OutHit, bool bIgnoreSelf);
 
 
@@ -990,7 +990,7 @@ class ENGINE_API UKismetSystemLibrary : public UBlueprintFunctionLibrary
 	* @param OutHits		A list of hits, sorted along the trace from start to finish.  The blocking hit will be the last hit, if there was one.
 	* @return				True if there was a hit, false otherwise.
 	*/
-	UFUNCTION(BlueprintCallable, Category = "Collision", meta = (bIgnoreSelf = "true", HidePin = "WorldContextObject", DefaultToSelf = "WorldContextObject", AutoCreateRefTerm = "ActorsToIgnore", DisplayName = "MultiBoxTraceForObjects", Keywords="sweep"))
+	UFUNCTION(BlueprintCallable, Category = "Collision", meta = (bIgnoreSelf = "true", WorldContext="WorldContextObject", AutoCreateRefTerm = "ActorsToIgnore", DisplayName = "MultiBoxTraceForObjects", Keywords="sweep"))
 	static bool BoxTraceMultiForObjects(UObject* WorldContextObject, const FVector Start, const FVector End, const FVector HalfSize, const FRotator Orientation, const TArray<TEnumAsByte<EObjectTypeQuery> > & ObjectTypes, bool bTraceComplex, const TArray<AActor*>& ActorsToIgnore, EDrawDebugTrace::Type DrawDebugType, TArray<FHitResult>& OutHits, bool bIgnoreSelf);
 
 	/**
@@ -1433,7 +1433,7 @@ class ENGINE_API UKismetSystemLibrary : public UBlueprintFunctionLibrary
 	 * @param WorldContextObject	World context
 	 * @param State					set true to supress transition message
 	 */
-	UFUNCTION(BlueprintCallable, Category = "Utilities", meta = (HidePin = "WorldContextObject", DefaultToSelf = "WorldContextObject"))
+	UFUNCTION(BlueprintCallable, Category = "Utilities", meta = (WorldContext="WorldContextObject"))
 	static void SetSupressViewportTransitionMessage(UObject* WorldContextObject, bool bState);
 
 	/**
diff --git a/Engine/Source/Runtime/GameplayAbilities/Public/Abilities/GameplayAbility.h b/Engine/Source/Runtime/GameplayAbilities/Public/Abilities/GameplayAbility.h
--- a/Engine/Source/Runtime/GameplayAbilities/Public/Abilities/GameplayAbility.h
+++ b/Engine/Source/Runtime/GameplayAbilities/Public/Abilities/GameplayAbility.h
@@ -749,10 +749,10 @@ protected:
 	//
 	// ----------------------------------------------------------------------------------------------------------------
 
-	UFUNCTION(BlueprintPure, Category = Ability, meta = (HidePin = "WorldContextObject", DefaultToSelf = "WorldContextObject"))
+	UFUNCTION(BlueprintPure, Category = Ability)
 	FGameplayAbilityTargetingLocationInfo MakeTargetLocationInfoFromOwnerActor();
 
-	UFUNCTION(BlueprintPure, Category = Ability, meta = (HidePin = "WorldContextObject", DefaultToSelf = "WorldContextObject"))
+	UFUNCTION(BlueprintPure, Category = Ability)
 	FGameplayAbilityTargetingLocationInfo MakeTargetLocationInfoFromOwnerSkeletalMeshComponent(FName SocketName);
 
 	// ----------------------------------------------------------------------------------------------------------------
diff --git a/Engine/Source/Runtime/GameplayTasks/Classes/Tasks/GameplayTask_SpawnActor.h b/Engine/Source/Runtime/GameplayTasks/Classes/Tasks/GameplayTask_SpawnActor.h
--- a/Engine/Source/Runtime/GameplayTasks/Classes/Tasks/GameplayTask_SpawnActor.h
+++ b/Engine/Source/Runtime/GameplayTasks/Classes/Tasks/GameplayTask_SpawnActor.h
@@ -34,10 +34,10 @@ public:
 	UFUNCTION(BlueprintCallable, Category = "GameplayTasks", meta = (DisplayName="Spawn Actor for Gameplay Task", AdvancedDisplay = "TaskOwner, bSpawnOnlyOnAuthority", DefaultToSelf = "TaskOwner", BlueprintInternalUseOnly = "TRUE"))
 	static UGameplayTask_SpawnActor* SpawnActor(TScriptInterface<IGameplayTaskOwnerInterface> TaskOwner, FVector SpawnLocation, FRotator SpawnRotation, TSubclassOf<AActor> Class, bool bSpawnOnlyOnAuthority = false);
 
-	UFUNCTION(BlueprintCallable, meta = (HidePin = "WorldContextObject", DefaultToSelf = "WorldContextObject", BlueprintInternalUseOnly = "true"), Category = "GameplayTasks")
+	UFUNCTION(BlueprintCallable, meta = (WorldContext="WorldContextObject", BlueprintInternalUseOnly = "true"), Category = "GameplayTasks")
 	virtual bool BeginSpawningActor(UObject* WorldContextObject, AActor*& SpawnedActor);
 
-	UFUNCTION(BlueprintCallable, meta = (HidePin = "WorldContextObject", DefaultToSelf = "WorldContextObject", BlueprintInternalUseOnly = "true"), Category = "GameplayTasks")
+	UFUNCTION(BlueprintCallable, meta = (WorldContext="WorldContextObject", BlueprintInternalUseOnly = "true"), Category = "GameplayTasks")
 	virtual void FinishSpawningActor(UObject* WorldContextObject, AActor* SpawnedActor);
 
 protected:
