fixed bounds of "standalone" nav relevant components (without navigation parent), which caused them to be ignored during navmesh generation
#ue4

diff --git a/Engine/Source/Runtime/Engine/Classes/AI/Navigation/NavRelevantComponent.h b/Engine/Source/Runtime/Engine/Classes/AI/Navigation/NavRelevantComponent.h
--- a/Engine/Source/Runtime/Engine/Classes/AI/Navigation/NavRelevantComponent.h
+++ b/Engine/Source/Runtime/Engine/Classes/AI/Navigation/NavRelevantComponent.h
@@ -48,6 +48,8 @@ protected:
 	UPROPERTY()
 	uint32 bAttachToOwnersRoot : 1;
 
+	uint32 bBoundsInitialized : 1;
+
 	UPROPERTY(transient)
 	UObject* CachedNavParent;
 };
diff --git a/Engine/Source/Runtime/Engine/Private/AI/Navigation/NavRelevantComponent.cpp b/Engine/Source/Runtime/Engine/Private/AI/Navigation/NavRelevantComponent.cpp
--- a/Engine/Source/Runtime/Engine/Private/AI/Navigation/NavRelevantComponent.cpp
+++ b/Engine/Source/Runtime/Engine/Private/AI/Navigation/NavRelevantComponent.cpp
@@ -7,6 +7,7 @@ UNavRelevantComponent::UNavRelevantComponent(const FObjectInitializer& ObjectIni
 {
 	bNavigationRelevant = true;
 	bAttachToOwnersRoot = true;
+	bBoundsInitialized = false;
 }
 
 void UNavRelevantComponent::OnRegister()
@@ -47,6 +48,11 @@ void UNavRelevantComponent::OnRegister()
 			}
 		}
 	}
+	else if (!bBoundsInitialized)
+	{
+		bBoundsInitialized = true;
+		UpdateNavigationBounds();
+	}
 
 	UNavigationSystem::OnComponentRegistered(this);
 }
