[NO NON-PII COMMIT MESSAGE]

diff --git a/Engine/Source/Editor/BlueprintGraph/Private/K2Node_GenericCreateObject.cpp b/Engine/Source/Editor/BlueprintGraph/Private/K2Node_GenericCreateObject.cpp
--- a/Engine/Source/Editor/BlueprintGraph/Private/K2Node_GenericCreateObject.cpp
+++ b/Engine/Source/Editor/BlueprintGraph/Private/K2Node_GenericCreateObject.cpp
@@ -72,9 +72,10 @@ void UK2Node_GenericCreateObject::ValidateNodeDuringCompilation(class FCompilerR
 void UK2Node_GenericCreateObject::EarlyValidation(class FCompilerResultsLog& MessageLog) const
 {
 	Super::EarlyValidation(MessageLog);
-
+	auto ClassPin = GetClassPin(&Pins);
+	const bool bAllowAbstract = ClassPin && ClassPin->LinkedTo.Num();
 	auto ClassToSpawn = GetClassToSpawn();
-	if (!UGameplayStatics::CanSpawnObjectOfClass(ClassToSpawn))
+	if (!UGameplayStatics::CanSpawnObjectOfClass(ClassToSpawn, bAllowAbstract))
 	{
 		MessageLog.Error(*FString::Printf(*LOCTEXT("GenericCreateObject_WrongClass", "Wrong class to spawn '%s' in @@").ToString(), *GetPathNameSafe(ClassToSpawn)), this);
 	}
diff --git a/Engine/Source/Runtime/Engine/Classes/Kismet/GameplayStatics.h b/Engine/Source/Runtime/Engine/Classes/Kismet/GameplayStatics.h
--- a/Engine/Source/Runtime/Engine/Classes/Kismet/GameplayStatics.h
+++ b/Engine/Source/Runtime/Engine/Classes/Kismet/GameplayStatics.h
@@ -34,7 +34,7 @@ class ENGINE_API UGameplayStatics : public UBlueprintFunctionLibrary
 	UFUNCTION(BlueprintCallable, Category = "Spawning", meta = (BlueprintInternalUseOnly = "true"))
 	static UObject* SpawnObject(TSubclassOf<UObject> ObjectClass, UObject* Outer);
 
-	static bool CanSpawnObjectOfClass(TSubclassOf<UObject> ObjectClass);
+	static bool CanSpawnObjectOfClass(TSubclassOf<UObject> ObjectClass, bool bAllowAbstract = false);
 
 	// --- Spawning functions ------------------------------
 
diff --git a/Engine/Source/Runtime/Engine/Private/GameplayStatics.cpp b/Engine/Source/Runtime/Engine/Private/GameplayStatics.cpp
--- a/Engine/Source/Runtime/Engine/Private/GameplayStatics.cpp
+++ b/Engine/Source/Runtime/Engine/Private/GameplayStatics.cpp
@@ -303,7 +303,7 @@ void UGameplayStatics::ApplyDamage(AActor* DamagedActor, float BaseDamage, ACont
 	}
 }
 
-bool UGameplayStatics::CanSpawnObjectOfClass(TSubclassOf<UObject> ObjectClass)
+bool UGameplayStatics::CanSpawnObjectOfClass(TSubclassOf<UObject> ObjectClass, bool bAllowAbstract)
 {
 	bool bBlueprintType = true;
 #if WITH_EDITOR
@@ -347,7 +347,8 @@ bool UGameplayStatics::CanSpawnObjectOfClass(TSubclassOf<UObject> ObjectClass)
 	return (nullptr != *ObjectClass)
 		&& bBlueprintType
 		&& !bForbiddenSpawn
-		&& !ObjectClass->HasAnyClassFlags(CLASS_Abstract | CLASS_Deprecated | CLASS_NewerVersionExists)
+		&& (bAllowAbstract || !ObjectClass->HasAnyClassFlags(CLASS_Abstract))
+		&& !ObjectClass->HasAnyClassFlags(CLASS_Deprecated | CLASS_NewerVersionExists)
 		&& !ObjectClass->IsChildOf(AActor::StaticClass())
 		&& !ObjectClass->IsChildOf(UActorComponent::StaticClass());
 }
