Animation: Fix exclusive GetResourceSize reporting in UAnimSequence

diff --git a/Engine/Source/Runtime/Engine/Private/Animation/AnimSequence.cpp b/Engine/Source/Runtime/Engine/Private/Animation/AnimSequence.cpp
--- a/Engine/Source/Runtime/Engine/Private/Animation/AnimSequence.cpp
+++ b/Engine/Source/Runtime/Engine/Private/Animation/AnimSequence.cpp
@@ -122,8 +122,15 @@ UAnimSequence::UAnimSequence(const FObjectInitializer& ObjectInitializer)
 
 SIZE_T UAnimSequence::GetResourceSize(EResourceSizeMode::Type Mode)
 {
-	const int32 ResourceSize = CompressedTrackOffsets.Num() == 0 ? GetApproxRawSize() : GetApproxCompressedSize();
-	return ResourceSize;
+	if (Mode == EResourceSizeMode::Exclusive)
+	{
+		// All of the sequence data is serialized and will be counted as part of the direct object size rather than as a resource
+		return 0;
+	}
+	else
+	{
+		return (CompressedTrackOffsets.Num() == 0) ? GetApproxRawSize() : GetApproxCompressedSize();
+	}
 }
 
 void UAnimSequence::GetAssetRegistryTags(TArray<FAssetRegistryTag>& OutTags) const
