Animation: Fix exclusive GetResourceSize reporting in USkeletalMesh

diff --git a/Engine/Source/Runtime/Engine/Private/SkeletalMesh.cpp b/Engine/Source/Runtime/Engine/Private/SkeletalMesh.cpp
--- a/Engine/Source/Runtime/Engine/Private/SkeletalMesh.cpp
+++ b/Engine/Source/Runtime/Engine/Private/SkeletalMesh.cpp
@@ -2141,23 +2141,23 @@ float USkeletalMesh::GetStreamingTextureFactor( int32 RequestedUVIndex )
 SIZE_T USkeletalMesh::GetResourceSize(EResourceSizeMode::Type Mode)
 {
 	SIZE_T ResourceSize = 0;
-	if(ImportedResource.IsValid())
+	if (ImportedResource.IsValid())
 	{
 		ResourceSize += ImportedResource->GetResourceSize();
 	}
 
-	for (const auto& MorphTarget : MorphTargets)
+	if (Mode == EResourceSizeMode::Inclusive)
 	{
-		ResourceSize += MorphTarget->GetResourceSize(Mode);
-	}
+		for (const auto& MorphTarget : MorphTargets)
+		{
+			ResourceSize += MorphTarget->GetResourceSize(Mode);
+		}
 
-	for (const auto& ClothingAsset : ClothingAssets)
-	{
-		ResourceSize += ClothingAsset.GetResourceSize();
-	}
+		for (const auto& ClothingAsset : ClothingAssets)
+		{
+			ResourceSize += ClothingAsset.GetResourceSize();
+		}
 
-	if(Mode == EResourceSizeMode::Inclusive)
-	{
 		TSet<UMaterialInterface*> UniqueMaterials;
 		for(int32 MaterialIndex = 0; MaterialIndex < Materials.Num(); ++MaterialIndex)
 		{
@@ -2177,7 +2177,7 @@ SIZE_T USkeletalMesh::GetResourceSize(EResourceSizeMode::Type Mode)
 		ResourceSize += RefBasesInvMatrix.GetAllocatedSize();
 		ResourceSize += RefSkeleton.GetDataSize();
 
-		if(BodySetup)
+		if (BodySetup)
 		{
 			ResourceSize += BodySetup->GetResourceSize(Mode);
 		}
