Fix material errors using Custom nodes with LandscapeGrassOutput or other CustomOutput node types.

diff --git a/Engine/Source/Runtime/Engine/Private/Materials/HLSLMaterialTranslator.h b/Engine/Source/Runtime/Engine/Private/Materials/HLSLMaterialTranslator.h
--- a/Engine/Source/Runtime/Engine/Private/Materials/HLSLMaterialTranslator.h
+++ b/Engine/Source/Runtime/Engine/Private/Materials/HLSLMaterialTranslator.h
@@ -427,6 +427,41 @@ public:
 
 			ResourcesString = TEXT("");
 
+			// Gather the implementation for any custom output expressions
+			{
+				TArray<UMaterialExpressionCustomOutput*> CustomOutputExpressions;
+				Material->GatherCustomOutputExpressions(CustomOutputExpressions);
+				TSet<UClass*> SeenCustomOutputExpressionsClases;
+
+				for (UMaterialExpressionCustomOutput* CustomOutput : CustomOutputExpressions)
+				{
+					if (SeenCustomOutputExpressionsClases.Contains(CustomOutput->GetClass()))
+					{
+						Errorf(TEXT("The material can contain only one %s node"), *CustomOutput->GetDescription());
+					}
+					else
+					{
+						SeenCustomOutputExpressionsClases.Add(CustomOutput->GetClass());
+
+						int32 NumOutputs = CustomOutput->GetNumOutputs();
+						ResourcesString += FString::Printf(TEXT("#define NUM_MATERIAL_OUTPUTS_%s %d\r\n"), *CustomOutput->GetFunctionName().ToUpper(), NumOutputs);
+						if (NumOutputs > 0)
+						{
+							for (int32 Index = 0; Index < NumOutputs; Index++)
+							{
+								FunctionStack.Empty();
+								FunctionStack.Add(FMaterialFunctionCompileState(NULL));
+								MaterialProperty = MP_MAX; // Indicates we're not compiling any material property.
+								ShaderFrequency = SF_Pixel;
+								TArray<FShaderCodeChunk> CustomExpressionChunks;
+								CurrentScopeChunks = &CustomExpressionChunks; //-V506
+								CustomOutput->Compile(this, Index, 0);
+							}
+						}
+					}
+				}
+			}
+
 			// Output the implementation for any custom expressions we will call below.
 			for(int32 ExpressionIndex = 0;ExpressionIndex < CustomExpressionImplementations.Num();ExpressionIndex++)
 			{
@@ -488,39 +523,7 @@ public:
 				}
 			}
 
-			// Gather and output the implementation for any custom output expressions
-			TArray<UMaterialExpressionCustomOutput*> CustomOutputExpressions;
-			Material->GatherCustomOutputExpressions(CustomOutputExpressions);
-			TSet<UClass*> SeenCustomOutputExpressionsClases;
-
-			for(UMaterialExpressionCustomOutput* CustomOutput : CustomOutputExpressions)
-			{
-				if(SeenCustomOutputExpressionsClases.Contains(CustomOutput->GetClass()))
-				{
-					Errorf(TEXT("The material can contain only one %s node"), *CustomOutput->GetDescription());
-				}
-				else
-				{
-					SeenCustomOutputExpressionsClases.Add(CustomOutput->GetClass());
-
-					int32 NumOutputs = CustomOutput->GetNumOutputs();
-					ResourcesString += FString::Printf(TEXT("#define NUM_MATERIAL_OUTPUTS_%s %d\r\n"), *CustomOutput->GetFunctionName().ToUpper(), NumOutputs);
-					if (NumOutputs > 0)
-					{
-						for (int32 Index = 0; Index < NumOutputs; Index++)
-						{
-							FunctionStack.Empty();
-							FunctionStack.Add(FMaterialFunctionCompileState(NULL));
-							MaterialProperty = MP_MAX; // Indicates we're not compiling any material property.
-							ShaderFrequency = SF_Pixel;
-							TArray<FShaderCodeChunk> CustomExpressionChunks;
-							CurrentScopeChunks = &CustomExpressionChunks; //-V506
-							CustomOutput->Compile(this, Index, 0);
-						}
-					}
-				}
-			}
-
+			// Output the implementation for any custom output expressions
 			for (int32 ExpressionIndex = 0; ExpressionIndex < CustomOutputImplementations.Num(); ExpressionIndex++)
 			{
 				ResourcesString += CustomOutputImplementations[ExpressionIndex] + "\r\n\r\n";
