Fix parameters connected to LandscapeGrass or other CustomOutput nodes not appearing in the Material Instance Constant editor

diff --git a/Engine/Source/Editor/MaterialEditor/Private/MaterialEditorUtilities.cpp b/Engine/Source/Editor/MaterialEditor/Private/MaterialEditorUtilities.cpp
--- a/Engine/Source/Editor/MaterialEditor/Private/MaterialEditorUtilities.cpp
+++ b/Engine/Source/Editor/MaterialEditor/Private/MaterialEditorUtilities.cpp
@@ -16,6 +16,7 @@
 #include "Materials/MaterialExpressionVectorParameter.h"
 #include "Materials/MaterialExpressionStaticSwitchParameter.h"
 #include "Materials/MaterialFunction.h"
+#include "Materials/MaterialExpressionCustomOutput.h"
 
 #include "MaterialEditorUtilities.h"
 #include "Toolkits/ToolkitManager.h"
@@ -212,6 +213,13 @@ void FMaterialEditorUtilities::GetVisibleMaterialParameters(const UMaterial* Mat
 			GetVisibleMaterialParametersFromExpression(FMaterialExpressionKey(ExpressionInput->Expression, ExpressionInput->OutputIndex), MaterialInstance, VisibleExpressions, FunctionStack);
 		}
 	}
+
+	TArray<UMaterialExpressionCustomOutput*> CustomOutputExpressions;
+	Material->GetAllCustomOutputExpressions(CustomOutputExpressions);
+	for (UMaterialExpressionCustomOutput* Expression : CustomOutputExpressions)
+	{
+		GetVisibleMaterialParametersFromExpression(FMaterialExpressionKey(Expression, 0), MaterialInstance, VisibleExpressions, FunctionStack);
+	}
 }
 
 bool FMaterialEditorUtilities::GetStaticSwitchExpressionValue(UMaterialInstance* MaterialInstance, UMaterialExpression* SwitchValueExpression, bool& OutValue, FGuid& OutExpressionID, TArray<FGetVisibleMaterialParametersFunctionState*>& FunctionStack)
diff --git a/Engine/Source/Runtime/Engine/Classes/Materials/Material.h b/Engine/Source/Runtime/Engine/Classes/Materials/Material.h
--- a/Engine/Source/Runtime/Engine/Classes/Materials/Material.h
+++ b/Engine/Source/Runtime/Engine/Classes/Materials/Material.h
@@ -1213,6 +1213,9 @@ public:
 	*/
 	ENGINE_API FExpressionInput* GetExpressionInputForProperty(EMaterialProperty InProperty);
 
+	/* Returns any UMaterialExpressionCustomOutput expressions */
+	ENGINE_API void GetAllCustomOutputExpressions(TArray<class UMaterialExpressionCustomOutput*>& OutCustomOutputs) const;
+
 	/**
 	 *	Get all referenced expressions (returns the chains for all properties).
 	 *
diff --git a/Engine/Source/Runtime/Engine/Private/Materials/Material.cpp b/Engine/Source/Runtime/Engine/Private/Materials/Material.cpp
--- a/Engine/Source/Runtime/Engine/Private/Materials/Material.cpp
+++ b/Engine/Source/Runtime/Engine/Private/Materials/Material.cpp
@@ -122,14 +122,7 @@ int32 FMaterialResource::CompilePropertyAndSetMaterialProperty(EMaterialProperty
 
 void FMaterialResource::GatherCustomOutputExpressions(TArray<UMaterialExpressionCustomOutput*>& OutCustomOutputs) const
 {
-	for (UMaterialExpression* Expression : Material->Expressions)
-	{
-		UMaterialExpressionCustomOutput* CustomOutput = Cast<UMaterialExpressionCustomOutput>(Expression);
-		if (CustomOutput)
-		{
-			OutCustomOutputs.Add(CustomOutput);
-		}
-	}
+	Material->GetAllCustomOutputExpressions(OutCustomOutputs);
 }
 
 void FMaterialResource::GetShaderMapId(EShaderPlatform Platform, FMaterialShaderMapId& OutId) const
@@ -3422,6 +3415,17 @@ FExpressionInput* UMaterial::GetExpressionInputForProperty(EMaterialProperty InP
 	return nullptr;
 }
 
+void UMaterial::GetAllCustomOutputExpressions(TArray<class UMaterialExpressionCustomOutput*>& OutCustomOutputs) const
+{
+	for (UMaterialExpression* Expression : Expressions)
+	{
+		UMaterialExpressionCustomOutput* CustomOutput = Cast<UMaterialExpressionCustomOutput>(Expression);
+		if (CustomOutput)
+		{
+			OutCustomOutputs.Add(CustomOutput);
+		}
+	}
+}
 
 bool UMaterial::GetAllReferencedExpressions(TArray<UMaterialExpression*>& OutExpressions, class FStaticParameterSet* InStaticParameterSet)
 {
@@ -3440,6 +3444,14 @@ bool UMaterial::GetAllReferencedExpressions(TArray<UMaterialExpression*>& OutExp
 		}
 	}
 
+	TArray<class UMaterialExpressionCustomOutput*> CustomOutputExpressions;
+	GetAllCustomOutputExpressions(CustomOutputExpressions);
+	for (UMaterialExpressionCustomOutput* Expression : CustomOutputExpressions)
+	{
+		TArray<FExpressionInput*> ProcessedInputs;
+		RecursiveGetExpressionChain(Expression, ProcessedInputs, OutExpressions, InStaticParameterSet);
+	}
+
 	return true;
 }
 
