Animation Compression: Defaults for the "Automatic" animation compression setting are now set to the same settings in UAnimSettings. Added comments where possible.

diff --git a/Engine/Source/Runtime/Engine/Classes/Animation/AnimCompress_Automatic.h b/Engine/Source/Runtime/Engine/Classes/Animation/AnimCompress_Automatic.h
--- a/Engine/Source/Runtime/Engine/Classes/Animation/AnimCompress_Automatic.h
+++ b/Engine/Source/Runtime/Engine/Classes/Animation/AnimCompress_Automatic.h
@@ -10,33 +10,49 @@
 #include "Animation/AnimCompress.h"
 #include "AnimCompress_Automatic.generated.h"
 
-UCLASS(hidecategories=AnimationCompressionAlgorithm)
+UCLASS(MinimalAPI)
 class UAnimCompress_Automatic : public UAnimCompress
 {
 	GENERATED_UCLASS_BODY()
 
-	/** Maximum amount of error that a compression technique can introduce in an end effector */
+	/** Maximum amount of error that a compression technique can introduce in an end effector 
+	* Determines the current setting for world - space error tolerance in the animation compressor.
+	* When requested, animation being compressed will also consider an alternative compression
+	* method if the end result of that method produces less error than the AlternativeCompressionThreshold.
+	* Also known as "Alternative Compression Threshold"
+	*/
 	UPROPERTY(EditAnywhere, Category=AnimationCompressionAlgorithm_Automatic)
 	float MaxEndEffectorError;
 
+	/** If true, the uniform bitwise techniques will be tried */
 	UPROPERTY(EditAnywhere, Category=AnimationCompressionAlgorithm_Automatic)
 	uint32 bTryFixedBitwiseCompression:1;
 
+	/** If true, the per-track compressor techniques will be tried */
 	UPROPERTY(EditAnywhere, Category=AnimationCompressionAlgorithm_Automatic)
 	uint32 bTryPerTrackBitwiseCompression:1;
 
+	/** If true, the linear key removal techniques will be tried */
 	UPROPERTY(EditAnywhere, Category=AnimationCompressionAlgorithm_Automatic)
 	uint32 bTryLinearKeyRemovalCompression:1;
 
+	/** If true, the resampling techniques will be tried */
 	UPROPERTY(EditAnywhere, Category=AnimationCompressionAlgorithm_Automatic)
 	uint32 bTryIntervalKeyRemoval:1;
 
+	/** If true, then the animation will be first recompressed with it's current compressor if non-NULL, or with the global default compressor (specified in the engine ini)
+	* Also known as "First Recompress Using Current Or Default"
+	*/
 	UPROPERTY(EditAnywhere, Category=AnimationCompressionAlgorithm_Automatic)
 	uint32 bRunCurrentDefaultCompressor:1;
 
+	/** If true and the existing compression error is greater than Max End Effector Error, then any compression technique (even one that increases the size) with a lower error will be used until it falls below the threshold
+	* Also known as "force below threshold"
+	*/
 	UPROPERTY(EditAnywhere, Category=AnimationCompressionAlgorithm_Automatic)
 	uint32 bAutoReplaceIfExistingErrorTooGreat:1;
 
+	/** If true and the existing compression error is greater than Max End Effector Error, then Max End Effector Error will be effectively raised to the existing error level */
 	UPROPERTY(EditAnywhere, Category=AnimationCompressionAlgorithm_Automatic)
 	uint32 bRaiseMaxErrorToExisting:1;
 
diff --git a/Engine/Source/Runtime/Engine/Classes/Animation/AnimationSettings.h b/Engine/Source/Runtime/Engine/Classes/Animation/AnimationSettings.h
--- a/Engine/Source/Runtime/Engine/Classes/Animation/AnimationSettings.h
+++ b/Engine/Source/Runtime/Engine/Classes/Animation/AnimationSettings.h
@@ -33,6 +33,13 @@ class ENGINE_API UAnimationSettings : public UDeveloperSettings
 	UPROPERTY(config, EditAnywhere, Category = Compression)
 	TEnumAsByte<AnimationCompressionFormat> TranslationCompressionFormat;
 
+	/** The alternate error threshold (0.0 means don't try anything other than the current / default scheme) 
+	* 
+	* Determines the current setting for world-space error tolerance in the animation compressor.
+	* When requested, animation being compressed will also consider an alternative compression
+	* method if the end result of that method produces less error than the AlternativeCompressionThreshold.
+	* Also known as "Max End Effector Error"
+	*/
 	UPROPERTY(config, EditAnywhere, Category = Compression, meta=(ClampMin = "0", UIMin = "0"))
 	float AlternativeCompressionThreshold;
 
@@ -41,25 +48,34 @@ class ENGINE_API UAnimationSettings : public UDeveloperSettings
 
 	UPROPERTY(config, EditAnywhere, Category = Compression)
 	bool bOnlyCheckForMissingSkeletalMeshes;
-	
+
+	/** If true and the existing compression error is greater than Alternative Compression Threshold, then any compression technique (even one that increases the size) with a lower error will be used until it falls below the threshold */
 	UPROPERTY(config, EditAnywhere, Category = Compression)
 	bool bForceBelowThreshold;
 
+	/** If true, then the animation will be first recompressed with it's current compressor if non-NULL, or with the global default compressor (specified in the engine ini) 
+	* Also known as "Run Current Default Compressor"
+	*/
 	UPROPERTY(config, EditAnywhere, Category = Compression)
 	bool bFirstRecompressUsingCurrentOrDefault;
 
+	/** If true and the existing compression error is greater than Alternative Compression Threshold, then Alternative Compression Threshold will be effectively raised to the existing error level */
 	UPROPERTY(config, EditAnywhere, Category = Compression)
 	bool bRaiseMaxErrorToExisting;
 
+	/** If true, the uniform bitwise techniques will be tried */
 	UPROPERTY(config, EditAnywhere, Category = Compression)
 	bool bTryFixedBitwiseCompression;
 
+	/** If true, the per-track compressor techniques will be tried */
 	UPROPERTY(config, EditAnywhere, Category = Compression)
 	bool bTryPerTrackBitwiseCompression;
 
+	/** If true, the linear key removal techniques will be tried */
 	UPROPERTY(config, EditAnywhere, Category = Compression)
 	bool bTryLinearKeyRemovalCompression;
 
+	/** If true, the resampling techniques will be tried */
 	UPROPERTY(config, EditAnywhere, Category = Compression)
 	bool bTryIntervalKeyRemoval;
 
diff --git a/Engine/Source/Runtime/Engine/Private/Animation/AnimCompress_Automatic.cpp b/Engine/Source/Runtime/Engine/Private/Animation/AnimCompress_Automatic.cpp
--- a/Engine/Source/Runtime/Engine/Private/Animation/AnimCompress_Automatic.cpp
+++ b/Engine/Source/Runtime/Engine/Private/Animation/AnimCompress_Automatic.cpp
@@ -4,19 +4,21 @@
 #include "EnginePrivate.h"
 #include "AnimationUtils.h"
 #include "Animation/AnimCompress_Automatic.h"
+#include "Animation/AnimationSettings.h"
 
 UAnimCompress_Automatic::UAnimCompress_Automatic(const FObjectInitializer& ObjectInitializer)
 	: Super(ObjectInitializer)
 {
 	Description = TEXT("Automatic");
-	MaxEndEffectorError = 1.0f;
-	bTryFixedBitwiseCompression = true;
-	bTryPerTrackBitwiseCompression = true;
-	bTryLinearKeyRemovalCompression = true;
-	bTryIntervalKeyRemoval = true;
-	bRunCurrentDefaultCompressor = false;
-	bAutoReplaceIfExistingErrorTooGreat = false;
-	bRaiseMaxErrorToExisting = false;
+	UAnimationSettings* AnimationSettings = UAnimationSettings::Get();
+	MaxEndEffectorError = AnimationSettings->AlternativeCompressionThreshold;
+	bTryFixedBitwiseCompression = AnimationSettings->bTryFixedBitwiseCompression;
+	bTryPerTrackBitwiseCompression = AnimationSettings->bTryPerTrackBitwiseCompression;
+	bTryLinearKeyRemovalCompression = AnimationSettings->bTryLinearKeyRemovalCompression;
+	bTryIntervalKeyRemoval = AnimationSettings->bTryIntervalKeyRemoval;
+	bRunCurrentDefaultCompressor = AnimationSettings->bFirstRecompressUsingCurrentOrDefault;
+	bAutoReplaceIfExistingErrorTooGreat = AnimationSettings->bForceBelowThreshold;
+	bRaiseMaxErrorToExisting = AnimationSettings->bRaiseMaxErrorToExisting;
 }
 
 
