Add SetHidden UFUNCTION() and make bHidden an Interp property.

diff --git a/Engine/Source/Runtime/Engine/Classes/GameFramework/Actor.h b/Engine/Source/Runtime/Engine/Classes/GameFramework/Actor.h
--- a/Engine/Source/Runtime/Engine/Classes/GameFramework/Actor.h
+++ b/Engine/Source/Runtime/Engine/Classes/GameFramework/Actor.h
@@ -106,8 +106,10 @@ public:
 	 * Allows us to only see this Actor in the Editor, and not in the actual game.
 	 * @see SetActorHiddenInGame()
 	 */
-	UPROPERTY(EditAnywhere, Category=Rendering, BlueprintReadOnly, replicated, meta=(DisplayName = "Actor Hidden In Game", SequencerTrackClass = "MovieSceneVisibilityTrack"))
+	UPROPERTY(Interp, EditAnywhere, Category=Rendering, BlueprintReadOnly, replicated, meta=(DisplayName = "Actor Hidden In Game", SequencerTrackClass = "MovieSceneVisibilityTrack"))
 	uint32 bHidden:1;
+	UFUNCTION()
+	void SetHidden(bool bInHidden) { bHidden = bInHidden; }
 
 	/** If true, when the actor is spawned it will be sent to the client but receive no further replication updates from the server afterwards. */
 	UPROPERTY()
