Sequencer: Add setter functions for interp camera properties.

diff --git a/Engine/Source/Runtime/Engine/Classes/Camera/CameraComponent.h b/Engine/Source/Runtime/Engine/Classes/Camera/CameraComponent.h
--- a/Engine/Source/Runtime/Engine/Classes/Camera/CameraComponent.h
+++ b/Engine/Source/Runtime/Engine/Classes/Camera/CameraComponent.h
@@ -17,30 +17,44 @@ class ENGINE_API UCameraComponent : public USceneComponent
 	/** The horizontal field of view (in degrees) in perspective mode (ignored in Orthographic mode) */
 	UPROPERTY(Interp, EditAnywhere, BlueprintReadWrite, Category=CameraSettings, meta=(UIMin = "5.0", UIMax = "170", ClampMin = "0.001", ClampMax = "360.0"))
 	float FieldOfView;
+	UFUNCTION(BlueprintCallable, Category=Camera)
+	void SetFieldOfView(float InFieldOfView) { FieldOfView = InFieldOfView; }
 
 	/** The desired width (in world units) of the orthographic view (ignored in Perspective mode) */
 	UPROPERTY(Interp, EditAnywhere, BlueprintReadWrite, Category=CameraSettings)
 	float OrthoWidth;
+	UFUNCTION(BlueprintCallable, Category=Camera)
+	void SetOrthoWidth(float InOrthoWidth) { OrthoWidth = InOrthoWidth; }
 
 	/** The near plane distance of the orthographic view (in world units) */
 	UPROPERTY(Interp, EditAnywhere, BlueprintReadWrite, Category=CameraSettings)
 	float OrthoNearClipPlane;
+	UFUNCTION(BlueprintCallable, Category=Camera)
+	void SetOrthoNearClipPlane(float InOrthoNearClipPlane) { OrthoNearClipPlane = InOrthoNearClipPlane; }
 
 	/** The far plane distance of the orthographic view (in world units) */
 	UPROPERTY(Interp, EditAnywhere, BlueprintReadWrite, Category=CameraSettings)
 	float OrthoFarClipPlane;
+	UFUNCTION(BlueprintCallable, Category=Camera)
+	void SetOrthoFarClipPlane(float InOrthoFarClipPlane) { OrthoFarClipPlane = InOrthoFarClipPlane; }
 
 	// Aspect Ratio (Width/Height)
 	UPROPERTY(Interp, EditAnywhere, BlueprintReadWrite, Category=CameraSettings, meta=(ClampMin = "0.1", ClampMax = "100.0", EditCondition="bConstrainAspectRatio"))
 	float AspectRatio;
+	UFUNCTION(BlueprintCallable, Category=Camera)
+	void SetAspectRatio(float InAspectRatio) { AspectRatio = InAspectRatio; }
 
 	// If bConstrainAspectRatio is true, black bars will be added if the destination view has a different aspect ratio than this camera requested.
 	UPROPERTY(Interp, EditAnywhere, BlueprintReadWrite, Category=CameraSettings)
 	uint32 bConstrainAspectRatio:1;
+	UFUNCTION(BlueprintCallable, Category=Camera)
+	void SetConstraintAspectRatio(bool bInConstrainAspectRatio) { bConstrainAspectRatio = bInConstrainAspectRatio; }
 
 	// If true, account for the field of view angle when computing which level of detail to use for meshes.
 	UPROPERTY(Interp, EditAnywhere, AdvancedDisplay, BlueprintReadWrite, Category=CameraSettings)
 	uint32 bUseFieldOfViewForLOD:1;
+	UFUNCTION(BlueprintCallable, Category=Camera)
+	void SetUseFieldOfViewForLOD(bool bInUseFieldOfViewForLOD) { bUseFieldOfViewForLOD = bInUseFieldOfViewForLOD; }
 
 	/**
 	 * If this camera component is placed on a pawn, should it use the view/control rotation of the pawn where possible?
@@ -52,10 +66,14 @@ class ENGINE_API UCameraComponent : public USceneComponent
 	// The type of camera
 	UPROPERTY(Interp, EditAnywhere, BlueprintReadWrite, Category=CameraSettings)
 	TEnumAsByte<ECameraProjectionMode::Type> ProjectionMode;
+	UFUNCTION(BlueprintCallable, Category=Camera)
+	void SetProjectionMode(ECameraProjectionMode::Type InProjectionMode) { ProjectionMode = InProjectionMode; }
 
 	/** Indicates if PostProcessSettings should be used when using this Camera to view through. */
 	UPROPERTY(Interp, EditAnywhere, BlueprintReadWrite, Category=CameraSettings, meta=(UIMin = "0.0", UIMax = "1.0"))
 	float PostProcessBlendWeight;
+	UFUNCTION(BlueprintCallable, Category=Camera)
+	void SetPostProcessBlendWeight(float InPostProcessBlendWeight) { PostProcessBlendWeight = InPostProcessBlendWeight; }
 
 	/** Post process settings to use for this camera. Don't forget to check the properties you want to override */
 	UPROPERTY(Interp, BlueprintReadWrite, Category=CameraSettings)
