ArrayLibrary functions use ProcessContext also for the Array (not only for the actual context object, that would be always ArrayLibrary CDO)

diff --git a/Engine/Source/Editor/KismetCompiler/Private/KismetCompilerVMBackend.cpp b/Engine/Source/Editor/KismetCompiler/Private/KismetCompilerVMBackend.cpp
--- a/Engine/Source/Editor/KismetCompiler/Private/KismetCompilerVMBackend.cpp
+++ b/Engine/Source/Editor/KismetCompiler/Private/KismetCompilerVMBackend.cpp
@@ -784,7 +784,7 @@ public:
 		{
 			// RValue property is used to clear value after Access Violation. See UObject::ProcessContextOpcod
 			// If the property from LHS is used, then the retured property (with CPF_ReturnParm) is cleared. But properties returned by ref are not cleared. 
-			UProperty* RValueProperty = nullptr; //Statement.LHS ? Statement.LHS->AssociatedVarProperty : nullptr;
+			UProperty* RValueProperty = Statement.LHS ? Statement.LHS->AssociatedVarProperty : nullptr;
 			CallContextWriter.TryStartContext(Statement.FunctionContext, /*bUnsafeToSkip=*/ bHasOutputValue, Statement.bIsInterfaceContext, RValueProperty);
 		}
 
diff --git a/Engine/Source/Runtime/CoreUObject/Private/UObject/ScriptCore.cpp b/Engine/Source/Runtime/CoreUObject/Private/UObject/ScriptCore.cpp
--- a/Engine/Source/Runtime/CoreUObject/Private/UObject/ScriptCore.cpp
+++ b/Engine/Source/Runtime/CoreUObject/Private/UObject/ScriptCore.cpp
@@ -1742,15 +1742,24 @@ void UObject::ProcessContextOpcode( FFrame& Stack, RESULT_DECL, bool bCanFailSil
 	UObject* NewContext = NULL;
 	Stack.Step( this, &NewContext );
 
+	uint8* const OriginalCode = Stack.Code;
+	const bool bValidContext = IsValid(NewContext);
 	// Execute or skip the following expression in the object's context.
-	if (IsValid(NewContext))
+	if (bValidContext)
 	{
 		Stack.Code += sizeof(CodeSkipSizeType)	// Code offset for NULL expressions.
 			+ sizeof(ScriptPointerType);		// Property corresponding to the r-value data, in case the l-value needs to be cleared
 		Stack.Step( NewContext, RESULT_PARAM );
 	}
-	else
+
+	if (!bValidContext || Stack.bArrayContextFailed)
 	{
+		if (Stack.bArrayContextFailed)
+		{
+			Stack.bArrayContextFailed = false;
+			Stack.Code = OriginalCode;
+		}
+
 		if (!bCanFailSilently)
 		{
 			if (NewContext && NewContext->IsPendingKill())
diff --git a/Engine/Source/Runtime/CoreUObject/Public/UObject/Stack.h b/Engine/Source/Runtime/CoreUObject/Public/UObject/Stack.h
--- a/Engine/Source/Runtime/CoreUObject/Public/UObject/Stack.h
+++ b/Engine/Source/Runtime/CoreUObject/Public/UObject/Stack.h
@@ -100,6 +100,8 @@ public:
 
 	/** Currently executed native function */
 	UFunction* CurrentNativeFunction;
+
+	bool bArrayContextFailed;
 public:
 
 	// Constructors.
@@ -183,6 +185,7 @@ inline FFrame::FFrame( UObject* InObject, UFunction* InNode, void* InLocals, FFr
 	, OutParms(NULL)
 	, PropertyChainForCompiledIn(InPropertyChainForCompiledIn)
 	, CurrentNativeFunction(NULL)
+	, bArrayContextFailed(false)
 {
 #if !(UE_BUILD_SHIPPING || UE_BUILD_TEST)
 	FScriptTraceStackNode StackNode(InNode->GetOuter()->GetFName(), InNode->GetFName());
diff --git a/Engine/Source/Runtime/Engine/Classes/Kismet/KismetArrayLibrary.h b/Engine/Source/Runtime/Engine/Classes/Kismet/KismetArrayLibrary.h
--- a/Engine/Source/Runtime/Engine/Classes/Kismet/KismetArrayLibrary.h
+++ b/Engine/Source/Runtime/Engine/Classes/Kismet/KismetArrayLibrary.h
@@ -200,7 +200,12 @@ public:
 		Stack.MostRecentProperty = nullptr;
  		Stack.StepCompiledIn<UArrayProperty>(NULL);
 		void* ArrayAddr = Stack.MostRecentPropertyAddress;
-		UArrayProperty* ArrayProperty = CastChecked<UArrayProperty>(Stack.MostRecentProperty);
+		UArrayProperty* ArrayProperty = Cast<UArrayProperty>(Stack.MostRecentProperty);
+		if (!ArrayProperty)
+		{
+			Stack.bArrayContextFailed = true;
+			return;
+		}
  
  		// Since NewItem isn't really an int, step the stack manually
  		const UProperty* InnerProp = ArrayProperty->Inner;
@@ -224,7 +229,12 @@ public:
 		Stack.MostRecentProperty = nullptr;
 		Stack.StepCompiledIn<UArrayProperty>(NULL);
 		void* ArrayAddr = Stack.MostRecentPropertyAddress;
-		UArrayProperty* ArrayProperty = CastChecked<UArrayProperty>(Stack.MostRecentProperty);
+		UArrayProperty* ArrayProperty = Cast<UArrayProperty>(Stack.MostRecentProperty);
+		if (!ArrayProperty)
+		{
+			Stack.bArrayContextFailed = true;
+			return;
+		}
 
 		// Since NewItem isn't really an int, step the stack manually
 		const UProperty* InnerProp = ArrayProperty->Inner;
@@ -248,7 +258,12 @@ public:
 		Stack.MostRecentProperty = nullptr;
 		Stack.StepCompiledIn<UArrayProperty>(NULL);
 		void* ArrayAddr = Stack.MostRecentPropertyAddress;
-		UArrayProperty* ArrayProperty = CastChecked<UArrayProperty>(Stack.MostRecentProperty);
+		UArrayProperty* ArrayProperty = Cast<UArrayProperty>(Stack.MostRecentProperty);
+		if (!ArrayProperty)
+		{
+			Stack.bArrayContextFailed = true;
+			return;
+		}
 
 		P_FINISH;
 
@@ -261,13 +276,22 @@ public:
 		Stack.MostRecentProperty = nullptr;
 		Stack.StepCompiledIn<UArrayProperty>(NULL);
 		void* TargetArrayAddr = Stack.MostRecentPropertyAddress;
-		UArrayProperty* TargetArrayProperty = CastChecked<UArrayProperty>(Stack.MostRecentProperty);
-
+		UArrayProperty* TargetArrayProperty = Cast<UArrayProperty>(Stack.MostRecentProperty);
+		if (!TargetArrayProperty)
+		{
+			Stack.bArrayContextFailed = true;
+			return;
+		}
 		// Retrieve the source array
 		Stack.MostRecentProperty = nullptr;
 		Stack.StepCompiledIn<UArrayProperty>(NULL);
 		void* SourceArrayAddr = Stack.MostRecentPropertyAddress;
-		UArrayProperty* SourceArrayProperty = CastChecked<UArrayProperty>(Stack.MostRecentProperty);
+		UArrayProperty* SourceArrayProperty = Cast<UArrayProperty>(Stack.MostRecentProperty);
+		if (!TargetArrayProperty)
+		{
+			Stack.bArrayContextFailed = true;
+			return;
+		}
 
 		P_FINISH;
 
@@ -279,7 +303,12 @@ public:
 		Stack.MostRecentProperty = nullptr;
 		Stack.StepCompiledIn<UArrayProperty>(NULL);
 		void* ArrayAddr = Stack.MostRecentPropertyAddress;
-		UArrayProperty* ArrayProperty = CastChecked<UArrayProperty>(Stack.MostRecentProperty);
+		UArrayProperty* ArrayProperty = Cast<UArrayProperty>(Stack.MostRecentProperty);
+		if (!ArrayProperty)
+		{
+			Stack.bArrayContextFailed = true;
+			return;
+		}
 
 		// Since NewItem isn't really an int, step the stack manually
 		const UProperty* InnerProp = ArrayProperty->Inner;
@@ -304,7 +333,12 @@ public:
 		Stack.MostRecentProperty = nullptr;
 		Stack.StepCompiledIn<UArrayProperty>(NULL);
 		void* ArrayAddr = Stack.MostRecentPropertyAddress;
-		UArrayProperty* ArrayProperty = CastChecked<UArrayProperty>(Stack.MostRecentProperty);
+		UArrayProperty* ArrayProperty = Cast<UArrayProperty>(Stack.MostRecentProperty);
+		if (!ArrayProperty)
+		{
+			Stack.bArrayContextFailed = true;
+			return;
+		}
 
 		P_GET_PROPERTY(UIntProperty, Index);
 		P_FINISH;
@@ -317,8 +351,12 @@ public:
 		Stack.MostRecentProperty = nullptr;
 		Stack.StepCompiledIn<UArrayProperty>(NULL);
 		void* ArrayAddr = Stack.MostRecentPropertyAddress;
-		UArrayProperty* ArrayProperty = CastChecked<UArrayProperty>(Stack.MostRecentProperty);
-
+		UArrayProperty* ArrayProperty = Cast<UArrayProperty>(Stack.MostRecentProperty);
+		if (!ArrayProperty)
+		{
+			Stack.bArrayContextFailed = true;
+			return;
+		}
 		// Since Item isn't really an int, step the stack manually
 		const UProperty* InnerProp = ArrayProperty->Inner;
 		const int32 PropertySize = InnerProp->ElementSize * InnerProp->ArrayDim;
@@ -345,8 +383,12 @@ public:
 		Stack.MostRecentProperty = nullptr;
 		Stack.StepCompiledIn<UArrayProperty>(NULL);
 		void* ArrayAddr = Stack.MostRecentPropertyAddress;
-		UArrayProperty* ArrayProperty = CastChecked<UArrayProperty>(Stack.MostRecentProperty);
-
+		UArrayProperty* ArrayProperty = Cast<UArrayProperty>(Stack.MostRecentProperty);
+		if (!ArrayProperty)
+		{
+			Stack.bArrayContextFailed = true;
+			return;
+		}
 		P_FINISH;
 
 		GenericArray_Clear(ArrayAddr, ArrayProperty);
@@ -357,8 +399,12 @@ public:
 		Stack.MostRecentProperty = nullptr;
 		Stack.StepCompiledIn<UArrayProperty>(NULL);
 		void* ArrayAddr = Stack.MostRecentPropertyAddress;
-		UArrayProperty* ArrayProperty = CastChecked<UArrayProperty>(Stack.MostRecentProperty);
-
+		UArrayProperty* ArrayProperty = Cast<UArrayProperty>(Stack.MostRecentProperty);
+		if (!ArrayProperty)
+		{
+			Stack.bArrayContextFailed = true;
+			return;
+		}
 		P_GET_PROPERTY(UIntProperty, Size);
 		P_FINISH;
 
@@ -370,8 +416,12 @@ public:
 		Stack.MostRecentProperty = nullptr;
 		Stack.StepCompiledIn<UArrayProperty>(NULL);
 		void* ArrayAddr = Stack.MostRecentPropertyAddress;
-		UArrayProperty* ArrayProperty = CastChecked<UArrayProperty>(Stack.MostRecentProperty);
-
+		UArrayProperty* ArrayProperty = Cast<UArrayProperty>(Stack.MostRecentProperty);
+		if (!ArrayProperty)
+		{
+			Stack.bArrayContextFailed = true;
+			return;
+		}
 		P_FINISH;
 
 		*(int32*)RESULT_PARAM = GenericArray_Length(ArrayAddr, ArrayProperty);
@@ -382,8 +432,12 @@ public:
 		Stack.MostRecentProperty = nullptr;
 		Stack.StepCompiledIn<UArrayProperty>(NULL);
 		void* ArrayAddr = Stack.MostRecentPropertyAddress;
-		UArrayProperty* ArrayProperty = CastChecked<UArrayProperty>(Stack.MostRecentProperty);
-
+		UArrayProperty* ArrayProperty = Cast<UArrayProperty>(Stack.MostRecentProperty);
+		if (!ArrayProperty)
+		{
+			Stack.bArrayContextFailed = true;
+			return;
+		}
 		P_FINISH;
 
 		*(int32*)RESULT_PARAM = GenericArray_LastIndex(ArrayAddr, ArrayProperty);
@@ -394,8 +448,12 @@ public:
 		Stack.MostRecentProperty = nullptr;
 		Stack.StepCompiledIn<UArrayProperty>(NULL);
 		void* ArrayAddr = Stack.MostRecentPropertyAddress;
-		UArrayProperty* ArrayProperty = CastChecked<UArrayProperty>(Stack.MostRecentProperty);
-
+		UArrayProperty* ArrayProperty = Cast<UArrayProperty>(Stack.MostRecentProperty);
+		if (!ArrayProperty)
+		{
+			Stack.bArrayContextFailed = true;
+			return;
+		}
 		P_GET_PROPERTY(UIntProperty, Index);
 
 		// Since Item isn't really an int, step the stack manually
@@ -420,8 +478,12 @@ public:
 		Stack.MostRecentProperty = nullptr;
 		Stack.StepCompiledIn<UArrayProperty>(NULL);
 		void* ArrayAddr = Stack.MostRecentPropertyAddress;
-		UArrayProperty* ArrayProperty = CastChecked<UArrayProperty>(Stack.MostRecentProperty);
-
+		UArrayProperty* ArrayProperty = Cast<UArrayProperty>(Stack.MostRecentProperty);
+		if (!ArrayProperty)
+		{
+			Stack.bArrayContextFailed = true;
+			return;
+		}
 		P_GET_PROPERTY(UIntProperty, Index);
 
 		// Since NewItem isn't really an int, step the stack manually
@@ -448,8 +510,12 @@ public:
 		Stack.MostRecentProperty = nullptr;
 		Stack.StepCompiledIn<UArrayProperty>(NULL);
 		void* ArrayAddr = Stack.MostRecentPropertyAddress;
-		UArrayProperty* ArrayProperty = CastChecked<UArrayProperty>(Stack.MostRecentProperty);
-
+		UArrayProperty* ArrayProperty = Cast<UArrayProperty>(Stack.MostRecentProperty);
+		if (!ArrayProperty)
+		{
+			Stack.bArrayContextFailed = true;
+			return;
+		}
 		// Since ItemToFind isn't really an int, step the stack manually
 		const UProperty* InnerProp = ArrayProperty->Inner;
 		const int32 PropertySize = InnerProp->ElementSize * InnerProp->ArrayDim;
@@ -477,8 +543,12 @@ public:
 		Stack.MostRecentProperty = nullptr;
 		Stack.StepCompiledIn<UArrayProperty>(NULL);
 		void* ArrayAddr = Stack.MostRecentPropertyAddress;
-		UArrayProperty* ArrayProperty = CastChecked<UArrayProperty>(Stack.MostRecentProperty);
-
+		UArrayProperty* ArrayProperty = Cast<UArrayProperty>(Stack.MostRecentProperty);
+		if (!ArrayProperty)
+		{
+			Stack.bArrayContextFailed = true;
+			return;
+		}
 		// Since ItemToFind isn't really an int, step the stack manually
 		const UProperty* InnerProp = ArrayProperty->Inner;
 		const int32 PropertySize = InnerProp->ElementSize * InnerProp->ArrayDim;
