Removed a warning and added an error to ensure that texture data is being correctly cleared.

diff --git a/Engine/Source/Runtime/Engine/Private/TextureDerivedData.cpp b/Engine/Source/Runtime/Engine/Private/TextureDerivedData.cpp
--- a/Engine/Source/Runtime/Engine/Private/TextureDerivedData.cpp
+++ b/Engine/Source/Runtime/Engine/Private/TextureDerivedData.cpp
@@ -693,7 +693,7 @@ class FTextureCacheDerivedDataWorker : public FNonAbandonableTask
 		if ( Texture.Source.HasHadBulkDataCleared() )
 		{
 			// don't do any work we can't reload this
-			UE_LOG(LogTexture, Warning, TEXT("Unable to load texture source data could be because it has been cleared. %s"), *Texture.GetName())
+			UE_LOG(LogTexture, Error, TEXT("Unable to load texture source data could be because it has been cleared. %s"), *Texture.GetName())
 			return;
 		}
 
@@ -1649,7 +1649,7 @@ void UTexture2D::WillNeverCacheCookedPlatformDataAgain()
 	Super::WillNeverCacheCookedPlatformDataAgain();
 #if WITH_EDITORONLY_DATA 
 	// "Source" is only in WITH_EDITORONLY_DATA
-	UE_LOG(LogTemp, Warning, TEXT("Cleared source texture data for texture %s"), *GetName());
+	// UE_LOG(LogTemp, Display, TEXT("Cleared source texture data for texture %s"), *GetName());
 	// clear source mips if we are in the editor then we don't have this luxury 
 	check( IsAsyncCacheComplete());
 
