Integrate 2641961 from Orion branch

Re-fix collision not being enabled on the spectator camera (got broken after stable main merge)

diff --git a/Engine/Source/Runtime/Engine/Private/DemoNetDriver.cpp b/Engine/Source/Runtime/Engine/Private/DemoNetDriver.cpp
--- a/Engine/Source/Runtime/Engine/Private/DemoNetDriver.cpp
+++ b/Engine/Source/Runtime/Engine/Private/DemoNetDriver.cpp
@@ -575,9 +575,6 @@ void UDemoNetDriver::TickDispatch(float DeltaSeconds)
 			if ( SpectatorController->GetSpectatorPawn() != NULL )
 			{
 				SpectatorController->GetSpectatorPawn()->CustomTimeDilation = SpectatorController->CustomTimeDilation;
-
-				// Disable collision on the spectator
-				SpectatorController->GetSpectatorPawn()->SetActorEnableCollision( false );
 					
 				SpectatorController->GetSpectatorPawn()->PrimaryActorTick.bTickEvenWhenPaused = true;
 
