Merging 2640056 (plus a few minor fixes)

//depot/UE4-Orion/...

to //depot/UE4/...

UE4 -  rollup of many performance changes

* increased max threads for task graph and added a way to reduce task threads at runtime for testing the optimal number
* added a way to track excessive allocation at runtime without using the memory profiler: LogGameThreadMallocChurn.Enable
* ps4: fixed a few problems with the platform properties, IsRunningGame was returning false
* added inline storage for one delegate in multicast delegates
* fixed excessive allocation on the script stack traceback
* fixed excessive malloc calls throughout the animation system
* added stats throughout the animation system
* reworked parallel queue ticks
* fixed excessive malloc calls in cloth and physics animation systems
* fixed GHitchThresholdCVar
* minor improvements to abtest
* added ability to change thread affinities on the fly so they can be tuned
* cvars to control animrate optimization and spew the results: EnableAnimRateOptimization, SpewAnimRateOptimization, DrawAnimRateOptimization, ForceAnimRate
* optimized malloc calls in cascade and gpu particles
* optimized malloc calls in input system
* removed visual logging and ability log from consoles and fixed a few compile errors related to doing that
* optimized gameplaytags
* added a few cvars to cause stalls in various parts of the pipeline to identify critical paths and bottlenecks: CriticalPathStall.*
* added console command to toggle all known outstanding perf tweaks: Orion.SetupPerfCVars

diff --git a/Engine/Source/Runtime/AnimGraphRuntime/Private/AnimNodes/AnimNode_TwoWayBlend.cpp b/Engine/Source/Runtime/AnimGraphRuntime/Private/AnimNodes/AnimNode_TwoWayBlend.cpp
--- a/Engine/Source/Runtime/AnimGraphRuntime/Private/AnimNodes/AnimNode_TwoWayBlend.cpp
+++ b/Engine/Source/Runtime/AnimGraphRuntime/Private/AnimNodes/AnimNode_TwoWayBlend.cpp
@@ -20,6 +20,7 @@ void FAnimationNode_TwoWayBlend::CacheBones(const FAnimationCacheBonesContext& C
 
 void FAnimationNode_TwoWayBlend::Update(const FAnimationUpdateContext& Context)
 {
+	QUICK_SCOPE_CYCLE_COUNTER(STAT_FAnimationNode_TwoWayBlend_Update);
 	EvaluateGraphExposedInputs.Execute(Context);
 
 	const float ActualAlpha = AlphaScaleBias.ApplyTo(Alpha);
diff --git a/Engine/Source/Runtime/Core/Private/Async/TaskGraph.cpp b/Engine/Source/Runtime/Core/Private/Async/TaskGraph.cpp
--- a/Engine/Source/Runtime/Core/Private/Async/TaskGraph.cpp
+++ b/Engine/Source/Runtime/Core/Private/Async/TaskGraph.cpp
@@ -187,6 +187,14 @@ FORCEINLINE void TestRandomizedThreads()
 
 #endif
 
+static int32 GNumWorkerThreadsToIgnore = 0;
+static FAutoConsoleVariableRef CVarNumWorkerThreadsToIgnore(
+	TEXT("TaskGraph.NumWorkerThreadsToIgnore"),
+	GNumWorkerThreadsToIgnore,
+	TEXT("Used to tune the number of task threads. Generally once you have found the right value, PlatformMisc::NumberOfWorkerThreadsToSpawn() should be hardcoded."),
+	ECVF_Cheat
+	);
+
 /** 
  *	FTaskQueue
  *	High performance, SINGLE threaded, FIFO task queue for the private queue on named threads.
@@ -1004,6 +1012,10 @@ public:
 				{
 					TASKGRAPH_SCOPE_CYCLE_COUNTER(5, STAT_TaskGraph_QueueTask_StalledUnnamedThreads_Pop);
 					TempTarget = StalledUnnamedThreads.Pop(); //@todo it is possible that a thread is in the process of stalling and we just missed it, non-fatal, but we could lose a whole task of potential parallelism.
+					if (TempTarget && GNumWorkerThreadsToIgnore && (TempTarget->GetThreadId() - NumNamedThreads) >= GetNumWorkerThreads())
+					{
+						TempTarget = nullptr;
+					}
 				}
 				if (TempTarget)
 				{
@@ -1011,7 +1023,8 @@ public:
 				}
 				else
 				{
-					ThreadToExecuteOn = ENamedThreads::Type((uint32(NextUnnamedThreadForTaskFromUnknownThread.Increment()) % uint32(NextUnnamedThreadMod)) + NumNamedThreads);
+					check(NextUnnamedThreadMod - GNumWorkerThreadsToIgnore > 0); // can't tune it to zero task threads
+					ThreadToExecuteOn = ENamedThreads::Type((uint32(NextUnnamedThreadForTaskFromUnknownThread.Increment()) % uint32(NextUnnamedThreadMod - GNumWorkerThreadsToIgnore)) + NumNamedThreads);
 				}
 				FTaskThread* Target = &Thread(ThreadToExecuteOn);
 				if (ThreadToExecuteOn != CurrentThreadIfKnown)
@@ -1043,7 +1056,9 @@ public:
 
 	virtual	int32 GetNumWorkerThreads() override
 	{
-		return NumThreads - NumNamedThreads;
+		int32 Result = NumThreads - NumNamedThreads - GNumWorkerThreadsToIgnore;
+		check(Result > 0); // can't tune it to zero task threads
+		return Result;
 	}
 
 	virtual ENamedThreads::Type GetCurrentThreadIfKnown(bool bLocalQueue) override
@@ -1331,7 +1346,7 @@ private:
 	enum
 	{
 		/** Compile time maximum number of threads. @todo Didn't really need to be a compile time constant. **/
-		MAX_THREADS=8
+		MAX_THREADS=12
 	};
 
 	/** Per thread data. **/
@@ -1480,85 +1495,3 @@ FGraphEvent::~FGraphEvent()
 	CheckDontCompleteUntilIsEmpty(); // We should not have any wait untils outstanding
 }
 
-
-//---
-
-#include "ParallelFor.h"
-
-static void TestParallelFor(const TArray<FString>& Args)
-{
-
-	ParallelFor(10, 
-		[](int32 Index)
-		{
-			UE_LOG(LogConsoleResponse, Display, TEXT("ParallelFor index=%d, thread=%x"), Index, FPlatformTLS::GetCurrentThreadId());
-		}
-	);
-
-	ParallelFor(10, 
-		[](int32 Index)
-		{
-			ParallelFor(10, 
-				[Index](int32 IndexInner)
-				{
-					UE_LOG(LogConsoleResponse, Display, TEXT("ParallelFor index=%d %d, thread=%x"), Index, IndexInner, FPlatformTLS::GetCurrentThreadId());
-				}
-			);
-		}
-	);
-
-	TArray<bool> TestBools;
-
-	TestBools.AddZeroed(10000);
-
-	ParallelFor(TestBools.Num(), [&TestBools](int32 Index){TestBools[Index] = true;});
-
-	for (int32 Index = 0; Index < TestBools.Num(); Index++)
-	{
-		check(TestBools[Index]);
-	}
-
-	TestBools.Empty(10000);
-	TestBools.AddZeroed(10000);
-
-	ParallelFor(TestBools.Num(), [&TestBools](int32 Index){TestBools[Index] = true;}, true);
-
-	for (int32 Index = 0; Index < TestBools.Num(); Index++)
-	{
-		check(TestBools[Index]);
-	}
-
-	uint32 TargetCrc = FCrc::MemCrc32(TestBools.GetData(), 1024);
-
-	{
-		double StartTime = FPlatformTime::Seconds();
-		ParallelFor(1024, 
-			[&TestBools, TargetCrc](int32 Index)
-			{
-				uint32 TestCrc = FCrc::MemCrc32(TestBools.GetData(), 1024);
-				check(TestCrc == TargetCrc);
-			}
-		);
-		UE_LOG(LogConsoleResponse, Display, TEXT("Parallel CRC of 1MB (%d threads) in %6.3fms"), FTaskGraphInterface::Get().GetNumWorkerThreads() + 1, float(FPlatformTime::Seconds() - StartTime) * 1000.0f);
-	}
-	{
-		double StartTime = FPlatformTime::Seconds();
-		ParallelFor(1024, 
-			[&TestBools, TargetCrc](int32 Index)
-			{
-				uint32 TestCrc = FCrc::MemCrc32(TestBools.GetData(), 1024);
-				check(TestCrc == TargetCrc);
-			}, 
-			true
-		);
-		UE_LOG(LogConsoleResponse, Display, TEXT("Serial CRC of 1MB in %6.3fms"), float(FPlatformTime::Seconds() - StartTime) * 1000.0f);
-	}
-}
-
-static FAutoConsoleCommand TestParallelForCmd(
-	TEXT("TestParallelFor"),
-	TEXT("Simple test of ParallelFor."),
-	FConsoleCommandWithArgsDelegate::CreateStatic(&TestParallelFor)
-	);
-
-
diff --git a/Engine/Source/Runtime/Core/Private/Containers/StackTracker.cpp b/Engine/Source/Runtime/Core/Private/Containers/StackTracker.cpp
--- a/Engine/Source/Runtime/Core/Private/Containers/StackTracker.cpp
+++ b/Engine/Source/Runtime/Core/Private/Containers/StackTracker.cpp
@@ -12,7 +12,7 @@ DEFINE_LOG_CATEGORY_STATIC(LogStackTracker, Log, All);
  * optionally stores a user data pointer that the tracker will take ownership of and delete upon reset
  * you must allocate the memory with FMemory::Malloc()
  */
-void FStackTracker::CaptureStackTrace(int32 EntriesToIgnore /*=2*/, void* UserData /*=NULL*/)
+void FStackTracker::CaptureStackTrace(int32 EntriesToIgnore, void* UserData, int32 StackLen)
 {
 	// Avoid re-rentrancy as the code uses TArray/TMap.
 	if( !bAvoidCapturing && bIsEnabled )
@@ -21,14 +21,17 @@ void FStackTracker::CaptureStackTrace(int32 EntriesToIgnore /*=2*/, void* UserDa
 		bAvoidCapturing = true;
 
 		// Capture callstack and create CRC.
-		uint64* FullBackTrace = NULL;
-		FullBackTrace = static_cast<uint64*>(FMemory_Alloca((MAX_BACKTRACE_DEPTH + EntriesToIgnore) * sizeof(uint64)));
+		int32 Size = (MAX_BACKTRACE_DEPTH + EntriesToIgnore) * sizeof(uint64);
+		uint64* FullBackTrace = static_cast<uint64*>(FMemory_Alloca(Size));
+
+		FMemory::Memzero(FullBackTrace, Size);
 
 		FPlatformStackWalk::CaptureStackBackTrace( FullBackTrace, MAX_BACKTRACE_DEPTH + EntriesToIgnore );
 		CA_ASSUME(FullBackTrace);
 
 		// Skip first NUM_ENTRIES_TO_SKIP entries as they are inside this code
 		uint64* BackTrace = &FullBackTrace[EntriesToIgnore];
+		FMemory::Memzero(BackTrace + StackLen, Size - sizeof(uint64) * (MAX_BACKTRACE_DEPTH - StackLen));
 		uint32 CRC = FCrc::MemCrc_DEPRECATED( BackTrace, MAX_BACKTRACE_DEPTH * sizeof(uint64) );
         
 		// Use index if found
@@ -72,8 +75,9 @@ void FStackTracker::CaptureStackTrace(int32 EntriesToIgnore /*=2*/, void* UserDa
 /**
  * Dumps capture stack trace summary to the passed in log.
  */
-void FStackTracker::DumpStackTraces( int32 StackThreshold, FOutputDevice& Ar )
+void FStackTracker::DumpStackTraces(int32 StackThreshold, FOutputDevice& Ar, float SampleCountCorrectionFactor)
 {
+	check(SampleCountCorrectionFactor > 0.0f);
 	// Avoid distorting results while we log them.
 	check( !bAvoidCapturing );
 	bAvoidCapturing = true;
@@ -108,7 +112,16 @@ void FStackTracker::DumpStackTraces( int32 StackThreshold, FOutputDevice& Ar )
 	}
 
 	// Log quick summary as we don't log each individual so totals in CSV won't represent real totals.
-	Ar.Logf(TEXT("Captured %i unique callstacks totalling %i function calls over %i frames, averaging %5.2f calls/frame, Avg Per Frame"), SortedCallStacks.Num(), (int32)TotalStackCount, FramesCaptured, (float) TotalStackCount / FramesCaptured);
+	if (SampleCountCorrectionFactor != 1.0f)
+	{
+		StackThreshold = FMath::Max<int32>(1, int32(float(StackThreshold) / SampleCountCorrectionFactor));
+		// assume here the user has already rolled the number of frames into the correction factor
+		Ar.Logf(TEXT("Captured %i unique callstacks averaging %f function calls per frame"), SortedCallStacks.Num(), float(TotalStackCount) * SampleCountCorrectionFactor);
+	}
+	else
+	{
+		Ar.Logf(TEXT("Captured %i unique callstacks totalling %i function calls over %i frames, averaging %5.2f calls/frame, Avg Per Frame"), SortedCallStacks.Num(), (int32)TotalStackCount, FramesCaptured, (float) TotalStackCount / FramesCaptured);
+	}
 
 	// Iterate over each callstack and write out info in human readable form in CSV format
 	for( int32 CallStackIndex=0; CallStackIndex<SortedCallStacks.Num(); CallStackIndex++ )
@@ -118,9 +131,19 @@ void FStackTracker::DumpStackTraces( int32 StackThreshold, FOutputDevice& Ar )
 		// Avoid log spam by only logging above threshold.
 		if( CallStack.StackCount > StackThreshold )
 		{
-			// First row is stack count.
-			FString CallStackString = FString::FromInt((int32)CallStack.StackCount);
-			CallStackString += FString::Printf( TEXT(",%5.2f"), static_cast<float>(CallStack.StackCount)/static_cast<float>(FramesCaptured) );
+			FString CallStackString;
+			if (SampleCountCorrectionFactor != 1.0f)
+			{
+				// First row is stack count.
+				CallStackString = FString::FromInt((int32)(float(CallStack.StackCount) * SampleCountCorrectionFactor));
+				CallStackString += FString::Printf( TEXT(",%5.2f"), float(CallStack.StackCount) * SampleCountCorrectionFactor);
+			}
+			else
+			{
+				// First row is stack count.
+				CallStackString = FString::FromInt((int32)CallStack.StackCount);
+				CallStackString += FString::Printf( TEXT(",%5.2f"), static_cast<float>(CallStack.StackCount)/static_cast<float>(FramesCaptured) );
+			}
 			
 
 			// Iterate over all addresses in the callstack to look up symbol name.
@@ -141,6 +164,10 @@ void FStackTracker::DumpStackTraces( int32 StackThreshold, FOutputDevice& Ar )
 				ReportFn(CallStack, CallStack.StackCount, Ar);
 			}
 		}
+		else
+		{
+			break;
+		}
 	}
 
 	// Done logging.
@@ -172,17 +199,31 @@ void FStackTracker::ResetTracking()
 /** Toggles tracking. */
 void FStackTracker::ToggleTracking()
 {
-	bIsEnabled = !bIsEnabled;
-	// Enabled
-	if( bIsEnabled )
-	{
-		UE_LOG(LogStackTracker, Log, TEXT("Stack tracking is now enabled."));
-		StartFrameCounter = GFrameCounter;
-	}
-	// Disabled.
-	else
+	ToggleTracking(!bIsEnabled, false);
+}
+
+void FStackTracker::ToggleTracking(bool bEnable, bool bSilent)
+{
+	if (bEnable != bIsEnabled)
 	{
-		StopFrameCounter = GFrameCounter;
-		UE_LOG(LogStackTracker, Log, TEXT("Stack tracking is now disabled."));
+		bIsEnabled = bEnable;
+		// Enabled
+		if( bIsEnabled )
+		{
+			if (!bSilent)
+			{
+				UE_LOG(LogStackTracker, Log, TEXT("Stack tracking is now enabled."));
+			}
+			StartFrameCounter = GFrameCounter;
+		}
+		// Disabled.
+		else
+		{
+			StopFrameCounter = GFrameCounter;
+			if (!bSilent)
+			{
+				UE_LOG(LogStackTracker, Log, TEXT("Stack tracking is now disabled."));
+			}
+		}
 	}
 }
diff --git a/Engine/Source/Runtime/Core/Private/HAL/UnrealMemory.cpp b/Engine/Source/Runtime/Core/Private/HAL/UnrealMemory.cpp
--- a/Engine/Source/Runtime/Core/Private/HAL/UnrealMemory.cpp
+++ b/Engine/Source/Runtime/Core/Private/HAL/UnrealMemory.cpp
@@ -40,6 +40,46 @@ void GCreateMalloc()
 		GMalloc = new FMallocThreadSafeProxy( GMalloc );
 	}
 }
+#if STATS
+	#define MALLOC_GT_HOOKS 1
+#else
+	#define MALLOC_GT_HOOKS 0
+#endif
+
+#if MALLOC_GT_HOOKS
+
+// This code is used to find memory allocations, you set up the lambda in the section of the code you are interested in and add a breakpoint to your lambda to see who is allocating memory
+
+//An example:
+//static TFunction<void(int32)> AnimHook(
+//	[](int32 Index)
+//{
+//	TickAnimationMallocStats[Index]++;
+//	if (++AllCount % 337 == 0)
+//	{
+//		BreakMe();
+//	}
+//}
+//);
+//extern CORE_API TFunction<void(int32)>* GGameThreadMallocHook;
+//TGuardValue<TFunction<void(int32)>*> RestoreHook(GGameThreadMallocHook, &AnimHook);
+
+
+CORE_API TFunction<void(int32)>* GGameThreadMallocHook = nullptr;
+
+FORCEINLINE static void DoGamethreadHook(int32 Index)
+{
+	if (GIsRunning && GGameThreadMallocHook // otherwise our hook might not be initialized yet
+		&& IsInGameThread() ) 
+	{
+		(*GGameThreadMallocHook)(Index);
+	}
+}
+#else
+FORCEINLINE static void DoGamethreadHook(int32 Index)
+{
+}
+#endif
 
 void* FMemory::Malloc( SIZE_T Count, uint32 Alignment ) 
 { 
@@ -48,6 +88,7 @@ void* FMemory::Malloc( SIZE_T Count, uint32 Alignment )
 		GCreateMalloc();
 		CA_ASSUME( GMalloc != NULL );	// Don't want to assert, but suppress static analysis warnings about potentially NULL GMalloc
 	}
+	DoGamethreadHook(0);
 	return GMalloc->Malloc( Count, Alignment );
 }
 
@@ -58,6 +99,7 @@ void* FMemory::Realloc( void* Original, SIZE_T Count, uint32 Alignment )
 		GCreateMalloc();	
 		CA_ASSUME( GMalloc != NULL );	// Don't want to assert, but suppress static analysis warnings about potentially NULL GMalloc
 	}
+	DoGamethreadHook(1);
 	return GMalloc->Realloc( Original, Count, Alignment );
 }	
 
@@ -68,6 +110,7 @@ void FMemory::Free( void* Original )
 		GCreateMalloc();
 		CA_ASSUME( GMalloc != NULL );	// Don't want to assert, but suppress static analysis warnings about potentially NULL GMalloc
 	}
+	DoGamethreadHook(2);
 	return GMalloc->Free( Original );
 }
 
diff --git a/Engine/Source/Runtime/Core/Public/Containers/StackTracker.h b/Engine/Source/Runtime/Core/Public/Containers/StackTracker.h
--- a/Engine/Source/Runtime/Core/Public/Containers/StackTracker.h
+++ b/Engine/Source/Runtime/Core/Public/Containers/StackTracker.h
@@ -67,19 +67,21 @@ public:
 	 * Captures the current stack and updates stack tracking information.
 	 * optionally stores a user data pointer that the tracker will take ownership of and delete upon reset
 	 * you must allocate the memory with FMemory::Malloc()
+	 * EntriesToIgnore are removed from the top of then stack, then we keep at most StackLen of the remaining entries.
 	 */
-	CORE_API void CaptureStackTrace( int32 EntriesToIgnore = 2, void* UserData = NULL );
+	CORE_API void CaptureStackTrace(int32 EntriesToIgnore = 2, void* UserData = nullptr, int32 StackLen = MAX_int32);
 
 	/**
 	 * Dumps capture stack trace summary to the passed in log.
 	 */
-	CORE_API void DumpStackTraces( int32 StackThreshold, FOutputDevice& Ar );
+	CORE_API void DumpStackTraces(int32 StackThreshold, FOutputDevice& Ar, float SampleCountCorrectionFactor = 1.0);
 
 	/** Resets stack tracking. Deletes all user pointers passed in via CaptureStackTrace() */
 	CORE_API void ResetTracking();
 
 	/** Toggles tracking. */
 	CORE_API void ToggleTracking();
+	CORE_API void ToggleTracking(bool bEnable, bool bSilent);
 
 private:
 
diff --git a/Engine/Source/Runtime/Core/Public/Delegates/DelegateSignatureImpl.inl b/Engine/Source/Runtime/Core/Public/Delegates/DelegateSignatureImpl.inl
--- a/Engine/Source/Runtime/Core/Public/Delegates/DelegateSignatureImpl.inl
+++ b/Engine/Source/Runtime/Core/Public/Delegates/DelegateSignatureImpl.inl
@@ -1764,7 +1764,7 @@ protected:
 
 		LockInvocationList();
 		{
-			const TArray<IDelegateInstance*>& InvocationList = GetInvocationList();
+			const TInvocationList& InvocationList = GetInvocationList();
 
 			// call bound functions in reverse order, so we ignore any instances that may be added by callees
 			for (int32 InvocationListIndex = InvocationList.Num() - 1; InvocationListIndex >= 0; --InvocationListIndex)
@@ -1796,7 +1796,7 @@ protected:
 	 */
 	void RemoveDelegateInstance( FDelegateHandle Handle )
 	{
-		const TArray<IDelegateInstance*>& InvocationList = GetInvocationList();
+		const TInvocationList& InvocationList = GetInvocationList();
 
 		// NOTE: We assume that this method is never called with a nullptr object, in which case the
 		//       the following algorithm would break down (it would remove the first found instance
diff --git a/Engine/Source/Runtime/Core/Public/Delegates/DelegateSignatureImpl_Variadics.inl b/Engine/Source/Runtime/Core/Public/Delegates/DelegateSignatureImpl_Variadics.inl
--- a/Engine/Source/Runtime/Core/Public/Delegates/DelegateSignatureImpl_Variadics.inl
+++ b/Engine/Source/Runtime/Core/Public/Delegates/DelegateSignatureImpl_Variadics.inl
@@ -798,7 +798,7 @@ public:
 
 		LockInvocationList();
 		{
-			const TArray<IDelegateInstance*>& LocalInvocationList = GetInvocationList();
+			const TInvocationList& LocalInvocationList = GetInvocationList();
 
 			// call bound functions in reverse order, so we ignore any instances that may be added by callees
 			for (int32 InvocationListIndex = LocalInvocationList.Num() - 1; InvocationListIndex >= 0; --InvocationListIndex)
@@ -831,7 +831,7 @@ protected:
 	 */
 	void RemoveDelegateInstance( FDelegateHandle Handle )
 	{
-		const TArray<IDelegateInstance*>& LocalInvocationList = GetInvocationList();
+		const TInvocationList& LocalInvocationList = GetInvocationList();
 
 		for (int32 InvocationListIndex = 0; InvocationListIndex < LocalInvocationList.Num(); ++InvocationListIndex)
 		{
diff --git a/Engine/Source/Runtime/Core/Public/Delegates/MulticastDelegateBase.h b/Engine/Source/Runtime/Core/Public/Delegates/MulticastDelegateBase.h
--- a/Engine/Source/Runtime/Core/Public/Delegates/MulticastDelegateBase.h
+++ b/Engine/Source/Runtime/Core/Public/Delegates/MulticastDelegateBase.h
@@ -6,6 +6,8 @@
 /**
  * Abstract base class for multicast delegates.
  */
+typedef TArray<IDelegateInstance*, TInlineAllocator<1> > TInvocationList;
+
 template<typename ObjectPtrType = FWeakObjectPtr>
 class FMulticastDelegateBase
 {
@@ -145,7 +147,7 @@ protected:
 	 *
 	 * @return The invocation list.
 	 */
-	inline const TArray<IDelegateInstance*>& GetInvocationList( ) const
+	inline const TInvocationList& GetInvocationList( ) const
 	{
 		return InvocationList;
 	}
@@ -164,12 +166,12 @@ protected:
 
 private:
 
+	/** Holds the collection of delegate instances to invoke. */
+	TInvocationList InvocationList;
+
 	/** Used to determine when a compaction should happen. */
 	int32 CompactionThreshold;
 
-	/** Holds the collection of delegate instances to invoke. */
-	TArray<IDelegateInstance*> InvocationList;
-
 	/** Holds a lock counter for the invocation list. */
 	mutable int32 InvocationListLockCount;
 };
diff --git a/Engine/Source/Runtime/CoreUObject/Public/UObject/Stack.h b/Engine/Source/Runtime/CoreUObject/Public/UObject/Stack.h
--- a/Engine/Source/Runtime/CoreUObject/Public/UObject/Stack.h
+++ b/Engine/Source/Runtime/CoreUObject/Public/UObject/Stack.h
@@ -112,7 +112,7 @@ public:
 #if !(UE_BUILD_SHIPPING || UE_BUILD_TEST)
 		if (GScriptStack.Num())
 		{
-			GScriptStack.Pop();
+			GScriptStack.Pop(false);
 		}
 #endif
 	}
diff --git a/Engine/Source/Runtime/Engine/Classes/Animation/BlendSpace.h b/Engine/Source/Runtime/Engine/Classes/Animation/BlendSpace.h
--- a/Engine/Source/Runtime/Engine/Classes/Animation/BlendSpace.h
+++ b/Engine/Source/Runtime/Engine/Classes/Animation/BlendSpace.h
@@ -48,7 +48,7 @@ protected:
 	virtual void SnapToBorder(FBlendSample& Sample) const override;
 	virtual EBlendSpaceAxis GetAxisToScale() const override;
 	virtual bool IsSameSamplePoint(const FVector& SamplePointA, const FVector& SamplePointB) const override;
-	virtual void GetRawSamplesFromBlendInput(const FVector &BlendInput, TArray<FGridBlendSample> & OutBlendSamples) const override;
+	virtual void GetRawSamplesFromBlendInput(const FVector &BlendInput, TArray<FGridBlendSample, TInlineAllocator<4> > & OutBlendSamples) const override;
 	// End UBlendSpaceBase interface
 
 private:
diff --git a/Engine/Source/Runtime/Engine/Classes/Animation/BlendSpace1D.h b/Engine/Source/Runtime/Engine/Classes/Animation/BlendSpace1D.h
--- a/Engine/Source/Runtime/Engine/Classes/Animation/BlendSpace1D.h
+++ b/Engine/Source/Runtime/Engine/Classes/Animation/BlendSpace1D.h
@@ -34,7 +34,7 @@ protected:
 	virtual void SnapToBorder(FBlendSample& Sample) const override;
 	virtual EBlendSpaceAxis GetAxisToScale() const override;
 	virtual bool IsSameSamplePoint(const FVector& SamplePointA, const FVector& SamplePointB) const override;
-	virtual void GetRawSamplesFromBlendInput(const FVector &BlendInput, TArray<FGridBlendSample> & OutBlendSamples) const override;
+	virtual void GetRawSamplesFromBlendInput(const FVector &BlendInput, TArray<FGridBlendSample, TInlineAllocator<4> > & OutBlendSamples) const override;
 	// End UBlendSpaceBase interface
 
 private:
diff --git a/Engine/Source/Runtime/Engine/Classes/Animation/BlendSpaceBase.h b/Engine/Source/Runtime/Engine/Classes/Animation/BlendSpaceBase.h
--- a/Engine/Source/Runtime/Engine/Classes/Animation/BlendSpaceBase.h
+++ b/Engine/Source/Runtime/Engine/Classes/Animation/BlendSpaceBase.h
@@ -410,7 +410,7 @@ protected:
 	 * @param	OutBlendSamples		Populated with the samples nearest the BlendInput 
 	 *
 	 */
-	virtual void GetRawSamplesFromBlendInput(const FVector &BlendInput, TArray<FGridBlendSample> & OutBlendSamples) const {}
+	virtual void GetRawSamplesFromBlendInput(const FVector &BlendInput, TArray<FGridBlendSample, TInlineAllocator<4> > & OutBlendSamples) const {}
 
 public:
 	
diff --git a/Engine/Source/Runtime/Engine/Classes/Components/SkeletalMeshComponent.h b/Engine/Source/Runtime/Engine/Classes/Components/SkeletalMeshComponent.h
--- a/Engine/Source/Runtime/Engine/Classes/Components/SkeletalMeshComponent.h
+++ b/Engine/Source/Runtime/Engine/Classes/Components/SkeletalMeshComponent.h
@@ -48,6 +48,24 @@ struct FAnimationEvaluationContext
 		Clear();
 	}
 
+	void Copy(const FAnimationEvaluationContext& Other)
+	{
+		AnimInstance = Other.AnimInstance;
+		SkeletalMesh = Other.SkeletalMesh;
+		SpaceBases.Reset();
+		SpaceBases.Append(Other.SpaceBases);
+		LocalAtoms.Reset();
+		LocalAtoms.Append(Other.LocalAtoms);
+		VertexAnims.Reset();
+		VertexAnims.Append(Other.VertexAnims);
+		RootBoneTranslation = Other.RootBoneTranslation;
+		Curve.InitFrom(Other.Curve);
+		bDoInterpolation = Other.bDoInterpolation;
+		bDoEvaluation = Other.bDoEvaluation;
+		bDuplicateToCacheBones = Other.bDuplicateToCacheBones;
+		bDuplicateToCacheCurve = Other.bDuplicateToCacheCurve;
+	}
+
 	void Clear()
 	{
 		AnimInstance = NULL;
diff --git a/Engine/Source/Runtime/Engine/Classes/Components/SkinnedMeshComponent.h b/Engine/Source/Runtime/Engine/Classes/Components/SkinnedMeshComponent.h
--- a/Engine/Source/Runtime/Engine/Classes/Components/SkinnedMeshComponent.h
+++ b/Engine/Source/Runtime/Engine/Classes/Components/SkinnedMeshComponent.h
@@ -912,4 +912,6 @@ public:
 	virtual bool IsPlayingRootMotion(){ return false; }
 
 	virtual bool IsPlayingRootMotionFromEverything(){ return false; }
+
+	bool ShouldUseUpdateRateOptimizations() const;
 };
diff --git a/Engine/Source/Runtime/Engine/Classes/Engine/EngineBaseTypes.h b/Engine/Source/Runtime/Engine/Classes/Engine/EngineBaseTypes.h
--- a/Engine/Source/Runtime/Engine/Classes/Engine/EngineBaseTypes.h
+++ b/Engine/Source/Runtime/Engine/Classes/Engine/EngineBaseTypes.h
@@ -297,8 +297,9 @@ private:
 	 * Queues a tick function for execution from the game thread
 	 * @param TickContext - context to tick in
 	 * @param StackForCycleDetection - Stack For Cycle Detection
+	 * @param bWasInterval - true if this was an interval tick
 	 */
-	void QueueTickFunctionParallel(const struct FTickContext& TickContext, TArray<FTickFunction*, TInlineAllocator<4> >& StackForCycleDetection, struct FTickGroupCompletionItem* AllCompletionEvents, int32 Index, bool bWasInterval);
+	void QueueTickFunctionParallel(const struct FTickContext& TickContext, TArray<FTickFunction*, TInlineAllocator<8> >& StackForCycleDetection, bool bWasInterval);
 
 	/** 
 	 * Logs the prerequisites
diff --git a/Engine/Source/Runtime/Engine/Private/AI/Navigation/RecastNavMeshGenerator.cpp b/Engine/Source/Runtime/Engine/Private/AI/Navigation/RecastNavMeshGenerator.cpp
--- a/Engine/Source/Runtime/Engine/Private/AI/Navigation/RecastNavMeshGenerator.cpp
+++ b/Engine/Source/Runtime/Engine/Private/AI/Navigation/RecastNavMeshGenerator.cpp
@@ -4167,7 +4167,7 @@ uint32 FRecastNavMeshGenerator::LogMemUsed() const
 	return GeneratorsMem + sizeof(FRecastNavMeshGenerator) + PendingDirtyTiles.GetAllocatedSize() + PendingDirtyTiles.GetAllocatedSize();
 }
 
-#if !(UE_BUILD_SHIPPING || UE_BUILD_TEST)
+#if !(UE_BUILD_SHIPPING || UE_BUILD_TEST) && ENABLE_VISUAL_LOG
 void FRecastNavMeshGenerator::GrabDebugSnapshot(struct FVisualLogEntry* Snapshot, const FBox& BoundingBox, const struct FLogCategoryBase& LogCategory, ELogVerbosity::Type LogVerbosity) const
 {
 	const UNavigationSystem* NavSys = UNavigationSystem::GetCurrent(GetWorld());
@@ -4259,7 +4259,7 @@ void FRecastNavMeshGenerator::GrabDebugSnapshot(struct FVisualLogEntry* Snapshot
 }
 #endif
 
-#if !(UE_BUILD_SHIPPING || UE_BUILD_TEST)
+#if !(UE_BUILD_SHIPPING || UE_BUILD_TEST) && ENABLE_VISUAL_LOG
 void FRecastNavMeshGenerator::ExportNavigationData(const FString& FileName) const
 {
 	const UNavigationSystem* NavSys = UNavigationSystem::GetCurrent(GetWorld());
diff --git a/Engine/Source/Runtime/Engine/Private/Animation/AnimCurveTypes.cpp b/Engine/Source/Runtime/Engine/Private/Animation/AnimCurveTypes.cpp
--- a/Engine/Source/Runtime/Engine/Private/Animation/AnimCurveTypes.cpp
+++ b/Engine/Source/Runtime/Engine/Private/Animation/AnimCurveTypes.cpp
@@ -520,7 +520,8 @@ void FBlendedCurve::Override(const FBlendedCurve& CurveToOverrideFrom, float Wei
 void FBlendedCurve::Override(const FBlendedCurve& CurveToOverrideFrom)
 {
 	InitFrom(CurveToOverrideFrom);
-	Elements = CurveToOverrideFrom.Elements;
+	Elements.Reset();
+	Elements.Append(CurveToOverrideFrom.Elements);
 
 	ValidateCurve(this);
 }
@@ -530,7 +531,7 @@ void FBlendedCurve::InitFrom(const USkeleton* Skeleton)
 	if(const FSmartNameMapping* Mapping = Skeleton->SmartNames.GetContainer(USkeleton::AnimCurveMappingName))
 	{
 		Mapping->FillUidArray(UIDList);
-		Elements.Empty(UIDList.Num());
+		Elements.Reset();
 		Elements.AddZeroed(UIDList.Num());
 	}
 
@@ -546,10 +547,15 @@ FBlendedCurve::FBlendedCurve(const class UAnimInstance* AnimInstance)
 
 void FBlendedCurve::CopyFrom(const FBlendedCurve& CurveToCopyFrom)
 {
-	UIDList = CurveToCopyFrom.UIDList;
-	Elements = CurveToCopyFrom.Elements;
-	bInitialized = true;
-	ValidateCurve(this);
+	if (&CurveToCopyFrom != this)
+	{
+		UIDList.Reset();
+		UIDList.Append(CurveToCopyFrom.UIDList);
+		Elements.Reset();
+		Elements.Append(CurveToCopyFrom.Elements);
+		bInitialized = true;
+		ValidateCurve(this);
+	}
 }
 
 void FBlendedCurve::MoveFrom(FBlendedCurve& CurveToMoveFrom)
diff --git a/Engine/Source/Runtime/Engine/Private/Animation/AnimInstance.cpp b/Engine/Source/Runtime/Engine/Private/Animation/AnimInstance.cpp
--- a/Engine/Source/Runtime/Engine/Private/Animation/AnimInstance.cpp
+++ b/Engine/Source/Runtime/Engine/Private/Animation/AnimInstance.cpp
@@ -269,12 +269,12 @@ void UAnimInstance::UninitializeAnimation()
 		}
 	}
 
-	ActiveAnimNotifyState.Empty();
-	EventCurves.Empty();
-	MorphTargetCurves.Empty();
-	MaterialParameterCurves.Empty();
-	MaterialParamatersToClear.Empty();
-	AnimNotifies.Empty();
+	ActiveAnimNotifyState.Reset();
+	EventCurves.Reset();
+	MorphTargetCurves.Reset();
+	MaterialParameterCurves.Reset();
+	MaterialParamatersToClear.Reset();
+	AnimNotifies.Reset();
 }
 
 #if WITH_EDITORONLY_DATA
@@ -305,6 +305,9 @@ bool UAnimInstance::UpdateSnapshotAndSkipRemainingUpdate()
 
 void UAnimInstance::UpdateAnimation(float DeltaSeconds)
 {
+	QUICK_SCOPE_CYCLE_COUNTER(STAT_UAnimInstance_UpdateAnimation);
+	FScopeCycleCounterUObject AnimScope(this);
+
 #if WITH_EDITOR
 	if (GIsEditor)
 	{
@@ -332,15 +335,18 @@ void UAnimInstance::UpdateAnimation(float DeltaSeconds)
 	}
 #endif
 
-	AnimNotifies.Empty();
+	{
+		QUICK_SCOPE_CYCLE_COUNTER(STAT_UAnimInstance_UpdateAnimation_MiscClear);
+		AnimNotifies.Reset();
 
-	ClearSlotNodeWeights();
+		ClearSlotNodeWeights();
 
-	// Reset the player tick list (but keep it presized)
-	UngroupedActivePlayers.Empty(UngroupedActivePlayers.Num());
-	for (int32 GroupIndex = 0; GroupIndex < SyncGroups.Num(); ++GroupIndex)
-	{
-		SyncGroups[GroupIndex].Reset();
+	    // Reset the player tick list (but keep it presized)
+	    UngroupedActivePlayers.Empty(UngroupedActivePlayers.Num());
+	    for (int32 GroupIndex = 0; GroupIndex < SyncGroups.Num(); ++GroupIndex)
+	    {
+		    SyncGroups[GroupIndex].Reset();
+	    }
 	}
 
 	{
@@ -363,60 +369,75 @@ void UAnimInstance::UpdateAnimation(float DeltaSeconds)
 		RootNode->Update(UpdateContext);
 	}
 
-	// Handle all players inside sync groups
-	for (int32 GroupIndex = 0; GroupIndex < SyncGroups.Num(); ++GroupIndex)
 	{
-		FAnimGroupInstance& SyncGroup = SyncGroups[GroupIndex];
+		QUICK_SCOPE_CYCLE_COUNTER(STAT_TickAssetPlayerInstances);
 
-		if (SyncGroup.ActivePlayers.Num() > 0)
+	// Handle all players inside sync groups
+		for (int32 GroupIndex = 0; GroupIndex < SyncGroups.Num(); ++GroupIndex)
 		{
-			const int32 GroupLeaderIndex = FMath::Max(SyncGroup.GroupLeaderIndex, 0);
-
-			// Tick the group leader
-			FAnimAssetTickContext TickContext(DeltaSeconds, RootMotionMode);
-			FAnimTickRecord& GroupLeader = SyncGroup.ActivePlayers[GroupLeaderIndex];
-			GroupLeader.SourceAsset->TickAssetPlayerInstance(GroupLeader, this, TickContext);
-			if (RootMotionMode == ERootMotionMode::RootMotionFromEverything && TickContext.RootMotionMovementParams.bHasRootMotion)
-			{
-				ExtractedRootMotion.AccumulateWithBlend(TickContext.RootMotionMovementParams.RootMotionTransform, GroupLeader.EffectiveBlendWeight);
-			}
-
-			// Update everything else to follow the leader
-			if (SyncGroup.ActivePlayers.Num() > 1)
+			FAnimGroupInstance& SyncGroup = SyncGroups[GroupIndex];
+    
+			if (SyncGroup.ActivePlayers.Num() > 0)
 			{
-				TickContext.ConvertToFollower();
-
-				for (int32 TickIndex = 0; TickIndex < SyncGroup.ActivePlayers.Num(); ++TickIndex)
+				const int32 GroupLeaderIndex = FMath::Max(SyncGroup.GroupLeaderIndex, 0);
+    
+				// Tick the group leader
+				FAnimAssetTickContext TickContext(DeltaSeconds, RootMotionMode);
+				FAnimTickRecord& GroupLeader = SyncGroup.ActivePlayers[GroupLeaderIndex];
+				{
+					QUICK_SCOPE_CYCLE_COUNTER(STAT_UAnimInstance_UpdateAnimation_TickAssetPlayerInstance);
+					FScopeCycleCounterUObject Scope(GroupLeader.SourceAsset);
+					GroupLeader.SourceAsset->TickAssetPlayerInstance(GroupLeader, this, TickContext);
+				}
+				if (RootMotionMode == ERootMotionMode::RootMotionFromEverything && TickContext.RootMotionMovementParams.bHasRootMotion)
+				{
+					ExtractedRootMotion.AccumulateWithBlend(TickContext.RootMotionMovementParams.RootMotionTransform, GroupLeader.EffectiveBlendWeight);
+				}
+    
+				// Update everything else to follow the leader
+				if (SyncGroup.ActivePlayers.Num() > 1)
 				{
-					if (TickIndex != GroupLeaderIndex)
+					TickContext.ConvertToFollower();
+    
+					for (int32 TickIndex = 0; TickIndex < SyncGroup.ActivePlayers.Num(); ++TickIndex)
 					{
-						const FAnimTickRecord& AssetPlayer = SyncGroup.ActivePlayers[TickIndex];
-						AssetPlayer.SourceAsset->TickAssetPlayerInstance(AssetPlayer, this, TickContext);
-						if (RootMotionMode == ERootMotionMode::RootMotionFromEverything && TickContext.RootMotionMovementParams.bHasRootMotion)
+						if (TickIndex != GroupLeaderIndex)
 						{
-							ExtractedRootMotion.AccumulateWithBlend(TickContext.RootMotionMovementParams.RootMotionTransform, AssetPlayer.EffectiveBlendWeight);
+							const FAnimTickRecord& AssetPlayer = SyncGroup.ActivePlayers[TickIndex];
+							{
+								QUICK_SCOPE_CYCLE_COUNTER(STAT_UAnimInstance_UpdateAnimation_TickAssetPlayerInstance);
+								FScopeCycleCounterUObject Scope(AssetPlayer.SourceAsset);
+								AssetPlayer.SourceAsset->TickAssetPlayerInstance(AssetPlayer, this, TickContext);
+							}
+							if (RootMotionMode == ERootMotionMode::RootMotionFromEverything && TickContext.RootMotionMovementParams.bHasRootMotion)
+							{
+								ExtractedRootMotion.AccumulateWithBlend(TickContext.RootMotionMovementParams.RootMotionTransform, AssetPlayer.EffectiveBlendWeight);
+							}
 						}
 					}
 				}
 			}
 		}
-	}
 
-	// Handle the remaining ungrouped animation players
-	for (int32 TickIndex = 0; TickIndex < UngroupedActivePlayers.Num(); ++TickIndex)
-	{
-		const FAnimTickRecord& AssetPlayerToTick = UngroupedActivePlayers[TickIndex];
-		FAnimAssetTickContext TickContext(DeltaSeconds, RootMotionMode);
-		AssetPlayerToTick.SourceAsset->TickAssetPlayerInstance(AssetPlayerToTick, this, TickContext);
-		if (RootMotionMode == ERootMotionMode::RootMotionFromEverything && TickContext.RootMotionMovementParams.bHasRootMotion)
+		// Handle the remaining ungrouped animation players
+		for (int32 TickIndex = 0; TickIndex < UngroupedActivePlayers.Num(); ++TickIndex)
 		{
-			ExtractedRootMotion.AccumulateWithBlend(TickContext.RootMotionMovementParams.RootMotionTransform, AssetPlayerToTick.EffectiveBlendWeight);
+			const FAnimTickRecord& AssetPlayerToTick = UngroupedActivePlayers[TickIndex];
+			FAnimAssetTickContext TickContext(DeltaSeconds, RootMotionMode);
+		    {
+			    QUICK_SCOPE_CYCLE_COUNTER(STAT_UAnimInstance_UpdateAnimation_TickAssetPlayerInstance);
+			    FScopeCycleCounterUObject Scope(AssetPlayerToTick.SourceAsset);
+				AssetPlayerToTick.SourceAsset->TickAssetPlayerInstance(AssetPlayerToTick, this, TickContext);
+		    }
+			if (RootMotionMode == ERootMotionMode::RootMotionFromEverything && TickContext.RootMotionMovementParams.bHasRootMotion)
+			{
+				ExtractedRootMotion.AccumulateWithBlend(TickContext.RootMotionMovementParams.RootMotionTransform, AssetPlayerToTick.EffectiveBlendWeight);
+			}
 		}
 	}
-
 	// update montage should run in game thread
 	// if we do multi threading, make sure this stays in game thread
-	Montage_Advance(DeltaSeconds);
+		Montage_Advance(DeltaSeconds);
 
 	// We may have just partially blended root motion, so make it up to 1 by
 	// blending in identity too
@@ -1488,7 +1509,7 @@ void UAnimInstance::SlotEvaluatePose(FName SlotNodeName, const FCompactPose& Sou
 
 void UAnimInstance::ReinitializeSlotNodes()
 {
-	SlotWeightTracker.Empty();
+	SlotWeightTracker.Reset();
 	
 	// Increment counter
 	SlotNodeInitializationCounter++;
@@ -1649,6 +1670,8 @@ float UAnimInstance::GetCurrentStateElapsedTime(int32 MachineIndex)
 
 void UAnimInstance::Montage_UpdateWeight(float DeltaSeconds)
 {
+	QUICK_SCOPE_CYCLE_COUNTER(STAT_UAnimInstance_Montage_UpdateWeight);
+
 	// go through all montage instances, and update them
 	// and make sure their weight is updated properly
 	for (int32 I=0; I<MontageInstances.Num(); ++I)
@@ -1662,6 +1685,8 @@ void UAnimInstance::Montage_UpdateWeight(float DeltaSeconds)
 
 void UAnimInstance::Montage_Advance(float DeltaSeconds)
 {
+	QUICK_SCOPE_CYCLE_COUNTER(STAT_UAnimInstance_Montage_Advance);
+
 	// We're about to tick montages, queue their events to they're triggered after batched anim notifies.
 	bQueueMontageEvents = true;
 
@@ -1761,7 +1786,7 @@ void UAnimInstance::TriggerQueuedMontageEvents()
 		{
 			TriggerMontageBlendingOutEvent(MontageBlendingOutEvent);
 		}
-		QueuedMontageBlendingOutEvents.Empty();
+		QueuedMontageBlendingOutEvents.Reset();
 	}
 
 	if (QueuedMontageEndedEvents.Num() > 0)
@@ -1770,7 +1795,7 @@ void UAnimInstance::TriggerQueuedMontageEvents()
 		{
 			TriggerMontageEndedEvent(MontageEndedEvent);
 		}
-		QueuedMontageEndedEvents.Empty();
+		QueuedMontageEndedEvents.Reset();
 	}
 }
 
diff --git a/Engine/Source/Runtime/Engine/Private/Animation/AnimNodeBase.cpp b/Engine/Source/Runtime/Engine/Private/Animation/AnimNodeBase.cpp
--- a/Engine/Source/Runtime/Engine/Private/Animation/AnimNodeBase.cpp
+++ b/Engine/Source/Runtime/Engine/Private/Animation/AnimNodeBase.cpp
@@ -57,6 +57,7 @@ void FPoseLinkBase::CacheBones(const FAnimationCacheBonesContext& Context)
 
 void FPoseLinkBase::Update(const FAnimationUpdateContext& Context)
 {
+	QUICK_SCOPE_CYCLE_COUNTER(STAT_FPoseLinkBase_Update);
 #if DO_CHECK
 	checkf( !bProcessed, TEXT( "Update already in progress, circular link for AnimInstance [%s] Blueprint [%s]" ), \
 		Context.AnimInstance ? *Context.AnimInstance->GetFullName() : TEXT( "None" ), Context.GetAnimBlueprintClass() ? *Context.GetAnimBlueprintClass()->GetFullName() : TEXT( "None" ) );
diff --git a/Engine/Source/Runtime/Engine/Private/Animation/AnimNode_StateMachine.cpp b/Engine/Source/Runtime/Engine/Private/Animation/AnimNode_StateMachine.cpp
--- a/Engine/Source/Runtime/Engine/Private/Animation/AnimNode_StateMachine.cpp
+++ b/Engine/Source/Runtime/Engine/Private/Animation/AnimNode_StateMachine.cpp
@@ -171,7 +171,8 @@ void FAnimNode_StateMachine::Initialize(const FAnimationInitializeContext& Conte
 		if (Machine->States.Num() > 0)
 		{
 			// Create a pose link for each state we can reach
-			StatePoseLinks.Empty(Machine->States.Num());
+			StatePoseLinks.Reset();
+			StatePoseLinks.Reserve(Machine->States.Num());
 			for (int32 StateIndex = 0; StateIndex < Machine->States.Num(); ++StateIndex)
 			{
 				FPoseLink* StatePoseLink = new (StatePoseLinks) FPoseLink();
@@ -184,8 +185,8 @@ void FAnimNode_StateMachine::Initialize(const FAnimationInitializeContext& Conte
 			}
 
 			// Reset transition related variables
-			StatesUpdated.Empty(StatesUpdated.Num());
-			ActiveTransitionArray.Empty(ActiveTransitionArray.Num());
+			StatesUpdated.Reset();
+			ActiveTransitionArray.Reset();
 		
 			// Move to the default state
 			SetState(Context, Machine->InitialState);
@@ -350,7 +351,7 @@ void FAnimNode_StateMachine::Update(const FAnimationUpdateContext& Context)
 		bFirstUpdate = false;
 	}
 
-	StatesUpdated.Empty(StatesUpdated.Num());
+	StatesUpdated.Reset();
 
 	// Tick the individual state/states that are active
 	if (ActiveTransitionArray.Num() > 0)
diff --git a/Engine/Source/Runtime/Engine/Private/Animation/BlendSpace.cpp b/Engine/Source/Runtime/Engine/Private/Animation/BlendSpace.cpp
--- a/Engine/Source/Runtime/Engine/Private/Animation/BlendSpace.cpp
+++ b/Engine/Source/Runtime/Engine/Private/Animation/BlendSpace.cpp
@@ -77,9 +77,9 @@ void UBlendSpace::GetGridSamplesFromBlendInput(const FVector &BlendInput, FGridB
 	}
 }
 
-void UBlendSpace::GetRawSamplesFromBlendInput(const FVector &BlendInput, TArray<FGridBlendSample> & OutBlendSamples) const
+void UBlendSpace::GetRawSamplesFromBlendInput(const FVector &BlendInput, TArray<FGridBlendSample, TInlineAllocator<4> > & OutBlendSamples) const
 {
-	OutBlendSamples.Empty(4);
+	OutBlendSamples.Reset();
 	OutBlendSamples.AddUninitialized(4);
 
 	GetGridSamplesFromBlendInput(BlendInput, OutBlendSamples[0], OutBlendSamples[1], OutBlendSamples[2], OutBlendSamples[3]);
diff --git a/Engine/Source/Runtime/Engine/Private/Animation/BlendSpace1D.cpp b/Engine/Source/Runtime/Engine/Private/Animation/BlendSpace1D.cpp
--- a/Engine/Source/Runtime/Engine/Private/Animation/BlendSpace1D.cpp
+++ b/Engine/Source/Runtime/Engine/Private/Animation/BlendSpace1D.cpp
@@ -62,7 +62,7 @@ bool UBlendSpace1D::IsSameSamplePoint(const FVector& SamplePointA, const FVector
 	return (Diff.X < Threshold);
 }
 
-void UBlendSpace1D::GetRawSamplesFromBlendInput(const FVector &BlendInput, TArray<FGridBlendSample> & OutBlendSamples) const
+void UBlendSpace1D::GetRawSamplesFromBlendInput(const FVector &BlendInput, TArray<FGridBlendSample, TInlineAllocator<4> > & OutBlendSamples) const
 {
 
 	FVector NormalizedBlendInput = GetNormalizedBlendInput(BlendInput);
diff --git a/Engine/Source/Runtime/Engine/Private/Animation/BlendSpaceBase.cpp b/Engine/Source/Runtime/Engine/Private/Animation/BlendSpaceBase.cpp
--- a/Engine/Source/Runtime/Engine/Private/Animation/BlendSpaceBase.cpp
+++ b/Engine/Source/Runtime/Engine/Private/Animation/BlendSpaceBase.cpp
@@ -138,15 +138,18 @@ void UBlendSpaceBase::TickAssetPlayerInstance(const FAnimTickRecord& Instance, c
 			UE_LOG(LogAnimation, Log, TEXT("BlendSpace(%s) - BlendInput(%s) : FilteredBlendInput(%s), FilterMultiplier(%0.2f)"), *GetName(), *Instance.BlendSpacePosition.ToString(), *BlendInput.ToString(), FilterMultiplier );
 		}
 
-		check (Instance.BlendSampleDataCache);
+		check(Instance.BlendSampleDataCache);
 
 		// For Target weight interpolation, we'll need to save old data, and interpolate to new data
-		TArray<FBlendSampleData> OldSampleDataList, NewSampleDataList;
+		static TArray<FBlendSampleData> OldSampleDataList;
+		static TArray<FBlendSampleData> NewSampleDataList;
+		check(IsInGameThread() && !OldSampleDataList.Num() && !NewSampleDataList.Num()); // this must be called non-recursively on the game thread
+
 		OldSampleDataList.Append(*Instance.BlendSampleDataCache);
 
 		// get sample data based on new input
 		// consolidate all samples and sort them, so that we can handle from biggest weight to smallest
-		Instance.BlendSampleDataCache->Empty();
+		Instance.BlendSampleDataCache->Reset();
 		// new sample data that will be used for evaluation
 		TArray<FBlendSampleData> & SampleDataList = *Instance.BlendSampleDataCache;
 
@@ -287,15 +290,19 @@ void UBlendSpaceBase::TickAssetPlayerInstance(const FAnimTickRecord& Instance, c
 				}
 			}
 		}
+		OldSampleDataList.Reset();
+		NewSampleDataList.Reset();
 	}
 }
 
 bool UBlendSpaceBase::GetSamplesFromBlendInput(const FVector &BlendInput, TArray<FBlendSampleData> & OutSampleDataList) const
 {
-	TArray<FGridBlendSample> RawGridSamples;
+	static TArray<FGridBlendSample, TInlineAllocator<4> > RawGridSamples;
+	check(IsInGameThread() && !RawGridSamples.Num()); // this must be called non-recursively from the gamethread
 	GetRawSamplesFromBlendInput(BlendInput, RawGridSamples);
 
-	OutSampleDataList.Empty();
+	OutSampleDataList.Reset();
+	OutSampleDataList.Reserve(RawGridSamples.Num() * FEditorElement::MAX_VERTICES);
 
 	// consolidate all samples
 	for (int32 SampleNum=0; SampleNum<RawGridSamples.Num(); ++SampleNum)
@@ -329,7 +336,7 @@ bool UBlendSpaceBase::GetSamplesFromBlendInput(const FVector &BlendInput, TArray
 		if (OutSampleDataList[I].TotalWeight < ZERO_ANIMWEIGHT_THRESH)
 		{
 			// cut anything in front of this 
-			OutSampleDataList.RemoveAt(I, TotalSample-I);
+			OutSampleDataList.RemoveAt(I, TotalSample-I, false); // we won't shrink here, that might screw up alloc optimization at a higher level, if not this is temp anyway
 			break;
 		}
 
@@ -341,7 +348,7 @@ bool UBlendSpaceBase::GetSamplesFromBlendInput(const FVector &BlendInput, TArray
 		// normalize to all weights
 		OutSampleDataList[I].TotalWeight /= TotalWeight;
 	}
-
+	RawGridSamples.Reset();
 	return (OutSampleDataList.Num()!=0);
 }
 
diff --git a/Engine/Source/Runtime/Engine/Private/Components/SceneComponent.cpp b/Engine/Source/Runtime/Engine/Private/Components/SceneComponent.cpp
--- a/Engine/Source/Runtime/Engine/Private/Components/SceneComponent.cpp
+++ b/Engine/Source/Runtime/Engine/Private/Components/SceneComponent.cpp
@@ -1112,7 +1112,8 @@ USceneComponent* USceneComponent::GetChildComponent(int32 ChildIndex) const
 
 void USceneComponent::GetChildrenComponents(bool bIncludeAllDescendants, TArray<USceneComponent*>& Children) const
 {
-	Children.Empty();
+	Children.Reset();
+	Children.Reserve(AttachChildren.Num());
 
 	if (bIncludeAllDescendants)
 	{
diff --git a/Engine/Source/Runtime/Engine/Private/Components/SkeletalMeshComponent.cpp b/Engine/Source/Runtime/Engine/Private/Components/SkeletalMeshComponent.cpp
--- a/Engine/Source/Runtime/Engine/Private/Components/SkeletalMeshComponent.cpp
+++ b/Engine/Source/Runtime/Engine/Private/Components/SkeletalMeshComponent.cpp
@@ -101,6 +101,7 @@ public:
 		if (USkeletalMeshComponent* Comp = SkeletalMeshComponent.Get())
 		{
 			const bool bPerformPostAnimEvaluation = true;
+			QUICK_SCOPE_CYCLE_COUNTER(STAT_USkeletalMeshComponent_CompleteParallelAnimationEvaluation);
 			Comp->CompleteParallelAnimationEvaluation(bPerformPostAnimEvaluation);
 		}
 	}
@@ -169,20 +170,20 @@ void USkeletalMeshComponent::RegisterPreClothTick(bool bRegister)
 {
 	if (bRegister != PreClothTickFunction.IsTickFunctionRegistered())
 	{
-	if (bRegister)
-	{
-		if (SetupActorComponentTickFunction(&PreClothTickFunction))
+		if (bRegister)
 		{
-			PreClothTickFunction.Target = this;
-			// Set a prereq for the pre cloth tick to happen after physics is finished
-			if (World != NULL)
+			if (SetupActorComponentTickFunction(&PreClothTickFunction))
 			{
-				PreClothTickFunction.AddPrerequisite(World, World->EndPhysicsTickFunction);
+				PreClothTickFunction.Target = this;
+				// Set a prereq for the pre cloth tick to happen after physics is finished
+				if (World != NULL)
+				{
+					PreClothTickFunction.AddPrerequisite(World, World->EndPhysicsTickFunction);
+				}
 			}
 		}
-	}
-	else
-	{
+		else
+		{
 			PreClothTickFunction.UnRegisterTickFunction();
 		}
 	}
@@ -274,8 +275,8 @@ void USkeletalMeshComponent::InitAnim(bool bForceReinit)
 
 		InitializeAnimScriptInstance(bForceReinit);
 
-		//Make sure we have a valid pose		
-		TickAnimation(0.f); 
+        //Make sure we have a valid pose		
+        TickAnimation(0.f); 
 
 		RefreshBoneTransforms();
 		UpdateComponentToWorld();
@@ -495,8 +496,9 @@ void USkeletalMeshComponent::LoadedFromAnotherClass(const FName& OldClassName)
 
 void USkeletalMeshComponent::TickAnimation(float DeltaTime)
 {
-	SCOPE_CYCLE_COUNTER(STAT_AnimTickTime);
+
 	SCOPE_CYCLE_COUNTER(STAT_AnimGameThreadTime);
+	SCOPE_CYCLE_COUNTER(STAT_AnimTickTime);
 	if (SkeletalMesh != NULL)
 	{
 		if (AnimScriptInstance != NULL)
@@ -508,22 +510,26 @@ void USkeletalMeshComponent::TickAnimation(float DeltaTime)
 			// AnimNotifies?  Please take a look and fix as we discussed.  Temporary fix:
 			if (AnimScriptInstance != NULL)
 			{
-				// now all tick/trigger/kismet is done
-				// add MorphTarget Curves from Kismet driven or any other source
-				// and overwrite if it exists
-				// Tick always should maintain this list, not Evaluate
-				for( auto Iter = MorphTargetCurves.CreateConstIterator(); Iter; ++Iter )
 				{
-					float *CurveValPtr = AnimScriptInstance->MorphTargetCurves.Find(Iter.Key());
-					if ( CurveValPtr )
-					{
-						// override the value if Kismet request was made
-						*CurveValPtr = Iter.Value();
-					}
-					else
-					{
-						AnimScriptInstance->MorphTargetCurves.Add(Iter.Key(), Iter.Value());
-					}				
+					QUICK_SCOPE_CYCLE_COUNTER(STAT_USkeletalMeshComponent_TickAnimation_MorphTargetCurves);
+
+				    // now all tick/trigger/kismet is done
+				    // add MorphTarget Curves from Kismet driven or any other source
+				    // and overwrite if it exists
+				    // Tick always should maintain this list, not Evaluate
+				    for( auto Iter = MorphTargetCurves.CreateConstIterator(); Iter; ++Iter )
+				    {
+					    float *CurveValPtr = AnimScriptInstance->MorphTargetCurves.Find(Iter.Key());
+					    if ( CurveValPtr )
+					    {
+						    // override the value if Kismet request was made
+						    *CurveValPtr = Iter.Value();
+					    }
+					    else
+					    {
+						    AnimScriptInstance->MorphTargetCurves.Add(Iter.Key(), Iter.Value());
+					    }				
+				    }
 				}
 
 				//Update material parameters
@@ -537,6 +543,8 @@ void USkeletalMeshComponent::UpdateMaterialParameters()
 {
 	if(AnimScriptInstance->MaterialParameterCurves.Num() > 0)
 	{
+		QUICK_SCOPE_CYCLE_COUNTER(STAT_USkeletalMeshComponent_UpdateMaterialParameters);
+
 		for(auto Iter = AnimScriptInstance->MaterialParameterCurves.CreateConstIterator(); Iter; ++Iter)
 		{
 			FName ParameterName = Iter.Key();
@@ -607,17 +615,34 @@ bool USkeletalMeshComponent::ShouldTickPose() const
 	return (Super::ShouldTickPose() && IsRegistered() && AnimScriptInstance && !bAutonomousTickPose && !bPauseAnims && GetWorld()->AreActorsInitialized() && !bNoSkeletonUpdate && !bAlreadyTickedThisFrame);
 }
 
+static FThreadSafeCounter Ticked;
+static FThreadSafeCounter NotTicked;
+
+static TAutoConsoleVariable<int32> CVarSpewAnimRateOptimization(
+	TEXT("SpewAnimRateOptimization"),
+	0,
+	TEXT("True to spew overall anim rate optimization tick rates."));
+
 void USkeletalMeshComponent::TickPose(float DeltaTime, bool bNeedsValidRootMotion)
 {
 	Super::TickPose(DeltaTime, bNeedsValidRootMotion);
 
-	if (AnimUpdateRateParams != nullptr && (!bEnableUpdateRateOptimizations || !AnimUpdateRateParams->ShouldSkipUpdate()))
+	if ((AnimUpdateRateParams != nullptr) && (!ShouldUseUpdateRateOptimizations() || !AnimUpdateRateParams->ShouldSkipUpdate()))
 	{
 		float TimeAdjustment = AnimUpdateRateParams->GetTimeAdjustment();
 		TickAnimation(DeltaTime + TimeAdjustment);
 		LastPoseTickTime = GetWorld()->TimeSeconds;
+		if (CVarSpewAnimRateOptimization.GetValueOnGameThread() > 0 && Ticked.Increment()==500)
+		{
+			UE_LOG(LogTemp, Display, TEXT("%d Ticked %d NotTicked"), Ticked.GetValue(), NotTicked.GetValue());
+			Ticked.Reset();
+			NotTicked.Reset();
+		}
+	}
+	else if (CVarSpewAnimRateOptimization.GetValueOnGameThread())
+	{
+		NotTicked.Increment();
 	}
-	//return bPoseTicked;
 }
 
 void USkeletalMeshComponent::TickComponent(float DeltaTime, enum ELevelTick TickType, FActorComponentTickFunction *ThisTickFunction)
@@ -1035,7 +1060,8 @@ void USkeletalMeshComponent::PerformAnimationEvaluation(const USkeletalMesh* InS
 }
 
 static TAutoConsoleVariable<int32> CVarHiPriSkinnedMeshesTicks(
-	TEXT("tick.HiPriSkinnedMeshes"),1,
+	TEXT("tick.HiPriSkinnedMeshes"),
+	1,
 	TEXT("If > 0, then schedule the skinned component ticks in a tick group before other ticks."));
 
 int32 GetCurveNumber(USkeleton* Skeleton)
@@ -1055,8 +1081,8 @@ int32 GetCurveNumber(USkeleton* Skeleton)
 
 void USkeletalMeshComponent::RefreshBoneTransforms(FActorComponentTickFunction* TickFunction)
 {
-	SCOPE_CYCLE_COUNTER(STAT_RefreshBoneTransforms);
 	SCOPE_CYCLE_COUNTER(STAT_AnimGameThreadTime);
+	SCOPE_CYCLE_COUNTER(STAT_RefreshBoneTransforms);
 
 	check(IsInGameThread()); //Only want to call this from the game thread as we set up tasks etc
 
@@ -1068,20 +1094,22 @@ void USkeletalMeshComponent::RefreshBoneTransforms(FActorComponentTickFunction*
 	// Recalculate the RequiredBones array, if necessary
 	if (!bRequiredBonesUpToDate)
 	{
+		QUICK_SCOPE_CYCLE_COUNTER(STAT_USkeletalMeshComponent_RefreshBoneTransforms_RecalcRequiredBones);
 		RecalcRequiredBones(PredictedLODLevel);
 	}
 
 	if (AnimScriptInstance)
 	{
+		QUICK_SCOPE_CYCLE_COUNTER(STAT_USkeletalMeshComponent_RefreshBoneTransforms_UpdateMontageEvaluationData);
 		AnimScriptInstance->UpdateMontageEvaluationData();
 	}
-	const bool bDoEvaluationRateOptimization = bEnableUpdateRateOptimizations && AnimUpdateRateParams->DoEvaluationRateOptimizations();
+	const bool bDoEvaluationRateOptimization = ShouldUseUpdateRateOptimizations() && AnimUpdateRateParams->DoEvaluationRateOptimizations();
 
 	//Handle update rate optimization setup
 	//Dont mark cache as invalid if we aren't performing optimization anyway
 	const bool bInvalidCachedBones = bDoEvaluationRateOptimization &&
 									 ((LocalAtoms.Num() != SkeletalMesh->RefSkeleton.GetNum())
-									  || (LocalAtoms.Num() != CachedLocalAtoms.Num())
+									 || (LocalAtoms.Num() != CachedLocalAtoms.Num())
 									 || (GetNumSpaceBases() != CachedSpaceBases.Num()));
 
 
@@ -1094,8 +1122,8 @@ void USkeletalMeshComponent::RefreshBoneTransforms(FActorComponentTickFunction*
 
 	const bool bDoParallelEvaluation = bShouldDoEvaluation && TickFunction && bDoPAE;
 
-	const bool bBlockOnTask = !bDoParallelEvaluation;   // If we aren't trying to do parallel evaluation then we
-														// will need to wait on an existing task.
+	const bool bBlockOnTask = !bDoParallelEvaluation;  // If we aren't trying to do parallel evaluation then we
+															// will need to wait on an existing task.
 
 	const bool bPerformPostAnimEvaluation = true;
 	if (HandleExistingParallelEvaluationTask(bBlockOnTask, bPerformPostAnimEvaluation))
@@ -1117,13 +1145,15 @@ void USkeletalMeshComponent::RefreshBoneTransforms(FActorComponentTickFunction*
 	if (!bDoEvaluationRateOptimization)
 	{
 		//If we aren't optimizing clear the cached local atoms
-		CachedLocalAtoms.Empty();
-		CachedSpaceBases.Empty();
+		CachedLocalAtoms.Reset();
+		CachedSpaceBases.Reset();
 		CachedCurve.Empty();
 	}
 
 	if (bDoParallelEvaluation)
 	{
+		QUICK_SCOPE_CYCLE_COUNTER(STAT_USkeletalMeshComponent_RefreshBoneTransforms_SetupParallel); 
+
 		// this makes no sense to do on the fly. If we want to use hipri ticks, this needs to be done at register time
 		// this does allow us to toggle it easily for performance testing
 		bool bDoHiPri = CVarHiPriSkinnedMeshesTicks.GetValueOnGameThread() > 0;
@@ -1135,29 +1165,33 @@ void USkeletalMeshComponent::RefreshBoneTransforms(FActorComponentTickFunction*
 		if (SkeletalMesh->RefSkeleton.GetNum() != AnimEvaluationContext.LocalAtoms.Num())
 		{
 			// Initialize Parallel Task arrays
-			AnimEvaluationContext.LocalAtoms = LocalAtoms;
-			AnimEvaluationContext.SpaceBases = GetSpaceBases();
-			AnimEvaluationContext.VertexAnims = ActiveVertexAnims;
+			AnimEvaluationContext.LocalAtoms.Reset();
+			AnimEvaluationContext.LocalAtoms.Append(LocalAtoms);
+			AnimEvaluationContext.SpaceBases.Reset();
+			AnimEvaluationContext.SpaceBases.Append(GetSpaceBases());
+			AnimEvaluationContext.VertexAnims.Reset();
+			AnimEvaluationContext.VertexAnims.Append(ActiveVertexAnims);
 		}
 
-		// start parallel work
-		check(!IsValidRef(ParallelAnimationEvaluationTask));
-		ParallelAnimationEvaluationTask = TGraphTask<FParallelAnimationEvaluationTask>::CreateTask().ConstructAndDispatchWhenReady(this);
+			// start parallel work
+			check(!IsValidRef(ParallelAnimationEvaluationTask));
+			ParallelAnimationEvaluationTask = TGraphTask<FParallelAnimationEvaluationTask>::CreateTask().ConstructAndDispatchWhenReady(this);
 
-		// set up a task to run on the game thread to accept the results
-		FGraphEventArray Prerequistes;
-		Prerequistes.Add(ParallelAnimationEvaluationTask);
-		FGraphEventRef TickCompletionEvent = TGraphTask<FParallelAnimationCompletionTask>::CreateTask(&Prerequistes).ConstructAndDispatchWhenReady(this);
+			// set up a task to run on the game thread to accept the results
+			FGraphEventArray Prerequistes;
+			Prerequistes.Add(ParallelAnimationEvaluationTask);
+			FGraphEventRef TickCompletionEvent = TGraphTask<FParallelAnimationCompletionTask>::CreateTask(&Prerequistes).ConstructAndDispatchWhenReady(this);
 
-		if ( TickFunction )
-		{
-		TickFunction->GetCompletionHandle()->DontCompleteUntil(TickCompletionEvent);
-	}
-	}
+			if ( TickFunction )
+			{
+				TickFunction->GetCompletionHandle()->DontCompleteUntil(TickCompletionEvent);
+			}
+		}
 	else
 	{
 		if (AnimEvaluationContext.bDoEvaluation)
 		{
+			QUICK_SCOPE_CYCLE_COUNTER(STAT_USkeletalMeshComponent_RefreshBoneTransforms_GamethreadEval);
 			if (AnimEvaluationContext.bDoInterpolation)
 			{
 				PerformAnimationEvaluation(SkeletalMesh, AnimScriptInstance, CachedSpaceBases, CachedLocalAtoms, ActiveVertexAnims, RootBoneTranslation, CachedCurve);
@@ -1169,9 +1203,13 @@ void USkeletalMeshComponent::RefreshBoneTransforms(FActorComponentTickFunction*
 		}
 		else if (!AnimEvaluationContext.bDoInterpolation)
 		{
-			LocalAtoms = CachedLocalAtoms;
-			GetEditableSpaceBases() = CachedSpaceBases;
-			AnimEvaluationContext.Curve = CachedCurve;
+			QUICK_SCOPE_CYCLE_COUNTER(STAT_USkeletalMeshComponent_RefreshBoneTransforms_CopyBones);
+			LocalAtoms.Reset();
+			LocalAtoms.Append(CachedLocalAtoms);
+			TArray<FTransform>& LocalEditableSpaceBases = GetEditableSpaceBases();
+			LocalEditableSpaceBases.Reset();
+			LocalEditableSpaceBases.Append(CachedSpaceBases);
+			AnimEvaluationContext.Curve.CopyFrom(CachedCurve);
 		}
 
 		PostAnimEvaluation(AnimEvaluationContext);
@@ -1199,8 +1237,10 @@ void USkeletalMeshComponent::PostAnimEvaluation(FAnimationEvaluationContext& Eva
 	
 	if (EvaluationContext.bDuplicateToCacheBones)
 	{
-		CachedSpaceBases = GetEditableSpaceBases();
-		CachedLocalAtoms = LocalAtoms;
+		CachedSpaceBases.Reset();
+		CachedSpaceBases.Append(GetEditableSpaceBases());
+		CachedLocalAtoms.Reset();
+		CachedLocalAtoms.Append(LocalAtoms);
 	}
 
 	if (EvaluationContext.bDuplicateToCacheCurve)
@@ -1218,7 +1258,9 @@ void USkeletalMeshComponent::PostAnimEvaluation(FAnimationEvaluationContext& Eva
 		{
 			// We need to prep our space bases for interp (TODO: Would a new Lerp function that took
 			// separate input and output be quicker than current copy + lerp in place?)
-			GetEditableSpaceBases() = GetSpaceBases();
+			TArray<FTransform>& LocalEditableSpaceBases = GetEditableSpaceBases();
+			LocalEditableSpaceBases.Reset();
+			LocalEditableSpaceBases.Append(GetSpaceBases());
 		}
 		FAnimationRuntime::LerpBoneTransforms(GetEditableSpaceBases(), CachedSpaceBases, Alpha, RequiredBones);
 
@@ -1304,8 +1346,8 @@ FBoxSphereBounds USkeletalMeshComponent::CalcBounds(const FTransform& LocalToWor
 			MasterPoseComponentInst->IsA((USkeletalMeshComponent::StaticClass())))
 		{
 			const USkeletalMeshComponent* BaseComponent = CastChecked<USkeletalMeshComponent>(MasterPoseComponentInst);
-		RootBoneOffset = BaseComponent->RootBoneTranslation; // Adjust bounds by root bone translation
-	}
+			RootBoneOffset = BaseComponent->RootBoneTranslation; // Adjust bounds by root bone translation
+		}
 	}
 
 	FBoxSphereBounds NewBounds = CalcMeshBound( RootBoneOffset, bHasValidBodies, LocalToWorld );
@@ -1896,7 +1938,7 @@ FRootMotionMovementParams USkeletalMeshComponent::ConsumeRootMotion()
 {
 	if (AnimScriptInstance)
 	{
-		float InterpAlpha = bEnableUpdateRateOptimizations ? AnimUpdateRateParams->GetRootMotionInterp() : 1.f;
+		float InterpAlpha = ShouldUseUpdateRateOptimizations() ? AnimUpdateRateParams->GetRootMotionInterp() : 1.f;
 		return AnimScriptInstance->ConsumeExtractedRootMotion(InterpAlpha);
 	}
 	return FRootMotionMovementParams();
@@ -2000,8 +2042,8 @@ bool USkeletalMeshComponent::ComponentIsTouchingSelectionBox(const FBox& InSelBB
 		// is consider touching
 		if (bMustEncompassEntireComponent)
 		{
-			return true;
-		}
+		return true;
+	}
 	}
 
 	return false;
@@ -2114,7 +2156,7 @@ void USkeletalMeshComponent::ValidateAnimation()
 		{
 			if (SkeletalMesh->Skeleton)
 			{
-			UE_LOG(LogAnimation, Warning, TEXT("Animation %s is incompatible with skeleton %s, removing animation from actor."), *AnimationData.AnimToPlay->GetName(), *SkeletalMesh->Skeleton->GetName());
+				UE_LOG(LogAnimation, Warning, TEXT("Animation %s is incompatible with skeleton %s, removing animation from actor."), *AnimationData.AnimToPlay->GetName(), *SkeletalMesh->Skeleton->GetName());
 			}
 			else
 			{
@@ -2130,7 +2172,7 @@ void USkeletalMeshComponent::ValidateAnimation()
 		{
 			if(SkeletalMesh->Skeleton)
 			{
-			UE_LOG(LogAnimation, Warning, TEXT("AnimBP %s is incompatible with skeleton %s, removing AnimBP from actor."), *AnimBlueprintGeneratedClass->GetName(), *SkeletalMesh->Skeleton->GetName());
+				UE_LOG(LogAnimation, Warning, TEXT("AnimBP %s is incompatible with skeleton %s, removing AnimBP from actor."), *AnimBlueprintGeneratedClass->GetName(), *SkeletalMesh->Skeleton->GetName());
 			}
 			else
 			{
diff --git a/Engine/Source/Runtime/Engine/Private/Components/SkinnedMeshComponent.cpp b/Engine/Source/Runtime/Engine/Private/Components/SkinnedMeshComponent.cpp
--- a/Engine/Source/Runtime/Engine/Private/Components/SkinnedMeshComponent.cpp
+++ b/Engine/Source/Runtime/Engine/Private/Components/SkinnedMeshComponent.cpp
@@ -26,6 +26,15 @@ FAutoConsoleVariableRef CVarSkeletalMeshLODBias(
 	ECVF_Scalability
 	);
 
+static TAutoConsoleVariable<int32> CVarEnableAnimRateOptimization(
+	TEXT("EnableAnimRateOptimization"),
+	1,
+	TEXT("True to anim rate optimization."));
+
+static TAutoConsoleVariable<int32> CVarDrawAnimRateOptimization(
+	TEXT("DrawAnimRateOptimization"),
+	0,
+	TEXT("True to draw color coded boxes for anim rate."));
 
 namespace FAnimUpdateRateManager
 {
@@ -121,6 +130,11 @@ namespace FAnimUpdateRateManager
 		}
 	}
 
+	static TAutoConsoleVariable<int32> CVarForceAnimRate(
+		TEXT("ForceAnimRate"),
+		0,
+		TEXT("Non-zero to force anim rate. 10 = eval anim every ten frames for those meshes that can do it. In some cases a frame is considered to be 30fps."));
+
 	void AnimUpdateRateSetParams(FAnimUpdateRateParametersTracker* Tracker, float DeltaTime, bool bRecentlyRendered, float MaxDistanceFactor, bool bNeedsValidRootMotion, bool bUsingRootMotionFromEverything)
 	{
 		// default rules for setting update rates
@@ -155,6 +169,12 @@ namespace FAnimUpdateRateManager
 				}
 			}
 
+			int32 ForceAnimRate = CVarForceAnimRate.GetValueOnGameThread();
+			if (ForceAnimRate)
+			{
+				DesiredEvaluationRate = ForceAnimRate;
+			}
+
 			if (bUsingRootMotionFromEverything && DesiredEvaluationRate > 1)
 			{
 				//Use look ahead mode that allows us to rate limit updates even when using root motion
@@ -410,19 +430,22 @@ bool USkinnedMeshComponent::ShouldUpdateTransform(bool bLODHasChanged) const
 	return (bLODHasChanged || bRecentlyRendered || (MeshComponentUpdateFlag == EMeshComponentUpdateFlag::AlwaysTickPoseAndRefreshBones));
 }
 
+bool USkinnedMeshComponent::ShouldUseUpdateRateOptimizations() const
+{
+	return bEnableUpdateRateOptimizations && CVarEnableAnimRateOptimization.GetValueOnGameThread() > 0;
+}
 
 void USkinnedMeshComponent::TickUpdateRate(float DeltaTime, bool bNeedsValidRootMotion)
 {
 	SCOPE_CYCLE_COUNTER(STAT_TickUpdateRate);
-	if( bEnableUpdateRateOptimizations )
+	if (ShouldUseUpdateRateOptimizations())
 	{
 		if (GetOwner())
 		{
 			// Tick Owner once per frame. All attached SkinnedMeshComponents will share the same settings.
 			FAnimUpdateRateManager::TickUpdateRateParameters(this, DeltaTime, bNeedsValidRootMotion);
 
-			// debug -- @todo: hook this up to a console command.
-			if (bDisplayDebugUpdateRateOptimizations)
+			if ((CVarDrawAnimRateOptimization.GetValueOnGameThread() > 0) || bDisplayDebugUpdateRateOptimizations)
 			{
 				FColor DrawColor = AnimUpdateRateParams->GetUpdateRateDebugColor();
 				DrawDebugBox(GetWorld(), Bounds.Origin, Bounds.BoxExtent, FQuat::Identity, DrawColor, false);
diff --git a/Engine/Source/Runtime/Engine/Private/GameEngine.cpp b/Engine/Source/Runtime/Engine/Private/GameEngine.cpp
--- a/Engine/Source/Runtime/Engine/Private/GameEngine.cpp
+++ b/Engine/Source/Runtime/Engine/Private/GameEngine.cpp
@@ -803,8 +803,128 @@ float UGameEngine::GetMaxTickRate(float DeltaTime, bool bAllowFrameRateSmoothing
 	return MaxTickRate;
 }
 
+#if !UE_BUILD_SHIPPING && !UE_BUILD_TEST
+
+#include "StackTracker.h"
+static TAutoConsoleVariable<int32> CVarLogGameThreadMallocChurn(
+	TEXT("LogGameThreadMallocChurn.Enable"),
+	0,
+	TEXT("If > 0, then collect sample game thread malloc, realloc and free, periodically print a report of the worst offenders."));
+
+static TAutoConsoleVariable<int32> CVarLogGameThreadMallocChurn_PrintFrequency(
+	TEXT("LogGameThreadMallocChurn.PrintFrequency"),
+	300,
+	TEXT("Number of frames between churn reports."));
+
+static TAutoConsoleVariable<int32> CVarLogGameThreadMallocChurn_Threshhold(
+	TEXT("LogGameThreadMallocChurn.Threshhold"),
+	10,
+	TEXT("Minimum average number of allocs per frame to include in the report."));
+
+static TAutoConsoleVariable<int32> CVarLogGameThreadMallocChurn_SampleFrequency(
+	TEXT("LogGameThreadMallocChurn.SampleFrequency"),
+	100,
+	TEXT("Number of allocs to skip between samples. This is used to prevent churn sampling from slowing the game down too much."));
+
+static TAutoConsoleVariable<int32> CVarLogGameThreadMallocChurn_StackIgnore(
+	TEXT("LogGameThreadMallocChurn.StackIgnore"),
+	2,
+	TEXT("Number of items to discard from the top of a stack frame."));
+
+static TAutoConsoleVariable<int32> CVarLogGameThreadMallocChurn_StackLen(
+	TEXT("LogGameThreadMallocChurn.StackLen"),
+	3,
+	TEXT("Maximum number of stack frame items to keep. This improves aggregation because calls that originate from multiple places but end up in the same place will be accounted together."));
+
+
+extern CORE_API TFunction<void(int32)>* GGameThreadMallocHook;
+
+struct FScopedSampleMallocChurn
+{
+	static FStackTracker GGameThreadMallocChurnTracker;
+	static uint64 DumpFrame;
+
+	bool bEnabled;
+	int32 CountDown;
+	TFunction<void(int32)> Hook;
+
+	FScopedSampleMallocChurn()
+		: bEnabled(CVarLogGameThreadMallocChurn.GetValueOnGameThread() > 0)
+		, CountDown(CVarLogGameThreadMallocChurn_SampleFrequency.GetValueOnGameThread())
+		, Hook(
+		[this](int32 Index)
+	{
+		if (--CountDown <= 0)
+		{
+			CountDown = CVarLogGameThreadMallocChurn_SampleFrequency.GetValueOnGameThread();
+			CollectSample();
+		}
+	}
+	)
+	{
+		if (bEnabled)
+		{
+			check(IsInGameThread());
+			check(!GGameThreadMallocHook);
+			if (!DumpFrame)
+			{
+				DumpFrame = GFrameCounter + CVarLogGameThreadMallocChurn_PrintFrequency.GetValueOnGameThread();
+				GGameThreadMallocChurnTracker.ResetTracking();
+			}
+			GGameThreadMallocChurnTracker.ToggleTracking(true, true);
+			GGameThreadMallocHook = &Hook;
+		}
+		else
+		{
+			check(IsInGameThread());
+			GGameThreadMallocChurnTracker.ToggleTracking(false, true);
+			if (DumpFrame)
+			{
+				DumpFrame = 0;
+				GGameThreadMallocChurnTracker.ResetTracking();
+			}
+		}
+	}
+	~FScopedSampleMallocChurn()
+	{
+		if (bEnabled)
+		{
+			check(IsInGameThread());
+			check(GGameThreadMallocHook == &Hook);
+			GGameThreadMallocHook = nullptr;
+			GGameThreadMallocChurnTracker.ToggleTracking(false, true);
+			check(DumpFrame);
+			if (GFrameCounter > DumpFrame)
+			{
+				PrintResultsAndReset();
+			}
+		}
+	}
+
+	void CollectSample()
+	{
+		check(IsInGameThread());
+		GGameThreadMallocChurnTracker.CaptureStackTrace(CVarLogGameThreadMallocChurn_StackIgnore.GetValueOnGameThread(), nullptr, CVarLogGameThreadMallocChurn_StackLen.GetValueOnGameThread());
+	}
+	void PrintResultsAndReset()
+	{
+		DumpFrame = GFrameCounter + CVarLogGameThreadMallocChurn_PrintFrequency.GetValueOnGameThread();
+		FOutputDeviceRedirector* Log = FOutputDeviceRedirector::Get();
+		float SampleAndFrameCorrection = float(CVarLogGameThreadMallocChurn_SampleFrequency.GetValueOnGameThread()) / float(CVarLogGameThreadMallocChurn_PrintFrequency.GetValueOnGameThread());
+		GGameThreadMallocChurnTracker.DumpStackTraces(CVarLogGameThreadMallocChurn_Threshhold.GetValueOnGameThread(), *Log, SampleAndFrameCorrection);
+		GGameThreadMallocChurnTracker.ResetTracking();
+	}
+};
+FStackTracker FScopedSampleMallocChurn::GGameThreadMallocChurnTracker;
+uint64 FScopedSampleMallocChurn::DumpFrame = 0;
+
+#endif
+
 void UGameEngine::Tick( float DeltaSeconds, bool bIdleMode )
 {
+#if !UE_BUILD_SHIPPING && !UE_BUILD_TEST
+	FScopedSampleMallocChurn ChurnTracker;
+#endif
 	SCOPE_CYCLE_COUNTER(STAT_GameEngineTick);
 	NETWORK_PROFILER(GNetworkProfiler.TrackFrameBegin());
 
diff --git a/Engine/Source/Runtime/Engine/Private/Particles/ParticleBeam2EmitterInstance.cpp b/Engine/Source/Runtime/Engine/Private/Particles/ParticleBeam2EmitterInstance.cpp
--- a/Engine/Source/Runtime/Engine/Private/Particles/ParticleBeam2EmitterInstance.cpp
+++ b/Engine/Source/Runtime/Engine/Private/Particles/ParticleBeam2EmitterInstance.cpp
@@ -1237,7 +1237,7 @@ FDynamicEmitterDataBase* FParticleBeam2EmitterInstance::GetDynamicData(bool bSel
 	}
 
 	// Allocate the dynamic data
-	FDynamicBeam2EmitterData* NewEmitterData = ::new FDynamicBeam2EmitterData(LODLevel->RequiredModule);
+	FDynamicBeam2EmitterData* NewEmitterData = new FDynamicBeam2EmitterData(LODLevel->RequiredModule);
 	{
 		SCOPE_CYCLE_COUNTER(STAT_ParticleMemTime);
 		INC_DWORD_STAT(STAT_DynamicEmitterCount);
@@ -1327,7 +1327,7 @@ bool FParticleBeam2EmitterInstance::UpdateDynamicData(FDynamicEmitterDataBase* D
  */
 FDynamicEmitterReplayDataBase* FParticleBeam2EmitterInstance::GetReplayData()
 {
-	FDynamicEmitterReplayDataBase* NewEmitterReplayData = ::new FDynamicBeam2EmitterReplayData();
+	FDynamicEmitterReplayDataBase* NewEmitterReplayData = new FDynamicBeam2EmitterReplayData();
 	check( NewEmitterReplayData != NULL );
 
 	if( !FillReplayData( *NewEmitterReplayData ) )
@@ -1573,7 +1573,7 @@ bool FParticleBeam2EmitterInstance::FillReplayData( FDynamicEmitterReplayDataBas
 
 	//@todo. SORTING IS A DIFFERENT ISSUE NOW! 
 	//		 GParticleView isn't going to be valid anymore?
-	uint8* PData = NewReplayData->ParticleData.GetData();
+	uint8* PData = NewReplayData->DataContainer.ParticleData;
 	for (int32 i = 0; i < NewReplayData->ActiveParticleCount; i++)
 	{
 		DECLARE_PARTICLE(Particle, ParticleData + ParticleStride * ParticleIndices[i]);
diff --git a/Engine/Source/Runtime/Engine/Private/Particles/ParticleComponents.cpp b/Engine/Source/Runtime/Engine/Private/Particles/ParticleComponents.cpp
--- a/Engine/Source/Runtime/Engine/Private/Particles/ParticleComponents.cpp
+++ b/Engine/Source/Runtime/Engine/Private/Particles/ParticleComponents.cpp
@@ -2969,7 +2969,7 @@ FDynamicEmitterDataBase* UParticleSystemComponent::CreateDynamicDataFromReplay(
 		case DET_Sprite:
 			{
 				// Allocate the dynamic data
-				FDynamicSpriteEmitterData* NewEmitterData = ::new FDynamicSpriteEmitterData(EmitterInstance->CurrentLODLevel->RequiredModule);
+				FDynamicSpriteEmitterData* NewEmitterData = new FDynamicSpriteEmitterData(EmitterInstance->CurrentLODLevel->RequiredModule);
 
 				// Fill in the source data
 				const FDynamicSpriteEmitterReplayData* SpriteEmitterReplayData =
@@ -3018,7 +3018,7 @@ FDynamicEmitterDataBase* UParticleSystemComponent::CreateDynamicDataFromReplay(
 		case DET_Beam2:
 			{
 				// Allocate the dynamic data
-				FDynamicBeam2EmitterData* NewEmitterData = ::new FDynamicBeam2EmitterData(EmitterInstance->CurrentLODLevel->RequiredModule);
+				FDynamicBeam2EmitterData* NewEmitterData = new FDynamicBeam2EmitterData(EmitterInstance->CurrentLODLevel->RequiredModule);
 
 				// Fill in the source data
 				const FDynamicBeam2EmitterReplayData* Beam2EmitterReplayData =
@@ -3035,7 +3035,7 @@ FDynamicEmitterDataBase* UParticleSystemComponent::CreateDynamicDataFromReplay(
 		case DET_Ribbon:
 			{
 				// Allocate the dynamic data
-				FDynamicRibbonEmitterData* NewEmitterData = ::new FDynamicRibbonEmitterData(EmitterInstance->CurrentLODLevel->RequiredModule);
+				FDynamicRibbonEmitterData* NewEmitterData = new FDynamicRibbonEmitterData(EmitterInstance->CurrentLODLevel->RequiredModule);
 
 				// Fill in the source data
 				const FDynamicRibbonEmitterReplayData* Trail2EmitterReplayData = static_cast<const FDynamicRibbonEmitterReplayData*>(EmitterReplayData);
@@ -3049,7 +3049,7 @@ FDynamicEmitterDataBase* UParticleSystemComponent::CreateDynamicDataFromReplay(
 		case DET_AnimTrail:
 			{
 				// Allocate the dynamic data
-				FDynamicAnimTrailEmitterData* NewEmitterData = ::new FDynamicAnimTrailEmitterData(EmitterInstance->CurrentLODLevel->RequiredModule);
+				FDynamicAnimTrailEmitterData* NewEmitterData = new FDynamicAnimTrailEmitterData(EmitterInstance->CurrentLODLevel->RequiredModule);
 				// Fill in the source data
 				const FDynamicTrailsEmitterReplayData* AnimTrailEmitterReplayData = static_cast<const FDynamicTrailsEmitterReplayData*>(EmitterReplayData);
 				NewEmitterData->Source = *AnimTrailEmitterReplayData;
@@ -3100,12 +3100,11 @@ FParticleDynamicData* UParticleSystemComponent::CreateDynamicData()
 		}
 	}
 
+
 	FParticleDynamicData* ParticleDynamicData = new FParticleDynamicData();
-	{
-		SCOPE_CYCLE_COUNTER(STAT_ParticleMemTime);
-		INC_DWORD_STAT(STAT_DynamicPSysCompCount);
-		INC_DWORD_STAT_BY(STAT_DynamicPSysCompMem, sizeof(FParticleDynamicData));
-	}
+	INC_DWORD_STAT(STAT_DynamicPSysCompCount);
+	INC_DWORD_STAT_BY(STAT_DynamicPSysCompMem, sizeof(FParticleDynamicData));
+
 	if (Template)
 	{
 		ParticleDynamicData->SystemPositionForMacroUVs = ComponentToWorld.TransformPosition(Template->MacroUVPosition);
@@ -3114,6 +3113,7 @@ FParticleDynamicData* UParticleSystemComponent::CreateDynamicData()
 
 	if( ReplayState == PRS_Replaying )
 	{
+		QUICK_SCOPE_CYCLE_COUNTER(STAT_ParticleSystemComponent_CreateDynamicData_Replay);
 		// Do we have any replay data to play back?
 		UParticleSystemReplay* ReplayData = FindReplayClipForIDNumber( ReplayClipIDNumber );
 		if( ReplayData != NULL )
@@ -3126,7 +3126,8 @@ FParticleDynamicData* UParticleSystemComponent::CreateDynamicData()
 
 
 				// Fill the emitter dynamic buffers with data from our replay
-				ParticleDynamicData->DynamicEmitterDataArray.Empty( CurReplayFrame.Emitters.Num() );
+				ParticleDynamicData->DynamicEmitterDataArray.Reset();
+				ParticleDynamicData->DynamicEmitterDataArray.Reserve(CurReplayFrame.Emitters.Num());
 				for( int32 CurEmitterIndex = 0; CurEmitterIndex < CurReplayFrame.Emitters.Num(); ++CurEmitterIndex )
 				{
 					const FParticleEmitterReplayFrame& CurEmitter = CurReplayFrame.Emitters[ CurEmitterIndex ];
@@ -3156,6 +3157,7 @@ FParticleDynamicData* UParticleSystemComponent::CreateDynamicData()
 		FParticleSystemReplayFrame* NewReplayFrame = NULL;
 		if( ReplayState == PRS_Capturing )
 		{
+			QUICK_SCOPE_CYCLE_COUNTER(STAT_ParticleSystemComponent_CreateDynamicData_Capture);
 			ForceAsyncWorkCompletion(ENSURE_AND_STALL);
 			check(IsInGameThread());
 			// If we don't have any replay data for this component yet, create some now
@@ -3190,7 +3192,9 @@ FParticleDynamicData* UParticleSystemComponent::CreateDynamicData()
 		// Is the particle system allowed to run?
 		if( bForcedInActive == false )
 		{
-			ParticleDynamicData->DynamicEmitterDataArray.Empty(EmitterInstances.Num());
+			QUICK_SCOPE_CYCLE_COUNTER(STAT_ParticleSystemComponent_CreateDynamicData_Gather);
+			ParticleDynamicData->DynamicEmitterDataArray.Reset();
+			ParticleDynamicData->DynamicEmitterDataArray.Reserve(EmitterInstances.Num());
 
 			for (int32 EmitterIndex = 0; EmitterIndex < EmitterInstances.Num(); EmitterIndex++)
 			{
@@ -3199,7 +3203,17 @@ FParticleDynamicData* UParticleSystemComponent::CreateDynamicData()
 				if (EmitterInst)
 				{
 					// Generate the dynamic data for this emitter
-					NewDynamicEmitterData = EmitterInst->GetDynamicData( IsOwnerSelected() );
+					{
+						QUICK_SCOPE_CYCLE_COUNTER(STAT_ParticleSystemComponent_GetDynamicData);
+						bool bIsOwnerSeleted = false;
+#if WITH_EDITOR
+						{
+							QUICK_SCOPE_CYCLE_COUNTER(STAT_ParticleSystemComponent_GetDynamicData_Selected);
+							bIsOwnerSeleted = IsOwnerSelected();
+						}
+#endif
+						NewDynamicEmitterData = EmitterInst->GetDynamicData(bIsOwnerSeleted);
+					}
 					if( NewDynamicEmitterData != NULL )
 					{
 						NewDynamicEmitterData->bValid = true;
@@ -3207,6 +3221,7 @@ FParticleDynamicData* UParticleSystemComponent::CreateDynamicData()
 						// Are we current capturing particle state?
 						if( ReplayState == PRS_Capturing )
 						{
+							QUICK_SCOPE_CYCLE_COUNTER(STAT_ParticleSystemComponent_CreateDynamicData_GatherCapture);
 							// Capture replay data for this particle system
 							// NOTE: This call should always succeed if GetDynamicData succeeded earlier
 							FDynamicEmitterReplayDataBase* NewEmitterReplayData = EmitterInst->GetReplayData();
diff --git a/Engine/Source/Runtime/Engine/Private/Particles/ParticleEmitterInstances.cpp b/Engine/Source/Runtime/Engine/Private/Particles/ParticleEmitterInstances.cpp
--- a/Engine/Source/Runtime/Engine/Private/Particles/ParticleEmitterInstances.cpp
+++ b/Engine/Source/Runtime/Engine/Private/Particles/ParticleEmitterInstances.cpp
@@ -132,6 +132,116 @@ DEFINE_STAT(STAT_DynamicAnimTrailCount_MAX);
 DEFINE_STAT(STAT_DynamicAnimTrailGTMem_MAX);
 DEFINE_STAT(STAT_DynamicUntrackedGTMem_MAX);
 
+/*-----------------------------------------------------------------------------
+	Fast allocators for small block data being sent to the render thread
+-----------------------------------------------------------------------------*/
+
+static TLockFreeFixedSizeAllocator<256> FastParticleSmallBlockAlloc256;
+static TLockFreeFixedSizeAllocator<384> FastParticleSmallBlockAlloc384;
+static TLockFreeFixedSizeAllocator<512> FastParticleSmallBlockAlloc512;
+static TLockFreeFixedSizeAllocator<768> FastParticleSmallBlockAlloc768;
+static TLockFreeFixedSizeAllocator<1024> FastParticleSmallBlockAlloc1024;
+static TLockFreeFixedSizeAllocator<1792> FastParticleSmallBlockAlloc1792;
+static TLockFreeFixedSizeAllocator<2048> FastParticleSmallBlockAlloc2048;
+
+FORCEINLINE static void* FastParticleSmallBlockAlloc(size_t AllocSize)
+{
+	check(AllocSize > 0);
+	if (AllocSize <= 768)
+	{
+		if (AllocSize <= 256)
+		{
+			return FastParticleSmallBlockAlloc256.Allocate();
+		}
+		if (AllocSize <= 384)
+		{
+			return FastParticleSmallBlockAlloc384.Allocate();
+		}
+		if (AllocSize <= 512)
+		{
+			return FastParticleSmallBlockAlloc512.Allocate();
+		}
+		return FastParticleSmallBlockAlloc768.Allocate();
+	}
+	if (AllocSize <= 1024)
+	{
+		return FastParticleSmallBlockAlloc1024.Allocate();
+	}
+	if (AllocSize <= 1792)
+	{
+		return FastParticleSmallBlockAlloc1792.Allocate();
+	}
+	if (AllocSize <= 2048)
+	{
+		return FastParticleSmallBlockAlloc2048.Allocate();
+	}
+	return FMemory::Malloc(AllocSize);
+}
+
+FORCEINLINE static void FastParticleSmallBlockFree(void *RawMemory, size_t AllocSize)
+{
+	check(AllocSize > 0);
+	if (AllocSize <= 768)
+	{
+		if (AllocSize <= 256)
+		{
+			return FastParticleSmallBlockAlloc256.Free(RawMemory);
+		}
+		if (AllocSize <= 384)
+		{
+			return FastParticleSmallBlockAlloc384.Free(RawMemory);
+		}
+		if (AllocSize <= 512)
+		{
+			return FastParticleSmallBlockAlloc512.Free(RawMemory);
+		}
+		return FastParticleSmallBlockAlloc768.Free(RawMemory);
+	}
+
+	if (AllocSize <= 1024)
+	{
+		return FastParticleSmallBlockAlloc1024.Free(RawMemory);
+	}
+	if (AllocSize <= 1792)
+	{
+		return FastParticleSmallBlockAlloc1792.Free(RawMemory);
+	}
+	if (AllocSize <= 2048)
+	{
+		return FastParticleSmallBlockAlloc2048.Free(RawMemory);
+	}
+	FMemory::Free(RawMemory);
+}
+
+
+void FParticleDataContainer::Alloc(int32 InParticleDataNumBytes, int32 InParticleIndicesNumShorts)
+{
+	check(InParticleDataNumBytes > 0 && ParticleIndicesNumShorts >= 0
+		&& InParticleDataNumBytes % sizeof(uint16) == 0); // we assume that the particle storage has reasonable alignment below
+	ParticleDataNumBytes = InParticleDataNumBytes;
+	ParticleIndicesNumShorts = InParticleIndicesNumShorts;
+
+	MemBlockSize = ParticleDataNumBytes + ParticleIndicesNumShorts * sizeof(uint16);
+
+	ParticleData = (uint8*)FastParticleSmallBlockAlloc(MemBlockSize);
+	ParticleIndices = (uint16*)(ParticleData + ParticleDataNumBytes);
+}
+
+void FParticleDataContainer::Free()
+{
+	if (ParticleData)
+	{
+		check(MemBlockSize > 0);
+		FastParticleSmallBlockFree(ParticleData, MemBlockSize);
+	}
+	MemBlockSize = 0;
+	ParticleDataNumBytes = 0;
+	ParticleIndicesNumShorts = 0;
+	ParticleData = nullptr;
+	ParticleIndices = nullptr;
+}
+
+
 /*-----------------------------------------------------------------------------
 	Information compiled from modules to build runtime emitter data.
 -----------------------------------------------------------------------------*/
@@ -2427,22 +2537,14 @@ bool FParticleEmitterInstance::FillReplayData( FDynamicEmitterReplayDataBase& Ou
 	}
 
 	int32 ParticleMemSize = MaxActiveParticles * ParticleStride;
-	int32 IndexMemSize = MaxActiveParticles * sizeof(uint16);
-	{
-		SCOPE_CYCLE_COUNTER(STAT_ParticleMemTime);
-		// This is 'render thread' memory in that we pass it over to the render thread for consumption
-		INC_DWORD_STAT_BY(STAT_RTParticleData, ParticleMemSize + IndexMemSize);
-		//SET_DWORD_STAT(STAT_RTParticleData_Largest, FMath::Max<uint32>(ParticleMemSize + IndexMemSize, GET_DWORD_STAT(STAT_RTParticleData_Largest)));
-	}
 
 	// Allocate particle memory
-	OutData.ParticleData.Empty(ParticleMemSize);
-	OutData.ParticleData.AddUninitialized(ParticleMemSize);
-	FMemory::BigBlockMemcpy(OutData.ParticleData.GetData(), ParticleData, ParticleMemSize);
-	// Allocate particle index memory
-	OutData.ParticleIndices.Empty(MaxActiveParticles);
-	OutData.ParticleIndices.AddUninitialized(MaxActiveParticles);
-	FMemory::BigBlockMemcpy(OutData.ParticleIndices.GetData(), ParticleIndices, IndexMemSize);
+
+	OutData.DataContainer.Alloc(ParticleMemSize, MaxActiveParticles);
+	INC_DWORD_STAT_BY(STAT_RTParticleData, OutData.DataContainer.MemBlockSize);
+
+	FMemory::BigBlockMemcpy(OutData.DataContainer.ParticleData, ParticleData, ParticleMemSize);
+	FMemory::Memcpy(OutData.DataContainer.ParticleIndices, ParticleIndices, OutData.DataContainer.ParticleIndicesNumShorts * sizeof(uint16));
 
 	// All particle emitter types derived from sprite emitters, so we can fill that data in here too!
 	{
@@ -2614,7 +2716,7 @@ FDynamicEmitterDataBase* FParticleSpriteEmitterInstance::GetDynamicData(bool bSe
 	}
 
 	// Allocate the dynamic data
-	FDynamicSpriteEmitterData* NewEmitterData = ::new FDynamicSpriteEmitterData(LODLevel->RequiredModule);
+	FDynamicSpriteEmitterData* NewEmitterData = new FDynamicSpriteEmitterData(LODLevel->RequiredModule);
 	{
 		SCOPE_CYCLE_COUNTER(STAT_ParticleMemTime);
 		INC_DWORD_STAT(STAT_DynamicEmitterCount);
@@ -2682,7 +2784,7 @@ FDynamicEmitterReplayDataBase* FParticleSpriteEmitterInstance::GetReplayData()
 		return NULL;
 	}
 
-	FDynamicEmitterReplayDataBase* NewEmitterReplayData = ::new FDynamicSpriteEmitterReplayData();
+	FDynamicEmitterReplayDataBase* NewEmitterReplayData = new FDynamicSpriteEmitterReplayData();
 	check( NewEmitterReplayData != NULL );
 
 	if( !FillReplayData( *NewEmitterReplayData ) )
@@ -3215,7 +3317,7 @@ FDynamicEmitterDataBase* FParticleMeshEmitterInstance::GetDynamicData(bool bSele
 	}
 
 	// Allocate the dynamic data
-	FDynamicMeshEmitterData* NewEmitterData = ::new FDynamicMeshEmitterData(LODLevel->RequiredModule);
+	FDynamicMeshEmitterData* NewEmitterData = new FDynamicMeshEmitterData(LODLevel->RequiredModule);
 	{
 		SCOPE_CYCLE_COUNTER(STAT_ParticleMemTime);
 		INC_DWORD_STAT(STAT_DynamicEmitterCount);
@@ -3293,7 +3395,7 @@ FDynamicEmitterReplayDataBase* FParticleMeshEmitterInstance::GetReplayData()
 		return NULL;
 	}
 
-	FDynamicEmitterReplayDataBase* NewEmitterReplayData = ::new FDynamicMeshEmitterReplayData();
+	FDynamicEmitterReplayDataBase* NewEmitterReplayData = new FDynamicMeshEmitterReplayData();
 	check( NewEmitterReplayData != NULL );
 
 	if( !FillReplayData( *NewEmitterReplayData ) )
@@ -3532,6 +3634,33 @@ bool FParticleMeshEmitterInstance::FillReplayData( FDynamicEmitterReplayDataBase
 	return true;
 }
 
+
+void* FDynamicEmitterDataBase::operator new(size_t AllocSize)
+{
+	return FastParticleSmallBlockAlloc(AllocSize);
+}
+
+void FDynamicEmitterDataBase::operator delete(void *RawMemory, size_t AllocSize)
+{
+	FastParticleSmallBlockFree(RawMemory, AllocSize);
+}	
+
+static TLockFreeFixedSizeAllocator<sizeof(FParticleDynamicData)> ParticleDynamicDataAllocator;
+
+void* FParticleDynamicData::operator new(size_t AllocSize)
+{
+	check(AllocSize == sizeof(FParticleDynamicData));
+	return ParticleDynamicDataAllocator.Allocate();
+	//return FMemory::Malloc(AllocSize);
+}
+
+void FParticleDynamicData::operator delete(void *RawMemory, size_t AllocSize)
+{
+	check(AllocSize == sizeof(FParticleDynamicData));
+	ParticleDynamicDataAllocator.Free(RawMemory);
+	//FMemory::Free(RawMemory);
+}	
+
 FDynamicEmitterDataBase::FDynamicEmitterDataBase(const UParticleModuleRequired* RequiredModule)
 	: bSelected(false)
 	, ParticleVertexFactory(nullptr)
diff --git a/Engine/Source/Runtime/Engine/Private/Particles/ParticleGpuSimulation.cpp b/Engine/Source/Runtime/Engine/Private/Particles/ParticleGpuSimulation.cpp
--- a/Engine/Source/Runtime/Engine/Private/Particles/ParticleGpuSimulation.cpp
+++ b/Engine/Source/Runtime/Engine/Private/Particles/ParticleGpuSimulation.cpp
@@ -2515,6 +2515,63 @@ public:
 	}
 };
 
+static struct FNewParticleAlloc
+{
+	TLockFreeFixedSizeAllocator<sizeof(TArray<FNewParticle>), FThreadSafeCounter> FreeTArrayFNewParticleArrays;
+	TLockFreePointerList<TArray<FNewParticle>>	FreeTArrayFNewParticle;
+	~FNewParticleAlloc()
+	{
+		while (true)
+		{
+			TArray<FNewParticle>* Recycle = FreeTArrayFNewParticle.Pop();
+			if (!Recycle)
+			{
+				break;
+			}
+			Recycle->Empty();
+			FreeTArrayFNewParticleArrays.Free(Recycle);
+		}
+	}
+
+} GNewParticleAlloc;
+
+static const int MaxNumParticlesToRecycle  = 512; // these can be quite large
+
+// recycle memory blocks for the NewParticle array
+static void FreeNewParticleArray(TArray<FNewParticle>& NewParticles)
+{
+	NewParticles.Reset();
+	const int MaxNumToRecycledArrays  = 100; 
+	int32 CurrentSize = NewParticles.GetSlack();
+	if (CurrentSize > 0 && CurrentSize <= MaxNumParticlesToRecycle && GNewParticleAlloc.FreeTArrayFNewParticleArrays.GetNumUsed().GetValue() < MaxNumToRecycledArrays)
+	{
+		TArray<FNewParticle>* Recycle = new (GNewParticleAlloc.FreeTArrayFNewParticleArrays.Allocate()) TArray<FNewParticle>;
+		Exchange(*Recycle, NewParticles);
+		check(Recycle->Num() == 0 && Recycle->GetSlack());
+		check(NewParticles.Num() == 0 && NewParticles.GetSlack() == 0);
+		GNewParticleAlloc.FreeTArrayFNewParticle.Push(Recycle);
+	}
+}
+
+static void GetNewParticleArray(TArray<FNewParticle>& NewParticles, int32 NumParticlesNeeded = -1)
+{
+	if (NumParticlesNeeded <= MaxNumParticlesToRecycle)
+	{
+		TArray<FNewParticle>* Recycle = GNewParticleAlloc.FreeTArrayFNewParticle.Pop();
+		if (Recycle)
+		{
+			Exchange(*Recycle, NewParticles);
+			Recycle->~TArray<FNewParticle>(); // this probably doesn't do anything, but type safety and all
+			GNewParticleAlloc.FreeTArrayFNewParticleArrays.Free(Recycle);
+			check(NewParticles.Num() == 0 && NewParticles.GetSlack());
+		}
+	}
+	if (NumParticlesNeeded > 0)
+	{
+		// this might realloc, but we need to get the small blocks out of the recycle list
+		NewParticles.Reserve(NumParticlesNeeded);
+	}
+}
 /**
  * Dynamic emitter data for Cascade.
  */
@@ -2568,6 +2625,11 @@ struct FGPUSpriteDynamicEmitterData : FDynamicEmitterDataBase
 		, bLocalVectorFieldTileY(false)
 		, bLocalVectorFieldTileZ(false)
 	{
+		GetNewParticleArray(NewParticles);
+	}
+	~FGPUSpriteDynamicEmitterData()
+	{
+		FreeNewParticleArray(NewParticles);
 	}
 
 	bool RendersWithTranslucentMaterial() const
@@ -2919,6 +2981,7 @@ public:
 	 */
 	virtual FDynamicEmitterDataBase* GetDynamicData(bool bSelected) override
 	{
+		QUICK_SCOPE_CYCLE_COUNTER(STAT_FDynamicEmitterDataBase_GetDynamicData);
 		check(Component);
 		check(SpriteTemplate);
 		check(FXSystem);
@@ -3017,8 +3080,7 @@ public:
 			Exchange(DynamicData->TilesToClear, TilesToClear);
 			Exchange(DynamicData->NewParticles, NewParticles);
 		}
-
-		NewParticles.Reset();
+		FreeNewParticleArray(NewParticles);
 		PendingDeltaSeconds = 0.0f;
 		PositionOffsetThisTick = FVector::ZeroVector;
 
@@ -3112,6 +3174,21 @@ public:
 		InitLocalVectorField();
 	}
 
+	FORCENOINLINE void ReserveNewParticles(int32 Num)
+	{
+		if (Num)
+		{
+			if (!(NewParticles.Num() + NewParticles.GetSlack()))
+			{
+				GetNewParticleArray(NewParticles, Num);
+			}
+			else
+			{
+				NewParticles.Reserve(Num);
+			}
+		}
+	}
+
 	/**
 	 * Simulates the emitter forward by the specified amount of time.
 	 */
@@ -3174,14 +3251,17 @@ public:
 			}
 
 
-
-			int32 FirstBurstParticleIndex = NewParticles.Num();
-			BurstInfo.Count = AllocateTilesForParticles(NewParticles, BurstInfo.Count, ActiveTileCount);
-
 			// Determine spawn count based on rate.
 			FSpawnInfo SpawnInfo = GetNumParticlesToSpawn(DeltaSeconds);
 			SpawnInfo.Count += ForceSpawnedParticles.Num();
 
+
+			int32 FirstBurstParticleIndex = NewParticles.Num();
+
+			ReserveNewParticles(FirstBurstParticleIndex + BurstInfo.Count + SpawnInfo.Count);
+
+			BurstInfo.Count = AllocateTilesForParticles(NewParticles, BurstInfo.Count, ActiveTileCount);
+
 			int32 FirstSpawnParticleIndex = NewParticles.Num();
 			SpawnInfo.Count = AllocateTilesForParticles(NewParticles, SpawnInfo.Count, ActiveTileCount);
 			SpawnFraction += LeftoverBurst;
@@ -3198,8 +3278,8 @@ public:
 				BuildNewParticles(NewParticles.GetData() + FirstSpawnParticleIndex, SpawnInfo, ForceSpawnedParticles);
 			}
 
-			ForceBurstSpawnedParticles.Empty();
-			ForceSpawnedParticles.Empty();
+			FreeNewParticleArray(ForceSpawnedParticles);
+			FreeNewParticleArray(ForceBurstSpawnedParticles);
 
 			int32 NewParticleCount = BurstInfo.Count + SpawnInfo.Count;
 			INC_DWORD_STAT_BY(STAT_GPUSpritesSpawned, NewParticleCount);
@@ -3500,6 +3580,10 @@ private:
 	 */
 	int32 AllocateTilesForParticles(TArray<FNewParticle>& InNewParticles, int32 NumNewParticles, int32& ActiveTileCount)
 	{
+		if (!NumNewParticles)
+		{
+			return 0;
+		}
 		// Need to allocate space in tiles for all new particles.
 		FParticleSimulationResources* SimulationResources = FXSystem->GetParticleSimulationResources();
 		uint32 TileIndex = (AllocatedTiles.IsValidIndex(TileToAllocateFrom)) ? AllocatedTiles[TileToAllocateFrom] : INDEX_NONE;
@@ -3822,6 +3906,10 @@ private:
 		FVector SpawnLocation = bUseLocalSpace ? FVector::ZeroVector : InLocation;
 
 		float Increment = DeltaTime / InSpawnCount;
+		if (InSpawnCount && !(ForceSpawnedParticles.Num() + ForceSpawnedParticles.GetSlack()))
+		{
+			GetNewParticleArray(ForceSpawnedParticles, InSpawnCount);
+		}
 		for (int32 i = 0; i < InSpawnCount; i++)
 		{
 
@@ -3831,7 +3919,10 @@ private:
 			Particle.RelativeTime = Increment*i;
 			ForceSpawnedParticles.Add(Particle);
 		}
-
+		if (InBurstCount && !(ForceBurstSpawnedParticles.Num() + ForceBurstSpawnedParticles.GetSlack()))
+		{
+			GetNewParticleArray(ForceBurstSpawnedParticles, InBurstCount);
+		}
 		for (int32 i = 0; i < InBurstCount; i++)
 		{
 			FNewParticle Particle;
@@ -4233,7 +4324,7 @@ void FFXSystem::SimulateGPUParticles(
 
 			// Add to the list of new particles.
 			NewParticles.Append(Simulation->NewParticles);
-			Simulation->NewParticles.Reset();
+			FreeNewParticleArray(Simulation->NewParticles);
 
 			// Reset pending simulation time. This prevents an emitter from simulating twice if we don't get an update from the game thread, e.g. the component didn't tick last frame.
 			Simulation->PerFrameSimulationParameters.DeltaSeconds = 0.0f;
diff --git a/Engine/Source/Runtime/Engine/Private/Particles/ParticleSystemRender.cpp b/Engine/Source/Runtime/Engine/Private/Particles/ParticleSystemRender.cpp
--- a/Engine/Source/Runtime/Engine/Private/Particles/ParticleSystemRender.cpp
+++ b/Engine/Source/Runtime/Engine/Private/Particles/ParticleSystemRender.cpp
@@ -341,7 +341,7 @@ FORCEINLINE FVector GetCameraOffsetFromPayload(
 }
 
 void FDynamicSpriteEmitterDataBase::SortSpriteParticles(int32 SortMode, bool bLocalSpace, 
-	int32 ParticleCount, const TArray<uint8>& ParticleData, int32 ParticleStride, const TArray<uint16>& ParticleIndices,
+	int32 ParticleCount, const uint8* ParticleData, int32 ParticleStride, const uint16* ParticleIndices,
 	const FSceneView* View, const FMatrix& LocalToWorld, FParticleOrder* ParticleOrder) const
 {
 	SCOPE_CYCLE_COUNTER(STAT_SortingTime);
@@ -359,7 +359,7 @@ void FDynamicSpriteEmitterDataBase::SortSpriteParticles(int32 SortMode, bool bLo
 	{
 		for (int32 ParticleIndex = 0; ParticleIndex < ParticleCount; ParticleIndex++)
 		{
-			DECLARE_PARTICLE(Particle, ParticleData.GetData() + ParticleStride * ParticleIndices[ParticleIndex]);
+			DECLARE_PARTICLE(Particle, ParticleData + ParticleStride * ParticleIndices[ParticleIndex]);
 			float InZ;
 			if (bLocalSpace)
 			{
@@ -379,7 +379,7 @@ void FDynamicSpriteEmitterDataBase::SortSpriteParticles(int32 SortMode, bool bLo
 	{
 		for (int32 ParticleIndex = 0; ParticleIndex < ParticleCount; ParticleIndex++)
 		{
-			DECLARE_PARTICLE(Particle, ParticleData.GetData() + ParticleStride * ParticleIndices[ParticleIndex]);
+			DECLARE_PARTICLE(Particle, ParticleData + ParticleStride * ParticleIndices[ParticleIndex]);
 			float InZ;
 			FVector Position;
 			if (bLocalSpace)
@@ -400,7 +400,7 @@ void FDynamicSpriteEmitterDataBase::SortSpriteParticles(int32 SortMode, bool bLo
 	{
 		for (int32 ParticleIndex = 0; ParticleIndex < ParticleCount; ParticleIndex++)
 		{
-			DECLARE_PARTICLE(Particle, ParticleData.GetData() + ParticleStride * ParticleIndices[ParticleIndex]);
+			DECLARE_PARTICLE(Particle, ParticleData + ParticleStride * ParticleIndices[ParticleIndex]);
 			ParticleOrder[ParticleIndex].ParticleIndex = ParticleIndex;
 			ParticleOrder[ParticleIndex].C = Particle.Flags & STATE_CounterMask;
 		}
@@ -410,7 +410,7 @@ void FDynamicSpriteEmitterDataBase::SortSpriteParticles(int32 SortMode, bool bLo
 	{
 		for (int32 ParticleIndex = 0; ParticleIndex < ParticleCount; ParticleIndex++)
 		{
-			DECLARE_PARTICLE(Particle, ParticleData.GetData() + ParticleStride * ParticleIndices[ParticleIndex]);
+			DECLARE_PARTICLE(Particle, ParticleData + ParticleStride * ParticleIndices[ParticleIndex]);
 			ParticleOrder[ParticleIndex].ParticleIndex = ParticleIndex;
 			ParticleOrder[ParticleIndex].C = (~Particle.Flags) & STATE_CounterMask;
 		}
@@ -435,7 +435,7 @@ void FDynamicSpriteEmitterDataBase::RenderDebug(const FParticleSystemSceneProxy*
 
 	for (int32 i = 0; i < SpriteSource.ActiveParticleCount; i++)
 	{
-		DECLARE_PARTICLE(Particle, SpriteSource.ParticleData.GetData() + SpriteSource.ParticleStride * SpriteSource.ParticleIndices[i]);
+		DECLARE_PARTICLE(Particle, SpriteSource.DataContainer.ParticleData + SpriteSource.ParticleStride * SpriteSource.DataContainer.ParticleIndices[i]);
 
 		FVector DrawLocation = LocalToWorld.TransformPosition(Particle.Location);
 		if (bCrosses)
@@ -591,8 +591,8 @@ bool FDynamicSpriteEmitterData::GetVertexAndIndexData(void* VertexData, void* Dy
 	FVector ParticleOldPosition;
 	float SubImageIndex = 0.0f;
 
-	const uint8* ParticleData = Source.ParticleData.GetData();
-	const uint16* ParticleIndices = Source.ParticleIndices.GetData();
+	const uint8* ParticleData = Source.DataContainer.ParticleData;
+	const uint16* ParticleIndices = Source.DataContainer.ParticleIndices;
 	const FParticleOrder* OrderedIndices = ParticleOrder;
 
 	for (int32 i = 0; i < ParticleCount; i++)
@@ -696,8 +696,8 @@ bool FDynamicSpriteEmitterData::GetVertexAndIndexDataNonInstanced(void* VertexDa
 	FVector ParticleOldPosition;
 	float SubImageIndex = 0.0f;
 
-	const uint8* ParticleData = Source.ParticleData.GetData();
-	const uint16* ParticleIndices = Source.ParticleIndices.GetData();
+	const uint8* ParticleData = Source.DataContainer.ParticleData;
+	const uint16* ParticleIndices = Source.DataContainer.ParticleIndices;
 	const FParticleOrder* OrderedIndices = ParticleOrder;
 
 	for (int32 i = 0; i < ParticleCount; i++)
@@ -810,8 +810,8 @@ void GatherParticleLightData(const FDynamicSpriteEmitterReplayDataBase& Source,
 			OutParticleLights.PerViewData.Reserve(OutParticleLights.PerViewData.Num() + ParticleCount);
 		}
 		
-		const uint8* ParticleData = Source.ParticleData.GetData();
-		const uint16* ParticleIndices = Source.ParticleIndices.GetData();
+		const uint8* ParticleData = Source.DataContainer.ParticleData;
+		const uint16* ParticleIndices = Source.DataContainer.ParticleIndices;
 
 		for (int32 i = 0; i < ParticleCount; i++)
 		{
@@ -1004,7 +1004,7 @@ void FDynamicSpriteEmitterData::GetDynamicMeshElementsEmitter(const FParticleSys
 								{
 									ParticleOrder = (FParticleOrder*)FMemStack::Get().Alloc(sizeof(FParticleOrder)* ParticleCount, ALIGNOF(FParticleOrder));
 									SortSpriteParticles(SourceData->SortMode, SourceData->bUseLocalSpace, SourceData->ActiveParticleCount, 
-										SourceData->ParticleData, SourceData->ParticleStride, SourceData->ParticleIndices,
+										SourceData->DataContainer.ParticleData, SourceData->ParticleStride, SourceData->DataContainer.ParticleIndices,
 										View, Proxy->GetLocalToWorld(), ParticleOrder);
 								}
 								// Fill vertex buffers.
@@ -1024,20 +1024,20 @@ void FDynamicSpriteEmitterData::GetDynamicMeshElementsEmitter(const FParticleSys
 						FParticleOrder* ParticleOrder = NULL;
 						if (bSort)
 						{
-							ParticleOrder = GParticleOrderPool.GetParticleOrderData(ParticleCount);
-							SortSpriteParticles(SourceData->SortMode, SourceData->bUseLocalSpace, SourceData->ActiveParticleCount, 
-								SourceData->ParticleData, SourceData->ParticleStride, SourceData->ParticleIndices,
-								View, Proxy->GetLocalToWorld(), ParticleOrder);
-						}
-						// Fill vertex buffers.
-						if(bInstanced)
-						{
-							GetVertexAndIndexData(Allocation.Buffer, DynamicParameterAllocation.Buffer, NULL, ParticleOrder, View->ViewMatrices.ViewOrigin, Proxy->GetLocalToWorld());
-						}
-						else
-						{
-							GetVertexAndIndexDataNonInstanced(Allocation.Buffer, DynamicParameterAllocation.Buffer, NULL, ParticleOrder, View->ViewMatrices.ViewOrigin, Proxy->GetLocalToWorld());
-						}
+						ParticleOrder = GParticleOrderPool.GetParticleOrderData(ParticleCount);
+						SortSpriteParticles(SourceData->SortMode, SourceData->bUseLocalSpace, SourceData->ActiveParticleCount, 
+								SourceData->DataContainer.ParticleData, SourceData->ParticleStride, SourceData->DataContainer.ParticleIndices,
+							View, Proxy->GetLocalToWorld(), ParticleOrder);
+					}
+				// Fill vertex buffers.
+				if(bInstanced)
+				{
+					GetVertexAndIndexData(Allocation.Buffer, DynamicParameterAllocation.Buffer, NULL, ParticleOrder, View->ViewMatrices.ViewOrigin, Proxy->GetLocalToWorld());
+				}
+				else
+				{
+					GetVertexAndIndexDataNonInstanced(Allocation.Buffer, DynamicParameterAllocation.Buffer, NULL, ParticleOrder, View->ViewMatrices.ViewOrigin, Proxy->GetLocalToWorld());
+				}
 					}
 				}
 
@@ -1436,8 +1436,8 @@ void FDynamicMeshEmitterData::GetDynamicMeshElementsEmitter(const FParticleSyste
 						Collector.AddTask(
 							[this, View, Proxy, Allocation, DynamicParameterAllocation]()
 							{
-								GetInstanceData(Allocation.Buffer, DynamicParameterAllocation.Buffer, Proxy, View);
-							}
+					GetInstanceData(Allocation.Buffer, DynamicParameterAllocation.Buffer, Proxy, View);
+				}
 						);
 					}
 					else
@@ -1473,8 +1473,8 @@ void FDynamicMeshEmitterData::GetDynamicMeshElementsEmitter(const FParticleSyste
 				}
 				else
 				{
-					GetInstanceData((void*)InstanceVerticesCPU->InstanceDataAllocationsCPU.GetData(), (void*)InstanceVerticesCPU->DynamicParameterDataAllocationsCPU.GetData(), Proxy, View);
-				}
+				GetInstanceData((void*)InstanceVerticesCPU->InstanceDataAllocationsCPU.GetData(), (void*)InstanceVerticesCPU->DynamicParameterDataAllocationsCPU.GetData(), Proxy, View);
+			}
 			}
 
 			Proxy->UpdateWorldSpacePrimitiveUniformBuffer();
@@ -1959,8 +1959,8 @@ void FDynamicMeshEmitterData::GetInstanceData(void* InstanceData, void* DynamicP
 
 	for (int32 i = ParticleCount - 1; i >= 0; i--)
 	{
-		const int32	CurrentIndex	= Source.ParticleIndices[i];
-		const uint8* ParticleBase	= Source.ParticleData.GetData() + CurrentIndex * Source.ParticleStride;
+		const int32	CurrentIndex	= Source.DataContainer.ParticleIndices[i];
+		const uint8* ParticleBase	= Source.DataContainer.ParticleData + CurrentIndex * Source.ParticleStride;
 		const FBaseParticle& Particle		= *((const FBaseParticle*) ParticleBase);
 		FMeshParticleInstanceVertex* CurrentInstanceVertex = (FMeshParticleInstanceVertex*)TempVert;
 		
@@ -2367,7 +2367,7 @@ void FDynamicBeam2EmitterData::RenderDirectLine(const FParticleSystemSceneProxy*
 {
 	for (int32 Beam = 0; Beam < Source.ActiveParticleCount; Beam++)
 	{
-		DECLARE_PARTICLE_PTR(Particle, Source.ParticleData.GetData() + Source.ParticleStride * Beam);
+		DECLARE_PARTICLE_PTR(Particle, Source.DataContainer.ParticleData + Source.ParticleStride * Beam);
 
 		FBeam2TypeDataPayload*	BeamPayloadData		= NULL;
 		FVector*				InterpolatedPoints	= NULL;
@@ -2431,7 +2431,7 @@ void FDynamicBeam2EmitterData::RenderLines(const FParticleSystemSceneProxy* Prox
 		// Tessellate the beam along the noise points
 		for (i = 0; i < Source.ActiveParticleCount; i++)
 		{
-			DECLARE_PARTICLE_PTR(Particle, Source.ParticleData.GetData() + Source.ParticleStride * i);
+			DECLARE_PARTICLE_PTR(Particle, Source.DataContainer.ParticleData + Source.ParticleStride * i);
 
 			// Retrieve the beam data from the particle.
 			FBeam2TypeDataPayload*	BeamPayloadData		= NULL;
@@ -2754,7 +2754,7 @@ void FDynamicBeam2EmitterData::RenderLines(const FParticleSystemSceneProxy* Prox
 		{
 			for (int32 i = 0; i < Source.ActiveParticleCount; i++)
 			{
-				DECLARE_PARTICLE_PTR(Particle, Source.ParticleData.GetData() + Source.ParticleStride * i);
+				DECLARE_PARTICLE_PTR(Particle, Source.DataContainer.ParticleData + Source.ParticleStride * i);
 				FBeam2TypeDataPayload* BeamPayloadData = (FBeam2TypeDataPayload*)((uint8*)Particle + Source.BeamDataOffset);
 				if (BeamPayloadData->TriangleCount == 0)
 				{
@@ -2773,7 +2773,7 @@ void FDynamicBeam2EmitterData::RenderLines(const FParticleSystemSceneProxy* Prox
 		{
 			for (int32 i = 0; i < Source.ActiveParticleCount; i++)
 			{
-				DECLARE_PARTICLE_PTR(Particle, Source.ParticleData.GetData() + Source.ParticleStride * i);
+				DECLARE_PARTICLE_PTR(Particle, Source.DataContainer.ParticleData + Source.ParticleStride * i);
 
 				FBeam2TypeDataPayload*	BeamPayloadData		= NULL;
 				FVector*				InterpolatedPoints	= NULL;
@@ -2854,7 +2854,7 @@ static int32 CreateDynamicBeam2EmitterIndices(TIndexType* OutIndex, const FDynam
 
 	for (int32 Beam = 0; Beam < Source.ActiveParticleCount; Beam++)
 	{
-		DECLARE_PARTICLE_PTR(Particle, Source.ParticleData.GetData() + Source.ParticleStride * Beam);
+		DECLARE_PARTICLE_PTR(Particle, Source.DataContainer.ParticleData + Source.ParticleStride * Beam);
 		FBeam2TypeDataPayload*	BeamPayloadData = (FBeam2TypeDataPayload*)((uint8*)Particle + Source.BeamDataOffset);
 		if (BeamPayloadData->TriangleCount == 0)
 		{
@@ -2963,7 +2963,7 @@ int32 FDynamicBeam2EmitterData::FillVertexData_NoNoise(FAsyncBufferFillData& Me)
 	{
 		for (int32 i = 0; i < Source.ActiveParticleCount; i++)
 		{
-			DECLARE_PARTICLE_PTR(Particle, Source.ParticleData.GetData() + Source.ParticleStride * i);
+			DECLARE_PARTICLE_PTR(Particle, Source.DataContainer.ParticleData + Source.ParticleStride * i);
 
 			FBeam2TypeDataPayload*	BeamPayloadData		= NULL;
 			FVector*				InterpolatedPoints	= NULL;
@@ -3142,7 +3142,7 @@ int32 FDynamicBeam2EmitterData::FillVertexData_NoNoise(FAsyncBufferFillData& Me)
 
 		for (int32 i = 0; i < Source.ActiveParticleCount; i++)
 		{
-			DECLARE_PARTICLE_PTR(Particle, Source.ParticleData.GetData() + Source.ParticleStride * i);
+			DECLARE_PARTICLE_PTR(Particle, Source.DataContainer.ParticleData + Source.ParticleStride * i);
 
 			FBeam2TypeDataPayload*	BeamPayloadData		= NULL;
 			FVector*				InterpolatedPoints	= NULL;
@@ -3417,7 +3417,7 @@ int32 FDynamicBeam2EmitterData::FillData_Noise(FAsyncBufferFillData& Me) const
 	// Tessellate the beam along the noise points
 	for (i = 0; i < Source.ActiveParticleCount; i++)
 	{
-		DECLARE_PARTICLE_PTR(Particle, Source.ParticleData.GetData() + Source.ParticleStride * i);
+		DECLARE_PARTICLE_PTR(Particle, Source.DataContainer.ParticleData + Source.ParticleStride * i);
 
 		// Retrieve the beam data from the particle.
 		FBeam2TypeDataPayload*	BeamPayloadData		= NULL;
@@ -4378,7 +4378,7 @@ int32 FDynamicBeam2EmitterData::FillData_InterpolatedNoise(FAsyncBufferFillData&
 	// Tessellate the beam along the noise points
 	for (i = 0; i < Source.ActiveParticleCount; i++)
 	{
-		DECLARE_PARTICLE_PTR(Particle, Source.ParticleData.GetData() + Source.ParticleStride * i);
+		DECLARE_PARTICLE_PTR(Particle, Source.DataContainer.ParticleData + Source.ParticleStride * i);
 
 		// Retrieve the beam data from the particle.
 		FBeam2TypeDataPayload*	BeamPayloadData		= NULL;
@@ -5045,8 +5045,8 @@ void FDynamicTrailsEmitterData::GetDynamicMeshElementsEmitter(const FParticleSys
 				int32 CheckTriangleCount = 0;
 				for (int32 ParticleIdx = 0; ParticleIdx < SourcePointer->ActiveParticleCount; ParticleIdx++)
 				{
-					int32 CurrentIndex = SourcePointer->ParticleIndices[ParticleIdx];
-					DECLARE_PARTICLE_PTR(CheckParticle, SourcePointer->ParticleData.GetData() + SourcePointer->ParticleStride * CurrentIndex);
+					int32 CurrentIndex = SourcePointer->DataContainer.ParticleIndices[ParticleIdx];
+					DECLARE_PARTICLE_PTR(CheckParticle, SourcePointer->DataContainer.ParticleData + SourcePointer->ParticleStride * CurrentIndex);
 					FTrailsBaseTypeDataPayload* TrailPayload = (FTrailsBaseTypeDataPayload*)((uint8*)CheckParticle + SourcePointer->TrailDataOffset);
 					if (TRAIL_EMITTER_IS_HEAD(TrailPayload->Flags) == false)
 					{
@@ -5120,8 +5120,8 @@ int32 FDynamicTrailsEmitterData::FillIndexData(struct FAsyncBufferFillData& Data
 	int32 CurrentTrail = 0;
 	for (int32 ParticleIdx = 0; ParticleIdx < SourcePointer->ActiveParticleCount; ParticleIdx++)
 	{
-		int32 CurrentIndex = SourcePointer->ParticleIndices[ParticleIdx];
-		DECLARE_PARTICLE_PTR(Particle, SourcePointer->ParticleData.GetData() + SourcePointer->ParticleStride * CurrentIndex);
+		int32 CurrentIndex = SourcePointer->DataContainer.ParticleIndices[ParticleIdx];
+		DECLARE_PARTICLE_PTR(Particle, SourcePointer->DataContainer.ParticleData + SourcePointer->ParticleStride * CurrentIndex);
 
 		FTrailsBaseTypeDataPayload* TrailPayload = (FTrailsBaseTypeDataPayload*)((uint8*)Particle + SourcePointer->TrailDataOffset);
 		if (TRAIL_EMITTER_IS_HEAD(TrailPayload->Flags) == false)
@@ -5201,7 +5201,7 @@ void FDynamicRibbonEmitterData::RenderDebug(const FParticleSystemSceneProxy* Pro
 		FColor PrevDrawColor;
 		FVector DrawTangentEnd;
 
-		const uint8* Address = Source.ParticleData.GetData();
+		const uint8* Address = Source.DataContainer.ParticleData;
 		const FRibbonTypeDataPayload* StartTrailPayload;
 		const FRibbonTypeDataPayload* EndTrailPayload = NULL;
 		const FBaseParticle* DebugParticle;
@@ -5211,7 +5211,7 @@ void FDynamicRibbonEmitterData::RenderDebug(const FParticleSystemSceneProxy* Pro
 
 		for (int32 ParticleIdx = 0; ParticleIdx < Source.ActiveParticleCount; ParticleIdx++)
 		{
-			DECLARE_PARTICLE_PTR(Particle, Address + Source.ParticleStride * Source.ParticleIndices[ParticleIdx]);
+			DECLARE_PARTICLE_PTR(Particle, Address + Source.ParticleStride * Source.DataContainer.ParticleIndices[ParticleIdx]);
 			StartTrailPayload = (FRibbonTypeDataPayload*)((uint8*)Particle + Source.TrailDataOffset);
 			if (TRAIL_EMITTER_IS_HEAD(StartTrailPayload->Flags) == 0)
 			{
@@ -5396,10 +5396,10 @@ int32 FDynamicRibbonEmitterData::FillVertexData(struct FAsyncBufferFillData& Dat
 	float CurrDistance = 0.0f;
 
 	FBaseParticle* PackingParticle;
-	const uint8* ParticleData = Source.ParticleData.GetData();
+	const uint8* ParticleData = Source.DataContainer.ParticleData;
 	for (int32 ParticleIdx = 0; ParticleIdx < Source.ActiveParticleCount; ParticleIdx++)
 	{
-		DECLARE_PARTICLE_PTR(Particle, ParticleData + Source.ParticleStride * Source.ParticleIndices[ParticleIdx]);
+		DECLARE_PARTICLE_PTR(Particle, ParticleData + Source.ParticleStride * Source.DataContainer.ParticleIndices[ParticleIdx]);
 		FRibbonTypeDataPayload* TrailPayload = (FRibbonTypeDataPayload*)((uint8*)Particle + Source.TrailDataOffset);
 		if (TRAIL_EMITTER_IS_HEAD(TrailPayload->Flags) == 0)
 		{
@@ -5721,7 +5721,7 @@ struct FAnimTrailParticleRenderData
 	
 	FAnimTrailParticleRenderData( const FDynamicTrailsEmitterReplayData& InSource, const FBaseParticle* InParticle, const FAnimTrailTypeDataPayload* InPayload )
 		:	Source(InSource),
-			ParticleDataAddress(Source.ParticleData.GetData()),
+			ParticleDataAddress(Source.DataContainer.ParticleData),
 			PrevPrevParticle(NULL),
 			PrevPrevPayload(NULL),
 			PrevParticle(NULL),
@@ -5949,12 +5949,12 @@ void FDynamicAnimTrailEmitterData::RenderDebug(const FParticleSystemSceneProxy*
 		float TiledU;
 		FLinearColor DummyColor;
 
-		const uint8* Address = Source.ParticleData.GetData();
+		const uint8* Address = Source.DataContainer.ParticleData;
 		const FAnimTrailTypeDataPayload* StartTrailPayload;
 
 		for (int32 ParticleIdx = 0; ParticleIdx < Source.ActiveParticleCount; ParticleIdx++)
 		{
-			DECLARE_PARTICLE_PTR(Particle, Address + Source.ParticleStride * Source.ParticleIndices[ParticleIdx]);
+			DECLARE_PARTICLE_PTR(Particle, Address + Source.ParticleStride * Source.DataContainer.ParticleIndices[ParticleIdx]);
 			StartTrailPayload = (FAnimTrailTypeDataPayload*)((uint8*)Particle + Source.TrailDataOffset);
 			if (TRAIL_EMITTER_IS_HEAD(StartTrailPayload->Flags) == 0)
 			{
@@ -6096,10 +6096,10 @@ int32 FDynamicAnimTrailEmitterData::FillVertexData(struct FAsyncBufferFillData&
 	// The distance tracking for tiling the 2nd UV set
 	float CurrDistance = 0.0f;
 
-	const uint8* ParticleData = Source.ParticleData.GetData();
+	const uint8* ParticleData = Source.DataContainer.ParticleData;
 	for (int32 ParticleIdx = 0; ParticleIdx < Source.ActiveParticleCount; ParticleIdx++)
 	{
-		DECLARE_PARTICLE_PTR(Particle, ParticleData + Source.ParticleStride * Source.ParticleIndices[ParticleIdx]);
+		DECLARE_PARTICLE_PTR(Particle, ParticleData + Source.ParticleStride * Source.DataContainer.ParticleIndices[ParticleIdx]);
 		FAnimTrailTypeDataPayload* TrailPayload = (FAnimTrailTypeDataPayload*)((uint8*)Particle + Source.TrailDataOffset);
 		if (TRAIL_EMITTER_IS_HEAD(TrailPayload->Flags) == 0)
 		{
diff --git a/Engine/Source/Runtime/Engine/Private/Particles/ParticleTrail2EmitterInstance.cpp b/Engine/Source/Runtime/Engine/Private/Particles/ParticleTrail2EmitterInstance.cpp
--- a/Engine/Source/Runtime/Engine/Private/Particles/ParticleTrail2EmitterInstance.cpp
+++ b/Engine/Source/Runtime/Engine/Private/Particles/ParticleTrail2EmitterInstance.cpp
@@ -2572,7 +2572,7 @@ FDynamicEmitterDataBase* FParticleRibbonEmitterInstance::GetDynamicData(bool bSe
 	}
 
 	// Allocate the dynamic data
-	FDynamicRibbonEmitterData* NewEmitterData = ::new FDynamicRibbonEmitterData(LODLevel->RequiredModule);
+	FDynamicRibbonEmitterData* NewEmitterData = new FDynamicRibbonEmitterData(LODLevel->RequiredModule);
 	{
 		SCOPE_CYCLE_COUNTER(STAT_ParticleMemTime);
 		INC_DWORD_STAT(STAT_DynamicEmitterCount);
@@ -2664,7 +2664,7 @@ FDynamicEmitterReplayDataBase* FParticleRibbonEmitterInstance::GetReplayData()
 		return NULL;
 	}
 
-	FDynamicEmitterReplayDataBase* NewEmitterReplayData = ::new FDynamicRibbonEmitterReplayData();
+	FDynamicEmitterReplayDataBase* NewEmitterReplayData = new FDynamicRibbonEmitterReplayData();
 	check(NewEmitterReplayData != NULL);
 
 	if (!FillReplayData(*NewEmitterReplayData))
@@ -3920,7 +3920,7 @@ FDynamicEmitterDataBase* FParticleAnimTrailEmitterInstance::GetDynamicData(bool
 	}
 
 	// Allocate the dynamic data
-	FDynamicAnimTrailEmitterData* NewEmitterData = ::new FDynamicAnimTrailEmitterData(LODLevel->RequiredModule);
+	FDynamicAnimTrailEmitterData* NewEmitterData = new FDynamicAnimTrailEmitterData(LODLevel->RequiredModule);
 	{
 		SCOPE_CYCLE_COUNTER(STAT_ParticleMemTime);
 		INC_DWORD_STAT(STAT_DynamicEmitterCount);
@@ -4026,7 +4026,7 @@ FDynamicEmitterReplayDataBase* FParticleAnimTrailEmitterInstance::GetReplayData(
 		return NULL;
 	}
 
-	FDynamicTrailsEmitterReplayData* NewEmitterReplayData = ::new FDynamicTrailsEmitterReplayData();
+	FDynamicTrailsEmitterReplayData* NewEmitterReplayData = new FDynamicTrailsEmitterReplayData();
 	check(NewEmitterReplayData != NULL);
 
 	if (!FillReplayData(*NewEmitterReplayData))
diff --git a/Engine/Source/Runtime/Engine/Private/PhysicsEngine/PhysAnim.cpp b/Engine/Source/Runtime/Engine/Private/PhysicsEngine/PhysAnim.cpp
--- a/Engine/Source/Runtime/Engine/Private/PhysicsEngine/PhysAnim.cpp
+++ b/Engine/Source/Runtime/Engine/Private/PhysicsEngine/PhysAnim.cpp
@@ -93,9 +93,9 @@ public:
 	}
 };
 
-
+typedef TArray<FAssetWorldBoneTM, TMemStackAllocator<ALIGNOF(FAssetWorldBoneTM)>> TAssetWorldBoneTMArray;
 // Use current pose to calculate world-space position of this bone without physics now.
-void UpdateWorldBoneTM(TArray<FAssetWorldBoneTM> & WorldBoneTMs, const TArray<FTransform>& InLocalAtoms, int32 BoneIndex, USkeletalMeshComponent* SkelComp, const FTransform &LocalToWorldTM, const FVector& Scale3D)
+void UpdateWorldBoneTM(TAssetWorldBoneTMArray& WorldBoneTMs, const TArray<FTransform>& InLocalAtoms, int32 BoneIndex, USkeletalMeshComponent* SkelComp, const FTransform &LocalToWorldTM, const FVector& Scale3D)
 {
 	// If its already up to date - do nothing
 	if(	WorldBoneTMs[BoneIndex].bUpToDate )
@@ -152,9 +152,9 @@ void USkeletalMeshComponent::PerformBlendPhysicsBones(const TArray<FBoneIndexTyp
 		return;
 	}
 
+	FMemMark Mark(FMemStack::Get());
 	// Make sure scratch space is big enough.
-	TArray<FAssetWorldBoneTM> WorldBoneTMs;
-	WorldBoneTMs.Reset();
+	TAssetWorldBoneTMArray WorldBoneTMs;
 	WorldBoneTMs.AddZeroed(GetNumSpaceBases());
 	
 	FTransform LocalToWorldTM = ComponentToWorld;
@@ -167,8 +167,7 @@ void USkeletalMeshComponent::PerformBlendPhysicsBones(const TArray<FBoneIndexTyp
 		FBodyInstance* BI;
 		FTransform TM;
 	};
-
-	TArray<FBodyTMPair> PendingBodyTMs;
+	TArray<FBodyTMPair, TMemStackAllocator<ALIGNOF(FBodyTMPair)>> PendingBodyTMs;
 
 #if WITH_PHYSX
 	// Lock the scenes we need (flags set in InitArticulated)
diff --git a/Engine/Source/Runtime/Engine/Private/SkeletalMeshComponentPhysics.cpp b/Engine/Source/Runtime/Engine/Private/SkeletalMeshComponentPhysics.cpp
--- a/Engine/Source/Runtime/Engine/Private/SkeletalMeshComponentPhysics.cpp
+++ b/Engine/Source/Runtime/Engine/Private/SkeletalMeshComponentPhysics.cpp
@@ -1937,7 +1937,27 @@ bool USkeletalMeshComponent::ComponentOverlapMultiImpl(TArray<struct FOverlapRes
 // convert a bone name from APEX stype to FBX style
 static FName GetConvertedBoneName(NxClothingAsset* ApexClothingAsset, int32 BoneIndex)
 {
-	return *FString(ApexClothingAsset->getBoneName(BoneIndex)).Replace(TEXT(" "), TEXT("-"));
+	QUICK_SCOPE_CYCLE_COUNTER(STAT_GetConvertedBoneName);
+	static FString Work;
+	check(IsInGameThread() && !Work.Len());
+	// FIXME, this still needs work, string conversion and hashing
+
+	const char* Bone = ApexClothingAsset->getBoneName(BoneIndex);
+	if (Bone)
+	{
+		int32 SrcLen = FCStringAnsi::Strlen(Bone) + 1;
+		int32 DestLen = FPlatformString::ConvertedLength<TCHAR>(Bone, SrcLen);
+		TArray<TCHAR>& Array = Work.GetCharArray();
+		Array.AddUninitialized(DestLen);
+		FPlatformString::Convert(Array.GetData(), DestLen, Bone, SrcLen);
+
+		Work.ReplaceInline(TEXT(" "), TEXT("-"));
+		FName Result(*Work);
+		Work.Reset();
+		//check(Result == FName(*FString(ApexClothingAsset->getBoneName(BoneIndex)).Replace(TEXT(" "), TEXT("-"))));
+		return Result;
+	}
+	return NAME_None;
 }
 
 void USkeletalMeshComponent::AddClothingBounds(FBoxSphereBounds& InOutBounds) const
@@ -3540,6 +3560,9 @@ void USkeletalMeshComponent::UpdateClothState(float DeltaTime)
 	// convert teleport mode to apex clothing teleport enum
 	physx::apex::ClothingTeleportMode::Enum CurTeleportMode = (physx::apex::ClothingTeleportMode::Enum)ClothTeleportMode;
 
+	static TArray<physx::PxMat44> BoneMatrices;
+	check(IsInGameThread() && !BoneMatrices.Num());
+
 	for(int32 ActorIdx=0; ActorIdx<NumActors; ActorIdx++)
 	{
 		// skip if ClothingActor is NULL or invalid
@@ -3550,13 +3573,11 @@ void USkeletalMeshComponent::UpdateClothState(float DeltaTime)
 
 		ApplyWindForCloth(ClothingActors[ActorIdx]);
 
-		TArray<physx::PxMat44> BoneMatrices;
-
 		NxClothingAsset* ClothingAsset = ClothingActors[ActorIdx].ParentClothingAsset;
 
 		uint32 NumUsedBones = ClothingAsset->getNumUsedBones();
 
-		BoneMatrices.Empty(NumUsedBones);
+		check(!BoneMatrices.Num());
 		BoneMatrices.AddUninitialized(NumUsedBones);
 
 		for(uint32 Index=0; Index < NumUsedBones; Index++)
@@ -3603,6 +3624,8 @@ void USkeletalMeshComponent::UpdateClothState(float DeltaTime)
 			sizeof(physx::PxMat44), 
 			NumUsedBones,
 			CurTeleportMode);
+
+		BoneMatrices.Reset();
 	}
 
 	// reset to Continuous
@@ -3755,13 +3778,13 @@ void USkeletalMeshComponent::GetUpdateClothSimulationData(TArray<FClothSimulData
 
 	if(NumClothingActors == 0 || bDisableClothSimulation)
 	{
-		OutClothSimData.Empty();
+		OutClothSimData.Reset();
 		return;
 	}
 
 	if(OutClothSimData.Num() != NumClothingActors)
 	{
-		OutClothSimData.Empty(NumClothingActors);
+		OutClothSimData.Reset();
 		OutClothSimData.AddZeroed(NumClothingActors);
 	}
 
@@ -3771,8 +3794,8 @@ void USkeletalMeshComponent::GetUpdateClothSimulationData(TArray<FClothSimulData
 	{
 		if(!IsValidClothingActor(ActorIndex))
 		{
-			OutClothSimData[ActorIndex].ClothSimulPositions.Empty();
-			OutClothSimData[ActorIndex].ClothSimulNormals.Empty();
+			OutClothSimData[ActorIndex].ClothSimulPositions.Reset();
+			OutClothSimData[ActorIndex].ClothSimulNormals.Reset();
 			continue;
 		}
 
@@ -3791,9 +3814,9 @@ void USkeletalMeshComponent::GetUpdateClothSimulationData(TArray<FClothSimulData
 
 				if(ClothData.ClothSimulPositions.Num() != NumSimulVertices)
 				{
-					ClothData.ClothSimulPositions.Empty(NumSimulVertices);
+					ClothData.ClothSimulPositions.Reset();
 					ClothData.ClothSimulPositions.AddUninitialized(NumSimulVertices);
-					ClothData.ClothSimulNormals.Empty(NumSimulVertices);
+					ClothData.ClothSimulNormals.Reset();
 					ClothData.ClothSimulNormals.AddUninitialized(NumSimulVertices);
 				}
 
@@ -3825,7 +3848,7 @@ void USkeletalMeshComponent::GetUpdateClothSimulationData(TArray<FClothSimulData
 	//no simulated vertices 
 	if(!bSimulated)
 	{
-		OutClothSimData.Empty();
+		OutClothSimData.Reset();
 	}
 #endif// #if WITH_APEX_CLOTHING
 }
diff --git a/Engine/Source/Runtime/Engine/Private/SkeletalRender.cpp b/Engine/Source/Runtime/Engine/Private/SkeletalRender.cpp
--- a/Engine/Source/Runtime/Engine/Private/SkeletalRender.cpp
+++ b/Engine/Source/Runtime/Engine/Private/SkeletalRender.cpp
@@ -202,7 +202,7 @@ void UpdateRefToLocalMatrices( TArray<FMatrix>& ReferenceToLocal, const USkinned
 	check( ThisMesh->RefBasesInvMatrix.Num() != 0 );
 	if(ReferenceToLocal.Num() != ThisMesh->RefBasesInvMatrix.Num())
 	{
-		ReferenceToLocal.Empty(ThisMesh->RefBasesInvMatrix.Num());
+		ReferenceToLocal.Reset();
 		ReferenceToLocal.AddUninitialized(ThisMesh->RefBasesInvMatrix.Num());
 	}
 
@@ -302,7 +302,7 @@ void UpdateCustomLeftRightVectors( TArray<FTwoVectors>& OutVectors, const USkinn
 
 	if(OutVectors.Num() != LODInfo.TriangleSortSettings.Num())
 	{
-		OutVectors.Empty(LODInfo.TriangleSortSettings.Num());
+		OutVectors.Reset();
 		OutVectors.AddUninitialized(LODInfo.TriangleSortSettings.Num());
 	}
 
diff --git a/Engine/Source/Runtime/Engine/Private/SkeletalRender.h b/Engine/Source/Runtime/Engine/Private/SkeletalRender.h
--- a/Engine/Source/Runtime/Engine/Private/SkeletalRender.h
+++ b/Engine/Source/Runtime/Engine/Private/SkeletalRender.h
@@ -11,14 +11,6 @@ extern const float MinVertexAnimBlendWeight;
 // largest blend weight for vertex anims
 extern const float MaxVertexAnimBlendWeight;
 
-/** Dynamic data updates needed by the rendering thread are sent with this */
-class FDynamicSkelMeshObjectData
-{
-public:
-	FDynamicSkelMeshObjectData(){}
-	virtual ~FDynamicSkelMeshObjectData(){}
-};
-
 /** 
 * Stores the data for updating instanced weights
 * Created by the game thread and sent to the rendering thread as an update 
diff --git a/Engine/Source/Runtime/Engine/Private/SkeletalRenderCPUSkin.cpp b/Engine/Source/Runtime/Engine/Private/SkeletalRenderCPUSkin.cpp
--- a/Engine/Source/Runtime/Engine/Private/SkeletalRenderCPUSkin.cpp
+++ b/Engine/Source/Runtime/Engine/Private/SkeletalRenderCPUSkin.cpp
@@ -168,13 +168,13 @@ void FSkeletalMeshObjectCPUSkin::Update(int32 LODIndex,USkinnedMeshComponent* In
 {
 	// create the new dynamic data for use by the rendering thread
 	// this data is only deleted when another update is sent
-	FDynamicSkelMeshObjectData* NewDynamicData = new FDynamicSkelMeshObjectDataCPUSkin(InMeshComponent,SkeletalMeshResource,LODIndex,ActiveVertexAnims);
+	FDynamicSkelMeshObjectDataCPUSkin* NewDynamicData = new FDynamicSkelMeshObjectDataCPUSkin(InMeshComponent,SkeletalMeshResource,LODIndex,ActiveVertexAnims);
 
 	// queue a call to update this data
 	ENQUEUE_UNIQUE_RENDER_COMMAND_TWOPARAMETER(
 		SkelMeshObjectUpdateDataCommand,
-		FSkeletalMeshObject*, MeshObject, this,
-		FDynamicSkelMeshObjectData*, NewDynamicData, NewDynamicData,
+		FSkeletalMeshObjectCPUSkin*, MeshObject, this,
+		FDynamicSkelMeshObjectDataCPUSkin*, NewDynamicData, NewDynamicData,
 	{
 		FScopeCycleCounter Context(MeshObject->GetStatId());
 		MeshObject->UpdateDynamicData_RenderThread(RHICmdList, NewDynamicData);
@@ -189,12 +189,12 @@ void FSkeletalMeshObjectCPUSkin::Update(int32 LODIndex,USkinnedMeshComponent* In
 	}
 }
 
-void FSkeletalMeshObjectCPUSkin::UpdateDynamicData_RenderThread(FRHICommandListImmediate& RHICmdList, FDynamicSkelMeshObjectData* InDynamicData)
+void FSkeletalMeshObjectCPUSkin::UpdateDynamicData_RenderThread(FRHICommandListImmediate& RHICmdList, FDynamicSkelMeshObjectDataCPUSkin* InDynamicData)
 {
 	// we should be done with the old data at this point
 	delete DynamicData;
 	// update with new data
-	DynamicData = (FDynamicSkelMeshObjectDataCPUSkin*)InDynamicData;	
+	DynamicData = InDynamicData;	
 	check(DynamicData);
 
 	// update vertices using the new data
diff --git a/Engine/Source/Runtime/Engine/Private/SkeletalRenderCPUSkin.h b/Engine/Source/Runtime/Engine/Private/SkeletalRenderCPUSkin.h
--- a/Engine/Source/Runtime/Engine/Private/SkeletalRenderCPUSkin.h
+++ b/Engine/Source/Runtime/Engine/Private/SkeletalRenderCPUSkin.h
@@ -80,7 +80,7 @@ private:
 * Stores the updated matrices needed to skin the verts.
 * Created by the game thread and sent to the rendering thread as an update 
 */
-class FDynamicSkelMeshObjectDataCPUSkin : public FDynamicSkelMeshObjectData
+class FDynamicSkelMeshObjectDataCPUSkin
 {
 public:
 
@@ -98,6 +98,10 @@ public:
 		const TArray<FActiveVertexAnim>& InActiveVertexAnims
 		);
 
+	virtual ~FDynamicSkelMeshObjectDataCPUSkin()
+	{
+	}
+
 	/** ref pose to local space transforms */
 	TArray<FMatrix> ReferenceToLocal;
 
@@ -142,7 +146,7 @@ public:
 	virtual void InitResources() override;
 	virtual void ReleaseResources() override;
 	virtual void Update(int32 LODIndex,USkinnedMeshComponent* InMeshComponent,const TArray<FActiveVertexAnim>& ActiveVertexAnims) override;
-	virtual void UpdateDynamicData_RenderThread(FRHICommandListImmediate& RHICmdList, FDynamicSkelMeshObjectData* InDynamicData) override;
+	void UpdateDynamicData_RenderThread(FRHICommandListImmediate& RHICmdList, FDynamicSkelMeshObjectDataCPUSkin* InDynamicData);
 	virtual void EnableBlendWeightRendering(bool bEnabled, const TArray<int32>& InBonesOfInterest) override;
 	virtual void CacheVertices(int32 LODIndex, bool bForce) const override;
 	virtual bool IsCPUSkinned() const override { return true; }
diff --git a/Engine/Source/Runtime/Engine/Private/SkeletalRenderGPUSkin.cpp b/Engine/Source/Runtime/Engine/Private/SkeletalRenderGPUSkin.cpp
--- a/Engine/Source/Runtime/Engine/Private/SkeletalRenderGPUSkin.cpp
+++ b/Engine/Source/Runtime/Engine/Private/SkeletalRenderGPUSkin.cpp
@@ -115,7 +115,11 @@ FSkeletalMeshObjectGPUSkin::FSkeletalMeshObjectGPUSkin(USkinnedMeshComponent* In
 FSkeletalMeshObjectGPUSkin::~FSkeletalMeshObjectGPUSkin()
 {
 	check(!RHIThreadFenceForDynamicData.GetReference());
-	delete DynamicData;
+	if (DynamicData)
+	{
+		FDynamicSkelMeshObjectDataGPUSkin::FreeDynamicSkelMeshObjectDataGPUSkin(DynamicData);
+	}
+	DynamicData = nullptr;
 }
 
 
@@ -189,14 +193,15 @@ void FSkeletalMeshObjectGPUSkin::Update(int32 LODIndex,USkinnedMeshComponent* In
 
 	// create the new dynamic data for use by the rendering thread
 	// this data is only deleted when another update is sent
-	FDynamicSkelMeshObjectDataGPUSkin* NewDynamicData = new FDynamicSkelMeshObjectDataGPUSkin(InMeshComponent,SkeletalMeshResource,LODIndex,ActiveVertexAnims);
+	FDynamicSkelMeshObjectDataGPUSkin* NewDynamicData = FDynamicSkelMeshObjectDataGPUSkin::AllocDynamicSkelMeshObjectDataGPUSkin();		
+	NewDynamicData->InitDynamicSkelMeshObjectDataGPUSkin(InMeshComponent,SkeletalMeshResource,LODIndex,ActiveVertexAnims);
 
 
 	// queue a call to update this data
 	ENQUEUE_UNIQUE_RENDER_COMMAND_TWOPARAMETER(
 		SkelMeshObjectUpdateDataCommand,
-		FSkeletalMeshObject*, MeshObject, this,
-		FDynamicSkelMeshObjectData*, NewDynamicData, NewDynamicData,
+		FSkeletalMeshObjectGPUSkin*, MeshObject, this,
+		FDynamicSkelMeshObjectDataGPUSkin*, NewDynamicData, NewDynamicData,
 	{
 		FScopeCycleCounter Context(MeshObject->GetStatId());
 		MeshObject->UpdateDynamicData_RenderThread(RHICmdList, NewDynamicData);
@@ -216,22 +221,25 @@ static TAutoConsoleVariable<int32> CVarDeferSkeletalDynamicDataUpdateUntilGDME(
 	0,
 	TEXT("If > 0, then do skeletal mesh dynamic data updates will be deferred until GDME. Experimental option."));
 
-void FSkeletalMeshObjectGPUSkin::UpdateDynamicData_RenderThread(FRHICommandListImmediate& RHICmdList, FDynamicSkelMeshObjectData* InDynamicData)
+void FSkeletalMeshObjectGPUSkin::UpdateDynamicData_RenderThread(FRHICommandListImmediate& RHICmdList, FDynamicSkelMeshObjectDataGPUSkin* InDynamicData)
 {
 	SCOPE_CYCLE_COUNTER(STAT_GPUSkinUpdateRTTime);
+	check(InDynamicData);
 	bool bMorphNeedsUpdate=false;
 	// figure out if the morphing vertex buffer needs to be updated. compare old vs new active morphs
 	bMorphNeedsUpdate = 
 		(bMorphNeedsUpdateDeferred && bNeedsUpdateDeferred) || // the need for an update sticks
-		(DynamicData ? (DynamicData->LODIndex != ((FDynamicSkelMeshObjectDataGPUSkin*)InDynamicData)->LODIndex ||
-		!DynamicData->ActiveVertexAnimsEqual(((FDynamicSkelMeshObjectDataGPUSkin*)InDynamicData)->ActiveVertexAnims))
+		(DynamicData ? (DynamicData->LODIndex != InDynamicData->LODIndex ||
+		!DynamicData->ActiveVertexAnimsEqual(InDynamicData->ActiveVertexAnims))
 		: true);
 
 	WaitForRHIThreadFenceForDynamicData();
-	delete DynamicData;
+	if (DynamicData)
+	{
+		FDynamicSkelMeshObjectDataGPUSkin::FreeDynamicSkelMeshObjectDataGPUSkin(DynamicData);
+	}
 	// update with new data
-	DynamicData = (FDynamicSkelMeshObjectDataGPUSkin*)InDynamicData;
-	check(DynamicData);
+	DynamicData = InDynamicData;
 
 	if (CVarDeferSkeletalDynamicDataUpdateUntilGDME.GetValueOnRenderThread())
 	{
@@ -471,13 +479,13 @@ void FSkeletalMeshObjectGPUSkin::FSkeletalMeshObjectLOD::UpdateMorphVertexBuffer
 			}
 		} // ApplyDelta
 
-			// Lock the real buffer.
-			FMorphGPUSkinVertex* ActualBuffer = (FMorphGPUSkinVertex*)RHILockVertexBuffer(MorphVertexBuffer.VertexBufferRHI, 0, Size, RLM_WriteOnly);
-			FMemory::Memcpy(ActualBuffer, Buffer, Size);
-			FMemory::Free(Buffer);
+		// Lock the real buffer.
+		FMorphGPUSkinVertex* ActualBuffer = (FMorphGPUSkinVertex*)RHILockVertexBuffer(MorphVertexBuffer.VertexBufferRHI, 0, Size, RLM_WriteOnly);
+		FMemory::Memcpy(ActualBuffer, Buffer, Size);
+		FMemory::Free(Buffer);
 
-			// Unlock the buffer.
-			RHIUnlockVertexBuffer(MorphVertexBuffer.VertexBufferRHI);
+		// Unlock the buffer.
+		RHIUnlockVertexBuffer(MorphVertexBuffer.VertexBufferRHI);
 		// set update flag
 		MorphVertexBuffer.bHasBeenUpdated = true;
 	}
@@ -1003,16 +1011,53 @@ const FTwoVectors& FSkeletalMeshObjectGPUSkin::GetCustomLeftRightVectors(int32 S
 FDynamicSkelMeshObjectDataGPUSkin
 -----------------------------------------------------------------------------*/
 
-FDynamicSkelMeshObjectDataGPUSkin::FDynamicSkelMeshObjectDataGPUSkin(
+static TLockFreePointerList<FDynamicSkelMeshObjectDataGPUSkin>	FreeDynamicSkelMeshObjectDataGPUSkins;
+
+void FDynamicSkelMeshObjectDataGPUSkin::Clear()
+{
+	ReferenceToLocal.Reset();
+	CustomLeftRightVectors.Reset();
+#if !(UE_BUILD_SHIPPING || UE_BUILD_TEST) 
+	MeshSpaceBases.Reset();
+#endif
+	LODIndex = 0;
+	ActiveVertexAnims.Reset();
+	NumWeightedActiveVertexAnims = 0;
+	ClothSimulUpdateData.Reset();
+	ClothBlendWeight = 0.0f;
+}
+
+
+FDynamicSkelMeshObjectDataGPUSkin* FDynamicSkelMeshObjectDataGPUSkin::AllocDynamicSkelMeshObjectDataGPUSkin()
+{
+	FDynamicSkelMeshObjectDataGPUSkin* Result = FreeDynamicSkelMeshObjectDataGPUSkins.Pop();
+	if (!Result)
+	{
+		Result = new FDynamicSkelMeshObjectDataGPUSkin;
+	}
+	return Result;
+}
+
+void FDynamicSkelMeshObjectDataGPUSkin::FreeDynamicSkelMeshObjectDataGPUSkin(FDynamicSkelMeshObjectDataGPUSkin* Who)
+{
+	check(Who);
+	Who->Clear();
+	FreeDynamicSkelMeshObjectDataGPUSkins.Push(Who);
+}
+
+void FDynamicSkelMeshObjectDataGPUSkin::InitDynamicSkelMeshObjectDataGPUSkin(
 	USkinnedMeshComponent* InMeshComponent,
 	FSkeletalMeshResource* InSkeletalMeshResource,
 	int32 InLODIndex,
 	const TArray<FActiveVertexAnim>& InActiveVertexAnims
 	)
-	:	LODIndex(InLODIndex)
-	,	ActiveVertexAnims(InActiveVertexAnims)
-	,	NumWeightedActiveVertexAnims(0)
 {
+	LODIndex = InLODIndex;
+	check(!ActiveVertexAnims.Num() && !ReferenceToLocal.Num() && !CustomLeftRightVectors.Num() && !ClothSimulUpdateData.Num());
+
+	// append instead of equals to avoid alloc
+	ActiveVertexAnims.Append(InActiveVertexAnims);
+	NumWeightedActiveVertexAnims = 0;
 	// update ReferenceToLocal
 	UpdateRefToLocalMatrices( ReferenceToLocal, InMeshComponent, InSkeletalMeshResource, LODIndex );
 
@@ -1020,7 +1065,9 @@ FDynamicSkelMeshObjectDataGPUSkin::FDynamicSkelMeshObjectDataGPUSkin(
 
 
 #if !(UE_BUILD_SHIPPING || UE_BUILD_TEST)
-	MeshSpaceBases = InMeshComponent->GetSpaceBases();
+	check(!MeshSpaceBases.Num());
+	// append instead of equals to avoid alloc
+	MeshSpaceBases.Append(InMeshComponent->GetSpaceBases());
 #endif
 
 	// find number of morphs that are currently weighted and will affect the mesh
@@ -1038,7 +1085,7 @@ FDynamicSkelMeshObjectDataGPUSkin::FDynamicSkelMeshObjectDataGPUSkin(
 		}
 		else
 		{
-			ActiveVertexAnims.RemoveAt(AnimIdx);
+			ActiveVertexAnims.RemoveAt(AnimIdx, 1, false);
 		}
 	}
 
@@ -1169,17 +1216,16 @@ void FPreviousPerBoneMotionBlur::StartAppend(FRHICommandListImmediate& RHICmdLis
 
 		AdvanceBufferLocation();
 
-	FBoneDataVertexBuffer& WriteTexture = PerChunkBoneMatricesTexture[GetWriteBufferIndex()];
-
-	if(WriteTexture.IsValid())
-	{
-		LockedData = WriteTexture.LockData();
-		check(LockedTexelPosition.GetValue() == 0); //otherwise it wasn't unlocked or it was unlocked before it was done being filled by async stuff
-		LockedTexelPosition.Set(0);
-		LockedTexelCount = WriteTexture.GetSizeX();
-	}
-}
-
+	    FBoneDataVertexBuffer& WriteTexture = PerChunkBoneMatricesTexture[GetWriteBufferIndex()];
+    
+	    if(WriteTexture.IsValid())
+	    {
+		    LockedData = WriteTexture.LockData();
+		    check(LockedTexelPosition.GetValue() == 0); //otherwise it wasn't unlocked or it was unlocked before it was done being filled by async stuff
+		    LockedTexelPosition.Set(0);
+		    LockedTexelCount = WriteTexture.GetSizeX();
+	    }
+    }
 	if (CVarMotionBlurDebug.GetValueOnRenderThread())
 	{
 		UE_LOG(LogEngine, Log, TEXT("r.MotionBlurDebug: BufferSize=%d Read=%d Write=%d IsLocked=%d"),
@@ -1190,6 +1236,7 @@ void FPreviousPerBoneMotionBlur::StartAppend(FRHICommandListImmediate& RHICmdLis
 uint32 FPreviousPerBoneMotionBlur::AppendData(FBoneSkinning*& OutChunkMatrices, uint32 BoneCount)
 {
 	check(LockedData);
+	checkSlow(DataStart);
 	checkSlow(BoneCount);
 
 	uint32 TexelCount = BoneCount * sizeof(FBoneSkinning) / sizeof(float) / 4;
@@ -1242,25 +1289,25 @@ void FPreviousPerBoneMotionBlur::EndAppend(FRHICommandListImmediate& RHICmdList)
 		PerChunkBoneMatricesTexture[GetWriteBufferIndex()].UnlockData(LockedTexelPosition.GetValue() * 4 * sizeof(float));
 		LockedTexelPosition.Set(0);
 
-		{
-			static int LogSpawmPrevent = 0;
-
-			if(bWarningBufferSizeExceeded)
-			{
-				bWarningBufferSizeExceeded = false;
-
-				if((LogSpawmPrevent % 16) == 0)
-				{
-					UE_LOG(LogSkeletalGPUSkinMesh, Warning, TEXT("Exceeded buffer for per bone motionblur for skinned mesh velocity rendering. Artifacts can occur. Change Content, increase buffer size or change to use FGlobalDynamicVertexBuffer."));
-				}
-				++LogSpawmPrevent;
-			}
-			else
-			{
-				LogSpawmPrevent = 0;
-			}
-		}
-	}
+	    {
+		    static int LogSpawmPrevent = 0;
+
+		    if(bWarningBufferSizeExceeded)
+		    {
+			    bWarningBufferSizeExceeded = false;
+
+			    if((LogSpawmPrevent % 16) == 0)
+			    {
+				    UE_LOG(LogSkeletalGPUSkinMesh, Warning, TEXT("Exceeded buffer for per bone motionblur for skinned mesh velocity rendering. Artifacts can occur. Change Content, increase buffer size or change to use FGlobalDynamicVertexBuffer."));
+			    }
+			    ++LogSpawmPrevent;
+		    }
+		    else
+		    {
+			    LogSpawmPrevent = 0;
+		    }
+	    }
+    }
 }
 
 void FPreviousPerBoneMotionBlur::SetVelocityPassCallback(const TFunction<bool(FRHICommandList& RHICmdList)>& InIsVelocityFunc)
diff --git a/Engine/Source/Runtime/Engine/Private/SkeletalRenderGPUSkin.h b/Engine/Source/Runtime/Engine/Private/SkeletalRenderGPUSkin.h
--- a/Engine/Source/Runtime/Engine/Private/SkeletalRenderGPUSkin.h
+++ b/Engine/Source/Runtime/Engine/Private/SkeletalRenderGPUSkin.h
@@ -19,10 +19,29 @@
 * Stores the updated matrices needed to skin the verts.
 * Created by the game thread and sent to the rendering thread as an update 
 */
-class FDynamicSkelMeshObjectDataGPUSkin : public FDynamicSkelMeshObjectData
+class FDynamicSkelMeshObjectDataGPUSkin
 {
+	/**
+	* Constructor, these are recycled, so you never use a constructor
+	*/
+	FDynamicSkelMeshObjectDataGPUSkin()
+	{
+		Clear();
+	}
+
+	virtual ~FDynamicSkelMeshObjectDataGPUSkin()
+	{
+		// we leak these
+		check(0);
+	}
+
+	void Clear();
+
 public:
 
+	static FDynamicSkelMeshObjectDataGPUSkin* AllocDynamicSkelMeshObjectDataGPUSkin();
+	static void FreeDynamicSkelMeshObjectDataGPUSkin(FDynamicSkelMeshObjectDataGPUSkin* Who);
+
 	/**
 	* Constructor
 	* Updates the ReferenceToLocal matrices using the new dynamic data.
@@ -30,7 +49,7 @@ public:
 	* @param	InLODIndex - each lod has its own bone map 
 	* @param	InActiveVertexAnims - vertex anims active for the mesh
 	*/
-	FDynamicSkelMeshObjectDataGPUSkin(
+	void InitDynamicSkelMeshObjectDataGPUSkin(
 		USkinnedMeshComponent* InMeshComponent,
 		FSkeletalMeshResource* InSkeletalMeshResource,
 		int32 InLODIndex,
@@ -279,7 +298,7 @@ public:
 	virtual void InitResources() override;
 	virtual void ReleaseResources() override;
 	virtual void Update(int32 LODIndex,USkinnedMeshComponent* InMeshComponent,const TArray<FActiveVertexAnim>& ActiveVertexAnims) override;
-	virtual void UpdateDynamicData_RenderThread(FRHICommandListImmediate& RHICmdList, FDynamicSkelMeshObjectData* InDynamicData) override;
+	void UpdateDynamicData_RenderThread(FRHICommandListImmediate& RHICmdList, FDynamicSkelMeshObjectDataGPUSkin* InDynamicData);
 	virtual void PreGDMECallback() override;
 	virtual const FVertexFactory* GetVertexFactory(int32 LODIndex,int32 ChunkIdx) const override;
 	virtual void CacheVertices(int32 LODIndex, bool bForce) const override {}
diff --git a/Engine/Source/Runtime/Engine/Private/TickTaskManager.cpp b/Engine/Source/Runtime/Engine/Private/TickTaskManager.cpp
--- a/Engine/Source/Runtime/Engine/Private/TickTaskManager.cpp
+++ b/Engine/Source/Runtime/Engine/Private/TickTaskManager.cpp
@@ -22,9 +22,9 @@ DECLARE_DWORD_COUNTER_STAT(TEXT("Ticks Queued"),STAT_TicksQueued,STATGROUP_Game)
 DECLARE_CYCLE_STAT(TEXT("TG_NewlySpawned"), STAT_TG_NewlySpawned, STATGROUP_TickGroups);
 
 static TAutoConsoleVariable<float> CVarStallStartFrame(
-	TEXT("tick.StallStartFrame"),
+	TEXT("CriticalPathStall.TickStartFrame"),
 	0.0f,
-	TEXT("Sleep for the given time in start frame. Time is given in ms. This is a debug option of little value."));
+	TEXT("Sleep for the given time in start frame. Time is given in ms. This is a debug option used for critical path analysis and forcing a change in the critical path."));
 
 static TAutoConsoleVariable<int32> CVarLogTicks(
 	TEXT("tick.LogTicks"),
@@ -43,9 +43,73 @@ static TAutoConsoleVariable<int32> CVarAllowAsyncComponentTicks(
 
 static TAutoConsoleVariable<int32> CVarAllowConcurrentQueue(
 	TEXT("tick.AllowConcurrentTickQueue"),
-	0,
+	1,
 	TEXT("If true, queue ticks concurrently."));
 
+template<typename InElementType, typename InAllocator = FDefaultAllocator>
+class TArrayWithThreadsafeAdd : public TArray<InElementType, InAllocator>
+{
+public:
+	typedef InElementType ElementType;
+	typedef InAllocator   Allocator;
+
+#if PLATFORM_COMPILER_HAS_VARIADIC_TEMPLATES
+	template <typename... ArgsType>
+	int32 EmplaceThreadsafe(ArgsType&&... Args)
+	{
+		const int32 Index = AddUninitializedThreadsafe(1);
+		new(this->GetData() + Index) ElementType(Forward<ArgsType>(Args)...);
+		return Index;
+	}
+
+#else
+	template <typename Arg0Type>
+	int32 EmplaceThreadsafe(Arg0Type&& Arg0)
+	{
+		const int32 Index = AddUninitializedThreadsafe(1);
+		new(this->GetData() + Index) ElementType(Forward<Arg0Type>(Arg0));
+		return Index;
+	}
+#endif
+
+
+	/**
+	 * Adds a given number of uninitialized elements into the array using an atomic increment on the array num
+	 *
+	 * Caution, the array must have sufficient slack or this will assert/crash. You must presize the array.
+	 *
+	 * Caution, AddUninitialized() will create elements without calling
+	 * the constructor and this is not appropriate for element types that
+	 * require a constructor to function properly.
+	 *
+	 * @param Count Number of elements to add.
+	 *
+	 * @returns Number of elements in array before addition.
+	 */
+	int32 AddUninitializedThreadsafe(int32 Count = 1)
+	{
+		checkSlow(Count >= 0);
+		const int32 OldNum = FPlatformAtomics::InterlockedAdd(&this->ArrayNum, Count);
+		check(this->ArrayNum + Count < this->ArrayMax);
+		return OldNum;
+	}
+
+	/**
+	 * Adds a new item to the end of the array, using atomics to update the current size of the array
+	 *
+	 * Caution, the array must have sufficient slack or this will assert/crash. You must presize the array.
+	 *
+	 * @param Item	The item to add
+	 * @return		Index to the new item
+	 */
+	FORCEINLINE int32 AddThreadsafe(const ElementType& Item) 
+	{ 
+		return EmplaceThreadsafe(Item); 
+	}
+
+};
+
+
 struct FTickContext
 {
 	/** Delta time to tick **/
@@ -81,65 +145,68 @@ struct FTickContext
 	}
 };
 
-/** Helper class define the task of ticking a component **/
-class FTickFunctionTask
-{
-	/** Actor to tick **/
-	FTickFunction*			Target;
-	/** tick context, here thread is desired execution thread **/
-	FTickContext			Context;
-	/** If true, log each tick **/
-	bool					bLogTick; 
-	/** If true, log prereqs **/
-	bool					bLogTicksShowPrerequistes; 
+/**
+ * Class that handles the actual tick tasks and starting and completing tick groups
+ */
+	/** Helper class define the task of ticking a component **/
+	class FTickFunctionTask
+	{
+		/** Actor to tick **/
+		FTickFunction*			Target;
+		/** tick context, here thread is desired execution thread **/
+		FTickContext			Context;
+		/** If true, log each tick **/
+		bool					bLogTick; 
+		/** If true, log prereqs **/
+		bool					bLogTicksShowPrerequistes; 
 public:
-	/** Constructor
-		* @param InTarget - Function to tick
-		* @param InContext - context to tick in, here thread is desired execution thread
-	**/
-	FORCEINLINE FTickFunctionTask(FTickFunction* InTarget, const FTickContext* InContext, bool InbLogTick, bool bInLogTicksShowPrerequistes)
-		: Target(InTarget)
-		, Context(*InContext)
-		, bLogTick(InbLogTick)
+		/** Constructor
+		 * @param InTarget - Function to tick
+		 * @param InContext - context to tick in, here thread is desired execution thread
+		**/
+		FORCEINLINE FTickFunctionTask(FTickFunction* InTarget, const FTickContext* InContext, bool InbLogTick, bool bInLogTicksShowPrerequistes)
+			: Target(InTarget)
+			, Context(*InContext)
+			, bLogTick(InbLogTick)
 		, bLogTicksShowPrerequistes(bInLogTicksShowPrerequistes)
-	{
-	}
-	FORCEINLINE TStatId GetStatId() const
-	{
+		{
+		}
+		FORCEINLINE TStatId GetStatId() const
+		{
 		RETURN_QUICK_DECLARE_CYCLE_STAT(FTickFunctionTask, STATGROUP_TaskGraphTasks);
-	}
-	/** return the thread for this task **/
-	FORCEINLINE ENamedThreads::Type GetDesiredThread()
-	{
+		}
+		/** return the thread for this task **/
+		FORCEINLINE ENamedThreads::Type GetDesiredThread()
+		{
 		if (Target->bHighPriority)
 		{
 			return ENamedThreads::HiPri(Context.Thread);
 		}
-		return Context.Thread;
-	}
-	FORCEINLINE static ESubsequentsMode::Type GetSubsequentsMode() 
-	{ 
-		return ESubsequentsMode::TrackSubsequents; 
-	}
-	/** 
-		*	Actually execute the tick.
-		*	@param	CurrentThread; the thread we are running on
-		*	@param	MyCompletionGraphEvent; my completion event. Not always useful since at the end of DoWork, you can assume you are done and hence further tasks do not need you as a prerequisite. 
-		*	However, MyCompletionGraphEvent can be useful for passing to other routines or when it is handy to set up subsequents before you actually do work.
-		**/
-	void DoTask(ENamedThreads::Type CurrentThread, const FGraphEventRef& MyCompletionGraphEvent)
-	{
-		if (bLogTick)
+			return Context.Thread;
+		}
+		FORCEINLINE static ESubsequentsMode::Type GetSubsequentsMode() 
+		{ 
+			return ESubsequentsMode::TrackSubsequents; 
+		}
+		/** 
+		 *	Actually execute the tick.
+		 *	@param	CurrentThread; the thread we are running on
+		 *	@param	MyCompletionGraphEvent; my completion event. Not always useful since at the end of DoWork, you can assume you are done and hence further tasks do not need you as a prerequisite. 
+		 *	However, MyCompletionGraphEvent can be useful for passing to other routines or when it is handy to set up subsequents before you actually do work.
+		 **/
+		void DoTask(ENamedThreads::Type CurrentThread, const FGraphEventRef& MyCompletionGraphEvent)
 		{
-			UE_LOG(LogTick, Log, TEXT("tick %s %6d %2d %s"), Target->bHighPriority ? TEXT("*") : TEXT(" "),GFrameCounter, (int32)CurrentThread, *Target->DiagnosticMessage());
-			if (bLogTicksShowPrerequistes)
+			if (bLogTick)
 			{
-				Target->ShowPrerequistes();
+			UE_LOG(LogTick, Log, TEXT("tick %s %6d %2d %s"), Target->bHighPriority ? TEXT("*") : TEXT(" "),GFrameCounter, (int32)CurrentThread, *Target->DiagnosticMessage());
+				if (bLogTicksShowPrerequistes)
+				{
+					Target->ShowPrerequistes();
+				}
 			}
-		}
-		Target->ExecuteTick(Context.DeltaSeconds, Context.TickType, CurrentThread, MyCompletionGraphEvent);
+			Target->ExecuteTick(Context.DeltaSeconds, Context.TickType, CurrentThread, MyCompletionGraphEvent);
 		Target->TaskPointer = nullptr;  // This is stale and a good time to clear it for safety
-	}
+		}
 };
 
 /**
@@ -148,13 +215,13 @@ public:
 class FTickTaskSequencer
 {
 	/** Completion handles for each phase of ticks */
-	FGraphEventArray	TickCompletionEvents[TG_MAX];
+	TArrayWithThreadsafeAdd<FGraphEventRef, TInlineAllocator<4> > TickCompletionEvents[TG_MAX];
 
 	/** HiPri Held tasks for each tick group. */
-	TArray<TGraphTask<FTickFunctionTask>*> HiPriTickTasks[TG_MAX];
+	TArrayWithThreadsafeAdd<TGraphTask<FTickFunctionTask>*> HiPriTickTasks[TG_MAX];
 
 	/** LowPri Held tasks for each tick group. */
-	TArray<TGraphTask<FTickFunctionTask>*> TickTasks[TG_MAX];
+	TArrayWithThreadsafeAdd<TGraphTask<FTickFunctionTask>*> TickTasks[TG_MAX];
 
 	/** These are waited for at the end of the frame; they are not on the critical path, but they have to be done before we leave the frame. */
 	FGraphEventArray CleanupTasks;
@@ -212,7 +279,6 @@ public:
 		{
 			UseContext.Thread = ENamedThreads::AnyThread;
 		}
-		QUICK_SCOPE_CYCLE_COUNTER(STAT_StartTickTask_ConstructAndDispatchWhenReady);
 
 		TickFunction->TaskPointer = TGraphTask<FTickFunctionTask>::CreateTask(Prerequisites, TickContext.Thread).ConstructAndHold(TickFunction, &UseContext, bLogTicks, bLogTicksShowPrerequistes);
 	}
@@ -231,6 +297,30 @@ public:
 		}
 		new (TickCompletionEvents[TickGroup]) FGraphEventRef(Task->GetCompletionEvent());
 	}
+	/** Add a completion handle to a tick group, parallel version **/
+	FORCEINLINE void AddTickTaskCompletionParallel(ETickingGroup TickGroup, TGraphTask<FTickFunctionTask>* Task, bool bHiPri)
+	{
+		checkSlow(TickGroup >=0 && TickGroup < TG_MAX);
+		if (bHiPri)
+		{
+			HiPriTickTasks[TickGroup].AddThreadsafe(Task);
+		}
+		else
+		{
+			TickTasks[TickGroup].AddThreadsafe(Task);
+		}
+		TickCompletionEvents[TickGroup].AddThreadsafe(Task->GetCompletionEvent());
+	}
+	/** Set up the lists for AddTickTaskCompletionParallel, since we are using AddThreadsafe, we need to presize the arrays **/
+	void SetupAddTickTaskCompletionParallel(int32 NumTicks)
+	{
+		for (int32 TickGroup = 0; TickGroup < TG_MAX; TickGroup++)
+		{
+			HiPriTickTasks[TickGroup].Reserve(NumTicks);
+			TickTasks[TickGroup].Reserve(NumTicks);
+			TickCompletionEvents[TickGroup].Reserve(NumTicks);
+		}
+	}
 	/**
 	 * Start a component tick task and add the completion handle
 	 *
@@ -257,6 +347,8 @@ public:
 	{
 		checkSlow(TickContext.Thread == ENamedThreads::GameThread);
 		StartTickTask(Prerequisites, TickFunction, TickContext);
+		TGraphTask<FTickFunctionTask>* Task = (TGraphTask<FTickFunctionTask>*)TickFunction->TaskPointer;
+		AddTickTaskCompletionParallel(TickFunction->ActualTickGroup, Task, TickFunction->bHighPriority);
 	}
 
 	/** 
@@ -511,18 +603,10 @@ public:
 				CumulativeCooldown += TickFunction->RelativeTickCooldown;
 
 				TickFunction->TickState = FTickFunction::ETickState::Enabled;
+				// we store a bit in here if this came from an interval queue 
+				AllTickFunctions.Add((FTickFunction*)(UPTRINT(TickFunction) | 1));
 
-				if (TickFunction->TickInterval > 0.f)
-				{
-					// we store a bit in here if this came from an interval queue 
-					AllTickFunctions.Add((FTickFunction*)(UPTRINT(TickFunction) | 1));
-
-					TickFunctionsToReschedule.Add(FTickScheduleDetails(TickFunction, TickFunction->TickInterval - (Context.DeltaSeconds - CumulativeCooldown))); // Give credit for any overrun
-				}
-				else
-				{
-					AllEnabledTickFunctions.Add(TickFunction);
-				}
+				TickFunctionsToReschedule.Add(FTickScheduleDetails(TickFunction, TickFunction->TickInterval - (Context.DeltaSeconds - CumulativeCooldown))); // Give credit for any overrun
 
 				AllCoolingDownTickFunctions.Head = TickFunction->Next;
 				TickFunction = TickFunction->Next;
@@ -540,13 +624,15 @@ public:
 
 	struct FTickScheduleDetails
 	{
-		FTickScheduleDetails(FTickFunction* InTickFunction, const float InCooldown)
+		FTickScheduleDetails(FTickFunction* InTickFunction, const float InCooldown, bool bInDeferredRemove = false)
 			: TickFunction(InTickFunction)
 			, Cooldown(InCooldown)
+			, bDeferredRemove(bInDeferredRemove)
 		{}
 
 		FTickFunction* TickFunction;
 		float Cooldown;
+		bool bDeferredRemove;
 	};
 
 	void RemoveAndRescheduleForInterval(FTickFunction* TickFunction)
@@ -554,15 +640,24 @@ public:
 		verify(AllEnabledTickFunctions.Remove(TickFunction) == 1);
 		TickFunctionsToReschedule.Add(FTickScheduleDetails(TickFunction, TickFunction->TickInterval));
 	}
-
-	/* Schedules tick functions for cooldown */
+	void RescheduleForIntervalParallel(FTickFunction* TickFunction)
+	{
+		// note we do the remove later!
+		TickFunctionsToReschedule.AddThreadsafe(FTickScheduleDetails(TickFunction, TickFunction->TickInterval, true));
+	}
+	/* Puts a TickFunction in to the cooldown state*/
+	void ReserveTickFunctionCooldowns(int32 NumToReserve)
+	{
+		TickFunctionsToReschedule.Reserve(NumToReserve);
+	}
+	/* Puts a TickFunction in to the cooldown state*/
 	void ScheduleTickFunctionCooldowns()
 	{
 		if (TickFunctionsToReschedule.Num() > 0)
 		{
 			SCOPE_CYCLE_COUNTER(STAT_ScheduleCooldowns);
 
-			TickFunctionsToReschedule.Sort([](const FTickScheduleDetails& A, const FTickScheduleDetails& B)
+				TickFunctionsToReschedule.Sort([](const FTickScheduleDetails& A, const FTickScheduleDetails& B)
 			{
 				return A.Cooldown < B.Cooldown;
 			});
@@ -573,46 +668,46 @@ public:
 			FTickFunction* ComparisonTickFunction = AllCoolingDownTickFunctions.Head;
 			while (ComparisonTickFunction && RescheduleIndex < TickFunctionsToReschedule.Num())
 			{
-				FTickFunction* TickFunction = TickFunctionsToReschedule[RescheduleIndex].TickFunction;
-				if (TickFunction->TickInterval > 0.f)
+				const float CooldownTime = TickFunctionsToReschedule[RescheduleIndex].Cooldown;
+				if ((CumulativeCooldown + ComparisonTickFunction->RelativeTickCooldown) > CooldownTime)
 				{
-					const float CooldownTime = TickFunctionsToReschedule[RescheduleIndex].Cooldown;
-					if ((CumulativeCooldown + ComparisonTickFunction->RelativeTickCooldown) > CooldownTime)
+					FTickFunction* TickFunction = TickFunctionsToReschedule[RescheduleIndex].TickFunction;
+					if (TickFunctionsToReschedule[RescheduleIndex].bDeferredRemove)
 					{
-						TickFunction->TickState = FTickFunction::ETickState::CoolingDown;
-						TickFunction->RelativeTickCooldown = CooldownTime - CumulativeCooldown;
+						verify(AllEnabledTickFunctions.Remove(TickFunction) == 1);
+					}
+					TickFunction->TickState = FTickFunction::ETickState::CoolingDown;
+					TickFunction->RelativeTickCooldown = CooldownTime - CumulativeCooldown;
 
-						if (PrevComparisonTickFunction)
-						{
-							PrevComparisonTickFunction->Next = TickFunction;
-						}
-						else
-						{
-							check(ComparisonTickFunction == AllCoolingDownTickFunctions.Head);
-							AllCoolingDownTickFunctions.Head = TickFunction;
-						}
-						TickFunction->Next = ComparisonTickFunction;
-						PrevComparisonTickFunction = TickFunction;
-						ComparisonTickFunction->RelativeTickCooldown -= TickFunction->RelativeTickCooldown;
-						CumulativeCooldown += TickFunction->RelativeTickCooldown;
-						++RescheduleIndex;
+					if (PrevComparisonTickFunction)
+					{
+						PrevComparisonTickFunction->Next = TickFunction;
 					}
 					else
 					{
-						CumulativeCooldown += ComparisonTickFunction->RelativeTickCooldown;
-						PrevComparisonTickFunction = ComparisonTickFunction;
-						ComparisonTickFunction = ComparisonTickFunction->Next;
+						check(ComparisonTickFunction == AllCoolingDownTickFunctions.Head);
+						AllCoolingDownTickFunctions.Head = TickFunction;
 					}
+					TickFunction->Next = ComparisonTickFunction;
+					PrevComparisonTickFunction = TickFunction;
+					ComparisonTickFunction->RelativeTickCooldown -= TickFunction->RelativeTickCooldown;
+					CumulativeCooldown += TickFunction->RelativeTickCooldown;
+					++RescheduleIndex;
 				}
 				else
 				{
-					AllEnabledTickFunctions.Add(TickFunction);
-					++RescheduleIndex;
+					CumulativeCooldown += ComparisonTickFunction->RelativeTickCooldown;
+					PrevComparisonTickFunction = ComparisonTickFunction;
+					ComparisonTickFunction = ComparisonTickFunction->Next;
 				}
-			}
+	}
 			for ( ; RescheduleIndex < TickFunctionsToReschedule.Num(); ++RescheduleIndex)
 			{
 				FTickFunction* TickFunction = TickFunctionsToReschedule[RescheduleIndex].TickFunction;
+				if (TickFunctionsToReschedule[RescheduleIndex].bDeferredRemove)
+				{
+					verify(AllEnabledTickFunctions.Remove(TickFunction) == 1);
+				}
 				const float CooldownTime = TickFunctionsToReschedule[RescheduleIndex].Cooldown;
 
 				TickFunction->TickState = FTickFunction::ETickState::CoolingDown;
@@ -633,7 +728,7 @@ public:
 				CumulativeCooldown += TickFunction->RelativeTickCooldown;
 			}
 			TickFunctionsToReschedule.Reset();
-		}
+	}
 	}
 
 	/* Queue all tick functions for execution */
@@ -646,10 +741,10 @@ public:
 			TickFunction->QueueTickFunction(TTS, Context);
 
 			if (TickFunction->TickInterval > 0.f)
-			{
+		{
 				It.RemoveCurrent();
 				TickFunctionsToReschedule.Add(FTickScheduleDetails(TickFunction, TickFunction->TickInterval));
-			}
+		}
 		}
 		int32 EnabledCooldownTicks = 0;
 		float CumulativeCooldown = 0.f;
@@ -659,14 +754,7 @@ public:
 			{
 				CumulativeCooldown += TickFunction->RelativeTickCooldown;
 				TickFunction->QueueTickFunction(TTS, Context);
-				if (TickFunction->TickInterval > 0.f)
-				{
-					TickFunctionsToReschedule.Add(FTickScheduleDetails(TickFunction, TickFunction->TickInterval - (Context.DeltaSeconds - CumulativeCooldown))); // Give credit for any overrun
-				}
-				else
-				{
-					AllEnabledTickFunctions.Add(TickFunction);
-				}
+				TickFunctionsToReschedule.Add(FTickScheduleDetails(TickFunction, TickFunction->TickInterval - (Context.DeltaSeconds - CumulativeCooldown))); // Give credit for any overrun
 				AllCoolingDownTickFunctions.Head = TickFunction->Next;
 			}
 			else
@@ -742,22 +830,13 @@ public:
 			if (TickFunction->bTickEvenWhenPaused)
 			{
 				if (CumulativeCooldown < InContext.DeltaSeconds)
-				{
-					TickFunction->TickState = FTickFunction::ETickState::Enabled;
-
+		{
 					TickFunction->TickVisitedGFrameCounter = GFrameCounter;
 					TickFunction->TickQueuedGFrameCounter = GFrameCounter;
 					TickFunction->ExecuteTick(InContext.DeltaSeconds, InContext.TickType, ENamedThreads::GameThread, FGraphEventRef());
 					TickFunction->TaskPointer = nullptr; // this is stale, clear it out now
 
-					if (TickFunction->TickInterval > 0.f)
-					{
-						TickFunctionsToReschedule.Add(FTickScheduleDetails(TickFunction, TickFunction->TickInterval - (InContext.DeltaSeconds - CumulativeCooldown))); // Give credit for any overrun
-					}
-					else
-					{
-						AllEnabledTickFunctions.Add(TickFunction);
-					}
+					TickFunctionsToReschedule.Add(FTickScheduleDetails(TickFunction, TickFunction->TickInterval - (InContext.DeltaSeconds - CumulativeCooldown))); // Give credit for any overrun
 				}
 				else
 				{
@@ -975,9 +1054,9 @@ public:
 						TickFunction->Next->RelativeTickCooldown += TickFunction->RelativeTickCooldown;
 						TickFunction->Next = nullptr;
 					}
-				}
-				else
-				{
+		        }
+		        else
+		        {
 					PrevComparisionFunction = ComparisonFunction;
 					ComparisonFunction = ComparisonFunction->Next;
 				}
@@ -1026,7 +1105,7 @@ private:
 	/** Master list of disabled tick functions **/
 	TSet<FTickFunction*>						AllDisabledTickFunctions;
 	/** Utility array to avoid memory reallocations when collecting functions to reschedule **/
-	TArray<FTickScheduleDetails>				TickFunctionsToReschedule;
+	TArrayWithThreadsafeAdd<FTickScheduleDetails>				TickFunctionsToReschedule;
 	/** List of tick functions added during a tick phase; these items are also duplicated in AllLiveTickFunctions for future frames **/
 	TSet<FTickFunction*>						NewlySpawnedTickFunctions;
 	/** tick context **/
@@ -1111,12 +1190,12 @@ public:
 
 		if (!bConcurrentQueue)
 		{
-			int32 TotalTickFunctions = 0;
-			for( int32 LevelIndex = 0; LevelIndex < LevelList.Num(); LevelIndex++ )
-			{
-				TotalTickFunctions += LevelList[LevelIndex]->StartFrame(Context);
-			}
-			INC_DWORD_STAT_BY(STAT_TicksQueued, TotalTickFunctions);
+		    int32 TotalTickFunctions = 0;
+		    for( int32 LevelIndex = 0; LevelIndex < LevelList.Num(); LevelIndex++ )
+		    {
+			    TotalTickFunctions += LevelList[LevelIndex]->StartFrame(Context);
+		    }
+		    INC_DWORD_STAT_BY(STAT_TicksQueued, TotalTickFunctions);
 			for( int32 LevelIndex = 0; LevelIndex < LevelList.Num(); LevelIndex++ )
 			{
 				LevelList[LevelIndex]->QueueAllTicks();
@@ -1128,8 +1207,13 @@ public:
 			{
 				LevelList[LevelIndex]->StartFrameParallel(Context, AllTickFunctions);
 			}
-			AllCompletionEvents.AddZeroed(AllTickFunctions.Num());
 			INC_DWORD_STAT_BY(STAT_TicksQueued, AllTickFunctions.Num());
+			FTickTaskSequencer& TTS = FTickTaskSequencer::Get();
+			TTS.SetupAddTickTaskCompletionParallel(AllTickFunctions.Num());
+			for( int32 LevelIndex = 0; LevelIndex < LevelList.Num(); LevelIndex++ )
+			{
+				LevelList[LevelIndex]->ReserveTickFunctionCooldowns(AllTickFunctions.Num());
+			}
 			ParallelFor(AllTickFunctions.Num(), 
 				[this](int32 Index)
 				{
@@ -1138,34 +1222,15 @@ public:
 					bool bWasInterval = !!(UPTRINT(TickFunction) & 1);
 					TickFunction = (FTickFunction*)(UPTRINT(TickFunction) & ~1);
 
-					TArray<FTickFunction*, TInlineAllocator<4> > StackForCycleDetection;
-					TickFunction->QueueTickFunctionParallel(Context, StackForCycleDetection, &AllCompletionEvents[0], Index, bWasInterval);
+					TArray<FTickFunction*, TInlineAllocator<8> > StackForCycleDetection;
+					TickFunction->QueueTickFunctionParallel(Context, StackForCycleDetection, bWasInterval);
 				}
 			);
-
-			{
-				SCOPE_CYCLE_COUNTER(STAT_FinalizeParallelQueue);
-
-				FTickTaskSequencer& TTS = FTickTaskSequencer::Get();
-				for (const FTickGroupCompletionItem& Comp : AllCompletionEvents)
-				{
-					check(Comp.Task);
-					TTS.AddTickTaskCompletion(Comp.ActualTickGroup, Comp.Task, Comp.bHiPri);
-					if (Comp.bNeedsToBeRemovedFromTickListsAndRescheduled)
-					{
-						FTickFunction* TickFunction = AllTickFunctions[&Comp - &AllCompletionEvents[0]];
-						TickFunction->TickTaskLevel->RemoveAndRescheduleForInterval(TickFunction);
-					}
-				}
-			}
-
-			for( int32 LevelIndex = 0; LevelIndex < LevelList.Num(); LevelIndex++ )
-			{
+		    for( int32 LevelIndex = 0; LevelIndex < LevelList.Num(); LevelIndex++ )
+		    {
 				LevelList[LevelIndex]->ScheduleTickFunctionCooldowns();
 			}
-
 			AllTickFunctions.Reset();
-			AllCompletionEvents.Reset();
 		}
 	}
 
@@ -1349,7 +1414,6 @@ private:
 
 	/** true during the tick phase, when true, tick function adds also go to the newly spawned list. **/
 	TArray<FTickFunction*> AllTickFunctions;
-	TArray<FTickGroupCompletionItem> AllCompletionEvents;
 };
 
 
@@ -1487,7 +1551,7 @@ void FTickFunction::QueueTickFunction(FTickTaskSequencer& TTS, const struct FTic
 		TickVisitedGFrameCounter = GFrameCounter;
 		if (TickState != FTickFunction::ETickState::Disabled)
 		{
-			ETickingGroup MaxPrerequisiteTickGroup = ETickingGroup(0);
+			ETickingGroup MaxPrerequisiteTickGroup =  ETickingGroup(0);
 
 			FGraphEventArray TaskPrerequisites;
 			for (int32 PrereqIndex = 0; PrereqIndex < Prerequisites.Num(); PrereqIndex++)
@@ -1513,7 +1577,7 @@ void FTickFunction::QueueTickFunction(FTickTaskSequencer& TTS, const struct FTic
 					}
 					else
 					{
-						MaxPrerequisiteTickGroup = FMath::Max<ETickingGroup>(MaxPrerequisiteTickGroup, Prereq->ActualTickGroup);
+						MaxPrerequisiteTickGroup =  FMath::Max<ETickingGroup>(MaxPrerequisiteTickGroup, Prereq->ActualTickGroup);
 						TaskPrerequisites.Add(Prereq->GetCompletionHandle());
 					}
 				}
@@ -1529,14 +1593,14 @@ void FTickFunction::QueueTickFunction(FTickTaskSequencer& TTS, const struct FTic
 
 			if (TickState == FTickFunction::ETickState::Enabled)
 			{
-				TTS.QueueTickTask(&TaskPrerequisites, this, TickContext);
-			}
+			TTS.QueueTickTask(&TaskPrerequisites, this, TickContext);
+		}
 		}
 		TickQueuedGFrameCounter = GFrameCounter;
 	}
 }
 
-void FTickFunction::QueueTickFunctionParallel(const struct FTickContext& TickContext, TArray<FTickFunction*, TInlineAllocator<4> >& StackForCycleDetection, struct FTickGroupCompletionItem* AllCompletionEvents, int32 Index, bool bWasInterval)
+void FTickFunction::QueueTickFunctionParallel(const struct FTickContext& TickContext, TArray<FTickFunction*, TInlineAllocator<8> >& StackForCycleDetection, bool bWasInterval)
 {
 
 	bool bProcessTick;
@@ -1570,7 +1634,7 @@ void FTickFunction::QueueTickFunctionParallel(const struct FTickContext& TickCon
 					else if (Prereq->bRegistered)
 					{
 						// recursive call to make sure my prerequisite is set up so I can use its completion handle
-						Prereq->QueueTickFunctionParallel(TickContext, StackForCycleDetection, AllCompletionEvents, -1, false);
+						Prereq->QueueTickFunctionParallel(TickContext, StackForCycleDetection, false);
 						if (!Prereq->TaskPointer)
 						{
 							//ok UE_LOG(LogTick, Warning, TEXT("While processing prerequisites for %s, could use %s because it is disabled."),*DiagnosticMessage(), *Prereq->DiagnosticMessage());
@@ -1596,6 +1660,10 @@ void FTickFunction::QueueTickFunctionParallel(const struct FTickContext& TickCon
 			if (TickState == FTickFunction::ETickState::Enabled)
 			{
 				FTickTaskSequencer::Get().QueueTickTaskParallel(&TaskPrerequisites, this, TickContext);
+				if (!bWasInterval && TickInterval > 0.f)
+				{
+					TickTaskLevel->RescheduleForIntervalParallel(this);
+				}
 			}
 		}
 		FPlatformMisc::MemoryBarrier();
@@ -1614,20 +1682,6 @@ void FTickFunction::QueueTickFunctionParallel(const struct FTickContext& TickCon
 			}
 		}
 	}
-	if (Index >= 0)
-	{
-		TGraphTask<FTickFunctionTask>* Task = (TGraphTask<FTickFunctionTask>*)TaskPointer;
-		check(Task);
-		FTickGroupCompletionItem& Comp = *(AllCompletionEvents + Index);
-		check(!Comp.Task);
-		Comp.Task = Task;
-		Comp.ActualTickGroup = ActualTickGroup;
-		if (!bWasInterval && TickInterval > 0.f && TickState == FTickFunction::ETickState::Enabled)
-		{
-			Comp.bNeedsToBeRemovedFromTickListsAndRescheduled = true;
-		}
-		Comp.bHiPri = bHighPriority;
-	}
 }
 
 /**
diff --git a/Engine/Source/Runtime/Engine/Private/UnrealEngine.cpp b/Engine/Source/Runtime/Engine/Private/UnrealEngine.cpp
--- a/Engine/Source/Runtime/Engine/Private/UnrealEngine.cpp
+++ b/Engine/Source/Runtime/Engine/Private/UnrealEngine.cpp
@@ -150,7 +150,7 @@ ENGINE_API uint32 GGPUFrameTime = 0;
 FSystemResolution GSystemResolution;
 
 /** Threshold for a frame to be considered a hitch (in seconds. */
-static TAutoConsoleVariable<float> GHitchThresholdCVar(
+static FAutoConsoleVariableRef GHitchThresholdCVar(
 	TEXT("t.HitchThreshold"),
 	GHitchThreshold,
 	TEXT("Time in seconds that is considered a hitch by \"stat dumphitches\"")
@@ -1013,8 +1013,8 @@ void UEngine::PreExit()
 		auto SavedHMD = HMDDevice;
 		{
 			FSuspendRenderingThread Suspend(false);
-			StereoRenderingDevice.Reset();
-			HMDDevice.Reset();
+	StereoRenderingDevice.Reset();
+	HMDDevice.Reset();
 		}
 		// shutdown will occur here.
 	}
@@ -1042,6 +1042,7 @@ void UEngine::TickDeferredCommands()
 	DeferredCommands.RemoveAt(0, DeferredCommandsCount);
 }
 
+
 #if !UE_BUILD_SHIPPING
 
 static TAutoConsoleVariable<int32> CVarABTestHistory(
@@ -1299,12 +1300,17 @@ struct FABTest
 
 	void Stop()
 	{
-		UE_LOG(LogConsoleResponse, Display, TEXT("Stopping AB test."));
+		UE_LOG(LogConsoleResponse, Display, TEXT("Running 'A' console command and stopping test."));
+		SwitchTest(0);
 		bABTestActive = false;
 	}
 	
 	void Start(FString* InABTestCmds)
 	{
+		if (bABTestActive)
+		{
+			Stop();
+		}
 		ABTestCmds[0] = InABTestCmds[0];
 		ABTestCmds[1] = InABTestCmds[1];
 
@@ -1327,7 +1333,7 @@ struct FABTest
 		RemainingPrint = ReportNum;
 
 		bABTestActive = true;
-		Get().SwitchTest(0);
+		SwitchTest(0);
 		UE_LOG(LogConsoleResponse, Display, TEXT("abtest started with A = '%s' and B = '%s'"), *ABTestCmds[0], *ABTestCmds[1]);
 	}
 
@@ -1899,8 +1905,8 @@ void UEngine::Serialize(FArchive& Ar)
 		// Only use the main audio device when counting memory
 		if (FAudioDevice* AudioDevice = GetMainAudioDevice())
 		{
-			AudioDevice->CountBytes(Ar);
-		}
+		AudioDevice->CountBytes(Ar);
+	}
 	}
 }
 
@@ -2064,10 +2070,10 @@ bool UEngine::InitializeAudioDeviceManager()
 					FAudioDevice* NewAudioDevice = AudioDeviceManager->CreateAudioDevice(MainAudioDeviceHandle, true);
 					if (NewAudioDevice)
 					{
-						AudioDeviceManager->SetActiveDevice(MainAudioDeviceHandle);
-					}
-					else
-					{
+							AudioDeviceManager->SetActiveDevice(MainAudioDeviceHandle);
+						}
+						else
+						{
 						ShutdownAudioDeviceManager();
 					}
 				}
@@ -6517,23 +6523,23 @@ void UEngine::UpdateRunningAverageDeltaTime(float DeltaTime, bool bAllowFrameRat
 	{
 		// Smooth the framerate if wanted. The code uses a simplistic running average. Other approaches, like reserving
 		// a percentage of time, ended up creating negative feedback loops in conjunction with GPU load and were abandonend.
-		if( DeltaTime < 0.0f )
-		{
+			if( DeltaTime < 0.0f )
+			{
 #if PLATFORM_ANDROID
-			UE_LOG(LogEngine, Warning, TEXT("Detected negative delta time - ignoring"));
-			DeltaTime = 0.01;
+				UE_LOG(LogEngine, Warning, TEXT("Detected negative delta time - ignoring"));
+				DeltaTime = 0.01;
 #elif (UE_BUILD_SHIPPING && WITH_EDITOR)
-			// End users don't have access to the secure parts of UDN. The localized string points to the release notes,
-			// which should include a link to the AMD CPU drivers download site.
-			UE_LOG(LogEngine, Fatal, TEXT("%s"), TEXT("CPU time drift detected! Please consult release notes on how to address this."));
+				// End users don't have access to the secure parts of UDN. The localized string points to the release notes,
+				// which should include a link to the AMD CPU drivers download site.
+				UE_LOG(LogEngine, Fatal, TEXT("%s"), TEXT("CPU time drift detected! Please consult release notes on how to address this."));
 #else
-			// Send developers to the support list thread.
-			UE_LOG(LogEngine, Fatal, TEXT("Negative delta time! Please see https://udn.epicgames.com/lists/showpost.php?list=ue3bugs&id=4364"));
+				// Send developers to the support list thread.
+				UE_LOG(LogEngine, Fatal, TEXT("Negative delta time! Please see https://udn.epicgames.com/lists/showpost.php?list=ue3bugs&id=4364"));
 #endif
-		}
+			}
 
-		// Keep track of running average over 300 frames, clamping at min of 5 FPS for individual delta times.
-		RunningAverageDeltaTime = FMath::Lerp<float>( RunningAverageDeltaTime, FMath::Min<float>( DeltaTime, 0.2f ), 1 / 300.f );
+			// Keep track of running average over 300 frames, clamping at min of 5 FPS for individual delta times.
+			RunningAverageDeltaTime = FMath::Lerp<float>( RunningAverageDeltaTime, FMath::Min<float>( DeltaTime, 0.2f ), 1 / 300.f );
 	}
 }
 
@@ -6545,19 +6551,19 @@ float UEngine::GetMaxTickRate(float DeltaTime, bool bAllowFrameRateSmoothing) co
 
 	if (bAllowFrameRateSmoothing && IsAllowedFramerateSmoothing())
 	{
-		// Work in FPS domain as that is what the function will return.
-		MaxTickRate = 1.f / RunningAverageDeltaTime;
+			// Work in FPS domain as that is what the function will return.
+			MaxTickRate = 1.f / RunningAverageDeltaTime;
 
-		// Clamp FPS into ini defined min/ max range.
-		if (SmoothedFrameRateRange.HasLowerBound())
-		{
-			MaxTickRate = FMath::Max( MaxTickRate, SmoothedFrameRateRange.GetLowerBoundValue() );
-		}
-		if (SmoothedFrameRateRange.HasUpperBound())
-		{
-			MaxTickRate = FMath::Min( MaxTickRate, SmoothedFrameRateRange.GetUpperBoundValue() );
+			// Clamp FPS into ini defined min/ max range.
+			if (SmoothedFrameRateRange.HasLowerBound())
+			{
+				MaxTickRate = FMath::Max( MaxTickRate, SmoothedFrameRateRange.GetLowerBoundValue() );
+			}
+			if (SmoothedFrameRateRange.HasUpperBound())
+			{
+				MaxTickRate = FMath::Min( MaxTickRate, SmoothedFrameRateRange.GetUpperBoundValue() );
+			}
 		}
-	}
 
 	if (CVarCauseHitches.GetValueOnGameThread())
 	{
@@ -7587,19 +7593,19 @@ void DrawStatsHUD( UWorld* World, FViewport* Viewport, FCanvas* Canvas, UCanvas*
 		for (const FDebugClass& DebugClass : DebugClasses)
 		{
 			if (DebugClass.Class)
-			{
+		{
 				TArray<UObject*> DebugObjectsOfClass;
 				const bool bIncludeDerivedClasses = true;
 				GetObjectsOfClass(DebugClass.Class, DebugObjectsOfClass, bIncludeDerivedClasses);
 				for (UObject* Obj : DebugObjectsOfClass)
-				{
+			{
 					if (!Obj)
-					{
-						continue;
-					}
+				{
+					continue;
+				}
 
 					if (Obj->GetWorld() && Obj->GetWorld() != World)
-					{
+				{
 						continue;
 					}
 
@@ -10663,19 +10669,19 @@ bool UEngine::CommitMapChange( FWorldContext &Context )
 		if (Context.PendingLevelStreamingStatusUpdates.Num() > 0)
 		{
 			for (const FLevelStreamingStatus& PendingUpdate : Context.PendingLevelStreamingStatusUpdates)
-			{
+				{
 				ULevelStreaming** Found = Context.World()->StreamingLevels.FindByPredicate([&](ULevelStreaming* Level){
 					return Level && Level->GetWorldAssetPackageFName() == PendingUpdate.PackageName;
 				});
 
 				if (Found)
-				{
+					{
 					(*Found)->bShouldBeLoaded  = PendingUpdate.bShouldBeLoaded;
 					(*Found)->bShouldBeVisible = PendingUpdate.bShouldBeVisible;
 					(*Found)->LevelLODIndex    = PendingUpdate.LODIndex;
-				}
-				else
-				{
+						}
+						else
+						{
 					UE_LOG(LogStreaming, Log, TEXT("Unable to find streaming object %s"), *PendingUpdate.PackageName.ToString());
 				}
 			}
@@ -11805,6 +11811,116 @@ bool UEngine::ToggleStatRaw(UWorld* World, FCommonViewportClient* ViewportClient
 }
 #endif
 
+
+static uint64 TaskThread = 0xFFFFFFFFFFFFFFFF;
+static uint64 GameThread = 0xFFFFFFFFFFFFFFFF;
+static uint64 RenderThread = 0xFFFFFFFFFFFFFFFF;
+static uint64 RHIThread = 0xFFFFFFFFFFFFFFFF;
+FThreadSafeCounter StallForTaskThread;
+
+void SetAffinityOnThread()
+{
+	if (IsInActualRenderingThread())
+	{
+		FPlatformProcess::SetThreadAffinityMask(RenderThread);
+		UE_LOG(LogConsoleResponse, Display, TEXT("RT     %016llX"), RenderThread);
+	}
+	else if (IsInRHIThread())
+	{
+		FPlatformProcess::SetThreadAffinityMask(RHIThread);
+		UE_LOG(LogConsoleResponse, Display, TEXT("RHI    %016llX"), RHIThread);
+	}
+	else if (IsInGameThread())
+	{
+		FPlatformProcess::SetThreadAffinityMask(GameThread);
+		UE_LOG(LogConsoleResponse, Display, TEXT("GT     %016llX"), GameThread);
+	}
+	else // assume task thread
+	{
+		int32 TaskThreadIndex = int32(FTaskGraphInterface::Get().GetCurrentThreadIfKnown()) - int32(ENamedThreads::ActualRenderingThread) - 1;
+		FPlatformProcess::SetThreadAffinityMask(TaskThread);
+		UE_LOG(LogConsoleResponse, Display, TEXT("Task%2d %016llX"), TaskThreadIndex, TaskThread);
+		StallForTaskThread.Decrement();
+		// we wait for the others to finish here so that we do all task threads
+		while (StallForTaskThread.GetValue())
+		{
+			FPlatformProcess::Sleep(.0001f);
+		}
+	}
+}
+
+
+static void SetupThreadAffinity(const TArray<FString>& Args)
+{
+	static bool LoadedDefaults = false;
+	if (!LoadedDefaults || (Args.Num() && Args[0] == TEXT("default")))
+	{
+		LoadedDefaults = true;
+		TaskThread = FPlatformAffinity::GetTaskGraphThreadMask();
+		GameThread = FPlatformAffinity::GetMainGameMask();
+		RenderThread = FPlatformAffinity::GetRenderingThreadMask();
+		RHIThread = FPlatformAffinity::GetRHIThreadMask();
+	}
+	for (int32 Index = 0; Index + 1 < Args.Num(); Index += 2)
+	{
+		uint64 Aff = FParse::HexNumber(*Args[Index + 1]); // this is only 32 bits
+		if (!Aff)
+		{
+			UE_LOG(LogConsoleResponse, Display, TEXT("Parsed 0 for affinity, using 0xFFFFFFFFFFFFFFFF instead"));
+			Aff = 0xFFFFFFFFFFFFFFFF;
+		}
+		if (Args[Index] == TEXT("GT"))
+		{
+			GameThread = Aff;
+		}
+		else if (Args[Index] == TEXT("RT"))
+		{
+			RenderThread = Aff;
+		}
+		else if (Args[Index] == TEXT("RHI"))
+		{
+			RHIThread = Aff;
+		}
+		else if (Args[Index] == TEXT("Task"))
+		{
+			TaskThread = Aff;
+		}
+		else
+		{
+			UE_LOG(LogConsoleResponse, Display, TEXT("Unrecognized thread name %s"), *Args[Index]);
+		}
+	}
+
+	StallForTaskThread.Reset();
+	StallForTaskThread.Add(FTaskGraphInterface::Get().GetNumWorkerThreads());
+
+	for (int32 Index = 0; Index < FTaskGraphInterface::Get().GetNumWorkerThreads(); Index++)
+	{
+		FSimpleDelegateGraphTask::CreateAndDispatchWhenReady(
+			FSimpleDelegateGraphTask::FDelegate::CreateStatic(&SetAffinityOnThread),
+			TStatId(), NULL, ENamedThreads::AnyThread);
+	}
+	FSimpleDelegateGraphTask::CreateAndDispatchWhenReady(
+		FSimpleDelegateGraphTask::FDelegate::CreateStatic(&SetAffinityOnThread),
+		TStatId(), NULL, ENamedThreads::RenderThread);
+	if (GRHIThread)
+	{
+		FSimpleDelegateGraphTask::CreateAndDispatchWhenReady(
+			FSimpleDelegateGraphTask::FDelegate::CreateStatic(&SetAffinityOnThread),
+			TStatId(), NULL, ENamedThreads::RHIThread);
+	}
+	check(IsInGameThread());
+	SetAffinityOnThread();
+	FlushRenderingCommands();
+	GLog->FlushThreadedLogs();
+}
+
+static FAutoConsoleCommand SetupThreadAffinityCmd(
+	TEXT("SetThreadAffinity"),
+	TEXT("Sets the thread affinity. A single arg of default resets the thread affinity, otherwise pairs of args [GT|RT|RHI|Task] [Hex affinity] sets the affinity."),
+	FConsoleCommandWithArgsDelegate::CreateStatic(&SetupThreadAffinity)
+	);
+
 // REVERB
 #if !(UE_BUILD_SHIPPING || UE_BUILD_TEST)
 int32 UEngine::RenderStatReverb(UWorld* World, FViewport* Viewport, FCanvas* Canvas, int32 X, int32 Y, const FVector* ViewLocation, const FRotator* ViewRotation)
diff --git a/Engine/Source/Runtime/Engine/Private/UserInterface/PlayerInput.cpp b/Engine/Source/Runtime/Engine/Private/UserInterface/PlayerInput.cpp
--- a/Engine/Source/Runtime/Engine/Private/UserInterface/PlayerInput.cpp
+++ b/Engine/Source/Runtime/Engine/Private/UserInterface/PlayerInput.cpp
@@ -890,9 +890,16 @@ void UPlayerInput::ProcessInputStack(const TArray<UInputComponent*>& InputCompon
 		{
 		}
 	};
-	TArray<FAxisDelegateDetails> AxisDelegates;
-	TArray<FVectorAxisDelegateDetails> VectorAxisDelegates;
-	TArray<FDelegateDispatchDetails> NonAxisDelegates;
+	static TArray<FAxisDelegateDetails> AxisDelegates;
+	static TArray<FVectorAxisDelegateDetails> VectorAxisDelegates;
+	static TArray<FDelegateDispatchDetails> NonAxisDelegates;
+	static TArray<FKey> KeysToConsume;
+	static TArray<FDelegateDispatchDetails> FoundChords;
+
+
+	// must be called non-recursively and on the game thread
+	check(IsInGameThread() && !AxisDelegates.Num() && !VectorAxisDelegates.Num() && !NonAxisDelegates.Num() && !KeysToConsume.Num() && !FoundChords.Num());
+
 	struct FDelegateDispatchDetailsSorter
 	{
 		bool operator()( const FDelegateDispatchDetails& A, const FDelegateDispatchDetails& B ) const
@@ -909,8 +916,7 @@ void UPlayerInput::ProcessInputStack(const TArray<UInputComponent*>& InputCompon
 		UInputComponent* const IC = InputComponentStack[StackIndex];
 		if (IC)
 		{
-			TArray<FKey> KeysToConsume;
-			TArray<FDelegateDispatchDetails> FoundChords;
+			check(!KeysToConsume.Num() && !FoundChords.Num())
 
 			for (int32 ActionIndex=0; ActionIndex<IC->GetNumActionBindings(); ++ActionIndex)
 			{
@@ -1081,6 +1087,8 @@ void UPlayerInput::ProcessInputStack(const TArray<UInputComponent*>& InputCompon
 					ConsumeKey(KeysToConsume[KeyIndex]);
 				}
 			}		
+			KeysToConsume.Reset();
+			FoundChords.Reset();
 		}
 	}
 
@@ -1140,6 +1148,9 @@ void UPlayerInput::ProcessInputStack(const TArray<UInputComponent*>& InputCompon
 	PlayerController->PostProcessInput(DeltaTime, bGamePaused);
 
 	FinishProcessingPlayerInput();
+	AxisDelegates.Reset();
+	VectorAxisDelegates.Reset();
+	NonAxisDelegates.Reset();
 }
 
 void UPlayerInput::DiscardPlayerInput()
diff --git a/Engine/Source/Runtime/Engine/Private/VisualLogger/VisualLogger.cpp b/Engine/Source/Runtime/Engine/Private/VisualLogger/VisualLogger.cpp
--- a/Engine/Source/Runtime/Engine/Private/VisualLogger/VisualLogger.cpp
+++ b/Engine/Source/Runtime/Engine/Private/VisualLogger/VisualLogger.cpp
@@ -14,10 +14,9 @@
 #endif
 
 
+DEFINE_LOG_CATEGORY(LogVisual);
 #if ENABLE_VISUAL_LOG 
-
 DEFINE_STAT(STAT_VisualLog);
-DEFINE_LOG_CATEGORY(LogVisual);
 
 TMap<UObject*, TArray<TWeakObjectPtr<const UObject> > > FVisualLogger::RedirectionMap;
 int32 FVisualLogger::bIsRecording = false;
@@ -512,10 +511,12 @@ public:
 #endif
 			}
 		}
+#if ENABLE_VISUAL_LOG
 		else if (FParse::Command(&Cmd, TEXT("LogNavOctree")))
 		{
 			FVisualLogger::NavigationDataDump(GetWorldForVisualLogger(nullptr), LogNavigation, ELogVerbosity::Log, INDEX_NONE, FBox());
 		}
+#endif
 		return false;
 	}
 } LogVisualizerExec;
diff --git a/Engine/Source/Runtime/Engine/Public/AI/Navigation/RecastNavMeshGenerator.h b/Engine/Source/Runtime/Engine/Public/AI/Navigation/RecastNavMeshGenerator.h
--- a/Engine/Source/Runtime/Engine/Public/AI/Navigation/RecastNavMeshGenerator.h
+++ b/Engine/Source/Runtime/Engine/Public/AI/Navigation/RecastNavMeshGenerator.h
@@ -405,7 +405,7 @@ public:
 
 	FBox GrowBoundingBox(const FBox& BBox, bool bIncludeAgentHeight) const;
 	
-#if !(UE_BUILD_SHIPPING || UE_BUILD_TEST)
+#if !(UE_BUILD_SHIPPING || UE_BUILD_TEST) && ENABLE_VISUAL_LOG
 	virtual void ExportNavigationData(const FString& FileName) const override;
 	virtual void GrabDebugSnapshot(struct FVisualLogEntry* Snapshot, const FBox& BoundingBox, const struct FLogCategoryBase& Category, ELogVerbosity::Type Verbosity) const override;
 #endif
diff --git a/Engine/Source/Runtime/Engine/Public/Animation/AnimCurveTypes.h b/Engine/Source/Runtime/Engine/Public/Animation/AnimCurveTypes.h
--- a/Engine/Source/Runtime/Engine/Public/Animation/AnimCurveTypes.h
+++ b/Engine/Source/Runtime/Engine/Public/Animation/AnimCurveTypes.h
@@ -269,10 +269,14 @@ struct ENGINE_API FBlendedCurve
 	void InitFrom(const class USkeleton* Skeleton);
 	void InitFrom(const FBlendedCurve& InCurveToInitFrom)
 	{
-		UIDList = InCurveToInitFrom.UIDList;
-		Elements.Reset(UIDList.Num());
-		Elements.AddZeroed(UIDList.Num());
-		bInitialized = true;
+		if (&InCurveToInitFrom != this)
+		{
+			UIDList.Reset();
+			UIDList.Append(InCurveToInitFrom.UIDList);
+			Elements.Reset();
+			Elements.AddZeroed(UIDList.Num());
+			bInitialized = true;
+		}
 	}
 
 	/** Set value of InUID to InValue */
diff --git a/Engine/Source/Runtime/Engine/Public/EngineDefines.h b/Engine/Source/Runtime/Engine/Public/EngineDefines.h
--- a/Engine/Source/Runtime/Engine/Public/EngineDefines.h
+++ b/Engine/Source/Runtime/Engine/Public/EngineDefines.h
@@ -55,7 +55,7 @@
 #endif //WITH_APEX_CLOTHING
 
 #ifndef ENABLE_VISUAL_LOG
-#	define ENABLE_VISUAL_LOG (1 && !NO_LOGGING && !USING_CODE_ANALYSIS && !(UE_BUILD_SHIPPING || UE_BUILD_TEST))
+	#define ENABLE_VISUAL_LOG (PLATFORM_DESKTOP && !NO_LOGGING && !USING_CODE_ANALYSIS && !(UE_BUILD_SHIPPING || UE_BUILD_TEST))
 #endif
 
 #ifndef WITH_FIXED_AREA_ENTERING_COST
diff --git a/Engine/Source/Runtime/Engine/Public/ParticleHelper.h b/Engine/Source/Runtime/Engine/Public/ParticleHelper.h
--- a/Engine/Source/Runtime/Engine/Public/ParticleHelper.h
+++ b/Engine/Source/Runtime/Engine/Public/ParticleHelper.h
@@ -1214,6 +1214,30 @@ enum EDynamicEmitterType
 	DET_Custom
 };
 
+struct FParticleDataContainer
+{
+	int32 MemBlockSize;
+	int32 ParticleDataNumBytes;
+	int32 ParticleIndicesNumShorts;
+	uint8* ParticleData; // this is also the memory block we allocated
+	uint16* ParticleIndices; // not allocated, this is at the end of the memory block
+
+	FParticleDataContainer()
+		: MemBlockSize(0)
+		, ParticleDataNumBytes(0)
+		, ParticleIndicesNumShorts(0)
+		, ParticleData(nullptr)
+		, ParticleIndices(nullptr)
+	{
+	}
+	~FParticleDataContainer()
+	{
+		Free();
+	}
+	void Alloc(int32 InParticleDataNumBytes, int32 InParticleIndicesNumShorts);
+	void Free();
+};
+
 /** Source data base class for all emitter types */
 struct FDynamicEmitterReplayDataBase
 {
@@ -1224,8 +1248,7 @@ struct FDynamicEmitterReplayDataBase
 	int32 ActiveParticleCount;
 
 	int32 ParticleStride;
-	TArray<uint8> ParticleData;
-	TArray<uint16> ParticleIndices;
+	FParticleDataContainer DataContainer;
 
 	FVector Scale;
 
@@ -1264,8 +1287,45 @@ struct FDynamicEmitterReplayDataBase
 
 		Ar << ActiveParticleCount;
 		Ar << ParticleStride;
+		
+		TArray<uint8> ParticleData;
+		TArray<uint16> ParticleIndices;
+
+		if (!Ar.IsLoading() && !Ar.IsObjectReferenceCollector())
+		{
+			if (DataContainer.ParticleDataNumBytes)
+			{
+				ParticleData.AddUninitialized(DataContainer.ParticleDataNumBytes);
+				FMemory::Memcpy(ParticleData.GetData(), DataContainer.ParticleData, DataContainer.ParticleDataNumBytes);
+			}
+			if (DataContainer.ParticleIndicesNumShorts)
+			{
+				ParticleIndices.AddUninitialized(DataContainer.ParticleIndicesNumShorts);
+				FMemory::Memcpy(ParticleIndices.GetData(), DataContainer.ParticleIndices, DataContainer.ParticleIndicesNumShorts * sizeof(uint16));
+			}
+		}
+
 		Ar << ParticleData;
 		Ar << ParticleIndices;
+
+		if (Ar.IsLoading())
+		{
+			DataContainer.Free();
+			if (ParticleData.Num())
+			{
+				DataContainer.Alloc(ParticleData.Num(), ParticleIndices.Num());
+				FMemory::Memcpy(DataContainer.ParticleData, ParticleData.GetData(), DataContainer.ParticleDataNumBytes);
+				if (DataContainer.ParticleIndicesNumShorts)
+				{
+					FMemory::Memcpy(DataContainer.ParticleIndices, ParticleIndices.GetData(), DataContainer.ParticleIndicesNumShorts * sizeof(uint16));
+				}
+			}
+			else
+			{
+				check(!ParticleIndices.Num());
+			}
+		}
+
 		Ar << Scale;
 		Ar << SortMode;
 		Ar << bOverrideSystemMacroUV;
@@ -1285,6 +1345,10 @@ struct FDynamicEmitterDataBase
 		ReturnVertexFactory();
 	}
 
+	/** Custom new/delete with recycling */
+	void* operator new(size_t Size);
+	void operator delete(void *RawMemory, size_t Size);
+
 	/**
 	 *	Create the render thread resources for this emitter data
 	 *
@@ -1461,7 +1525,7 @@ struct FDynamicSpriteEmitterDataBase : public FDynamicEmitterDataBase
 	 *	@param	ParticleOrder		The array to fill in with ordered indices
 	 */
 	void SortSpriteParticles(int32 SortMode, bool bLocalSpace, 
-		int32 ParticleCount, const TArray<uint8>& ParticleData, int32 ParticleStride, const TArray<uint16>& ParticleIndices,
+		int32 ParticleCount, const uint8* ParticleData, int32 ParticleStride, const uint16* ParticleIndices,
 		const FSceneView* View, const FMatrix& LocalToWorld, FParticleOrder* ParticleOrder) const;
 
 	/**
@@ -2298,30 +2362,33 @@ public:
 	{
 	}
 
-	virtual ~FParticleDynamicData()
+	~FParticleDynamicData()
 	{
 		ClearEmitterDataArray();
 	}
 
+	/** Custom new/delete with recycling */
+	void* operator new(size_t Size);
+	void operator delete(void *RawMemory, size_t Size);
+
 	void ClearEmitterDataArray()
 	{
 		for (int32 Index = 0; Index < DynamicEmitterDataArray.Num(); Index++)
 		{
 			FDynamicEmitterDataBase* Data =	DynamicEmitterDataArray[Index];
 			delete Data;
-			DynamicEmitterDataArray[Index] = NULL;
 		}
-		DynamicEmitterDataArray.Empty();
+		DynamicEmitterDataArray.Reset();
 	}
 
 	uint32 GetMemoryFootprint( void ) const { return( sizeof( *this ) + DynamicEmitterDataArray.GetAllocatedSize() ); }
 
-	// Variables
-	TArray<FDynamicEmitterDataBase*>	DynamicEmitterDataArray;
-
 	/** The Current Emmitter we are rendering **/
 	uint32 EmitterIndex;
 
+	// Variables
+	TArray<FDynamicEmitterDataBase*, TInlineAllocator<12> >	DynamicEmitterDataArray;
+
 	/** World space position that UVs generated with the ParticleMacroUV material node will be centered on. */
 	FVector SystemPositionForMacroUVs;
 
diff --git a/Engine/Source/Runtime/Engine/Public/SkeletalRenderPublic.h b/Engine/Source/Runtime/Engine/Public/SkeletalRenderPublic.h
--- a/Engine/Source/Runtime/Engine/Public/SkeletalRenderPublic.h
+++ b/Engine/Source/Runtime/Engine/Public/SkeletalRenderPublic.h
@@ -40,12 +40,6 @@ public:
 	 */
 	virtual void Update(int32 LODIndex,USkinnedMeshComponent* InMeshComponent,const TArray<FActiveVertexAnim>& ActiveVertexAnims) = 0;
 
-	/**
-	 * Called by the rendering thread to update the current dynamic data
-	 * @param	InDynamicData - data that was created by the game thread for use by the rendering thread
-	 */
-	virtual void UpdateDynamicData_RenderThread(FRHICommandListImmediate& RHICmdList, class FDynamicSkelMeshObjectData* InDynamicData) = 0;
-
 	/**
 	 * Called by FSkeletalMeshObject prior to GDME. This allows the GPU skin version to update bones etc now that we know we are going to render
 	 */
diff --git a/Engine/Source/Runtime/GameplayAbilities/Public/AbilitySystemLog.h b/Engine/Source/Runtime/GameplayAbilities/Public/AbilitySystemLog.h
--- a/Engine/Source/Runtime/GameplayAbilities/Public/AbilitySystemLog.h
+++ b/Engine/Source/Runtime/GameplayAbilities/Public/AbilitySystemLog.h
@@ -14,7 +14,7 @@
 GAMEPLAYABILITIES_API DECLARE_LOG_CATEGORY_EXTERN(LogAbilitySystem, Warning, All);
 GAMEPLAYABILITIES_API DECLARE_LOG_CATEGORY_EXTERN(VLogAbilitySystem, Warning, All);
 
-#if NO_LOGGING
+#if NO_LOGGING || !PLATFORM_DESKTOP
 
 // Without logging enabled we pass ability system through to UE_LOG which only handles Fatal verbosity in NO_LOGGING
 #define ABILITY_LOG(Verbosity, Format, ...) \
diff --git a/Engine/Source/Runtime/GameplayTags/Private/GameplayTagsManager.cpp b/Engine/Source/Runtime/GameplayTags/Private/GameplayTagsManager.cpp
--- a/Engine/Source/Runtime/GameplayTags/Private/GameplayTagsManager.cpp
+++ b/Engine/Source/Runtime/GameplayTags/Private/GameplayTagsManager.cpp
@@ -198,7 +198,7 @@ void UGameplayTagsManager::ConstructNetIndex()
 
 	{
 		FScopeLock Lock(&GameplayTagNodeMapCritical);
-		GameplayTagNodeMap.GenerateValueArray(NetworkGameplayTagNodeIndex);
+	GameplayTagNodeMap.GenerateValueArray(NetworkGameplayTagNodeIndex);
 	}
 
 	NetworkGameplayTagNodeIndex.Sort(FCompareFGameplayTagNodeByTag());
@@ -426,13 +426,13 @@ int32 UGameplayTagsManager::InsertTagIntoNodeArray(FName Tag, TWeakPtr<FGameplay
 
 		{
 			FScopeLock Lock(&GameplayTagMapCritical);
-			GameplayTagMap.Add(TagNode->GetCompleteTag(), GameplayTag);
+		GameplayTagMap.Add(TagNode->GetCompleteTag(), GameplayTag);
 		}
 
 		{
 			FScopeLock Lock(&GameplayTagNodeMapCritical);
-			GameplayTagNodeMap.Add(GameplayTag, TagNode);
-		}
+		GameplayTagNodeMap.Add(GameplayTag, TagNode);
+	}
 	}
 	else if (NodeArray[InsertionIdx]->CategoryDescription.IsEmpty() && !CategoryDescription.IsEmpty())
 	{
@@ -659,46 +659,74 @@ void UGameplayTagsManager::AddChildrenTags(FGameplayTagContainer& TagContainer,
 
 bool UGameplayTagsManager::GameplayTagsMatch(const FGameplayTag& GameplayTagOne, TEnumAsByte<EGameplayTagMatchType::Type> MatchTypeOne, const FGameplayTag& GameplayTagTwo, TEnumAsByte<EGameplayTagMatchType::Type> MatchTypeTwo) const
 {
-	TSet<FName> Tags1;
-	TSet<FName> Tags2;
+	//fixme, why is there a critical section here
+	FScopeLock Lock(&GameplayTagNodeMapCritical);
+	check(IsInGameThread());
+	if (MatchTypeOne == EGameplayTagMatchType::Explicit && MatchTypeTwo == EGameplayTagMatchType::Explicit)
+	{
+		const TSharedPtr<FGameplayTagNode>* TagNode1 = GameplayTagNodeMap.Find(GameplayTagOne);
+		const TSharedPtr<FGameplayTagNode>* TagNode2 = GameplayTagNodeMap.Find(GameplayTagTwo);
+		if (TagNode1 && TagNode2)
+		{
+			return (*TagNode1)->GetCompleteTag() == (*TagNode2)->GetCompleteTag();
+		}
+		return false;
+	}
+	static TSet<FName> Tags1;
+	static TSet<FName> Tags2;
+
+	check(!Tags1.Num() && !Tags2.Num()); // must be game thread, cannot call this recursively
 
+
+	if (MatchTypeOne == EGameplayTagMatchType::IncludeParentTags)
 	{
-		FScopeLock Lock(&GameplayTagNodeMapCritical);
-		if (MatchTypeOne == EGameplayTagMatchType::Explicit)
+		const TSharedPtr<FGameplayTagNode>* TagNode = GameplayTagNodeMap.Find(GameplayTagOne);
+		if (TagNode)
 		{
-			const TSharedPtr<FGameplayTagNode>* TagNode = GameplayTagNodeMap.Find(GameplayTagOne);
-			if (TagNode)
-			{
-				Tags1.Add((*TagNode)->GetCompleteTag());
-			}
+			GetAllParentNodeNames(Tags1, *TagNode);
 		}
-		if (MatchTypeTwo == EGameplayTagMatchType::Explicit)
+	}
+	if (MatchTypeTwo == EGameplayTagMatchType::IncludeParentTags)
+	{
+		const TSharedPtr<FGameplayTagNode>* TagNode = GameplayTagNodeMap.Find(GameplayTagTwo);
+		if (TagNode)
 		{
-			const TSharedPtr<FGameplayTagNode>* TagNode = GameplayTagNodeMap.Find(GameplayTagTwo);
-			if (TagNode)
-			{
-				Tags2.Add((*TagNode)->GetCompleteTag());
-			}
+			GetAllParentNodeNames(Tags2, *TagNode);
 		}
-		if (MatchTypeOne == EGameplayTagMatchType::IncludeParentTags)
+	}
+	bool bResult = false;
+
+	if (MatchTypeOne == EGameplayTagMatchType::Explicit)
+	{
+		const TSharedPtr<FGameplayTagNode>* TagNode = GameplayTagNodeMap.Find(GameplayTagOne);
+		if (TagNode)
 		{
-			const TSharedPtr<FGameplayTagNode>* TagNode = GameplayTagNodeMap.Find(GameplayTagOne);
-			if (TagNode)
-			{
-				GetAllParentNodeNames(Tags1, *TagNode);
-			}
+			bResult = Tags2.Contains((*TagNode)->GetCompleteTag());
+		}
+	}
+	else if (MatchTypeTwo == EGameplayTagMatchType::Explicit)
+	{
+		const TSharedPtr<FGameplayTagNode>* TagNode = GameplayTagNodeMap.Find(GameplayTagTwo);
+		if (TagNode)
+		{
+			bResult = Tags1.Contains((*TagNode)->GetCompleteTag());
 		}
-		if (MatchTypeTwo == EGameplayTagMatchType::IncludeParentTags)
+	}
+	else
+	{
+		for (auto& Tag1 : Tags1)
 		{
-			const TSharedPtr<FGameplayTagNode>* TagNode = GameplayTagNodeMap.Find(GameplayTagTwo);
-			if (TagNode)
+			if (Tags2.Contains(Tag1))
 			{
-				GetAllParentNodeNames(Tags2, *TagNode);
+				bResult = true;
+				break;
 			}
 		}
 	}
 
-	return Tags1.Intersect(Tags2).Num() > 0;
+	Tags1.Reset();
+	Tags2.Reset();
+	return bResult;
 }
 
 void UGameplayTagsManager::GetAllParentNodeNames(TSet<FName>& NamesList, const TSharedPtr<FGameplayTagNode> GameplayTag) const
diff --git a/Engine/Source/Runtime/GameplayTasks/Classes/GameplayTasksComponent.h b/Engine/Source/Runtime/GameplayTasks/Classes/GameplayTasksComponent.h
--- a/Engine/Source/Runtime/GameplayTasks/Classes/GameplayTasksComponent.h
+++ b/Engine/Source/Runtime/GameplayTasks/Classes/GameplayTasksComponent.h
@@ -117,7 +117,7 @@ public:
 
 	static EGameplayTaskRunResult RunGameplayTask(IGameplayTaskOwnerInterface& TaskOwner, UGameplayTask& Task, uint8 Priority, FGameplayResourceSet AdditionalRequiredResources, FGameplayResourceSet AdditionalClaimedResources);
 	
-#if !(UE_BUILD_SHIPPING || UE_BUILD_TEST)
+#if !(UE_BUILD_SHIPPING || UE_BUILD_TEST) && ENABLE_VISUAL_LOG
 	FString GetTickingTasksDescription() const;
 	FString GetTasksPriorityQueueDescription() const;
 #endif // !(UE_BUILD_SHIPPING || UE_BUILD_TEST)
diff --git a/Engine/Source/Runtime/GameplayTasks/Private/GameplayTasksComponent.cpp b/Engine/Source/Runtime/GameplayTasks/Private/GameplayTasksComponent.cpp
--- a/Engine/Source/Runtime/GameplayTasks/Private/GameplayTasksComponent.cpp
+++ b/Engine/Source/Runtime/GameplayTasks/Private/GameplayTasksComponent.cpp
@@ -428,7 +428,7 @@ void UGameplayTasksComponent::SetCurrentlyClaimedResources(FGameplayResourceSet
 //----------------------------------------------------------------------//
 // debugging
 //----------------------------------------------------------------------//
-#if !(UE_BUILD_SHIPPING || UE_BUILD_TEST)
+#if !(UE_BUILD_SHIPPING || UE_BUILD_TEST) && ENABLE_VISUAL_LOG
 FString UGameplayTasksComponent::GetTickingTasksDescription() const
 {
 	FString TasksDescription;
diff --git a/Engine/Source/Runtime/Renderer/Private/DeferredShadingRenderer.cpp b/Engine/Source/Runtime/Renderer/Private/DeferredShadingRenderer.cpp
--- a/Engine/Source/Runtime/Renderer/Private/DeferredShadingRenderer.cpp
+++ b/Engine/Source/Runtime/Renderer/Private/DeferredShadingRenderer.cpp
@@ -767,7 +767,7 @@ void FDeferredShadingSceneRenderer::RenderOcclusion(FRHICommandListImmediate& RH
 			for (int32 Dest = 1; Dest < NumFrames; Dest++)
 			{
 				OcclusionSubmittedFence[Dest] = OcclusionSubmittedFence[Dest - 1];
-			}
+		}
 			OcclusionSubmittedFence[0] = RHICmdList.RHIThreadFence();
 		}
 	}
@@ -780,6 +780,11 @@ static void ServiceLocalQueue()
 	FTaskGraphInterface::Get().ProcessThreadUntilIdle(ENamedThreads::RenderThread_Local);
 }
 
+static TAutoConsoleVariable<float> CVarStallInitViews(
+	TEXT("CriticalPathStall.AfterInitViews"),
+	0.0f,
+	TEXT("Sleep for the given time after InitViews. Time is given in ms. This is a debug option used for critical path analysis and forcing a change in the critical path."));
+
 void FDeferredShadingSceneRenderer::Render(FRHICommandListImmediate& RHICmdList)
 {
 	FSceneRenderTargets& SceneContext = FSceneRenderTargets::Get(RHICmdList);
@@ -815,6 +820,13 @@ void FDeferredShadingSceneRenderer::Render(FRHICommandListImmediate& RHICmdList)
 
 	// Find the visible primitives.
 	InitViews(RHICmdList);
+#if !UE_BUILD_SHIPPING
+	if (CVarStallInitViews.GetValueOnRenderThread() > 0.0f)
+	{
+		QUICK_SCOPE_CYCLE_COUNTER(STAT_InitViews_Intentional_Stall);
+		FPlatformProcess::Sleep(CVarStallInitViews.GetValueOnRenderThread() / 1000.0f);
+	}
+#endif
 
 	if (GRHICommandList.UseParallelAlgorithms())
 	{
@@ -1047,20 +1059,20 @@ void FDeferredShadingSceneRenderer::Render(FRHICommandListImmediate& RHICmdList)
 	ServiceLocalQueue();
 
 	if (ViewFamily.EngineShowFlags.VisualizeLightCulling)
-	{
-		// clear out emissive and baked lighting (not too efficient but simple and only needed for this debug view)
-		SceneContext.BeginRenderingSceneColor(RHICmdList);
-		RHICmdList.Clear(true, FLinearColor(0, 0, 0, 0), false, (float)ERHIZBuffer::FarPlane, false, 0, FIntRect());
-	}
+	  {
+		  // clear out emissive and baked lighting (not too efficient but simple and only needed for this debug view)
+		  SceneContext.BeginRenderingSceneColor(RHICmdList);
+		  RHICmdList.Clear(true, FLinearColor(0, 0, 0, 0), false, (float)ERHIZBuffer::FarPlane, false, 0, FIntRect());
+	  }
 
-	SceneContext.DBufferA.SafeRelease();
-	SceneContext.DBufferB.SafeRelease();
-	SceneContext.DBufferC.SafeRelease();
+	  SceneContext.DBufferA.SafeRelease();
+	  SceneContext.DBufferB.SafeRelease();
+	  SceneContext.DBufferC.SafeRelease();
 
-	// only temporarily available after early z pass and until base pass
-	check(!SceneContext.DBufferA);
-	check(!SceneContext.DBufferB);
-	check(!SceneContext.DBufferC);
+	  // only temporarily available after early z pass and until base pass
+	  check(!SceneContext.DBufferA);
+	  check(!SceneContext.DBufferB);
+	  check(!SceneContext.DBufferC);
 
 	if (bRequiresFarZQuadClear)
 	{
@@ -1691,10 +1703,10 @@ bool FDeferredShadingSceneRenderer::RenderBasePass(FRHICommandListImmediate& RHI
 				FViewInfo& View = Views[ViewIndex];
 				RenderBasePassViewParallel(View, RHICmdList);
 			}
-			bDirty = true; // assume dirty since we are not going to wait
-			if (FVelocityRendering::OutputsToGBuffer())
-			{
-				GPrevPerBoneMotionBlur.EndAppendFence(RHICmdList);
+				bDirty = true; // assume dirty since we are not going to wait
+				if (FVelocityRendering::OutputsToGBuffer())
+				{
+					GPrevPerBoneMotionBlur.EndAppendFence(RHICmdList);
 			}
 		}
 		else
