Fixed crash importing skeletal mesh LODs due to calculating resource size on LODs which do not have adjacency index buffers

https://jira.ol.epicgames.net/browse/UE-19276

diff --git a/Engine/Source/Runtime/Engine/Private/SkeletalMesh.cpp b/Engine/Source/Runtime/Engine/Private/SkeletalMesh.cpp
--- a/Engine/Source/Runtime/Engine/Private/SkeletalMesh.cpp
+++ b/Engine/Source/Runtime/Engine/Private/SkeletalMesh.cpp
@@ -1642,10 +1642,13 @@ SIZE_T FStaticLODModel::GetResourceSize() const
 		ResourceSize += IndexBuffer->GetResourceDataSize(); 
 	}
 
-	const FRawStaticIndexBuffer16or32Interface* AdjacentIndexBuffer = AdjacencyMultiSizeIndexContainer.GetIndexBuffer();
-	if (AdjacentIndexBuffer)
+	if( AdjacencyMultiSizeIndexContainer.IsIndexBufferValid() )
 	{
-		ResourceSize += AdjacentIndexBuffer->GetResourceDataSize();
+		const FRawStaticIndexBuffer16or32Interface* AdjacentIndexBuffer = AdjacencyMultiSizeIndexContainer.GetIndexBuffer();
+		if(AdjacentIndexBuffer)
+		{
+			ResourceSize += AdjacentIndexBuffer->GetResourceDataSize();
+		}
 	}
 
 	ResourceSize += VertexBufferGPUSkin.GetVertexDataSize();
