Added p.bUseUnifiedHeightfield console variable to allow disabling the use of PhysX unified heightfield support (enabled by default).  When it is enabled, landscape collision thickness does not function and CCD is the only option to prevent fast moving rigid body penetration.

diff --git a/Engine/Source/Runtime/Engine/Private/PhysicsEngine/PhysLevel.cpp b/Engine/Source/Runtime/Engine/Private/PhysicsEngine/PhysLevel.cpp
--- a/Engine/Source/Runtime/Engine/Private/PhysicsEngine/PhysLevel.cpp
+++ b/Engine/Source/Runtime/Engine/Private/PhysicsEngine/PhysLevel.cpp
@@ -86,6 +86,12 @@ static TAutoConsoleVariable<int32> CVarAPEXSortDynamicChunksByBenefit(
 	TEXT("True if APEX should sort dynamic chunks by benefit."),
 	ECVF_Default);
 
+static TAutoConsoleVariable<int32> CVarUseUnifiedHeightfield(
+	TEXT("p.bUseUnifiedHeightfield"),
+	1,
+	TEXT("Whether to use the PhysX unified heightfield. This feature of PhysX makes landscape collision consistent with triangle meshes but the thickness parameter is not supported for unified heightfields. 1 enables and 0 disables. Default: 1"),
+	ECVF_ReadOnly);
+
 
 
 //////////////////////////////////////////////////////////////////////////
@@ -362,10 +368,18 @@ void InitGamePhys()
 	PxInitVehicleSDK(*GPhysXSDK);
 #endif
 
-	//Turn on PhysX 3.3 unified height field collision detection. 
-	//This approach shares the collision detection code between meshes and height fields such that height fields behave identically to the equivalent terrain created as a mesh. 
-	//This approach facilitates mixing the use of height fields and meshes in the application with no tangible difference in collision behavior between the two approaches
-	PxRegisterUnifiedHeightFields(*GPhysXSDK);
+	if (CVarUseUnifiedHeightfield.GetValueOnGameThread())
+	{
+		//Turn on PhysX 3.3 unified height field collision detection. 
+		//This approach shares the collision detection code between meshes and height fields such that height fields behave identically to the equivalent terrain created as a mesh. 
+		//This approach facilitates mixing the use of height fields and meshes in the application with no tangible difference in collision behavior between the two approaches except that 
+		//heightfield thickness is not supported for unified heightfields.
+		PxRegisterUnifiedHeightFields(*GPhysXSDK);
+	}
+	else
+	{
+		PxRegisterHeightFields(*GPhysXSDK);
+	}
 
 	if( FParse::Param( FCommandLine::Get(), TEXT( "PVD" ) ) )
 	{
