32bpp HDR encoding for simple element rendering.

diff --git a/Engine/Shaders/SimpleElementPixelShader.usf b/Engine/Shaders/SimpleElementPixelShader.usf
--- a/Engine/Shaders/SimpleElementPixelShader.usf
+++ b/Engine/Shaders/SimpleElementPixelShader.usf
@@ -60,6 +60,43 @@ void PerformEditorCompositingDepthTest(float4 SvPosition)
 
 #endif
 
+float3 SimpleElementFrameBufferBlendOp(float4 Source)
+{
+#if (COMPILER_GLSL_ES2)
+	half4 Dest = FramebufferFetchES2();
+	Dest.rgb = Decode32BPPHDR(Dest);
+#else
+	half4 Dest = half4(0, 0, 0, 0);
+#endif
+
+#if SE_BLEND_MODE == SE_BLEND_OPAQUE || SE_BLEND_MODE == SE_BLEND_MASKED || SE_BLEND_MODE == SE_BLEND_MASKEDDISTANCEFIELD || SE_BLEND_MODE == SE_BLEND_MASKEDDISTANCEFIELDSHADOWED
+	return Source.rgb;
+#elif SE_BLEND_MODE == SE_BLEND_TRANSLUCENT || SE_BLEND_MODE == SE_BLEND_TRANSLUCENTALPHAONLY || SE_BLEND_MODE == SE_BLEND_ALPHABLEND || SE_BLEND_MODE == SE_BLEND_TRANSLUCENTDISTANCEFIELD || SE_BLEND_MODE == SE_BLEND_TRANSLUCENTDISTANCEFIELDSHADOWED
+	return (Source.rgb*Source.a) + (Dest.rgb*(1.0 - Source.a));
+#elif SE_BLEND_MODE == SE_BLEND_ADDITIVE
+	return Source.rgb + Dest.rgb;
+#elif SE_BLEND_MODE == SE_BLEND_MODULATE
+	return Source.rgb * Dest.rgb;
+#elif SE_BLEND_MODE == SE_BLEND_ALPHACOMPOSITE
+	return Source.rgb + (Dest.rgb*(1.0 - Source.a));
+#endif
+
+	return Source.rgb;
+}
+
+float4 SimpleElementEncodeFor32BPPHDR(float4 OutColor)
+{
+#if ES2_PROFILE && USE_32BPP_HDR
+	if (GetHDR32bppEncodeMode() == HDR_ENCODE_RGBA)
+	{
+		OutColor.rgb = SimpleElementFrameBufferBlendOp(OutColor);
+	}
+	return Encode32BPPHDR(OutColor, float2(0,0)); // TODO: fix for mosaic.
+#else
+	return RETURN_COLOR(OutColor);
+#endif
+}
+
 void Main(
 	in float2 TextureCoordinate : TEXCOORD0,
 	in float4 Color : TEXCOORD1,
@@ -86,6 +123,8 @@ void Main(
 	// Set the G buffer bits that indicate unlit
 	OutWorldNormal = 0;
 #endif
+
+	OutColor = SimpleElementEncodeFor32BPPHDR(OutColor);
 }
 
 void AlphaOnlyMain(
@@ -113,8 +152,11 @@ void AlphaOnlyMain(
 	// Set the G buffer bits that indicate unlit
 	OutWorldNormal = 0;
 #endif
+
+	OutColor = SimpleElementEncodeFor32BPPHDR(OutColor);
 }
 
+
 #if METAL_PROFILE
 static const half Gamma = 1.0f;
 #else
@@ -152,6 +194,8 @@ void GammaMain(
 	// Set the G buffer bits that indicate unlit
 	OutWorldNormal = 0;
 #endif
+
+	OutColor = SimpleElementEncodeFor32BPPHDR(OutColor);
 }
 
 void GammaAlphaOnlyMain(
@@ -184,6 +228,7 @@ void GammaAlphaOnlyMain(
 	// Set the G buffer bits that indicate unlit
 	OutWorldNormal = 0;
 #endif
+	OutColor = SimpleElementEncodeFor32BPPHDR(OutColor);
 }
 
 float ClipRef;
@@ -220,6 +265,7 @@ void GammaMaskedMain(
 	// Set the G buffer bits that indicate unlit
 	OutWorldNormal = 0;
 #endif
+	OutColor = SimpleElementEncodeFor32BPPHDR(OutColor);
 }
 
 
@@ -307,8 +353,9 @@ void GammaDistanceFieldMain(
 	{
 		// Gamma correct the output color.
 		OutColor.rgb = pow(saturate(OutColor.rgb),Gamma);	
-	}	
-	OutColor = RETURN_COLOR(OutColor);
+	}
+
+	OutColor = SimpleElementEncodeFor32BPPHDR(OutColor);
 }
 
 TextureCube CubeTexture;
diff --git a/Engine/Source/Runtime/Engine/Private/BatchedElements.cpp b/Engine/Source/Runtime/Engine/Private/BatchedElements.cpp
--- a/Engine/Source/Runtime/Engine/Private/BatchedElements.cpp
+++ b/Engine/Source/Runtime/Engine/Private/BatchedElements.cpp
@@ -296,8 +296,14 @@ void FBatchedElements::AddSprite(
 }
 
 /** Translates a ESimpleElementBlendMode into a RHI state change for rendering a mesh with the blend mode normally. */
-static void SetBlendState(FRHICommandList& RHICmdList, ESimpleElementBlendMode BlendMode)
+static void SetBlendState(FRHICommandList& RHICmdList, ESimpleElementBlendMode BlendMode, bool bEncodedHDR)
 {
+	if (bEncodedHDR)
+	{
+		RHICmdList.SetBlendState(TStaticBlendState<>::GetRHI());
+		return;
+	}
+
 	switch(BlendMode)
 	{
 	case SE_BLEND_Opaque:
@@ -363,21 +369,85 @@ static void SetHitTestingBlendState(FRHICommandList& RHICmdList, ESimpleElementB
 	}
 }
 
-FGlobalBoundShaderState FBatchedElements::SimpleBoundShaderState;
-FGlobalBoundShaderState FBatchedElements::RegularSRGBBoundShaderState;
-FGlobalBoundShaderState FBatchedElements::RegularLinearBoundShaderState;
-FGlobalBoundShaderState FBatchedElements::MaskedBoundShaderState;
-FGlobalBoundShaderState FBatchedElements::DistanceFieldBoundShaderState;
-FGlobalBoundShaderState FBatchedElements::HitTestingBoundShaderState;
-FGlobalBoundShaderState FBatchedElements::ColorChannelMaskShaderState;
-FGlobalBoundShaderState FBatchedElements::AlphaOnlyShaderState;
-FGlobalBoundShaderState FBatchedElements::GammaAlphaOnlyShaderState;
+FBatchedElements::FSimpleElementBSSContainer FBatchedElements::SimpleBoundShaderState;
+FBatchedElements::FSimpleElementBSSContainer FBatchedElements::RegularSRGBBoundShaderState;
+FBatchedElements::FSimpleElementBSSContainer FBatchedElements::RegularLinearBoundShaderState;
+FBatchedElements::FSimpleElementBSSContainer FBatchedElements::MaskedBoundShaderState;
+FBatchedElements::FSimpleElementBSSContainer FBatchedElements::DistanceFieldBoundShaderState;
+FBatchedElements::FSimpleElementBSSContainer FBatchedElements::HitTestingBoundShaderState;
+FBatchedElements::FSimpleElementBSSContainer FBatchedElements::ColorChannelMaskShaderState;
+FBatchedElements::FSimpleElementBSSContainer FBatchedElements::AlphaOnlyShaderState;
+FBatchedElements::FSimpleElementBSSContainer FBatchedElements::GammaAlphaOnlyShaderState;
 
 /** Global alpha ref test value for rendering masked batched elements */
 float GBatchedElementAlphaRefVal = 128.f;
 /** Global smoothing width for rendering batched elements with distance field blend modes */
 float GBatchedElementSmoothWidth = 12;
 
+template<class TSimpleElementPixelShader>
+static TSimpleElementPixelShader* GetPixelShader(bool bEncoded, ESimpleElementBlendMode BlendMode, ERHIFeatureLevel::Type FeatureLevel)
+{
+	if (bEncoded)
+	{
+		// When encoding blending occurs in the shader. return the appropriate blend shader.
+		switch (BlendMode)
+		{
+			case SE_BLEND_Opaque:
+				return *TShaderMapRef<FEncodedSimpleElement<TSimpleElementPixelShader, SE_BLEND_Opaque> >(GetGlobalShaderMap(FeatureLevel));
+			case SE_BLEND_Masked:
+				return *TShaderMapRef<FEncodedSimpleElement<TSimpleElementPixelShader, SE_BLEND_Masked> >(GetGlobalShaderMap(FeatureLevel));
+			case SE_BLEND_Translucent:
+				return *TShaderMapRef<FEncodedSimpleElement<TSimpleElementPixelShader, SE_BLEND_Translucent> >(GetGlobalShaderMap(FeatureLevel));
+			case SE_BLEND_Additive:
+				return *TShaderMapRef<FEncodedSimpleElement<TSimpleElementPixelShader, SE_BLEND_Additive> >(GetGlobalShaderMap(FeatureLevel));
+			case SE_BLEND_Modulate:
+				return *TShaderMapRef<FEncodedSimpleElement<TSimpleElementPixelShader, SE_BLEND_Modulate> >(GetGlobalShaderMap(FeatureLevel));
+			case SE_BLEND_MaskedDistanceField:
+				return *TShaderMapRef<FEncodedSimpleElement<TSimpleElementPixelShader, SE_BLEND_MaskedDistanceField> >(GetGlobalShaderMap(FeatureLevel));
+			case SE_BLEND_MaskedDistanceFieldShadowed:
+				return *TShaderMapRef<FEncodedSimpleElement<TSimpleElementPixelShader, SE_BLEND_MaskedDistanceFieldShadowed> >(GetGlobalShaderMap(FeatureLevel));
+			case SE_BLEND_AlphaComposite:
+				return *TShaderMapRef<FEncodedSimpleElement<TSimpleElementPixelShader, SE_BLEND_AlphaComposite> >(GetGlobalShaderMap(FeatureLevel));
+			case SE_BLEND_AlphaBlend:
+				return *TShaderMapRef<FEncodedSimpleElement<TSimpleElementPixelShader, SE_BLEND_AlphaBlend> >(GetGlobalShaderMap(FeatureLevel));
+			case SE_BLEND_TranslucentAlphaOnly:
+				return *TShaderMapRef<FEncodedSimpleElement<TSimpleElementPixelShader, SE_BLEND_TranslucentAlphaOnly> >(GetGlobalShaderMap(FeatureLevel));
+
+			default:
+				checkNoEntry();
+		}
+	}
+	return *TShaderMapRef<TSimpleElementPixelShader>(GetGlobalShaderMap(FeatureLevel));
+}
+
+static bool Is32BppHDREncoded(const FSceneView* View)
+{
+	if (View == nullptr || View->GetFeatureLevel() >= ERHIFeatureLevel::SM4)
+	{
+		return false;
+	}
+
+	static auto* MobileHDRCvar = IConsoleManager::Get().FindTConsoleVariableDataInt(TEXT("r.MobileHDR"));
+	bool bMobileHDR = MobileHDRCvar->GetValueOnRenderThread() == 1;
+
+	static auto* MobileHDR32bppModeCvar = IConsoleManager::Get().FindTConsoleVariableDataInt(TEXT("r.MobileHDR32bppMode"));
+	bool bMobileHDR32Bpp = bMobileHDR && (GSupportsRenderTargetFormat_PF_FloatRGBA == false || MobileHDR32bppModeCvar->GetValueOnRenderThread() != 0);
+
+	if (bMobileHDR32Bpp)
+	{
+		switch (MobileHDR32bppModeCvar->GetValueOnRenderThread())
+		{
+		case 1:
+			return false;
+		case 2:
+			return true;
+		default:
+			return (GSupportsHDR32bppEncodeModeIntrinsic && GSupportsShaderFramebufferFetch);
+		}
+	}
+	return false;
+}
+
 /**
  * Sets the appropriate vertex and pixel shader.
  */
@@ -398,7 +468,10 @@ void FBatchedElements::PrepareShaders(
 {
 	// used to mask individual channels and desaturate
 	FMatrix ColorWeights( FPlane(1, 0, 0, 0), FPlane(0, 1, 0, 0), FPlane(0, 0, 1, 0), FPlane(0, 0, 0, 0) );
-	
+
+	// bEncodedHDR requires that blend states are disabled.
+	bool bEncodedHDR = Is32BppHDREncoded(View);
+
 	float GammaToUse = Gamma;
 	if(BlendMode >= SE_BLEND_RGBA_MASK_START && BlendMode <= SE_BLEND_RGBA_MASK_END)
 	{
@@ -427,13 +500,13 @@ void FBatchedElements::PrepareShaders(
 		// If we are only to draw the alpha channel, make the Blend state opaque, to allow easy identification of the alpha values
 		if( bAlphaOnly )
 		{
-			SetBlendState(RHICmdList, SE_BLEND_Opaque);
+			SetBlendState(RHICmdList, SE_BLEND_Opaque, bEncodedHDR);
 
 			R.W = G.W = B.W = 1.0f;
 		}
 		else
 		{
-			SetBlendState(RHICmdList, !bAlphaChannel ? SE_BLEND_Opaque : SE_BLEND_Translucent); // If alpha channel is disabled, do not allow alpha blending
+			SetBlendState(RHICmdList, !bAlphaChannel ? SE_BLEND_Opaque : SE_BLEND_Translucent, bEncodedHDR); // If alpha channel is disabled, do not allow alpha blending
 
 			// Determine the red, green, blue and alpha components of their respective weights to enable that colours prominence
 			R.X = bRedChannel ? 1.0f : 0.0f;
@@ -472,7 +545,7 @@ void FBatchedElements::PrepareShaders(
 		if (bHitTesting)
 		{
 			TShaderMapRef<FSimpleElementHitProxyPS> HitTestingPixelShader(GetGlobalShaderMap(FeatureLevel));
-			SetGlobalBoundShaderState(RHICmdList, FeatureLevel, HitTestingBoundShaderState, GSimpleElementVertexDeclaration.VertexDeclarationRHI,
+			SetGlobalBoundShaderState(RHICmdList, FeatureLevel, HitTestingBoundShaderState.GetBSS(bEncodedHDR, BlendMode), GSimpleElementVertexDeclaration.VertexDeclarationRHI,
 				*VertexShader, *HitTestingPixelShader);
 
 			HitTestingPixelShader->SetParameters(RHICmdList, Texture);
@@ -489,18 +562,18 @@ void FBatchedElements::PrepareShaders(
 
 				if (Texture->bSRGB)
 				{
-					TShaderMapRef<FSimpleElementMaskedGammaPS_SRGB> MaskedPixelShader(GetGlobalShaderMap(FeatureLevel));
-					SetGlobalBoundShaderState(RHICmdList, FeatureLevel, MaskedBoundShaderState, GSimpleElementVertexDeclaration.VertexDeclarationRHI,
-						*VertexShader, *MaskedPixelShader);
+					auto* MaskedPixelShader = GetPixelShader<FSimpleElementMaskedGammaPS_SRGB>(bEncodedHDR, BlendMode, FeatureLevel);
+					SetGlobalBoundShaderState(RHICmdList, FeatureLevel, MaskedBoundShaderState.GetBSS(bEncodedHDR, BlendMode), GSimpleElementVertexDeclaration.VertexDeclarationRHI,
+						*VertexShader, MaskedPixelShader);
 
 					MaskedPixelShader->SetEditorCompositingParameters(RHICmdList, View, DepthTexture);
 					MaskedPixelShader->SetParameters(RHICmdList, Texture, Gamma, GBatchedElementAlphaRefVal / 255.0f, BlendMode);
 				}
 				else
 				{
-					TShaderMapRef<FSimpleElementMaskedGammaPS_Linear> MaskedPixelShader(GetGlobalShaderMap(FeatureLevel));
-					SetGlobalBoundShaderState(RHICmdList, FeatureLevel, MaskedBoundShaderState, GSimpleElementVertexDeclaration.VertexDeclarationRHI,
-						*VertexShader, *MaskedPixelShader);
+					auto* MaskedPixelShader = GetPixelShader<FSimpleElementMaskedGammaPS_Linear>(bEncodedHDR, BlendMode, FeatureLevel);
+					SetGlobalBoundShaderState(RHICmdList, FeatureLevel, MaskedBoundShaderState.GetBSS(bEncodedHDR, BlendMode), GSimpleElementVertexDeclaration.VertexDeclarationRHI,
+						*VertexShader, MaskedPixelShader);
 
 					MaskedPixelShader->SetEditorCompositingParameters(RHICmdList, View, DepthTexture);
 					MaskedPixelShader->SetParameters(RHICmdList, Texture, Gamma, GBatchedElementAlphaRefVal / 255.0f, BlendMode);
@@ -518,7 +591,10 @@ void FBatchedElements::PrepareShaders(
 					BlendMode == SE_BLEND_TranslucentDistanceFieldShadowed)
 				{
 					// enable alpha blending and disable clip ref value for translucent rendering
-					RHICmdList.SetBlendState(TStaticBlendState<CW_RGB, BO_Add, BF_SourceAlpha, BF_InverseSourceAlpha>::GetRHI());
+					if (!bEncodedHDR)
+					{
+						RHICmdList.SetBlendState(TStaticBlendState<CW_RGB, BO_Add, BF_SourceAlpha, BF_InverseSourceAlpha>::GetRHI());
+					}
 					AlphaRefVal = 0.0f;
 				}
 				else
@@ -528,7 +604,7 @@ void FBatchedElements::PrepareShaders(
 				}
 				
 				TShaderMapRef<FSimpleElementDistanceFieldGammaPS> DistanceFieldPixelShader(GetGlobalShaderMap(FeatureLevel));
-				SetGlobalBoundShaderState(RHICmdList, FeatureLevel, DistanceFieldBoundShaderState, GSimpleElementVertexDeclaration.VertexDeclarationRHI,
+				SetGlobalBoundShaderState(RHICmdList, FeatureLevel, DistanceFieldBoundShaderState.GetBSS(bEncodedHDR, BlendMode), GSimpleElementVertexDeclaration.VertexDeclarationRHI,
 					*VertexShader, *DistanceFieldPixelShader );			
 
 				// @todo - expose these as options for batch rendering
@@ -557,12 +633,12 @@ void FBatchedElements::PrepareShaders(
 			}
 			else if(BlendMode == SE_BLEND_TranslucentAlphaOnly)
 			{
-				SetBlendState(RHICmdList, BlendMode);
+				SetBlendState(RHICmdList, BlendMode, bEncodedHDR);
 
 				if (FMath::Abs(Gamma - 1.0f) < KINDA_SMALL_NUMBER)
 				{
 					TShaderMapRef<FSimpleElementAlphaOnlyPS> AlphaOnlyPixelShader(GetGlobalShaderMap(FeatureLevel));
-					SetGlobalBoundShaderState(RHICmdList, FeatureLevel, AlphaOnlyShaderState, GSimpleElementVertexDeclaration.VertexDeclarationRHI,
+					SetGlobalBoundShaderState(RHICmdList, FeatureLevel, AlphaOnlyShaderState.GetBSS(bEncodedHDR, BlendMode), GSimpleElementVertexDeclaration.VertexDeclarationRHI,
 						*VertexShader, *AlphaOnlyPixelShader);
 
 					AlphaOnlyPixelShader->SetParameters(RHICmdList, Texture);
@@ -571,7 +647,7 @@ void FBatchedElements::PrepareShaders(
 				else
 				{
 					TShaderMapRef<FSimpleElementGammaAlphaOnlyPS> GammaAlphaOnlyPixelShader(GetGlobalShaderMap(FeatureLevel));
-					SetGlobalBoundShaderState(RHICmdList, FeatureLevel, GammaAlphaOnlyShaderState, GSimpleElementVertexDeclaration.VertexDeclarationRHI,
+					SetGlobalBoundShaderState(RHICmdList, FeatureLevel, GammaAlphaOnlyShaderState.GetBSS(bEncodedHDR, BlendMode), GSimpleElementVertexDeclaration.VertexDeclarationRHI,
 						*VertexShader, *GammaAlphaOnlyPixelShader);
 
 					GammaAlphaOnlyPixelShader->SetParameters(RHICmdList, Texture, Gamma, BlendMode);
@@ -581,19 +657,19 @@ void FBatchedElements::PrepareShaders(
 			else if(BlendMode >= SE_BLEND_RGBA_MASK_START && BlendMode <= SE_BLEND_RGBA_MASK_END)
 			{
 				TShaderMapRef<FSimpleElementColorChannelMaskPS> ColorChannelMaskPixelShader(GetGlobalShaderMap(FeatureLevel));
-				SetGlobalBoundShaderState(RHICmdList, FeatureLevel, ColorChannelMaskShaderState, GSimpleElementVertexDeclaration.VertexDeclarationRHI,
+				SetGlobalBoundShaderState(RHICmdList, FeatureLevel, ColorChannelMaskShaderState.GetBSS(bEncodedHDR, BlendMode), GSimpleElementVertexDeclaration.VertexDeclarationRHI,
 					*VertexShader, *ColorChannelMaskPixelShader );
 			
 				ColorChannelMaskPixelShader->SetParameters(RHICmdList, Texture, ColorWeights, GammaToUse );
 			}
 			else
 			{
-				SetBlendState(RHICmdList, BlendMode);
+				SetBlendState(RHICmdList, BlendMode, bEncodedHDR);
 	
 				if (FMath::Abs(Gamma - 1.0f) < KINDA_SMALL_NUMBER)
 				{
 					TShaderMapRef<FSimpleElementPS> PixelShader(GetGlobalShaderMap(FeatureLevel));
-					SetGlobalBoundShaderState(RHICmdList, FeatureLevel, SimpleBoundShaderState, GSimpleElementVertexDeclaration.VertexDeclarationRHI,
+					SetGlobalBoundShaderState(RHICmdList, FeatureLevel, SimpleBoundShaderState.GetBSS(bEncodedHDR, BlendMode), GSimpleElementVertexDeclaration.VertexDeclarationRHI,
 						*VertexShader, *PixelShader);
 
 					PixelShader->SetParameters(RHICmdList, Texture);
@@ -609,12 +685,12 @@ void FBatchedElements::PrepareShaders(
 					if (Texture->bSRGB)
 					{
 						BasePixelShader = *PixelShader_SRGB;
-						BoundShaderState = &RegularSRGBBoundShaderState;
+						BoundShaderState = &RegularSRGBBoundShaderState.GetBSS(bEncodedHDR, BlendMode);
 					}
 					else
 					{
 						BasePixelShader = *PixelShader_Linear;
-						BoundShaderState = &RegularLinearBoundShaderState;
+						BoundShaderState = &RegularLinearBoundShaderState.GetBSS(bEncodedHDR, BlendMode);
 					}
 
 					SetGlobalBoundShaderState(RHICmdList, FeatureLevel, *BoundShaderState, GSimpleElementVertexDeclaration.VertexDeclarationRHI,
diff --git a/Engine/Source/Runtime/Engine/Private/SimpleElementShaders.cpp b/Engine/Source/Runtime/Engine/Private/SimpleElementShaders.cpp
--- a/Engine/Source/Runtime/Engine/Private/SimpleElementShaders.cpp
+++ b/Engine/Source/Runtime/Engine/Private/SimpleElementShaders.cpp
@@ -7,6 +7,7 @@
 #include "EnginePrivate.h"
 #include "ShaderParameterUtils.h"
 #include "SimpleElementShaders.h"
+#include "SceneTypes.h"
 
 /*------------------------------------------------------------------------------
 	Simple element vertex shader.
@@ -307,3 +308,34 @@ IMPLEMENT_SHADER_TYPE(,FSimpleElementDistanceFieldGammaPS,TEXT("SimpleElementPix
 IMPLEMENT_SHADER_TYPE(,FSimpleElementHitProxyPS,TEXT("SimpleElementHitProxyPixelShader"),TEXT("Main"),SF_Pixel);
 IMPLEMENT_SHADER_TYPE(,FSimpleElementColorChannelMaskPS,TEXT("SimpleElementColorChannelMaskPixelShader"),TEXT("Main"),SF_Pixel);
 
+// 32 Bpp HDR Encoded implementations of simple element shaders.
+#define BLEND_VARIATION(SHADERCLASS,VARIATION, SHADERFILENAME, SHADERENTRYFUNC)\
+	typedef FEncodedSimpleElement<SHADERCLASS, VARIATION> SHADERCLASS##Encoded##VARIATION; \
+	IMPLEMENT_SHADER_TYPE(template<>, SHADERCLASS##Encoded##VARIATION, SHADERFILENAME, SHADERENTRYFUNC, SF_Pixel);
+
+#define IMPLEMENT_ENCODEDSHADERS(SHADERCLASS, SHADERFILENAME, SHADERENTRYFUNC) \
+	BLEND_VARIATION(SHADERCLASS, SE_BLEND_Opaque, SHADERFILENAME, SHADERENTRYFUNC) \
+	BLEND_VARIATION(SHADERCLASS, SE_BLEND_Masked, SHADERFILENAME, SHADERENTRYFUNC) \
+	BLEND_VARIATION(SHADERCLASS, SE_BLEND_Translucent, SHADERFILENAME, SHADERENTRYFUNC) \
+	BLEND_VARIATION(SHADERCLASS, SE_BLEND_Additive, SHADERFILENAME, SHADERENTRYFUNC) \
+	BLEND_VARIATION(SHADERCLASS, SE_BLEND_Modulate, SHADERFILENAME, SHADERENTRYFUNC) \
+	BLEND_VARIATION(SHADERCLASS, SE_BLEND_MaskedDistanceField, SHADERFILENAME, SHADERENTRYFUNC) \
+	BLEND_VARIATION(SHADERCLASS, SE_BLEND_MaskedDistanceFieldShadowed, SHADERFILENAME, SHADERENTRYFUNC) \
+	BLEND_VARIATION(SHADERCLASS, SE_BLEND_TranslucentDistanceField, SHADERFILENAME, SHADERENTRYFUNC) \
+	BLEND_VARIATION(SHADERCLASS, SE_BLEND_AlphaComposite, SHADERFILENAME, SHADERENTRYFUNC) \
+	BLEND_VARIATION(SHADERCLASS, SE_BLEND_TranslucentDistanceFieldShadowed, SHADERFILENAME, SHADERENTRYFUNC) \
+	BLEND_VARIATION(SHADERCLASS, SE_BLEND_AlphaBlend, SHADERFILENAME, SHADERENTRYFUNC) \
+	BLEND_VARIATION(SHADERCLASS, SE_BLEND_TranslucentAlphaOnly, SHADERFILENAME, SHADERENTRYFUNC)
+
+IMPLEMENT_ENCODEDSHADERS(FSimpleElementPS, TEXT("SimpleElementPixelShader"), TEXT("Main"));
+IMPLEMENT_ENCODEDSHADERS(FSimpleElementAlphaOnlyPS, TEXT("SimpleElementPixelShader"), TEXT("AlphaOnlyMain"));
+IMPLEMENT_ENCODEDSHADERS(FSimpleElementGammaPS_SRGB, TEXT("SimpleElementPixelShader"), TEXT("GammaMain"));
+IMPLEMENT_ENCODEDSHADERS(FSimpleElementGammaPS_Linear, TEXT("SimpleElementPixelShader"), TEXT("GammaMain"));
+IMPLEMENT_ENCODEDSHADERS(FSimpleElementGammaAlphaOnlyPS, TEXT("SimpleElementPixelShader"), TEXT("GammaAlphaOnlyMain"));
+IMPLEMENT_ENCODEDSHADERS(FSimpleElementMaskedGammaPS_SRGB, TEXT("SimpleElementPixelShader"), TEXT("GammaMaskedMain"));
+IMPLEMENT_ENCODEDSHADERS(FSimpleElementMaskedGammaPS_Linear, TEXT("SimpleElementPixelShader"), TEXT("GammaMaskedMain"));
+IMPLEMENT_ENCODEDSHADERS(FSimpleElementDistanceFieldGammaPS, TEXT("SimpleElementPixelShader"), TEXT("GammaDistanceFieldMain"));
+IMPLEMENT_ENCODEDSHADERS(FSimpleElementColorChannelMaskPS, TEXT("SimpleElementColorChannelMaskPixelShader"), TEXT("Main"));
+
+#undef IMPLEMENT_ENCODEDSHADERS
+#undef BLEND_VARIATION
diff --git a/Engine/Source/Runtime/Engine/Public/BatchedElements.h b/Engine/Source/Runtime/Engine/Public/BatchedElements.h
--- a/Engine/Source/Runtime/Engine/Public/BatchedElements.h
+++ b/Engine/Source/Runtime/Engine/Public/BatchedElements.h
@@ -296,23 +296,36 @@ private:
 	TArray<FBatchedQuadMeshElement> QuadMeshElements;
 
 	/** bound shader state for the fast path */
-	static FGlobalBoundShaderState SimpleBoundShaderState;
+	class FSimpleElementBSSContainer
+	{
+		static const uint32 NumBSS = (uint32)SE_BLEND_RGBA_MASK_START;
+		FGlobalBoundShaderState UnencodedBSS;
+		FGlobalBoundShaderState EncodedBSS[NumBSS];
+	public:
+		FGlobalBoundShaderState& GetBSS(bool bEncoded, ESimpleElementBlendMode BlendMode)
+		{
+			check((uint32)BlendMode < NumBSS);
+			return bEncoded ? EncodedBSS[BlendMode] : UnencodedBSS;
+		}
+	};
+
+	static FSimpleElementBSSContainer SimpleBoundShaderState;
 	/** bound shader state for the regular mesh elements with a linear texture */
-	static FGlobalBoundShaderState RegularLinearBoundShaderState;
+	static FSimpleElementBSSContainer RegularLinearBoundShaderState;
 	/** bound shader state for the regular mesh elements with an sRGB texture */
-	static FGlobalBoundShaderState RegularSRGBBoundShaderState;
+	static FSimpleElementBSSContainer RegularSRGBBoundShaderState;
 	/** bound shader state for masked mesh elements */
-	static FGlobalBoundShaderState MaskedBoundShaderState;
+	static FSimpleElementBSSContainer MaskedBoundShaderState;
 	/** bound shader state for masked mesh elements */
-	static FGlobalBoundShaderState DistanceFieldBoundShaderState;
+	static FSimpleElementBSSContainer DistanceFieldBoundShaderState;
 	/** bound shader state for the hit testing mesh elements */
-	static FGlobalBoundShaderState HitTestingBoundShaderState;
+	static FSimpleElementBSSContainer HitTestingBoundShaderState;
 	/** bound shader state for color masked elements */
-	static FGlobalBoundShaderState ColorChannelMaskShaderState;
+	static FSimpleElementBSSContainer ColorChannelMaskShaderState;
 	/** bound shader state for alpha only fonts */
-	static FGlobalBoundShaderState AlphaOnlyShaderState;
+	static FSimpleElementBSSContainer AlphaOnlyShaderState;
 	/** bound shader state for gamma corrected alpha only fonts */
-	static FGlobalBoundShaderState GammaAlphaOnlyShaderState;
+	static FSimpleElementBSSContainer GammaAlphaOnlyShaderState;
 
 	/**
 	 * Sets the appropriate vertex and pixel shader.
diff --git a/Engine/Source/Runtime/Engine/Public/SimpleElementShaders.h b/Engine/Source/Runtime/Engine/Public/SimpleElementShaders.h
--- a/Engine/Source/Runtime/Engine/Public/SimpleElementShaders.h
+++ b/Engine/Source/Runtime/Engine/Public/SimpleElementShaders.h
@@ -317,6 +317,39 @@ private:
 	FShaderParameter Gamma;
 };
 
+template <class TSimpleElementBase, uint32 EncodedBlendMode>
+class FEncodedSimpleElement : public TSimpleElementBase
+{
+	DECLARE_SHADER_TYPE(FEncodedSimpleElement, Global);
+public:
+	FEncodedSimpleElement(const ShaderMetaType::CompiledShaderInitializerType& Initializer) : TSimpleElementBase(Initializer) {}
+	FEncodedSimpleElement() {}
+
+	static void ModifyCompilationEnvironment(EShaderPlatform Platform, FShaderCompilerEnvironment& OutEnvironment)
+	{
+		TSimpleElementBase::ModifyCompilationEnvironment(Platform, OutEnvironment); 
+		OutEnvironment.SetDefine(TEXT("SE_BLEND_OPAQUE"), (uint32)ESimpleElementBlendMode::SE_BLEND_Opaque);
+		OutEnvironment.SetDefine(TEXT("SE_BLEND_MASKED"), (uint32)ESimpleElementBlendMode::SE_BLEND_Masked);
+		OutEnvironment.SetDefine(TEXT("SE_BLEND_TRANSLUCENT"), (uint32)ESimpleElementBlendMode::SE_BLEND_Translucent);
+		OutEnvironment.SetDefine(TEXT("SE_BLEND_ADDITIVE"), (uint32)ESimpleElementBlendMode::SE_BLEND_Additive);
+		OutEnvironment.SetDefine(TEXT("SE_BLEND_MODULATE"), (uint32)ESimpleElementBlendMode::SE_BLEND_Modulate);
+		OutEnvironment.SetDefine(TEXT("SE_BLEND_MASKEDDISTANCEFIELD"), (uint32)ESimpleElementBlendMode::SE_BLEND_MaskedDistanceField);
+		OutEnvironment.SetDefine(TEXT("SE_BLEND_MASKEDDISTANCEFIELDSHADOWED"), (uint32)ESimpleElementBlendMode::SE_BLEND_MaskedDistanceFieldShadowed);
+		OutEnvironment.SetDefine(TEXT("SE_BLEND_ALPHACOMPOSITE"), (uint32)ESimpleElementBlendMode::SE_BLEND_AlphaComposite);
+		OutEnvironment.SetDefine(TEXT("SE_BLEND_ALPHABLEND"), (uint32)ESimpleElementBlendMode::SE_BLEND_AlphaBlend);
+		OutEnvironment.SetDefine(TEXT("SE_BLEND_TRANSLUCENTALPHAONLY"), (uint32)ESimpleElementBlendMode::SE_BLEND_TranslucentAlphaOnly);
+		OutEnvironment.SetDefine(TEXT("SE_BLEND_TRANSLUCENTDISTANCEFIELD"), (uint32)ESimpleElementBlendMode::SE_BLEND_TranslucentDistanceField);
+		OutEnvironment.SetDefine(TEXT("SE_BLEND_TRANSLUCENTDISTANCEFIELDSHADOWED"), (uint32)ESimpleElementBlendMode::SE_BLEND_TranslucentDistanceFieldShadowed);
+		OutEnvironment.SetDefine(TEXT("SE_BLEND_MODE"), EncodedBlendMode);
+		OutEnvironment.SetDefine(TEXT("USE_32BPP_HDR"), 1u);
+	}
+
+	static bool ShouldCache(EShaderPlatform Platform)
+	{
+		return ((IsES2Platform(Platform) && IsPCPlatform(Platform)) || Platform == SP_OPENGL_ES2) && TSimpleElementBase::ShouldCache(Platform);
+	}
+};
+
 typedef FSimpleElementGammaPS<true> FSimpleElementGammaPS_SRGB;
 typedef FSimpleElementGammaPS<false> FSimpleElementGammaPS_Linear;
 typedef FSimpleElementMaskedGammaPS<true> FSimpleElementMaskedGammaPS_SRGB;
