Fix for minor issue in which distortion would not render for Opaqe materials which were overriden to be translucent.
Relevancy calculations were checking the base material's resource rather than the static permutation resource of the instance being rendered.

diff --git a/Engine/Source/Runtime/Engine/Private/Materials/MaterialInterface.cpp b/Engine/Source/Runtime/Engine/Private/Materials/MaterialInterface.cpp
--- a/Engine/Source/Runtime/Engine/Private/Materials/MaterialInterface.cpp
+++ b/Engine/Source/Runtime/Engine/Private/Materials/MaterialInterface.cpp
@@ -53,7 +53,7 @@ FMaterialRelevance UMaterialInterface::GetRelevance_Internal(const UMaterial* Ma
 {
 	if(Material)
 	{
-		const FMaterialResource* MaterialResource = Material->GetMaterialResource(InFeatureLevel);
+		const FMaterialResource* MaterialResource = GetMaterialResource(InFeatureLevel);
 		const bool bIsTranslucent = IsTranslucentBlendMode((EBlendMode)GetBlendMode());
 
 		EMaterialShadingModel ShadingModel = GetShadingModel();
