Make enum to string conversion in blueprints use enum value instead of index.

diff --git a/Engine/Source/Editor/BlueprintGraph/Private/K2Node_GetEnumeratorName.cpp b/Engine/Source/Editor/BlueprintGraph/Private/K2Node_GetEnumeratorName.cpp
--- a/Engine/Source/Editor/BlueprintGraph/Private/K2Node_GetEnumeratorName.cpp
+++ b/Engine/Source/Editor/BlueprintGraph/Private/K2Node_GetEnumeratorName.cpp
@@ -148,9 +148,9 @@ void UK2Node_GetEnumeratorName::ExpandNode(class FKismetCompilerContext& Compile
 	Schema->TrySetDefaultObject(*EnumPin, Enum);
 	check(EnumPin->DefaultObject == Enum);
 
-	//INDEX PIN
+	//VALUE PIN
 	UEdGraphPin* OrgInputPin = FindPinChecked(EnumeratorPinName);
-	UEdGraphPin* IndexPin = CallGetName->FindPinChecked(TEXT("EnumeratorIndex"));
+	UEdGraphPin* IndexPin = CallGetName->FindPinChecked(TEXT("EnumeratorValue"));
 	check(EGPD_Input == IndexPin->Direction && Schema->PC_Byte == IndexPin->PinType.PinCategory);
 	CompilerContext.MovePinLinksToIntermediate(*OrgInputPin, *IndexPin);
 
diff --git a/Engine/Source/Runtime/Engine/Classes/Kismet/KismetNodeHelperLibrary.h b/Engine/Source/Runtime/Engine/Classes/Kismet/KismetNodeHelperLibrary.h
--- a/Engine/Source/Runtime/Engine/Classes/Kismet/KismetNodeHelperLibrary.h
+++ b/Engine/Source/Runtime/Engine/Classes/Kismet/KismetNodeHelperLibrary.h
@@ -111,21 +111,21 @@ class ENGINE_API UKismetNodeHelperLibrary : public UBlueprintFunctionLibrary
 	 * Gets enumerator name.
 	 *
 	 * @param Enum - Enumeration
-	 * @param EnumeratorIndex - Index of searched enumeration
+	 * @param EnumeratorValue - Value of searched enumeration
 	 * @return - name of the searched enumerator, or NAME_None
 	 */
 	UFUNCTION(BlueprintPure, meta=(BlueprintInternalUseOnly = "TRUE"))
-	static FName GetEnumeratorName(const UEnum* Enum, uint8 EnumeratorIndex);
+	static FName GetEnumeratorName(const UEnum* Enum, uint8 EnumeratorValue);
 
 	/**
 	 * Gets enumerator name as FString. Use DeisplayName when possible.
 	 *
 	 * @param Enum - Enumeration
-	 * @param EnumeratorIndex - Index of searched enumeration
+	 * @param EnumeratorValue - Value of searched enumeration
 	 * @return - name of the searched enumerator, or NAME_None
 	 */
 	UFUNCTION(BlueprintPure, meta=(BlueprintInternalUseOnly = "TRUE"))
-	static FString GetEnumeratorUserFriendlyName(const UEnum* Enum, uint8 EnumeratorIndex);
+	static FString GetEnumeratorUserFriendlyName(const UEnum* Enum, uint8 EnumeratorValue);
 
 	/**
 	 * @param Enum - Enumeration
diff --git a/Engine/Source/Runtime/Engine/Private/KismetNodeHelperLibrary.cpp b/Engine/Source/Runtime/Engine/Private/KismetNodeHelperLibrary.cpp
--- a/Engine/Source/Runtime/Engine/Private/KismetNodeHelperLibrary.cpp
+++ b/Engine/Source/Runtime/Engine/Private/KismetNodeHelperLibrary.cpp
@@ -141,15 +141,17 @@ int32 UKismetNodeHelperLibrary::GetFirstUnmarkedBit(int32 Data, int32 StartIdx,
 	return INDEX_NONE;
 }
 
-FName UKismetNodeHelperLibrary::GetEnumeratorName(const UEnum* Enum, uint8 EnumeratorIndex)
+FName UKismetNodeHelperLibrary::GetEnumeratorName(const UEnum* Enum, uint8 EnumeratorValue)
 {
+	int32 EnumeratorIndex = Enum->GetIndexByValue(EnumeratorValue);
 	return (NULL != Enum) ? Enum->GetEnum(EnumeratorIndex) : NAME_None;
 }
 
-FString UKismetNodeHelperLibrary::GetEnumeratorUserFriendlyName(const UEnum* Enum, uint8 EnumeratorIndex)
+FString UKismetNodeHelperLibrary::GetEnumeratorUserFriendlyName(const UEnum* Enum, uint8 EnumeratorValue)
 {
 	if (NULL != Enum)
 	{
+		int32 EnumeratorIndex = Enum->GetIndexByValue(EnumeratorValue);
 		return Enum->GetEnumText(EnumeratorIndex).ToString();
 	}
 
