Disable parallel evaluation when paused as we are not supposed to spawn new tasks during pause frames (crash fix)

diff --git a/Engine/Source/Runtime/Engine/Private/Components/SkeletalMeshComponent.cpp b/Engine/Source/Runtime/Engine/Private/Components/SkeletalMeshComponent.cpp
--- a/Engine/Source/Runtime/Engine/Private/Components/SkeletalMeshComponent.cpp
+++ b/Engine/Source/Runtime/Engine/Private/Components/SkeletalMeshComponent.cpp
@@ -1118,7 +1118,7 @@ void USkeletalMeshComponent::RefreshBoneTransforms(FActorComponentTickFunction*
 
 	const bool bShouldDoEvaluation = !bDoEvaluationRateOptimization || bInvalidCachedBones || !AnimUpdateRateParams->ShouldSkipEvaluation();
 
-	const bool bDoPAE = !!CVarUseParallelAnimationEvaluation.GetValueOnGameThread() && FApp::ShouldUseThreadingForPerformance();
+	const bool bDoPAE = !!CVarUseParallelAnimationEvaluation.GetValueOnGameThread() && FApp::ShouldUseThreadingForPerformance() && !GetWorld()->IsPaused();
 
 	const bool bDoParallelEvaluation = bShouldDoEvaluation && TickFunction && bDoPAE;
 
