#merge, replays: Smooth on all axes for replays even if on ground

diff --git a/Engine/Source/Runtime/Engine/Private/Components/CharacterMovementComponent.cpp b/Engine/Source/Runtime/Engine/Private/Components/CharacterMovementComponent.cpp
--- a/Engine/Source/Runtime/Engine/Private/Components/CharacterMovementComponent.cpp
+++ b/Engine/Source/Runtime/Engine/Private/Components/CharacterMovementComponent.cpp
@@ -5778,12 +5778,12 @@ void UCharacterMovementComponent::SmoothClientPosition(float DeltaSeconds)
 			{
 				ClientData->MeshRotationOffset = FQuat::Identity;
 			}
-		}
 
-		if (IsMovingOnGround())
-		{
-			// don't smooth Z position if walking on ground
-			ClientData->MeshTranslationOffset.Z = 0;
+			if (IsMovingOnGround())
+			{
+				// don't smooth Z position if walking on ground
+				ClientData->MeshTranslationOffset.Z = 0;
+			}
 		}
 
 		const FVector NewRelTranslation = UpdatedComponent->GetComponentToWorld().InverseTransformVectorNoScale(ClientData->MeshTranslationOffset) + CharacterOwner->GetBaseTranslationOffset();
