Fix invalid scale spam when scaling bodyinstance. This is really just clamping to make sure that we never get 0 as animation uses size 0 all the time

diff --git a/Engine/Source/Runtime/Engine/Private/PhysicsEngine/BodyInstance.cpp b/Engine/Source/Runtime/Engine/Private/PhysicsEngine/BodyInstance.cpp
--- a/Engine/Source/Runtime/Engine/Private/PhysicsEngine/BodyInstance.cpp
+++ b/Engine/Source/Runtime/Engine/Private/PhysicsEngine/BodyInstance.cpp
@@ -2019,8 +2019,9 @@ namespace EScaleMode
 }
 
 //computes the relative scaling vectors based on scale mode used
-void ComputeScalingVectors(EScaleMode::Type ScaleMode, const FVector& NewScale3D, FVector& OutScale3D, FVector& OutScale3DAbs)
+void ComputeScalingVectors(EScaleMode::Type ScaleMode, const FVector& InScale3D, FVector& OutScale3D, FVector& OutScale3DAbs)
 {
+	const FVector NewScale3D = InScale3D.IsNearlyZero() ? FVector(KINDA_SMALL_NUMBER) : InScale3D;	//min scale
 	const FVector NewScale3DAbs = NewScale3D.GetAbs();
 	switch (ScaleMode)
 	{
