Better fix for translation being in the rotation matrix.

diff --git a/Engine/Source/Editor/UnrealEd/Private/EditorViewportClient.cpp b/Engine/Source/Editor/UnrealEd/Private/EditorViewportClient.cpp
--- a/Engine/Source/Editor/UnrealEd/Private/EditorViewportClient.cpp
+++ b/Engine/Source/Editor/UnrealEd/Private/EditorViewportClient.cpp
@@ -656,7 +656,7 @@ FSceneView* FEditorViewportClient::CalcSceneView(FSceneViewFamily* ViewFamily)
 		{
 			if (bUsingOrbitCamera)
 			{
-				ViewInitOptions.ViewRotationMatrix = ViewTransform.ComputeOrbitMatrix().RemoveTranslation();
+				ViewInitOptions.ViewRotationMatrix = FTranslationMatrix(ViewLocation) * ViewTransform.ComputeOrbitMatrix();
 			}
 			else
 			{
@@ -807,9 +807,6 @@ FSceneView* FEditorViewportClient::CalcSceneView(FSceneViewFamily* ViewFamily)
 				check(false);
 			}
 
-			ViewInitOptions.ViewOrigin += ViewInitOptions.ViewRotationMatrix.InverseTransformPosition( FVector::ZeroVector );
-			ViewInitOptions.ViewRotationMatrix = ViewInitOptions.ViewRotationMatrix.RemoveTranslation();
-
 			ViewInitOptions.ProjectionMatrix = FReversedZOrthoMatrix(
 				OrthoWidth,
 				OrthoHeight,
diff --git a/Engine/Source/Runtime/Engine/Private/SceneView.cpp b/Engine/Source/Runtime/Engine/Private/SceneView.cpp
--- a/Engine/Source/Runtime/Engine/Private/SceneView.cpp
+++ b/Engine/Source/Runtime/Engine/Private/SceneView.cpp
@@ -302,9 +302,16 @@ FSceneView::FSceneView(const FSceneViewInitOptions& InitOptions)
 	check(UnscaledViewRect.Min.Y >= 0);
 	check(UnscaledViewRect.Width() > 0);
 	check(UnscaledViewRect.Height() > 0);
-	check(InitOptions.ViewRotationMatrix.GetOrigin().IsNearlyZero());
 
-	ViewMatrices.ViewMatrix = FTranslationMatrix(-InitOptions.ViewOrigin) * InitOptions.ViewRotationMatrix;
+	FVector ViewOrigin = InitOptions.ViewOrigin;
+	FMatrix ViewRotationMatrix = InitOptions.ViewRotationMatrix;
+	if( !ViewRotationMatrix.GetOrigin().IsNearlyZero( 0.0f ) )
+	{
+		ViewOrigin += ViewRotationMatrix.InverseTransformPosition( FVector::ZeroVector );
+		ViewRotationMatrix = ViewRotationMatrix.RemoveTranslation();
+	}
+
+	ViewMatrices.ViewMatrix = FTranslationMatrix(-ViewOrigin) * ViewRotationMatrix;
 
 	// Adjust the projection matrix for the current RHI.
 	ViewMatrices.ProjMatrix = AdjustProjectionMatrixForRHI(ProjectionMatrixUnadjustedForRHI);
@@ -315,7 +322,7 @@ FSceneView::FSceneView(const FSceneViewInitOptions& InitOptions)
 	FMatrix InvProjectionMatrix = ViewMatrices.GetInvProjMatrix();
 
 	// For precision reasons the view matrix inverse is calculated independently.
-	InvViewMatrix = InitOptions.ViewRotationMatrix.GetTransposed() * FTranslationMatrix(InitOptions.ViewOrigin);
+	InvViewMatrix = ViewRotationMatrix.GetTransposed() * FTranslationMatrix(ViewOrigin);
 	InvViewProjectionMatrix = InvProjectionMatrix * InvViewMatrix;
 
 	bool bApplyPreViewTranslation = true;
@@ -324,13 +331,13 @@ FSceneView::FSceneView(const FSceneViewInitOptions& InitOptions)
 	// Calculate the view origin from the view/projection matrices.
 	if(IsPerspectiveProjection())
 	{
-		ViewMatrices.ViewOrigin = InitOptions.ViewOrigin;
+		ViewMatrices.ViewOrigin = ViewOrigin;
 	}
 #if WITH_EDITOR
 	else if (InitOptions.bUseFauxOrthoViewPos)
 	{
 		float DistanceToViewOrigin = WORLD_MAX;
-		ViewMatrices.ViewOrigin = FVector4(InvViewMatrix.TransformVector(FVector(0,0,-1)).GetSafeNormal()*DistanceToViewOrigin,1) + InitOptions.ViewOrigin;
+		ViewMatrices.ViewOrigin = FVector4(InvViewMatrix.TransformVector(FVector(0,0,-1)).GetSafeNormal()*DistanceToViewOrigin,1) + ViewOrigin;
 		bViewOriginIsFudged = true;
 	}
 #endif
@@ -342,7 +349,7 @@ FSceneView::FSceneView(const FSceneViewInitOptions& InitOptions)
 	}
 
 	/** The view transform, starting from world-space points translated by -ViewOrigin. */
-	FMatrix TranslatedViewMatrix = InitOptions.ViewRotationMatrix;
+	FMatrix TranslatedViewMatrix = ViewRotationMatrix;
 	FMatrix InvTranslatedViewMatrix = TranslatedViewMatrix.GetTransposed();
 
 	// Translate world-space so its origin is at ViewOrigin for improved precision.
@@ -350,7 +357,7 @@ FSceneView::FSceneView(const FSceneViewInitOptions& InitOptions)
 	// in that case so the same value may be used in shaders for both the world-space translation and the camera's world position.
 	if(bApplyPreViewTranslation)
 	{
-		ViewMatrices.PreViewTranslation = -FVector(ViewMatrices.ViewOrigin);
+		ViewMatrices.PreViewTranslation = -FVector(ViewOrigin);
 
 #if !(UE_BUILD_SHIPPING || UE_BUILD_TEST)
 		{
@@ -382,7 +389,7 @@ FSceneView::FSceneView(const FSceneViewInitOptions& InitOptions)
 	if (bViewOriginIsFudged)
 	{
 		TranslatedViewMatrix = FTranslationMatrix(-ViewMatrices.PreViewTranslation)
-			* FTranslationMatrix(-InitOptions.ViewOrigin) * InitOptions.ViewRotationMatrix;
+			* FTranslationMatrix(-ViewOrigin) * ViewRotationMatrix;
 		InvTranslatedViewMatrix = TranslatedViewMatrix.Inverse();
 	}
 	
