Better fix for parallel evaluation crash when paused

diff --git a/Engine/Source/Runtime/Engine/Classes/Engine/EngineBaseTypes.h b/Engine/Source/Runtime/Engine/Classes/Engine/EngineBaseTypes.h
--- a/Engine/Source/Runtime/Engine/Classes/Engine/EngineBaseTypes.h
+++ b/Engine/Source/Runtime/Engine/Classes/Engine/EngineBaseTypes.h
@@ -235,7 +235,8 @@ public:
 	void SetTickFunctionEnable(bool bInEnabled);
 	/** Returns whether the tick function is currently enabled */
 	bool IsTickFunctionEnabled() const { return TickState != ETickState::Disabled; }
-
+	/** Returns whether it is valid to access this tick function's completion handle */
+	bool IsCompletionHandleValid() const { return TaskPointer != nullptr; }
 	/**
 	* Gets the current completion handle of this tick function, so it can be delayed until a later point when some additional
 	* tasks have been completed.  Only valid after TG_PreAsyncWork has started and then only until the TickFunction finishes
diff --git a/Engine/Source/Runtime/Engine/Private/Components/SkeletalMeshComponent.cpp b/Engine/Source/Runtime/Engine/Private/Components/SkeletalMeshComponent.cpp
--- a/Engine/Source/Runtime/Engine/Private/Components/SkeletalMeshComponent.cpp
+++ b/Engine/Source/Runtime/Engine/Private/Components/SkeletalMeshComponent.cpp
@@ -1118,9 +1118,9 @@ void USkeletalMeshComponent::RefreshBoneTransforms(FActorComponentTickFunction*
 
 	const bool bShouldDoEvaluation = !bDoEvaluationRateOptimization || bInvalidCachedBones || !AnimUpdateRateParams->ShouldSkipEvaluation();
 
-	const bool bDoPAE = !!CVarUseParallelAnimationEvaluation.GetValueOnGameThread() && FApp::ShouldUseThreadingForPerformance() && !GetWorld()->IsPaused();
+	const bool bDoPAE = !!CVarUseParallelAnimationEvaluation.GetValueOnGameThread() && FApp::ShouldUseThreadingForPerformance();
 
-	const bool bDoParallelEvaluation = bShouldDoEvaluation && TickFunction && bDoPAE;
+	const bool bDoParallelEvaluation = bDoPAE && bShouldDoEvaluation && TickFunction && TickFunction->IsCompletionHandleValid();
 
 	const bool bBlockOnTask = !bDoParallelEvaluation;  // If we aren't trying to do parallel evaluation then we
 															// will need to wait on an existing task.
