Fixed GetInvViewMatrix

diff --git a/Engine/Source/Runtime/Engine/Public/SceneView.h b/Engine/Source/Runtime/Engine/Public/SceneView.h
--- a/Engine/Source/Runtime/Engine/Public/SceneView.h
+++ b/Engine/Source/Runtime/Engine/Public/SceneView.h
@@ -241,7 +241,7 @@ struct FViewMatrices
 
 	FMatrix GetInvViewMatrix() const
 	{
-		return ViewMatrix.RemoveTranslation().GetTransposed() * FTranslationMatrix( ViewMatrix.GetOrigin() );
+		return FTranslationMatrix( -ViewMatrix.GetOrigin() ) * ViewMatrix.RemoveTranslation().GetTransposed();
 	}
 
 	FMatrix GetInvViewProjMatrix() const
