Change 2645290 by [REDACTED_PERSON][REDACTED_HANDLE] on 2015/08/05 14:46:25

Fixed InstancePreReplication() to include the object being replicated

Merging 2645290

//depot/UE4-UT/...

to //depot/UE4/...

diff --git a/Engine/Source/Runtime/Engine/Classes/Engine/BlueprintGeneratedClass.h b/Engine/Source/Runtime/Engine/Classes/Engine/BlueprintGeneratedClass.h
--- a/Engine/Source/Runtime/Engine/Classes/Engine/BlueprintGeneratedClass.h
+++ b/Engine/Source/Runtime/Engine/Classes/Engine/BlueprintGeneratedClass.h
@@ -503,12 +503,12 @@ public:
 	/** called to gather blueprint replicated properties */
 	virtual void GetLifetimeBlueprintReplicationList(TArray<class FLifetimeProperty>& OutLifetimeProps) const;
 	/** called prior to replication of an instance of this BP class */
-	virtual void InstancePreReplication(class IRepChangedPropertyTracker& ChangedPropertyTracker) const
+	virtual void InstancePreReplication(UObject* Obj, class IRepChangedPropertyTracker& ChangedPropertyTracker) const
 	{
 		UBlueprintGeneratedClass* SuperBPClass = Cast<UBlueprintGeneratedClass>(GetSuperStruct());
 		if (SuperBPClass != NULL)
 		{
-			SuperBPClass->InstancePreReplication(ChangedPropertyTracker);
+			SuperBPClass->InstancePreReplication(Obj, ChangedPropertyTracker);
 		}
 	}
 };
diff --git a/Engine/Source/Runtime/Engine/Private/Actor.cpp b/Engine/Source/Runtime/Engine/Private/Actor.cpp
--- a/Engine/Source/Runtime/Engine/Private/Actor.cpp
+++ b/Engine/Source/Runtime/Engine/Private/Actor.cpp
@@ -751,7 +751,7 @@ void AActor::PreReplication( IRepChangedPropertyTracker & ChangedPropertyTracker
 	UBlueprintGeneratedClass* BPClass = Cast<UBlueprintGeneratedClass>(GetClass());
 	if (BPClass != NULL)
 	{
-		BPClass->InstancePreReplication(ChangedPropertyTracker);
+		BPClass->InstancePreReplication(this, ChangedPropertyTracker);
 	}
 }
 
