Fix build break

diff --git a/Engine/Source/Runtime/Engine/Private/Components/ActorComponent.cpp b/Engine/Source/Runtime/Engine/Private/Components/ActorComponent.cpp
--- a/Engine/Source/Runtime/Engine/Private/Components/ActorComponent.cpp
+++ b/Engine/Source/Runtime/Engine/Private/Components/ActorComponent.cpp
@@ -1443,7 +1443,7 @@ void UActorComponent::PreReplication(IRepChangedPropertyTracker & ChangedPropert
 	UBlueprintGeneratedClass* BPClass = Cast<UBlueprintGeneratedClass>(GetClass());
 	if (BPClass != NULL)
 	{
-		BPClass->InstancePreReplication(ChangedPropertyTracker);
+		BPClass->InstancePreReplication(this, ChangedPropertyTracker);
 	}
 }
 
