Restoring fast view inverse now that it's correct.

diff --git a/Engine/Source/Runtime/Engine/Private/SceneView.cpp b/Engine/Source/Runtime/Engine/Private/SceneView.cpp
--- a/Engine/Source/Runtime/Engine/Private/SceneView.cpp
+++ b/Engine/Source/Runtime/Engine/Private/SceneView.cpp
@@ -582,9 +582,7 @@ void FSceneView::UpdateViewMatrix()
 	ViewMatrices.InvTranslatedViewProjectionMatrix = ViewMatrices.TranslatedViewProjectionMatrix.Inverse();
 
 	ViewProjectionMatrix = ViewMatrices.GetViewProjMatrix();
-
-	// Need to use a full view matrix inverse here as there are translations present
-	InvViewMatrix = ViewMatrices.ViewMatrix.Inverse();
+	InvViewMatrix = ViewMatrices.GetInvViewMatrix();
 	InvViewProjectionMatrix = ViewMatrices.GetInvProjMatrix() * InvViewMatrix;
 }
 
