Ping calculation fixes:

* Fix bug where extra acks were contributing to ping
* Call FlushNet on frames where SendAck was called
* Use wall time when calculating pings
* Stop subtracting half of frame time from ping

diff --git a/Engine/Source/Runtime/Engine/Private/NetConnection.cpp b/Engine/Source/Runtime/Engine/Private/NetConnection.cpp
--- a/Engine/Source/Runtime/Engine/Private/NetConnection.cpp
+++ b/Engine/Source/Runtime/Engine/Private/NetConnection.cpp
@@ -594,8 +594,11 @@ void UNetConnection::FlushNet(bool bIgnoreSimulation)
 #endif
 		// Update stuff.
 		const int32 Index = OutPacketId & (ARRAY_COUNT(OutLagPacketId)-1);
-		OutLagPacketId [Index] = OutPacketId;
-		OutLagTime     [Index] = Driver->Time;
+
+		// Remember the actual time this packet was sent out, so we can compute ping when the ack comes back
+		OutLagPacketId[Index]	= OutPacketId;
+		OutLagTime[Index]		= FPlatformTime::Seconds();
+
 		OutPacketId++;
 		Driver->OutPackets++;
 		LastSendTime = Driver->Time;
@@ -747,11 +750,6 @@ void UNetConnection::ReceivedPacket( FBitReader& Reader )
 				}
 			}
 
-
-			// For clientside monitoring of ping, watch out for duplicate acks (they come one tick later than the original ack),
-			//	because if you don't, the calculated ping value will be high by one tick
-			bool bPotentialDupeAck = false;
-
 			// Resend any old reliable packets that the receiver hasn't acknowledged.
 			if( AckPacketId>OutAckPacketId )
 			{
@@ -766,15 +764,16 @@ void UNetConnection::ReceivedPacket( FBitReader& Reader )
 			{
 				//warning: Double-ack logic makes this unmeasurable.
 				//OutOrdAcc++;
-
-				bPotentialDupeAck = true;
 			}
 
-			// Update lag.
-			int32 Index = AckPacketId & (ARRAY_COUNT(OutLagPacketId)-1);
-			if( OutLagPacketId[Index]==AckPacketId )
+			// Update ping
+			const int32 Index = AckPacketId & (ARRAY_COUNT(OutLagPacketId)-1);
+
+			if ( OutLagPacketId[Index] == AckPacketId )
 			{
-				float NewLag = Driver->Time - OutLagTime[Index] - (FrameTime/2.f);
+				OutLagPacketId[Index] = -1;		// Only use the ack once
+
+				const float NewLag = FPlatformTime::Seconds() - OutLagTime[Index];
 
 				LagAcc += NewLag;
 				LagCount++;
@@ -782,13 +781,9 @@ void UNetConnection::ReceivedPacket( FBitReader& Reader )
 				if (PlayerController != NULL)
 				{
 					// Verify PingAck's, and pass notification up
-					if (bPingAck && bHasPingAckData &&
-						PingAckDataCache[(AckPacketId % MAX_PACKETID)/PING_ACK_PACKET_INTERVAL] == InPingAckData)
-					{
-						PlayerController->UpdatePing(NewLag);
-					}
-					// For clients monitoring their own ping, trigger UpdatePing so long as this is not a duplicate ack
-					else if (Driver->ServerConnection != NULL && !bPotentialDupeAck)
+					const bool bServerVerified = bPingAck && bHasPingAckData && PingAckDataCache[( AckPacketId % MAX_PACKETID ) / PING_ACK_PACKET_INTERVAL] == InPingAckData;
+					
+					if ( bServerVerified || !Driver->IsServer() )
 					{
 						PlayerController->UpdatePing(NewLag);
 					}
@@ -1181,6 +1176,8 @@ void UNetConnection::SendAck( int32 AckPacketId, bool FirstTime/*=1*/, bool bHav
 
 		AllowMerge = false;
 
+		TimeSensitive = 1;
+
 		UE_LOG(LogNetTraffic, Log, TEXT("   Send ack %i"), AckPacketId);
 	}
 }
