[UE-13328] Fix potential crash on PIE for BP class instances w/ component object properties serialized as NULL.

change summary:
- added AActor::NativeConstructedComponentToPropertyMap to associate native UObjectProperty fields w/ natively-constructed component instances at construction time
- modified AActor::ResetOwnedComponents() to ensure that UObjectProperty fields that referenced natively-constructed component properties post-construction will fall back to those values post-load if they serialized as NULL

notes:
- also intended to address UE-18284 (broken component property editing in the Blueprint editor)

diff --git a/Engine/Source/Runtime/Engine/Classes/GameFramework/Actor.h b/Engine/Source/Runtime/Engine/Classes/GameFramework/Actor.h
--- a/Engine/Source/Runtime/Engine/Classes/GameFramework/Actor.h
+++ b/Engine/Source/Runtime/Engine/Classes/GameFramework/Actor.h
@@ -2484,6 +2484,9 @@ private:
 	/** List of replicated components. */
 	TArray<UActorComponent*> ReplicatedComponents;
 
+	/** Maps natively-constructed components to properties that reference them. */
+	TMultiMap<UActorComponent*, UObjectProperty*> NativeConstructedComponentToPropertyMap;
+
 public:
 
 	/** Array of ActorComponents that are created by blueprints and serialized per-instance. */
diff --git a/Engine/Source/Runtime/Engine/Private/Actor.cpp b/Engine/Source/Runtime/Engine/Private/Actor.cpp
--- a/Engine/Source/Runtime/Engine/Private/Actor.cpp
+++ b/Engine/Source/Runtime/Engine/Private/Actor.cpp
@@ -193,6 +193,37 @@ bool AActor::CheckActorComponents()
 
 void AActor::ResetOwnedComponents()
 {
+	// Identify any natively-constructed components referenced by properties that either failed to serialize or came in as NULL.
+	if(HasAnyFlags(RF_WasLoaded) && NativeConstructedComponentToPropertyMap.Num() > 0)
+	{
+		TInlineComponentArray<UActorComponent*> ComponentsToDestroy;
+		for (auto Component : OwnedComponents)
+		{
+			// Only consider native components
+			if (Component->CreationMethod == EComponentCreationMethod::Native)
+			{
+				// Find the property or properties that previously referenced the natively-constructed component.
+				TArray<UObjectProperty*> Properties;
+				NativeConstructedComponentToPropertyMap.MultiFind(Component, Properties);
+
+				// Determine if the property or properties are no longer valid references (either it got serialized out that way or something failed during load)
+				for (auto ObjProp : Properties)
+				{
+					check(ObjProp != nullptr);
+					UActorComponent* ActorComponent = Cast<UActorComponent>(ObjProp->GetObjectPropertyValue_InContainer(this));
+					if (ActorComponent == nullptr)
+					{
+						// Restore the natively-constructed component instance
+						ObjProp->SetObjectPropertyValue_InContainer(this, Component);
+					}
+				}
+			}
+		}
+
+		// Clear out the mapping as we don't need it anymore
+		NativeConstructedComponentToPropertyMap.Empty();
+	}
+
 	TArray<UObject*> ActorChildren;
 	OwnedComponents.Empty();
 	ReplicatedComponents.Empty();
@@ -219,6 +250,27 @@ void AActor::PostInitProperties()
 	RegisterAllActorTickFunctions(true,false); // component will be handled when they are registered
 	RemoteRole = (bReplicates ? ROLE_SimulatedProxy : ROLE_None);
 
+	// Map natively-constructed component instances to any serialized properties that might reference them.
+	// We'll use this information post-load to determine if any owned components may not have been serialized through the reference property (i.e. in case the serialized property value ends up being NULL).
+	if(!HasAnyFlags(RF_WasLoaded))
+	{
+		NativeConstructedComponentToPropertyMap.Empty(OwnedComponents.Num());
+		for(TFieldIterator<UObjectProperty> PropertyIt(GetClass(), EFieldIteratorFlags::IncludeSuper); PropertyIt; ++PropertyIt)
+		{
+			UObjectProperty* ObjProp = *PropertyIt;
+
+			// Ignore transient properties since they won't be serialized
+			if(!ObjProp->HasAnyPropertyFlags(CPF_Transient))
+			{
+				UActorComponent* ActorComponent = Cast<UActorComponent>(ObjProp->GetObjectPropertyValue_InContainer(this));
+				if(ActorComponent != nullptr && ActorComponent->CreationMethod == EComponentCreationMethod::Native)
+				{
+					NativeConstructedComponentToPropertyMap.Add(ActorComponent, ObjProp);
+				}
+			}
+		}
+	}
+
 	// Make sure the OwnedComponents list correct.  
 	// Under some circumstances sub-object instancing can result in bogus/duplicate entries.
 	// This is not necessary for CDOs
