Move skeletal mesh cloth destroy to FinishDestroy. This is needed for when we shut down the editor and call BeginDestroy on all objects regardless of order. The result is that when one of the worlds destroys we flush the global apex buffer and things get deleted out of order.

diff --git a/Engine/Source/Runtime/Engine/Classes/Engine/SkeletalMesh.h b/Engine/Source/Runtime/Engine/Classes/Engine/SkeletalMesh.h
--- a/Engine/Source/Runtime/Engine/Classes/Engine/SkeletalMesh.h
+++ b/Engine/Source/Runtime/Engine/Classes/Engine/SkeletalMesh.h
@@ -737,6 +737,7 @@ public:
 	virtual void GetAssetRegistryTagMetadata(TMap<FName, FAssetRegistryTagMetadata>& OutMetadata) const override;
 #endif // WITH_EDITOR
 	virtual void BeginDestroy() override;
+	virtual void FinishDestroy() override;
 	virtual bool IsReadyForFinishDestroy() override;
 	virtual void PreSave() override;
 	virtual void Serialize(FArchive& Ar) override;
diff --git a/Engine/Source/Runtime/Engine/Private/SkeletalMesh.cpp b/Engine/Source/Runtime/Engine/Private/SkeletalMesh.cpp
--- a/Engine/Source/Runtime/Engine/Private/SkeletalMesh.cpp
+++ b/Engine/Source/Runtime/Engine/Private/SkeletalMesh.cpp
@@ -2404,6 +2404,11 @@ void USkeletalMesh::BeginDestroy()
 	// Release the mesh's render resources.
 	ReleaseResources();
 
+
+}
+
+void USkeletalMesh::FinishDestroy()
+{
 #if WITH_APEX_CLOTHING
 	// release clothing assets
 	for (FClothingAssetData& Data : ClothingAssets)
@@ -2415,6 +2420,8 @@ void USkeletalMesh::BeginDestroy()
 		}
 	}
 #endif // #if WITH_APEX_CLOTHING
+
+	Super::FinishDestroy();
 }
 
 
