#ue4 - Fix missing replication of AttachChildren array. Avoid crashing under certain networking conditions when joining in progress with attached components.

- Note: just made this a non-fatal assert for now, a more complete "fix" that this is warning about will come later, it's a larger refactor that is not worth the risk of pushing to 4.9 so late.

diff --git a/Engine/Source/Runtime/Engine/Private/Components/SceneComponent.cpp b/Engine/Source/Runtime/Engine/Private/Components/SceneComponent.cpp
--- a/Engine/Source/Runtime/Engine/Private/Components/SceneComponent.cpp
+++ b/Engine/Source/Runtime/Engine/Private/Components/SceneComponent.cpp
@@ -1384,7 +1384,7 @@ void USceneComponent::DetachFromParent(bool bMaintainWorldPosition, bool bCallMo
 		}
 
 		// Make sure parent points to us if we're registered
-		checkf(!bRegistered || AttachParent->AttachChildren.Contains(this), TEXT("Attempt to detach SceneComponent '%s' owned by '%s' from AttachParent '%s' while not attached."), *GetName(), (Owner ? *Owner->GetName() : TEXT("Unowned")), *AttachParent->GetName());
+		ensureMsgf(!bRegistered || AttachParent->AttachChildren.Contains(this), TEXT("Attempt to detach SceneComponent '%s' owned by '%s' from AttachParent '%s' while not attached."), *GetName(), (Owner ? *Owner->GetName() : TEXT("Unowned")), *AttachParent->GetName());
 
 		if (bCallModify)
 		{
@@ -2270,6 +2270,7 @@ void USceneComponent::GetLifetimeReplicatedProps(TArray< FLifetimeProperty > & O
 	DOREPLIFETIME(USceneComponent, bAbsoluteScale);
 	DOREPLIFETIME(USceneComponent, bVisible);
 	DOREPLIFETIME(USceneComponent, AttachParent);
+	DOREPLIFETIME(USceneComponent, AttachChildren);
 	DOREPLIFETIME(USceneComponent, AttachSocketName);
 	DOREPLIFETIME(USceneComponent, RelativeLocation);
 	DOREPLIFETIME(USceneComponent, RelativeRotation);
