Allow to subtract server and client frame time from ping if desired

diff --git a/Engine/Source/Runtime/Engine/Private/NetConnection.cpp b/Engine/Source/Runtime/Engine/Private/NetConnection.cpp
--- a/Engine/Source/Runtime/Engine/Private/NetConnection.cpp
+++ b/Engine/Source/Runtime/Engine/Private/NetConnection.cpp
@@ -14,6 +14,11 @@
 #include "GameFramework/GameMode.h"
 #include "Runtime/PacketHandlers/PacketHandler/Public/PacketHandler.h"
 
+#if !UE_BUILD_SHIPPING
+static TAutoConsoleVariable<int32> CVarPingExcludeFrameTime( TEXT( "net.PingExcludeFrameTime" ), 0, TEXT( "Calculate RTT time between NIC's of server and client." ) );
+static TAutoConsoleVariable<int32> CVarPingDisplayServerTime( TEXT( "net.PingDisplayServerTime" ), 0, TEXT( "Show server frame time" ) );
+#endif
+
 /*-----------------------------------------------------------------------------
 	UNetConnection implementation.
 -----------------------------------------------------------------------------*/
@@ -734,13 +739,23 @@ void UNetConnection::ReceivedPacket( FBitReader& Reader )
 				return;
 			}
 
-
 			// Detect ack packets on the server, containing PingAckData
-			bool bPingAck = Driver->ServerConnection == NULL && (AckPacketId % PING_ACK_PACKET_INTERVAL) == 0;
-			bool bHasPingAckData = false;
-			uint32 InPingAckData = 0;
+			const bool bPingAck		= Driver->IsServer() && ( AckPacketId % PING_ACK_PACKET_INTERVAL ) == 0;
+			bool bHasPingAckData	= false;
+			uint32 InPingAckData	= 0;
+
+#if !UE_BUILD_SHIPPING
+			double ServerFrameTime	= 0;
 
-			if (bPingAck)
+			if ( !Driver->IsServer() )
+			{
+				uint8 FrameTimeByte	= 0;
+				Reader << FrameTimeByte;
+				ServerFrameTime = ( double )FrameTimeByte / 1000;
+			}
+#endif
+
+			if ( bPingAck )
 			{
 				bHasPingAckData = !!Reader.ReadBit();
 
@@ -773,7 +788,18 @@ void UNetConnection::ReceivedPacket( FBitReader& Reader )
 			{
 				OutLagPacketId[Index] = -1;		// Only use the ack once
 
-				const float NewLag = FPlatformTime::Seconds() - OutLagTime[Index];
+#if !UE_BUILD_SHIPPING
+				if ( CVarPingDisplayServerTime.GetValueOnGameThread() > 0 )
+				{
+					UE_LOG( LogNetTraffic, Warning, TEXT( "ServerFrameTime: %2.2f" ), ServerFrameTime * 1000.0f );
+				}
+
+				const float GameTime	= ServerFrameTime + FrameTime;
+				const float RTT			= ( FPlatformTime::Seconds() - OutLagTime[Index] ) - ( CVarPingExcludeFrameTime.GetValueOnGameThread() ? GameTime : 0.0f );
+				const float NewLag		= FMath::Max( RTT, 0.0f );
+#else
+				const float NewLag		= FPlatformTime::Seconds() - OutLagTime[Index];
+#endif
 
 				LagAcc += NewLag;
 				LagCount++;
@@ -1158,7 +1184,15 @@ void UNetConnection::SendAck( int32 AckPacketId, bool FirstTime/*=1*/, bool bHav
 		AckData.WriteBit( 1 );
 		AckData.WriteIntWrapped(AckPacketId, MAX_PACKETID);
 
-		const bool bPingAck = Driver->ServerConnection != NULL && (AckPacketId % PING_ACK_PACKET_INTERVAL) == 0;
+#if !UE_BUILD_SHIPPING
+		if ( Driver->IsServer() )
+		{
+			uint8 FrameTimeByte = FMath::Min( FMath::FloorToInt( FrameTime * 1000 ), 255 );
+			AckData << FrameTimeByte;
+		}
+#endif
+
+		const bool bPingAck = !Driver->IsServer() && ( AckPacketId % PING_ACK_PACKET_INTERVAL ) == 0;
 
 		if ( bPingAck )
 		{
