Core: Minor readability improvements to on-screen stat display
- Added alternating background colors
- Switched to a solid color for each row rather than the gradient BG
- Made a few tiny optimizations
[INTEGRATE] Change 2646008 by [REDACTED_PERSON][REDACTED_HANDLE] on 2015/08/05 21:56:30

diff --git a/Engine/Source/Runtime/Engine/Private/StatsRender2.cpp b/Engine/Source/Runtime/Engine/Private/StatsRender2.cpp
--- a/Engine/Source/Runtime/Engine/Private/StatsRender2.cpp
+++ b/Engine/Source/Runtime/Engine/Private/StatsRender2.cpp
@@ -17,6 +17,9 @@ enum class EStatRenderConsts
 	NUM_COLUMNS = 5,
 };
 
+/** Should we use a solid fill or a graident? */
+const bool bUseFlatBackgroundForStats = true;
+
 /** Enumerates stat font types and maximum length of the stat names. */
 enum class EStatFontTypes : int32
 {
@@ -85,7 +88,7 @@ struct FStatRenderGlobals
 	FLinearColor GroupColor;
 
 	/** Color used as the background for every other stat item to make it easier to read across lines. */
-	FLinearColor BackgroundColor;
+	FLinearColor BackgroundColors[2];
 
 	/** The font used for rendering stats. */
 	UFont* StatFont;
@@ -105,12 +108,13 @@ struct FStatRenderGlobals
 		StatColor(0.f,1.f,0.f),
 		HeadingColor(1.f,0.2f,0.f),
 		GroupColor(FLinearColor::White),
-		BackgroundColor(0.05f, 0.05f, 0.05f, 0.90f),	// dark gray mostly occluding the background
 		StatFont(nullptr),
 		StatFontType(EStatFontTypes::NumFonts),
 		bNeedRefresh(true)
 	{
-		SetNewFont( EStatFontTypes::Small );
+		BackgroundColors[0] = FLinearColor(0.05f, 0.05f, 0.05f, 0.92f); // dark gray mostly occluding the background
+		BackgroundColors[1] = FLinearColor(0.02f, 0.02f, 0.02f, 0.88f); // slightly different to help make long lines more readable
+		SetNewFont(EStatFontTypes::Small);
 	}
 
 	/**
@@ -193,6 +197,20 @@ struct FStatRenderGlobals
 			bNeedRefresh = true;
 		}
 	}
+
+	/** Returns the background texture for stat rows */
+	FTexture* GetBackgroundTexture() const
+	{
+		if (bUseFlatBackgroundForStats)
+		{
+			return GWhiteTexture;
+		}
+		else
+		{
+			UTexture2D* BackgroundTexture = UCanvas::StaticClass()->GetDefaultObject<UCanvas>()->GradientTexture0;
+			return (BackgroundTexture != nullptr) ? BackgroundTexture->Resource : nullptr;
+		}
+	}
 };
 
 FStatRenderGlobals& GetStatRenderGlobals()
@@ -239,10 +257,12 @@ StatCmdEngineExec;
 
 static void RightJustify(FCanvas* Canvas, int32 X, int32 Y, TCHAR const* Text, FLinearColor const& Color)
 {
+	const FStatRenderGlobals& Globals = GetStatRenderGlobals();
+
 	int32 StatColumnSpaceSizeX, StatColumnSpaceSizeY;
-	StringSize(GetStatRenderGlobals().StatFont, StatColumnSpaceSizeX, StatColumnSpaceSizeY, Text);
+	StringSize(Globals.StatFont, StatColumnSpaceSizeX, StatColumnSpaceSizeY, Text);
 
-	Canvas->DrawShadowedString(X + GetStatRenderGlobals().InterColumnOffset - StatColumnSpaceSizeX,Y,Text,GetStatRenderGlobals().StatFont,Color);
+	Canvas->DrawShadowedString(X + Globals.InterColumnOffset - StatColumnSpaceSizeX, Y, Text, Globals.StatFont, Color);
 }
 
 
@@ -257,7 +277,8 @@ static void RightJustify(FCanvas* Canvas, int32 X, int32 Y, TCHAR const* Text, F
  */
 static int32 RenderCycle( const FComplexStatMessage& Item, class FCanvas* Canvas, int32 X, int32 Y, const int32 Indent, const bool bStackStat )
 {
-	FColor Color = GetStatRenderGlobals().StatColor.ToFColor(true);
+	const FStatRenderGlobals& Globals = GetStatRenderGlobals();
+	FColor Color = Globals.StatColor.ToFColor(true);
 
 	check(Item.NameAndInfo.GetFlag(EStatMetaFlags::IsCycle));
 
@@ -272,11 +293,11 @@ static int32 RenderCycle( const FComplexStatMessage& Item, class FCanvas* Canvas
 		// If show inclusive and and show exclusive is on, then it will choose color based on inclusive average
 		// #YRX_STATS: 2015-06-09 This is slow, fix this
 		FString CounterName = Item.GetShortName().ToString();
-		CounterName.RemoveFromStart(TEXT("STAT_"));
+		CounterName.RemoveFromStart(TEXT("STAT_"), ESearchCase::CaseSensitive);
 		GEngine->GetStatValueColoration(CounterName, InMs, Color);
 
 		const float MaxMeter = 33.3f; // the time of a "full bar" in ms
-		const int32 MeterWidth = GetStatRenderGlobals().AfterNameColumnOffset;
+		const int32 MeterWidth = Globals.AfterNameColumnOffset;
 
 		int32 BarWidth = int32((InMs / MaxMeter) * MeterWidth);
 		if (BarWidth > 2)
@@ -286,7 +307,7 @@ static int32 RenderCycle( const FComplexStatMessage& Item, class FCanvas* Canvas
 				BarWidth = MeterWidth;
 			}
 
-			FCanvasBoxItem BoxItem( FVector2D(X + MeterWidth - BarWidth, Y + .4f * GetStatRenderGlobals().GetFontHeight()), FVector2D(BarWidth, 0.2f * GetStatRenderGlobals().GetFontHeight()) );
+			FCanvasBoxItem BoxItem(FVector2D(X + MeterWidth - BarWidth, Y + .4f * Globals.GetFontHeight()), FVector2D(BarWidth, 0.2f * Globals.GetFontHeight()));
 			BoxItem.SetColor( FLinearColor::Red );
 			BoxItem.Draw( Canvas );		
 		}
@@ -296,9 +317,9 @@ static int32 RenderCycle( const FComplexStatMessage& Item, class FCanvas* Canvas
 	const FString StatDesc = Item.GetDescription();
 	const FString StatDisplay = StatDesc.Len() == 0 ? Item.GetShortName().GetPlainNameString() : StatDesc;
 
-	Canvas->DrawShadowedString(X + IndentWidth, Y, *ShortenName(*StatDisplay), GetStatRenderGlobals().StatFont, Color);
+	Canvas->DrawShadowedString(X + IndentWidth, Y, *ShortenName(*StatDisplay), Globals.StatFont, Color);
 
-	int32 CurrX = X + GetStatRenderGlobals().AfterNameColumnOffset;
+	int32 CurrX = X + Globals.AfterNameColumnOffset;
 	// Now append the call count
 	if( bStackStat )
 	{
@@ -306,7 +327,7 @@ static int32 RenderCycle( const FComplexStatMessage& Item, class FCanvas* Canvas
 		{
 			RightJustify(Canvas,CurrX,Y,*FString::Printf(TEXT("%u"), Item.GetValue_CallCount(EComplexStatField::IncAve)),Color);
 		}
-		CurrX += GetStatRenderGlobals().InterColumnOffset;
+		CurrX += Globals.InterColumnOffset;
 	}
 
 	// Add the two inclusive columns if asked
@@ -314,14 +335,14 @@ static int32 RenderCycle( const FComplexStatMessage& Item, class FCanvas* Canvas
 	{
 		RightJustify(Canvas,CurrX,Y,*FString::Printf(TEXT("%1.2f ms"),FPlatformTime::ToMilliseconds(Item.GetValue_Duration(EComplexStatField::IncAve))),Color); 
 	}
-	CurrX += GetStatRenderGlobals().InterColumnOffset;
+	CurrX += Globals.InterColumnOffset;
 
 	if( bIsInitialized )
 	{
 		RightJustify(Canvas,CurrX,Y,*FString::Printf(TEXT("%1.2f ms"),FPlatformTime::ToMilliseconds(Item.GetValue_Duration(EComplexStatField::IncMax))),Color);
 		
 	}
-	CurrX += GetStatRenderGlobals().InterColumnOffset;
+	CurrX += Globals.InterColumnOffset;
 
 	if( bStackStat )
 	{
@@ -330,15 +351,15 @@ static int32 RenderCycle( const FComplexStatMessage& Item, class FCanvas* Canvas
 		{
 			RightJustify(Canvas,CurrX,Y,*FString::Printf(TEXT("%1.2f ms"),FPlatformTime::ToMilliseconds(Item.GetValue_Duration(EComplexStatField::ExcAve))),Color);
 		}
-		CurrX += GetStatRenderGlobals().InterColumnOffset;
+		CurrX += Globals.InterColumnOffset;
 
 		if( bIsInitialized )
 		{
 			RightJustify(Canvas,CurrX,Y,*FString::Printf(TEXT("%1.2f ms"),FPlatformTime::ToMilliseconds(Item.GetValue_Duration(EComplexStatField::ExcMax))),Color);
 		}
-		CurrX += GetStatRenderGlobals().InterColumnOffset;
+		CurrX += Globals.InterColumnOffset;
 	}
-	return GetStatRenderGlobals().GetFontHeight();
+	return Globals.GetFontHeight();
 }
 
 /**
@@ -357,24 +378,25 @@ static int32 RenderGroupedHeadings(class FCanvas* Canvas,int X,int32 Y,const boo
 
 	static const TCHAR* CaptionFlat = TEXT("Cycle counters (flat)");
 	static const TCHAR* CaptionHier = TEXT("Cycle counters (hierarchy)");
+	const FStatRenderGlobals& Globals = GetStatRenderGlobals();
 
-	Canvas->DrawShadowedString(X,Y,bIsHierarchy?CaptionHier:CaptionFlat,GetStatRenderGlobals().StatFont,GetStatRenderGlobals().HeadingColor);
+	Canvas->DrawShadowedString(X, Y, bIsHierarchy ? CaptionHier : CaptionFlat, Globals.StatFont, Globals.HeadingColor);
 
-	int32 CurrX = X + GetStatRenderGlobals().AfterNameColumnOffset;
-	RightJustify(Canvas,CurrX,Y,TEXT("CallCount"),GetStatRenderGlobals().HeadingColor);
-	CurrX += GetStatRenderGlobals().InterColumnOffset;
+	int32 CurrX = X + Globals.AfterNameColumnOffset;
+	RightJustify(Canvas, CurrX, Y, TEXT("CallCount"), Globals.HeadingColor);
+	CurrX += Globals.InterColumnOffset;
 
-	RightJustify(Canvas,CurrX,Y,TEXT("InclusiveAvg"),GetStatRenderGlobals().HeadingColor);
-	CurrX += GetStatRenderGlobals().InterColumnOffset;
-	RightJustify(Canvas,CurrX,Y,TEXT("InclusiveMax"),GetStatRenderGlobals().HeadingColor);
-	CurrX += GetStatRenderGlobals().InterColumnOffset;
+	RightJustify(Canvas, CurrX, Y, TEXT("InclusiveAvg"), Globals.HeadingColor);
+	CurrX += Globals.InterColumnOffset;
+	RightJustify(Canvas, CurrX, Y, TEXT("InclusiveMax"), Globals.HeadingColor);
+	CurrX += Globals.InterColumnOffset;
 
-	RightJustify(Canvas,CurrX,Y,TEXT("ExclusiveAvg"),GetStatRenderGlobals().HeadingColor);
-	CurrX += GetStatRenderGlobals().InterColumnOffset;
-	RightJustify(Canvas,CurrX,Y,TEXT("ExclusiveMax"),GetStatRenderGlobals().HeadingColor);
-	CurrX += GetStatRenderGlobals().InterColumnOffset;
+	RightJustify(Canvas, CurrX, Y, TEXT("ExclusiveAvg"), Globals.HeadingColor);
+	CurrX += Globals.InterColumnOffset;
+	RightJustify(Canvas, CurrX, Y, TEXT("ExclusiveMax"), Globals.HeadingColor);
+	CurrX += Globals.InterColumnOffset;
 
-	return GetStatRenderGlobals().GetFontHeight() + (GetStatRenderGlobals().GetFontHeight() / 3);
+	return Globals.GetFontHeight() + (Globals.GetFontHeight() / 3);
 }
 
 /**
@@ -390,19 +412,20 @@ static int32 RenderCounterHeadings(class FCanvas* Canvas,int32 X,int32 Y)
 {
 	// The heading looks like:
 	// Stat [32chars]    Value [8chars]    Average [8chars]
+	const FStatRenderGlobals& Globals = GetStatRenderGlobals();
 
-	Canvas->DrawShadowedString(X,Y,TEXT("Counters"),GetStatRenderGlobals().StatFont,GetStatRenderGlobals().HeadingColor);
+	Canvas->DrawShadowedString(X, Y, TEXT("Counters"), Globals.StatFont, Globals.HeadingColor);
 
 	// Determine where the first column goes
-	int32 CurrX = X + GetStatRenderGlobals().AfterNameColumnOffset;
+	int32 CurrX = X + Globals.AfterNameColumnOffset;
 
 	// Draw the average column label.
-	RightJustify(Canvas,CurrX,Y,TEXT("Average"),GetStatRenderGlobals().HeadingColor);
-	CurrX += GetStatRenderGlobals().InterColumnOffset;
+	RightJustify(Canvas, CurrX, Y, TEXT("Average"), Globals.HeadingColor);
+	CurrX += Globals.InterColumnOffset;
 
 	// Draw the max column label.
-	RightJustify(Canvas,CurrX,Y,TEXT("Max"),GetStatRenderGlobals().HeadingColor);
-	return GetStatRenderGlobals().GetFontHeight() + (GetStatRenderGlobals().GetFontHeight() / 3);
+	RightJustify(Canvas, CurrX, Y, TEXT("Max"), Globals.HeadingColor);
+	return Globals.GetFontHeight() + (Globals.GetFontHeight() / 3);
 }
 
 /**
@@ -418,23 +441,24 @@ static int32 RenderMemoryHeadings(class FCanvas* Canvas,int32 X,int32 Y)
 {
 	// The heading looks like:
 	// Stat [32chars]    MemUsed [8chars]    PhysMem [8chars]
+	const FStatRenderGlobals& Globals = GetStatRenderGlobals();
 
-	Canvas->DrawShadowedString(X,Y,TEXT("Memory Counters"),GetStatRenderGlobals().StatFont,GetStatRenderGlobals().HeadingColor);
+	Canvas->DrawShadowedString(X, Y, TEXT("Memory Counters"), Globals.StatFont, Globals.HeadingColor);
 
 	// Determine where the first column goes
-	int32 CurrX = X + GetStatRenderGlobals().AfterNameColumnOffset;
-	RightJustify(Canvas,CurrX,Y,TEXT("MemUsedMax"),GetStatRenderGlobals().HeadingColor);
+	int32 CurrX = X + Globals.AfterNameColumnOffset;
+	RightJustify(Canvas, CurrX, Y, TEXT("MemUsedMax"), Globals.HeadingColor);
 
-	CurrX += GetStatRenderGlobals().InterColumnOffset;
-	RightJustify(Canvas,CurrX,Y,TEXT("MemUsedMax%"),GetStatRenderGlobals().HeadingColor);
+	CurrX += Globals.InterColumnOffset;
+	RightJustify(Canvas, CurrX, Y, TEXT("MemUsedMax%"), Globals.HeadingColor);
 
-	CurrX += GetStatRenderGlobals().InterColumnOffset;
-	RightJustify(Canvas,CurrX,Y,TEXT("MemPool"),GetStatRenderGlobals().HeadingColor);
+	CurrX += Globals.InterColumnOffset;
+	RightJustify(Canvas, CurrX, Y, TEXT("MemPool"), Globals.HeadingColor);
 
-	CurrX += GetStatRenderGlobals().InterColumnOffset;
-	RightJustify(Canvas,CurrX,Y,TEXT("Pool Capacity"),GetStatRenderGlobals().HeadingColor);
+	CurrX += Globals.InterColumnOffset;
+	RightJustify(Canvas, CurrX, Y, TEXT("Pool Capacity"), Globals.HeadingColor);
 
-	return GetStatRenderGlobals().GetFontHeight() + (GetStatRenderGlobals().GetFontHeight() / 3);
+	return Globals.GetFontHeight() + (Globals.GetFontHeight() / 3);
 }
 
 static FString FormatStatValueFloat( const float Value )
@@ -482,40 +506,42 @@ static int32 RenderMemoryCounter(const FGameThreadHudData& ViewData, const FComp
 	FPlatformMemory::EMemoryCounterRegion Region = FPlatformMemory::EMemoryCounterRegion(All.NameAndInfo.GetField<EMemoryRegion>());
 	// At this moment we only have memory stats that are marked as non frame stats, so can't be cleared every frame.
 	//const bool bDisplayAll = All.NameAndInfo.GetFlag(EStatMetaFlags::ShouldClearEveryFrame);
-	
+	const FStatRenderGlobals& Globals = GetStatRenderGlobals();
 	const float MaxMemUsed = All.GetValue_double(EComplexStatField::IncMax);
 
 	// Draw the label
-	Canvas->DrawShadowedString(X,Y,*ShortenName(*All.GetDescription()),GetStatRenderGlobals().StatFont,GetStatRenderGlobals().StatColor);
-	int32 CurrX = X + GetStatRenderGlobals().AfterNameColumnOffset;
+	Canvas->DrawShadowedString(X, Y, *ShortenName(*All.GetDescription()), Globals.StatFont, Globals.StatColor);
+	int32 CurrX = X + Globals.AfterNameColumnOffset;
 
 	// always use MB for easier comparisons
 	const bool bAutoType = false;
 
 	// Now append the max value of the stat
-	RightJustify(Canvas,CurrX,Y,*GetMemoryString(MaxMemUsed, bAutoType),GetStatRenderGlobals().StatColor);
-	CurrX += GetStatRenderGlobals().InterColumnOffset;
+	RightJustify(Canvas, CurrX, Y, *GetMemoryString(MaxMemUsed, bAutoType), Globals.StatColor);
+	CurrX += Globals.InterColumnOffset;
 	if (ViewData.PoolCapacity.Contains(Region))
 	{
-		RightJustify(Canvas,CurrX,Y,*FString::Printf(TEXT("%.0f%%"), float(100.0 * MaxMemUsed / double(ViewData.PoolCapacity[Region]))),GetStatRenderGlobals().StatColor);
+		RightJustify(Canvas, CurrX, Y, *FString::Printf(TEXT("%.0f%%"), float(100.0 * MaxMemUsed / double(ViewData.PoolCapacity[Region]))), Globals.StatColor);
 	}
-	CurrX += GetStatRenderGlobals().InterColumnOffset;
+	CurrX += Globals.InterColumnOffset;
 	if (ViewData.PoolAbbreviation.Contains(Region))
 	{
-		RightJustify(Canvas,CurrX,Y,*ViewData.PoolAbbreviation[Region],GetStatRenderGlobals().StatColor);
+		RightJustify(Canvas, CurrX, Y, *ViewData.PoolAbbreviation[Region], Globals.StatColor);
 	}
-	CurrX += GetStatRenderGlobals().InterColumnOffset;
+	CurrX += Globals.InterColumnOffset;
 	if (ViewData.PoolCapacity.Contains(Region))
 	{
-		RightJustify(Canvas,CurrX,Y,*GetMemoryString(double(ViewData.PoolCapacity[Region]), bAutoType),GetStatRenderGlobals().StatColor);
+		RightJustify(Canvas, CurrX, Y, *GetMemoryString(double(ViewData.PoolCapacity[Region]), bAutoType), Globals.StatColor);
 	}
-	CurrX += GetStatRenderGlobals().InterColumnOffset;
+	CurrX += Globals.InterColumnOffset;
 
-	return GetStatRenderGlobals().GetFontHeight();
+	return Globals.GetFontHeight();
 }
 
 static int32 RenderCounter(const FGameThreadHudData& ViewData, const FComplexStatMessage& All,class FCanvas* Canvas,int32 X,int32 Y)
 {
+	const FStatRenderGlobals& Globals = GetStatRenderGlobals();
+
 	// If this is a cycle, render it as a cycle. This is a special case for manually set cycle counters.
 	const bool bIsCycle = All.NameAndInfo.GetFlag(EStatMetaFlags::IsCycle);
 	if( bIsCycle )
@@ -526,8 +552,8 @@ static int32 RenderCounter(const FGameThreadHudData& ViewData, const FComplexSta
 	const bool bDisplayAll = All.NameAndInfo.GetFlag(EStatMetaFlags::ShouldClearEveryFrame);
 
 	// Draw the label
-	Canvas->DrawShadowedString(X,Y,*ShortenName(*All.GetDescription()),GetStatRenderGlobals().StatFont,GetStatRenderGlobals().StatColor);
-	int32 CurrX = X + GetStatRenderGlobals().AfterNameColumnOffset;
+	Canvas->DrawShadowedString(X, Y, *ShortenName(*All.GetDescription()), Globals.StatFont, Globals.StatColor);
+	int32 CurrX = X + Globals.AfterNameColumnOffset;
 
 	if( bDisplayAll )
 	{
@@ -535,33 +561,34 @@ static int32 RenderCounter(const FGameThreadHudData& ViewData, const FComplexSta
 		if (All.NameAndInfo.GetField<EStatDataType>() == EStatDataType::ST_double)
 		{
 			const FString ValueFormatted = FormatStatValueFloat( All.GetValue_double( EComplexStatField::IncAve ) );
-			RightJustify( Canvas, CurrX, Y, *ValueFormatted, GetStatRenderGlobals().StatColor );
+			RightJustify(Canvas, CurrX, Y, *ValueFormatted, Globals.StatColor);
 		}
 		else if (All.NameAndInfo.GetField<EStatDataType>() == EStatDataType::ST_int64)
 		{
 			const FString ValueFormatted = FormatStatValueInt64( All.GetValue_int64( EComplexStatField::IncAve ) );
-			RightJustify( Canvas, CurrX, Y, *ValueFormatted, GetStatRenderGlobals().StatColor );
+			RightJustify(Canvas, CurrX, Y, *ValueFormatted, Globals.StatColor);
 		}
 	}
-	CurrX += GetStatRenderGlobals().InterColumnOffset;
+	CurrX += Globals.InterColumnOffset;
 
 	// Append the maximum.
 	if (All.NameAndInfo.GetField<EStatDataType>() == EStatDataType::ST_double)
 	{
 		const FString ValueFormatted = FormatStatValueFloat( All.GetValue_double( EComplexStatField::IncMax ) );
-		RightJustify( Canvas, CurrX, Y, *ValueFormatted, GetStatRenderGlobals().StatColor );
+		RightJustify(Canvas, CurrX, Y, *ValueFormatted, Globals.StatColor);
 	}
 	else if (All.NameAndInfo.GetField<EStatDataType>() == EStatDataType::ST_int64)
 	{
 		const FString ValueFormatted = FormatStatValueInt64( All.GetValue_int64( EComplexStatField::IncMax ) );
-		RightJustify( Canvas, CurrX, Y, *ValueFormatted, GetStatRenderGlobals().StatColor );
+		RightJustify(Canvas, CurrX, Y, *ValueFormatted, Globals.StatColor);
 	}
-	return GetStatRenderGlobals().GetFontHeight();
+	return Globals.GetFontHeight();
 }
 
 void RenderHierCycles( FCanvas* Canvas, int32 X, int32& Y, const FHudGroup& HudGroup )
 {
-	UTexture2D* BackgroundTexture = UCanvas::StaticClass()->GetDefaultObject<UCanvas>()->GradientTexture0;
+	const FStatRenderGlobals& Globals = GetStatRenderGlobals();
+	const FTexture* BackgroundTexture = Globals.GetBackgroundTexture();
 
 	// Render all cycle counters.
 	for( int32 RowIndex = 0; RowIndex < HudGroup.HierAggregate.Num(); ++RowIndex )
@@ -569,11 +596,11 @@ void RenderHierCycles( FCanvas* Canvas, int32 X, int32& Y, const FHudGroup& HudG
 		const FComplexStatMessage& ComplexStat = HudGroup.HierAggregate[RowIndex];
 		const int32 Indent = HudGroup.Indentation[RowIndex];
 
-		if(BackgroundTexture != nullptr)
+		if (BackgroundTexture != nullptr)
 		{
-			Canvas->DrawTile( X, Y + GetStatRenderGlobals().GetYOffset(), GetStatRenderGlobals().AfterNameColumnOffset + GetStatRenderGlobals().InterColumnOffset * (int32)EStatRenderConsts::NUM_COLUMNS, GetStatRenderGlobals().GetFontHeight(),
+			Canvas->DrawTile(X, Y + Globals.GetYOffset(), Globals.AfterNameColumnOffset + Globals.InterColumnOffset * (int32)EStatRenderConsts::NUM_COLUMNS, Globals.GetFontHeight(),
 				0, 0, 1, 1,
-				GetStatRenderGlobals().BackgroundColor, BackgroundTexture->Resource, true );
+				Globals.BackgroundColors[RowIndex & 1], BackgroundTexture, true);
 		}
 
 		Y += RenderCycle( ComplexStat, Canvas, X, Y, Indent, true );
@@ -583,18 +610,19 @@ void RenderHierCycles( FCanvas* Canvas, int32 X, int32& Y, const FHudGroup& HudG
 template< typename T >
 void RenderArrayOfStats( FCanvas* Canvas, int32 X, int32& Y, const TArray<FComplexStatMessage>& Aggregates, const FGameThreadHudData& ViewData, const T& FunctionToCall )
 {
-	UTexture2D* BackgroundTexture = UCanvas::StaticClass()->GetDefaultObject<UCanvas>()->GradientTexture0;
+	const FStatRenderGlobals& Globals = GetStatRenderGlobals();
+	const FTexture* BackgroundTexture = Globals.GetBackgroundTexture();
 
 	// Render all counters.
 	for( int32 RowIndex = 0; RowIndex < Aggregates.Num(); ++RowIndex )
 	{
 		const FComplexStatMessage& ComplexStat = Aggregates[RowIndex];
 
-		if(BackgroundTexture != nullptr)
+		if (BackgroundTexture != nullptr)
 		{
-			Canvas->DrawTile( X, Y + GetStatRenderGlobals().GetYOffset(), GetStatRenderGlobals().AfterNameColumnOffset + GetStatRenderGlobals().InterColumnOffset * (int32)EStatRenderConsts::NUM_COLUMNS, GetStatRenderGlobals().GetFontHeight(),
+			Canvas->DrawTile(X, Y + Globals.GetYOffset(), Globals.AfterNameColumnOffset + Globals.InterColumnOffset * (int32)EStatRenderConsts::NUM_COLUMNS, Globals.GetFontHeight(),
 				0, 0, 1, 1,
-				GetStatRenderGlobals().BackgroundColor, BackgroundTexture->Resource, true );
+				Globals.BackgroundColors[RowIndex & 1], BackgroundTexture, true);
 		}
 
 		Y += FunctionToCall( ViewData, ComplexStat, Canvas, X, Y );
@@ -619,12 +647,14 @@ static void RenderGroupedWithHierarchy(const FGameThreadHudData& ViewData, FView
 	// Grab texture for rendering text background.
 	UTexture2D* BackgroundTexture = UCanvas::StaticClass()->GetDefaultObject<UCanvas>()->DefaultTexture;
 
+	const FStatRenderGlobals& Globals = GetStatRenderGlobals();
+
 	// Render all groups.
 	for( int32 GroupIndex = 0; GroupIndex < ViewData.HudGroups.Num(); ++GroupIndex )
 	{
 		// If the stat isn't enabled for this particular viewport, skip
 		FString StatGroupName = ViewData.GroupNames[GroupIndex].ToString();
-		StatGroupName.RemoveFromStart(TEXT("STATGROUP_"));
+		StatGroupName.RemoveFromStart(TEXT("STATGROUP_"), ESearchCase::CaseSensitive);
 		if (!Viewport->GetClient() || !Viewport->GetClient()->IsStatEnabled(StatGroupName))
 		{
 			continue;
@@ -634,8 +664,8 @@ static void RenderGroupedWithHierarchy(const FGameThreadHudData& ViewData, FView
 		const FName& GroupName = ViewData.GroupNames[GroupIndex];
 		const FString& GroupDesc = ViewData.GroupDescriptions[GroupIndex];
 		const FString GroupLongName = FString::Printf( TEXT("%s [%s]"), *GroupDesc, *GroupName.GetPlainNameString() );
-		Canvas->DrawShadowedString( X, Y, *GroupLongName, GetStatRenderGlobals().StatFont, GetStatRenderGlobals().GroupColor );
-		Y += GetStatRenderGlobals().GetFontHeight();
+		Canvas->DrawShadowedString(X, Y, *GroupLongName, Globals.StatFont, Globals.GroupColor);
+		Y += Globals.GetFontHeight();
 
 		const FHudGroup& HudGroup = ViewData.HudGroups[GroupIndex];
 		const bool bHasHierarchy = !!HudGroup.HierAggregate.Num();
@@ -651,14 +681,14 @@ static void RenderGroupedWithHierarchy(const FGameThreadHudData& ViewData, FView
 		if( bHasHierarchy )
 		{
 			RenderHierCycles( Canvas, X, Y, HudGroup );
-			Y += GetStatRenderGlobals().GetFontHeight();
+			Y += Globals.GetFontHeight();
 		}
 
 		// Render flat.
 		if( bHasFlat )
 		{
 			RenderArrayOfStats(Canvas,X,Y,HudGroup.FlatAggregate, ViewData, RenderFlatCycle);
-			Y += GetStatRenderGlobals().GetFontHeight();
+			Y += Globals.GetFontHeight();
 		}
 
 		// Render memory counters.
@@ -666,7 +696,7 @@ static void RenderGroupedWithHierarchy(const FGameThreadHudData& ViewData, FView
 		{
 			Y += RenderMemoryHeadings(Canvas,X,Y);
 			RenderArrayOfStats(Canvas,X,Y,HudGroup.MemoryAggregate, ViewData, RenderMemoryCounter);
-			Y += GetStatRenderGlobals().GetFontHeight();
+			Y += Globals.GetFontHeight();
 		}
 
 		// Render remaining counters.
@@ -674,7 +704,7 @@ static void RenderGroupedWithHierarchy(const FGameThreadHudData& ViewData, FView
 		{
 			Y += RenderCounterHeadings(Canvas,X,Y);
 			RenderArrayOfStats(Canvas,X,Y,HudGroup.CountersAggregate, ViewData, RenderCounter);
-			Y += GetStatRenderGlobals().GetFontHeight();
+			Y += Globals.GetFontHeight();
 		}
 	}
 }
@@ -696,7 +726,10 @@ void RenderStats(FViewport* Viewport, class FCanvas* Canvas, int32 X, int32 Y)
 	{
 		return;
 	}
-	GetStatRenderGlobals().Initialize( Viewport->GetSizeXY() );
+
+	FStatRenderGlobals& Globals = GetStatRenderGlobals();
+	Globals.Initialize(Viewport->GetSizeXY());
+
 	if( !ViewData->bDrawOnlyRawStats )
 	{
 		RenderGroupedWithHierarchy(*ViewData, Viewport, Canvas, X, Y);
@@ -706,8 +739,8 @@ void RenderStats(FViewport* Viewport, class FCanvas* Canvas, int32 X, int32 Y)
 		// Render all counters.
 		for( int32 RowIndex = 0; RowIndex < ViewData->GroupDescriptions.Num(); ++RowIndex )
 		{
-			Canvas->DrawShadowedString(X, Y, *ViewData->GroupDescriptions[RowIndex], GetStatRenderGlobals().StatFont, GetStatRenderGlobals().StatColor);
-			Y += GetStatRenderGlobals().GetFontHeight();
+			Canvas->DrawShadowedString(X, Y, *ViewData->GroupDescriptions[RowIndex], Globals.StatFont, Globals.StatColor);
+			Y += Globals.GetFontHeight();
 		}
 	}
 }
