Fix for a crash when concurrently creating render states for primitives

diff --git a/Engine/Source/Runtime/Engine/Private/Materials/MaterialInstance.cpp b/Engine/Source/Runtime/Engine/Private/Materials/MaterialInstance.cpp
--- a/Engine/Source/Runtime/Engine/Private/Materials/MaterialInstance.cpp
+++ b/Engine/Source/Runtime/Engine/Private/Materials/MaterialInstance.cpp
@@ -1085,7 +1085,7 @@ void UMaterialInstance::CopyMaterialInstanceParameters(UMaterialInterface* Mater
 
 FMaterialResource* UMaterialInstance::GetMaterialResource(ERHIFeatureLevel::Type InFeatureLevel, EMaterialQualityLevel::Type QualityLevel)
 {
-	check(IsInGameThread());
+	check(!IsInActualRenderingThread());
 
 	if (QualityLevel == EMaterialQualityLevel::Num)
 	{
@@ -1104,7 +1104,7 @@ FMaterialResource* UMaterialInstance::GetMaterialResource(ERHIFeatureLevel::Type
 
 const FMaterialResource* UMaterialInstance::GetMaterialResource(ERHIFeatureLevel::Type InFeatureLevel, EMaterialQualityLevel::Type QualityLevel) const
 {
-	check(IsInGameThread());
+	check(!IsInActualRenderingThread());
 
 	if (QualityLevel == EMaterialQualityLevel::Num)
 	{
