Added customization for FAssetImportInfo

 - FAssetImportInfo is now marked as USTRUCT and is stored as a member value on UAssetImportData rather than using inheritance. This makes customization writing much easier and better affords UObject interoperability.
 - Removed UAssetImportData::bDirty as it was only every used to enable/disable a button on the static mesh editor, and its presence caused confusion and inconsistency.
 - Fixed issue where timestamps did not get updated correctly
 - Fixed various reimport paths that were not updating their source file information correctly.

This addresses UETOOL-466 and UE-16406.

diff --git a/Engine/Plugins/2D/Paper2D/Source/Paper2D/Private/PaperTileMap.cpp b/Engine/Plugins/2D/Paper2D/Source/Paper2D/Private/PaperTileMap.cpp
--- a/Engine/Plugins/2D/Paper2D/Source/Paper2D/Private/PaperTileMap.cpp
+++ b/Engine/Plugins/2D/Paper2D/Source/Paper2D/Private/PaperTileMap.cpp
@@ -240,7 +240,7 @@ void UPaperTileMap::GetAssetRegistryTags(TArray<FAssetRegistryTag>& OutTags) con
 {
 	if (AssetImportData)
 	{
-		OutTags.Add( FAssetRegistryTag(SourceFileTagName(), AssetImportData->ToJson(), FAssetRegistryTag::TT_Hidden) );
+		OutTags.Add( FAssetRegistryTag(SourceFileTagName(), AssetImportData->GetSourceData().ToJson(), FAssetRegistryTag::TT_Hidden) );
 	}
 
 	Super::GetAssetRegistryTags(OutTags);
diff --git a/Engine/Plugins/2D/Paper2D/Source/Paper2DEditor/Private/TileMapEditing/TileMapAssetImportData.cpp b/Engine/Plugins/2D/Paper2D/Source/Paper2DEditor/Private/TileMapEditing/TileMapAssetImportData.cpp
--- a/Engine/Plugins/2D/Paper2D/Source/Paper2DEditor/Private/TileMapEditing/TileMapAssetImportData.cpp
+++ b/Engine/Plugins/2D/Paper2D/Source/Paper2DEditor/Private/TileMapEditing/TileMapAssetImportData.cpp
@@ -25,7 +25,7 @@ UTileMapAssetImportData* UTileMapAssetImportData::GetImportDataForTileMap(UPaper
 		// Try to preserve the source file path if possible
 		if (TileMap->AssetImportData != nullptr)
 		{
-			ImportData->CopyFrom(*TileMap->AssetImportData);
+			ImportData->SourceData = TileMap->AssetImportData->SourceData;
 		}
 
 		TileMap->AssetImportData = ImportData;
diff --git a/Engine/Plugins/2D/Paper2D/Source/PaperSpriteSheetImporter/Private/PaperSpriteSheet.cpp b/Engine/Plugins/2D/Paper2D/Source/PaperSpriteSheetImporter/Private/PaperSpriteSheet.cpp
--- a/Engine/Plugins/2D/Paper2D/Source/PaperSpriteSheetImporter/Private/PaperSpriteSheet.cpp
+++ b/Engine/Plugins/2D/Paper2D/Source/PaperSpriteSheetImporter/Private/PaperSpriteSheet.cpp
@@ -24,7 +24,7 @@ void UPaperSpriteSheet::GetAssetRegistryTags(TArray<FAssetRegistryTag>& OutTags)
 {
 	if (AssetImportData)
 	{
-		OutTags.Add( FAssetRegistryTag(SourceFileTagName(), AssetImportData->ToJson(), FAssetRegistryTag::TT_Hidden) );
+		OutTags.Add( FAssetRegistryTag(SourceFileTagName(), AssetImportData->GetSourceData().ToJson(), FAssetRegistryTag::TT_Hidden) );
 	}
 
 	Super::GetAssetRegistryTags(OutTags);
diff --git a/Engine/Plugins/2D/Paper2D/Source/PaperSpriteSheetImporter/Private/PaperSpriteSheetImportFactory.cpp b/Engine/Plugins/2D/Paper2D/Source/PaperSpriteSheetImporter/Private/PaperSpriteSheetImportFactory.cpp
--- a/Engine/Plugins/2D/Paper2D/Source/PaperSpriteSheetImporter/Private/PaperSpriteSheetImportFactory.cpp
+++ b/Engine/Plugins/2D/Paper2D/Source/PaperSpriteSheetImporter/Private/PaperSpriteSheetImportFactory.cpp
@@ -72,7 +72,6 @@ UObject* UPaperSpriteSheetImportFactory::FactoryCreateText(UClass* InClass, UObj
 		if (Importer.PerformImport(LongPackagePath, Flags, SpriteSheet))
 		{
 			SpriteSheet->AssetImportData->Update(CurrentFilename);
-			SpriteSheet->AssetImportData->bDirty = false;
 
 			Result = SpriteSheet;
 		}
diff --git a/Engine/Plugins/2D/Paper2D/Source/PaperSpriteSheetImporter/Private/PaperSpriteSheetReimportFactory.cpp b/Engine/Plugins/2D/Paper2D/Source/PaperSpriteSheetImporter/Private/PaperSpriteSheetReimportFactory.cpp
--- a/Engine/Plugins/2D/Paper2D/Source/PaperSpriteSheetImporter/Private/PaperSpriteSheetReimportFactory.cpp
+++ b/Engine/Plugins/2D/Paper2D/Source/PaperSpriteSheetImporter/Private/PaperSpriteSheetReimportFactory.cpp
@@ -64,6 +64,9 @@ EReimportResult::Type UPaperSpriteSheetReimportFactory::Reimport(UObject* Obj)
 	if (UFactory::StaticImportObject(SpriteSheet->GetClass(), SpriteSheet->GetOuter(), *SpriteSheet->GetName(), RF_Public | RF_Standalone, *Filename, nullptr, this))
 	{
 		UE_LOG(LogPaperSpriteSheetImporter, Log, TEXT("Imported successfully"));
+
+		SpriteSheet->AssetImportData->Update(Filename);
+		
 		// Try to find the outer package so we can dirty it up
 		if (SpriteSheet->GetOuter())
 		{
diff --git a/Engine/Plugins/2D/Paper2D/Source/SpriterImporter/Private/PaperSpriterImportData.cpp b/Engine/Plugins/2D/Paper2D/Source/SpriterImporter/Private/PaperSpriterImportData.cpp
--- a/Engine/Plugins/2D/Paper2D/Source/SpriterImporter/Private/PaperSpriterImportData.cpp
+++ b/Engine/Plugins/2D/Paper2D/Source/SpriterImporter/Private/PaperSpriterImportData.cpp
@@ -16,7 +16,7 @@ void UPaperSpriterImportData::GetAssetRegistryTags(TArray<FAssetRegistryTag>& Ou
 {
 	if (AssetImportData != nullptr)
 	{
-		OutTags.Add( FAssetRegistryTag(SourceFileTagName(), AssetImportData->ToJson(), FAssetRegistryTag::TT_Hidden) );
+		OutTags.Add( FAssetRegistryTag(SourceFileTagName(), AssetImportData->GetSourceData().ToJson(), FAssetRegistryTag::TT_Hidden) );
 	}
 
 	Super::GetAssetRegistryTags(OutTags);
diff --git a/Engine/Plugins/2D/Paper2D/Source/SpriterImporter/Private/SpriterImporterFactory.cpp b/Engine/Plugins/2D/Paper2D/Source/SpriterImporter/Private/SpriterImporterFactory.cpp
--- a/Engine/Plugins/2D/Paper2D/Source/SpriterImporter/Private/SpriterImporterFactory.cpp
+++ b/Engine/Plugins/2D/Paper2D/Source/SpriterImporter/Private/SpriterImporterFactory.cpp
@@ -544,7 +544,6 @@ UObject* USpriterImporterFactory::FactoryCreateText(UClass* InClass, UObject* In
 // 						UFbxAnimSequenceImportData* ImportData = UFbxAnimSequenceImportData::GetImportDataForAnimSequence(NewAnimation, TemplateImportData);
 // 						ImportData->SourceFilePath = FReimportManager::SanitizeImportFilename(UFactory::CurrentFilename, NewAnimation);
 // 						ImportData->SourceFileTimestamp = IFileManager::Get().GetTimeStamp(*UFactory::CurrentFilename).ToString();
-// 						ImportData->bDirty = false;
 // 					}
 
 
diff --git a/Engine/Plugins/Editor/SpeedTreeImporter/Source/SpeedTreeImporter/Private/ReimportSpeedTreeFactory.cpp b/Engine/Plugins/Editor/SpeedTreeImporter/Source/SpeedTreeImporter/Private/ReimportSpeedTreeFactory.cpp
--- a/Engine/Plugins/Editor/SpeedTreeImporter/Source/SpeedTreeImporter/Private/ReimportSpeedTreeFactory.cpp
+++ b/Engine/Plugins/Editor/SpeedTreeImporter/Source/SpeedTreeImporter/Private/ReimportSpeedTreeFactory.cpp
@@ -82,6 +82,8 @@ EReimportResult::Type UReimportSpeedTreeFactory::Reimport(UObject* Obj)
 	{
 		UE_LOG(LogEditorFactories, Log, TEXT("-- imported successfully"));
 
+		Mesh->AssetImportData->Update(Filename);
+
 		// Mark the package dirty after the successful import
 		Mesh->MarkPackageDirty();
 
diff --git a/Engine/Plugins/Editor/SpeedTreeImporter/Source/SpeedTreeImporter/Private/SpeedTreeImportFactory.cpp b/Engine/Plugins/Editor/SpeedTreeImporter/Source/SpeedTreeImporter/Private/SpeedTreeImportFactory.cpp
--- a/Engine/Plugins/Editor/SpeedTreeImporter/Source/SpeedTreeImporter/Private/SpeedTreeImportFactory.cpp
+++ b/Engine/Plugins/Editor/SpeedTreeImporter/Source/SpeedTreeImporter/Private/SpeedTreeImportFactory.cpp
@@ -1234,7 +1234,6 @@ UObject* USpeedTreeImportFactory::FactoryCreateBinary( UClass* InClass, UObject*
 
 				// @todo AssetImportData make a data class for speed tree assets
 				StaticMesh->AssetImportData->Update(UFactory::GetCurrentFilename());
-				StaticMesh->AssetImportData->bDirty = false;
 				
 				// clear out any old data
 				StaticMesh->SourceModels.Empty();
diff --git a/Engine/Plugins/ScriptPlugin/Source/ScriptPlugin/Private/ScriptBlueprint.cpp b/Engine/Plugins/ScriptPlugin/Source/ScriptPlugin/Private/ScriptBlueprint.cpp
--- a/Engine/Plugins/ScriptPlugin/Source/ScriptPlugin/Private/ScriptBlueprint.cpp
+++ b/Engine/Plugins/ScriptPlugin/Source/ScriptPlugin/Private/ScriptBlueprint.cpp
@@ -21,7 +21,7 @@ void UScriptBlueprint::GetAssetRegistryTags(TArray<FAssetRegistryTag>& OutTags)
 {
 	if (AssetImportData)
 	{
-		OutTags.Add( FAssetRegistryTag(SourceFileTagName(), AssetImportData->ToJson(), FAssetRegistryTag::TT_Hidden) );
+		OutTags.Add( FAssetRegistryTag(SourceFileTagName(), AssetImportData->GetSourceData().ToJson(), FAssetRegistryTag::TT_Hidden) );
 	}
 
 	Super::GetAssetRegistryTags(OutTags);
@@ -67,7 +67,7 @@ bool UScriptBlueprint::IsCodeDirty() const
 		return true;
 	}
 
-	const TArray<FAssetImportInfo::FSourceFile>& Data = AssetImportData->GetSourceFileData();
+	const TArray<FAssetImportInfo::FSourceFile>& Data = AssetImportData->SourceData.SourceFiles;
 
 	if (Data.Num() == 1)
 	{
diff --git a/Engine/Source/Developer/AssetTools/Private/AssetTypeActions/AssetTypeActions_AnimSequence.cpp b/Engine/Source/Developer/AssetTools/Private/AssetTypeActions/AssetTypeActions_AnimSequence.cpp
--- a/Engine/Source/Developer/AssetTools/Private/AssetTypeActions/AssetTypeActions_AnimSequence.cpp
+++ b/Engine/Source/Developer/AssetTools/Private/AssetTypeActions/AssetTypeActions_AnimSequence.cpp
@@ -69,19 +69,17 @@ void FAssetTypeActions_AnimSequence::ExecuteReimportWithNewSource(TArray<TWeakOb
 		auto Object = (*ObjIt).Get();
 		if (Object)
 		{
-			FAssetImportInfo OldImportData;
-
 			// Make note of the old import data
-			OldImportData.CopyFrom(*Object->AssetImportData);
+			FAssetImportInfo OldImportData = Object->AssetImportData->SourceData;
 			// And reset the import data
-			Object->AssetImportData->CopyFrom(EmptyImportInfo);
+			Object->AssetImportData->SourceData = EmptyImportInfo;
 
 			bool bSuccess = FReimportManager::Instance()->Reimport(Object, /*bAskForNewFileIfMissing=*/true);
 
 			// restore the old source path in case reimport was not successful
 			if (!bSuccess)
 			{
-				Object->AssetImportData->CopyFrom(OldImportData);
+				Object->AssetImportData->SourceData = OldImportData;
 			}
 		}
 	}
diff --git a/Engine/Source/Editor/DestructibleMeshEditor/Private/SDestructibleMeshEditorViewport.cpp b/Engine/Source/Editor/DestructibleMeshEditor/Private/SDestructibleMeshEditorViewport.cpp
--- a/Engine/Source/Editor/DestructibleMeshEditor/Private/SDestructibleMeshEditorViewport.cpp
+++ b/Engine/Source/Editor/DestructibleMeshEditor/Private/SDestructibleMeshEditorViewport.cpp
@@ -246,11 +246,11 @@ bool FDestructibleMeshEditorViewportClient::CanRefreshFromStaticMesh()
 		return false;
 	}
 
-	FAssetImportInfo* ImportInfo = DestructibleMesh->SourceStaticMesh->AssetImportData;
+	const auto* ImportInfo = DestructibleMesh->SourceStaticMesh->AssetImportData;
 	FDateTime CurrentSourceTimestamp = FDateTime::MinValue();
-	if (ImportInfo && ImportInfo->GetSourceFileData().Num() == 1)
+	if (ImportInfo && ImportInfo->SourceData.SourceFiles.Num() == 1)
 	{
-		CurrentSourceTimestamp = ImportInfo->GetSourceFileData()[0].Timestamp;
+		CurrentSourceTimestamp = ImportInfo->SourceData.SourceFiles[0].Timestamp;
 	}
 
 	return (CurrentSourceTimestamp > DestructibleMesh->SourceSMImportTimestamp);
diff --git a/Engine/Source/Editor/DetailCustomizations/Private/AssetImportDataCustomization.cpp b/Engine/Source/Editor/DetailCustomizations/Private/AssetImportDataCustomization.cpp
new file mode 100644
--- /dev/null
+++ b/Engine/Source/Editor/DetailCustomizations/Private/AssetImportDataCustomization.cpp
@@ -0,0 +1,159 @@
+﻿// Copyright 1998-2015 Epic Games, Inc. All Rights Reserved.
+
+#include "DetailCustomizationsPrivatePCH.h"
+#include "AssetImportDataCustomization.h"
+#include "SEditableText.h"
+#include "EditorReimportHandler.h"
+#include "ObjectEditorUtils.h"
+
+#define LOCTEXT_NAMESPACE "AssetImportDataCustomization"
+
+TSharedRef<IPropertyTypeCustomization> FAssetImportDataCustomization::MakeInstance()
+{
+	return MakeShareable( new FAssetImportDataCustomization );
+}
+
+void FAssetImportDataCustomization::CustomizeHeader( TSharedRef<IPropertyHandle> InPropertyHandle, FDetailWidgetRow& HeaderRow, IPropertyTypeCustomizationUtils& CustomizationUtils )
+{
+}
+
+void FAssetImportDataCustomization::CustomizeChildren( TSharedRef<IPropertyHandle> InPropertyHandle, IDetailChildrenBuilder& ChildBuilder, IPropertyTypeCustomizationUtils& CustomizationUtils )
+{
+	PropertyHandle = InPropertyHandle;
+
+	FAssetImportInfo* Info = GetEditStruct();
+	if (!Info)
+	{
+		return;
+	}
+
+	auto Font = IDetailLayoutBuilder::GetDetailFont();
+
+	const FText SourceFileText = LOCTEXT("SourceFile", "Source File");
+	for (int32 Index = 0; Index < Info->SourceFiles.Num(); ++Index)
+	{
+		ChildBuilder.AddChildContent(SourceFileText)
+		.NameContent()
+		[
+			SNew(STextBlock)
+			.Text(SourceFileText)
+			.Font(Font)
+		]
+		.ValueContent()
+		.MaxDesiredWidth(TOptional<float>())
+		[
+			SNew(SVerticalBox)
+
+			+ SVerticalBox::Slot()
+			.AutoHeight()
+			[
+				SNew(SHorizontalBox)
+
+				+ SHorizontalBox::Slot()
+				[
+					SNew(SEditableText)
+					.IsReadOnly(true)
+					.Text(this, &FAssetImportDataCustomization::GetFilenameText, Index)
+					.ToolTipText(this, &FAssetImportDataCustomization::GetFilenameText, Index)
+					.Font(Font)
+				]
+
+				+SHorizontalBox::Slot()
+				.AutoWidth()
+				[
+					SNew(SButton)
+					.OnClicked(this, &FAssetImportDataCustomization::OnChangePathClicked, Index)
+					.ToolTipText(LOCTEXT("ChangePath_Tooltip", "Browse for a new source file path"))
+					[
+						SNew(STextBlock)
+						.Text(LOCTEXT("...", "..."))
+						.Font(Font)
+					]
+				]
+			]
+
+			+ SVerticalBox::Slot()
+			.AutoHeight()
+			[
+				SNew(SEditableText)
+				.IsReadOnly(true)
+				.Text(this, &FAssetImportDataCustomization::GetTimestampText, Index)
+				.Font(Font)
+			]
+		];
+	}
+}
+
+FText FAssetImportDataCustomization::GetFilenameText(int32 Index) const
+{
+	if (FAssetImportInfo* Info = GetEditStruct())
+	{
+		if (Info->SourceFiles.IsValidIndex(Index))
+		{
+			return FText::FromString(Info->SourceFiles[Index].RelativeFilename);
+		}
+	}
+	return FText();
+}
+
+FText FAssetImportDataCustomization::GetTimestampText(int32 Index) const
+{
+	if (FAssetImportInfo* Info = GetEditStruct())
+	{
+		if (Info->SourceFiles.IsValidIndex(Index))
+		{
+			return FText::FromString(Info->SourceFiles[Index].Timestamp.ToString());
+		}
+	}
+	return FText();
+}
+
+
+FAssetImportInfo* FAssetImportDataCustomization::GetEditStruct() const
+{
+	static TArray<FAssetImportInfo*> AssetImportInfo;
+	AssetImportInfo.Reset();
+
+	if(PropertyHandle->IsValidHandle())
+	{
+		PropertyHandle->AccessRawData(reinterpret_cast<TArray<void*>&>(AssetImportInfo));
+	}
+
+	if (AssetImportInfo.Num() == 1)
+	{
+		return AssetImportInfo[0];
+	}
+	return nullptr;
+}
+
+UAssetImportData* FAssetImportDataCustomization::GetOuterClass() const
+{
+	static TArray<UObject*> OuterObjects;
+	OuterObjects.Reset();
+
+	PropertyHandle->GetOuterObjects(OuterObjects);
+
+	return OuterObjects.Num() ? Cast<UAssetImportData>(OuterObjects[0]) : nullptr;
+}
+
+FReply FAssetImportDataCustomization::OnChangePathClicked(int32 Index) const
+{
+	UAssetImportData* ImportData = GetOuterClass();
+	UObject* Obj = ImportData ? ImportData->GetOuter() : nullptr;
+
+	if (!Obj)
+	{
+		return FReply::Handled();
+	}
+
+	TArray<FString> OpenFilenames;
+	FReimportManager::Instance()->GetNewReimportPath(Obj, OpenFilenames);
+	if (OpenFilenames.Num() == 1)
+	{
+		ImportData->UpdateFilenameOnly(FPaths::ConvertRelativePathToFull(OpenFilenames[0]), Index);
+	}
+	return FReply::Handled();
+}
+
+
+#undef LOCTEXT_NAMESPACE
diff --git a/Engine/Source/Editor/DetailCustomizations/Private/AssetImportDataCustomization.h b/Engine/Source/Editor/DetailCustomizations/Private/AssetImportDataCustomization.h
new file mode 100644
--- /dev/null
+++ b/Engine/Source/Editor/DetailCustomizations/Private/AssetImportDataCustomization.h
@@ -0,0 +1,37 @@
+// Copyright 1998-2015 Epic Games, Inc. All Rights Reserved.
+
+#pragma once
+
+class FAssetImportDataCustomization : public IPropertyTypeCustomization
+{
+public:
+	/** Makes a new instance of this detail layout class for a specific detail view requesting it */
+	static TSharedRef<IPropertyTypeCustomization> MakeInstance();
+
+	/** IDetailCustomization interface */
+	virtual void CustomizeHeader( TSharedRef<IPropertyHandle> PropertyHandle, FDetailWidgetRow& HeaderRow, IPropertyTypeCustomizationUtils& CustomizationUtils ) override;
+	virtual void CustomizeChildren( TSharedRef<IPropertyHandle> PropertyHandle, IDetailChildrenBuilder& ChildBuilder, IPropertyTypeCustomizationUtils& CustomizationUtils ) override;
+
+private:
+
+	/** Handle the user wanting to change a specific source path */
+	FReply OnChangePathClicked(int32 Index) const;
+
+	/** Update the filename strings */
+	void UpdateCachedData();
+
+	/** Access the struct we are editing */
+	FAssetImportInfo* GetEditStruct() const;
+
+	/** Access the outer class that contains this struct */
+	UAssetImportData* GetOuterClass() const;
+
+	/** Get text for the UI */
+	FText GetFilenameText(int32 Index) const;
+	FText GetTimestampText(int32 Index) const;
+
+private:
+
+	/** Property handle of the property we're editing */
+	TSharedPtr<IPropertyHandle> PropertyHandle;
+};
diff --git a/Engine/Source/Editor/DetailCustomizations/Private/DetailCustomizations.cpp b/Engine/Source/Editor/DetailCustomizations/Private/DetailCustomizations.cpp
--- a/Engine/Source/Editor/DetailCustomizations/Private/DetailCustomizations.cpp
+++ b/Engine/Source/Editor/DetailCustomizations/Private/DetailCustomizations.cpp
@@ -102,6 +102,7 @@
 #include "HierarchicalSimplificationCustomizations.h"
 #include "PostProcessSettingsCustomization.h"
 #include "ConfigEditorPropertyDetails.h"
+#include "AssetImportDataCustomization.h"
 #include "TextCustomization.h"
 
 IMPLEMENT_MODULE( FDetailCustomizationsModule, DetailCustomizations );
@@ -208,6 +209,7 @@ void FDetailCustomizationsModule::RegisterPropertyTypeCustomizations()
 	RegisterCustomPropertyTypeLayout("DistanceDatum", FOnGetPropertyTypeCustomizationInstance::CreateStatic(&FDistanceDatumStructCustomization::MakeInstance));
 	RegisterCustomPropertyTypeLayout("HierarchicalSimplification", FOnGetPropertyTypeCustomizationInstance::CreateStatic(&FHierarchicalSimplificationCustomizations::MakeInstance));
 	RegisterCustomPropertyTypeLayout("PostProcessSettings", FOnGetPropertyTypeCustomizationInstance::CreateStatic(&FPostProcessSettingsCustomization::MakeInstance));
+	RegisterCustomPropertyTypeLayout("AssetImportInfo", FOnGetPropertyTypeCustomizationInstance::CreateStatic(&FAssetImportDataCustomization::MakeInstance));
 	RegisterCustomPropertyTypeLayout("WeightedBlendable", FOnGetPropertyTypeCustomizationInstance::CreateStatic(&FWeightedBlendableCustomization::MakeInstance));
 }
 
diff --git a/Engine/Source/Editor/StaticMeshEditor/Private/StaticMeshEditorTools.cpp b/Engine/Source/Editor/StaticMeshEditor/Private/StaticMeshEditorTools.cpp
--- a/Engine/Source/Editor/StaticMeshEditor/Private/StaticMeshEditorTools.cpp
+++ b/Engine/Source/Editor/StaticMeshEditor/Private/StaticMeshEditorTools.cpp
@@ -85,26 +85,6 @@ void FStaticMeshDetails::CustomizeDetails( class IDetailLayoutBuilder& DetailBui
 			}
 		}
 	}
-
-	// Only add the reimport button if we have reimport settings we can modify
-	// Note: this will get rebuilt if the asset is reimported so we don't need to use .Visibility on the button
-	const UStaticMesh* StaticMesh = StaticMeshEditor.GetStaticMesh();
-	check(StaticMesh);
-	if (StaticMesh->AssetImportData && StaticMesh->AssetImportData->GetClass() != UAssetImportData::StaticClass())
-	{
-		ImportCategory.AddCustomRow(LOCTEXT("ReimportStaticMesh", "Reimport Static Mesh"), true)
-			.ValueContent()
-			[
-				SNew(SButton)
-				.OnClicked(this, &FStaticMeshDetails::Reimport)
-				.IsEnabled(this, &FStaticMeshDetails::CanReimport)
-				[
-					SNew(STextBlock)
-					.Text(LOCTEXT("Reimport", "Reimport"))
-					.Font(IDetailLayoutBuilder::GetDetailFont())
-				]
-			];
-	}
 }
 
 void SConvexDecomposition::Construct(const FArguments& InArgs)
@@ -207,25 +187,6 @@ void FStaticMeshDetails::ApplyChanges()
 	}
 }
 
-FReply FStaticMeshDetails::Reimport()
-{
-	UStaticMesh* StaticMesh = StaticMeshEditor.GetStaticMesh();
-	check(StaticMesh);
-	FReimportManager::Instance()->Reimport(StaticMesh, true);
-	return FReply::Handled();
-}
-
-bool FStaticMeshDetails::CanReimport() const
-{
-	const UStaticMesh* StaticMesh = StaticMeshEditor.GetStaticMesh();
-	check(StaticMesh);
-	if (StaticMesh->AssetImportData)
-	{
-		return StaticMesh->AssetImportData->bDirty;
-	}
-	return false;
-}
-
 SConvexDecomposition::~SConvexDecomposition()
 {
 
diff --git a/Engine/Source/Editor/StaticMeshEditor/Private/StaticMeshEditorTools.h b/Engine/Source/Editor/StaticMeshEditor/Private/StaticMeshEditorTools.h
--- a/Engine/Source/Editor/StaticMeshEditor/Private/StaticMeshEditorTools.h
+++ b/Engine/Source/Editor/StaticMeshEditor/Private/StaticMeshEditorTools.h
@@ -34,12 +34,6 @@ public:
 
 	/** Applies level of detail changes to the static mesh */
 	void ApplyChanges();
-
-	/** Reimports the current static mesh */
-	FReply Reimport();
-
-	/** Whether the user should reimport based on the changes */
-	bool CanReimport() const;
 private:
 	/** Level of detail settings for the details panel */
 	TSharedPtr<FLevelOfDetailSettingsLayout> LevelOfDetailSettings;
diff --git a/Engine/Source/Editor/UnrealEd/Private/AutoReimport/AssetSourceFilenameCache.cpp b/Engine/Source/Editor/UnrealEd/Private/AutoReimport/AssetSourceFilenameCache.cpp
--- a/Engine/Source/Editor/UnrealEd/Private/AutoReimport/AssetSourceFilenameCache.cpp
+++ b/Engine/Source/Editor/UnrealEd/Private/AutoReimport/AssetSourceFilenameCache.cpp
@@ -97,7 +97,7 @@ void FAssetSourceFilenameCache::HandleOnAssetAdded(const FAssetData& AssetData)
 	TOptional<FAssetImportInfo> ImportData = ExtractAssetImportInfo(AssetData.TagsAndValues);
 	if (ImportData.IsSet())
 	{
-		for (const auto& SourceFile : ImportData->GetSourceFileData())
+		for (const auto& SourceFile : ImportData->SourceFiles)
 		{
 			SourceFileToObjectPathCache.FindOrAdd(FPaths::GetCleanFilename(SourceFile.RelativeFilename)).Add(AssetData.ObjectPath);
 		}
@@ -109,7 +109,7 @@ void FAssetSourceFilenameCache::HandleOnAssetRemoved(const FAssetData& AssetData
 	TOptional<FAssetImportInfo> ImportData = ExtractAssetImportInfo(AssetData.TagsAndValues);
 	if (ImportData.IsSet())
 	{
-		for (auto& SourceFile : ImportData->GetSourceFileData())
+		for (auto& SourceFile : ImportData->SourceFiles)
 		{
 			FString CleanFilename = FPaths::GetCleanFilename(SourceFile.RelativeFilename);
 			if (auto* Objects = SourceFileToObjectPathCache.Find(CleanFilename))
@@ -131,7 +131,7 @@ void FAssetSourceFilenameCache::HandleOnAssetRenamed(const FAssetData& AssetData
 	{
 		FName OldPathName = *OldPath;
 
-		for (auto& SourceFile : ImportData->GetSourceFileData())
+		for (auto& SourceFile : ImportData->SourceFiles)
 		{
 			FString CleanFilename = FPaths::GetCleanFilename(SourceFile.RelativeFilename);
 
@@ -155,7 +155,7 @@ void FAssetSourceFilenameCache::HandleOnAssetUpdated(const FAssetImportInfo& Old
 {
 	FName ObjectPath = FName(*ImportData->GetOuter()->GetPathName());
 
-	for (const auto& SourceFile : OldData.GetSourceFileData())
+	for (const auto& SourceFile : OldData.SourceFiles)
 	{
 		FString CleanFilename = FPaths::GetCleanFilename(SourceFile.RelativeFilename);
 		if (auto* Objects = SourceFileToObjectPathCache.Find(CleanFilename))
@@ -168,7 +168,7 @@ void FAssetSourceFilenameCache::HandleOnAssetUpdated(const FAssetImportInfo& Old
 		}
 	}
 
-	for (const auto& SourceFile : ImportData->GetSourceFileData())
+	for (const auto& SourceFile : ImportData->SourceData.SourceFiles)
 	{
 		SourceFileToObjectPathCache.FindOrAdd(FPaths::GetCleanFilename(SourceFile.RelativeFilename)).Add(ObjectPath);
 	}
@@ -189,7 +189,7 @@ TArray<FAssetData> FAssetSourceFilenameCache::GetAssetsPertainingToFile(const IA
 				auto AssetPackagePath = FPackageName::LongPackageNameToFilename(Asset.PackagePath.ToString() / TEXT(""));
 
 				// Attempt to find the matching source filename in the list of imported sorce files (generally there are only one of these)
-				const bool bWasImportedFromFile = ImportInfo->GetSourceFileData().ContainsByPredicate([&](const FAssetImportInfo::FSourceFile& File){
+				const bool bWasImportedFromFile = ImportInfo->SourceFiles.ContainsByPredicate([&](const FAssetImportInfo::FSourceFile& File){
 
 					if (AbsoluteFilename == FPaths::ConvertRelativePathToFull(AssetPackagePath / File.RelativeFilename) ||
 						AbsoluteFilename == FPaths::ConvertRelativePathToFull(File.RelativeFilename))
diff --git a/Engine/Source/Editor/UnrealEd/Private/AutoReimport/AutoReimportManager.cpp b/Engine/Source/Editor/UnrealEd/Private/AutoReimport/AutoReimportManager.cpp
--- a/Engine/Source/Editor/UnrealEd/Private/AutoReimport/AutoReimportManager.cpp
+++ b/Engine/Source/Editor/UnrealEd/Private/AutoReimport/AutoReimportManager.cpp
@@ -361,7 +361,7 @@ void FAutoReimportManager::OnAssetRenamed(const FAssetData& AssetData, const FSt
 		return;
 	}
 
-	const TArray<FAssetImportInfo::FSourceFile>& SourceFileData = ImportInfo->GetSourceFileData();
+	const TArray<FAssetImportInfo::FSourceFile>& SourceFileData = ImportInfo->SourceFiles;
 
 	// We move the file with the asset provided it is the only file referenced, and sits right beside the uasset file
 	if (SourceFileData.Num() == 1 && !SourceFileData[0].RelativeFilename.GetCharArray().ContainsByPredicate([](const TCHAR Char) { return Char == '/' || Char == '\\'; }))
diff --git a/Engine/Source/Editor/UnrealEd/Private/AutoReimport/AutoReimportUtilities.cpp b/Engine/Source/Editor/UnrealEd/Private/AutoReimport/AutoReimportUtilities.cpp
--- a/Engine/Source/Editor/UnrealEd/Private/AutoReimport/AutoReimportUtilities.cpp
+++ b/Engine/Source/Editor/UnrealEd/Private/AutoReimport/AutoReimportUtilities.cpp
@@ -36,7 +36,7 @@ namespace Utils
 				TOptional<FAssetImportInfo> ImportInfo = FAssetImportInfo::FromJson(Tag.Value);
 				if (ImportInfo.IsSet())
 				{
-					for (const auto& File : ImportInfo->GetSourceFileData())
+					for (const auto& File : ImportInfo->SourceFiles)
 					{
 						OutSourceFiles.Emplace(UAssetImportData::ResolveImportFilename(File.RelativeFilename, Object->GetOutermost()));
 					}
diff --git a/Engine/Source/Editor/UnrealEd/Private/AutoReimport/ContentDirectoryMonitor.cpp b/Engine/Source/Editor/UnrealEd/Private/AutoReimport/ContentDirectoryMonitor.cpp
--- a/Engine/Source/Editor/UnrealEd/Private/AutoReimport/ContentDirectoryMonitor.cpp
+++ b/Engine/Source/Editor/UnrealEd/Private/AutoReimport/ContentDirectoryMonitor.cpp
@@ -370,9 +370,9 @@ void FContentDirectoryMonitor::ProcessModifications(const IAssetRegistry& Regist
 
 					// Check if the source file that this asset last imported was the same as the one we're going to reimport.
 					// If it is, there's no reason to auto-reimport it
-					if (Info.IsSet() && Info->GetSourceFileData().Num() == 1)
+					if (Info.IsSet() && Info->SourceFiles.Num() == 1)
 					{
-						if (Info->GetSourceFileData()[0].FileHash == Change.FileData.FileHash)
+						if (Info->SourceFiles[0].FileHash == Change.FileData.FileHash)
 						{
 							continue;
 						}
diff --git a/Engine/Source/Editor/UnrealEd/Private/Factories/ApexDestructibleAssetImport.cpp b/Engine/Source/Editor/UnrealEd/Private/Factories/ApexDestructibleAssetImport.cpp
--- a/Engine/Source/Editor/UnrealEd/Private/Factories/ApexDestructibleAssetImport.cpp
+++ b/Engine/Source/Editor/UnrealEd/Private/Factories/ApexDestructibleAssetImport.cpp
@@ -1006,7 +1006,6 @@ UDestructibleMesh* ImportDestructibleMeshFromApexDestructibleAsset(UObject* InPa
 		// @todo AssetImportData make a data class for Apex destructible assets
 		DestructibleMesh->AssetImportData = NewObject<UAssetImportData>(DestructibleMesh);
 		DestructibleMesh->AssetImportData->Update(UFactory::CurrentFilename);
-		DestructibleMesh->AssetImportData->bDirty = false;
 	}
 
 	DestructibleMesh->PreEditChange(NULL);
diff --git a/Engine/Source/Editor/UnrealEd/Private/Factories/EditorFactories.cpp b/Engine/Source/Editor/UnrealEd/Private/Factories/EditorFactories.cpp
--- a/Engine/Source/Editor/UnrealEd/Private/Factories/EditorFactories.cpp
+++ b/Engine/Source/Editor/UnrealEd/Private/Factories/EditorFactories.cpp
@@ -1897,6 +1897,8 @@ EReimportResult::Type UReimportSoundFactory::Reimport( UObject* Obj )
 	{
 		UE_LOG(LogEditorFactories, Log, TEXT("-- imported successfully") );
 
+		SoundWave->AssetImportData->Update(Filename);
+
 		// Mark the package dirty after the successful import
 		SoundWave->MarkPackageDirty();
 	}
@@ -5513,6 +5515,9 @@ EReimportResult::Type UReimportTextureFactory::Reimport( UObject* Obj )
 	if (UFactory::StaticImportObject(pTex->GetClass(), pTex->GetOuter(), *pTex->GetName(), RF_Public|RF_Standalone, *ResolvedSourceFilePath, nullptr, this))
 	{
 		UE_LOG(LogEditorFactories, Log, TEXT("-- imported successfully") );
+
+		pTex->AssetImportData->Update(ResolvedSourceFilePath);
+
 		// Try to find the outer package so we can dirty it up
 		if (pTex->GetOuter())
 		{
@@ -5629,7 +5634,7 @@ EReimportResult::Type UReimportFbxStaticMeshFactory::Reimport( UObject* Obj )
 
 	if( !bOperationCanceled && ensure(ImportData) )
 	{
-		const FString& Filename = ImportData->GetFirstFilename();
+		const FString Filename = ImportData->GetFirstFilename();
 		const FString FileExtension = FPaths::GetExtension(Filename);
 		const bool bIsValidFile = FileExtension.Equals( TEXT("fbx"), ESearchCase::IgnoreCase ) || FileExtension.Equals( "obj",  ESearchCase::IgnoreCase );
 
@@ -5699,6 +5704,8 @@ EReimportResult::Type UReimportFbxStaticMeshFactory::Reimport( UObject* Obj )
 				Mesh->PositiveBoundsExtension = PositiveBoundsExtension;
 				Mesh->NegativeBoundsExtension = NegativeBoundsExtension;
 
+				Mesh->AssetImportData->Update(Filename);
+
 				// Try to find the outer package so we can dirty it up
 				if (Mesh->GetOuter())
 				{
@@ -5855,6 +5862,8 @@ EReimportResult::Type UReimportFbxSkeletalMeshFactory::Reimport( UObject* Obj )
 			{
 				UE_LOG(LogEditorFactories, Log, TEXT("-- imported successfully") );
 
+				SkeletalMesh->AssetImportData->Update(Filename);
+
 				// Try to find the outer package so we can dirty it up
 				if (SkeletalMesh->GetOuter())
 				{
@@ -6021,7 +6030,6 @@ EReimportResult::Type UReimportFbxAnimSequenceFactory::Reimport( UObject* Obj )
 
 		// update the data in case the file source has changed
 		ImportData->Update(UFactory::CurrentFilename);
-		ImportData->bDirty = false;
 
 		// Try to find the outer package so we can dirty it up
 		if (AnimSequence->GetOuter())
diff --git a/Engine/Source/Editor/UnrealEd/Private/Factories/VectorFieldFactory.cpp b/Engine/Source/Editor/UnrealEd/Private/Factories/VectorFieldFactory.cpp
--- a/Engine/Source/Editor/UnrealEd/Private/Factories/VectorFieldFactory.cpp
+++ b/Engine/Source/Editor/UnrealEd/Private/Factories/VectorFieldFactory.cpp
@@ -305,13 +305,13 @@ EReimportResult::Type UReimportVectorFieldStaticFactory::Reimport( UObject* Obj
 
 	UVectorFieldStatic* VectorFieldStatic = Cast<UVectorFieldStatic>(Obj);
 
-	if ( VectorFieldStatic->AssetImportData->GetSourceFileData().Num() != 1 )
+	if ( VectorFieldStatic->AssetImportData->SourceData.SourceFiles.Num() != 1 )
 	{
 		// No source art path. Can't reimport.
 		return EReimportResult::Failed;
 	}
 
-	const FString& ReImportFilename = VectorFieldStatic->AssetImportData->GetSourceFileData()[0].RelativeFilename;
+	FString ReImportFilename = VectorFieldStatic->AssetImportData->GetFirstFilename();
 
 	UE_LOG(LogVectorFieldFactory, Log, TEXT("Performing atomic reimport of [%s]"), *ReImportFilename);
 
@@ -325,6 +325,7 @@ EReimportResult::Type UReimportVectorFieldStaticFactory::Reimport( UObject* Obj
 	if (UFactory::StaticImportObject(VectorFieldStatic->GetClass(), VectorFieldStatic->GetOuter(), *VectorFieldStatic->GetName(), RF_Public|RF_Standalone, *ReImportFilename, NULL, this))
 	{
 		UE_LOG(LogVectorFieldFactory, Log, TEXT("Reimported successfully") );
+		VectorFieldStatic->AssetImportData->Update(ReImportFilename);
 		VectorFieldStatic->MarkPackageDirty();
 	}
 	else
diff --git a/Engine/Source/Editor/UnrealEd/Private/Fbx/FbxAnimSequenceImportData.cpp b/Engine/Source/Editor/UnrealEd/Private/Fbx/FbxAnimSequenceImportData.cpp
--- a/Engine/Source/Editor/UnrealEd/Private/Fbx/FbxAnimSequenceImportData.cpp
+++ b/Engine/Source/Editor/UnrealEd/Private/Fbx/FbxAnimSequenceImportData.cpp
@@ -18,10 +18,10 @@ UFbxAnimSequenceImportData* UFbxAnimSequenceImportData::GetImportDataForAnimSequ
 	{
 		ImportData = NewObject<UFbxAnimSequenceImportData>(AnimSequence, NAME_None, RF_NoFlags, TemplateForCreation);
 
-		// Try to preserve the source file path if possible
+		// Try to preserve the source file data if possible
 		if ( AnimSequence->AssetImportData != NULL )
 		{
-			ImportData->CopyFrom(*AnimSequence->AssetImportData);
+			ImportData->SourceData = AnimSequence->AssetImportData->SourceData;
 		}
 
 		AnimSequence->AssetImportData = ImportData;
diff --git a/Engine/Source/Editor/UnrealEd/Private/Fbx/FbxSkeletalMeshImport.cpp b/Engine/Source/Editor/UnrealEd/Private/Fbx/FbxSkeletalMeshImport.cpp
--- a/Engine/Source/Editor/UnrealEd/Private/Fbx/FbxSkeletalMeshImport.cpp
+++ b/Engine/Source/Editor/UnrealEd/Private/Fbx/FbxSkeletalMeshImport.cpp
@@ -1325,7 +1325,6 @@ USkeletalMesh* UnFbx::FFbxImporter::ImportSkeletalMesh(UObject* InParent, TArray
 		// Store the current file path and timestamp for re-import purposes
 		UFbxSkeletalMeshImportData* ImportData = UFbxSkeletalMeshImportData::GetImportDataForSkeletalMesh(SkeletalMesh, TemplateImportData);
 		SkeletalMesh->AssetImportData->Update(UFactory::CurrentFilename);
-		SkeletalMesh->AssetImportData->bDirty = false;
 
 		SkeletalMesh->CalculateInvRefMatrices();
 		SkeletalMesh->PostEditChange();
diff --git a/Engine/Source/Editor/UnrealEd/Private/Fbx/FbxSkeletalMeshImportData.cpp b/Engine/Source/Editor/UnrealEd/Private/Fbx/FbxSkeletalMeshImportData.cpp
--- a/Engine/Source/Editor/UnrealEd/Private/Fbx/FbxSkeletalMeshImportData.cpp
+++ b/Engine/Source/Editor/UnrealEd/Private/Fbx/FbxSkeletalMeshImportData.cpp
@@ -17,10 +17,10 @@ UFbxSkeletalMeshImportData* UFbxSkeletalMeshImportData::GetImportDataForSkeletal
 	{
 		ImportData = NewObject<UFbxSkeletalMeshImportData>(SkeletalMesh, NAME_None, RF_NoFlags, TemplateForCreation);
 
-		// Try to preserve the source file path if possible
+		// Try to preserve the source file data if possible
 		if ( SkeletalMesh->AssetImportData != NULL )
 		{
-			ImportData->CopyFrom(*SkeletalMesh->AssetImportData);
+			ImportData->SourceData = SkeletalMesh->AssetImportData->SourceData;
 		}
 
 		SkeletalMesh->AssetImportData = ImportData;
diff --git a/Engine/Source/Editor/UnrealEd/Private/Fbx/FbxStaticMeshImport.cpp b/Engine/Source/Editor/UnrealEd/Private/Fbx/FbxStaticMeshImport.cpp
--- a/Engine/Source/Editor/UnrealEd/Private/Fbx/FbxStaticMeshImport.cpp
+++ b/Engine/Source/Editor/UnrealEd/Private/Fbx/FbxStaticMeshImport.cpp
@@ -1207,7 +1207,6 @@ UStaticMesh* UnFbx::FFbxImporter::ImportStaticMeshAsSingle(UObject* InParent, TA
 
 		UFbxStaticMeshImportData* ImportData = UFbxStaticMeshImportData::GetImportDataForStaticMesh(StaticMesh, TemplateImportData);
 		ImportData->Update(UFactory::CurrentFilename);
-		ImportData->bDirty = false;
 
 		// @todo This overrides restored values currently but we need to be able to import over the existing settings if the user chose to do so.
 		SrcModel.BuildSettings.bRemoveDegenerates = ImportOptions->bRemoveDegenerates;
diff --git a/Engine/Source/Editor/UnrealEd/Private/Fbx/FbxStaticMeshImportData.cpp b/Engine/Source/Editor/UnrealEd/Private/Fbx/FbxStaticMeshImportData.cpp
--- a/Engine/Source/Editor/UnrealEd/Private/Fbx/FbxStaticMeshImportData.cpp
+++ b/Engine/Source/Editor/UnrealEd/Private/Fbx/FbxStaticMeshImportData.cpp
@@ -27,7 +27,7 @@ UFbxStaticMeshImportData* UFbxStaticMeshImportData::GetImportDataForStaticMesh(U
 		// Try to preserve the source file path if possible
 		if ( StaticMesh->AssetImportData != NULL )
 		{
-			ImportData->CopyFrom(*StaticMesh->AssetImportData);
+			ImportData->SourceData = StaticMesh->AssetImportData->SourceData;
 		}
 
 		StaticMesh->AssetImportData = ImportData;
diff --git a/Engine/Source/Editor/UnrealEd/Private/SkeletalMeshEdit.cpp b/Engine/Source/Editor/UnrealEd/Private/SkeletalMeshEdit.cpp
--- a/Engine/Source/Editor/UnrealEd/Private/SkeletalMeshEdit.cpp
+++ b/Engine/Source/Editor/UnrealEd/Private/SkeletalMeshEdit.cpp
@@ -78,7 +78,6 @@ UAnimSequence * UEditorEngine::ImportFbxAnimation( USkeleton* Skeleton, UObject*
 				// since to know full path, reimport will need to do same
 				UFbxAnimSequenceImportData* ImportData = UFbxAnimSequenceImportData::GetImportDataForAnimSequence(NewAnimation, TemplateImportData);
 				ImportData->Update(UFactory::CurrentFilename);
-				ImportData->bDirty = false;
 			}
 		}
 	}
@@ -549,7 +548,6 @@ UAnimSequence * UnFbx::FFbxImporter::ImportAnimations(USkeleton* Skeleton, UObje
 		// since to know full path, reimport will need to do same
 		UFbxAnimSequenceImportData* ImportData = UFbxAnimSequenceImportData::GetImportDataForAnimSequence(DestSeq, TemplateImportData);
 		ImportData->Update(UFactory::CurrentFilename);
-		ImportData->bDirty = false;
 
 		ImportAnimation(Skeleton, DestSeq, Name, SortedLinks, NodeArray, CurAnimStack, ResampleRate, AnimTimeSpan);
 
diff --git a/Engine/Source/Editor/UnrealEd/Private/UnrealEdSrv.cpp b/Engine/Source/Editor/UnrealEd/Private/UnrealEdSrv.cpp
--- a/Engine/Source/Editor/UnrealEd/Private/UnrealEdSrv.cpp
+++ b/Engine/Source/Editor/UnrealEd/Private/UnrealEdSrv.cpp
@@ -1106,7 +1106,7 @@ bool UUnrealEdEngine::Exec( UWorld* InWorld, const TCHAR* Stream, FOutputDevice&
 				FAssetImportInfo AssetImportInfo;
 
 				bool bModified = false;
-				for (const auto& File : Data->GetSourceFileData())
+				for (const auto& File : Data->SourceData.SourceFiles)
 				{
 					if( !File.RelativeFilename.IsEmpty() && !SearchTerms.ContainsByPredicate([&](const FString& SearchTerm){ return File.RelativeFilename.Contains(SearchTerm); }) )
 					{
@@ -1122,7 +1122,7 @@ bool UUnrealEdEngine::Exec( UWorld* InWorld, const TCHAR* Stream, FOutputDevice&
 				if (bModified)
 				{
 					Data->Modify();
-					Data->CopyFrom(AssetImportInfo);
+					Data->SourceData = AssetImportInfo;
 					return true;
 				}
 
diff --git a/Engine/Source/Editor/UnrealEd/Public/EditorReimportHandler.h b/Engine/Source/Editor/UnrealEd/Public/EditorReimportHandler.h
--- a/Engine/Source/Editor/UnrealEd/Public/EditorReimportHandler.h
+++ b/Engine/Source/Editor/UnrealEd/Public/EditorReimportHandler.h
@@ -76,9 +76,8 @@ public:
 	DECLARE_MULTICAST_DELEGATE_TwoParams(FPostReimportNotification, UObject*, bool);
 	FPostReimportNotification& OnPostReimport(){ return PostReimport; }
 
-private:
 	/** Opens a file dialog to request a new reimport path */
-	void GetNewReimportPath(UObject* Obj, TArray<FString>& InOutFilenames);
+	UNREALED_API void GetNewReimportPath(UObject* Obj, TArray<FString>& InOutFilenames);
 
 private:
 	/** Reimport handlers registered with this manager */
diff --git a/Engine/Source/Runtime/Engine/Classes/Animation/AnimSequence.h b/Engine/Source/Runtime/Engine/Classes/Animation/AnimSequence.h
--- a/Engine/Source/Runtime/Engine/Classes/Animation/AnimSequence.h
+++ b/Engine/Source/Runtime/Engine/Classes/Animation/AnimSequence.h
@@ -508,7 +508,7 @@ class ENGINE_API UAnimSequence : public UAnimSequenceBase
 	uint32 bWasCompressedWithoutTranslations:1;
 
 	/** Importing data and options used for this mesh */
-	UPROPERTY(EditAnywhere, Instanced, Category=Reimport)
+	UPROPERTY(EditAnywhere, Instanced, Category=ImportSettings)
 	class UAssetImportData* AssetImportData;
 
 	/***  for Reimport **/
diff --git a/Engine/Source/Runtime/Engine/Classes/EditorFramework/AssetImportData.h b/Engine/Source/Runtime/Engine/Classes/EditorFramework/AssetImportData.h
--- a/Engine/Source/Runtime/Engine/Classes/EditorFramework/AssetImportData.h
+++ b/Engine/Source/Runtime/Engine/Classes/EditorFramework/AssetImportData.h
@@ -5,14 +5,26 @@
 #include "AssetImportData.generated.h"
 
 /** Struct that is used to store an array of asset import data as an asset registry tag */
+USTRUCT()
 struct ENGINE_API FAssetImportInfo
 {
 #if WITH_EDITORONLY_DATA
+
+	GENERATED_BODY()
+
+	FAssetImportInfo() {}
+
+	FAssetImportInfo(const FAssetImportInfo& In) : SourceFiles(In.SourceFiles) { }
+	FAssetImportInfo& operator=(const FAssetImportInfo& In) { SourceFiles = In.SourceFiles; return *this; }
+
+	FAssetImportInfo(FAssetImportInfo&& In) : SourceFiles(MoveTemp(In.SourceFiles)) { }
+	FAssetImportInfo& operator=(FAssetImportInfo&& In) { SourceFiles = MoveTemp(In.SourceFiles); return *this; }
+
 	struct FSourceFile
 	{
 		FSourceFile(FString InRelativeFilename, const FDateTime& InTimestamp = 0, const FMD5Hash& InFileHash = FMD5Hash())
 			: RelativeFilename(MoveTemp(InRelativeFilename))
-			, Timestamp(0)
+			, Timestamp(InTimestamp)
 			, FileHash(InFileHash)
 		{}
 
@@ -35,31 +47,22 @@ struct ENGINE_API FAssetImportInfo
 	/** Insert information pertaining to a new source file to this structure */
 	void Insert(const FSourceFile& InSourceFile) { SourceFiles.Add(InSourceFile); }
 
-	/** Copy the information from another asset import info structure into this one. */
-	void CopyFrom(const FAssetImportInfo& Other) { SourceFiles = Other.SourceFiles; }
-
-	/** Const public access to the source file data */
-	const TArray<FSourceFile>& GetSourceFileData() const { return SourceFiles; }
-
-protected:
-
 	/** Array of information pertaining to the source files that this asset was imported from */
 	TArray<FSourceFile> SourceFiles;
+
 #endif // WITH_EDITORONLY_DATA
 };
 
 DECLARE_MULTICAST_DELEGATE_TwoParams(FOnImportDataChanged, const FAssetImportInfo& /*OldData*/, const class UAssetImportData* /* NewData */);
 
+/* todo: Make this class better suited to multiple import paths - maybe have FAssetImportInfo use a map rather than array? */
 UCLASS()
-class ENGINE_API UAssetImportData : public UObject, public FAssetImportInfo
+class ENGINE_API UAssetImportData : public UObject
 {
 public:
 	GENERATED_UCLASS_BODY()
 
 #if WITH_EDITORONLY_DATA
-	/** If true, the settings have been modified but not reimported yet */
-	UPROPERTY()
-	uint32 bDirty : 1;
 
 	/** Path to the resource used to construct this static mesh. Relative to the object's package, BaseDir() or absolute */
 	UPROPERTY()
@@ -69,6 +72,10 @@ public:
 	UPROPERTY()
 	FString SourceFileTimestamp_DEPRECATED;
 
+	/** Source file data describing the files that were used to import this asset. */
+	UPROPERTY(VisibleAnywhere, Category=ImportSettings)
+	FAssetImportInfo SourceData;
+
 public:
 
 	/** Static event that is broadcast whenever any asset has updated its import data */
@@ -77,12 +84,18 @@ public:
 	/** Update this import data using the specified file. Called when an asset has been imported from a file. */
 	void Update(const FString& AbsoluteFilename);
 
-	/** Update this import data using the specified filename. Will not update the imported timestamp or MD5. */
+	/** Update this import data using the specified filename. Will not update the imported timestamp or MD5 (so we can update files when they move). */
 	void UpdateFilenameOnly(const FString& InPath);
 
+	/** Update the filename at the specific specified index. Will not update the imported timestamp or MD5 (so we can update files when they move). */
+	void UpdateFilenameOnly(const FString& InPath, int32 Index);
+
 	/** Helper function to return the first filename stored in this data. The resulting filename will be absolute (ie, not relative to the asset).  */
 	FString GetFirstFilename() const;
 
+	/** Const access to the source file data */
+	const FAssetImportInfo& GetSourceData() const { return SourceData; }
+
 	/** Extract all the (resolved) filenames from this data  */
 	void ExtractFilenames(TArray<FString>& AbsoluteFilenames) const;
 
@@ -104,7 +117,6 @@ protected:
 
 	/** Overridden serialize function to write out the underlying data as json */
 	virtual void Serialize(FArchive& Ar) override;
-	virtual void PostEditChangeProperty(struct FPropertyChangedEvent& PropertyChangedEvent) override;
 	
 #endif		// WITH_EDITORONLY_DATA
 };
diff --git a/Engine/Source/Runtime/Engine/Classes/Engine/SkeletalMesh.h b/Engine/Source/Runtime/Engine/Classes/Engine/SkeletalMesh.h
--- a/Engine/Source/Runtime/Engine/Classes/Engine/SkeletalMesh.h
+++ b/Engine/Source/Runtime/Engine/Classes/Engine/SkeletalMesh.h
@@ -607,7 +607,7 @@ public:
 #if WITH_EDITORONLY_DATA
 
 	/** Importing data and options used for this mesh */
-	UPROPERTY(EditAnywhere, Instanced, Category = Reimport)
+	UPROPERTY(VisibleAnywhere, Instanced, Category = ImportSettings)
 	class UAssetImportData* AssetImportData;
 
 	/** Path to the resource used to construct this skeletal mesh */
diff --git a/Engine/Source/Runtime/Engine/Private/Animation/AnimSequence.cpp b/Engine/Source/Runtime/Engine/Private/Animation/AnimSequence.cpp
--- a/Engine/Source/Runtime/Engine/Private/Animation/AnimSequence.cpp
+++ b/Engine/Source/Runtime/Engine/Private/Animation/AnimSequence.cpp
@@ -145,7 +145,7 @@ void UAnimSequence::GetAssetRegistryTags(TArray<FAssetRegistryTag>& OutTags) con
 #if WITH_EDITORONLY_DATA
 	if (AssetImportData)
 	{
-		OutTags.Add( FAssetRegistryTag(SourceFileTagName(), AssetImportData->ToJson(), FAssetRegistryTag::TT_Hidden) );
+		OutTags.Add( FAssetRegistryTag(SourceFileTagName(), AssetImportData->GetSourceData().ToJson(), FAssetRegistryTag::TT_Hidden) );
 	}
 #endif
 
@@ -308,7 +308,7 @@ void UAnimSequence::Serialize(FArchive& Ar)
 		// AssetImportData should always have been set up in the constructor where this is relevant
 		FAssetImportInfo Info;
 		Info.Insert(FAssetImportInfo::FSourceFile(SourceFilePath_DEPRECATED));
-		AssetImportData->CopyFrom(Info);
+		AssetImportData->SourceData = MoveTemp(Info);
 		
 		SourceFilePath_DEPRECATED = TEXT("");
 		SourceFileTimestamp_DEPRECATED = TEXT("");
diff --git a/Engine/Source/Runtime/Engine/Private/CurveBase.cpp b/Engine/Source/Runtime/Engine/Private/CurveBase.cpp
--- a/Engine/Source/Runtime/Engine/Private/CurveBase.cpp
+++ b/Engine/Source/Runtime/Engine/Private/CurveBase.cpp
@@ -1097,7 +1097,7 @@ void UCurveBase::GetAssetRegistryTags(TArray<FAssetRegistryTag>& OutTags) const
 {
 	if (AssetImportData)
 	{
-		OutTags.Add(FAssetRegistryTag(SourceFileTagName(), AssetImportData->ToJson(), FAssetRegistryTag::TT_Hidden));
+		OutTags.Add(FAssetRegistryTag(SourceFileTagName(), AssetImportData->GetSourceData().ToJson(), FAssetRegistryTag::TT_Hidden));
 	}
 
 	Super::GetAssetRegistryTags(OutTags);
@@ -1110,7 +1110,7 @@ void UCurveBase::PostLoad()
 	{
 		FAssetImportInfo Info;
 		Info.Insert(FAssetImportInfo::FSourceFile(ImportPath_DEPRECATED));
-		AssetImportData->CopyFrom(Info);
+		AssetImportData->SourceData = MoveTemp(Info);
 	}
 }
 #endif
diff --git a/Engine/Source/Runtime/Engine/Private/CurveTable.cpp b/Engine/Source/Runtime/Engine/Private/CurveTable.cpp
--- a/Engine/Source/Runtime/Engine/Private/CurveTable.cpp
+++ b/Engine/Source/Runtime/Engine/Private/CurveTable.cpp
@@ -101,7 +101,7 @@ void UCurveTable::GetAssetRegistryTags(TArray<FAssetRegistryTag>& OutTags) const
 {
 	if (AssetImportData)
 	{
-		OutTags.Add( FAssetRegistryTag(SourceFileTagName(), AssetImportData->ToJson(), FAssetRegistryTag::TT_Hidden) );
+		OutTags.Add( FAssetRegistryTag(SourceFileTagName(), AssetImportData->GetSourceData().ToJson(), FAssetRegistryTag::TT_Hidden) );
 	}
 
 	Super::GetAssetRegistryTags(OutTags);
@@ -114,7 +114,7 @@ void UCurveTable::PostLoad()
 	{
 		FAssetImportInfo Info;
 		Info.Insert(FAssetImportInfo::FSourceFile(ImportPath_DEPRECATED));
-		AssetImportData->CopyFrom(Info);
+		AssetImportData->SourceData = MoveTemp(Info);
 	}
 }
 #endif
diff --git a/Engine/Source/Runtime/Engine/Private/DataTable.cpp b/Engine/Source/Runtime/Engine/Private/DataTable.cpp
--- a/Engine/Source/Runtime/Engine/Private/DataTable.cpp
+++ b/Engine/Source/Runtime/Engine/Private/DataTable.cpp
@@ -163,7 +163,7 @@ void UDataTable::GetAssetRegistryTags(TArray<FAssetRegistryTag>& OutTags) const
 {
 	if (AssetImportData)
 	{
-		OutTags.Add( FAssetRegistryTag(SourceFileTagName(), AssetImportData->ToJson(), FAssetRegistryTag::TT_Hidden) );
+		OutTags.Add( FAssetRegistryTag(SourceFileTagName(), AssetImportData->GetSourceData().ToJson(), FAssetRegistryTag::TT_Hidden) );
 	}
 
 	Super::GetAssetRegistryTags(OutTags);
@@ -176,7 +176,7 @@ void UDataTable::PostLoad()
 	{
 		FAssetImportInfo Info;
 		Info.Insert(FAssetImportInfo::FSourceFile(ImportPath_DEPRECATED));
-		AssetImportData->CopyFrom(Info);
+		AssetImportData->SourceData = MoveTemp(Info);
 	}
 }
 #endif
diff --git a/Engine/Source/Runtime/Engine/Private/DestructibleMesh.cpp b/Engine/Source/Runtime/Engine/Private/DestructibleMesh.cpp
--- a/Engine/Source/Runtime/Engine/Private/DestructibleMesh.cpp
+++ b/Engine/Source/Runtime/Engine/Private/DestructibleMesh.cpp
@@ -647,9 +647,9 @@ ENGINE_API bool UDestructibleMesh::BuildFromStaticMesh( UStaticMesh& StaticMesh
 	SourceStaticMesh = &StaticMesh;
 
 	SourceSMImportTimestamp = FDateTime::MinValue();
-	if ( StaticMesh.AssetImportData && StaticMesh.AssetImportData->GetSourceFileData().Num() == 1 )
+	if ( StaticMesh.AssetImportData && StaticMesh.AssetImportData->SourceData.SourceFiles.Num() == 1 )
 	{
-		SourceSMImportTimestamp = StaticMesh.AssetImportData->GetSourceFileData()[0].Timestamp;
+		SourceSMImportTimestamp = StaticMesh.AssetImportData->SourceData.SourceFiles[0].Timestamp;
 	}
 
 #if WITH_APEX
diff --git a/Engine/Source/Runtime/Engine/Private/EditorFramework/AssetImportData.cpp b/Engine/Source/Runtime/Engine/Private/EditorFramework/AssetImportData.cpp
--- a/Engine/Source/Runtime/Engine/Private/EditorFramework/AssetImportData.cpp
+++ b/Engine/Source/Runtime/Engine/Private/EditorFramework/AssetImportData.cpp
@@ -8,9 +8,6 @@
 UAssetImportData::UAssetImportData(const FObjectInitializer& ObjectInitializer)
 	: Super(ObjectInitializer)
 {
-#if WITH_EDITORONLY_DATA
-	bDirty = false;
-#endif
 }
 
 #if WITH_EDITORONLY_DATA
@@ -87,47 +84,60 @@ TOptional<FAssetImportInfo> FAssetImportInfo::FromJson(FString InJsonString)
 void UAssetImportData::UpdateFilenameOnly(const FString& InPath)
 {
 	// Try and retain the MD5 and timestamp if possible
-	if (SourceFiles.Num() == 1)
+	if (SourceData.SourceFiles.Num() == 1)
 	{
-		SourceFiles[0].RelativeFilename = SanitizeImportFilename(InPath);
+		SourceData.SourceFiles[0].RelativeFilename = SanitizeImportFilename(InPath);
 	}
 	else
 	{
-		SourceFiles.Reset();
-		SourceFiles.Emplace(SanitizeImportFilename(InPath));
+		SourceData.SourceFiles.Reset();
+		SourceData.SourceFiles.Emplace(SanitizeImportFilename(InPath));
+	}
+}
+
+void UAssetImportData::UpdateFilenameOnly(const FString& InPath, int32 Index)
+{
+	if (SourceData.SourceFiles.IsValidIndex(Index))
+	{
+		SourceData.SourceFiles[Index].RelativeFilename = SanitizeImportFilename(InPath);
 	}
 }
 
 void UAssetImportData::Update(const FString& InPath)
 {
-	FAssetImportInfo Old;
-	Old.CopyFrom(*this);
+	FAssetImportInfo Old = SourceData;
 
-	SourceFiles.Reset();
-	SourceFiles.Emplace(
+	// Reset our current data
+	SourceData.SourceFiles.Reset();
+	SourceData.SourceFiles.Emplace(
 		SanitizeImportFilename(InPath),
 		IFileManager::Get().GetTimeStamp(*InPath),
 		FMD5Hash::HashFile(*InPath)
 		);
-
-	bDirty = true;
 	
 	OnImportDataChanged.Broadcast(Old, this);
 }
 
 FString UAssetImportData::GetFirstFilename() const
 {
-	return SourceFiles.Num() > 0 ? ResolveImportFilename(SourceFiles[0].RelativeFilename) : FString();
+	return SourceData.SourceFiles.Num() > 0 ? ResolveImportFilename(SourceData.SourceFiles[0].RelativeFilename) : FString();
 }
 
 void UAssetImportData::ExtractFilenames(TArray<FString>& AbsoluteFilenames) const
 {
-	for (const auto& File : SourceFiles)
+	for (const auto& File : SourceData.SourceFiles)
 	{
 		AbsoluteFilenames.Add(ResolveImportFilename(File.RelativeFilename));
 	}
 }
 
+TArray<FString> UAssetImportData::ExtractFilenames() const
+{
+	TArray<FString> Temp;
+	ExtractFilenames(Temp);
+	return Temp;
+}
+
 FString UAssetImportData::SanitizeImportFilename(const FString& InPath) const
 {
 	const UPackage* Package = GetOutermost();
@@ -182,15 +192,15 @@ void UAssetImportData::Serialize(FArchive& Ar)
 		if (Ar.IsLoading())
 		{
 			Ar << Json;
-			TOptional<FAssetImportInfo> Copy = FromJson(MoveTemp(Json));
+			TOptional<FAssetImportInfo> Copy = FAssetImportInfo::FromJson(MoveTemp(Json));
 			if (Copy.IsSet())
 			{
-				CopyFrom(Copy.GetValue());
+				SourceData = MoveTemp(Copy.GetValue());
 			}
 		}
 		else if (Ar.IsSaving())
 		{
-			Json = ToJson();
+			Json = SourceData.ToJson();
 			Ar << Json;
 		}
 	}
@@ -200,7 +210,7 @@ void UAssetImportData::Serialize(FArchive& Ar)
 
 void UAssetImportData::PostLoad()
 {
-	if (!SourceFilePath_DEPRECATED.IsEmpty() && SourceFiles.Num() == 0)
+	if (!SourceFilePath_DEPRECATED.IsEmpty() && SourceData.SourceFiles.Num() == 0)
 	{
 		FDateTime SourceDateTime;
 		if (!FDateTime::Parse(SourceFileTimestamp_DEPRECATED,SourceDateTime))
@@ -208,7 +218,7 @@ void UAssetImportData::PostLoad()
 			SourceDateTime = 0;
 		}
 
-		SourceFiles.Add(FSourceFile(MoveTemp(SourceFilePath_DEPRECATED),SourceDateTime));
+		SourceData.SourceFiles.Add(FAssetImportInfo::FSourceFile(MoveTemp(SourceFilePath_DEPRECATED),SourceDateTime));
 
 		SourceFilePath_DEPRECATED.Empty();
 		SourceFileTimestamp_DEPRECATED.Empty();
@@ -217,11 +227,4 @@ void UAssetImportData::PostLoad()
 	Super::PostLoad();
 }
 
-void UAssetImportData::PostEditChangeProperty(FPropertyChangedEvent& PropertyChangedEvent)
-{
-	Super::PostEditChangeProperty(PropertyChangedEvent);
-
-	bDirty = true;
-}
-
 #endif // WITH_EDITORONLY_DATA
\ No newline at end of file
diff --git a/Engine/Source/Runtime/Engine/Private/SkeletalMesh.cpp b/Engine/Source/Runtime/Engine/Private/SkeletalMesh.cpp
--- a/Engine/Source/Runtime/Engine/Private/SkeletalMesh.cpp
+++ b/Engine/Source/Runtime/Engine/Private/SkeletalMesh.cpp
@@ -2511,7 +2511,7 @@ void USkeletalMesh::Serialize( FArchive& Ar )
 		// AssetImportData should always have been set up in the constructor where this is relevant
 		FAssetImportInfo Info;
 		Info.Insert(FAssetImportInfo::FSourceFile(SourceFilePath_DEPRECATED));
-		AssetImportData->CopyFrom(Info);
+		AssetImportData->SourceData = MoveTemp(Info);
 		
 		SourceFilePath_DEPRECATED = TEXT("");
 		SourceFileTimestamp_DEPRECATED = TEXT("");
@@ -2876,7 +2876,7 @@ void USkeletalMesh::GetAssetRegistryTags(TArray<FAssetRegistryTag>& OutTags) con
 #if WITH_EDITORONLY_DATA
 	if (AssetImportData)
 	{
-		OutTags.Add( FAssetRegistryTag(SourceFileTagName(), AssetImportData->ToJson(), FAssetRegistryTag::TT_Hidden) );
+		OutTags.Add( FAssetRegistryTag(SourceFileTagName(), AssetImportData->GetSourceData().ToJson(), FAssetRegistryTag::TT_Hidden) );
 	}
 #endif
 	
diff --git a/Engine/Source/Runtime/Engine/Private/SoundWave.cpp b/Engine/Source/Runtime/Engine/Private/SoundWave.cpp
--- a/Engine/Source/Runtime/Engine/Private/SoundWave.cpp
+++ b/Engine/Source/Runtime/Engine/Private/SoundWave.cpp
@@ -144,7 +144,7 @@ void USoundWave::GetAssetRegistryTags(TArray<FAssetRegistryTag>& OutTags) const
 #if WITH_EDITORONLY_DATA
 	if (AssetImportData)
 	{
-		OutTags.Add( FAssetRegistryTag(SourceFileTagName(), AssetImportData->ToJson(), FAssetRegistryTag::TT_Hidden) );
+		OutTags.Add( FAssetRegistryTag(SourceFileTagName(), AssetImportData->GetSourceData().ToJson(), FAssetRegistryTag::TT_Hidden) );
 	}
 #endif
 	// GetCompressedDataSize could technically modify this->CompressedFormatData therefore it is not const, however this information
@@ -370,7 +370,7 @@ void USoundWave::PostLoad()
 	{
 		FAssetImportInfo Info;
 		Info.Insert(FAssetImportInfo::FSourceFile(SourceFilePath_DEPRECATED));
-		AssetImportData->CopyFrom(Info);
+		AssetImportData->SourceData = MoveTemp(Info);
 	}
 #endif // #if WITH_EDITORONLY_DATA
 
diff --git a/Engine/Source/Runtime/Engine/Private/StaticMesh.cpp b/Engine/Source/Runtime/Engine/Private/StaticMesh.cpp
--- a/Engine/Source/Runtime/Engine/Private/StaticMesh.cpp
+++ b/Engine/Source/Runtime/Engine/Private/StaticMesh.cpp
@@ -1564,7 +1564,7 @@ void UStaticMesh::GetAssetRegistryTags(TArray<FAssetRegistryTag>& OutTags) const
 #if WITH_EDITORONLY_DATA
 	if (AssetImportData)
 	{
-		OutTags.Add( FAssetRegistryTag(SourceFileTagName(), AssetImportData->ToJson(), FAssetRegistryTag::TT_Hidden) );
+		OutTags.Add( FAssetRegistryTag(SourceFileTagName(), AssetImportData->GetSourceData().ToJson(), FAssetRegistryTag::TT_Hidden) );
 	}
 #endif
 
@@ -1940,7 +1940,7 @@ void UStaticMesh::Serialize(FArchive& Ar)
 		// AssetImportData should always have been set up in the constructor where this is relevant
 		FAssetImportInfo Info;
 		Info.Insert(FAssetImportInfo::FSourceFile(SourceFilePath_DEPRECATED));
-		AssetImportData->CopyFrom(Info);
+		AssetImportData->SourceData = MoveTemp(Info);
 		
 		SourceFilePath_DEPRECATED = TEXT("");
 		SourceFileTimestamp_DEPRECATED = TEXT("");
diff --git a/Engine/Source/Runtime/Engine/Private/Texture.cpp b/Engine/Source/Runtime/Engine/Private/Texture.cpp
--- a/Engine/Source/Runtime/Engine/Private/Texture.cpp
+++ b/Engine/Source/Runtime/Engine/Private/Texture.cpp
@@ -264,7 +264,7 @@ void UTexture::PostLoad()
 	{
 		FAssetImportInfo Info;
 		Info.Insert(FAssetImportInfo::FSourceFile(SourceFilePath_DEPRECATED));
-		AssetImportData->CopyFrom(Info);
+		AssetImportData->SourceData = MoveTemp(Info);
 	}
 #endif
 
@@ -383,7 +383,7 @@ void UTexture::GetAssetRegistryTags(TArray<FAssetRegistryTag>& OutTags) const
 {
 	if (AssetImportData)
 	{
-		OutTags.Add( FAssetRegistryTag(SourceFileTagName(), AssetImportData->ToJson(), FAssetRegistryTag::TT_Hidden) );
+		OutTags.Add( FAssetRegistryTag(SourceFileTagName(), AssetImportData->GetSourceData().ToJson(), FAssetRegistryTag::TT_Hidden) );
 	}
 
 	Super::GetAssetRegistryTags(OutTags);
diff --git a/Engine/Source/Runtime/Engine/Private/VectorField.cpp b/Engine/Source/Runtime/Engine/Private/VectorField.cpp
--- a/Engine/Source/Runtime/Engine/Private/VectorField.cpp
+++ b/Engine/Source/Runtime/Engine/Private/VectorField.cpp
@@ -330,7 +330,7 @@ void UVectorFieldStatic::PostLoad()
 	{
 		FAssetImportInfo Info;
 		Info.Insert(FAssetImportInfo::FSourceFile(SourceFilePath_DEPRECATED));
-		AssetImportData->CopyFrom(Info);
+		AssetImportData->SourceData = MoveTemp(Info);
 	}
 #endif
 }
@@ -354,7 +354,7 @@ void UVectorFieldStatic::GetAssetRegistryTags(TArray<FAssetRegistryTag>& OutTags
 {
 	if (AssetImportData)
 	{
-		OutTags.Add( FAssetRegistryTag(SourceFileTagName(), AssetImportData->ToJson(), FAssetRegistryTag::TT_Hidden) );
+		OutTags.Add( FAssetRegistryTag(SourceFileTagName(), AssetImportData->GetSourceData().ToJson(), FAssetRegistryTag::TT_Hidden) );
 	}
 
 	Super::GetAssetRegistryTags(OutTags);
diff --git a/Engine/Source/Runtime/GeometryCache/Private/GeometryCache.cpp b/Engine/Source/Runtime/GeometryCache/Private/GeometryCache.cpp
--- a/Engine/Source/Runtime/GeometryCache/Private/GeometryCache.cpp
+++ b/Engine/Source/Runtime/GeometryCache/Private/GeometryCache.cpp
@@ -69,7 +69,7 @@ void UGeometryCache::GetAssetRegistryTags(TArray<FAssetRegistryTag>& OutTags) co
 #if WITH_EDITORONLY_DATA
 	if (AssetImportData)
 	{
-		OutTags.Add(FAssetRegistryTag(SourceFileTagName(), AssetImportData->ToJson(), FAssetRegistryTag::TT_Hidden));
+		OutTags.Add(FAssetRegistryTag(SourceFileTagName(), AssetImportData->GetSourceData().ToJson(), FAssetRegistryTag::TT_Hidden));
 	}
 #endif
 
