#ue4 - Fix replicated transforms not being updated on clients until the next transform or replication update.

udn: https://udn.unrealengine.com/questions/238900/scenecomponent-seems-to-use-old-data-for-replicati.html

diff --git a/Engine/Source/Runtime/Engine/Classes/Components/SceneComponent.h b/Engine/Source/Runtime/Engine/Classes/Components/SceneComponent.h
--- a/Engine/Source/Runtime/Engine/Classes/Components/SceneComponent.h
+++ b/Engine/Source/Runtime/Engine/Classes/Components/SceneComponent.h
@@ -224,7 +224,7 @@ public:
 
 private:
 
-	bool NetUpdateTransform;
+	bool bNetUpdateTransform;
 
 	FName NetOldAttachSocketName;
 	USceneComponent *NetOldAttachParent;
@@ -237,6 +237,7 @@ private:
 
 	virtual void PreNetReceive() override;
 	virtual void PostNetReceive() override;
+	virtual void PostRepNotifies() override;
 
 public:
 
diff --git a/Engine/Source/Runtime/Engine/Private/Components/SceneComponent.cpp b/Engine/Source/Runtime/Engine/Private/Components/SceneComponent.cpp
--- a/Engine/Source/Runtime/Engine/Private/Components/SceneComponent.cpp
+++ b/Engine/Source/Runtime/Engine/Private/Components/SceneComponent.cpp
@@ -49,8 +49,6 @@ USceneComponent::USceneComponent(const FObjectInitializer& ObjectInitializer /*=
 	// default behavior is visible
 	bVisible = true;
 	bAutoActivate = false;
-
-	NetUpdateTransform = false;
 }
 
 void USceneComponent::AddReferencedObjects(UObject* InThis, FReferenceCollector& Collector)
@@ -2208,7 +2206,7 @@ FBoxSphereBounds USceneComponent::GetPlacementExtent() const
 
 void USceneComponent::OnRep_Transform()
 {
-	NetUpdateTransform = true;
+	bNetUpdateTransform = true;
 }
 
 void USceneComponent::OnRep_Visibility(bool OldValue)
@@ -2253,11 +2251,14 @@ void USceneComponent::PostNetReceive()
 		
 		AttachTo(NetOldAttachParent, NetOldAttachSocketName);
 	}
+}
 
-	if (NetUpdateTransform)
+void USceneComponent::PostRepNotifies()
+{
+	if (bNetUpdateTransform)
 	{
 		UpdateComponentToWorld(true);
-		NetUpdateTransform = false;
+		bNetUpdateTransform = false;
 	}
 }
 
