Fixed json serialization of source file timestamps

This addresses UE-19725

diff --git a/Engine/Source/Runtime/Engine/Private/EditorFramework/AssetImportData.cpp b/Engine/Source/Runtime/Engine/Private/EditorFramework/AssetImportData.cpp
--- a/Engine/Source/Runtime/Engine/Private/EditorFramework/AssetImportData.cpp
+++ b/Engine/Source/Runtime/Engine/Private/EditorFramework/AssetImportData.cpp
@@ -22,7 +22,7 @@ FString FAssetImportInfo::ToJson() const
 
 	for (int32 Index = 0; Index < SourceFiles.Num(); ++Index)
 	{
-		Json += FString::Printf(TEXT("{ \"RelativeFilename\" : \"%s\", \"Timestamp\" : %d, \"FileMD5\" : \"%s\" }"),
+		Json += FString::Printf(TEXT("{ \"RelativeFilename\" : \"%s\", \"Timestamp\" : \"%d\", \"FileMD5\" : \"%s\" }"),
 			*SourceFiles[Index].RelativeFilename,
 			SourceFiles[Index].Timestamp.ToUnixTimestamp(),
 			*LexicalConversion::ToString(SourceFiles[Index].FileHash)
@@ -75,7 +75,7 @@ TOptional<FAssetImportInfo> FAssetImportInfo::FromJson(FString InJsonString)
 		FMD5Hash FileHash;
 		LexicalConversion::FromString(FileHash, *MD5String);
 
-		Info.SourceFiles.Emplace(MoveTemp(RelativeFilename), UnixTimestamp, FileHash);
+		Info.SourceFiles.Emplace(MoveTemp(RelativeFilename), FDateTime::FromUnixTimestamp(UnixTimestamp), FileHash);
 	}
 
 	return Info;
