#Fix client replication issue with spawn actor when actor has relative transform
https://jira.ol.epicgames.net/browse/UE-18660

diff --git a/Engine/Source/Runtime/Engine/Classes/Engine/World.h b/Engine/Source/Runtime/Engine/Classes/Engine/World.h
--- a/Engine/Source/Runtime/Engine/Classes/Engine/World.h
+++ b/Engine/Source/Runtime/Engine/Classes/Engine/World.h
@@ -2793,6 +2793,17 @@ public:
 	 */
 	AActor* SpawnActor( UClass* Class, FTransform const* Transform, const FActorSpawnParameters& SpawnParameters = FActorSpawnParameters());
 
+	/**
+	 * Spawn Actors with given absolute transform (override root component transform) and SpawnParameters
+	 * 
+	 * @param	Class					Class to Spawn
+	 * @param	AbsoluteTransform		World Transform to spawn on - without considering CDO's relative transform, thus Absolute
+	 * @param	SpawmParameters			Spawm Parameters
+	 *
+	 * @return	Actor that just spawned
+	 */
+	AActor* SpawnActorAbsolute( UClass* Class, FTransform const& AbsoluteTransform, const FActorSpawnParameters& SpawnParameters = FActorSpawnParameters());
+
 	/** Templated version of SpawnActor that allows you to specify a class type via the template type */
 	template< class T >
 	T* SpawnActor( const FActorSpawnParameters& SpawnParameters = FActorSpawnParameters() )
@@ -2832,7 +2843,23 @@ public:
 	{
 		return CastChecked<T>(SpawnActor(Class, &Transform, SpawnParameters), ECastCheckedType::NullAllowed);
 	}
-	
+
+	/** Templated version of SpawnActorAbsolute that allows you to specify absolute location and rotation in addition to class type via the template type */
+	template< class T >
+	T* SpawnActorAbsolute(FVector const& AbsoluteLocation, FRotator const& AbsoluteRotation, const FActorSpawnParameters& SpawnParameters = FActorSpawnParameters())
+	{
+		return CastChecked<T>(SpawnActorAbsolute(T::StaticClass(), FTransform(AbsoluteRotation, AbsoluteLocation), SpawnParameters), ECastCheckedType::NullAllowed);
+	}
+
+	/** 
+	 *  Templated version of SpawnActorAbsolute that allows you to specify whole absolute Transform
+	 *  class type via parameter while the return type is a parent class of that type 
+	 */
+	template< class T >
+	T* SpawnActorAbsolute(UClass* Class, FTransform const& Transform,const FActorSpawnParameters& SpawnParameters = FActorSpawnParameters())
+	{
+		return CastChecked<T>(SpawnActorAbsolute(Class, &Transform, SpawnParameters), ECastCheckedType::NullAllowed);
+	}
 	/**
 	* Spawns given class and returns class T pointer, forcibly sets world position. WILL NOT run Construction Script of Blueprints 
 	* to give caller an opportunity to set parameters beforehand.  Caller is responsible for invoking construction
diff --git a/Engine/Source/Runtime/Engine/Private/LevelActor.cpp b/Engine/Source/Runtime/Engine/Private/LevelActor.cpp
--- a/Engine/Source/Runtime/Engine/Private/LevelActor.cpp
+++ b/Engine/Source/Runtime/Engine/Private/LevelActor.cpp
@@ -230,6 +230,27 @@ void LineCheckTracker::CaptureLineCheck(int32 LineCheckFlags, const FVector* Ext
 	TArray<FString>	ThisFramePawnSpawns;
 #endif	//PERF_SHOW_MULTI_PAWN_SPAWN_FRAMES
 
+AActor* UWorld::SpawnActorAbsolute( UClass* Class, FTransform const& AbsoluteTransform, const FActorSpawnParameters& SpawnParameters)
+{
+	AActor* Template = SpawnParameters.Template;
+
+	if(!Template)
+	{
+		// Use class's default actor as a template.
+		Template = Class->GetDefaultObject<AActor>();
+	}
+
+	FTransform NewTransform = AbsoluteTransform;
+	USceneComponent* TemplateRootComponent = (Template)? Template->GetRootComponent() : NULL;
+	if(TemplateRootComponent)
+	{
+		TemplateRootComponent->UpdateComponentToWorld();
+		NewTransform = TemplateRootComponent->GetComponentToWorld().Inverse() * NewTransform;
+	}
+
+	return SpawnActor(Class, &NewTransform, SpawnParameters);
+}
+
 AActor* UWorld::SpawnActor( UClass* Class, FVector const* Location, FRotator const* Rotation, const FActorSpawnParameters& SpawnParameters )
 {
 	FTransform Transform;
diff --git a/Engine/Source/Runtime/Engine/Private/PackageMapClient.cpp b/Engine/Source/Runtime/Engine/Private/PackageMapClient.cpp
--- a/Engine/Source/Runtime/Engine/Private/PackageMapClient.cpp
+++ b/Engine/Source/Runtime/Engine/Private/PackageMapClient.cpp
@@ -341,8 +341,7 @@ bool UPackageMapClient::SerializeNewActor(FArchive& Ar, class UActorChannel *Cha
 					SpawnInfo.SpawnCollisionHandlingOverride = ESpawnActorCollisionHandlingMethod::AlwaysSpawn;
 					SpawnInfo.bRemoteOwned = true;
 					SpawnInfo.bNoFail = true;
-					Actor = Connection->Driver->GetWorld()->SpawnActor(Archetype->GetClass(), &Location, &Rotation, SpawnInfo );
-
+					Actor = Connection->Driver->GetWorld()->SpawnActorAbsolute(Archetype->GetClass(), FTransform(Rotation, Location), SpawnInfo );
 					// Velocity was serialized by the server
 					if (bSerializeVelocity)
 					{
