Comment improvement for documentation accuracy.

diff --git a/Engine/Source/Runtime/Engine/Classes/Components/PrimitiveComponent.h b/Engine/Source/Runtime/Engine/Classes/Components/PrimitiveComponent.h
--- a/Engine/Source/Runtime/Engine/Classes/Components/PrimitiveComponent.h
+++ b/Engine/Source/Runtime/Engine/Classes/Components/PrimitiveComponent.h
@@ -1738,7 +1738,7 @@ public:
 	virtual bool SweepComponent(FHitResult& OutHit, const FVector Start, const FVector End, const FCollisionShape &CollisionShape, bool bTraceComplex=false);
 	
 	/** 
-	 *  Test the collision of the supplied component at the supplied location/rotation, and determine if it overlaps this component
+	 *  Test the collision of the supplied component at the supplied location/rotation, and determine if it overlaps this component.
 	 *  @note This overload taking rotation as a FQuat is slightly faster than the version using FRotator.
 	 *  @note This simply calls the virtual ComponentOverlapComponentImpl() which can be overridden to implement custom behavior.
 	 *  @param  PrimComp        Component to use geometry from to test against this component. Transform of this component is ignored.
@@ -1758,7 +1758,7 @@ protected:
 public:
 	
 	/** 
-	 *  Test the collision of the supplied Sphere at the supplied location, and determine if it overlaps this component
+	 *  Test the collision of the supplied shape at the supplied location, and determine if it overlaps this component.
 	 *
 	 *  @param  Pos             Location to place PrimComp geometry at 
 	 *	@param	Rot				Rotation of PrimComp geometry
