Fix pvd nodebug not working with ip address

diff --git a/Engine/Source/Runtime/Engine/Private/PhysicsEngine/PhysUtils.cpp b/Engine/Source/Runtime/Engine/Private/PhysicsEngine/PhysUtils.cpp
--- a/Engine/Source/Runtime/Engine/Private/PhysicsEngine/PhysUtils.cpp
+++ b/Engine/Source/Runtime/Engine/Private/PhysicsEngine/PhysUtils.cpp
@@ -555,14 +555,18 @@ bool ExecPhysCommands(const TCHAR* Cmd, FOutputDevice* Ar, UWorld* InWorld)
 		{
 			if(FParse::Command(&Cmd, TEXT("CONNECT")))
 			{
+				
+				
+				const bool bVizualization = !FParse::Command(&Cmd, TEXT("NODEBUG"));
+
 				// setup connection parameters
 				FString Host = TEXT("localhost");
-				if(*Cmd)
+				if (*Cmd)
 				{
 					Host = Cmd;
 				}
 
-				const bool bVizualization = !FParse::Command(&Cmd, TEXT("NODEBUG"));
+				
 
 				PvdConnect(Host, bVizualization);
 
