[UE-19299] Fix failure to match component data cache entries during Blueprint class reinstancing.

change summary:
- modified FActorComponentInstanceData::MatchesComponent() to match templates by name rather than by address during reinstancing only, since archetypes will source from old/new class in that situation.

diff --git a/Engine/Source/Runtime/Engine/Private/ComponentInstanceDataCache.cpp b/Engine/Source/Runtime/Engine/Private/ComponentInstanceDataCache.cpp
--- a/Engine/Source/Runtime/Engine/Private/ComponentInstanceDataCache.cpp
+++ b/Engine/Source/Runtime/Engine/Private/ComponentInstanceDataCache.cpp
@@ -81,7 +81,7 @@ bool FActorComponentInstanceData::MatchesComponent(const UActorComponent* Compon
 	bool bMatches = false;
 	if (   Component
 		&& Component->CreationMethod == SourceComponentCreationMethod
-		&& ComponentTemplate == SourceComponentTemplate)
+		&& (ComponentTemplate == SourceComponentTemplate || (GIsReinstancing && ComponentTemplate->GetFName() == SourceComponentTemplate->GetFName())))
 	{
 		if (SourceComponentCreationMethod != EComponentCreationMethod::UserConstructionScript)
 		{
@@ -95,13 +95,16 @@ bool FActorComponentInstanceData::MatchesComponent(const UActorComponent* Compon
 			{
 				for (const UActorComponent* BlueprintCreatedComponent : ComponentOwner->BlueprintCreatedComponents)
 				{
-					if (   BlueprintCreatedComponent
-						&& (BlueprintCreatedComponent->GetArchetype() == SourceComponentTemplate)
-						&& (BlueprintCreatedComponent->CreationMethod == SourceComponentCreationMethod)
-						&& (++FoundSerializedComponentsOfType == SourceComponentTypeSerializedIndex))
+					if (BlueprintCreatedComponent != nullptr)
 					{
-						bMatches = (BlueprintCreatedComponent == Component);
-						break;
+						const UObject* BlueprintComponentTemplate = BlueprintCreatedComponent->GetArchetype();
+						if (   (BlueprintComponentTemplate == SourceComponentTemplate || (GIsReinstancing && BlueprintComponentTemplate->GetFName() == SourceComponentTemplate->GetFName()))
+							&& (BlueprintCreatedComponent->CreationMethod == SourceComponentCreationMethod)
+							&& (++FoundSerializedComponentsOfType == SourceComponentTypeSerializedIndex))
+						{
+							bMatches = (BlueprintCreatedComponent == Component);
+							break;
+						}
 					}
 				}
 			}
