Fix the crash when Presave

- not the case I was looking for but still valid nonetheless
https://jira.ol.epicgames.net/browse/UE-19481

diff --git a/Engine/Source/Runtime/Engine/Private/Animation/AnimSequence.cpp b/Engine/Source/Runtime/Engine/Private/Animation/AnimSequence.cpp
--- a/Engine/Source/Runtime/Engine/Private/Animation/AnimSequence.cpp
+++ b/Engine/Source/Runtime/Engine/Private/Animation/AnimSequence.cpp
@@ -341,7 +341,9 @@ void UAnimSequence::PreSave()
 	}
 
 	// make sure if it does contain transform curvesm it contains source data
-	check (!DoesContainTransformCurves() || SourceRawAnimationData.Num() != 0);
+	// empty track animation still can be made by retargeting to invalid skeleton
+	// make sure to not trigger ensure if RawAnimationData is also null
+	check (!DoesContainTransformCurves() || (RawAnimationData.Num()==0 || SourceRawAnimationData.Num() != 0));
 #endif
 
 	Super::PreSave();
