Fix crash in the world outliner if a folder is added around an actor with no root component

diff --git a/Engine/Source/Runtime/Engine/Private/ActorEditorUtils.cpp b/Engine/Source/Runtime/Engine/Private/ActorEditorUtils.cpp
--- a/Engine/Source/Runtime/Engine/Private/ActorEditorUtils.cpp
+++ b/Engine/Source/Runtime/Engine/Private/ActorEditorUtils.cpp
@@ -74,12 +74,15 @@ namespace FActorEditorUtils
 			return false;
 		}
 
-		for (auto* ChildComponent : InActor->GetRootComponent()->AttachChildren)
+		if( InActor->GetRootComponent() )
 		{
-			AActor* ChildActor = ChildComponent ? ChildComponent->GetOwner() : nullptr;
-			if (ChildActor && ChildActor != InActor && !TraverseActorTree_ParentFirst(ChildActor, InPredicate))
+			for(auto* ChildComponent : InActor->GetRootComponent()->AttachChildren)
 			{
-				return false;
+				AActor* ChildActor = ChildComponent ? ChildComponent->GetOwner() : nullptr;
+				if(ChildActor && ChildActor != InActor && !TraverseActorTree_ParentFirst(ChildActor, InPredicate))
+				{
+					return false;
+				}
 			}
 		}
 		
