Fix CIS.

diff --git a/Engine/Source/Runtime/Engine/Classes/Engine/World.h b/Engine/Source/Runtime/Engine/Classes/Engine/World.h
--- a/Engine/Source/Runtime/Engine/Classes/Engine/World.h
+++ b/Engine/Source/Runtime/Engine/Classes/Engine/World.h
@@ -2858,7 +2858,7 @@ public:
 	template< class T >
 	T* SpawnActorAbsolute(UClass* Class, FTransform const& Transform,const FActorSpawnParameters& SpawnParameters = FActorSpawnParameters())
 	{
-		return CastChecked<T>(SpawnActorAbsolute(Class, &Transform, SpawnParameters), ECastCheckedType::NullAllowed);
+		return CastChecked<T>(SpawnActorAbsolute(Class, Transform, SpawnParameters), ECastCheckedType::NullAllowed);
 	}
 	/**
 	* Spawns given class and returns class T pointer, forcibly sets world position. WILL NOT run Construction Script of Blueprints 
