Don't send server frametime during acks to clients unless we need to. Always send a bit even in shipping to keep protocol format consistent.

diff --git a/Engine/Source/Runtime/Engine/Classes/Engine/NetConnection.h b/Engine/Source/Runtime/Engine/Classes/Engine/NetConnection.h
--- a/Engine/Source/Runtime/Engine/Classes/Engine/NetConnection.h
+++ b/Engine/Source/Runtime/Engine/Classes/Engine/NetConnection.h
@@ -123,11 +123,11 @@ struct DelayedPacket
 #endif
 
 // The interval between ack packets, which is used to decide which acks are used for ping validation checks (must be power of two)
-#define PING_ACK_PACKET_INTERVAL 16
+#define PING_ACK_CHECKSUM_PACKET_INTERVAL 16
 
 // Used by the client, for setting a minimum delay between PingAck's
-//	(optionally, 'PING_ACK_PACKET_INTERVAL' can be tweaked, so that the interval checks take a similar amount of time as this delay)
-#define PING_ACK_DELAY 0.5
+//	(optionally, 'PING_ACK_CHECKSUM_PACKET_INTERVAL' can be tweaked, so that the interval checks take a similar amount of time as this delay)
+#define PING_ACK_CHECKSUM_DELAY 0.5
 
 UCLASS(customConstructor, Abstract, MinimalAPI, transient, config=Engine)
 class UNetConnection : public UPlayer
@@ -276,9 +276,10 @@ public:
 	int32			OutPacketId;			// Most recently sent packet.
 	int32 			OutAckPacketId;			// Most recently acked outgoing packet.
 
-	uint32			PingAckDataCache[MAX_PACKETID/PING_ACK_PACKET_INTERVAL];	// Caches packet data on the server, for verifying pings
-	float			LastPingAck;												// The time of the most recent PingAck on the client
-	int32			LastPingAckPacketId;										// The PacketId of the last PingAck, on the server
+	uint32			PingChecksumAckCache[MAX_PACKETID / PING_ACK_CHECKSUM_PACKET_INTERVAL];	// Caches packet data on the server, for verifying pings
+	float			LastPingChecksumAck;													// The time of the most ping checksum generated on the client
+	int32			LastPingChecksumAckPacketId;											// The PacketId of the last PingAck, on the server
+	bool			bLastHasServerFrameTime;
 
 	// Channel table.
 	class UChannel*		Channels		[ MAX_CHANNELS ];
@@ -432,7 +433,7 @@ public:
 	ENGINE_API virtual void AssertValid();
 
 	/** Send an acknowledgment. */
-	ENGINE_API virtual void SendAck( int32 PacketId, bool FirstTime=1, bool bHavePingAckData=0, uint32 PingAckData=0 );
+	ENGINE_API virtual void SendAck( int32 PacketId, bool FirstTime=1, bool bHasPingChecksum=false, uint32 PingChecksum=0 );
 
 	/**
 	 * flushes any pending data, bundling it into a packet and sending it via LowLevelSend()
diff --git a/Engine/Source/Runtime/Engine/Private/NetConnection.cpp b/Engine/Source/Runtime/Engine/Private/NetConnection.cpp
--- a/Engine/Source/Runtime/Engine/Private/NetConnection.cpp
+++ b/Engine/Source/Runtime/Engine/Private/NetConnection.cpp
@@ -62,8 +62,9 @@ UNetConnection::UNetConnection(const FObjectInitializer& ObjectInitializer)
 ,	InPacketId			( -1 )
 ,	OutPacketId			( 0 ) // must be initialized as OutAckPacketId + 1 so loss of first packet can be detected
 ,	OutAckPacketId		( -1 )
-,	LastPingAck			( 0.f )
-,	LastPingAckPacketId	( -1 )
+,	LastPingChecksumAck	( 0.f )
+,	LastPingChecksumAckPacketId( -1 )
+,	bLastHasServerFrameTime( false )
 ,	ClientWorldPackageName( NAME_None )
 {
 }
@@ -654,14 +655,19 @@ void UNetConnection::ReceivedNak( int32 NakPacketId )
  * @NumBits		Number of bits in data blob
  * @PacketId	PacketId of the packet used to produce this ack data
 */
-static uint32 CalcPingAckData( const uint8* BunchData, const int32 NumBits, const int32 PacketId )
+static uint32 CalcPingChecksum( const uint8* BunchData, const int32 NumBits, const int32 PacketId )
 {
 	// Simply walk the bunch data, based upon OutPacketId, to get 'good enough' random data for verification
-	const int32	BunchBytesFloor = NumBits / 8;
-	const int32	PingAckIdx		= ( PacketId % MAX_PACKETID ) / PING_ACK_PACKET_INTERVAL;
+	const int32	BunchBytesFloor		= NumBits / 8;
+	const int32	PingChecksumAckIdx	= ( PacketId % MAX_PACKETID ) / PING_ACK_CHECKSUM_PACKET_INTERVAL;
 
-	return	( BunchData[( PingAckIdx + 0 ) % BunchBytesFloor] << 24 )	+ ( BunchData[( PingAckIdx + 1 ) % BunchBytesFloor] << 16 ) +
-			( BunchData[( PingAckIdx + 2 ) % BunchBytesFloor] << 8 )	+ ( BunchData[( PingAckIdx + 3 ) % BunchBytesFloor] );
+	return	( BunchData[( PingChecksumAckIdx + 0 ) % BunchBytesFloor] << 24 ) + ( BunchData[( PingChecksumAckIdx + 1 ) % BunchBytesFloor] << 16 ) +
+		( BunchData[( PingChecksumAckIdx + 2 ) % BunchBytesFloor] << 8 ) + ( BunchData[( PingChecksumAckIdx + 3 ) % BunchBytesFloor] );
+}
+
+static bool IsAckChecksumPacket( const int32 PacketId )
+{
+	return ( PacketId % PING_ACK_CHECKSUM_PACKET_INTERVAL ) == 0;
 }
 
 void UNetConnection::ReceivedPacket( FBitReader& Reader )
@@ -707,9 +713,9 @@ void UNetConnection::ReceivedPacket( FBitReader& Reader )
 	}
 
 	// Detect packets on the client, which should trigger PingAck verification
-	bool bPendingSendPingAck = Driver->ServerConnection != NULL && (PacketId % PING_ACK_PACKET_INTERVAL) == 0;
-	bool bGotPingAckData = false;
-	uint32 OutPingAckData = 0;
+	const bool bCanHavePingChecksum	= !Driver->IsServer() && IsAckChecksumPacket( PacketId );
+	bool bGotPingChecksum			= false;
+	uint32 OutPingChecksum			= 0;
 
 	bool bSkipAck = false;
 
@@ -731,7 +737,7 @@ void UNetConnection::ReceivedPacket( FBitReader& Reader )
 			LastRecvAckTime = Driver->Time;
 
 			// This is an acknowledgment.
-			int32 AckPacketId = MakeRelative(Reader.ReadInt(MAX_PACKETID),OutAckPacketId,MAX_PACKETID);
+			const int32 AckPacketId = MakeRelative(Reader.ReadInt(MAX_PACKETID),OutAckPacketId,MAX_PACKETID);
 
 			if( Reader.IsError() )
 			{
@@ -739,29 +745,42 @@ void UNetConnection::ReceivedPacket( FBitReader& Reader )
 				return;
 			}
 
-			// Detect ack packets on the server, containing PingAckData
-			const bool bPingAck		= Driver->IsServer() && ( AckPacketId % PING_ACK_PACKET_INTERVAL ) == 0;
-			bool bHasPingAckData	= false;
-			uint32 InPingAckData	= 0;
+			double ServerFrameTime = 0;
 
-#if !UE_BUILD_SHIPPING
-			double ServerFrameTime	= 0;
+			// If this is the server, we're reading in the request to send them our frame time
+			// If this is the client, we're reading in confirmation that our request to get frame time from server is granted
+			const bool bHasServerFrameTime = !!Reader.ReadBit();
 
+#if !UE_BUILD_SHIPPING	// We never actually read it for shipping
 			if ( !Driver->IsServer() )
 			{
-				uint8 FrameTimeByte	= 0;
-				Reader << FrameTimeByte;
-				ServerFrameTime = ( double )FrameTimeByte / 1000;
+				if ( bHasServerFrameTime )
+				{
+					// As a client, our request was granted, read the frame time
+					uint8 FrameTimeByte	= 0;
+					Reader << FrameTimeByte;
+					ServerFrameTime = ( double )FrameTimeByte / 1000;
+				}
+			}
+			else
+			{
+				// Server remembers so he can use during SendAck to notify to client of his frame time
+				bLastHasServerFrameTime = bHasServerFrameTime;
 			}
 #endif
 
-			if ( bPingAck )
+			// Detect ack packets on the server, containing PingChecksum
+			const bool bCanHavePingChecksum	= Driver->IsServer() && IsAckChecksumPacket( AckPacketId );
+			bool bHasPingChecksum			= false;
+			uint32 InPingChecksum			= 0;
+
+			if ( bCanHavePingChecksum )
 			{
-				bHasPingAckData = !!Reader.ReadBit();
+				bHasPingChecksum = !!Reader.ReadBit();
 
-				if (bHasPingAckData)
+				if ( bHasPingChecksum )
 				{
-					Reader.Serialize(&InPingAckData, sizeof(uint32));
+					Reader.Serialize( &InPingChecksum, sizeof( uint32 ) );
 				}
 			}
 
@@ -807,8 +826,8 @@ void UNetConnection::ReceivedPacket( FBitReader& Reader )
 				if (PlayerController != NULL)
 				{
 					// Verify PingAck's, and pass notification up
-					const bool bServerVerified = bPingAck && bHasPingAckData && PingAckDataCache[( AckPacketId % MAX_PACKETID ) / PING_ACK_PACKET_INTERVAL] == InPingAckData;
-					
+					const bool bServerVerified = bCanHavePingChecksum && bHasPingChecksum && PingChecksumAckCache[( AckPacketId % MAX_PACKETID ) / PING_ACK_CHECKSUM_PACKET_INTERVAL] == InPingChecksum;
+
 					if ( bServerVerified || !Driver->IsServer() )
 					{
 						PlayerController->UpdatePing(NewLag);
@@ -961,13 +980,12 @@ void UNetConnection::ReceivedPacket( FBitReader& Reader )
 				return;
 			}
 
-			// Handle grabbing of PingAckData, for client PingAck verification (need a minimum of 32bits data, for OutPingAckData)
-			if (bPendingSendPingAck && !bGotPingAckData && BunchDataBits >= 32 &&
-				FMath::Abs(Driver->Time - LastPingAck) >= PING_ACK_DELAY)
+			// Handle grabbing of PingChecksum, for client PingAck verification (need a minimum of 32bits data, for OutPingChecksum)
+			if ( bCanHavePingChecksum && !bGotPingChecksum && BunchDataBits >= 32 && FMath::Abs( Driver->Time - LastPingChecksumAck ) >= PING_ACK_CHECKSUM_DELAY )
 			{
-				OutPingAckData = CalcPingAckData( Bunch.GetData(), Bunch.GetNumBits(), PacketId );
-				bGotPingAckData = true;
-				LastPingAck = Driver->Time;
+				OutPingChecksum = CalcPingChecksum( Bunch.GetData(), Bunch.GetNumBits(), PacketId );
+				bGotPingChecksum = true;
+				LastPingChecksumAck = Driver->Time;
 			}
 
 			// Receiving data.
@@ -1064,7 +1082,7 @@ void UNetConnection::ReceivedPacket( FBitReader& Reader )
 	// Acknowledge the packet.
 	if ( !bSkipAck )
 	{
-		SendAck(PacketId, true, bGotPingAckData, OutPingAckData);
+		SendAck( PacketId, true, bGotPingChecksum, OutPingChecksum );
 	}
 }
 
@@ -1167,7 +1185,7 @@ void UNetConnection::PurgeAcks()
 }
 
 
-void UNetConnection::SendAck( int32 AckPacketId, bool FirstTime/*=1*/, bool bHavePingAckData/*=0*/, uint32 PingAckData/*=0*/ )
+void UNetConnection::SendAck( int32 AckPacketId, bool FirstTime/*=1*/, bool bHasPingChecksum/*=0*/, uint32 PingChecksum/*=0*/ )
 {
 	ValidateSendBuffer();
 
@@ -1183,24 +1201,31 @@ void UNetConnection::SendAck( int32 AckPacketId, bool FirstTime/*=1*/, bool bHav
 
 		AckData.WriteBit( 1 );
 		AckData.WriteIntWrapped(AckPacketId, MAX_PACKETID);
+		
+#if !UE_BUILD_SHIPPING	// We never actually write it for shipping
+		const bool bHasServerFrameTime = Driver->IsServer() ? bLastHasServerFrameTime : ( CVarPingExcludeFrameTime.GetValueOnGameThread() > 0 ? true : false );
 
-#if !UE_BUILD_SHIPPING
-		if ( Driver->IsServer() )
+		AckData.WriteBit( bHasServerFrameTime ? 1 : 0 );
+
+		if ( Driver->IsServer() && bHasServerFrameTime )
 		{
 			uint8 FrameTimeByte = FMath::Min( FMath::FloorToInt( FrameTime * 1000 ), 255 );
 			AckData << FrameTimeByte;
 		}
+#else
+		// We still write the bit in shipping to keep the format the same
+		AckData.WriteBit( 0 );
 #endif
 
-		const bool bPingAck = !Driver->IsServer() && ( AckPacketId % PING_ACK_PACKET_INTERVAL ) == 0;
+		const bool bCanHavePingChecksum = !Driver->IsServer() && IsAckChecksumPacket( AckPacketId );
 
-		if ( bPingAck )
+		if ( bCanHavePingChecksum )
 		{
-			AckData.WriteBit( bHavePingAckData );
+			AckData.WriteBit( bHasPingChecksum );
 
-			if ( bHavePingAckData )
+			if ( bHasPingChecksum )
 			{
-				AckData.Serialize( &PingAckData, sizeof( uint32 ) );
+				AckData.Serialize( &PingChecksum, sizeof( uint32 ) );
 			}
 		}
 
@@ -1295,17 +1320,16 @@ int32 UNetConnection::SendRawBunch( FOutBunch& Bunch, bool InAllowMerge )
 	}
 
 	// Verified client ping tracking - caches some semi-random bytes of the packet, for ping validation
-	if (Driver->ServerConnection == NULL && Bunch.PacketId != LastPingAckPacketId && (Bunch.PacketId % PING_ACK_PACKET_INTERVAL) == 0 &&
-		Bunch.GetNumBits() >= 32)
+	if ( Driver->IsServer() && Bunch.PacketId != LastPingChecksumAckPacketId && IsAckChecksumPacket( Bunch.PacketId ) && Bunch.GetNumBits() >= 32 )
 	{
-		const uint32 PingAckData = CalcPingAckData( Bunch.GetData(), Bunch.GetNumBits(), Bunch.PacketId );
+		const uint32 PingChecksum = CalcPingChecksum( Bunch.GetData(), Bunch.GetNumBits(), Bunch.PacketId );
 
-		const int32 PingAckIdx = (Bunch.PacketId % MAX_PACKETID) / PING_ACK_PACKET_INTERVAL;
+		const int32 PingChecksumAckIdx = ( Bunch.PacketId % MAX_PACKETID ) / PING_ACK_CHECKSUM_PACKET_INTERVAL;
 
-		PingAckDataCache[PingAckIdx] = PingAckData;
+		PingChecksumAckCache[PingChecksumAckIdx] = PingChecksum;
 
 		// Multiple bunches get written per-packet, but PingAck only uses the first bunch, so don't process more bunches this PacketId
-		LastPingAckPacketId = Bunch.PacketId;
+		LastPingChecksumAckPacketId = Bunch.PacketId;
 	}
 
 	return Bunch.PacketId;
