#ue4 - Network clients now sync physics simulation state (on/off) with server if bReplicateMovement is true.

diff --git a/Engine/Source/Runtime/Engine/Classes/GameFramework/Actor.h b/Engine/Source/Runtime/Engine/Classes/GameFramework/Actor.h
--- a/Engine/Source/Runtime/Engine/Classes/GameFramework/Actor.h
+++ b/Engine/Source/Runtime/Engine/Classes/GameFramework/Actor.h
@@ -188,6 +188,9 @@ private:
 	UPROPERTY()
 	uint32 bActorEnableCollision:1;
 
+	/** Flag indicating we have checked initial simulating physics state to sync networked proxies to the server. */
+	uint32 bNetCheckedInitialPhysicsState:1;
+
 protected:
 	/**
 	 * If true, this actor will replicate to remote machines
@@ -1765,6 +1768,12 @@ public:
 	/** Update and smooth simulated physic state, replaces PostNetReceiveLocation() and PostNetReceiveVelocity() */
 	virtual void PostNetReceivePhysicState();
 
+protected:
+	/** Sync IsSimulatingPhysics() with ReplicatedMovement.bRepPhysics */
+	void SyncReplicatedPhysicsSimulation();
+
+public:
+
 	/** 
 	 * Set the owner of this Actor, used primarily for network replication. 
 	 * @param NewOwner	The Actor whom takes over ownership of this Actor
diff --git a/Engine/Source/Runtime/Engine/Private/ActorReplication.cpp b/Engine/Source/Runtime/Engine/Private/ActorReplication.cpp
--- a/Engine/Source/Runtime/Engine/Private/ActorReplication.cpp
+++ b/Engine/Source/Runtime/Engine/Private/ActorReplication.cpp
@@ -15,6 +15,7 @@
 //
 static bool		SavedbHidden;
 static AActor*	SavedOwner;
+static bool		SavedbRepPhysics;
 
 #define DEPRECATED_NET_PRIORITY -17.0f // Pick an arbitrary invalid priority, check for that
 
@@ -93,10 +94,19 @@ void AActor::PreNetReceive()
 {
 	SavedbHidden = bHidden;
 	SavedOwner = Owner;
+	SavedbRepPhysics = ReplicatedMovement.bRepPhysics;
 }
 
 void AActor::PostNetReceive()
 {
+	if (!bNetCheckedInitialPhysicsState)
+	{
+		// Initially we need to sync the state regardless of whether bRepPhysics has "changed" since it may not currently match IsSimulatingPhysics().
+		SyncReplicatedPhysicsSimulation();
+		SavedbRepPhysics = ReplicatedMovement.bRepPhysics;
+		bNetCheckedInitialPhysicsState = true;
+	}
+
 	ExchangeB( bHidden, SavedbHidden );
 	Exchange ( Owner, SavedOwner );
 
@@ -112,16 +122,31 @@ void AActor::PostNetReceive()
 
 void AActor::OnRep_ReplicatedMovement()
 {
-	if( RootComponent && RootComponent->IsSimulatingPhysics() )
+	if (RootComponent)
 	{
-		PostNetReceivePhysicState();
-	}
-	else
-	{
-		if (Role == ROLE_SimulatedProxy)
+		if (SavedbRepPhysics != ReplicatedMovement.bRepPhysics)
 		{
-			PostNetReceiveVelocity(ReplicatedMovement.LinearVelocity);
-			PostNetReceiveLocationAndRotation();
+			// Turn on/off physics sim to match server.
+			SyncReplicatedPhysicsSimulation();
+		}
+
+		if (ReplicatedMovement.bRepPhysics)
+		{
+			// Sync physics state
+			checkSlow(RootComponent->IsSimulatingPhysics());
+			PostNetReceivePhysicState();
+		}
+		else
+		{
+			// Attachment trumps global position updates, see GatherCurrentMovement().
+			if (!RootComponent->AttachParent)
+			{
+				if (Role == ROLE_SimulatedProxy)
+				{
+					PostNetReceiveVelocity(ReplicatedMovement.LinearVelocity);
+					PostNetReceiveLocationAndRotation();
+				}
+			}
 		}
 	}
 }
@@ -151,6 +176,18 @@ void AActor::PostNetReceivePhysicState()
 	}
 }
 
+void AActor::SyncReplicatedPhysicsSimulation()
+{
+	if (bReplicateMovement && RootComponent && (RootComponent->IsSimulatingPhysics() != ReplicatedMovement.bRepPhysics))
+	{
+		UPrimitiveComponent* RootPrimComp = Cast<UPrimitiveComponent>(RootComponent);
+		if (RootPrimComp)
+		{
+			RootPrimComp->SetSimulatePhysics(ReplicatedMovement.bRepPhysics);
+		}
+	}
+}
+
 bool AActor::IsNetRelevantFor(const APlayerController* RealViewer, const AActor* ViewTarget, const FVector& SrcLocation) const
 {
 	return IsNetRelevantFor(Cast<AActor>(RealViewer), ViewTarget, SrcLocation);
@@ -201,6 +238,7 @@ void AActor::GatherCurrentMovement()
 		RootPrimComp->GetRigidBodyState(RBState);
 
 		ReplicatedMovement.FillFrom(RBState);
+		ReplicatedMovement.bRepPhysics = true;
 	}
 	else if(RootComponent != NULL)
 	{
@@ -222,8 +260,9 @@ void AActor::GatherCurrentMovement()
 			ReplicatedMovement.Rotation = RootComponent->GetComponentRotation();
 			ReplicatedMovement.LinearVelocity = GetVelocity();
 			ReplicatedMovement.AngularVelocity = FVector::ZeroVector;
-			ReplicatedMovement.bRepPhysics = false;
 		}
+
+		ReplicatedMovement.bRepPhysics = false;
 	}
 }
 
