#ue4 - Network attachment uses less bandwidth when root component is attached and is set to replicate. Actor no longer replicates redundant AttachmentReplication struct in this case.

diff --git a/Engine/Source/Runtime/Engine/Classes/Components/SceneComponent.h b/Engine/Source/Runtime/Engine/Classes/Components/SceneComponent.h
--- a/Engine/Source/Runtime/Engine/Classes/Components/SceneComponent.h
+++ b/Engine/Source/Runtime/Engine/Classes/Components/SceneComponent.h
@@ -104,7 +104,7 @@ public:
 	virtual void GetLifetimeReplicatedProps(TArray<FLifetimeProperty>& OutLifetimeProps) const override;
 
 	/** What we are currently attached to. If valid, RelativeLocation etc. are used relative to this object */
-	UPROPERTY(Replicated)
+	UPROPERTY(ReplicatedUsing=OnRep_AttachParent)
 	class USceneComponent* AttachParent;
 
 	/** List of child SceneComponents that are attached to us. */
@@ -112,7 +112,7 @@ public:
 	TArray< USceneComponent* > AttachChildren;
 
 	/** Optional socket name on AttachParent that we are attached to. */
-	UPROPERTY(Replicated)
+	UPROPERTY(ReplicatedUsing=OnRep_AttachSocketName)
 	FName AttachSocketName;
 
 	/** if true, will call GetCustomLocation instead or returning the location part of ComponentToWorld */
@@ -225,13 +225,19 @@ public:
 private:
 
 	bool bNetUpdateTransform;
-
+	bool bNetUpdateAttachment;
 	FName NetOldAttachSocketName;
 	USceneComponent *NetOldAttachParent;
 
 	UFUNCTION()
 	void OnRep_Transform();
 
+	UFUNCTION()
+	void OnRep_AttachParent();
+
+	UFUNCTION()
+	void OnRep_AttachSocketName();
+
 	UFUNCTION()
 	void OnRep_Visibility(bool OldValue);
 
diff --git a/Engine/Source/Runtime/Engine/Private/Actor.cpp b/Engine/Source/Runtime/Engine/Private/Actor.cpp
--- a/Engine/Source/Runtime/Engine/Private/Actor.cpp
+++ b/Engine/Source/Runtime/Engine/Private/Actor.cpp
@@ -805,6 +805,9 @@ void AActor::PreReplication( IRepChangedPropertyTracker & ChangedPropertyTracker
 
 	DOREPLIFETIME_ACTIVE_OVERRIDE( AActor, ReplicatedMovement, bReplicateMovement );
 
+	// Don't need to replicate AttachmentReplication if the root component replicates, because it already handles it.
+	DOREPLIFETIME_ACTIVE_OVERRIDE( AActor, AttachmentReplication, RootComponent && !RootComponent->GetIsReplicated() );
+
 	UBlueprintGeneratedClass* BPClass = Cast<UBlueprintGeneratedClass>(GetClass());
 	if (BPClass != NULL)
 	{
diff --git a/Engine/Source/Runtime/Engine/Private/ActorReplication.cpp b/Engine/Source/Runtime/Engine/Private/ActorReplication.cpp
--- a/Engine/Source/Runtime/Engine/Private/ActorReplication.cpp
+++ b/Engine/Source/Runtime/Engine/Private/ActorReplication.cpp
@@ -282,7 +282,7 @@ void AActor::GetLifetimeReplicatedProps( TArray< FLifetimeProperty > & OutLifeti
 
 	DOREPLIFETIME( AActor, bTearOff );
 	DOREPLIFETIME( AActor, bCanBeDamaged );
-	DOREPLIFETIME( AActor, AttachmentReplication );
+	DOREPLIFETIME_CONDITION( AActor, AttachmentReplication, COND_Custom );
 
 	DOREPLIFETIME( AActor, Instigator );
 
diff --git a/Engine/Source/Runtime/Engine/Private/Components/SceneComponent.cpp b/Engine/Source/Runtime/Engine/Private/Components/SceneComponent.cpp
--- a/Engine/Source/Runtime/Engine/Private/Components/SceneComponent.cpp
+++ b/Engine/Source/Runtime/Engine/Private/Components/SceneComponent.cpp
@@ -2209,6 +2209,16 @@ void USceneComponent::OnRep_Transform()
 	bNetUpdateTransform = true;
 }
 
+void USceneComponent::OnRep_AttachParent()
+{
+	bNetUpdateAttachment = true;
+}
+
+void USceneComponent::OnRep_AttachSocketName()
+{
+	bNetUpdateAttachment = true;
+}
+
 void USceneComponent::OnRep_Visibility(bool OldValue)
 {
 	bool ReppedValue = bVisible;
@@ -2220,6 +2230,8 @@ void USceneComponent::PreNetReceive()
 {
 	Super::PreNetReceive();
 
+	bNetUpdateTransform = false;
+	bNetUpdateAttachment = false;
 	NetOldAttachSocketName = AttachSocketName;
 	NetOldAttachParent = AttachParent;
 }
@@ -2229,32 +2241,27 @@ void USceneComponent::PostNetReceive()
 	Super::PostNetReceive();
 
 	// If we have no attach parent, attach to parent's root component.
-	bool UpdateAttach = false;
-
-	UpdateAttach |= (NetOldAttachParent != AttachParent);
-	UpdateAttach |= (NetOldAttachSocketName != AttachSocketName);
-
 	if (AttachParent == NULL)
 	{
 		USceneComponent * ParentRoot = GetOwner()->GetRootComponent();
 		if (ParentRoot != this)
 		{
-			UpdateAttach = true;
+			bNetUpdateAttachment = true;
 			AttachParent = ParentRoot;
 		}
 	}
-	
-	if (UpdateAttach)
+}
+
+void USceneComponent::PostRepNotifies()
+{
+	if (bNetUpdateAttachment)
 	{
 		Exchange(NetOldAttachParent, AttachParent);
 		Exchange(NetOldAttachSocketName, AttachSocketName);
-		
 		AttachTo(NetOldAttachParent, NetOldAttachSocketName);
+		bNetUpdateAttachment = false;
 	}
-}
 
-void USceneComponent::PostRepNotifies()
-{
 	if (bNetUpdateTransform)
 	{
 		UpdateComponentToWorld(true);
