Temporarily commenting out check in UMaterialInstance::GetMaterialResource() to unblock opening the landscape editor

diff --git a/Engine/Source/Runtime/Engine/Private/Materials/MaterialInstance.cpp b/Engine/Source/Runtime/Engine/Private/Materials/MaterialInstance.cpp
--- a/Engine/Source/Runtime/Engine/Private/Materials/MaterialInstance.cpp
+++ b/Engine/Source/Runtime/Engine/Private/Materials/MaterialInstance.cpp
@@ -1100,7 +1100,7 @@ FMaterialResource* UMaterialInstance::GetMaterialResource(ERHIFeatureLevel::Type
 
 const FMaterialResource* UMaterialInstance::GetMaterialResource(ERHIFeatureLevel::Type InFeatureLevel, EMaterialQualityLevel::Type QualityLevel) const
 {
-	check(!IsInActualRenderingThread());
+	//check(!IsInActualRenderingThread());
 
 	if (QualityLevel == EMaterialQualityLevel::Num)
 	{
