Custom thunk float divide float and float modulo float to provide a script stack trace in case of divide by zero

Merging 2647889

//depot/UE4-UT/...

to //depot/UE4/...

diff --git a/Engine/Source/Runtime/Engine/Classes/Kismet/KismetMathLibrary.h b/Engine/Source/Runtime/Engine/Classes/Kismet/KismetMathLibrary.h
--- a/Engine/Source/Runtime/Engine/Classes/Kismet/KismetMathLibrary.h
+++ b/Engine/Source/Runtime/Engine/Classes/Kismet/KismetMathLibrary.h
@@ -253,13 +253,53 @@ class ENGINE_API UKismetMathLibrary : public UBlueprintFunctionLibrary
 	static float Multiply_IntFloat(int32 A, float B);
 
 	/* Division (A / B) */
-	UFUNCTION(BlueprintPure, meta=(DisplayName = "float / float", CompactNodeTitle = "/", Keywords = "/ divide division"), Category="Math|Float")
+	UFUNCTION(BlueprintPure, CustomThunk, meta=(DisplayName = "float / float", CompactNodeTitle = "/", Keywords = "/ divide division"), Category="Math|Float")
 	static float Divide_FloatFloat(float A, float B = 1.f);
+	
+	static float GenericDivide_FloatFloat(float A, float B);
+
+	/* Custom thunk to allow script stack trace in case of divide by zero */
+	DECLARE_FUNCTION(execDivide_FloatFloat)
+	{
+		P_GET_PROPERTY(UFloatProperty, A);
+		P_GET_PROPERTY(UFloatProperty, B);
+
+		P_FINISH;
+
+		if (B == 0.f)
+		{
+			FFrame::KismetExecutionMessage(*FString::Printf(TEXT("Divide by zero\n%s"), *Stack.GetStackTrace()), ELogVerbosity::Warning);
+			*(float*)RESULT_PARAM = 0;
+			return;
+		}
+
+		*(float*)RESULT_PARAM = GenericDivide_FloatFloat(A, B);
+	}
 
 	/* Modulo (A % B) */
-	UFUNCTION(BlueprintPure, meta=(DisplayName = "% (float)", CompactNodeTitle = "%", Keywords = "% modulus"), Category="Math|Float")
+	UFUNCTION(BlueprintPure, CustomThunk, meta = (DisplayName = "% (float)", CompactNodeTitle = "%", Keywords = "% modulus"), Category = "Math|Float")
 	static float Percent_FloatFloat(float A, float B = 1.f);
 
+	static float GenericPercent_FloatFloat(float A, float B);
+
+	/* Custom thunk to allow script stack trace in case of modulo by zero */
+	DECLARE_FUNCTION(execPercent_FloatFloat)
+	{
+		P_GET_PROPERTY(UFloatProperty, A);
+		P_GET_PROPERTY(UFloatProperty, B);
+
+		P_FINISH;
+
+		if (B == 0.f)
+		{
+			FFrame::KismetExecutionMessage(*FString::Printf(TEXT("Modulo by zero\n%s"), *Stack.GetStackTrace()), ELogVerbosity::Warning);
+			*(float*)RESULT_PARAM = 0;
+			return;
+		}
+
+		*(float*)RESULT_PARAM = GenericPercent_FloatFloat(A, B);
+	}
+
 	/** Returns the fractional part of a float. */
 	UFUNCTION(BlueprintPure, Category="Math|Float")
 	static float Fraction(float A);
diff --git a/Engine/Source/Runtime/Engine/Private/KismetMathLibrary.cpp b/Engine/Source/Runtime/Engine/Private/KismetMathLibrary.cpp
--- a/Engine/Source/Runtime/Engine/Private/KismetMathLibrary.cpp
+++ b/Engine/Source/Runtime/Engine/Private/KismetMathLibrary.cpp
@@ -276,27 +276,27 @@ float UKismetMathLibrary::Multiply_IntFloat(int32 A, float B)
 	return A * B;
 }	
 
+/* This function is custom thunked, the real function is GenericDivide_FloatFloat */
 float UKismetMathLibrary::Divide_FloatFloat(float A, float B)
 {
-	if (B == 0.f)
-	{
-		//@TODO: EXCEPTION: Throw script exception 
-		FFrame::KismetExecutionMessage(TEXT("Divide by zero: Divide_FloatFloat"), ELogVerbosity::Warning);
-		return 0.f;
-	}
+	check(0);
+	return 0;
+}
 
+float UKismetMathLibrary::GenericDivide_FloatFloat(float A, float B)
+{
 	return A / B;
 }	
 
+/* This function is custom thunked, the real function is GenericPercent_FloatFloat */
 float UKismetMathLibrary::Percent_FloatFloat(float A, float B)
 {
-	if (B == 0.f)
-	{
-		//@TODO: EXCEPTION: Throw script exception 
-		FFrame::KismetExecutionMessage(TEXT("Modulo by zero"), ELogVerbosity::Warning);
-		return 0.f;
-	}
+	check(0);
+	return 0;
+}
 
+float UKismetMathLibrary::GenericPercent_FloatFloat(float A, float B)
+{
 	return (B != 0.f) ? FMath::Fmod(A, B) : 0.f;
 }	
 
