Fix check in FUniformExpressionSet::CreateUniformBuffer so that it doesn't itself cause a crash when it fires

diff --git a/Engine/Source/Runtime/Engine/Private/Materials/MaterialUniformExpressions.cpp b/Engine/Source/Runtime/Engine/Private/Materials/MaterialUniformExpressions.cpp
--- a/Engine/Source/Runtime/Engine/Private/Materials/MaterialUniformExpressions.cpp
+++ b/Engine/Source/Runtime/Engine/Private/Materials/MaterialUniformExpressions.cpp
@@ -352,7 +352,10 @@ FUniformBufferRHIRef FUniformExpressionSet::CreateUniformBuffer(const FMaterialR
 			if (Value && Value->Resource)
 			{
 				//@todo-rco: Help track down a invalid values
-				checkf(Value->IsA(UTexture::StaticClass()), TEXT("Expecting a UTexture! Value='%s' class='%s'"), *Value->GetName(), *Value->GetClass()->GetName());
+				checkf(Value->IsA(UTexture::StaticClass()),
+					   TEXT("Expecting a UTexture! Value='%s' class='%s'"),
+					   Value->IsValidLowLevel() ? *Value->GetName() : TEXT("!Value->IsValidLowLevel()"),
+					   Value->IsValidLowLevel() ? *Value->GetClass()->GetName() : TEXT("!Value->IsValidLowLevel()"));
 
 				// UMaterial / UMaterialInstance should have caused all dependent textures to be PostLoaded, which initializes their rendering resource
 				checkf(Value->TextureReference.TextureReferenceRHI, TEXT("Texture %s of class %s had invalid texture reference"), *Value->GetName(), *Value->GetClass()->GetName());
