Restore IsValidLowLevel() check to try to help find UE-17608

diff --git a/Engine/Source/Runtime/Engine/Private/Materials/MaterialUniformExpressions.cpp b/Engine/Source/Runtime/Engine/Private/Materials/MaterialUniformExpressions.cpp
--- a/Engine/Source/Runtime/Engine/Private/Materials/MaterialUniformExpressions.cpp
+++ b/Engine/Source/Runtime/Engine/Private/Materials/MaterialUniformExpressions.cpp
@@ -352,7 +352,7 @@ FUniformBufferRHIRef FUniformExpressionSet::CreateUniformBuffer(const FMaterialR
 			if (Value && Value->Resource)
 			{
 				//@todo-rco: Help track down a invalid values
-				checkf(Value->IsA(UTexture::StaticClass()),
+				checkf(Value->IsValidLowLevel() && Value->IsA(UTexture::StaticClass()),
 					   TEXT("Expecting a UTexture! Value='%s' class='%s'"),
 					   Value->IsValidLowLevel() ? *Value->GetName() : TEXT("!Value->IsValidLowLevel()"),
 					   Value->IsValidLowLevel() ? *Value->GetClass()->GetName() : TEXT("!Value->IsValidLowLevel()"));
