UE-17608 - split into our own check

diff --git a/Engine/Source/Runtime/Engine/Private/Materials/MaterialUniformExpressions.cpp b/Engine/Source/Runtime/Engine/Private/Materials/MaterialUniformExpressions.cpp
--- a/Engine/Source/Runtime/Engine/Private/Materials/MaterialUniformExpressions.cpp
+++ b/Engine/Source/Runtime/Engine/Private/Materials/MaterialUniformExpressions.cpp
@@ -349,13 +349,14 @@ FUniformBufferRHIRef FUniformExpressionSet::CreateUniformBuffer(const FMaterialR
 			const UTexture* Value;
 			ESamplerSourceMode SourceMode;
 			Uniform2DTextureExpressions[ExpressionIndex]->GetTextureValue(MaterialRenderContext,MaterialRenderContext.Material,Value,SourceMode);
+
+			// gmartin: Trying to locate UE-17608
+			checkf(Value && Value->IsValidLowLevel(), TEXT("Texture not valid! UE-17608!"));
+
 			if (Value && Value->Resource)
 			{
 				//@todo-rco: Help track down a invalid values
-				checkf(Value->IsValidLowLevel() && Value->IsA(UTexture::StaticClass()),
-					   TEXT("Expecting a UTexture! Value='%s' class='%s'"),
-					   Value->IsValidLowLevel() ? *Value->GetName() : TEXT("!Value->IsValidLowLevel()"),
-					   Value->IsValidLowLevel() ? *Value->GetClass()->GetName() : TEXT("!Value->IsValidLowLevel()"));
+				checkf(Value->IsA(UTexture::StaticClass()), TEXT("Expecting a UTexture! Value='%s' class='%s'"), *Value->GetName(), *Value->GetClass()->GetName());
 
 				// UMaterial / UMaterialInstance should have caused all dependent textures to be PostLoaded, which initializes their rendering resource
 				checkf(Value->TextureReference.TextureReferenceRHI, TEXT("Texture %s of class %s had invalid texture reference"), *Value->GetName(), *Value->GetClass()->GetName());
