#ue4 - Add SceneComponent::IsQueryCollisionEnabled() and IsPhysicsCollisionEnabled(), since IsCollisionEnabled() is more general than we always want.

- MoveComponent now only sweeps if query collision is enabled (not just any collision, so now excludes PhysicsOnly collision).
- UpdateOverlaps now only generates overlaps if query collision is enabled (not just any collision, so now excludes PhysicsOnly collision).
- A few other locations were changed to only consider query collision, where it didn't make sense to include PhysicsOnly collision.

diff --git a/Engine/Source/Runtime/Engine/Classes/Components/SceneComponent.h b/Engine/Source/Runtime/Engine/Classes/Components/SceneComponent.h
--- a/Engine/Source/Runtime/Engine/Classes/Components/SceneComponent.h
+++ b/Engine/Source/Runtime/Engine/Classes/Components/SceneComponent.h
@@ -961,12 +961,26 @@ public:
 	/** Returns the form of collision for this component */
 	virtual ECollisionEnabled::Type GetCollisionEnabled() const;
 
-	/** Utility to see if there is any form of collision enabled on this component */
-	bool IsCollisionEnabled() const
+	/** Utility to see if there is any form of collision (query or physics) enabled on this component. */
+	FORCEINLINE_DEBUGGABLE bool IsCollisionEnabled() const
 	{
 		return GetCollisionEnabled() != ECollisionEnabled::NoCollision;
 	}
 
+	/** Utility to see if there is any query collision enabled on this component. */
+	FORCEINLINE_DEBUGGABLE bool IsQueryCollisionEnabled() const
+	{
+		const ECollisionEnabled::Type CollisionSetting = GetCollisionEnabled();
+		return (CollisionSetting == ECollisionEnabled::QueryAndPhysics) || (CollisionSetting == ECollisionEnabled::QueryOnly);
+	}
+
+	/** Utility to see if there is any physics collision enabled on this component. */
+	FORCEINLINE_DEBUGGABLE bool IsPhysicsCollisionEnabled() const
+	{
+		const ECollisionEnabled::Type CollisionSetting = GetCollisionEnabled();
+		return (CollisionSetting == ECollisionEnabled::QueryAndPhysics) || (CollisionSetting == ECollisionEnabled::PhysicsOnly);
+	}
+
 	/** Returns the response that this component has to a specific collision channel. */
 	virtual ECollisionResponse GetCollisionResponseToChannel(ECollisionChannel Channel) const;
 
diff --git a/Engine/Source/Runtime/Engine/Private/Actor.cpp b/Engine/Source/Runtime/Engine/Private/Actor.cpp
--- a/Engine/Source/Runtime/Engine/Private/Actor.cpp
+++ b/Engine/Source/Runtime/Engine/Private/Actor.cpp
@@ -353,7 +353,7 @@ bool AActor::TeleportTo( const FVector& DestLocation, const FRotator& DestRotati
 	UPrimitiveComponent* ActorPrimComp = Cast<UPrimitiveComponent>(RootComponent);
 	if ( ActorPrimComp )
 	{
-		if (!bNoCheck && (ActorPrimComp->IsCollisionEnabled() || (GetNetMode() != NM_Client)) )
+		if (!bNoCheck && (ActorPrimComp->IsQueryCollisionEnabled() || (GetNetMode() != NM_Client)) )
 		{
 			// Apply the pivot offset to the desired location
 			FVector Offset = GetRootComponent()->Bounds.Origin - PrevLocation;
diff --git a/Engine/Source/Runtime/Engine/Private/Components/CharacterMovementComponent.cpp b/Engine/Source/Runtime/Engine/Private/Components/CharacterMovementComponent.cpp
--- a/Engine/Source/Runtime/Engine/Private/Components/CharacterMovementComponent.cpp
+++ b/Engine/Source/Runtime/Engine/Private/Components/CharacterMovementComponent.cpp
@@ -3624,7 +3624,7 @@ void UCharacterMovementComponent::PhysWalking(float deltaTime, int32 Iterations)
 		return;
 	}
 
-	if (!UpdatedComponent->IsCollisionEnabled())
+	if (!UpdatedComponent->IsQueryCollisionEnabled())
 	{
 		SetMovementMode(MOVE_Walking);
 		return;
@@ -3732,7 +3732,7 @@ void UCharacterMovementComponent::PhysWalking(float deltaTime, int32 Iterations)
 			{
 				// see if it is OK to jump
 				// @todo collision : only thing that can be problem is that oldbase has world collision on
-				bool bMustJump = bZeroDelta || (OldBase == NULL || (!OldBase->IsCollisionEnabled() && MovementBaseUtility::IsDynamicBase(OldBase)));
+				bool bMustJump = bZeroDelta || (OldBase == NULL || (!OldBase->IsQueryCollisionEnabled() && MovementBaseUtility::IsDynamicBase(OldBase)));
 				if ( (bMustJump || !bCheckedFall) && CheckFall(OldFloor, CurrentFloor.HitResult, Delta, OldLocation, remainingTime, timeTick, Iterations, bMustJump) )
 				{
 					return;
@@ -3784,7 +3784,7 @@ void UCharacterMovementComponent::PhysWalking(float deltaTime, int32 Iterations)
 			// See if we need to start falling.
 			if (!CurrentFloor.IsWalkableFloor() && !CurrentFloor.HitResult.bStartPenetrating)
 			{
-				const bool bMustJump = bJustTeleported || bZeroDelta || (OldBase == NULL || (!OldBase->IsCollisionEnabled() && MovementBaseUtility::IsDynamicBase(OldBase)));
+				const bool bMustJump = bJustTeleported || bZeroDelta || (OldBase == NULL || (!OldBase->IsQueryCollisionEnabled() && MovementBaseUtility::IsDynamicBase(OldBase)));
 				if ((bMustJump || !bCheckedFall) && CheckFall(OldFloor, CurrentFloor.HitResult, Delta, OldLocation, remainingTime, timeTick, Iterations, bMustJump) )
 				{
 					return;
@@ -4712,7 +4712,7 @@ void UCharacterMovementComponent::ComputeFloorDist(const FVector& CapsuleLocatio
 	OutFloorResult.Clear();
 
 	// No collision, no floor...
-	if (!UpdatedComponent->IsCollisionEnabled())
+	if (!UpdatedComponent->IsQueryCollisionEnabled())
 	{
 		return;
 	}
@@ -4858,7 +4858,7 @@ void UCharacterMovementComponent::FindFloor(const FVector& CapsuleLocation, FFin
 	SCOPE_CYCLE_COUNTER(STAT_CharFindFloor);
 
 	// No collision, no floor...
-	if (!HasValidData() || !UpdatedComponent->IsCollisionEnabled())
+	if (!HasValidData() || !UpdatedComponent->IsQueryCollisionEnabled())
 	{
 		OutFloorResult.Clear();
 		return;
@@ -4892,7 +4892,7 @@ void UCharacterMovementComponent::FindFloor(const FVector& CapsuleLocation, FFin
 
 			if (MovementBase != NULL)
 			{
-				MutableThis->bForceNextFloorCheck = !MovementBase->IsCollisionEnabled()
+				MutableThis->bForceNextFloorCheck = !MovementBase->IsQueryCollisionEnabled()
 				|| MovementBase->GetCollisionResponseToChannel(CollisionChannel) != ECR_Block
 				|| MovementBaseUtility::IsDynamicBase(MovementBase);
 			}
diff --git a/Engine/Source/Runtime/Engine/Private/Components/PrimitiveComponent.cpp b/Engine/Source/Runtime/Engine/Private/Components/PrimitiveComponent.cpp
--- a/Engine/Source/Runtime/Engine/Private/Components/PrimitiveComponent.cpp
+++ b/Engine/Source/Runtime/Engine/Private/Components/PrimitiveComponent.cpp
@@ -1435,7 +1435,7 @@ bool UPrimitiveComponent::MoveComponentImpl( const FVector& Delta, const FQuat&
 		// not sweeping, just go directly to the new transform
 		bMoved = InternalSetWorldLocationAndRotation(TraceEnd, NewRotationQuat, bSkipPhysicsMove, Teleport);
 		bRotationOnly = (DeltaSizeSq == 0);
-		bIncludesOverlapsAtEnd = bRotationOnly && (AreSymmetricRotations(InitialRotationQuat, NewRotationQuat, GetComponentScale())) && IsCollisionEnabled();
+		bIncludesOverlapsAtEnd = bRotationOnly && (AreSymmetricRotations(InitialRotationQuat, NewRotationQuat, GetComponentScale())) && IsQueryCollisionEnabled();
 	}
 	else
 	{
@@ -1443,7 +1443,7 @@ bool UPrimitiveComponent::MoveComponentImpl( const FVector& Delta, const FQuat&
 		FVector NewLocation = TraceStart;
 
 		// Perform movement collision checking if needed for this actor.
-		const bool bCollisionEnabled = IsCollisionEnabled();
+		const bool bCollisionEnabled = IsQueryCollisionEnabled();
 		if( bCollisionEnabled && (DeltaSizeSq > 0.f))
 		{
 #if !(UE_BUILD_SHIPPING || UE_BUILD_TEST)
@@ -1754,7 +1754,7 @@ bool UPrimitiveComponent::IsNavigationRelevant() const
 	}
 
 	const FCollisionResponseContainer& ResponseToChannels = GetCollisionResponseToChannels();
-	return IsCollisionEnabled() &&
+	return IsQueryCollisionEnabled() &&
 		(ResponseToChannels.GetResponse(ECC_Pawn) == ECR_Block || ResponseToChannels.GetResponse(ECC_Vehicle) == ECR_Block);
 }
 
@@ -2241,7 +2241,7 @@ bool UPrimitiveComponent::AreAllCollideableDescendantsRelative(bool bAllowCached
 				{
 					// Can we possibly collide with the component?
 					UPrimitiveComponent* const CurrentPrimitive = Cast<UPrimitiveComponent>(CurrentComp);
-					if (CurrentPrimitive && CurrentPrimitive->bGenerateOverlapEvents && CurrentPrimitive->IsCollisionEnabled() && CurrentPrimitive->GetCollisionResponseToChannel(GetCollisionObjectType()) != ECR_Ignore)
+					if (CurrentPrimitive && CurrentPrimitive->bGenerateOverlapEvents && CurrentPrimitive->IsQueryCollisionEnabled() && CurrentPrimitive->GetCollisionResponseToChannel(GetCollisionObjectType()) != ECR_Ignore)
 					{
 						MutableThis->bCachedAllCollideableDescendantsRelative = false;
 						MutableThis->LastCheckedAllCollideableDescendantsTime = MyWorld->GetTimeSeconds();
@@ -2322,7 +2322,7 @@ void UPrimitiveComponent::UpdateOverlaps(const TArray<FOverlapInfo>* NewPendingO
 	}
 
 	// first, dispatch any pending overlaps
-	if (bGenerateOverlapEvents && IsCollisionEnabled())
+	if (bGenerateOverlapEvents && IsQueryCollisionEnabled())
 	{
 		// if we haven't begun play, we're still setting things up (e.g. we might be inside one of the construction scripts)
 		// so we don't want to generate overlaps yet.
@@ -2484,7 +2484,7 @@ void UPrimitiveComponent::UpdatePhysicsVolume( bool bTriggerNotifiers )
 	{
 		SCOPE_CYCLE_COUNTER(STAT_UpdatePhysicsVolume);
 
-		if (bGenerateOverlapEvents && IsCollisionEnabled())
+		if (bGenerateOverlapEvents && IsQueryCollisionEnabled())
 		{
 			APhysicsVolume* BestVolume = GetWorld()->GetDefaultPhysicsVolume();
 			int32 BestPriority = BestVolume->Priority;
diff --git a/Engine/Source/Runtime/Engine/Private/LevelActor.cpp b/Engine/Source/Runtime/Engine/Private/LevelActor.cpp
--- a/Engine/Source/Runtime/Engine/Private/LevelActor.cpp
+++ b/Engine/Source/Runtime/Engine/Private/LevelActor.cpp
@@ -925,7 +925,7 @@ bool UWorld::EncroachingBlockingGeometry(AActor* TestActor, FVector TestLocation
 		// the world, and that component is the only one we care about encroaching (since the movement code will happily embedding
 		// other components in the world during movement updates)
 		UPrimitiveComponent* const MovedPrimComp = MoveComponent->UpdatedPrimitive;
-		if (MovedPrimComp->IsCollisionEnabled())
+		if (MovedPrimComp->IsQueryCollisionEnabled())
 		{
 			// might not be the root, so we need to compute the transform
 			FTransform const CompToRoot = MovedPrimComp->GetComponentToWorld() * WorldToOldRoot;
@@ -951,7 +951,7 @@ bool UWorld::EncroachingBlockingGeometry(AActor* TestActor, FVector TestLocation
 	{
 		// This actor does not have a movement component, so we'll assume all components are potentially important to keep out of the world
 		UPrimitiveComponent* const RootPrimComp = Cast<UPrimitiveComponent>(RootComponent);
-		if (RootPrimComp && RootPrimComp->IsCollisionEnabled())
+		if (RootPrimComp && RootPrimComp->IsQueryCollisionEnabled())
 		{
 			if (ComponentEncroachesBlockingGeometry(this, TestActor, RootPrimComp, TestRootToWorld, ProposedAdjustment))
 			{
@@ -975,7 +975,7 @@ bool UWorld::EncroachingBlockingGeometry(AActor* TestActor, FVector TestLocation
 
 		for (auto Child : Children)
 		{
-			if (Child->IsCollisionEnabled())
+			if (Child->IsQueryCollisionEnabled())
 			{
 				UPrimitiveComponent* const PrimComp = Cast<UPrimitiveComponent>(Child);
 				if (PrimComp)
