#ue4 - Avoid triggering incorrect Begin/End overlap events when BeginOverlap is used to disable collision or overlap events on the other overlapped object.

diff --git a/Engine/Source/Runtime/Engine/Private/Components/PrimitiveComponent.cpp b/Engine/Source/Runtime/Engine/Private/Components/PrimitiveComponent.cpp
--- a/Engine/Source/Runtime/Engine/Private/Components/PrimitiveComponent.cpp
+++ b/Engine/Source/Runtime/Engine/Private/Components/PrimitiveComponent.cpp
@@ -105,6 +105,28 @@ private:
 	const AActor& MyOwner;
 };
 
+// Predicate to remove components from overlaps array that can no longer overlap
+struct FPredicateFilterCannotOverlap
+{
+	FPredicateFilterCannotOverlap(const UPrimitiveComponent& OwningComponent)
+	: MyComponent(OwningComponent)
+	{
+	}
+
+	// return true if not collideable
+	bool operator() (const FOverlapInfo& Info) const
+	{
+		UPrimitiveComponent* OtherComp = Info.OverlapInfo.GetComponent();
+		return !OtherComp
+			|| !OtherComp->bGenerateOverlapEvents
+			|| !MyComponent.bGenerateOverlapEvents
+			|| MyComponent.GetCollisionResponseToComponent(OtherComp) != ECR_Overlap;
+	}
+
+private:
+	const UPrimitiveComponent& MyComponent;
+};
+
 
 ///////////////////////////////////////////////////////////////////////////////
 // PRIMITIVE COMPONENT
@@ -2354,6 +2376,12 @@ void UPrimitiveComponent::UpdateOverlaps(const TArray<FOverlapInfo>* NewPendingO
 				{
 					UE_LOG(LogPrimitiveComponent, VeryVerbose, TEXT("%s->%s Skipping overlap test!"), *GetNameSafe(GetOwner()), *GetName());
 					NewOverlappingComponents = *OverlapsAtEndLocation;
+
+					// BeginComponentOverlap may have disabled what we thought were valid overlaps at the end.
+					if (NewPendingOverlaps && NewPendingOverlaps->Num() > 0)
+					{
+						NewOverlappingComponents.RemoveAllSwap(FPredicateFilterCannotOverlap(*this), /*bAllowShrinking*/ false);
+					}
 				}
 				else
 				{
@@ -2389,7 +2417,6 @@ void UPrimitiveComponent::UpdateOverlaps(const TArray<FOverlapInfo>* NewPendingO
 				if (bIgnoreChildren)
 				{
 					OldOverlappingComponents = OverlappingComponents.FilterByPredicate(FPredicateOverlapHasDifferentActor(*MyActor));
-					checkSlow(!NewOverlappingComponents.FindByPredicate(FPredicateOverlapHasSameActor(*MyActor)));
 				}
 				else
 				{
