Rendering: Fixed messed up format specifier in "Built static mesh" message

[INTEGRATION] Change 2648990 by [REDACTED_PERSON][REDACTED_HANDLE] on 2015/08/08 14:24:07

diff --git a/Engine/Source/Runtime/Engine/Private/StaticMesh.cpp b/Engine/Source/Runtime/Engine/Private/StaticMesh.cpp
--- a/Engine/Source/Runtime/Engine/Private/StaticMesh.cpp
+++ b/Engine/Source/Runtime/Engine/Private/StaticMesh.cpp
@@ -1197,7 +1197,7 @@ void FStaticMeshRenderData::Cache(UStaticMesh* Owner, const FStaticMeshLODSettin
 		GetDerivedDataCacheRef().Put(*DerivedDataKey, DerivedData);
 
 		int32 T1 = FPlatformTime::Cycles();
-		UE_LOG(LogStaticMesh,Log,TEXT("Built static mesh [%f.2s] %s"),
+		UE_LOG(LogStaticMesh,Log,TEXT("Built static mesh [%.2fs] %s"),
 			FPlatformTime::ToMilliseconds(T1-T0) / 1000.0f,
 			*Owner->GetPathName()
 			);
