UE-17608 - Allow nulls to pass check

diff --git a/Engine/Source/Runtime/Engine/Private/Materials/MaterialUniformExpressions.cpp b/Engine/Source/Runtime/Engine/Private/Materials/MaterialUniformExpressions.cpp
--- a/Engine/Source/Runtime/Engine/Private/Materials/MaterialUniformExpressions.cpp
+++ b/Engine/Source/Runtime/Engine/Private/Materials/MaterialUniformExpressions.cpp
@@ -351,7 +351,7 @@ FUniformBufferRHIRef FUniformExpressionSet::CreateUniformBuffer(const FMaterialR
 			Uniform2DTextureExpressions[ExpressionIndex]->GetTextureValue(MaterialRenderContext,MaterialRenderContext.Material,Value,SourceMode);
 
 			// gmartin: Trying to locate UE-17608
-			checkf(Value && Value->IsValidLowLevel(), TEXT("Texture not valid! UE-17608!"));
+			checkf(!Value || Value->IsValidLowLevel(), TEXT("Texture not valid! UE-17608!"));
 
 			if (Value && Value->Resource)
 			{
