Github pull request : GitHub #1435

diff --git a/Engine/Source/Runtime/Engine/Classes/GameFramework/Character.h b/Engine/Source/Runtime/Engine/Classes/GameFramework/Character.h
--- a/Engine/Source/Runtime/Engine/Classes/GameFramework/Character.h
+++ b/Engine/Source/Runtime/Engine/Classes/GameFramework/Character.h
@@ -460,7 +460,7 @@ public:
 	UFUNCTION(BlueprintCallable, Category=Animation)
 	virtual float PlayAnimMontage(class UAnimMontage* AnimMontage, float InPlayRate = 1.f, FName StartSectionName = NAME_None);
 
-	/** Stop Animation Montage. If NULL, it will stop what's currently active **/
+	/** Stop Animation Montage. If NULL, it will stop what's currently active. The Blend Out Time is taken from the montage asset that is being stopped. **/
 	UFUNCTION(BlueprintCallable, Category=Animation)
 	virtual void StopAnimMontage(class UAnimMontage* AnimMontage = NULL);
 
