Fix for cooking static mesh LODs when using cooked editor packages

diff --git a/Engine/Source/Runtime/Engine/Private/StaticMesh.cpp b/Engine/Source/Runtime/Engine/Private/StaticMesh.cpp
--- a/Engine/Source/Runtime/Engine/Private/StaticMesh.cpp
+++ b/Engine/Source/Runtime/Engine/Private/StaticMesh.cpp
@@ -622,7 +622,7 @@ void FStaticMeshRenderData::Serialize(FArchive& Ar, UStaticMesh* Owner, bool bCo
 
 	// Note: this is all derived data, native versioning is not needed, but be sure to bump STATICMESH_DERIVEDDATA_VER when modifying!
 #if WITH_EDITOR
-	const bool bHasEditorData = !(Owner->GetOutermost()->PackageFlags & PKG_FilterEditorOnly);
+	const bool bHasEditorData = !Owner->GetOutermost()->bIsCookedForEditor;
 	if (Ar.IsSaving() && bHasEditorData)
 	{
 		ResolveSectionInfo(Owner);
