Merge CVAR to disable all morph targets.  Expose SkeletalMesh property to disable morph targets on a specific mesh.

diff --git a/Engine/Source/Runtime/Engine/Classes/Components/SkinnedMeshComponent.h b/Engine/Source/Runtime/Engine/Classes/Components/SkinnedMeshComponent.h
--- a/Engine/Source/Runtime/Engine/Classes/Components/SkinnedMeshComponent.h
+++ b/Engine/Source/Runtime/Engine/Classes/Components/SkinnedMeshComponent.h
@@ -294,7 +294,7 @@ public:
 	uint32 bDisplayBones:1;
 
 	/** Disable Morphtarget for this component. */
-	UPROPERTY()
+	UPROPERTY(EditAnywhere, AdvancedDisplay, BlueprintReadWrite, Category = SkeletalMesh)
 	uint32 bDisableMorphTarget:1;
 
 	/** Don't bother rendering the skin. */
diff --git a/Engine/Source/Runtime/Engine/Private/Components/SkinnedMeshComponent.cpp b/Engine/Source/Runtime/Engine/Private/Components/SkinnedMeshComponent.cpp
--- a/Engine/Source/Runtime/Engine/Private/Components/SkinnedMeshComponent.cpp
+++ b/Engine/Source/Runtime/Engine/Private/Components/SkinnedMeshComponent.cpp
@@ -36,6 +36,8 @@ static TAutoConsoleVariable<int32> CVarDrawAnimRateOptimization(
 	0,
 	TEXT("True to draw color coded boxes for anim rate."));
 
+static TAutoConsoleVariable<int32> CVarEnableMorphTargets(TEXT("r.EnableMorphTargets"), 1, TEXT("Enable Morph Targets"));
+
 namespace FAnimUpdateRateManager
 {
 	// Global counter to spread SkinnedMeshComponent tick updates.
@@ -316,6 +318,15 @@ void USkinnedMeshComponent::CreateRenderState_Concurrent()
 
 			// verifies vertex animations are valid
 			RefreshActiveVertexAnims();
+
+			const bool bMorphTargetsAllowed = CVarEnableMorphTargets.GetValueOnGameThread() != 0;
+
+			// Are morph targets disabled for this LOD?
+			if (SkeletalMesh->LODInfo[UseLOD].bHasBeenSimplified || bDisableMorphTarget || !bMorphTargetsAllowed)
+			{
+				ActiveVertexAnims.Empty();
+			}
+
 			MeshObject->Update(UseLOD,this,ActiveVertexAnims);  // send to rendering thread
 		}
 
@@ -371,8 +382,10 @@ void USkinnedMeshComponent::SendRenderDynamicData_Concurrent()
 
 		int32 UseLOD = PredictedLODLevel;
 
+		const bool bMorphTargetsAllowed = CVarEnableMorphTargets.GetValueOnGameThread() != 0;
+
 		// Are morph targets disabled for this LOD?
-		if ( SkeletalMesh->LODInfo[ UseLOD ].bHasBeenSimplified || bDisableMorphTarget )
+		if (SkeletalMesh->LODInfo[UseLOD].bHasBeenSimplified || bDisableMorphTarget || !bMorphTargetsAllowed)
 		{
 			ActiveVertexAnims.Empty();
 		}
