Fix morph target check

diff --git a/Engine/Source/Runtime/Engine/Private/Components/SkinnedMeshComponent.cpp b/Engine/Source/Runtime/Engine/Private/Components/SkinnedMeshComponent.cpp
--- a/Engine/Source/Runtime/Engine/Private/Components/SkinnedMeshComponent.cpp
+++ b/Engine/Source/Runtime/Engine/Private/Components/SkinnedMeshComponent.cpp
@@ -319,7 +319,7 @@ void USkinnedMeshComponent::CreateRenderState_Concurrent()
 			// verifies vertex animations are valid
 			RefreshActiveVertexAnims();
 
-			const bool bMorphTargetsAllowed = CVarEnableMorphTargets.GetValueOnGameThread() != 0;
+			const bool bMorphTargetsAllowed = CVarEnableMorphTargets.GetValueOnAnyThread() != 0;
 
 			// Are morph targets disabled for this LOD?
 			if (SkeletalMesh->LODInfo[UseLOD].bHasBeenSimplified || bDisableMorphTarget || !bMorphTargetsAllowed)
@@ -382,7 +382,7 @@ void USkinnedMeshComponent::SendRenderDynamicData_Concurrent()
 
 		int32 UseLOD = PredictedLODLevel;
 
-		const bool bMorphTargetsAllowed = CVarEnableMorphTargets.GetValueOnGameThread() != 0;
+		const bool bMorphTargetsAllowed = CVarEnableMorphTargets.GetValueOnAnyThread() != 0;
 
 		// Are morph targets disabled for this LOD?
 		if (SkeletalMesh->LODInfo[UseLOD].bHasBeenSimplified || bDisableMorphTarget || !bMorphTargetsAllowed)
