UE-17608 - More log info to try and narrow down the cause

diff --git a/Engine/Source/Runtime/Engine/Private/Materials/MaterialUniformExpressions.cpp b/Engine/Source/Runtime/Engine/Private/Materials/MaterialUniformExpressions.cpp
--- a/Engine/Source/Runtime/Engine/Private/Materials/MaterialUniformExpressions.cpp
+++ b/Engine/Source/Runtime/Engine/Private/Materials/MaterialUniformExpressions.cpp
@@ -351,7 +351,9 @@ FUniformBufferRHIRef FUniformExpressionSet::CreateUniformBuffer(const FMaterialR
 			Uniform2DTextureExpressions[ExpressionIndex]->GetTextureValue(MaterialRenderContext,MaterialRenderContext.Material,Value,SourceMode);
 
 			// gmartin: Trying to locate UE-17608
-			checkf(!Value || Value->IsValidLowLevel(), TEXT("Texture not valid! UE-17608!"));
+			checkf(!Value || Value->IsValidLowLevel(), TEXT("Texture not valid! UE-17608! Parameter (%s)"),
+				   Uniform2DTextureExpressions[ExpressionIndex]->GetType() == &FMaterialUniformExpressionTextureParameter::StaticType ?
+				   	*((FMaterialUniformExpressionTextureParameter&)*Uniform2DTextureExpressions[ExpressionIndex]).GetParameterName().ToString() : TEXT("non-parameter"));
 
 			if (Value && Value->Resource)
 			{
diff --git a/Engine/Source/Runtime/Engine/Private/Materials/MaterialUniformExpressions.h b/Engine/Source/Runtime/Engine/Private/Materials/MaterialUniformExpressions.h
--- a/Engine/Source/Runtime/Engine/Private/Materials/MaterialUniformExpressions.h
+++ b/Engine/Source/Runtime/Engine/Private/Materials/MaterialUniformExpressions.h
@@ -322,6 +322,12 @@ public:
 	{
 		return false;
 	}
+
+	FName GetParameterName() const
+	{
+		return ParameterName;
+	}
+
 	virtual bool IsIdentical(const FMaterialUniformExpression* OtherExpression) const
 	{
 		if (GetType() != OtherExpression->GetType())
