Remove improper code from destructors that were getting hit on dedicated server ctrl+c. IsInGameThread may not be valid during object destruction.

Merging 2649735

//depot/UE4-UT/...

to //depot/UE4/...

diff --git a/Engine/Plugins/Runtime/WindowsMoviePlayer/Source/WindowsMoviePlayer/Private/WindowsMovieStreamer.cpp b/Engine/Plugins/Runtime/WindowsMoviePlayer/Source/WindowsMoviePlayer/Private/WindowsMovieStreamer.cpp
--- a/Engine/Plugins/Runtime/WindowsMoviePlayer/Source/WindowsMoviePlayer/Private/WindowsMovieStreamer.cpp
+++ b/Engine/Plugins/Runtime/WindowsMoviePlayer/Source/WindowsMoviePlayer/Private/WindowsMovieStreamer.cpp
@@ -37,7 +37,6 @@ FMediaFoundationMovieStreamer::~FMediaFoundationMovieStreamer()
 	CloseMovie();
 	CleanupRenderingResources();
 
-	FlushRenderingCommands();
 	TextureFreeList.Empty();
 }
 
diff --git a/Engine/Source/Runtime/Engine/Private/Rendering/SubsurfaceProfile.cpp b/Engine/Source/Runtime/Engine/Private/Rendering/SubsurfaceProfile.cpp
--- a/Engine/Source/Runtime/Engine/Private/Rendering/SubsurfaceProfile.cpp
+++ b/Engine/Source/Runtime/Engine/Private/Rendering/SubsurfaceProfile.cpp
@@ -27,13 +27,6 @@ FSubsurfaceProfileTexture::FSubsurfaceProfileTexture()
 
 FSubsurfaceProfileTexture::~FSubsurfaceProfileTexture()
 {
-	check(IsInGameThread());
-
-	// we assume all ~USubsurfaceProfile() have been called already
-	for (int32 i = 0; i < SubsurfaceProfileEntries.Num(); ++i)
-	{
-		check(SubsurfaceProfileEntries[i].Profile == 0);
-	}
 }
 
 int32 FSubsurfaceProfileTexture::AddProfile(const FSubsurfaceProfileStruct Settings, const USubsurfaceProfile* InProfile)
