#ue4 - Ensure if we pass an invalid delta time to UCharacterMovementComponent::CalcRootMotionVelocity()

diff --git a/Engine/Source/Runtime/Engine/Private/Components/CharacterMovementComponent.cpp b/Engine/Source/Runtime/Engine/Private/Components/CharacterMovementComponent.cpp
--- a/Engine/Source/Runtime/Engine/Private/Components/CharacterMovementComponent.cpp
+++ b/Engine/Source/Runtime/Engine/Private/Components/CharacterMovementComponent.cpp
@@ -1152,15 +1152,22 @@ void UCharacterMovementComponent::SimulateRootMotion(float DeltaSeconds, const F
 
 FVector UCharacterMovementComponent::CalcRootMotionVelocity(const FVector& RootMotionDeltaMove, float DeltaSeconds, const FVector& CurrentVelocity) const
 {
-	FVector RootMotionVelocity = RootMotionDeltaMove / DeltaSeconds;
-	
-	// Do not override Velocity.Z if in falling physics, we want to keep the effect of gravity.
-	if (IsFalling())
+	if (ensure(DeltaSeconds > 0.f))
 	{
-		RootMotionVelocity.Z = CurrentVelocity.Z;
+		FVector RootMotionVelocity = RootMotionDeltaMove / DeltaSeconds;
+
+		// Do not override Velocity.Z if in falling physics, we want to keep the effect of gravity.
+		if (IsFalling())
+		{
+			RootMotionVelocity.Z = CurrentVelocity.Z;
+		}
+
+		return RootMotionVelocity;
+	}
+	else
+	{
+		return CurrentVelocity;
 	}
-	
-	return RootMotionVelocity;
 }
 
 
