fixed UE-19834 A change in main broke material post process blendable references

 solved backward compatibility of PostProcessComponent for Blendables (lost the array elements recently when the feature was updated)

diff --git a/Engine/Source/Runtime/Engine/Classes/Components/PostProcessComponent.h b/Engine/Source/Runtime/Engine/Classes/Components/PostProcessComponent.h
--- a/Engine/Source/Runtime/Engine/Classes/Components/PostProcessComponent.h
+++ b/Engine/Source/Runtime/Engine/Classes/Components/PostProcessComponent.h
@@ -63,6 +63,10 @@ protected:
 
 	virtual void OnRegister() override;
 	virtual void OnUnregister() override;
+
+	// UObject interface
+	virtual void Serialize(FArchive& Ar) override;
+	// End of UObject interface
 };
 
 
diff --git a/Engine/Source/Runtime/Engine/Private/UnrealClient.cpp b/Engine/Source/Runtime/Engine/Private/UnrealClient.cpp
--- a/Engine/Source/Runtime/Engine/Private/UnrealClient.cpp
+++ b/Engine/Source/Runtime/Engine/Private/UnrealClient.cpp
@@ -909,6 +909,16 @@ void UPostProcessComponent::OnUnregister()
 	GetWorld()->PostProcessVolumes.RemoveSingle(this);
 }
 
+void UPostProcessComponent::Serialize(FArchive& Ar)
+{
+	Super::Serialize(Ar);
+
+	if(Ar.IsLoading())
+	{
+		Settings.OnAfterLoad();
+	}
+}
+
 /**
 *	Starts a new rendering frame. Called from the game thread thread.
 */
