Physics: Optimization for ConvertOverlapResults when there are several overlaps

[INTEGRATE] Change 2650497 by [REDACTED_PERSON][REDACTED_HANDLE] on 2015/08/10 21:13:19

diff --git a/Engine/Source/Runtime/Engine/Private/Collision/CollisionConversions.cpp b/Engine/Source/Runtime/Engine/Private/Collision/CollisionConversions.cpp
--- a/Engine/Source/Runtime/Engine/Private/Collision/CollisionConversions.cpp
+++ b/Engine/Source/Runtime/Engine/Private/Collision/CollisionConversions.cpp
@@ -11,6 +11,29 @@
 #include "Components/LineBatchComponent.h"
 #include "PhysicalMaterials/PhysicalMaterial.h"
 
+// Used to place overlaps into a TMap when deduplicating them
+struct FOverlapKey
+{
+	UPrimitiveComponent* Component;
+	int32 ComponentIndex;
+
+	FOverlapKey(UPrimitiveComponent* InComponent, int32 InComponentIndex)
+		: Component(InComponent)
+		, ComponentIndex(InComponentIndex)
+	{
+	}
+
+	friend bool operator==(const FOverlapKey& X, const FOverlapKey& Y)
+	{
+		return (X.Component == Y.Component) && (X.ComponentIndex == Y.ComponentIndex);
+	}
+};
+
+uint32 GetTypeHash(const FOverlapKey& Key)
+{
+	return GetTypeHash(Key.Component) ^ GetTypeHash(Key.ComponentIndex);
+}
+
 
 #define DRAW_OVERLAPPING_TRIS (!(UE_BUILD_SHIPPING || UE_BUILD_TEST))
 extern TAutoConsoleVariable<int32> CVarShowInitialOverlaps;
@@ -996,19 +1019,62 @@ bool ConvertOverlapResults(int32 NumOverlaps, PxOverlapHit* POverlapResults, con
 	OutOverlaps.Reserve(OutOverlaps.Num() + NumOverlaps);
 	bool bBlockingFound = false;
 
-	for(int32 i=0; i<NumOverlaps; i++)
+	// This number was not empirically determined, just a rough rule of thumb
+	if (OutOverlaps.Num() + NumOverlaps < 6)
 	{
-		FOverlapResult NewOverlap;		
-		
-		ConvertQueryOverlap( POverlapResults[i].shape, POverlapResults[i].actor, NewOverlap, QueryFilter);
+		// N^2 approach, no maps
+		for (int32 i = 0; i < NumOverlaps; i++)
+		{
+			FOverlapResult NewOverlap;
+			ConvertQueryOverlap(POverlapResults[i].shape, POverlapResults[i].actor, NewOverlap, QueryFilter);
 
+			if (NewOverlap.bBlockingHit)
+			{
+				bBlockingFound = true;
+			}
 
-		if(NewOverlap.bBlockingHit)
+			AddUniqueOverlap(OutOverlaps, NewOverlap);
+		}
+	}
+	else
+	{
+		// Map from an overlap to the position in the result array
+		TMap<FOverlapKey, int32> OverlapMap;
+		OverlapMap.Reserve(OutOverlaps.Num());
+
+		// Fill in the map with existing hits
+		for (int32 ExistingIndex = 0; ExistingIndex < OutOverlaps.Num(); ++ExistingIndex)
 		{
-			bBlockingFound = true;
+			const FOverlapResult& ExistingOverlap = OutOverlaps[ExistingIndex];
+			OverlapMap.Add(FOverlapKey(ExistingOverlap.Component.Get(), ExistingOverlap.ItemIndex), ExistingIndex);
 		}
 
-		AddUniqueOverlap(OutOverlaps, NewOverlap);
+		for (int32 PResultIndex = 0; PResultIndex < NumOverlaps; ++PResultIndex)
+		{
+			FOverlapResult NewOverlap;
+			ConvertQueryOverlap(POverlapResults[PResultIndex].shape, POverlapResults[PResultIndex].actor, NewOverlap, QueryFilter);
+
+			if (NewOverlap.bBlockingHit)
+			{
+				bBlockingFound = true;
+			}
+
+			int32& DestinationIndex = OverlapMap.FindOrAdd(FOverlapKey(NewOverlap.Component.Get(), NewOverlap.ItemIndex));
+			if (DestinationIndex < OutOverlaps.Num())
+			{
+				FOverlapResult& ExistingOverlap = OutOverlaps[DestinationIndex];
+
+				// If we had a non-blocking overlap with this component, but now we have a blocking one, use that one instead!
+				if (!ExistingOverlap.bBlockingHit && NewOverlap.bBlockingHit)
+				{
+					ExistingOverlap = NewOverlap;
+				}
+			}
+			else
+			{
+				DestinationIndex = OutOverlaps.Add(NewOverlap);
+			}
+		}
 	}
 
 	return bBlockingFound;
