Engine: Disabled ticking by default for ACameraActor, there is no work to be done per-tick (subclasses can still set it)

[INTEGRATE] Change 2650543 by [REDACTED_PERSON][REDACTED_HANDLE] on 2015/08/10 22:10:48

diff --git a/Engine/Source/Runtime/Engine/Private/Camera/CameraActor.cpp b/Engine/Source/Runtime/Engine/Private/Camera/CameraActor.cpp
--- a/Engine/Source/Runtime/Engine/Private/Camera/CameraActor.cpp
+++ b/Engine/Source/Runtime/Engine/Private/Camera/CameraActor.cpp
@@ -32,8 +32,6 @@ ACameraActor::ACameraActor(const FObjectInitializer& ObjectInitializer)
 	AspectRatio_DEPRECATED = 1.777778f;
 	PostProcessBlendWeight_DEPRECATED = 1.0f;
 	// End of deprecated property initialization
-
-	PrimaryActorTick.bCanEverTick = true;
 }
 
 void ACameraActor::Serialize(FArchive& Ar)
