UEBP-40 C++ code generation (WIP):
- timelines
- component delegate binding
- various improvements

diff --git a/Engine/Source/Developer/BlueprintCompilerCppBackend/Private/BlueprintCompilerCppBackend.cpp b/Engine/Source/Developer/BlueprintCompilerCppBackend/Private/BlueprintCompilerCppBackend.cpp
--- a/Engine/Source/Developer/BlueprintCompilerCppBackend/Private/BlueprintCompilerCppBackend.cpp
+++ b/Engine/Source/Developer/BlueprintCompilerCppBackend/Private/BlueprintCompilerCppBackend.cpp
@@ -230,7 +230,7 @@ void FBlueprintCompilerCppBackend::EmitCallDelegateStatment(FKismetFunctionConte
 
 	Emit(Body, TEXT("\t\t\t"));
 	Emit(Body, *SafeContextScope.GetAdditionalIndent());
-	Emit(Body, *FString::Printf(TEXT("%s.Broadcast(%s)"), *TermToText(Statement.FunctionContext, false), *EmitMethodInputParameterList(Statement)));
+	Emit(Body, *FString::Printf(TEXT("%s.Broadcast(%s);"), *TermToText(Statement.FunctionContext, false), *EmitMethodInputParameterList(Statement)));
 }
 
 FString FBlueprintCompilerCppBackend::EmitMethodInputParameterList(FBlueprintCompiledStatement& Statement)
diff --git a/Engine/Source/Developer/BlueprintCompilerCppBackend/Private/BlueprintCompilerCppBackendUtils.cpp b/Engine/Source/Developer/BlueprintCompilerCppBackend/Private/BlueprintCompilerCppBackendUtils.cpp
--- a/Engine/Source/Developer/BlueprintCompilerCppBackend/Private/BlueprintCompilerCppBackendUtils.cpp
+++ b/Engine/Source/Developer/BlueprintCompilerCppBackend/Private/BlueprintCompilerCppBackendUtils.cpp
@@ -6,6 +6,7 @@
 #include "EdGraphSchema_K2.h"
 #include "Editor/UnrealEd/Public/Kismet2/StructureEditorUtils.h"
 #include "Engine/InheritableComponentHandler.h"
+#include "Engine/DynamicBlueprintBinding.h"
 
 FString FSafeContextScopedEmmitter::GetAdditionalIndent() const
 {
@@ -532,10 +533,11 @@ FString FEmitHelper::LiteralTerm(const FEdGraphPinType& Type, const FString& Cus
 		else
 		{
 			//@todo:  This needs to be more robust, since import text isn't really proper for struct construction.
-			ensure(CustomValue.IsEmpty());
+			const bool bEmptyCustomValue = CustomValue.IsEmpty() || (CustomValue == TEXT("()"));
+			ensure(bEmptyCustomValue);
 			
 			const FString StructName = FString(TEXT("F")) + StructType->GetName();
-			const FString ConstructBody = CustomValue.IsEmpty()
+			const FString ConstructBody = bEmptyCustomValue
 				? (StructType->IsA<UUserDefinedStruct>() ? TEXT("::GetDefaultValue()") : TEXT("{}"))
 				: CustomValue;
 			return StructName + ConstructBody;
@@ -1041,6 +1043,15 @@ FString FEmitDefaultValueHelper::GenerateConstructor(UClass* InBPGC)
 				HandleClassSubobject(Context, TimelineTemplate);
 			}
 		}
+
+		for (auto DynamicBindingObject : BPGC->DynamicBindingObjects)
+		{
+			if (DynamicBindingObject)
+			{
+				HandleClassSubobject(Context, DynamicBindingObject);
+			}
+		}
+
 		Context.DecreaseIndent();
 		Context.AddLine(TEXT("}"));
 
@@ -1071,23 +1082,7 @@ FString FEmitDefaultValueHelper::GenerateConstructor(UClass* InBPGC)
 	// TIMELINES
 	TSet<const UProperty*> HandledProperties;
 	const bool bIsActor = BPGC->IsChildOf<AActor>();
-	if (bIsActor)
-	{
-		for (auto TimelineTemplate : BPGC->Timelines)
-		{
-			FString NativeName;
-			if (Context.FindObject(TimelineTemplate, NativeName))
-			{
-				Context.AddLine(FString::Printf(TEXT("UBlueprintGeneratedClass::CreateTimelineComponent(this, %s);"), *NativeName));
-				auto Prop = FindField<UObjectPropertyBase>(BPGC, *UTimelineTemplate::TimelineTemplateNameToVariableName(TimelineTemplate->GetFName()));
-				if (Prop)
-				{
-					HandledProperties.Add(Prop);
-				}
-			}
-		}
-	}
-	
+
 	// COMPONENTS
 	for (auto Property : TFieldRange<const UProperty>(BPGC))
 	{
diff --git a/Engine/Source/Editor/KismetCompiler/Private/KismetCompiler.cpp b/Engine/Source/Editor/KismetCompiler/Private/KismetCompiler.cpp
--- a/Engine/Source/Editor/KismetCompiler/Private/KismetCompiler.cpp
+++ b/Engine/Source/Editor/KismetCompiler/Private/KismetCompiler.cpp
@@ -1833,7 +1833,7 @@ void FKismetCompilerContext::BuildDynamicBindingObjects(UBlueprintGeneratedClass
 
 				if (DynamicBindingClass)
 				{
-					UDynamicBlueprintBinding* DynamicBindingObject = Class->GetDynamicBindingObject(DynamicBindingClass);
+					UDynamicBlueprintBinding* DynamicBindingObject = UBlueprintGeneratedClass::GetDynamicBindingObject(Class, DynamicBindingClass);
 					if (DynamicBindingObject == NULL)
 					{
 						DynamicBindingObject = NewObject<UDynamicBlueprintBinding>(Class, DynamicBindingClass);
diff --git a/Engine/Source/Runtime/Engine/Classes/Engine/BlueprintGeneratedClass.h b/Engine/Source/Runtime/Engine/Classes/Engine/BlueprintGeneratedClass.h
--- a/Engine/Source/Runtime/Engine/Classes/Engine/BlueprintGeneratedClass.h
+++ b/Engine/Source/Runtime/Engine/Classes/Engine/BlueprintGeneratedClass.h
@@ -451,7 +451,7 @@ public:
 	UActorComponent* FindComponentTemplateByName(const FName& TemplateName) const;
 
 	/** Create Timeline objects for this Actor based on the Timelines array*/
-	virtual void CreateComponentsForActor(AActor* Actor) const;
+	static void CreateComponentsForActor(const UClass* ThisClass, AActor* Actor);
 	static void CreateTimelineComponent(AActor* Actor, const UTimelineTemplate* TimelineTemplate);
 
 	// UObject interface
@@ -493,13 +493,15 @@ public:
 #endif
 
 	/** Bind functions on supplied actor to delegates */
-	void BindDynamicDelegates(UObject* InInstance) const;
+	static void BindDynamicDelegates(const UClass* ThisClass, UObject* InInstance);
 
+	// Finds the desired dynamic binding object for this blueprint generated class
+	static UDynamicBlueprintBinding* GetDynamicBindingObject(const UClass* ThisClass, UClass* BindingClass);
+
+#if WITH_EDITOR
 	/** Unbind functions on supplied actor from delegates tied to a specific property */
 	void UnbindDynamicDelegatesForProperty(UObject* InInstance, const UObjectProperty* InObjectProperty);
-	
-	// Finds the desired dynamic binding object for this blueprint generated class
-	UDynamicBlueprintBinding* GetDynamicBindingObject(UClass* InClass) const;
+#endif
 
 	/** called to gather blueprint replicated properties */
 	virtual void GetLifetimeBlueprintReplicationList(TArray<class FLifetimeProperty>& OutLifetimeProps) const;
diff --git a/Engine/Source/Runtime/Engine/Classes/Engine/InputDelegateBinding.h b/Engine/Source/Runtime/Engine/Classes/Engine/InputDelegateBinding.h
--- a/Engine/Source/Runtime/Engine/Classes/Engine/InputDelegateBinding.h
+++ b/Engine/Source/Runtime/Engine/Classes/Engine/InputDelegateBinding.h
@@ -32,6 +32,6 @@ class ENGINE_API UInputDelegateBinding : public UDynamicBlueprintBinding
 	GENERATED_UCLASS_BODY()
 
 	virtual void BindToInputComponent(UInputComponent* InputComponent) const { };
-
-	static void BindInputDelegates(const UBlueprintGeneratedClass* BGClass, UInputComponent* InputComponent);
+	static bool SupportsInputDelegate(const UClass* InClass);
+	static void BindInputDelegates(const UClass* InClass, UInputComponent* InputComponent);
 };
diff --git a/Engine/Source/Runtime/Engine/Private/Actor.cpp b/Engine/Source/Runtime/Engine/Private/Actor.cpp
--- a/Engine/Source/Runtime/Engine/Private/Actor.cpp
+++ b/Engine/Source/Runtime/Engine/Private/Actor.cpp
@@ -2816,11 +2816,9 @@ void AActor::EnableInput(APlayerController* PlayerController)
 			InputComponent->bBlockInput = bBlockInput;
 			InputComponent->Priority = InputPriority;
 
-			// Only do this if this actor is of a blueprint class
-			UBlueprintGeneratedClass* BGClass = Cast<UBlueprintGeneratedClass>(GetClass());
-			if(BGClass != NULL)
+			if (UInputDelegateBinding::SupportsInputDelegate(GetClass()))
 			{
-				UInputDelegateBinding::BindInputDelegates(BGClass, InputComponent);
+				UInputDelegateBinding::BindInputDelegates(GetClass(), InputComponent);
 			}
 		}
 		else
diff --git a/Engine/Source/Runtime/Engine/Private/ActorConstruction.cpp b/Engine/Source/Runtime/Engine/Private/ActorConstruction.cpp
--- a/Engine/Source/Runtime/Engine/Private/ActorConstruction.cpp
+++ b/Engine/Source/Runtime/Engine/Private/ActorConstruction.cpp
@@ -445,7 +445,7 @@ void AActor::ExecuteConstruction(const FTransform& Transform, const FComponentIn
 					CurrentBPGClass->SimpleConstructionScript->ExecuteScriptOnActor(this, Transform, bIsDefaultTransform);
 				}
 				// Now that the construction scripts have been run, we can create timelines and hook them up
-				CurrentBPGClass->CreateComponentsForActor(this);
+				UBlueprintGeneratedClass::CreateComponentsForActor(CurrentBPGClass, this);
 			}
 
 			// If we passed in cached data, we apply it now, so that the UserConstructionScript can use the updated values
@@ -478,7 +478,7 @@ void AActor::ExecuteConstruction(const FTransform& Transform, const FComponentIn
 			}
 
 			// Bind any delegates on components			
-			((UBlueprintGeneratedClass*)GetClass())->BindDynamicDelegates(this); // We have a BP stack, we must have a UBlueprintGeneratedClass...
+			UBlueprintGeneratedClass::BindDynamicDelegates(GetClass(), this); // We have a BP stack, we must have a UBlueprintGeneratedClass...
 
 			// Apply any cached data procedural components
 			// @TODO Don't re-apply to components we already applied to above
@@ -508,6 +508,7 @@ void AActor::ExecuteConstruction(const FTransform& Transform, const FComponentIn
 	}
 	else
 	{
+		UBlueprintGeneratedClass::CreateComponentsForActor(GetClass(), this);
 #if WITH_EDITOR
 		bool bDoUserConstructionScript;
 		GConfig->GetBool(TEXT("Kismet"), TEXT("bTurnOffEditorConstructionScript"), bDoUserConstructionScript, GEngineIni);
@@ -517,6 +518,7 @@ void AActor::ExecuteConstruction(const FTransform& Transform, const FComponentIn
 			// Then run the user script, which is responsible for calling its own super, if desired
 			ProcessUserConstructionScript();
 		}
+		UBlueprintGeneratedClass::BindDynamicDelegates(GetClass(), this);
 	}
 
 	GetWorld()->UpdateCullDistanceVolumes(this);
diff --git a/Engine/Source/Runtime/Engine/Private/BlueprintGeneratedClass.cpp b/Engine/Source/Runtime/Engine/Private/BlueprintGeneratedClass.cpp
--- a/Engine/Source/Runtime/Engine/Private/BlueprintGeneratedClass.cpp
+++ b/Engine/Source/Runtime/Engine/Private/BlueprintGeneratedClass.cpp
@@ -328,44 +328,74 @@ UObject* UBlueprintGeneratedClass::FindArchetype(UClass* ArchetypeClass, const F
 	return Archetype;
 }
 
-UDynamicBlueprintBinding* UBlueprintGeneratedClass::GetDynamicBindingObject(UClass* Class) const
+UDynamicBlueprintBinding* UBlueprintGeneratedClass::GetDynamicBindingObject(const UClass* ThisClass, UClass* BindingClass)
 {
-	UDynamicBlueprintBinding* DynamicBindingObject = NULL;
-	for (int32 Index = 0; Index < DynamicBindingObjects.Num(); ++Index)
+	check(ThisClass);
+	UDynamicBlueprintBinding* DynamicBlueprintBinding = nullptr;
+	if (auto BPGC = Cast<UBlueprintGeneratedClass>(ThisClass))
 	{
-		if (DynamicBindingObjects[Index]->GetClass() == Class)
+		for (auto DynamicBindingObject : BPGC->DynamicBindingObjects)
 		{
-			DynamicBindingObject = DynamicBindingObjects[Index];
-			break;
+			if (DynamicBindingObject && (DynamicBindingObject->GetClass() == BindingClass))
+			{
+				DynamicBlueprintBinding = DynamicBindingObject;
+				break;
+			}
 		}
 	}
-	return DynamicBindingObject;
+	else
+	{
+		for (auto MiscObj : ThisClass->MiscObjects)
+		{
+			auto DynamicBindingObject = Cast<UDynamicBlueprintBinding>(MiscObj);
+			if (DynamicBindingObject && (DynamicBindingObject->GetClass() == BindingClass))
+			{
+				DynamicBlueprintBinding = DynamicBindingObject;
+				break;
+			}
+		}
+	}
+	return DynamicBlueprintBinding;
 }
 
-void UBlueprintGeneratedClass::BindDynamicDelegates(UObject* InInstance) const
+void UBlueprintGeneratedClass::BindDynamicDelegates(const UClass* ThisClass, UObject* InInstance)
 {
-	if(!InInstance->IsA(this))
+	check(ThisClass && InInstance);
+	if (!InInstance->IsA(ThisClass))
 	{
-		UE_LOG(LogBlueprint, Warning, TEXT("BindComponentDelegates: '%s' is not an instance of '%s'."), *InInstance->GetName(), *GetName());
+		UE_LOG(LogBlueprint, Warning, TEXT("BindComponentDelegates: '%s' is not an instance of '%s'."), *InInstance->GetName(), *ThisClass->GetName());
 		return;
 	}
 
-	for (int32 Index = 0; Index < DynamicBindingObjects.Num(); ++Index)
+	if (auto BPGC = Cast<UBlueprintGeneratedClass>(ThisClass))
 	{
-		if ( ensure(DynamicBindingObjects[Index] != NULL) )
+		for (auto DynamicBindingObject : BPGC->DynamicBindingObjects)
 		{
-			DynamicBindingObjects[Index]->BindDynamicDelegates(InInstance);
+			if (ensure(DynamicBindingObject))
+			{
+				DynamicBindingObject->BindDynamicDelegates(InInstance);
+			}
+		}
+	}
+	else
+	{
+		for (auto MiscObj : ThisClass->MiscObjects)
+		{
+			auto DynamicBindingObject = Cast<UDynamicBlueprintBinding>(MiscObj);
+			if (DynamicBindingObject)
+			{
+				DynamicBindingObject->BindDynamicDelegates(InInstance);
+			}
 		}
 	}
 
-	// call on super class, if it's a BlueprintGeneratedClass
-	UBlueprintGeneratedClass* BGClass = Cast<UBlueprintGeneratedClass>(GetSuperStruct());
-	if(BGClass != NULL)
+	if (auto TheSuperClass = ThisClass->GetSuperClass())
 	{
-		BGClass->BindDynamicDelegates(InInstance);
+		BindDynamicDelegates(TheSuperClass, InInstance);
 	}
 }
 
+#if WITH_EDITOR
 void UBlueprintGeneratedClass::UnbindDynamicDelegatesForProperty(UObject* InInstance, const UObjectProperty* InObjectProperty)
 {
 	for (int32 Index = 0; Index < DynamicBindingObjects.Num(); ++Index)
@@ -376,6 +406,7 @@ void UBlueprintGeneratedClass::UnbindDynamicDelegatesForProperty(UObject* InInst
 		}
 	}
 }
+#endif
 
 bool UBlueprintGeneratedClass::GetGeneratedClassesHierarchy(const UClass* InClass, TArray<const UBlueprintGeneratedClass*>& OutBPGClasses)
 {
@@ -520,21 +551,33 @@ void UBlueprintGeneratedClass::CreateTimelineComponent(AActor* Actor, const UTim
 	}
 }
 
-void UBlueprintGeneratedClass::CreateComponentsForActor(AActor* Actor) const
+void UBlueprintGeneratedClass::CreateComponentsForActor(const UClass* ThisClass, AActor* Actor)
 {
-	check(Actor != NULL);
+	check(ThisClass && Actor);
 	check(!Actor->IsTemplate());
 	check(!Actor->IsPendingKill());
 
-	// Iterate over each timeline template
-	for(int32 i=0; i<Timelines.Num(); i++)
+	if (auto BPGC = Cast<const UBlueprintGeneratedClass>(ThisClass))
 	{
-		const UTimelineTemplate* TimelineTemplate = Timelines[i];
-
-		// Not fatal if NULL, but shouldn't happen and ignored if not wired up in graph
-		if(TimelineTemplate && TimelineTemplate->bValidatedAsWired)
+		for (auto TimelineTemplate : BPGC->Timelines)
 		{
-			CreateTimelineComponent(Actor, TimelineTemplate);
+			// Not fatal if NULL, but shouldn't happen and ignored if not wired up in graph
+			if (TimelineTemplate && TimelineTemplate->bValidatedAsWired)
+			{
+				CreateTimelineComponent(Actor, TimelineTemplate);
+			}
+		}
+	}
+	else
+	{
+		for (auto MiscObj : ThisClass->MiscObjects)
+		{
+			auto TimelineTemplate = Cast<const UTimelineTemplate>(MiscObj);
+			// Not fatal if NULL, but shouldn't happen and ignored if not wired up in graph
+			if (TimelineTemplate && TimelineTemplate->bValidatedAsWired)
+			{
+				CreateTimelineComponent(Actor, TimelineTemplate);
+			}
 		}
 	}
 }
diff --git a/Engine/Source/Runtime/Engine/Private/InputDelegateBinding.cpp b/Engine/Source/Runtime/Engine/Private/InputDelegateBinding.cpp
--- a/Engine/Source/Runtime/Engine/Private/InputDelegateBinding.cpp
+++ b/Engine/Source/Runtime/Engine/Private/InputDelegateBinding.cpp
@@ -15,7 +15,15 @@ UInputDelegateBinding::UInputDelegateBinding(const FObjectInitializer& ObjectIni
 {
 }
 
-void UInputDelegateBinding::BindInputDelegates(const UBlueprintGeneratedClass* BGClass, UInputComponent* InputComponent)
+bool UInputDelegateBinding::SupportsInputDelegate(const UClass* InClass)
+{
+	return InClass && (
+		InClass->MiscObjects.Num() 
+		|| Cast<UBlueprintGeneratedClass>(InClass)
+		|| SupportsInputDelegate(InClass->GetSuperClass()));
+}
+
+void UInputDelegateBinding::BindInputDelegates(const UClass* InClass, UInputComponent* InputComponent)
 {
 	static UClass* InputBindingClasses[] = { 
 												UInputActionDelegateBinding::StaticClass(), 
@@ -26,13 +34,15 @@ void UInputDelegateBinding::BindInputDelegates(const UBlueprintGeneratedClass* B
 												UInputVectorAxisDelegateBinding::StaticClass(),
 										   };
 
-	if (BGClass)
+	if (InClass)
 	{
-		BindInputDelegates(Cast<UBlueprintGeneratedClass>(BGClass->GetSuperStruct()), InputComponent);
+		BindInputDelegates(InClass->GetSuperClass(), InputComponent);
 
 		for (int32 Index = 0; Index < ARRAY_COUNT(InputBindingClasses); ++Index)
 		{
-			UInputDelegateBinding* BindingObject = CastChecked<UInputDelegateBinding>(BGClass->GetDynamicBindingObject(InputBindingClasses[Index]), ECastCheckedType::NullAllowed);
+			UInputDelegateBinding* BindingObject = CastChecked<UInputDelegateBinding>(
+				UBlueprintGeneratedClass::GetDynamicBindingObject(InClass, InputBindingClasses[Index])
+				, ECastCheckedType::NullAllowed);
 			if (BindingObject)
 			{
 				BindingObject->BindToInputComponent(InputComponent);
diff --git a/Engine/Source/Runtime/Engine/Private/LevelScriptActor.cpp b/Engine/Source/Runtime/Engine/Private/LevelScriptActor.cpp
--- a/Engine/Source/Runtime/Engine/Private/LevelScriptActor.cpp
+++ b/Engine/Source/Runtime/Engine/Private/LevelScriptActor.cpp
@@ -54,16 +54,14 @@ void ALevelScriptActor::BeginDestroy()
 
 void ALevelScriptActor::PreInitializeComponents()
 {
-	UBlueprintGeneratedClass* BGClass = Cast<UBlueprintGeneratedClass>(GetClass());
-	if(BGClass != NULL)
+	if (UInputDelegateBinding::SupportsInputDelegate(GetClass()))
 	{
 		// create an InputComponent object so that the level script actor can bind key events
 		InputComponent = NewObject<UInputComponent>(this);
 		InputComponent->RegisterComponent();
 
-		UInputDelegateBinding::BindInputDelegates(BGClass, InputComponent);
+		UInputDelegateBinding::BindInputDelegates(GetClass(), InputComponent);
 	}
-
 	Super::PreInitializeComponents();
 }
 
diff --git a/Engine/Source/Runtime/Engine/Private/Pawn.cpp b/Engine/Source/Runtime/Engine/Private/Pawn.cpp
--- a/Engine/Source/Runtime/Engine/Private/Pawn.cpp
+++ b/Engine/Source/Runtime/Engine/Private/Pawn.cpp
@@ -360,12 +360,12 @@ void APawn::PawnClientRestart()
 			{
 				SetupPlayerInputComponent(InputComponent);
 				InputComponent->RegisterComponent();
-				UBlueprintGeneratedClass* BGClass = Cast<UBlueprintGeneratedClass>(GetClass());
-				if(BGClass != NULL)
+				if (UInputDelegateBinding::SupportsInputDelegate(GetClass()))
 				{
 					InputComponent->bBlockInput = bBlockInput;
-					UInputDelegateBinding::BindInputDelegates(BGClass, InputComponent);
+					UInputDelegateBinding::BindInputDelegates(GetClass(), InputComponent);
 				}
+
 			}
 		}
 	}
diff --git a/Engine/Source/Runtime/Engine/Private/PlayerController.cpp b/Engine/Source/Runtime/Engine/Private/PlayerController.cpp
--- a/Engine/Source/Runtime/Engine/Private/PlayerController.cpp
+++ b/Engine/Source/Runtime/Engine/Private/PlayerController.cpp
@@ -2207,12 +2207,10 @@ void APlayerController::SetupInputComponent()
 		InputComponent->RegisterComponent();
 	}
 
-	// Only do this if this actor is of a blueprint class
-	UBlueprintGeneratedClass* BGClass = Cast<UBlueprintGeneratedClass>(GetClass());
-	if(BGClass != NULL)
+	if (UInputDelegateBinding::SupportsInputDelegate(GetClass()))
 	{
 		InputComponent->bBlockInput = bBlockInput;
-		UInputDelegateBinding::BindInputDelegates(BGClass, InputComponent);
+		UInputDelegateBinding::BindInputDelegates(GetClass(), InputComponent);
 	}
 }
 
diff --git a/Engine/Source/Runtime/UMG/Private/WidgetBlueprintGeneratedClass.cpp b/Engine/Source/Runtime/UMG/Private/WidgetBlueprintGeneratedClass.cpp
--- a/Engine/Source/Runtime/UMG/Private/WidgetBlueprintGeneratedClass.cpp
+++ b/Engine/Source/Runtime/UMG/Private/WidgetBlueprintGeneratedClass.cpp
@@ -126,7 +126,7 @@ void UWidgetBlueprintGeneratedClass::InitializeWidget(UUserWidget* UserWidget) c
 		});
 
 		// Bind any delegates on widgets
-		BindDynamicDelegates(UserWidget);
+		BindDynamicDelegates(this, UserWidget);
 
 		//TODO UMG Add OnWidgetInitialized?
 	}
